Add 'Optional' EditorOnly data for UMaterialInterface and UMaterialFunctionInterface

diff --git a/Engine/Source/Runtime/Engine/Private/Materials/HLSLMaterialTranslator.cpp b/Engine/Source/Runtime/Engine/Private/Materials/HLSLMaterialTranslator.cpp
--- a/Engine/Source/Runtime/Engine/Private/Materials/HLSLMaterialTranslator.cpp
+++ b/Engine/Source/Runtime/Engine/Private/Materials/HLSLMaterialTranslator.cpp
@@ -257,6 +257,8 @@ FHLSLMaterialTranslator::FHLSLMaterialTranslator(FMaterial* InMaterial,
 
 	FMemory::Memzero(MaterialAttributesReturned);
 
+	InStaticParameters.GetMaterialLayers(CachedMaterialLayers);
+
 	SharedPixelProperties[MP_Normal] = true;
 	SharedPixelProperties[MP_Tangent] = true;
 	SharedPixelProperties[MP_EmissiveColor] = true;
@@ -357,7 +359,8 @@ void FHLSLMaterialTranslator::AssignShaderFrequencyScope(EShaderFrequency InShad
 	CurrentScopeID = (uint64)InShaderFrequency;
 }
 
-void FHLSLMaterialTranslator::GatherCustomVertexInterpolators(TArray<UMaterialExpression*> Expressions)
+template<typename ExpressionsArrayType>
+void FHLSLMaterialTranslator::GatherCustomVertexInterpolators(const ExpressionsArrayType& Expressions)
 {
 	for (UMaterialExpression* Expression : Expressions)
 	{
@@ -396,10 +399,7 @@ void FHLSLMaterialTranslator::GatherCustomVertexInterpolators(TArray<UMaterialEx
 				FunctionCall->LinkFunctionIntoCaller(this);
 				PushFunction(&LocalState);
 
-				if (const TArray<TObjectPtr<UMaterialExpression>>* FunctionExpressions = FunctionCall->MaterialFunction->GetFunctionExpressions())
-				{
-					GatherCustomVertexInterpolators(*FunctionExpressions);
-				}
+				GatherCustomVertexInterpolators(FunctionCall->MaterialFunction->GetExpressions());
 
 				FMaterialFunctionCompileState* CompileState = PopFunction();
 				check(CompileState->ExpressionStack.Num() == 0);
@@ -423,11 +423,8 @@ void FHLSLMaterialTranslator::GatherCustomVertexInterpolators(TArray<UMaterialEx
 						FMaterialFunctionCompileState LocalState(Layer);
 						Layer->LinkFunctionIntoCaller(this);
 						PushFunction(&LocalState);
-
-						if (const TArray<TObjectPtr<UMaterialExpression>>* FunctionExpressions = Layer->MaterialFunction->GetFunctionExpressions())
-						{
-							GatherCustomVertexInterpolators(*FunctionExpressions);
-						}
+						
+						GatherCustomVertexInterpolators(Layer->MaterialFunction->GetExpressions());
 
 						FMaterialFunctionCompileState* CompileState = PopFunction();
 						check(CompileState->ExpressionStack.Num() == 0);
@@ -443,10 +440,7 @@ void FHLSLMaterialTranslator::GatherCustomVertexInterpolators(TArray<UMaterialEx
 						Blend->LinkFunctionIntoCaller(this);
 						PushFunction(&LocalState);
 
-						if (const TArray<TObjectPtr<UMaterialExpression>>* FunctionExpressions = Blend->MaterialFunction->GetFunctionExpressions())
-						{
-							GatherCustomVertexInterpolators(*FunctionExpressions);
-						}
+						GatherCustomVertexInterpolators(Blend->MaterialFunction->GetExpressions());
 
 						FMaterialFunctionCompileState* CompileState = PopFunction();
 						check(CompileState->ExpressionStack.Num() == 0);
@@ -509,7 +503,8 @@ void FHLSLMaterialTranslator::CompileCustomOutputs(TArray<UMaterialExpressionCus
 	}
 }
 
-EMaterialExpressionVisitResult FHLSLMaterialTranslator::VisitExpressionsRecursive(TArray<UMaterialExpression*> Expressions, IMaterialExpressionVisitor& InVisitor)
+template<typename ExpressionsArrayType>
+EMaterialExpressionVisitResult FHLSLMaterialTranslator::VisitExpressionsRecursive(const ExpressionsArrayType& Expressions, IMaterialExpressionVisitor& InVisitor)
 {
 	EMaterialExpressionVisitResult VisitResult = MVR_CONTINUE;
 	for (UMaterialExpression* Expression : Expressions)
@@ -528,10 +523,7 @@ EMaterialExpressionVisitResult FHLSLMaterialTranslator::VisitExpressionsRecursiv
 				FunctionCall->LinkFunctionIntoCaller(this);
 				PushFunction(&LocalState);
 
-				if (const TArray<TObjectPtr<UMaterialExpression>>* FunctionExpressions = FunctionCall->MaterialFunction->GetFunctionExpressions())
-				{
-					VisitResult = VisitExpressionsRecursive(*FunctionExpressions, InVisitor);
-				}
+				VisitResult = VisitExpressionsRecursive(FunctionCall->MaterialFunction->GetExpressions(), InVisitor);
 
 				FMaterialFunctionCompileState* CompileState = PopFunction();
 				check(CompileState->ExpressionStack.Num() == 0);
@@ -561,10 +553,7 @@ EMaterialExpressionVisitResult FHLSLMaterialTranslator::VisitExpressionsRecursiv
 						Layer->LinkFunctionIntoCaller(this);
 						PushFunction(&LocalState);
 
-						if (const TArray<TObjectPtr<UMaterialExpression>>* FunctionExpressions = Layer->MaterialFunction->GetFunctionExpressions())
-						{
-							VisitResult = VisitExpressionsRecursive(*FunctionExpressions, InVisitor);
-						}
+						VisitResult = VisitExpressionsRecursive(Layer->MaterialFunction->GetExpressions(), InVisitor);
 
 						FMaterialFunctionCompileState* CompileState = PopFunction();
 						check(CompileState->ExpressionStack.Num() == 0);
@@ -585,10 +574,7 @@ EMaterialExpressionVisitResult FHLSLMaterialTranslator::VisitExpressionsRecursiv
 						Blend->LinkFunctionIntoCaller(this);
 						PushFunction(&LocalState);
 
-						if (const TArray<TObjectPtr<UMaterialExpression>>* FunctionExpressions = Blend->MaterialFunction->GetFunctionExpressions())
-						{
-							VisitResult = VisitExpressionsRecursive(*FunctionExpressions, InVisitor);
-						}
+						VisitResult = VisitExpressionsRecursive(Blend->MaterialFunction->GetExpressions(), InVisitor);
 
 						FMaterialFunctionCompileState* CompileState = PopFunction();
 						check(CompileState->ExpressionStack.Num() == 0);
@@ -701,7 +687,7 @@ bool FHLSLMaterialTranslator::Translate()
 		bCompileForComputeShader = Material->IsLightFunction();
 
 		const bool bStrataEnabled = Engine_IsStrataEnabled();
-		FStrataMaterialInput* FrontMaterialInput = Material->GetMaterialInterface() ? &Material->GetMaterialInterface()->GetMaterial()->FrontMaterial : nullptr;
+		FStrataMaterialInput* FrontMaterialInput = Material->GetMaterialInterface() ? &Material->GetMaterialInterface()->GetMaterial()->GetEditorOnlyData()->FrontMaterial : nullptr;
 		UMaterialExpression* FrontMaterialExpr = FrontMaterialInput ? FrontMaterialInput->GetTracedInput().Expression : nullptr;
 		if (bStrataEnabled && FrontMaterialExpr)
 		{
@@ -6865,7 +6851,7 @@ int32 FHLSLMaterialTranslator::StaticBoolParameter(FName ParameterName,bool bDef
 	FMaterialParameterInfo ParameterInfo = GetParameterAssociationInfo();
 	ParameterInfo.Name = ParameterName;
 
-	for (const FStaticSwitchParameter& Parameter : StaticParameters.StaticSwitchParameters)
+	for (const FStaticSwitchParameter& Parameter : StaticParameters.EditorOnly.StaticSwitchParameters)
 	{
 		if (Parameter.ParameterInfo == ParameterInfo)
 		{
@@ -6888,7 +6874,7 @@ int32 FHLSLMaterialTranslator::StaticComponentMask(int32 Vector,FName ParameterN
 	FMaterialParameterInfo ParameterInfo = GetParameterAssociationInfo();
 	ParameterInfo.Name = ParameterName;
 
-	for (const FStaticComponentMaskParameter& Parameter : StaticParameters.StaticComponentMaskParameters)
+	for (const FStaticComponentMaskParameter& Parameter : StaticParameters.EditorOnly.StaticComponentMaskParameters)
 	{
 		if (Parameter.ParameterInfo == ParameterInfo)
 		{
@@ -6905,7 +6891,7 @@ int32 FHLSLMaterialTranslator::StaticComponentMask(int32 Vector,FName ParameterN
 
 const FMaterialLayersFunctions* FHLSLMaterialTranslator::GetMaterialLayers()
 {
-	return StaticParameters.bHasMaterialLayers ? &StaticParameters.MaterialLayers : nullptr;
+	return StaticParameters.bHasMaterialLayers ? &CachedMaterialLayers : nullptr;
 }
 
 bool FHLSLMaterialTranslator::GetStaticBoolValue(int32 BoolIndex, bool& bSucceeded)
@@ -6945,9 +6931,9 @@ int32 FHLSLMaterialTranslator::StaticTerrainLayerWeight(FName LayerName,int32 De
 	bool bAtLeastOneWeightBasedBlend = false;
 
 	int32 NumActiveTerrainLayerWeightParameters = 0;
-	for(int32 ParameterIndex = 0;ParameterIndex < StaticParameters.TerrainLayerWeightParameters.Num(); ++ParameterIndex)
+	for(int32 ParameterIndex = 0;ParameterIndex < StaticParameters.EditorOnly.TerrainLayerWeightParameters.Num(); ++ParameterIndex)
 	{
-		const FStaticTerrainLayerWeightParameter& Parameter = StaticParameters.TerrainLayerWeightParameters[ParameterIndex];
+		const FStaticTerrainLayerWeightParameter& Parameter = StaticParameters.EditorOnly.TerrainLayerWeightParameters[ParameterIndex];
 		if (Parameter.WeightmapIndex != INDEX_NONE)
 		{
 			NumActiveTerrainLayerWeightParameters++;
diff --git a/Engine/Source/Runtime/Engine/Private/Materials/HLSLMaterialTranslator.h b/Engine/Source/Runtime/Engine/Private/Materials/HLSLMaterialTranslator.h
--- a/Engine/Source/Runtime/Engine/Private/Materials/HLSLMaterialTranslator.h
+++ b/Engine/Source/Runtime/Engine/Private/Materials/HLSLMaterialTranslator.h
@@ -252,6 +252,7 @@ protected:
 	/** Compilation output which will be stored in the DDC. */
 	FMaterialCompilationOutput& MaterialCompilationOutput;
 	FStaticParameterSet StaticParameters;
+	FMaterialLayersFunctions CachedMaterialLayers;
 	EShaderPlatform Platform;
 	/** Quality level being compiled for. */
 	EMaterialQualityLevel::Type QualityLevel;
@@ -458,11 +459,14 @@ public:
 	void AssignTempScope(TArray<FShaderCodeChunk>& InScope);
 	void AssignShaderFrequencyScope(EShaderFrequency InShaderFrequency);
 
-	void GatherCustomVertexInterpolators(TArray<UMaterialExpression*> Expressions);
+	template<typename ExpressionsArrayType>
+	void GatherCustomVertexInterpolators(const ExpressionsArrayType& Expressions);
 
 	void CompileCustomOutputs(TArray<UMaterialExpressionCustomOutput*>& CustomOutputExpressions, TSet<UClass*>& SeenCustomOutputExpressionsClasses, bool bIsBeforeAttributes);
 
-	EMaterialExpressionVisitResult VisitExpressionsRecursive(TArray<UMaterialExpression*> Expressions, IMaterialExpressionVisitor& InVisitor);
+	template<typename ExpressionsArrayType>
+	EMaterialExpressionVisitResult VisitExpressionsRecursive(const ExpressionsArrayType& Expressions, IMaterialExpressionVisitor& InVisitor);
+
 	EMaterialExpressionVisitResult VisitExpressionsForProperty(EMaterialProperty InProperty, IMaterialExpressionVisitor& InVisitor);
 
 	void ValidateVtPropertyLimits();
diff --git a/Engine/Source/Runtime/Engine/Private/Materials/Material.cpp b/Engine/Source/Runtime/Engine/Private/Materials/Material.cpp
--- a/Engine/Source/Runtime/Engine/Private/Materials/Material.cpp
+++ b/Engine/Source/Runtime/Engine/Private/Materials/Material.cpp
@@ -918,13 +918,6 @@ UMaterial::UMaterial(const FObjectInitializer& ObjectInitializer)
 #if WITH_EDITORONLY_DATA
 	DiffuseColor_DEPRECATED.Constant = FColor(128,128,128);
 	SpecularColor_DEPRECATED.Constant = FColor(128,128,128);
-	BaseColor.Constant = FColor(128,128,128);	
-	Metallic.Constant = 0.0f;
-	Specular.Constant = 0.5f;
-	Roughness.Constant = 0.5f;
-	
-	Opacity.Constant = 1.0f;
-	OpacityMask.Constant = 1.0f;
 #endif
 	OpacityMaskClipValue = 0.3333f;
 	bCastDynamicShadowAsMasked = false;
@@ -1345,10 +1338,45 @@ bool UMaterial::IsCompilingOrHadCompileError(ERHIFeatureLevel::Type InFeatureLev
 }
 
 #if WITH_EDITOR
+TConstArrayView<TObjectPtr<UMaterialExpression>> UMaterial::GetExpressions() const
+{
+	return GetEditorOnlyData()->ExpressionCollection.Expressions;
+}
+
+TConstArrayView<TObjectPtr<UMaterialExpressionComment>> UMaterial::GetEditorComments() const
+{
+	return GetEditorOnlyData()->ExpressionCollection.EditorComments;
+}
+
+UMaterialExpressionExecBegin* UMaterial::GetExpressionExecBegin() const
+{
+	return GetEditorOnlyData()->ExpressionCollection.ExpressionExecBegin;
+}
+
+UMaterialExpressionExecEnd* UMaterial::GetExpressionExecEnd() const
+{
+	return GetEditorOnlyData()->ExpressionCollection.ExpressionExecEnd;
+}
+
+const FMaterialExpressionCollection& UMaterial::GetExpressionCollection() const
+{
+	return GetEditorOnlyData()->ExpressionCollection;
+}
+
+FMaterialExpressionCollection& UMaterial::GetExpressionCollection()
+{
+	return GetEditorOnlyData()->ExpressionCollection;
+}
+
+void UMaterial::AssignExpressionCollection(const FMaterialExpressionCollection& InCollection)
+{
+	GetEditorOnlyData()->ExpressionCollection = InCollection;
+}
+
 bool UMaterial::SetParameterValueEditorOnly(const FName& ParameterName, const FMaterialParameterMetadata& Meta)
 {
 	bool bResult = false;
-	for (TObjectPtr<UMaterialExpression>& Expression : Expressions)
+	for (UMaterialExpression* Expression : GetExpressions())
 	{
 		if (Expression && Expression->SetParameterValue(ParameterName, Meta, EMaterialExpressionSetParameterValueFlags::SendPostEditChangeProperty))
 		{
@@ -1364,15 +1392,11 @@ bool UMaterial::SetParameterValueEditorOnly(const FName& ParameterName, const FM
 
 				for (UMaterialFunctionInterface* Function : Functions)
 				{
-					const TArray<TObjectPtr<UMaterialExpression>>* ExpressionPtr = Function->GetFunctionExpressions();
-					if (ExpressionPtr)
+					for (const TObjectPtr<UMaterialExpression>& FunctionExpression : Function->GetExpressions())
 					{
-						for (const TObjectPtr<UMaterialExpression>& FunctionExpression : *ExpressionPtr)
+						if (FunctionExpression && FunctionExpression->SetParameterValue(ParameterName, Meta, EMaterialExpressionSetParameterValueFlags::SendPostEditChangeProperty))
 						{
-							if (FunctionExpression && FunctionExpression->SetParameterValue(ParameterName, Meta, EMaterialExpressionSetParameterValueFlags::SendPostEditChangeProperty))
-							{
-								bResult = true;
-							}
+							bResult = true;
 						}
 					}
 				}
@@ -1751,7 +1775,7 @@ void UMaterial::FixupMaterialUsageAfterLoad()
 #if WITH_EDITORONLY_DATA
 bool UMaterial::IterateDependentFunctions(TFunctionRef<bool(UMaterialFunctionInterface*)> Predicate) const
 {
-	for (UMaterialExpression* Expression : Expressions)
+	for (UMaterialExpression* Expression : GetExpressions())
 	{
 		if (UMaterialExpressionMaterialFunctionCall* FunctionCall = Cast<UMaterialExpressionMaterialFunctionCall>(Expression))
 		{
@@ -1882,8 +1906,6 @@ void UMaterial::UpdateCachedExpressionData()
 	}
 
 	FMaterialCachedExpressionData* LocalCachedExpressionData = new FMaterialCachedExpressionData();
-	LocalCachedExpressionData->Reset();
-
 	FMaterialCachedHLSLTree* LocalCachedTree = nullptr;
 	if (IsUsingNewHLSLGenerator())
 	{
@@ -1893,18 +1915,34 @@ void UMaterial::UpdateCachedExpressionData()
 	}
 	else
 	{
+		if (!bUseMaterialAttributes)
+		{
+			for (int32 PropertyIndex = 0; PropertyIndex < MP_MAX; ++PropertyIndex)
+			{
+				const EMaterialProperty Property = (EMaterialProperty)PropertyIndex;
+				const FExpressionInput* Input = GetExpressionInputForProperty(Property);
+				if (Input && Input->IsConnected())
+				{
+					LocalCachedExpressionData->SetPropertyConnected(Property);
+				}
+			}
+		}
+
 		FMaterialCachedExpressionContext Context;
-		LocalCachedExpressionData->UpdateForExpressions(Context, Expressions, EMaterialParameterAssociation::GlobalParameter, -1);
+		LocalCachedExpressionData->UpdateForExpressions(Context, GetExpressions(), EMaterialParameterAssociation::GlobalParameter, -1);
 	}
 
 	if (LocalCachedExpressionData->bHasMaterialLayers)
 	{
 		// Set all layers as linked to parent (there is no parent for base UMaterials)
-		LocalCachedExpressionData->MaterialLayers.LinkAllLayersToParent();
+		LocalCachedExpressionData->EditorOnlyData->MaterialLayers.LinkAllLayersToParent();
 	}
 
+	LocalCachedExpressionData->Validate();
+
 	CachedExpressionData.Reset(LocalCachedExpressionData);
 	CachedHLSLTree.Reset(LocalCachedTree);
+	EditorOnlyData->CachedExpressionData = LocalCachedExpressionData->EditorOnlyData;
 
 	FObjectCacheEventSink::NotifyReferencedTextureChanged_Concurrent(this);
 }
@@ -1914,7 +1952,7 @@ bool UMaterial::GetParameterValue(EMaterialParameterType Type, const FMemoryImag
 {
 	if (EnumHasAnyFlags(Flags, EMaterialGetParameterValueFlags::CheckNonOverrides) && CachedExpressionData)
 	{
-		return CachedExpressionData->Parameters.GetParameterValue(Type, ParameterInfo, OutResult);
+		return CachedExpressionData->GetParameterValue(Type, ParameterInfo, OutResult);
 	}
 
 	return false;
@@ -1924,7 +1962,14 @@ bool UMaterial::GetMaterialLayers(FMaterialLayersFunctions& OutLayers, TMicRecur
 {
 	if (CachedExpressionData && CachedExpressionData->bHasMaterialLayers)
 	{
-		OutLayers = CachedExpressionData->MaterialLayers;
+		OutLayers.GetRuntime() = CachedExpressionData->MaterialLayers;
+#if WITH_EDITORONLY_DATA
+		if (CachedExpressionData->EditorOnlyData)
+		{
+			OutLayers.EditorOnly = CachedExpressionData->EditorOnlyData->MaterialLayers;
+		}
+		OutLayers.Validate();
+#endif // WITH_EDITORONLY_DATA
 		return true;
 	}
 	return false;
@@ -2023,7 +2068,7 @@ static FAutoConsoleVariable GCompileMaterialsForShaderFormatCVar(
 void UMaterial::GetForceRecompileTextureIdsHash(FSHAHash &TextureReferencesHash)
 {
 	TArray<UTexture*> ForceRecompileTextures;
-	for (const UMaterialExpression *MaterialExpression : Expressions)
+	for (const UMaterialExpression *MaterialExpression : GetExpressions())
 	{
 		if (MaterialExpression == nullptr)
 		{
@@ -2058,7 +2103,7 @@ void UMaterial::GetForceRecompileTextureIdsHash(FSHAHash &TextureReferencesHash)
 
 bool UMaterial::IsTextureForceRecompileCacheRessource(UTexture *Texture)
 {
-	for (const UMaterialExpression *MaterialExpression : Expressions)
+	for (const UMaterialExpression *MaterialExpression : GetExpressions())
 	{
 		if (MaterialExpression == nullptr)
 		{
@@ -2382,18 +2427,22 @@ void UMaterial::ReleaseResourcesAndMutateDDCKey(const FGuid& TransformationId)
 bool UMaterial::AttemptInsertNewGroupName(const FString & InNewName)
 {
 #if WITH_EDITOR
-	FParameterGroupData* ParameterGroupDataElement = ParameterGroupData.FindByPredicate([&InNewName](const FParameterGroupData& DataElement)
+	UMaterialEditorOnlyData* LocalData = GetEditorOnlyData();
+	if (LocalData)
 	{
-		return InNewName == DataElement.GroupName;
-	});
+		FParameterGroupData* ParameterGroupDataElement = LocalData->ParameterGroupData.FindByPredicate([&InNewName](const FParameterGroupData& DataElement)
+		{
+			return InNewName == DataElement.GroupName;
+		});
 
-	if (ParameterGroupDataElement == nullptr)
-	{
-		FParameterGroupData NewGroupData;
-		NewGroupData.GroupName = InNewName;
-		NewGroupData.GroupSortPriority = 0;
-		ParameterGroupData.Add(NewGroupData);
-		return true;
+		if (ParameterGroupDataElement == nullptr)
+		{
+			FParameterGroupData NewGroupData;
+			NewGroupData.GroupName = InNewName;
+			NewGroupData.GroupSortPriority = 0;
+			LocalData->ParameterGroupData.Add(NewGroupData);
+			return true;
+		}
 	}
 #endif
 	return false;
@@ -2550,12 +2599,13 @@ void UMaterial::Serialize(FArchive& Ar)
 #if WITH_EDITORONLY_DATA
 	if (MaterialDomain == MD_Volume && Ar.IsLoading() && Ar.CustomVer(FRenderingObjectVersion::GUID) < FRenderingObjectVersion::VolumeExtinctionBecomesRGB)
 	{
-		if (Opacity.IsConnected()) // Base material input cannot have default values so we only deal with connected expression
+		UMaterialEditorOnlyData* EditorOnly = GetEditorOnlyData();
+		if (EditorOnly && EditorOnly->Opacity.IsConnected()) // Base material input cannot have default values so we only deal with connected expression
 		{
 			// Change expression output from the Opacity to SubSurfaceColor that is now representing RGB extinction. Leave opacity connected as it is unused now anyway
-			SubsurfaceColor.Connect(Opacity.OutputIndex, Opacity.Expression);
+			EditorOnly->SubsurfaceColor.Connect(EditorOnly->Opacity.OutputIndex, EditorOnly->Opacity.Expression);
 			// Now disconnect Opacity
-			Opacity.Expression = nullptr;
+			EditorOnly->Opacity.Expression = nullptr;
 
 			// Now force the material to recompile and we use a hash of the original StateId.
 			// This is to avoid having different StateId each time we load the material and to not forever recompile it,i.e. use a cached version.
@@ -2591,8 +2641,9 @@ void UMaterial::BackwardsCompatibilityInputConversion()
 #if WITH_EDITOR
 	if( ShadingModel != MSM_Unlit )
 	{
+		UMaterialEditorOnlyData* EditorOnly = GetEditorOnlyData();
 		bool bIsDS = DiffuseColor_DEPRECATED.IsConnected() || SpecularColor_DEPRECATED.IsConnected();
-		bool bIsBMS = BaseColor.IsConnected() || Metallic.IsConnected() || Specular.IsConnected();
+		bool bIsBMS = EditorOnly->BaseColor.IsConnected() || EditorOnly->Metallic.IsConnected() || EditorOnly->Specular.IsConnected();
 
 		if( bIsDS && !bIsBMS )
 		{
@@ -2601,7 +2652,7 @@ void UMaterial::BackwardsCompatibilityInputConversion()
 			check( GConvertFromDiffSpecMaterialFunction );
 
 			UMaterialExpressionMaterialFunctionCall* FunctionExpression = NewObject<UMaterialExpressionMaterialFunctionCall>(this);
-			Expressions.Add( FunctionExpression );
+			EditorOnly->ExpressionCollection.Expressions.Add( FunctionExpression );
 
 			FunctionExpression->MaterialExpressionEditorX += 200;
 
@@ -2618,9 +2669,9 @@ void UMaterial::BackwardsCompatibilityInputConversion()
 				FunctionExpression->GetInput(1)->Connect( SpecularColor_DEPRECATED.OutputIndex, SpecularColor_DEPRECATED.Expression );
 			}
 
-			BaseColor.Connect( 0, FunctionExpression );
-			Metallic.Connect( 1, FunctionExpression );
-			Specular.Connect( 2, FunctionExpression );
+			EditorOnly->BaseColor.Connect( 0, FunctionExpression );
+			EditorOnly->Metallic.Connect( 1, FunctionExpression );
+			EditorOnly->Specular.Connect( 2, FunctionExpression );
 		}
 	}
 #endif // WITH_EDITOR
@@ -2635,65 +2686,67 @@ void UMaterial::BackwardsCompatibilityVirtualTextureOutputConversion()
 		// Change this guid if you change the conversion code below
 		static FGuid BackwardsCompatibilityVirtualTextureOutputConversionGuid(TEXT("BABD7074-001F-4FC2-BDE5-3A0C436F4414"));
 
+		UMaterialEditorOnlyData* EditorOnly = GetEditorOnlyData();
+
 		MaterialDomain = MD_Surface;
 
 		if (!bUseMaterialAttributes)
 		{
 			// Create a new UMaterialExpressionRuntimeVirtualTextureOutput node and route the old material attribute output to it.
 			UMaterialExpressionRuntimeVirtualTextureOutput* OutputExpression = NewObject<UMaterialExpressionRuntimeVirtualTextureOutput>(this);
-			Expressions.Add(OutputExpression);
+			EditorOnly->ExpressionCollection.Expressions.Add(OutputExpression);
 
 			OutputExpression->MaterialExpressionEditorX = EditorX;
 			OutputExpression->MaterialExpressionEditorY = EditorY - 300;
 
-			if (BaseColor.IsConnected())
+			if (EditorOnly->BaseColor.IsConnected())
 			{
-				OutputExpression->GetInput(0)->Connect(BaseColor.OutputIndex, BaseColor.Expression);
+				OutputExpression->GetInput(0)->Connect(EditorOnly->BaseColor.OutputIndex, EditorOnly->BaseColor.Expression);
 			}
-			if (Specular.IsConnected())
+			if (EditorOnly->Specular.IsConnected())
 			{
-				OutputExpression->GetInput(1)->Connect(Specular.OutputIndex, Specular.Expression);
+				OutputExpression->GetInput(1)->Connect(EditorOnly->Specular.OutputIndex, EditorOnly->Specular.Expression);
 			}
-			if (Roughness.IsConnected())
+			if (EditorOnly->Roughness.IsConnected())
 			{
-				OutputExpression->GetInput(2)->Connect(Roughness.OutputIndex, Roughness.Expression);
+				OutputExpression->GetInput(2)->Connect(EditorOnly->Roughness.OutputIndex, EditorOnly->Roughness.Expression);
 			}
-			if (Normal.IsConnected())
+			if (EditorOnly->Normal.IsConnected())
 			{
 				if (bTangentSpaceNormal)
 				{
-					OutputExpression->GetInput(3)->Connect(Normal.OutputIndex, Normal.Expression);
+					OutputExpression->GetInput(3)->Connect(EditorOnly->Normal.OutputIndex, EditorOnly->Normal.Expression);
 				}
 				else
 				{
 					// Apply the tangent space to world transform that would be applied in the material output.
 					UMaterialExpressionTransform* TransformExpression = NewObject<UMaterialExpressionTransform>(this);
-					Expressions.Add(TransformExpression);
+					EditorOnly->ExpressionCollection.Expressions.Add(TransformExpression);
 
 					TransformExpression->MaterialExpressionEditorX = EditorX - 300;
 					TransformExpression->MaterialExpressionEditorY = EditorY - 300;
 					TransformExpression->TransformSourceType = TRANSFORMSOURCE_Tangent;
 					TransformExpression->TransformType = TRANSFORM_World;
-					TransformExpression->Input.Connect(Normal.OutputIndex, Normal.Expression);
+					TransformExpression->Input.Connect(EditorOnly->Normal.OutputIndex, EditorOnly->Normal.Expression);
 
 					OutputExpression->GetInput(3)->Connect(0, TransformExpression);
 				}
 			}
-			if (Opacity.IsConnected())
+			if (EditorOnly->Opacity.IsConnected())
 			{
-				OutputExpression->GetInput(5)->Connect(Opacity.OutputIndex, Opacity.Expression);
+				OutputExpression->GetInput(5)->Connect(EditorOnly->Opacity.OutputIndex, EditorOnly->Opacity.Expression);
 			}
 			if (BlendMode != BLEND_Opaque)
 			{
 				// Full alpha blend modes were mostly/always used with MD_RuntimeVirtualTexture to allow pin connections.
 				// But we will assume the intention for any associated MD_Surface output is opaque or alpha mask and force convert here.
-				if (Opacity.IsConnected())
+				if (EditorOnly->Opacity.IsConnected())
 				{
-					OpacityMask.Connect(Opacity.OutputIndex, Opacity.Expression);
-					Opacity.Expression = nullptr;
+					EditorOnly->OpacityMask.Connect(EditorOnly->Opacity.OutputIndex, EditorOnly->Opacity.Expression);
+					EditorOnly->Opacity.Expression = nullptr;
 				}
-				BlendMode = OpacityMask.IsConnected() ? BLEND_Masked : BLEND_Opaque;
-				bCanMaskedBeAssumedOpaque = !OpacityMask.Expression && !(OpacityMask.UseConstant && OpacityMask.Constant < 0.999f);
+				BlendMode = EditorOnly->OpacityMask.IsConnected() ? BLEND_Masked : BLEND_Opaque;
+				bCanMaskedBeAssumedOpaque = !EditorOnly->OpacityMask.Expression && !(EditorOnly->OpacityMask.UseConstant && EditorOnly->OpacityMask.Constant < 0.999f);
 			}
 		}
 
@@ -2710,6 +2763,8 @@ void UMaterial::BackwardsCompatibilityDecalConversion()
 	{
 		// Change this guid if you change the conversion code below
 		static FGuid BackwardsCompatibilityDecalConversionGuid(TEXT("352069F8-1B8C-406A-9B88-6946BCDF2C10"));
+		
+		UMaterialEditorOnlyData* EditorOnly = GetEditorOnlyData();
 
 		GMaterialsThatNeedDecalFix.Clear(this);
 
@@ -2730,56 +2785,56 @@ void UMaterial::BackwardsCompatibilityDecalConversion()
 		// Disconnect outputs according to old DBuffer blend mode.
 		if (DecalBlendMode == DBM_DBuffer_Normal || DecalBlendMode == DBM_DBuffer_Roughness || DecalBlendMode == DBM_DBuffer_NormalRoughness)
 		{
-			BaseColor.Expression = nullptr;
+			EditorOnly->BaseColor.Expression = nullptr;
 		}
 		if (DecalBlendMode == DBM_DBuffer_Color || DecalBlendMode == DBM_DBuffer_Roughness || DecalBlendMode == DBM_DBuffer_ColorRoughness || DecalBlendMode == DBM_AlphaComposite)
 		{
-			Normal.Expression = nullptr;
+			EditorOnly->Normal.Expression = nullptr;
 		}
 		if (DecalBlendMode == DBM_DBuffer_Color || DecalBlendMode == DBM_DBuffer_Normal || DecalBlendMode == DBM_DBuffer_ColorNormal)
 		{
-			Roughness.Expression = Specular.Expression = Metallic.Expression = nullptr;
+			EditorOnly->Roughness.Expression = EditorOnly->Specular.Expression = EditorOnly->Metallic.Expression = nullptr;
 		}
 
 		// Previously translucent decals used default values in all unconnected attributes (except for normal).
 		// For backwards compatibility we connect those attributes with defaults.
 		if (DecalBlendMode == DBM_Translucent || DecalBlendMode == DBM_AlphaComposite || DecalBlendMode == DBM_Stain)
 		{
-			if (!BaseColor.IsConnected() || !Metallic.IsConnected())
+			if (!EditorOnly->BaseColor.IsConnected() || !EditorOnly->Metallic.IsConnected())
 			{
 				UMaterialExpressionConstant* Expression = NewObject<UMaterialExpressionConstant>(this);
-				Expressions.Add(Expression);
+				EditorOnly->ExpressionCollection.Expressions.Add(Expression);
 
 				Expression->MaterialExpressionEditorX = EditorX - 100;
 				Expression->MaterialExpressionEditorY = EditorY - 120;
 				Expression->R = 0.f;
 
-				if (!BaseColor.IsConnected())
+				if (!EditorOnly->BaseColor.IsConnected())
 				{
-					BaseColor.Connect(0, Expression);
+					EditorOnly->BaseColor.Connect(0, Expression);
 				}
-				if (!Metallic.IsConnected())
+				if (!EditorOnly->Metallic.IsConnected())
 				{
-					Metallic.Connect(0, Expression);
+					EditorOnly->Metallic.Connect(0, Expression);
 				}
 			}
 
-			if (!Roughness.IsConnected() || !Specular.IsConnected())
+			if (!EditorOnly->Roughness.IsConnected() || !EditorOnly->Specular.IsConnected())
 			{
 				UMaterialExpressionConstant* Expression = NewObject<UMaterialExpressionConstant>(this);
-				Expressions.Add(Expression);
+				EditorOnly->ExpressionCollection.Expressions.Add(Expression);
 
 				Expression->MaterialExpressionEditorX = EditorX - 100;
 				Expression->MaterialExpressionEditorY = EditorY - 60;
 				Expression->R = .5f;
 
-				if (!Roughness.IsConnected())
+				if (!EditorOnly->Roughness.IsConnected())
 				{
-					Roughness.Connect(0, Expression);
+					EditorOnly->Roughness.Connect(0, Expression);
 				}
-				if (!Specular.IsConnected())
+				if (!EditorOnly->Specular.IsConnected())
 				{
-					Specular.Connect(0, Expression);
+					EditorOnly->Specular.Connect(0, Expression);
 				}
 			}
 		}
@@ -2863,6 +2918,8 @@ void UMaterial::ConvertMaterialToStrataMaterial()
 			}
 		}
 	};
+	
+	UMaterialEditorOnlyData* EditorOnly = GetEditorOnlyData();
 
 	bool bCustomNodesGathered = false;
 	UMaterialExpressionThinTranslucentMaterialOutput* ThinTranslucentOutput = nullptr;
@@ -2874,7 +2931,6 @@ void UMaterial::ConvertMaterialToStrataMaterial()
 		if (!bCustomNodesGathered)
 		{
 			bCustomNodesGathered = true;
-
 			TArray<class UMaterialExpressionCustomOutput*> CustomOutputExpressions;
 			GetAllCustomOutputExpressions(CustomOutputExpressions);
 			for (UMaterialExpressionCustomOutput* Expression : CustomOutputExpressions)
@@ -2920,10 +2976,10 @@ void UMaterial::ConvertMaterialToStrataMaterial()
 	bool bRelinkCustomOutputNodes = false;
 	UMaterialExpressionStrataLegacyConversion* ConvertNode = nullptr;
 	// Connect all the legacy pin into the conversion node
-	if (bUseMaterialAttributes && MaterialAttributes.Expression && !MaterialAttributes.Expression->IsResultStrataMaterial(MaterialAttributes.OutputIndex)) // M_Rifle cause issues there
+	if (bUseMaterialAttributes && EditorOnly->MaterialAttributes.Expression && !EditorOnly->MaterialAttributes.Expression->IsResultStrataMaterial(EditorOnly->MaterialAttributes.OutputIndex)) // M_Rifle cause issues there
 	{
 		UMaterialExpressionBreakMaterialAttributes* BreakMatAtt = NewObject<UMaterialExpressionBreakMaterialAttributes>(this);
-		MoveConnectionTo(MaterialAttributes, BreakMatAtt, 0);
+		MoveConnectionTo(EditorOnly->MaterialAttributes, BreakMatAtt, 0);
 
 		ConvertNode = NewObject<UMaterialExpressionStrataLegacyConversion>(this);
 		ConvertNode->BaseColor.Connect(0, BreakMatAtt);
@@ -2947,15 +3003,15 @@ void UMaterial::ConvertMaterialToStrataMaterial()
 		// * Forward inputs to the root node (Do not reconnect the Opacity as we handle the opacity by internally within the conversion node)
 		// * Forward masked opacity only if blend mode is set to masked, as certain material (e.g., FlattenVT requires to have no OpacityMask plugged)
 		bUseMaterialAttributes = false;
-		FrontMaterial.Connect(0, ConvertNode);
+		EditorOnly->FrontMaterial.Connect(0, ConvertNode);
 		if (BlendMode == BLEND_Masked)
 		{
-			OpacityMask.Connect(7, BreakMatAtt);
+			EditorOnly->OpacityMask.Connect(7, BreakMatAtt);
 		}
-		WorldPositionOffset.Connect(10, BreakMatAtt);
-		AmbientOcclusion.Connect(14, BreakMatAtt);
-		PixelDepthOffset.Connect(24, BreakMatAtt);
-		Refraction.Connect(15, BreakMatAtt);
+		EditorOnly->WorldPositionOffset.Connect(10, BreakMatAtt);
+		EditorOnly->AmbientOcclusion.Connect(14, BreakMatAtt);
+		EditorOnly->PixelDepthOffset.Connect(24, BreakMatAtt);
+		EditorOnly->Refraction.Connect(15, BreakMatAtt);
 
 		if (ShadingModel == MSM_FromMaterialExpression)
 		{
@@ -2981,7 +3037,7 @@ void UMaterial::ConvertMaterialToStrataMaterial()
 		ConvertLegacyToStrataBlendMode();
 		bInvalidateShader = true;
 	}
-	else if (!bUseMaterialAttributes && !FrontMaterial.IsConnected())
+	else if (!bUseMaterialAttributes && !EditorOnly->FrontMaterial.IsConnected())
 	{
 		// STRATA_TODO for material conversion
 		//  - WorldPositionOffset can remain on the end point node
@@ -2992,18 +3048,18 @@ void UMaterial::ConvertMaterialToStrataMaterial()
 		{
 			ConvertNode = NewObject<UMaterialExpressionStrataLegacyConversion>(this);
 			ConvertNode->SubsurfaceProfile = bRequireNoSubsurfaceProfile ? nullptr : SubsurfaceProfile;
-			MoveConnectionTo(BaseColor, ConvertNode, 0);
-			MoveConnectionTo(Metallic, ConvertNode, 1);
-			MoveConnectionTo(Specular, ConvertNode, 2);
-			MoveConnectionTo(Roughness, ConvertNode, 3);
-			MoveConnectionTo(Anisotropy, ConvertNode, 4);
-			MoveConnectionTo(EmissiveColor, ConvertNode, 5);
-			CopyConnectionTo(Normal, ConvertNode, 6);
-			MoveConnectionTo(Tangent, ConvertNode, 7);
-			MoveConnectionTo(SubsurfaceColor, ConvertNode, 8);
-			MoveConnectionTo(ClearCoat, ConvertNode, 9);
-			MoveConnectionTo(ClearCoatRoughness, ConvertNode, 10);
-			MoveConnectionTo(Opacity, ConvertNode, 11);
+			MoveConnectionTo(EditorOnly->BaseColor, ConvertNode, 0);
+			MoveConnectionTo(EditorOnly->Metallic, ConvertNode, 1);
+			MoveConnectionTo(EditorOnly->Specular, ConvertNode, 2);
+			MoveConnectionTo(EditorOnly->Roughness, ConvertNode, 3);
+			MoveConnectionTo(EditorOnly->Anisotropy, ConvertNode, 4);
+			MoveConnectionTo(EditorOnly->EmissiveColor, ConvertNode, 5);
+			CopyConnectionTo(EditorOnly->Normal, ConvertNode, 6);
+			MoveConnectionTo(EditorOnly->Tangent, ConvertNode, 7);
+			MoveConnectionTo(EditorOnly->SubsurfaceColor, ConvertNode, 8);
+			MoveConnectionTo(EditorOnly->ClearCoat, ConvertNode, 9);
+			MoveConnectionTo(EditorOnly->ClearCoatRoughness, ConvertNode, 10);
+			MoveConnectionTo(EditorOnly->Opacity, ConvertNode, 11);
 			bRelinkCustomOutputNodes = true;
 			
 			// Shading Model
@@ -3014,10 +3070,10 @@ void UMaterial::ConvertMaterialToStrataMaterial()
 			// rebuilding the final Shading model (see RebuildShadingModelField())
 			if (ShadingModel == MSM_FromMaterialExpression)
 			{
-				check(ShadingModelFromMaterialExpression.IsConnected());
+				check(EditorOnly->ShadingModelFromMaterialExpression.IsConnected());
 
 				// Reconnect the shading model expression
-				MoveConnectionTo(ShadingModelFromMaterialExpression, ConvertNode, 18);
+				MoveConnectionTo(EditorOnly->ShadingModelFromMaterialExpression, ConvertNode, 18);
 
 				// Store strata shading model of the converted material. 
 				GatherCustomNodes();
@@ -3044,19 +3100,19 @@ void UMaterial::ConvertMaterialToStrataMaterial()
 				check(ConvertNode->ConvertedStrataMaterialInfo.CountShadingModels() == 1);
 			}
 
-			FrontMaterial.Connect(0, ConvertNode);
+			EditorOnly->FrontMaterial.Connect(0, ConvertNode);
 			bInvalidateShader = true;
 		}
 		else if (MaterialDomain == MD_Volume)
 		{
 			UMaterialExpressionStrataVolumetricFogCloudBSDF* VolBSDF = NewObject<UMaterialExpressionStrataVolumetricFogCloudBSDF>(this);
-			MoveConnectionTo(BaseColor, VolBSDF, 0);		// Albedo
-			MoveConnectionTo(SubsurfaceColor, VolBSDF, 1);	// Extinction
-			MoveConnectionTo(EmissiveColor, VolBSDF, 2);	// EmissiveColor
-			MoveConnectionTo(AmbientOcclusion, VolBSDF, 3);	// AmbientOcclusion
+			MoveConnectionTo(EditorOnly->BaseColor, VolBSDF, 0);		// Albedo
+			MoveConnectionTo(EditorOnly->SubsurfaceColor, VolBSDF, 1);	// Extinction
+			MoveConnectionTo(EditorOnly->EmissiveColor, VolBSDF, 2);	// EmissiveColor
+			MoveConnectionTo(EditorOnly->AmbientOcclusion, VolBSDF, 3);	// AmbientOcclusion
 
 			// STRATA_TODO remove the VolumetricAdvancedOutput node and add the input onto FogCloudBSDF even if only used by the cloud renderer?
-			FrontMaterial.Connect(0, VolBSDF);
+			EditorOnly->FrontMaterial.Connect(0, VolBSDF);
 			bInvalidateShader = true;
 		}
 		else if (MaterialDomain == MD_LightFunction)
@@ -3068,9 +3124,9 @@ void UMaterial::ConvertMaterialToStrataMaterial()
 
 			// Only Emissive & Opacity are valid input for PostProcess material
 			UMaterialExpressionStrataLightFunction* LightFunctionNode = NewObject<UMaterialExpressionStrataLightFunction>(this);
-			MoveConnectionTo(EmissiveColor, LightFunctionNode, 0);
+			MoveConnectionTo(EditorOnly->EmissiveColor, LightFunctionNode, 0);
 
-			FrontMaterial.Connect(0, LightFunctionNode);
+			EditorOnly->FrontMaterial.Connect(0, LightFunctionNode);
 			bInvalidateShader = true;
 		}
 		else if (MaterialDomain == MD_PostProcess)
@@ -3090,23 +3146,23 @@ void UMaterial::ConvertMaterialToStrataMaterial()
 			{
 				// If the blending mode was translucent, we must multiply the color by the opacity before the output node
 				UMaterialExpressionMultiply* ColorMultiplyOpacityNode = NewObject<UMaterialExpressionMultiply>(this);
-				CopyConnectionTo(Opacity, ColorMultiplyOpacityNode, 0);
-				MoveConnectionTo(EmissiveColor, ColorMultiplyOpacityNode, 1);
+				CopyConnectionTo(EditorOnly->Opacity, ColorMultiplyOpacityNode, 0);
+				MoveConnectionTo(EditorOnly->EmissiveColor, ColorMultiplyOpacityNode, 1);
 
 				PostProcNode->Color.Connect(0, ColorMultiplyOpacityNode);
-				MoveConnectionTo(Opacity, PostProcNode, 1);
+				MoveConnectionTo(EditorOnly->Opacity, PostProcNode, 1);
 			}
 			else
 			{
 				// Only Emissive & Opacity are valid input for PostProcess material
-				MoveConnectionTo(EmissiveColor, PostProcNode, 0);
+				MoveConnectionTo(EditorOnly->EmissiveColor, PostProcNode, 0);
 				if (GetBlendMode() != BLEND_Additive) // Legacy material additive was ignoring opacity
 				{
-					MoveConnectionTo(Opacity, PostProcNode, 1);
+					MoveConnectionTo(EditorOnly->Opacity, PostProcNode, 1);
 				}
 			}
 
-			FrontMaterial.Connect(0, PostProcNode);
+			EditorOnly->FrontMaterial.Connect(0, PostProcNode);
 			bInvalidateShader = true;
 		}
 		else if (MaterialDomain == MD_DeferredDecal)
@@ -3117,18 +3173,18 @@ void UMaterial::ConvertMaterialToStrataMaterial()
 			ShadingModels.AddShadingModel(MSM_DefaultLit);
 
 			ConvertNode = NewObject<UMaterialExpressionStrataLegacyConversion>(this);
-			MoveConnectionTo(BaseColor, ConvertNode, 0);
-			MoveConnectionTo(Metallic, ConvertNode, 1);
-			MoveConnectionTo(Specular, ConvertNode, 2);
-			MoveConnectionTo(Roughness, ConvertNode, 3);
-			MoveConnectionTo(Anisotropy, ConvertNode, 4);
-			MoveConnectionTo(EmissiveColor, ConvertNode, 5);
-			CopyConnectionTo(Normal, ConvertNode, 6);
-			MoveConnectionTo(Tangent, ConvertNode, 7);
-			MoveConnectionTo(SubsurfaceColor, ConvertNode, 8);
-			MoveConnectionTo(ClearCoat, ConvertNode, 9);
-			MoveConnectionTo(ClearCoatRoughness, ConvertNode, 10);
-			MoveConnectionTo(Opacity, ConvertNode, 11);
+			MoveConnectionTo(EditorOnly->BaseColor, ConvertNode, 0);
+			MoveConnectionTo(EditorOnly->Metallic, ConvertNode, 1);
+			MoveConnectionTo(EditorOnly->Specular, ConvertNode, 2);
+			MoveConnectionTo(EditorOnly->Roughness, ConvertNode, 3);
+			MoveConnectionTo(EditorOnly->Anisotropy, ConvertNode, 4);
+			MoveConnectionTo(EditorOnly->EmissiveColor, ConvertNode, 5);
+			CopyConnectionTo(EditorOnly->Normal, ConvertNode, 6);
+			MoveConnectionTo(EditorOnly->Tangent, ConvertNode, 7);
+			MoveConnectionTo(EditorOnly->SubsurfaceColor, ConvertNode, 8);
+			MoveConnectionTo(EditorOnly->ClearCoat, ConvertNode, 9);
+			MoveConnectionTo(EditorOnly->ClearCoatRoughness, ConvertNode, 10);
+			MoveConnectionTo(EditorOnly->Opacity, ConvertNode, 11);
 
 			// Add constant for the Unlit shading model
 			UMaterialExpressionConstant* ShadingModelNode = NewObject<UMaterialExpressionConstant>(this);
@@ -3143,7 +3199,7 @@ void UMaterial::ConvertMaterialToStrataMaterial()
 			UMaterialExpressionStrataConvertToDecal* ConvertToDecalNode= NewObject<UMaterialExpressionStrataConvertToDecal>(this);
 			ConvertToDecalNode->DecalMaterial.Connect(0, ConvertNode);
 
-			FrontMaterial.Connect(0, ConvertToDecalNode);
+			EditorOnly->FrontMaterial.Connect(0, ConvertToDecalNode);
 			bInvalidateShader = true;
 		}
 
@@ -3191,6 +3247,15 @@ void UMaterial::ConvertMaterialToStrataMaterial()
 
 TMap<FGuid, UMaterialInterface*> LightingGuidFixupMap;
 
+template<typename InputType>
+static void MoveExpressionInput(InputType& Src, InputType& OutDst)
+{
+	if (Src.IsConnected() || Src.IsConstant())
+	{
+		OutDst = MoveTemp(Src);
+	}
+}
+
 void UMaterial::PostLoad()
 {
 	LLM_SCOPE(ELLMTag::Materials);
@@ -3220,35 +3285,94 @@ void UMaterial::PostLoad()
 	const int32 RenderObjVer = GetLinkerCustomVersion(FRenderingObjectVersion::GUID);
 	const int32 UE5MainVer = GetLinkerCustomVersion(FUE5MainStreamObjectVersion::GUID);
 
+	UMaterialEditorOnlyData* EditorOnly = GetEditorOnlyData();
+
+	MoveExpressionInput(BaseColor_DEPRECATED, EditorOnly->BaseColor);
+	MoveExpressionInput(Metallic_DEPRECATED, EditorOnly->Metallic);
+	MoveExpressionInput(Specular_DEPRECATED, EditorOnly->Specular);
+	MoveExpressionInput(Roughness_DEPRECATED, EditorOnly->Roughness);
+	MoveExpressionInput(Anisotropy_DEPRECATED, EditorOnly->Anisotropy);
+	MoveExpressionInput(Normal_DEPRECATED, EditorOnly->Normal);
+	MoveExpressionInput(Tangent_DEPRECATED, EditorOnly->Tangent);
+	MoveExpressionInput(EmissiveColor_DEPRECATED, EditorOnly->EmissiveColor);
+	MoveExpressionInput(Tangent_DEPRECATED, EditorOnly->Tangent);
+	MoveExpressionInput(Opacity_DEPRECATED, EditorOnly->Opacity);
+	MoveExpressionInput(OpacityMask_DEPRECATED, EditorOnly->OpacityMask);
+	MoveExpressionInput(WorldPositionOffset_DEPRECATED, EditorOnly->WorldPositionOffset);
+	MoveExpressionInput(SubsurfaceColor_DEPRECATED, EditorOnly->SubsurfaceColor);
+	MoveExpressionInput(ClearCoat_DEPRECATED, EditorOnly->ClearCoat);
+	MoveExpressionInput(ClearCoatRoughness_DEPRECATED, EditorOnly->ClearCoatRoughness);
+	MoveExpressionInput(AmbientOcclusion_DEPRECATED, EditorOnly->AmbientOcclusion);
+	MoveExpressionInput(Refraction_DEPRECATED, EditorOnly->Refraction);
+	MoveExpressionInput(MaterialAttributes_DEPRECATED, EditorOnly->MaterialAttributes);
+	MoveExpressionInput(PixelDepthOffset_DEPRECATED, EditorOnly->PixelDepthOffset);
+	MoveExpressionInput(ShadingModelFromMaterialExpression_DEPRECATED, EditorOnly->ShadingModelFromMaterialExpression);
+	MoveExpressionInput(FrontMaterial_DEPRECATED, EditorOnly->FrontMaterial);
+	for (int32 i = 0; i < 8; ++i)
+	{
+		MoveExpressionInput(CustomizedUVs_DEPRECATED[i], EditorOnly->CustomizedUVs[i]);
+	}
+
 	DoMaterialAttributeReorder(&DiffuseColor_DEPRECATED, UEVer, RenderObjVer, UE5MainVer);
 	DoMaterialAttributeReorder(&SpecularColor_DEPRECATED, UEVer, RenderObjVer, UE5MainVer);
-	DoMaterialAttributeReorder(&BaseColor, UEVer, RenderObjVer, UE5MainVer);
-	DoMaterialAttributeReorder(&Metallic, UEVer, RenderObjVer, UE5MainVer);
-	DoMaterialAttributeReorder(&Specular, UEVer, RenderObjVer, UE5MainVer);
-	DoMaterialAttributeReorder(&Roughness, UEVer, RenderObjVer, UE5MainVer);
-	DoMaterialAttributeReorder(&Anisotropy, UEVer, RenderObjVer, UE5MainVer);
-	DoMaterialAttributeReorder(&Normal, UEVer, RenderObjVer, UE5MainVer);
-	DoMaterialAttributeReorder(&Tangent, UEVer, RenderObjVer, UE5MainVer);
-	DoMaterialAttributeReorder(&EmissiveColor, UEVer, RenderObjVer, UE5MainVer);
-	DoMaterialAttributeReorder(&Opacity, UEVer, RenderObjVer, UE5MainVer);
-	DoMaterialAttributeReorder(&OpacityMask, UEVer, RenderObjVer, UE5MainVer);
-	DoMaterialAttributeReorder(&WorldPositionOffset, UEVer, RenderObjVer, UE5MainVer);
-	DoMaterialAttributeReorder(&SubsurfaceColor, UEVer, RenderObjVer, UE5MainVer);
-	DoMaterialAttributeReorder(&ClearCoat, UEVer, RenderObjVer, UE5MainVer);
-	DoMaterialAttributeReorder(&ClearCoatRoughness, UEVer, RenderObjVer, UE5MainVer);
-	DoMaterialAttributeReorder(&AmbientOcclusion, UEVer, RenderObjVer, UE5MainVer);
-	DoMaterialAttributeReorder(&Refraction, UEVer, RenderObjVer, UE5MainVer);
-	DoMaterialAttributeReorder(&CustomizedUVs[0], UEVer, RenderObjVer, UE5MainVer);
-	DoMaterialAttributeReorder(&CustomizedUVs[1], UEVer, RenderObjVer, UE5MainVer);
-	DoMaterialAttributeReorder(&CustomizedUVs[2], UEVer, RenderObjVer, UE5MainVer);
-	DoMaterialAttributeReorder(&CustomizedUVs[3], UEVer, RenderObjVer, UE5MainVer);
-	DoMaterialAttributeReorder(&CustomizedUVs[4], UEVer, RenderObjVer, UE5MainVer);
-	DoMaterialAttributeReorder(&CustomizedUVs[5], UEVer, RenderObjVer, UE5MainVer);
-	DoMaterialAttributeReorder(&CustomizedUVs[6], UEVer, RenderObjVer, UE5MainVer);
-	DoMaterialAttributeReorder(&CustomizedUVs[7], UEVer, RenderObjVer, UE5MainVer);
-	DoMaterialAttributeReorder(&PixelDepthOffset, UEVer, RenderObjVer, UE5MainVer);
-	DoMaterialAttributeReorder(&ShadingModelFromMaterialExpression, UEVer, RenderObjVer, UE5MainVer);
-	DoMaterialAttributeReorder(&FrontMaterial, UEVer, RenderObjVer, UE5MainVer);
+	DoMaterialAttributeReorder(&EditorOnly->BaseColor, UEVer, RenderObjVer, UE5MainVer);
+	DoMaterialAttributeReorder(&EditorOnly->Metallic, UEVer, RenderObjVer, UE5MainVer);
+	DoMaterialAttributeReorder(&EditorOnly->Specular, UEVer, RenderObjVer, UE5MainVer);
+	DoMaterialAttributeReorder(&EditorOnly->Roughness, UEVer, RenderObjVer, UE5MainVer);
+	DoMaterialAttributeReorder(&EditorOnly->Anisotropy, UEVer, RenderObjVer, UE5MainVer);
+	DoMaterialAttributeReorder(&EditorOnly->Normal, UEVer, RenderObjVer, UE5MainVer);
+	DoMaterialAttributeReorder(&EditorOnly->Tangent, UEVer, RenderObjVer, UE5MainVer);
+	DoMaterialAttributeReorder(&EditorOnly->EmissiveColor, UEVer, RenderObjVer, UE5MainVer);
+	DoMaterialAttributeReorder(&EditorOnly->Opacity, UEVer, RenderObjVer, UE5MainVer);
+	DoMaterialAttributeReorder(&EditorOnly->OpacityMask, UEVer, RenderObjVer, UE5MainVer);
+	DoMaterialAttributeReorder(&EditorOnly->WorldPositionOffset, UEVer, RenderObjVer, UE5MainVer);
+	DoMaterialAttributeReorder(&EditorOnly->SubsurfaceColor, UEVer, RenderObjVer, UE5MainVer);
+	DoMaterialAttributeReorder(&EditorOnly->ClearCoat, UEVer, RenderObjVer, UE5MainVer);
+	DoMaterialAttributeReorder(&EditorOnly->ClearCoatRoughness, UEVer, RenderObjVer, UE5MainVer);
+	DoMaterialAttributeReorder(&EditorOnly->AmbientOcclusion, UEVer, RenderObjVer, UE5MainVer);
+	DoMaterialAttributeReorder(&EditorOnly->Refraction, UEVer, RenderObjVer, UE5MainVer);
+	DoMaterialAttributeReorder(&EditorOnly->CustomizedUVs[0], UEVer, RenderObjVer, UE5MainVer);
+	DoMaterialAttributeReorder(&EditorOnly->CustomizedUVs[1], UEVer, RenderObjVer, UE5MainVer);
+	DoMaterialAttributeReorder(&EditorOnly->CustomizedUVs[2], UEVer, RenderObjVer, UE5MainVer);
+	DoMaterialAttributeReorder(&EditorOnly->CustomizedUVs[3], UEVer, RenderObjVer, UE5MainVer);
+	DoMaterialAttributeReorder(&EditorOnly->CustomizedUVs[4], UEVer, RenderObjVer, UE5MainVer);
+	DoMaterialAttributeReorder(&EditorOnly->CustomizedUVs[5], UEVer, RenderObjVer, UE5MainVer);
+	DoMaterialAttributeReorder(&EditorOnly->CustomizedUVs[6], UEVer, RenderObjVer, UE5MainVer);
+	DoMaterialAttributeReorder(&EditorOnly->CustomizedUVs[7], UEVer, RenderObjVer, UE5MainVer);
+	DoMaterialAttributeReorder(&EditorOnly->PixelDepthOffset, UEVer, RenderObjVer, UE5MainVer);
+	DoMaterialAttributeReorder(&EditorOnly->ShadingModelFromMaterialExpression, UEVer, RenderObjVer, UE5MainVer);
+	DoMaterialAttributeReorder(&EditorOnly->FrontMaterial, UEVer, RenderObjVer, UE5MainVer);
+
+	if (EditorOnly && ParameterGroupData_DEPRECATED.Num() > 0)
+	{
+		ensure(EditorOnly->ParameterGroupData.Num() == 0);
+		EditorOnly->ParameterGroupData = MoveTemp(ParameterGroupData_DEPRECATED);
+	}
+
+	if (EditorOnly && Expressions_DEPRECATED.Num() > 0)
+	{
+		ensure(EditorOnly->ExpressionCollection.Expressions.Num() == 0);
+		EditorOnly->ExpressionCollection.Expressions = MoveTemp(Expressions_DEPRECATED);
+	}
+
+	if (EditorOnly && EditorComments_DEPRECATED.Num() > 0)
+	{
+		ensure(EditorOnly->ExpressionCollection.EditorComments.Num() == 0);
+		EditorOnly->ExpressionCollection.EditorComments = MoveTemp(EditorComments_DEPRECATED);
+	}
+
+	if (EditorOnly && ExpressionExecBegin_DEPRECATED)
+	{
+		ensure(!EditorOnly->ExpressionCollection.ExpressionExecBegin);
+		EditorOnly->ExpressionCollection.ExpressionExecBegin = MoveTemp(ExpressionExecBegin_DEPRECATED);
+	}
+
+	if (EditorOnly && ExpressionExecEnd_DEPRECATED)
+	{
+		ensure(!EditorOnly->ExpressionCollection.ExpressionExecEnd);
+		EditorOnly->ExpressionCollection.ExpressionExecEnd = MoveTemp(ExpressionExecEnd_DEPRECATED);
+	}
+
 #endif // WITH_EDITORONLY_DATA
 
 	if (!IsDefaultMaterial())
@@ -3270,11 +3394,11 @@ void UMaterial::PostLoad()
 #if WITH_EDITORONLY_DATA
 	// Ensure expressions have been postloaded before we use them for compiling
 	// Any UObjects used by material compilation must be postloaded here
-	for (int32 ExpressionIndex = 0; ExpressionIndex < Expressions.Num(); ExpressionIndex++)
+	for (UMaterialExpression* Expression : GetExpressions())
 	{
-		if (Expressions[ExpressionIndex])
+		if (Expression)
 		{
-			Expressions[ExpressionIndex]->ConditionalPostLoad();
+			Expression->ConditionalPostLoad();
 		}
 	}
 #endif // WITH_EDITORONLY_DATA
@@ -3321,7 +3445,7 @@ void UMaterial::PostLoad()
 	if (GIsEditor)
 	{
 		// Clean up any removed material expression classes	
-		if (Expressions.Remove(NULL) != 0)
+		if (EditorOnly->ExpressionCollection.Expressions.Remove(nullptr) != 0)
 		{
 			// Force this material to recompile because its expressions have changed
 			// We're not providing a deterministic transformation guid because there could be many different ways expression
@@ -3346,10 +3470,9 @@ void UMaterial::PostLoad()
 	if ( GMaterialsThatNeedSamplerFixup.Get( this ) )
 	{
 		GMaterialsThatNeedSamplerFixup.Clear( this );
-		const int32 ExpressionCount = Expressions.Num();
-		for ( int32 ExpressionIndex = 0; ExpressionIndex < ExpressionCount; ++ExpressionIndex )
+		for (UMaterialExpression* Expression : GetExpressions())
 		{
-			UMaterialExpressionTextureBase* TextureExpression = Cast<UMaterialExpressionTextureBase>(Expressions[ExpressionIndex]);
+			UMaterialExpressionTextureBase* TextureExpression = Cast<UMaterialExpressionTextureBase>(Expression);
 			if ( TextureExpression && TextureExpression->Texture )
 			{
 				switch( TextureExpression->Texture->CompressionSettings )
@@ -3457,27 +3580,27 @@ void UMaterial::PostLoad()
 	if (GMaterialsThatNeedExpressionsFlipped.Get(this))
 	{
 		GMaterialsThatNeedExpressionsFlipped.Clear(this);
-		FlipExpressionPositions(Expressions, EditorComments, true, this);
+		FlipExpressionPositions(GetExpressions(), EditorOnly->ExpressionCollection.EditorComments, true, this);
 	}
 	else if (GMaterialsThatNeedCoordinateCheck.Get(this))
 	{
 		GMaterialsThatNeedCoordinateCheck.Clear(this);
 		if (HasFlippedCoordinates())
 		{
-			FlipExpressionPositions(Expressions, EditorComments, false, this);
+			FlipExpressionPositions(GetExpressions(), EditorOnly->ExpressionCollection.EditorComments, false, this);
 		}
-		FixCommentPositions(EditorComments);
+		FixCommentPositions(EditorOnly->ExpressionCollection.EditorComments);
 	}
 	else if (GMaterialsThatNeedCommentFix.Get(this))
 	{
 		GMaterialsThatNeedCommentFix.Clear(this);
-		FixCommentPositions(EditorComments);
+		FixCommentPositions(EditorOnly->ExpressionCollection.EditorComments);
 	}
 
 	if (GMaterialsThatNeedFeatureLevelSM6Fix.Get(this))
 	{
 		GMaterialsThatNeedFeatureLevelSM6Fix.Clear(this);
-		if (FixFeatureLevelNodesForSM6(Expressions))
+		if (FixFeatureLevelNodesForSM6(EditorOnly->ExpressionCollection.Expressions))
 		{
 			// Change this guid if you change the conversion logic.
 			static FGuid BackwardsCompatibilityFeatureLevelSM6ConversionGuid(TEXT("FC75DED7-2FB3-463B-B56C-8295871A340C"));
@@ -3552,59 +3675,64 @@ void UMaterial::GetShaderTypes(EShaderPlatform ShaderPlatform, const ITargetPlat
 }
 #endif // WITH_EDITOR
 
+bool UMaterial::IsPropertyConnected(EMaterialProperty Property) const
+{
+	return GetCachedExpressionData().IsPropertyConnected(Property);
+}
+
 bool UMaterial::HasBaseColorConnected() const
 {
-	return BaseColor.IsConnected() || GetCachedExpressionData().IsMaterialAttributePropertyConnected(MP_BaseColor);
+	return IsPropertyConnected(MP_BaseColor);
 }
 
 bool UMaterial::HasRoughnessConnected() const
 {
-	return Roughness.IsConnected() || GetCachedExpressionData().IsMaterialAttributePropertyConnected(MP_Roughness);
+	return IsPropertyConnected(MP_Roughness);
 }
 
 bool UMaterial::HasAmbientOcclusionConnected() const
 {
-	return AmbientOcclusion.IsConnected() || GetCachedExpressionData().IsMaterialAttributePropertyConnected(MP_AmbientOcclusion);
+	return IsPropertyConnected(MP_AmbientOcclusion);
 }
 
 bool UMaterial::HasNormalConnected() const
 {
-	return Normal.IsConnected() || GetCachedExpressionData().IsMaterialAttributePropertyConnected(MP_Normal);
+	return IsPropertyConnected(MP_Normal);
 }
 
 bool UMaterial::HasSpecularConnected() const
 {
-	return Specular.IsConnected() || GetCachedExpressionData().IsMaterialAttributePropertyConnected(MP_Specular);
+	return IsPropertyConnected(MP_Specular);
 }
 
 bool UMaterial::HasMetallicConnected() const
 {
-	return Metallic.IsConnected() || GetCachedExpressionData().IsMaterialAttributePropertyConnected(MP_Metallic);
+	return IsPropertyConnected(MP_Metallic);
 }
 
 bool UMaterial::HasEmissiveColorConnected() const
 {
-	return EmissiveColor.IsConnected() || GetCachedExpressionData().IsMaterialAttributePropertyConnected(MP_EmissiveColor);
+	return IsPropertyConnected(MP_EmissiveColor);
 }
 
 bool UMaterial::HasAnisotropyConnected() const
 {
-	return Anisotropy.IsConnected() || GetCachedExpressionData().IsMaterialAttributePropertyConnected(MP_Anisotropy);
+	return IsPropertyConnected(MP_Anisotropy);
 }
 
 bool UMaterial::HasStrataFrontMaterialConnected() const
 {
-	return FrontMaterial.IsConnected();
+	return IsPropertyConnected(MP_FrontMaterial);
 }
 
 bool UMaterial::HasVertexPositionOffsetConnected() const
 {
-	return WorldPositionOffset.IsConnected() || GetCachedExpressionData().IsMaterialAttributePropertyConnected(MP_WorldPositionOffset);
+	return IsPropertyConnected(MP_WorldPositionOffset);
 }
 
 bool UMaterial::HasPixelDepthOffsetConnected() const
 {
-	return PixelDepthOffset.IsConnected() || GetCachedExpressionData().IsMaterialAttributePropertyConnected(MP_PixelDepthOffset);
+	return IsPropertyConnected(MP_PixelDepthOffset);
 }
 
 void UMaterial::PropagateDataToMaterialProxy()
@@ -3685,6 +3813,7 @@ bool UMaterial::CanEditChange(const FProperty* InProperty) const
 {
 	if (InProperty)
 	{
+		const UMaterialEditorOnlyData* EditorOnly = GetEditorOnlyData();
 		FString PropertyName = InProperty->GetName();
 
 		if (PropertyName == GET_MEMBER_NAME_STRING_CHECKED(UMaterial, PhysMaterial) || PropertyName == GET_MEMBER_NAME_STRING_CHECKED(UMaterial, PhysMaterialMask))
@@ -3777,7 +3906,7 @@ bool UMaterial::CanEditChange(const FProperty* InProperty) const
 		}
 		else if (PropertyName == GET_MEMBER_NAME_STRING_CHECKED(UMaterial, RefractionDepthBias))
 		{
-			return Refraction.IsConnected();
+			return EditorOnly->Refraction.IsConnected();
 		}
 	
 		if (PropertyName == GET_MEMBER_NAME_STRING_CHECKED(UMaterial, TranslucencyPass)
@@ -3834,18 +3963,19 @@ void UMaterial::CreateExecutionFlowExpressions()
 {
 	if (IsUsingControlFlow())
 	{
-		if (!ExpressionExecBegin)
+		UMaterialEditorOnlyData* EditorOnly = GetEditorOnlyData();
+		if (!EditorOnly->ExpressionCollection.ExpressionExecBegin)
 		{
-			ExpressionExecBegin = NewObject<UMaterialExpressionExecBegin>(this);
-			ExpressionExecBegin->Material = this;
-			Expressions.Add(ExpressionExecBegin);
+			EditorOnly->ExpressionCollection.ExpressionExecBegin = NewObject<UMaterialExpressionExecBegin>(this);
+			EditorOnly->ExpressionCollection.ExpressionExecBegin->Material = this;
+			EditorOnly->ExpressionCollection.Expressions.Add(EditorOnly->ExpressionCollection.ExpressionExecBegin);
 		}
 
-		if (!ExpressionExecEnd)
+		if (!EditorOnly->ExpressionCollection.ExpressionExecEnd)
 		{
-			ExpressionExecEnd = NewObject<UMaterialExpressionExecEnd>(this);
-			ExpressionExecEnd->Material = this;
-			Expressions.Add(ExpressionExecEnd);
+			EditorOnly->ExpressionCollection.ExpressionExecEnd = NewObject<UMaterialExpressionExecEnd>(this);
+			EditorOnly->ExpressionCollection.ExpressionExecEnd->Material = this;
+			EditorOnly->ExpressionCollection.Expressions.Add(EditorOnly->ExpressionCollection.ExpressionExecEnd);
 		}
 	}
 }
@@ -3874,28 +4004,30 @@ void UMaterial::PostEditChangePropertyInternal(FPropertyChangedEvent& PropertyCh
 	//Cancel any current compilation jobs that are in flight for this material.
 	CancelOutstandingCompilation();
 
+	const UMaterialEditorOnlyData* EditorOnly = GetEditorOnlyData();
+
 	// check for distortion in material 
 	{
 		bUsesDistortion = false;
 		// check for a distortion value
-		if (Refraction.Expression
-			|| (Refraction.UseConstant && FMath::Abs(Refraction.Constant - 1.0f) >= UE_KINDA_SMALL_NUMBER))
+		if (EditorOnly->Refraction.Expression
+			|| (EditorOnly->Refraction.UseConstant && FMath::Abs(EditorOnly->Refraction.Constant - 1.0f) >= UE_KINDA_SMALL_NUMBER))
 		{
 			bUsesDistortion = true;
 		}
 
 		// check the material attributes for refraction expressions as well
-		if (MaterialAttributes.Expression)
+		if (EditorOnly->MaterialAttributes.Expression)
 		{
 			// handle make attribute expressions
-			UMaterialExpressionMakeMaterialAttributes * MakeAttributeExpression = Cast<UMaterialExpressionMakeMaterialAttributes>(MaterialAttributes.Expression);
+			UMaterialExpressionMakeMaterialAttributes * MakeAttributeExpression = Cast<UMaterialExpressionMakeMaterialAttributes>(EditorOnly->MaterialAttributes.Expression);
 			if (MakeAttributeExpression && MakeAttributeExpression->Refraction.Expression)
 			{
 				bUsesDistortion = true;
 			}
 
 			// handle set attribute expressions
-			UMaterialExpressionSetMaterialAttributes * SetAttributeExpression = Cast<UMaterialExpressionSetMaterialAttributes>(MaterialAttributes.Expression);
+			UMaterialExpressionSetMaterialAttributes * SetAttributeExpression = Cast<UMaterialExpressionSetMaterialAttributes>(EditorOnly->MaterialAttributes.Expression);
 			if (SetAttributeExpression)
 			{
 				for (int32 Index = 0; Index < SetAttributeExpression->Inputs.Num(); Index++)
@@ -3912,7 +4044,7 @@ void UMaterial::PostEditChangePropertyInternal(FPropertyChangedEvent& PropertyCh
 	}
 
 	//If we can be sure this material would be the same opaque as it is masked then allow it to be assumed opaque.
-	bCanMaskedBeAssumedOpaque = !OpacityMask.Expression && !(OpacityMask.UseConstant && OpacityMask.Constant < 0.999f) && !bUseMaterialAttributes;
+	bCanMaskedBeAssumedOpaque = !EditorOnly->OpacityMask.Expression && !(EditorOnly->OpacityMask.UseConstant && EditorOnly->OpacityMask.Constant < 0.999f) && !bUseMaterialAttributes;
 
 	// If the strata blending mode is changed, make sure we update the legacy blend mode too for the shaders to be filtered to passes correctly at runtime.
 	if (Engine_IsStrataEnabled())
@@ -4041,10 +4173,9 @@ bool UMaterial::IsDynamicParameter(const UMaterialExpression* Expression)
 void UMaterial::BuildEditorParameterList()
 {
 	EditorParameters.Empty();
-
-	for(int32 MaterialExpressionIndex = 0 ; MaterialExpressionIndex < Expressions.Num() ; ++MaterialExpressionIndex)
+	for(UMaterialExpression* Expression : GetExpressions())
 	{
-		AddExpressionParameter(Expressions[MaterialExpressionIndex], EditorParameters);
+		AddExpressionParameter(Expression, EditorParameters);
 	}
 }
 
@@ -4070,15 +4201,15 @@ bool UMaterial::HasDuplicateParameters(const UMaterialExpression* Expression)
 	return false;
 }
 
-bool UMaterial::HasDuplicateDynamicParameters(const UMaterialExpression* Expression)
+bool UMaterial::HasDuplicateDynamicParameters(const UMaterialExpression* InExpression)
 {
-	const UMaterialExpressionDynamicParameter* DynParam = Cast<UMaterialExpressionDynamicParameter>(Expression);
+	const UMaterialExpressionDynamicParameter* DynParam = Cast<UMaterialExpressionDynamicParameter>(InExpression);
 	if (DynParam)
 	{
-		for (int32 ExpIndex = 0; ExpIndex < Expressions.Num(); ExpIndex++)
+		for (UMaterialExpression* Expression : GetExpressions())
 		{
-			UMaterialExpressionDynamicParameter* CheckDynParam = Cast<UMaterialExpressionDynamicParameter>(Expressions[ExpIndex]);
-			if (CheckDynParam != Expression)
+			UMaterialExpressionDynamicParameter* CheckDynParam = Cast<UMaterialExpressionDynamicParameter>(Expression);
+			if (CheckDynParam != InExpression)
 			{
 				return true;
 			}
@@ -4087,14 +4218,14 @@ bool UMaterial::HasDuplicateDynamicParameters(const UMaterialExpression* Express
 	return false;
 }
 
-void UMaterial::UpdateExpressionDynamicParameters(const UMaterialExpression* Expression)
+void UMaterial::UpdateExpressionDynamicParameters(const UMaterialExpression* InExpression)
 {
-	const UMaterialExpressionDynamicParameter* DynParam = Cast<UMaterialExpressionDynamicParameter>(Expression);
+	const UMaterialExpressionDynamicParameter* DynParam = Cast<UMaterialExpressionDynamicParameter>(InExpression);
 	if (DynParam)
 	{
-		for (int32 ExpIndex = 0; ExpIndex < Expressions.Num(); ExpIndex++)
+		for (UMaterialExpression* Expression : GetExpressions())
 		{
-			UMaterialExpressionDynamicParameter* CheckParam = Cast<UMaterialExpressionDynamicParameter>(Expressions[ExpIndex]);
+			UMaterialExpressionDynamicParameter* CheckParam = Cast<UMaterialExpressionDynamicParameter>(Expression);
 			if (CheckParam && CheckParam->CopyDynamicParameterProperties(DynParam))
 			{
 				CheckParam->GraphNode->ReconstructNode();
@@ -4132,8 +4263,6 @@ void UMaterial::PropagateExpressionParameterChanges(const FName& ParameterName,
 
 void UMaterial::UpdateExpressionParameterName(UMaterialExpression* Expression)
 {
-	FName ExpressionName;
-
 	for(TMap<FName, TArray<UMaterialExpression*> >::TIterator Iter(EditorParameters); Iter; ++Iter)
 	{
 		if(Iter.Value().Remove(Expression) > 0)
@@ -4153,13 +4282,15 @@ void UMaterial::RebuildShadingModelField()
 {
 	ShadingModels.ClearShadingModels();
 	const bool bStrataEnabled = Engine_IsStrataEnabled();
-	if (bStrataEnabled && FrontMaterial.IsConnected())
+	UMaterialEditorOnlyData* EditorOnly = GetEditorOnlyData();
+
+	if (bStrataEnabled && EditorOnly->FrontMaterial.IsConnected())
 	{
 		FStrataMaterialInfo StrataMaterialInfo;
-		check(this->FrontMaterial.Expression);
-		if (this->FrontMaterial.Expression->IsResultStrataMaterial(this->FrontMaterial.OutputIndex))
+		check(EditorOnly->FrontMaterial.Expression);
+		if (EditorOnly->FrontMaterial.Expression->IsResultStrataMaterial(EditorOnly->FrontMaterial.OutputIndex))
 		{
-			this->FrontMaterial.Expression->GatherStrataMaterialInfo(StrataMaterialInfo, this->FrontMaterial.OutputIndex);
+			EditorOnly->FrontMaterial.Expression->GatherStrataMaterialInfo(StrataMaterialInfo, EditorOnly->FrontMaterial.OutputIndex);
 			this->CachedConnectedInputs = StrataMaterialInfo.GetPropertyConnected();
 		}
 
@@ -4948,32 +5079,33 @@ static void SetMaterialInputDescription(FStrataMaterialInput& Input, bool bHidde
 
 bool UMaterial::GetExpressionInputDescription(EMaterialProperty InProperty, FMaterialInputDescription& OutDescription)
 {
+	UMaterialEditorOnlyData* EditorOnly = GetEditorOnlyData();
 	switch (InProperty)
 	{
-	case MP_EmissiveColor: SetMaterialInputDescription(EmissiveColor, false, OutDescription); return true;
-	case MP_Opacity: SetMaterialInputDescription(Opacity, false, OutDescription); return true;
-	case MP_OpacityMask: SetMaterialInputDescription(OpacityMask, false, OutDescription); return true;
-	case MP_BaseColor: SetMaterialInputDescription(BaseColor, false, OutDescription); return true;
-	case MP_Metallic: SetMaterialInputDescription(Metallic, false, OutDescription); return true;
-	case MP_Specular: SetMaterialInputDescription(Specular, false, OutDescription); return true;
-	case MP_Roughness: SetMaterialInputDescription(Roughness, false, OutDescription); return true;
-	case MP_Anisotropy: SetMaterialInputDescription(Anisotropy, false, OutDescription); return true;
-	case MP_Normal: SetMaterialInputDescription(Normal, false, OutDescription); return true;
-	case MP_Tangent: SetMaterialInputDescription(Tangent, false, OutDescription); return true;
-	case MP_WorldPositionOffset: SetMaterialInputDescription(WorldPositionOffset, false, OutDescription); return true;
-	case MP_SubsurfaceColor: SetMaterialInputDescription(SubsurfaceColor, false, OutDescription); return true;
-	case MP_CustomData0: SetMaterialInputDescription(ClearCoat, false, OutDescription); return true;
-	case MP_CustomData1: SetMaterialInputDescription(ClearCoatRoughness, false, OutDescription); return true;
-	case MP_AmbientOcclusion: SetMaterialInputDescription(AmbientOcclusion, false, OutDescription); return true;
-	case MP_Refraction: SetMaterialInputDescription(Refraction, false, OutDescription); return true;
-	case MP_MaterialAttributes: SetMaterialInputDescription(MaterialAttributes, false, OutDescription); return true;
-	case MP_PixelDepthOffset: SetMaterialInputDescription(PixelDepthOffset, false, OutDescription); return true;
-	case MP_ShadingModel: SetMaterialInputDescription(ShadingModelFromMaterialExpression, false, OutDescription); return true;
-	case MP_FrontMaterial: SetMaterialInputDescription(FrontMaterial, false, OutDescription); return true;
+	case MP_EmissiveColor: SetMaterialInputDescription(EditorOnly->EmissiveColor, false, OutDescription); return true;
+	case MP_Opacity: SetMaterialInputDescription(EditorOnly->Opacity, false, OutDescription); return true;
+	case MP_OpacityMask: SetMaterialInputDescription(EditorOnly->OpacityMask, false, OutDescription); return true;
+	case MP_BaseColor: SetMaterialInputDescription(EditorOnly->BaseColor, false, OutDescription); return true;
+	case MP_Metallic: SetMaterialInputDescription(EditorOnly->Metallic, false, OutDescription); return true;
+	case MP_Specular: SetMaterialInputDescription(EditorOnly->Specular, false, OutDescription); return true;
+	case MP_Roughness: SetMaterialInputDescription(EditorOnly->Roughness, false, OutDescription); return true;
+	case MP_Anisotropy: SetMaterialInputDescription(EditorOnly->Anisotropy, false, OutDescription); return true;
+	case MP_Normal: SetMaterialInputDescription(EditorOnly->Normal, false, OutDescription); return true;
+	case MP_Tangent: SetMaterialInputDescription(EditorOnly->Tangent, false, OutDescription); return true;
+	case MP_WorldPositionOffset: SetMaterialInputDescription(EditorOnly->WorldPositionOffset, false, OutDescription); return true;
+	case MP_SubsurfaceColor: SetMaterialInputDescription(EditorOnly->SubsurfaceColor, false, OutDescription); return true;
+	case MP_CustomData0: SetMaterialInputDescription(EditorOnly->ClearCoat, false, OutDescription); return true;
+	case MP_CustomData1: SetMaterialInputDescription(EditorOnly->ClearCoatRoughness, false, OutDescription); return true;
+	case MP_AmbientOcclusion: SetMaterialInputDescription(EditorOnly->AmbientOcclusion, false, OutDescription); return true;
+	case MP_Refraction: SetMaterialInputDescription(EditorOnly->Refraction, false, OutDescription); return true;
+	case MP_MaterialAttributes: SetMaterialInputDescription(EditorOnly->MaterialAttributes, false, OutDescription); return true;
+	case MP_PixelDepthOffset: SetMaterialInputDescription(EditorOnly->PixelDepthOffset, false, OutDescription); return true;
+	case MP_ShadingModel: SetMaterialInputDescription(EditorOnly->ShadingModelFromMaterialExpression, false, OutDescription); return true;
+	case MP_FrontMaterial: SetMaterialInputDescription(EditorOnly->FrontMaterial, false, OutDescription); return true;
 	default:
 		if (InProperty >= MP_CustomizedUVs0 && InProperty <= MP_CustomizedUVs7)
 		{
-			SetMaterialInputDescription(CustomizedUVs[InProperty - MP_CustomizedUVs0], false, OutDescription);
+			SetMaterialInputDescription(EditorOnly->CustomizedUVs[InProperty - MP_CustomizedUVs0], false, OutDescription);
 			return true;
 		}
 		return false;
@@ -4985,7 +5117,7 @@ bool UMaterial::GetExpressionInputDescription(EMaterialProperty InProperty, FMat
 #if WITH_EDITORONLY_DATA
 void UMaterial::GetAllFunctionOutputExpressions(TArray<class UMaterialExpressionFunctionOutput*>& OutFunctionOutputs) const
 {
-	for (UMaterialExpression* Expression : Expressions)
+	for (UMaterialExpression* Expression : GetExpressions())
 	{
 		UMaterialExpressionFunctionOutput* FunctionOutput = Cast<UMaterialExpressionFunctionOutput>(Expression);
 		if (FunctionOutput)
@@ -4997,7 +5129,7 @@ void UMaterial::GetAllFunctionOutputExpressions(TArray<class UMaterialExpression
 
 void UMaterial::GetAllCustomOutputExpressions(TArray<class UMaterialExpressionCustomOutput*>& OutCustomOutputs) const
 {
-	for (UMaterialExpression* Expression : Expressions)
+	for (UMaterialExpression* Expression : GetExpressions())
 	{
 		UMaterialExpressionCustomOutput* CustomOutput = Cast<UMaterialExpressionCustomOutput>(Expression);
 		if (CustomOutput)
@@ -5009,7 +5141,7 @@ void UMaterial::GetAllCustomOutputExpressions(TArray<class UMaterialExpressionCu
 
 void UMaterial::GetAllExpressionsForCustomInterpolators(TArray<class UMaterialExpression*>& OutExpressions) const
 {
-	for (UMaterialExpression* Expression : Expressions)
+	for (UMaterialExpression* Expression : GetExpressions())
 	{
 		if (Expression &&
 			(Expression->IsA(UMaterialExpressionVertexInterpolator::StaticClass()) ||
@@ -5151,16 +5283,19 @@ bool UMaterial::GetExpressionsInPropertyChain(EMaterialProperty InProperty,
 
 bool UMaterial::GetGroupSortPriority(const FString& InGroupName, int32& OutSortPriority) const
 {
-	const FParameterGroupData* ParameterGroupDataElement = ParameterGroupData.FindByPredicate([&InGroupName](const FParameterGroupData& DataElement)
+	const UMaterialEditorOnlyData* LocalData = GetEditorOnlyData();
+	if (LocalData)
 	{
-		return InGroupName == DataElement.GroupName;
-	});
-	if (ParameterGroupDataElement != nullptr)
-	{
-		OutSortPriority = ParameterGroupDataElement->GroupSortPriority;
-		return true;
+		const FParameterGroupData* ParameterGroupDataElement = LocalData->ParameterGroupData.FindByPredicate([&InGroupName](const FParameterGroupData& DataElement)
+		{
+			return InGroupName == DataElement.GroupName;
+		});
+		if (ParameterGroupDataElement != nullptr)
+		{
+			OutSortPriority = ParameterGroupDataElement->GroupSortPriority;
+			return true;
+		}
 	}
-
 	return false;
 }
 
@@ -5319,9 +5454,9 @@ bool UMaterial::RecursiveGetExpressionChain(UMaterialExpression* InExpression, T
 						{
 							bool bUseInputA = StaticSwitchExp->DefaultValue;
 							FName StaticSwitchExpName = StaticSwitchExp->ParameterName;
-							for (int32 CheckIdx = 0; CheckIdx < InStaticParameterSet->StaticSwitchParameters.Num(); CheckIdx++)
+							for (int32 CheckIdx = 0; CheckIdx < InStaticParameterSet->EditorOnly.StaticSwitchParameters.Num(); CheckIdx++)
 							{
-								FStaticSwitchParameter& SwitchParam = InStaticParameterSet->StaticSwitchParameters[CheckIdx];
+								FStaticSwitchParameter& SwitchParam = InStaticParameterSet->EditorOnly.StaticSwitchParameters[CheckIdx];
 								if (SwitchParam.ParameterInfo.Name == StaticSwitchExpName)
 								{
 									// Found it...
@@ -5423,48 +5558,49 @@ void UMaterial::RecursiveUpdateRealtimePreview( UMaterialExpression* InExpressio
 int32 UMaterial::CompilePropertyEx( FMaterialCompiler* Compiler, const FGuid& AttributeID )
 {
 	const EMaterialProperty Property = FMaterialAttributeDefinitionMap::GetProperty(AttributeID);
+	UMaterialEditorOnlyData* EditorOnly = GetEditorOnlyData();
 
 	if (IsUsingControlFlow())
 	{
-		check(ExpressionExecBegin);
-		return ExpressionExecBegin->Compile(Compiler, UMaterialExpression::CompileExecutionOutputIndex);
+		check(EditorOnly->ExpressionCollection.ExpressionExecBegin);
+		return EditorOnly->ExpressionCollection.ExpressionExecBegin->Compile(Compiler, UMaterialExpression::CompileExecutionOutputIndex);
 	}
 
 	if( bUseMaterialAttributes && MP_DiffuseColor != Property && MP_SpecularColor != Property )
 	{
-		return MaterialAttributes.CompileWithDefault(Compiler, AttributeID);
+		return EditorOnly->MaterialAttributes.CompileWithDefault(Compiler, AttributeID);
 	}
 
 	switch (Property)
 	{
-		case MP_Opacity:				return Opacity.CompileWithDefault(Compiler, Property);
-		case MP_OpacityMask:			return OpacityMask.CompileWithDefault(Compiler, Property);
-		case MP_Metallic:				return Metallic.CompileWithDefault(Compiler, Property);
-		case MP_Specular:				return Specular.CompileWithDefault(Compiler, Property);
-		case MP_Roughness:				return Roughness.CompileWithDefault(Compiler, Property);
-		case MP_Anisotropy:				return Anisotropy.CompileWithDefault(Compiler, Property);
-		case MP_CustomData0:			return ClearCoat.CompileWithDefault(Compiler, Property);
-		case MP_CustomData1:			return ClearCoatRoughness.CompileWithDefault(Compiler, Property);
-		case MP_AmbientOcclusion:		return AmbientOcclusion.CompileWithDefault(Compiler, Property);
-		case MP_Refraction:				return Refraction.CompileWithDefault(Compiler, Property);
-		case MP_EmissiveColor:			return EmissiveColor.CompileWithDefault(Compiler, Property);
-		case MP_BaseColor:				return BaseColor.CompileWithDefault(Compiler, Property);
-		case MP_SubsurfaceColor:		return SubsurfaceColor.CompileWithDefault(Compiler, Property);
-		case MP_Normal:					return Normal.CompileWithDefault(Compiler, Property);
-		case MP_Tangent:				return Tangent.CompileWithDefault(Compiler, Property);
-		case MP_WorldPositionOffset:	return WorldPositionOffset.CompileWithDefault(Compiler, Property);
-		case MP_PixelDepthOffset:		return PixelDepthOffset.CompileWithDefault(Compiler, Property);
-		case MP_ShadingModel:			return ShadingModelFromMaterialExpression.CompileWithDefault(Compiler, Property);
-		case MP_FrontMaterial:			return FrontMaterial.CompileWithDefault(Compiler, Property);
+		case MP_Opacity:				return EditorOnly->Opacity.CompileWithDefault(Compiler, Property);
+		case MP_OpacityMask:			return EditorOnly->OpacityMask.CompileWithDefault(Compiler, Property);
+		case MP_Metallic:				return EditorOnly->Metallic.CompileWithDefault(Compiler, Property);
+		case MP_Specular:				return EditorOnly->Specular.CompileWithDefault(Compiler, Property);
+		case MP_Roughness:				return EditorOnly->Roughness.CompileWithDefault(Compiler, Property);
+		case MP_Anisotropy:				return EditorOnly->Anisotropy.CompileWithDefault(Compiler, Property);
+		case MP_CustomData0:			return EditorOnly->ClearCoat.CompileWithDefault(Compiler, Property);
+		case MP_CustomData1:			return EditorOnly->ClearCoatRoughness.CompileWithDefault(Compiler, Property);
+		case MP_AmbientOcclusion:		return EditorOnly->AmbientOcclusion.CompileWithDefault(Compiler, Property);
+		case MP_Refraction:				return EditorOnly->Refraction.CompileWithDefault(Compiler, Property);
+		case MP_EmissiveColor:			return EditorOnly->EmissiveColor.CompileWithDefault(Compiler, Property);
+		case MP_BaseColor:				return EditorOnly->BaseColor.CompileWithDefault(Compiler, Property);
+		case MP_SubsurfaceColor:		return EditorOnly->SubsurfaceColor.CompileWithDefault(Compiler, Property);
+		case MP_Normal:					return EditorOnly->Normal.CompileWithDefault(Compiler, Property);
+		case MP_Tangent:				return EditorOnly->Tangent.CompileWithDefault(Compiler, Property);
+		case MP_WorldPositionOffset:	return EditorOnly->WorldPositionOffset.CompileWithDefault(Compiler, Property);
+		case MP_PixelDepthOffset:		return EditorOnly->PixelDepthOffset.CompileWithDefault(Compiler, Property);
+		case MP_ShadingModel:			return EditorOnly->ShadingModelFromMaterialExpression.CompileWithDefault(Compiler, Property);
+		case MP_FrontMaterial:			return EditorOnly->FrontMaterial.CompileWithDefault(Compiler, Property);
 
 		default:
 			if (Property >= MP_CustomizedUVs0 && Property <= MP_CustomizedUVs7)
 			{
 				const int32 TextureCoordinateIndex = Property - MP_CustomizedUVs0;
 
-				if (TextureCoordinateIndex < NumCustomizedUVs && CustomizedUVs[TextureCoordinateIndex].Expression)
+				if (TextureCoordinateIndex < NumCustomizedUVs && EditorOnly->CustomizedUVs[TextureCoordinateIndex].Expression)
 				{
-					return CustomizedUVs[TextureCoordinateIndex].CompileWithDefault(Compiler, Property);
+					return EditorOnly->CustomizedUVs[TextureCoordinateIndex].CompileWithDefault(Compiler, Property);
 				}
 				else
 				{
@@ -5980,7 +6116,7 @@ bool UMaterial::IsPropertyActiveInEditor(EMaterialProperty InProperty) const
 		ShadingModels,
 		TranslucencyLightingMode,
 		BlendableOutputAlpha,
-		Refraction.IsConnected(),
+		GetEditorOnlyData()->Refraction.IsConnected(),
 		IsShadingModelFromMaterialExpression(),
 		IsTranslucencyWritingVelocity(),
 		IsPropertySupported(InProperty));
@@ -5989,6 +6125,11 @@ bool UMaterial::IsPropertyActiveInEditor(EMaterialProperty InProperty) const
 
 bool UMaterial::IsPropertyActiveInDerived(EMaterialProperty InProperty, const UMaterialInterface* DerivedMaterial) const
 {
+#if WITH_EDITOR
+	// Make sure our cached data is consistent
+	ensureMsgf(!(GetEditorOnlyData()->Refraction.IsConnected() && !GetCachedExpressionData().IsPropertyConnected(MP_Refraction)),
+		TEXT("GetCachedExpressionData() says refraction isn't connected, but GetEditorOnlyData() says it is"));
+#endif
 	return IsPropertyActive_Internal(InProperty,
 		MaterialDomain,
 		DerivedMaterial->GetBlendMode(),
@@ -5996,14 +6137,14 @@ bool UMaterial::IsPropertyActiveInDerived(EMaterialProperty InProperty, const UM
 		DerivedMaterial->GetShadingModels(),
 		TranslucencyLightingMode,
 		BlendableOutputAlpha,
-		Refraction.IsConnected(),
+		GetCachedExpressionData().IsPropertyConnected(MP_Refraction),
 		DerivedMaterial->IsShadingModelFromMaterialExpression(),
 		IsTranslucencyWritingVelocity(),
 		IsPropertySupported(InProperty));
 }
 
 #if WITH_EDITORONLY_DATA
-void UMaterial::FlipExpressionPositions(const TArray<UMaterialExpression*>& Expressions, const TArray<UMaterialExpressionComment*>& Comments, bool bScaleCoords, UMaterial* InMaterial)
+void UMaterial::FlipExpressionPositions(TConstArrayView<TObjectPtr<UMaterialExpression>> Expressions, TConstArrayView<TObjectPtr<UMaterialExpressionComment>> Comments, bool bScaleCoords, UMaterial* InMaterial)
 {
 	// Rough estimate of average increase in node size for the new editor
 	const float PosScaling = bScaleCoords ? 1.25f : 1.0f;
@@ -6034,7 +6175,7 @@ void UMaterial::FlipExpressionPositions(const TArray<UMaterialExpression*>& Expr
 	}
 }
 
-void UMaterial::FixCommentPositions(const TArray<UMaterialExpressionComment*>& Comments)
+void UMaterial::FixCommentPositions(TConstArrayView<TObjectPtr<UMaterialExpressionComment>> Comments)
 {
 	// equivalent to 1/1.25 * 0.25 to get the amount that should have been used when first flipping
 	const float SizeScaling = 0.2f;
@@ -6109,4 +6250,15 @@ void UMaterial::GetLightingGuidChain(bool bIncludeTextures, TArray<FGuid>& OutGu
 #endif
 }
 
+UMaterialEditorOnlyData::UMaterialEditorOnlyData()
+{
+	BaseColor.Constant = FColor(128, 128, 128);
+	Metallic.Constant = 0.0f;
+	Specular.Constant = 0.5f;
+	Roughness.Constant = 0.5f;
+
+	Opacity.Constant = 1.0f;
+	OpacityMask.Constant = 1.0f;
+}
+
 #undef LOCTEXT_NAMESPACE
