Add 'Optional' EditorOnly data for UMaterialInterface and UMaterialFunctionInterface

diff --git a/Engine/Source/Runtime/Engine/Private/Materials/MaterialCachedData.cpp b/Engine/Source/Runtime/Engine/Private/Materials/MaterialCachedData.cpp
--- a/Engine/Source/Runtime/Engine/Private/Materials/MaterialCachedData.cpp
+++ b/Engine/Source/Runtime/Engine/Private/Materials/MaterialCachedData.cpp
@@ -45,38 +45,30 @@
 
 const FMaterialCachedParameterEntry FMaterialCachedParameterEntry::EmptyData{};
 const FMaterialCachedExpressionData FMaterialCachedExpressionData::EmptyData{};
+const FMaterialCachedExpressionEditorOnlyData FMaterialCachedExpressionEditorOnlyData::EmptyData{};
 
-void FMaterialCachedExpressionData::Reset()
+static_assert((uint32)(EMaterialProperty::MP_MAX)-1 <= (8 * sizeof(FMaterialCachedExpressionData::PropertyConnectedBitmask)), "PropertyConnectedBitmask cannot contain entire EMaterialProperty enumeration.");
+
+FMaterialCachedExpressionData::FMaterialCachedExpressionData()
+	: bHasMaterialLayers(false)
+	, bHasRuntimeVirtualTextureOutput(false)
+	, bHasSceneColor(false)
+	, bHasPerInstanceCustomData(false)
+	, bHasPerInstanceRandom(false)
+	, bHasVertexInterpolator(false)
 {
-	Parameters.Reset();
-	ReferencedTextures.Reset();
-	FunctionInfos.Reset();
-	ParameterCollectionInfos.Reset();
-	GrassTypes.Reset();
-	DynamicParameterNames.Reset();
-	QualityLevelsUsed.Reset();
 	QualityLevelsUsed.AddDefaulted(EMaterialQualityLevel::Num);
-	bHasMaterialLayers = false;
-	bHasRuntimeVirtualTextureOutput = false;
-	bHasSceneColor = false;
-	bHasPerInstanceCustomData = false;
-	bHasPerInstanceRandom = false;
-	bHasVertexInterpolator = false;
-	MaterialAttributesPropertyConnectedBitmask = 0;
-
 #if WITH_EDITORONLY_DATA
-	LandscapeLayerNames.Reset();
+	EditorOnlyData = MakeShared<FMaterialCachedExpressionEditorOnlyData>();
 #endif // WITH_EDITORONLY_DATA
-
-	static_assert((uint32)(EMaterialProperty::MP_MAX)-1 <= (8 * sizeof(MaterialAttributesPropertyConnectedBitmask)), "MaterialAttributesPropertyConnectedBitmask cannot contain entire EMaterialProperty enumeration.");
 }
 
 void FMaterialCachedExpressionData::AddReferencedObjects(FReferenceCollector& Collector)
 {
 	Collector.AddReferencedObjects(ReferencedTextures);
+	Collector.AddReferencedObjects(GrassTypes);
 	Collector.AddReferencedObjects(MaterialLayers.Layers);
 	Collector.AddReferencedObjects(MaterialLayers.Blends);
-	Collector.AddReferencedObjects(GrassTypes);
 	for (FMaterialFunctionInfo& FunctionInfo : FunctionInfos)
 	{
 		Collector.AddReferencedObject(FunctionInfo.Function);
@@ -106,28 +98,31 @@ void FMaterialCachedExpressionData::AppendReferencedParameterCollectionIdsTo(TAr
 }
 
 #if WITH_EDITOR
-static int32 TryAddParameter(FMaterialCachedParameters& CachedParameters,
+static int32 TryAddParameter(FMaterialCachedExpressionData& CachedData,
 	EMaterialParameterType Type,
 	const FMaterialParameterInfo& ParameterInfo,
 	const FMaterialCachedParameterEditorInfo& InEditorInfo)
 {
-	FMaterialCachedParameterEntry& Entry = CachedParameters.GetParameterTypeEntry(Type);
+	check(CachedData.EditorOnlyData);
+	FMaterialCachedParameterEntry& Entry = CachedData.GetParameterTypeEntry(Type);
+	FMaterialCachedParameterEditorEntry& EditorEntry = CachedData.EditorOnlyData->EditorEntries[(int32)Type];
+
 	FSetElementId ElementId = Entry.ParameterInfoSet.FindId(ParameterInfo);
 	int32 Index = INDEX_NONE;
 	if (!ElementId.IsValidId())
 	{
 		ElementId = Entry.ParameterInfoSet.Add(ParameterInfo);
 		Index = ElementId.AsInteger();
-		Entry.EditorInfo.Insert(InEditorInfo, Index);
+		EditorEntry.EditorInfo.Insert(InEditorInfo, Index);
 		// should be valid as long as we don't ever remove elements from ParameterInfoSet
-		check(Entry.ParameterInfoSet.Num() == Entry.EditorInfo.Num());
+		check(Entry.ParameterInfoSet.Num() == EditorEntry.EditorInfo.Num());
 		return Index;
 	}
 
 	// Update any editor values that haven't been set yet
 	// TODO still need to do this??
 	Index = ElementId.AsInteger();
-	FMaterialCachedParameterEditorInfo& EditorInfo = Entry.EditorInfo[Index];
+	FMaterialCachedParameterEditorInfo& EditorInfo = EditorEntry.EditorInfo[Index];
 	if (!EditorInfo.ExpressionGuid.IsValid())
 	{
 		EditorInfo.ExpressionGuid = InEditorInfo.ExpressionGuid;
@@ -146,9 +141,16 @@ static int32 TryAddParameter(FMaterialCachedParameters& CachedParameters,
 	return INDEX_NONE;
 }
 
-void FMaterialCachedParameters::AddParameter(const FMaterialParameterInfo& ParameterInfo, const FMaterialParameterMetadata& ParameterMeta, UObject*& OutReferencedTexture)
+void FMaterialCachedExpressionData::AddParameter(const FMaterialParameterInfo& ParameterInfo, const FMaterialParameterMetadata& ParameterMeta, UObject*& OutReferencedTexture)
 {
-	const FMaterialCachedParameterEditorInfo EditorInfo(ParameterMeta.ExpressionGuid, ParameterMeta.Description, ParameterMeta.Group, ParameterMeta.SortPriority);
+	check(EditorOnlyData);
+	int32 AssetIndex = INDEX_NONE;
+	if (!ParameterMeta.AssetPath.IsEmpty())
+	{
+		AssetIndex = EditorOnlyData->AssetPaths.AddUnique(ParameterMeta.AssetPath);
+	}
+
+	const FMaterialCachedParameterEditorInfo EditorInfo(ParameterMeta.ExpressionGuid, ParameterMeta.Description, ParameterMeta.Group, ParameterMeta.SortPriority, AssetIndex);
 	const int32 Index = TryAddParameter(*this, ParameterMeta.Value.Type, ParameterInfo, EditorInfo);
 	if (Index != INDEX_NONE)
 	{
@@ -156,24 +158,24 @@ void FMaterialCachedParameters::AddParameter(const FMaterialParameterInfo& Param
 		{
 		case EMaterialParameterType::Scalar:
 			ScalarValues.Insert(ParameterMeta.Value.AsScalar(), Index);
-			ScalarMinMaxValues.Insert(FVector2D(ParameterMeta.ScalarMin, ParameterMeta.ScalarMax), Index);
+			EditorOnlyData->ScalarMinMaxValues.Insert(FVector2D(ParameterMeta.ScalarMin, ParameterMeta.ScalarMax), Index);
 			ScalarPrimitiveDataIndexValues.Insert(ParameterMeta.PrimitiveDataIndex, Index);
 			if (ParameterMeta.bUsedAsAtlasPosition)
 			{
-				ScalarCurveValues.Insert(ParameterMeta.ScalarCurve.Get(), Index);
-				ScalarCurveAtlasValues.Insert(ParameterMeta.ScalarAtlas.Get(), Index);
+				EditorOnlyData->ScalarCurveValues.Insert(ParameterMeta.ScalarCurve.Get(), Index);
+				EditorOnlyData->ScalarCurveAtlasValues.Insert(ParameterMeta.ScalarAtlas.Get(), Index);
 				OutReferencedTexture = ParameterMeta.ScalarAtlas.Get();
 			}
 			else
 			{
-				ScalarCurveValues.Insert(nullptr, Index);
-				ScalarCurveAtlasValues.Insert(nullptr, Index);
+				EditorOnlyData->ScalarCurveValues.Insert(nullptr, Index);
+				EditorOnlyData->ScalarCurveAtlasValues.Insert(nullptr, Index);
 			}
 			break;
 		case EMaterialParameterType::Vector:
 			VectorValues.Insert(ParameterMeta.Value.AsLinearColor(), Index);
-			VectorChannelNameValues.Insert(ParameterMeta.ChannelNames, Index);
-			VectorUsedAsChannelMaskValues.Insert(ParameterMeta.bUsedAsChannelMask, Index);
+			EditorOnlyData->VectorChannelNameValues.Insert(ParameterMeta.ChannelNames, Index);
+			EditorOnlyData->VectorUsedAsChannelMaskValues.Insert(ParameterMeta.bUsedAsChannelMask, Index);
 			VectorPrimitiveDataIndexValues.Insert(ParameterMeta.PrimitiveDataIndex, Index);
 			break;
 		case EMaterialParameterType::DoubleVector:
@@ -181,7 +183,7 @@ void FMaterialCachedParameters::AddParameter(const FMaterialParameterInfo& Param
 			break;
 		case EMaterialParameterType::Texture:
 			TextureValues.Insert(ParameterMeta.Value.Texture, Index);
-			TextureChannelNameValues.Insert(ParameterMeta.ChannelNames, Index);
+			EditorOnlyData->TextureChannelNameValues.Insert(ParameterMeta.ChannelNames, Index);
 			OutReferencedTexture = ParameterMeta.Value.Texture;
 			break;
 		case EMaterialParameterType::Font:
@@ -197,10 +199,10 @@ void FMaterialCachedParameters::AddParameter(const FMaterialParameterInfo& Param
 			OutReferencedTexture = ParameterMeta.Value.RuntimeVirtualTexture;
 			break;
 		case EMaterialParameterType::StaticSwitch:
-			StaticSwitchValues.Insert(ParameterMeta.Value.AsStaticSwitch(), Index);
+			EditorOnlyData->StaticSwitchValues.Insert(ParameterMeta.Value.AsStaticSwitch(), Index);
 			break;
 		case EMaterialParameterType::StaticComponentMask:
-			StaticComponentMaskValues.Insert(ParameterMeta.Value.AsStaticComponentMask(), Index);
+			EditorOnlyData->StaticComponentMaskValues.Insert(ParameterMeta.Value.AsStaticComponentMask(), Index);
 			break;
 		default:
 			checkNoEntry();
@@ -225,11 +227,7 @@ void FMaterialCachedExpressionData::UpdateForFunction(const FMaterialCachedExpre
 	FMaterialCachedExpressionData* Self = this;
 	auto ProcessFunction = [Self, &LocalContext, Association, ParameterIndex](UMaterialFunctionInterface* InFunction) -> bool
 	{
-		const TArray<TObjectPtr<UMaterialExpression>>* FunctionExpressions = InFunction->GetFunctionExpressions();
-		if (FunctionExpressions)
-		{
-			Self->UpdateForExpressions(LocalContext, *FunctionExpressions, Association, ParameterIndex);
-		}
+		Self->UpdateForExpressions(LocalContext, InFunction->GetExpressions(), Association, ParameterIndex);
 
 		FMaterialFunctionInfo NewFunctionInfo;
 		NewFunctionInfo.Function = InFunction;
@@ -256,8 +254,10 @@ void FMaterialCachedExpressionData::UpdateForLayerFunctions(const FMaterialCache
 	}
 }
 
-void FMaterialCachedExpressionData::UpdateForExpressions(const FMaterialCachedExpressionContext& Context, const TArray<TObjectPtr<UMaterialExpression>>& Expressions, EMaterialParameterAssociation Association, int32 ParameterIndex)
+void FMaterialCachedExpressionData::UpdateForExpressions(const FMaterialCachedExpressionContext& Context, TConstArrayView<TObjectPtr<UMaterialExpression>> Expressions, EMaterialParameterAssociation Association, int32 ParameterIndex)
 {
+	check(EditorOnlyData);
+
 	for (UMaterialExpression* Expression : Expressions)
 	{
 		if (!Expression)
@@ -287,7 +287,7 @@ void FMaterialCachedExpressionData::UpdateForExpressions(const FMaterialCachedEx
 			}
 
 			const FMaterialParameterInfo ParameterInfo(ParameterName, Association, ParameterIndex);
-			Parameters.AddParameter(ParameterInfo, ParameterMeta, ReferencedTexture);
+			AddParameter(ParameterInfo, ParameterMeta, ReferencedTexture);
 		}
 
 		// We first try to extract the referenced texture from the parameter value, that way we'll also get the proper texture in case value is overriden by a function instance
@@ -303,7 +303,7 @@ void FMaterialCachedExpressionData::UpdateForExpressions(const FMaterialCachedEx
 			ReferencedTextures.AddUnique(ReferencedTexture);
 		}
 
-		Expression->GetLandscapeLayerNames(LandscapeLayerNames);
+		Expression->GetLandscapeLayerNames(EditorOnlyData->LandscapeLayerNames);
 
 		if (UMaterialExpressionCollectionParameter* ExpressionCollectionParameter = Cast<UMaterialExpressionCollectionParameter>(Expression))
 		{
@@ -381,7 +381,9 @@ void FMaterialCachedExpressionData::UpdateForExpressions(const FMaterialCachedEx
 
 				// TODO(?) - Layers for MIs are currently duplicated here and in FStaticParameterSet
 				bHasMaterialLayers = true;
-				MaterialLayers = Layers;
+				MaterialLayers = Layers.GetRuntime();
+				EditorOnlyData->MaterialLayers = Layers.EditorOnly;
+				FMaterialLayersFunctions::Validate(MaterialLayers, EditorOnlyData->MaterialLayers);
 				LayersExpression->RebuildLayerGraph(false);
 			}
 		}
@@ -398,30 +400,25 @@ void FMaterialCachedExpressionData::UpdateForExpressions(const FMaterialCachedEx
 		}
 		else if (UMaterialExpressionSetMaterialAttributes* SetMatAttributes = Cast<UMaterialExpressionSetMaterialAttributes>(Expression))
 		{
-			for (int PinIndex = 0; PinIndex < SetMatAttributes->AttributeSetTypes.Num(); ++PinIndex)
+			for (int32 PinIndex = 0; PinIndex < SetMatAttributes->AttributeSetTypes.Num(); ++PinIndex)
 			{
 				// For this material attribute pin do we have something connected?
 				const FGuid& Guid = SetMatAttributes->AttributeSetTypes[PinIndex];
 				const FExpressionInput& AttributeInput = SetMatAttributes->Inputs[PinIndex + 1];
 				const EMaterialProperty MaterialProperty = FMaterialAttributeDefinitionMap::GetProperty(Guid);
-
-				// Only set the material property if it hasn't been set yet.  We want to specifically avoid a Set Material Attributes node which doesn't have a 
-				// attribute set from disabling the attribute from a different Set Material Attributes node which does have it enabled.
-				if (!IsMaterialAttributePropertyConnected(MaterialProperty))
+				if (AttributeInput.Expression)
 				{
-					SetMaterialAttributePropertyConnected(MaterialProperty, AttributeInput.Expression ? true : false);
+					SetPropertyConnected(MaterialProperty);
 				}
 			}
-			}
-			else if (UMaterialExpressionMakeMaterialAttributes* MakeMatAttributes = Cast<UMaterialExpressionMakeMaterialAttributes>(Expression))
-			{
-			// Only set the material property if it hasn't been set yet.  We want to specifically avoid a Set Material Attributes node which doesn't have a 
-			// attribute set from disabling the attribute from a different Set Material Attributes node which does have it enabled.
+		}
+		else if (UMaterialExpressionMakeMaterialAttributes* MakeMatAttributes = Cast<UMaterialExpressionMakeMaterialAttributes>(Expression))
+		{
 			auto SetMatAttributeConditionally = [&](EMaterialProperty InMaterialProperty, bool InIsConnected)
 			{
-				if (!IsMaterialAttributePropertyConnected(InMaterialProperty))
+				if (InIsConnected)
 				{
-					SetMaterialAttributePropertyConnected(InMaterialProperty, InIsConnected);
+					SetPropertyConnected(InMaterialProperty);
 				}
 			};
 
@@ -487,7 +484,7 @@ void PrepareHLSLTree(UE::HLSLTree::FEmitContext& EmitContext,
 				const EMaterialProperty Property = FMaterialAttributeDefinitionMap::GetProperty(AttributeID);
 				if (CachedTree.IsAttributeUsed(EmitContext, *EmitResultScope, ResultType, Property))
 				{
-					CachedData.SetMaterialAttributePropertyConnected(Property, true);
+					CachedData.SetPropertyConnected(Property);
 				}
 			}
 		}
@@ -514,88 +511,60 @@ void FMaterialCachedExpressionData::UpdateForCachedHLSLTree(const FMaterialCache
 	::Private::PrepareHLSLTree(EmitContext, CachedTree, *this, SF_Vertex);
 }
 
-#endif // WITH_EDITOR
-
-void FMaterialCachedParameterEntry::Reset()
-{
-	ParameterInfoSet.Reset();
-#if WITH_EDITORONLY_DATA
-	EditorInfo.Reset();
-#endif
-}
-
-void FMaterialCachedParameters::Reset()
+void FMaterialCachedExpressionData::Validate()
 {
-	for (int32 i = 0; i < NumMaterialRuntimeParameterTypes; ++i)
-	{
-		RuntimeEntries[i].Reset();
-	}
-#if WITH_EDITORONLY_DATA
-	for (int32 i = 0; i < NumMaterialEditorOnlyParameterTypes; ++i)
+	if (EditorOnlyData)
 	{
-		EditorOnlyEntries[i].Reset();
+		for (int32 TypeIndex = 0; TypeIndex < NumMaterialParameterTypes; ++TypeIndex)
+		{
+			const FMaterialCachedParameterEditorEntry& EditorEntry = EditorOnlyData->EditorEntries[TypeIndex];
+			const FMaterialCachedParameterEntry& Entry = GetParameterTypeEntry((EMaterialParameterType)TypeIndex);
+			check(EditorEntry.EditorInfo.Num() == Entry.ParameterInfoSet.Num());
+		}
+		FMaterialLayersFunctions::Validate(MaterialLayers, EditorOnlyData->MaterialLayers);
 	}
-#endif
-
-	ScalarPrimitiveDataIndexValues.Reset();
-	VectorPrimitiveDataIndexValues.Reset();
-	ScalarValues.Reset();
-	VectorValues.Reset();
-	DoubleVectorValues.Reset();
-	TextureValues.Reset();
-	FontValues.Reset();
-	FontPageValues.Reset();
-	RuntimeVirtualTextureValues.Reset();
-
-#if WITH_EDITORONLY_DATA
-	StaticSwitchValues.Reset();
-	StaticComponentMaskValues.Reset();
-	ScalarMinMaxValues.Reset();
-	ScalarCurveValues.Reset();
-	ScalarCurveAtlasValues.Reset();
-	VectorChannelNameValues.Reset();
-	VectorUsedAsChannelMaskValues.Reset();
-	TextureChannelNameValues.Reset();
-#endif // WITH_EDITORONLY_DATA
 }
 
-int32 FMaterialCachedParameters::FindParameterIndex(EMaterialParameterType Type, const FMemoryImageMaterialParameterInfo& ParameterInfo) const
+#endif // WITH_EDITOR
+
+int32 FMaterialCachedExpressionData::FindParameterIndex(EMaterialParameterType Type, const FMemoryImageMaterialParameterInfo& ParameterInfo) const
 {
 	const FMaterialCachedParameterEntry& Entry = GetParameterTypeEntry(Type);
 	const FSetElementId ElementId = Entry.ParameterInfoSet.FindId(FMaterialParameterInfo(ParameterInfo));
 	return ElementId.AsInteger();
 }
 
-void FMaterialCachedParameters::GetParameterValueByIndex(EMaterialParameterType Type, int32 ParameterIndex, FMaterialParameterMetadata& OutResult) const
+void FMaterialCachedExpressionData::GetParameterValueByIndex(EMaterialParameterType Type, int32 ParameterIndex, FMaterialParameterMetadata& OutResult) const
 {
-	const FMaterialCachedParameterEntry& Entry = GetParameterTypeEntry(Type);
-	
 #if WITH_EDITORONLY_DATA
-	const bool bIsEditorOnlyDataStripped = Entry.EditorInfo.Num() == 0;
-	if (!bIsEditorOnlyDataStripped)
+	if (EditorOnlyData)
 	{
-		const FMaterialCachedParameterEditorInfo& EditorInfo = Entry.EditorInfo[ParameterIndex];
+		const FMaterialCachedParameterEditorEntry& EditorEntry = EditorOnlyData->EditorEntries[(int32)Type];
+		const FMaterialCachedParameterEditorInfo& EditorInfo = EditorEntry.EditorInfo[ParameterIndex];
 		OutResult.ExpressionGuid = EditorInfo.ExpressionGuid;
 		OutResult.Description = EditorInfo.Description;
 		OutResult.Group = EditorInfo.Group;
 		OutResult.SortPriority = EditorInfo.SortPriority;
+		if (EditorInfo.AssetIndex != INDEX_NONE)
+		{
+			OutResult.AssetPath = EditorOnlyData->AssetPaths[EditorInfo.AssetIndex];
+		}
 	}
-#endif
+#endif // WITH_EDITORONLY_DATA
 
 	switch (Type)
 	{
 	case EMaterialParameterType::Scalar:
 		OutResult.Value = ScalarValues[ParameterIndex];
 		OutResult.PrimitiveDataIndex = ScalarPrimitiveDataIndexValues[ParameterIndex];
-
 #if WITH_EDITORONLY_DATA
-		if (!bIsEditorOnlyDataStripped)
+		if (EditorOnlyData)
 		{
-			OutResult.ScalarMin = ScalarMinMaxValues[ParameterIndex].X;
-			OutResult.ScalarMax = ScalarMinMaxValues[ParameterIndex].Y;
+			OutResult.ScalarMin = EditorOnlyData->ScalarMinMaxValues[ParameterIndex].X;
+			OutResult.ScalarMax = EditorOnlyData->ScalarMinMaxValues[ParameterIndex].Y;
 			{
-				const TSoftObjectPtr<UCurveLinearColor>& Curve = ScalarCurveValues[ParameterIndex];
-				const TSoftObjectPtr<UCurveLinearColorAtlas>& Atlas = ScalarCurveAtlasValues[ParameterIndex];
+				const TSoftObjectPtr<UCurveLinearColor>& Curve = EditorOnlyData->ScalarCurveValues[ParameterIndex];
+				const TSoftObjectPtr<UCurveLinearColorAtlas>& Atlas = EditorOnlyData->ScalarCurveAtlasValues[ParameterIndex];
 				if (!Curve.IsNull() && !Atlas.IsNull())
 				{
 					OutResult.ScalarCurve = Curve;
@@ -609,12 +578,11 @@ void FMaterialCachedParameters::GetParameterValueByIndex(EMaterialParameterType
 	case EMaterialParameterType::Vector:
 		OutResult.Value = VectorValues[ParameterIndex];
 		OutResult.PrimitiveDataIndex = VectorPrimitiveDataIndexValues[ParameterIndex];
-
 #if  WITH_EDITORONLY_DATA
-		if (!bIsEditorOnlyDataStripped)
+		if (EditorOnlyData)
 		{
-			OutResult.ChannelNames = VectorChannelNameValues[ParameterIndex];
-			OutResult.bUsedAsChannelMask = VectorUsedAsChannelMaskValues[ParameterIndex];
+			OutResult.ChannelNames = EditorOnlyData->VectorChannelNameValues[ParameterIndex];
+			OutResult.bUsedAsChannelMask = EditorOnlyData->VectorUsedAsChannelMaskValues[ParameterIndex];
 		}
 #endif // WITH_EDITORONLY_DATA
 		break;
@@ -624,9 +592,9 @@ void FMaterialCachedParameters::GetParameterValueByIndex(EMaterialParameterType
 	case EMaterialParameterType::Texture:
 		OutResult.Value = TextureValues[ParameterIndex].LoadSynchronous();
 #if WITH_EDITORONLY_DATA
-		if (!bIsEditorOnlyDataStripped)
+		if (EditorOnlyData)
 		{
-			OutResult.ChannelNames = TextureChannelNameValues[ParameterIndex];
+			OutResult.ChannelNames = EditorOnlyData->TextureChannelNameValues[ParameterIndex];
 		}
 #endif // WITH_EDITORONLY_DATA
 		break;
@@ -638,15 +606,15 @@ void FMaterialCachedParameters::GetParameterValueByIndex(EMaterialParameterType
 		break;
 #if WITH_EDITORONLY_DATA
 	case EMaterialParameterType::StaticSwitch:
-		if (!bIsEditorOnlyDataStripped)
+		if (EditorOnlyData)
 		{
-			OutResult.Value = StaticSwitchValues[ParameterIndex];
+			OutResult.Value = EditorOnlyData->StaticSwitchValues[ParameterIndex];
 		}
 		break;
 	case EMaterialParameterType::StaticComponentMask:
-		if (!bIsEditorOnlyDataStripped)
+		if (EditorOnlyData)
 		{
-			OutResult.Value = StaticComponentMaskValues[ParameterIndex];
+			OutResult.Value = EditorOnlyData->StaticComponentMaskValues[ParameterIndex];
 		}
 		break;
 #endif // WITH_EDITORONLY_DATA
@@ -656,7 +624,7 @@ void FMaterialCachedParameters::GetParameterValueByIndex(EMaterialParameterType
 	}
 }
 
-bool FMaterialCachedParameters::GetParameterValue(EMaterialParameterType Type, const FMemoryImageMaterialParameterInfo& ParameterInfo, FMaterialParameterMetadata& OutResult) const
+bool FMaterialCachedExpressionData::GetParameterValue(EMaterialParameterType Type, const FMemoryImageMaterialParameterInfo& ParameterInfo, FMaterialParameterMetadata& OutResult) const
 {
 	const int32 Index = FindParameterIndex(Type, ParameterInfo);
 	if (Index != INDEX_NONE)
@@ -668,15 +636,19 @@ bool FMaterialCachedParameters::GetParameterValue(EMaterialParameterType Type, c
 	return false;
 }
 
-#if WITH_EDITORONLY_DATA
-const FGuid& FMaterialCachedParameters::GetExpressionGuid(EMaterialParameterType Type, int32 Index) const
+const FGuid& FMaterialCachedExpressionData::GetExpressionGuid(EMaterialParameterType Type, int32 Index) const
 {
-	const FMaterialCachedParameterEntry& Entry = GetParameterTypeEntry(Type);
-	return Entry.EditorInfo[Index].ExpressionGuid;
-}
+#if WITH_EDITORONLY_DATA
+	if (EditorOnlyData)
+	{
+		return EditorOnlyData->EditorEntries[(int32)Type].EditorInfo[Index].ExpressionGuid;
+	}
 #endif // WITH_EDITORONLY_DATA
+	static const FGuid EmptyGuid;
+	return EmptyGuid;
+}
 
-void FMaterialCachedParameters::GetAllParametersOfType(EMaterialParameterType Type, TMap<FMaterialParameterInfo, FMaterialParameterMetadata>& OutParameters) const
+void FMaterialCachedExpressionData::GetAllParametersOfType(EMaterialParameterType Type, TMap<FMaterialParameterInfo, FMaterialParameterMetadata>& OutParameters) const
 {
 	const FMaterialCachedParameterEntry& Entry = GetParameterTypeEntry(Type);
 	const int32 NumParameters = Entry.ParameterInfoSet.Num();
@@ -690,7 +662,7 @@ void FMaterialCachedParameters::GetAllParametersOfType(EMaterialParameterType Ty
 	}
 }
 
-void FMaterialCachedParameters::GetAllParameterInfoOfType(EMaterialParameterType Type, TArray<FMaterialParameterInfo>& OutParameterInfo, TArray<FGuid>& OutParameterIds) const
+void FMaterialCachedExpressionData::GetAllParameterInfoOfType(EMaterialParameterType Type, TArray<FMaterialParameterInfo>& OutParameterInfo, TArray<FGuid>& OutParameterIds) const
 {
 	const FMaterialCachedParameterEntry& Entry = GetParameterTypeEntry(Type);
 	const int32 NumParameters = Entry.ParameterInfoSet.Num();
@@ -699,22 +671,13 @@ void FMaterialCachedParameters::GetAllParameterInfoOfType(EMaterialParameterType
 
 	for (TSet<FMaterialParameterInfo>::TConstIterator It(Entry.ParameterInfoSet); It; ++It)
 	{
+		const int32 ParameterIndex = It.GetId().AsInteger();
 		OutParameterInfo.Add(*It);
-#if WITH_EDITORONLY_DATA
-		// cooked materials can strip out expression guids
-		if (Entry.EditorInfo.Num() != 0)
-		{
-			OutParameterIds.Add(Entry.EditorInfo[It.GetId().AsInteger()].ExpressionGuid);
-		}
-		else
-#endif
-		{
-			OutParameterIds.Add(FGuid());
-		}
+		OutParameterIds.Add(GetExpressionGuid(Type, ParameterIndex));
 	}
 }
 
-void FMaterialCachedParameters::GetAllGlobalParametersOfType(EMaterialParameterType Type, TMap<FMaterialParameterInfo, FMaterialParameterMetadata>& OutParameters) const
+void FMaterialCachedExpressionData::GetAllGlobalParametersOfType(EMaterialParameterType Type, TMap<FMaterialParameterInfo, FMaterialParameterMetadata>& OutParameters) const
 {
 	const FMaterialCachedParameterEntry& Entry = GetParameterTypeEntry(Type);
 	const int32 NumParameters = Entry.ParameterInfoSet.Num();
@@ -734,7 +697,7 @@ void FMaterialCachedParameters::GetAllGlobalParametersOfType(EMaterialParameterT
 	}
 }
 
-void FMaterialCachedParameters::GetAllGlobalParameterInfoOfType(EMaterialParameterType Type, TArray<FMaterialParameterInfo>& OutParameterInfo, TArray<FGuid>& OutParameterIds) const
+void FMaterialCachedExpressionData::GetAllGlobalParameterInfoOfType(EMaterialParameterType Type, TArray<FMaterialParameterInfo>& OutParameterInfo, TArray<FGuid>& OutParameterIds) const
 {
 	const FMaterialCachedParameterEntry& Entry = GetParameterTypeEntry(Type);
 	const int32 NumParameters = Entry.ParameterInfoSet.Num();
@@ -746,18 +709,9 @@ void FMaterialCachedParameters::GetAllGlobalParameterInfoOfType(EMaterialParamet
 		const FMaterialParameterInfo& ParameterInfo = *It;
 		if (ParameterInfo.Association == GlobalParameter)
 		{
-			OutParameterInfo.Add(ParameterInfo);
-#if WITH_EDITORONLY_DATA
-			// cooked materials can strip out expression guids
-			if (Entry.EditorInfo.Num() != 0)
-			{
-				OutParameterIds.Add(Entry.EditorInfo[It.GetId().AsInteger()].ExpressionGuid);
-			}
-			else
-#endif
-			{
-				OutParameterIds.Add(FGuid());
-			}
+			const int32 ParameterIndex = It.GetId().AsInteger();
+			OutParameterInfo.Add(*It);
+			OutParameterIds.Add(GetExpressionGuid(Type, ParameterIndex));
 		}
 	}
 }
diff --git a/Engine/Source/Runtime/Engine/Private/Materials/MaterialCachedHLSLTree.cpp b/Engine/Source/Runtime/Engine/Private/Materials/MaterialCachedHLSLTree.cpp
--- a/Engine/Source/Runtime/Engine/Private/Materials/MaterialCachedHLSLTree.cpp
+++ b/Engine/Source/Runtime/Engine/Private/Materials/MaterialCachedHLSLTree.cpp
@@ -54,7 +54,7 @@ bool FMaterialCachedHLSLTree::GenerateTree(UMaterial* Material, const FMaterialL
 {
 	const EMaterialShadingModel DefaultShadingModel = Material->GetShadingModels().GetFirstShadingModel();
 
-	for (UMaterialExpression* Expression : Material->Expressions)
+	for (UMaterialExpression* Expression : Material->GetExpressions())
 	{
 		UMaterialExpressionCustomOutput* CustomOutput = Cast<UMaterialExpressionCustomOutput>(Expression);
 		// We don't want anything with HasCustomSourceOutput() here (VertexInterpolators)
diff --git a/Engine/Source/Runtime/Engine/Private/Materials/MaterialExpressionHLSL.cpp b/Engine/Source/Runtime/Engine/Private/Materials/MaterialExpressionHLSL.cpp
--- a/Engine/Source/Runtime/Engine/Private/Materials/MaterialExpressionHLSL.cpp
+++ b/Engine/Source/Runtime/Engine/Private/Materials/MaterialExpressionHLSL.cpp
@@ -1627,7 +1627,7 @@ bool UMaterialExpressionMaterialAttributeLayers::GenerateHLSLExpression(FMateria
 	{
 		UMaterialFunctionInterface* LayerFunction = MaterialLayers.Layers[LayerIndex];
 		const FExpression* LayerExpression = nullptr;
-		if (LayerFunction && MaterialLayers.LayerStates[LayerIndex])
+		if (LayerFunction && MaterialLayers.EditorOnly.LayerStates[LayerIndex])
 		{
 			const EMaterialFunctionUsage Usage = LayerFunction->GetMaterialFunctionUsage();
 			if (Usage != EMaterialFunctionUsage::MaterialLayer)
@@ -1670,7 +1670,7 @@ bool UMaterialExpressionMaterialAttributeLayers::GenerateHLSLExpression(FMateria
 			return Generator.Errorf(TEXT("Invalid number of layers (%d) and blends (%d)"), MaterialLayers.Layers.Num(), MaterialLayers.Blends.Num());
 		}
 
-		if (MaterialLayers.Layers[LayerIndex] && MaterialLayers.LayerStates[LayerIndex])
+		if (MaterialLayers.Layers[LayerIndex] && MaterialLayers.EditorOnly.LayerStates[LayerIndex])
 		{
 			const FExpression* LayerExpression = LayerExpressions[LayerIndex];
 			if (!LayerExpression)
@@ -1709,7 +1709,7 @@ bool UMaterialExpressionMaterialAttributeLayers::GenerateHLSLExpression(FMateria
 		}
 	}
 
-	OutExpression = Generator.GetTree().NewExpression<Material::FExpressionMaterialLayers>(BottomLayerExpression, &MaterialLayers);
+	OutExpression = Generator.GetTree().NewExpression<Material::FExpressionMaterialLayers>(BottomLayerExpression, MaterialLayers);
 	return true;
 }
 
diff --git a/Engine/Source/Runtime/Engine/Private/Materials/MaterialExpressions.cpp b/Engine/Source/Runtime/Engine/Private/Materials/MaterialExpressions.cpp
--- a/Engine/Source/Runtime/Engine/Private/Materials/MaterialExpressions.cpp
+++ b/Engine/Source/Runtime/Engine/Private/Materials/MaterialExpressions.cpp
@@ -504,7 +504,7 @@ void ValidateParameterNameInternal(class UMaterialExpression* ExpressionToValida
 
 				bFoundValidName = true;
 
-				for (const UMaterialExpression* Expression : OwningMaterial->Expressions)
+				for (const UMaterialExpression* Expression : OwningMaterial->GetExpressions())
 				{
 					if (Expression != nullptr && Expression->HasAParameterName())
 					{
@@ -528,7 +528,7 @@ void ValidateParameterNameInternal(class UMaterialExpression* ExpressionToValida
 		if (bAllowDuplicateName)
 		{
 			// Check for any matching values
-			for (UMaterialExpression* Expression : OwningMaterial->Expressions)
+			for (UMaterialExpression* Expression : OwningMaterial->GetExpressions())
 			{
 				if (Expression != nullptr && Expression->HasAParameterName())
 				{
@@ -657,6 +657,34 @@ int32 CompileStrataBlendFunction(FMaterialCompiler* Compiler, const int32 A, con
 	return INDEX_NONE;
 }
 
+void FMaterialExpressionCollection::AddExpression(UMaterialExpression* InExpression)
+{
+	Expressions.AddUnique(InExpression);
+}
+
+void FMaterialExpressionCollection::RemoveExpression(UMaterialExpression* InExpression)
+{
+	Expressions.Remove(InExpression);
+}
+
+void FMaterialExpressionCollection::AddComment(UMaterialExpressionComment* InExpression)
+{
+	EditorComments.AddUnique(InExpression);
+}
+
+void FMaterialExpressionCollection::RemoveComment(UMaterialExpressionComment* InExpression)
+{
+	EditorComments.Remove(InExpression);
+}
+
+void FMaterialExpressionCollection::Empty()
+{
+	Expressions.Empty();
+	EditorComments.Empty();
+	ExpressionExecBegin = nullptr;
+	ExpressionExecEnd = nullptr;
+}
+
 #if WITH_EDITOR
 void UMaterialExpression::InitializeNumExecutionInputs(TArrayView<UMaterialExpression*> Expressions)
 {
@@ -753,7 +781,7 @@ void UMaterialExpression::CopyMaterialExpressions(const TArray<UMaterialExpressi
 			SrcToDestMap.Add( SrcExpression, NewExpression );
 
 			// Add to list of material expressions associated with the copy buffer.
-			Material->Expressions.Add( NewExpression );
+			Material->GetExpressionCollection().AddExpression( NewExpression );
 
 			// There can be only one default mesh paint texture.
 			UMaterialExpressionTextureBase* TextureSample = Cast<UMaterialExpressionTextureBase>( NewExpression );
@@ -818,7 +846,7 @@ void UMaterialExpression::CopyMaterialExpressions(const TArray<UMaterialExpressi
 		NewComment->Material = Material;
 
 		// Add reference to the material
-		Material->EditorComments.Add(NewComment);
+		Material->GetExpressionCollection().AddComment(NewComment);
 
 		// Add to the output array.
 		OutNewComments.Add(NewComment);
@@ -1618,12 +1646,10 @@ void UMaterialExpressionTextureBase::PostEditChangeProperty(FPropertyChangedEven
 		const FName PropertyName = PropertyChangedEvent.Property->GetFName();
 		if (PropertyName == FName(TEXT("IsDefaultMeshpaintTexture")))
 		{
-			const TArray<UMaterialExpression*>& Expressions = this->Material->GetMaterial()->Expressions;
-
 			// Check for other defaulted textures in THIS material (does not search sub levels ie functions etc, as these are ignored in the texture painter). 
-			for (auto ItExpressions = Expressions.CreateConstIterator(); ItExpressions; ItExpressions++)
+			for (UMaterialExpression* Expression : this->Material->GetMaterial()->GetExpressions())
 			{
-				UMaterialExpressionTextureBase* TextureSample = Cast<UMaterialExpressionTextureBase>(*ItExpressions);
+				UMaterialExpressionTextureBase* TextureSample = Cast<UMaterialExpressionTextureBase>(Expression);
 				if (TextureSample != nullptr && TextureSample != this)
 				{
 					if(TextureSample->IsDefaultMeshpaintTexture)
@@ -5964,15 +5990,15 @@ UMaterialExpressionMakeMaterialAttributes::UMaterialExpressionMakeMaterialAttrib
 void UMaterialExpressionMakeMaterialAttributes::Serialize(FStructuredArchive::FRecord Record)
 {
 	Super::Serialize(Record);
+#if WITH_EDITOR
 	FArchive& UnderlyingArchive = Record.GetUnderlyingArchive();
-
 	UnderlyingArchive.UsingCustomVersion(FRenderingObjectVersion::GUID);
-	
 	if (UnderlyingArchive.CustomVer(FRenderingObjectVersion::GUID) < FRenderingObjectVersion::FixedLegacyMaterialAttributeNodeTypes)
 	{
 		// Update the legacy masks else fail on vec3 to vec2 conversion
 		Refraction.SetMask(1, 1, 1, 0, 0);
 	}
+#endif // WITH_EDITOR
 }
 
 #if WITH_EDITOR
@@ -6942,8 +6968,9 @@ UMaterialExpressionMaterialAttributeLayers::UMaterialExpressionMaterialAttribute
 	};
 	static FConstructorStatics ConstructorStatics;
 
+#if WITH_EDITOR
 	DefaultLayers.AddDefaultBackgroundLayer();
-
+#endif
 #if WITH_EDITORONLY_DATA
 	MenuCategories.Add(ConstructorStatics.NAME_MaterialAttributes);
 	MenuCategories.Add(ConstructorStatics.NAME_Parameters);
@@ -7950,14 +7977,14 @@ void UMaterialExpressionVectorParameter::ApplyChannelNames()
 void UMaterialExpressionVectorParameter::ValidateParameterName(const bool bAllowDuplicateName)
 {
 	bool bOverrideDuplicateBehavior = false;
-	TArray<UMaterialExpression*> Expressions;
+	TConstArrayView<TObjectPtr<UMaterialExpression>> Expressions;
 	if (Material)
 	{
-		Expressions = Material->Expressions;
+		Expressions = Material->GetExpressions();
 	}
 	else if (Function)
 	{
-		Expressions = Function->FunctionExpressions;
+		Expressions = Function->GetExpressions();
 	}
 
 	for (UMaterialExpression* Expression : Expressions)
@@ -8316,14 +8343,14 @@ void UMaterialExpressionScalarParameter::PostEditChangeProperty(FPropertyChanged
 void UMaterialExpressionScalarParameter::ValidateParameterName(const bool bAllowDuplicateName)
 {
 	bool bOverrideDuplicateBehavior = false;
-	TArray<UMaterialExpression*> Expressions;
+	TConstArrayView<TObjectPtr<UMaterialExpression>> Expressions;
 	if (Material)
 	{
-		Expressions = Material->Expressions;
+		Expressions = Material->GetExpressions();
 	}
 	else if (Function)
 	{
-		Expressions = Function->FunctionExpressions;
+		Expressions = Function->GetExpressions();
 	}
 
 	for (UMaterialExpression* Expression : Expressions)
@@ -9057,10 +9084,6 @@ void UMaterialExpressionFeatureLevelSwitch::Serialize(FStructuredArchive::FRecor
 	FArchive& UnderlyingArchive = Record.GetUnderlyingArchive();
 	UnderlyingArchive.UsingCustomVersion(FRenderingObjectVersion::GUID);
 
-#if !WITH_EDITORONLY_DATA
-	FExpressionInput Inputs[ERHIFeatureLevel::Num];
-#endif
-
 	if (UnderlyingArchive.IsLoading() && UnderlyingArchive.UEVer() < VER_UE4_RENAME_SM3_TO_ES3_1)
 	{
 		// Copy the ES2 input to SM3 (since SM3 will now become ES3_1 and we don't want broken content)
@@ -9549,9 +9572,9 @@ void UMaterialExpressionDynamicParameter::PostLoad()
 void UMaterialExpressionDynamicParameter::UpdateDynamicParameterProperties()
 {
 	check(Material);
-	for (int32 ExpIndex = 0; ExpIndex < Material->Expressions.Num(); ExpIndex++)
+	for (UMaterialExpression* Expression : Material->GetExpressions())
 	{
-		const UMaterialExpressionDynamicParameter* DynParam = Cast<UMaterialExpressionDynamicParameter>(Material->Expressions[ExpIndex]);
+		const UMaterialExpressionDynamicParameter* DynParam = Cast<UMaterialExpressionDynamicParameter>(Expression);
 		if (CopyDynamicParameterProperties(DynParam))
 		{
 			break;
@@ -11149,12 +11172,12 @@ void UMaterialExpressionPinBase::DeleteReroutePins()
 		if (Reroute.Expression)
 		{
 			Reroute.Expression->Modify();
-			Material->Expressions.Remove(Reroute.Expression);
+			Material->GetExpressionCollection().RemoveExpression(Reroute.Expression);
 			Reroute.Expression->MarkAsGarbage();
 		}
 		else
 		{
-			Material->Expressions.Remove(Reroute.Expression);
+			Material->GetExpressionCollection().RemoveExpression(Reroute.Expression);
 		}
 	}
 	ReroutePins.Empty();
@@ -11193,7 +11216,7 @@ void UMaterialExpressionPinBase::PostEditChangeProperty(FPropertyChangedEvent& P
 			AddedReroute.Expression->Material = Material;
 			AddedReroute.Name = AddedReroute.Name.IsNone() ? FName(FString::Printf(TEXT("Pin %u"), AddedRerouteIndex + 1)) : AddedReroute.Name;
 
-			Material->Expressions.Add(AddedReroute.Expression);
+			Material->GetExpressionCollection().AddExpression(AddedReroute.Expression);
 		}
 		else if (PropertyChangedEvent.ChangeType == EPropertyChangeType::ArrayRemove)
 		{
@@ -11203,12 +11226,12 @@ void UMaterialExpressionPinBase::PostEditChangeProperty(FPropertyChangedEvent& P
 			if (RemovedReroute)
 			{
 				RemovedReroute->Modify();
-				Material->Expressions.Remove(RemovedReroute);
+				Material->GetExpressionCollection().RemoveExpression(RemovedReroute);
 				RemovedReroute->MarkAsGarbage();
 			}
 			else
 			{
-				Material->Expressions.Remove(RemovedReroute);
+				Material->GetExpressionCollection().RemoveExpression(RemovedReroute);
 			}
 		}
 		else if (PropertyChangedEvent.ChangeType == EPropertyChangeType::ArrayClear)
@@ -11218,12 +11241,12 @@ void UMaterialExpressionPinBase::PostEditChangeProperty(FPropertyChangedEvent& P
 				if (RemovedReroute)
 				{
 					RemovedReroute->Modify();
-					Material->Expressions.Remove(RemovedReroute);
+					Material->GetExpressionCollection().RemoveExpression(RemovedReroute);
 					RemovedReroute->MarkAsGarbage();
 				}
 				else
 				{
-					Material->Expressions.Remove(RemovedReroute);
+					Material->GetExpressionCollection().RemoveExpression(RemovedReroute);
 				}
 			}
 		}
@@ -12601,6 +12624,9 @@ UMaterialFunctionInterface::UMaterialFunctionInterface(const FObjectInitializer&
 
 void UMaterialFunctionInterface::PostInitProperties()
 {
+#if WITH_EDITORONLY_DATA
+	EditorOnlyData = CreateEditorOnlyData();
+#endif
 	Super::PostInitProperties();
 
 	// Initialize StateId to something unique, in case this is a new function
@@ -12610,7 +12636,6 @@ void UMaterialFunctionInterface::PostInitProperties()
 void UMaterialFunctionInterface::PostLoad()
 {
 	Super::PostLoad();
-
 	if (!StateId.IsValid())
 	{
 		StateId = FGuid::NewGuid();
@@ -12634,7 +12659,58 @@ void UMaterialFunctionInterface::GetAssetRegistryTags(TArray<FAssetRegistryTag>&
 #endif
 }
 
+UMaterialFunctionInterface* UMaterialFunctionInterface::GetBaseFunctionInterface()
+{
+	return GetBaseFunction();
+}
+
+const UMaterialFunctionInterface* UMaterialFunctionInterface::GetBaseFunctionInterface() const
+{
+	return GetBaseFunction();
+}
+
 #if WITH_EDITOR
+
+TConstArrayView<TObjectPtr<UMaterialExpression>> UMaterialFunctionInterface::GetExpressions() const
+{
+	const UMaterialFunction* BaseFunction = GetBaseFunction();
+	if (BaseFunction)
+	{
+		return BaseFunction->GetExpressions();
+	}
+	return TConstArrayView<TObjectPtr<UMaterialExpression>>();
+}
+
+const FString& UMaterialFunctionInterface::GetDescription() const
+{
+	const UMaterialFunction* BaseFunction = GetBaseFunction();
+	if (BaseFunction)
+	{
+		return BaseFunction->Description;
+	}
+	static const FString EmptyString;
+	return EmptyString;
+}
+
+bool UMaterialFunctionInterface::GetReentrantFlag() const
+{
+	const UMaterialFunction* BaseFunction = GetBaseFunction();
+	if (BaseFunction)
+	{
+		return BaseFunction->GetReentrantFlag();
+	}
+	return false;
+}
+
+void UMaterialFunctionInterface::SetReentrantFlag(bool bIsReentrant)
+{
+	UMaterialFunction* BaseFunction = GetBaseFunction();
+	if (BaseFunction)
+	{
+		BaseFunction->SetReentrantFlag(bIsReentrant);
+	}
+}
+
 void UMaterialFunctionInterface::PostEditChangeProperty(FPropertyChangedEvent& PropertyChangedEvent)
 {
 	if (false) // temporary to unblock people, tracked as UE-109810
@@ -12786,6 +12862,24 @@ bool UMaterialFunctionInterface::OverrideNamedStaticComponentMaskParameter(const
 }
 #endif // WITH_EDITOR
 
+#if WITH_EDITORONLY_DATA
+const UClass* UMaterialFunctionInterface::GetEditorOnlyDataClass() const
+{
+	return UMaterialFunctionInterfaceEditorOnlyData::StaticClass();
+}
+
+UMaterialFunctionInterfaceEditorOnlyData* UMaterialFunctionInterface::CreateEditorOnlyData()
+{
+	const UClass* EditorOnlyClass = GetEditorOnlyDataClass();
+	check(EditorOnlyClass);
+	check(EditorOnlyClass->HasAllClassFlags(CLASS_Optional));
+
+	const FString EditorOnlyName = FString::Printf(TEXT("%sEditorOnlyData"), *GetName());
+	const EObjectFlags EditorOnlyFlags = GetMaskedFlags(RF_PropagateToSubObjects);
+	return NewObject<UMaterialFunctionInterfaceEditorOnlyData>(this, EditorOnlyClass, *EditorOnlyName, EditorOnlyFlags);
+}
+#endif // WITH_EDITORONLY_DATA
+
 ///////////////////////////////////////////////////////////////////////////////
 // UMaterialFunctionMaterialLayer
 ///////////////////////////////////////////////////////////////////////////////
@@ -12841,13 +12935,13 @@ UMaterialInterface* UMaterialFunction::GetPreviewMaterial()
 	{
 		PreviewMaterial = NewObject<UMaterial>(this, NAME_None, RF_Transient | RF_Public);
 
-		PreviewMaterial->Expressions = FunctionExpressions;
+		PreviewMaterial->AssignExpressionCollection(GetExpressionCollection());
 
 		//Find the first output expression and use that. 
-		for( int32 i=0; i < FunctionExpressions.Num() ; ++i)
+		for (UMaterialExpression* Expression : GetExpressions())
 		{
-			UMaterialExpressionFunctionOutput* Output = Cast<UMaterialExpressionFunctionOutput>(FunctionExpressions[i]);
-			if( Output )
+			UMaterialExpressionFunctionOutput* Output = Cast<UMaterialExpressionFunctionOutput>(Expression);
+			if (Output)
 			{
 				Output->ConnectToPreviewMaterial(PreviewMaterial, 0);
 			}
@@ -12865,9 +12959,8 @@ void UMaterialFunction::UpdateInputOutputTypes()
 	CombinedInputTypes = 0;
 	CombinedOutputTypes = 0;
 
-	for (int32 ExpressionIndex = 0; ExpressionIndex < FunctionExpressions.Num(); ExpressionIndex++)
+	for (UMaterialExpression* CurrentExpression : GetExpressions())
 	{
-		UMaterialExpression* CurrentExpression = FunctionExpressions[ExpressionIndex];
 		UMaterialExpressionFunctionOutput* OutputExpression = Cast<UMaterialExpressionFunctionOutput>(CurrentExpression);
 		UMaterialExpressionFunctionInput* InputExpression = Cast<UMaterialExpressionFunctionInput>(CurrentExpression);
 
@@ -12886,7 +12979,7 @@ void UMaterialFunction::UpdateInputOutputTypes()
 void UMaterialFunction::UpdateDependentFunctionCandidates()
 {
 	DependentFunctionExpressionCandidates.Reset();
-	for (UMaterialExpression* CurrentExpression : FunctionExpressions)
+	for (UMaterialExpression* CurrentExpression : GetExpressions())
 	{
 		if (UMaterialExpressionMaterialFunctionCall* MaterialFunctionExpression = Cast<UMaterialExpressionMaterialFunctionCall>(CurrentExpression))
 		{
@@ -12979,12 +13072,38 @@ void UMaterialFunction::PostLoad()
 	Super::PostLoad();
 
 #if WITH_EDITORONLY_DATA
-	for (int32 ExpressionIndex = 0; ExpressionIndex < FunctionExpressions.Num(); ExpressionIndex++)
+	UMaterialFunctionEditorOnlyData* EditorOnly = GetEditorOnlyData();
+
+	if (EditorOnly && FunctionExpressions_DEPRECATED.Num() > 0)
+	{
+		ensure(EditorOnly->ExpressionCollection.Expressions.Num() == 0);
+		EditorOnly->ExpressionCollection.Expressions = MoveTemp(FunctionExpressions_DEPRECATED);
+	}
+
+	if (EditorOnly && FunctionEditorComments_DEPRECATED.Num() > 0)
+	{
+		ensure(EditorOnly->ExpressionCollection.EditorComments.Num() == 0);
+		EditorOnly->ExpressionCollection.EditorComments = MoveTemp(FunctionEditorComments_DEPRECATED);
+	}
+
+	if (EditorOnly && ExpressionExecBegin_DEPRECATED)
+	{
+		ensure(!EditorOnly->ExpressionCollection.ExpressionExecBegin);
+		EditorOnly->ExpressionCollection.ExpressionExecBegin = MoveTemp(ExpressionExecBegin_DEPRECATED);
+	}
+
+	if (EditorOnly && ExpressionExecEnd_DEPRECATED)
+	{
+		ensure(!EditorOnly->ExpressionCollection.ExpressionExecEnd);
+		EditorOnly->ExpressionCollection.ExpressionExecEnd = MoveTemp(ExpressionExecEnd_DEPRECATED);
+	}
+
+	for (UMaterialExpression* Expression : EditorOnly->ExpressionCollection.Expressions)
 	{
 		// Expressions whose type was removed can be nullptr
-		if (FunctionExpressions[ExpressionIndex])
+		if (Expression)
 		{
-			FunctionExpressions[ExpressionIndex]->ConditionalPostLoad();
+			Expression->ConditionalPostLoad();
 		}
 	}
 #endif // WITH_EDITORONLY_DATA
@@ -12999,7 +13118,7 @@ void UMaterialFunction::PostLoad()
 	if (GIsEditor)
 	{
 		// Clean up any removed material expression classes	
-		if (FunctionExpressions.Remove(nullptr) != 0)
+		if (EditorOnly->ExpressionCollection.Expressions.Remove(nullptr) != 0)
 		{
 			// Force this function to recompile because its expressions have changed
 			// Warning: any content taking this path will recompile every load until saved!
@@ -13011,36 +13130,35 @@ void UMaterialFunction::PostLoad()
 	if (GMaterialFunctionsThatNeedExpressionsFlipped.Get(this))
 	{
 		GMaterialFunctionsThatNeedExpressionsFlipped.Clear(this);
-		UMaterial::FlipExpressionPositions(FunctionExpressions, FunctionEditorComments, true);
+		UMaterial::FlipExpressionPositions(EditorOnly->ExpressionCollection.Expressions, EditorOnly->ExpressionCollection.EditorComments, true);
 	}
 	else if (GMaterialFunctionsThatNeedCoordinateCheck.Get(this))
 	{
 		GMaterialFunctionsThatNeedCoordinateCheck.Clear(this);
 		if (HasFlippedCoordinates())
 		{
-			UMaterial::FlipExpressionPositions(FunctionExpressions, FunctionEditorComments, false);
+			UMaterial::FlipExpressionPositions(EditorOnly->ExpressionCollection.Expressions, EditorOnly->ExpressionCollection.EditorComments, false);
 		}
-		UMaterial::FixCommentPositions(FunctionEditorComments);
+		UMaterial::FixCommentPositions(EditorOnly->ExpressionCollection.EditorComments);
 	}
 	else if (GMaterialFunctionsThatNeedCommentFix.Get(this))
 	{
 		GMaterialFunctionsThatNeedCommentFix.Clear(this);
-		UMaterial::FixCommentPositions(FunctionEditorComments);
+		UMaterial::FixCommentPositions(EditorOnly->ExpressionCollection.EditorComments);
 	}
 
 	if (GMaterialFunctionsThatNeedFeatureLevelSM6Fix.Get(this))
 	{
 		GMaterialFunctionsThatNeedFeatureLevelSM6Fix.Clear(this);
-		UMaterial::FixFeatureLevelNodesForSM6(FunctionExpressions);
+		UMaterial::FixFeatureLevelNodesForSM6(EditorOnly->ExpressionCollection.Expressions);
 	}
 
 	if (GMaterialFunctionsThatNeedSamplerFixup.Get(this))
 	{
 		GMaterialFunctionsThatNeedSamplerFixup.Clear(this);
-		const int32 ExpressionCount = FunctionExpressions.Num();
-		for (int32 ExpressionIndex = 0; ExpressionIndex < ExpressionCount; ++ExpressionIndex)
+		for (UMaterialExpression* Expression : EditorOnly->ExpressionCollection.Expressions)
 		{
-			UMaterialExpressionTextureBase* TextureExpression = Cast<UMaterialExpressionTextureBase>(FunctionExpressions[ExpressionIndex]);
+			UMaterialExpressionTextureBase* TextureExpression = Cast<UMaterialExpressionTextureBase>(Expression);
 			if (TextureExpression && TextureExpression->Texture)
 			{
 				switch (TextureExpression->Texture->CompressionSettings)
@@ -13069,13 +13187,49 @@ void UMaterialFunction::PostLoad()
 #endif // #if WITH_EDITOR
 }
 
+#if WITH_EDITORONLY_DATA
+TConstArrayView<TObjectPtr<UMaterialExpression>> UMaterialFunction::GetExpressions() const
+{
+	return GetEditorOnlyData()->ExpressionCollection.Expressions;
+}
+
+TConstArrayView<TObjectPtr<UMaterialExpressionComment>> UMaterialFunction::GetEditorComments() const
+{
+	return GetEditorOnlyData()->ExpressionCollection.EditorComments;
+}
+
+UMaterialExpressionExecBegin* UMaterialFunction::GetExpressionExecBegin() const
+{
+	return GetEditorOnlyData()->ExpressionCollection.ExpressionExecBegin;
+}
+
+UMaterialExpressionExecEnd* UMaterialFunction::GetExpressionExecEnd() const
+{
+	return GetEditorOnlyData()->ExpressionCollection.ExpressionExecEnd;
+}
+
+const FMaterialExpressionCollection& UMaterialFunction::GetExpressionCollection() const
+{
+	return GetEditorOnlyData()->ExpressionCollection;
+}
+
+FMaterialExpressionCollection& UMaterialFunction::GetExpressionCollection()
+{
+	return GetEditorOnlyData()->ExpressionCollection;
+}
+
+void UMaterialFunction::AssignExpressionCollection(const FMaterialExpressionCollection& InCollection)
+{
+	GetEditorOnlyData()->ExpressionCollection = InCollection;
+}
+#endif // WITH_EDITORONLY_DATA
+
 #if WITH_EDITOR
 
 void UMaterialFunction::UpdateFromFunctionResource()
 {
-	for (int32 ExpressionIndex = 0; ExpressionIndex < FunctionExpressions.Num(); ExpressionIndex++)
+	for (UMaterialExpression* CurrentExpression : GetExpressions())
 	{
-		UMaterialExpression* CurrentExpression = FunctionExpressions[ExpressionIndex];
 		UMaterialExpressionMaterialFunctionCall* MaterialFunctionExpression = Cast<UMaterialExpressionMaterialFunctionCall>(CurrentExpression);
 		if (MaterialFunctionExpression)
 		{
@@ -13086,9 +13240,8 @@ void UMaterialFunction::UpdateFromFunctionResource()
 
 void UMaterialFunction::GetInputsAndOutputs(TArray<FFunctionExpressionInput>& OutInputs, TArray<FFunctionExpressionOutput>& OutOutputs) const
 {
-	for (int32 ExpressionIndex = 0; ExpressionIndex < FunctionExpressions.Num(); ExpressionIndex++)
+	for (UMaterialExpression* CurrentExpression : GetExpressions())
 	{
-		UMaterialExpression* CurrentExpression = FunctionExpressions[ExpressionIndex];
 		UMaterialExpressionFunctionOutput* OutputExpression = Cast<UMaterialExpressionFunctionOutput>(CurrentExpression);
 		UMaterialExpressionFunctionInput* InputExpression = Cast<UMaterialExpressionFunctionInput>(CurrentExpression);
 
@@ -13214,7 +13367,7 @@ bool UMaterialFunction::ValidateFunctionUsage(FMaterialCompiler* Compiler, const
 	if (GetMaterialFunctionUsage() == EMaterialFunctionUsage::MaterialLayer)
 	{
 		// Material layers must have a single MA input and output only
-		for (UMaterialExpression* Expression : FunctionExpressions)
+		for (UMaterialExpression* Expression : GetExpressions())
 		{
 			if (UMaterialExpressionFunctionInput* InputExpression = Cast<UMaterialExpressionFunctionInput>(Expression))
 			{
@@ -13250,7 +13403,7 @@ bool UMaterialFunction::ValidateFunctionUsage(FMaterialCompiler* Compiler, const
 	else if (GetMaterialFunctionUsage() == EMaterialFunctionUsage::MaterialLayerBlend)
 	{
 		// Material layer blends can have up to two MA inputs and single MA output only
-		for (UMaterialExpression* Expression : FunctionExpressions)
+		for (UMaterialExpression* Expression : GetExpressions())
 		{
 			if (UMaterialExpressionFunctionInput* InputExpression = Cast<UMaterialExpressionFunctionInput>(Expression))
 			{
@@ -13312,9 +13465,8 @@ int32 UMaterialFunction::Compile(FMaterialCompiler* Compiler, const FFunctionExp
 void UMaterialFunction::LinkIntoCaller(const TArray<FFunctionExpressionInput>& CallerInputs)
 {
 	// Go through all the function's input expressions and hook their inputs up to the corresponding expression in the material being compiled.
-	for (int32 ExpressionIndex = 0; ExpressionIndex < FunctionExpressions.Num(); ExpressionIndex++)
+	for (UMaterialExpression* CurrentExpression : GetExpressions())
 	{
-		UMaterialExpression* CurrentExpression = FunctionExpressions[ExpressionIndex];
 		UMaterialExpressionFunctionInput* InputExpression = Cast<UMaterialExpressionFunctionInput>(CurrentExpression);
 
 		if (InputExpression)
@@ -13347,9 +13499,8 @@ void UMaterialFunction::LinkIntoCaller(const TArray<FFunctionExpressionInput>& C
 
 void UMaterialFunction::UnlinkFromCaller()
 {
-	for (int32 ExpressionIndex = 0; ExpressionIndex < FunctionExpressions.Num(); ExpressionIndex++)
+	for (UMaterialExpression* CurrentExpression : GetExpressions())
 	{
-		UMaterialExpression* CurrentExpression = FunctionExpressions[ExpressionIndex];
 		UMaterialExpressionFunctionInput* InputExpression = Cast<UMaterialExpressionFunctionInput>(CurrentExpression);
 
 		if (InputExpression)
@@ -13388,9 +13539,8 @@ bool UMaterialFunction::IsDependent(UMaterialFunctionInterface* OtherFunction)
 #endif
 
 	bool bIsDependent = false;
-	for (int32 ExpressionIndex = 0; ExpressionIndex < FunctionExpressions.Num(); ExpressionIndex++)
+	for (UMaterialExpression* CurrentExpression : GetExpressions())
 	{
-		UMaterialExpression* CurrentExpression = FunctionExpressions[ExpressionIndex];
 		UMaterialExpressionMaterialFunctionCall* MaterialFunctionExpression = Cast<UMaterialExpressionMaterialFunctionCall>(CurrentExpression);
 		if (MaterialFunctionExpression && MaterialFunctionExpression->MaterialFunction)
 		{
@@ -13450,9 +13600,9 @@ bool UMaterialFunction::HasFlippedCoordinates() const
 	uint32 ReversedInputCount = 0;
 	uint32 StandardInputCount = 0;
 
-	for (int32 Index = 0; Index < FunctionExpressions.Num(); ++Index)
+	for (UMaterialExpression* CurrentExpression : GetExpressions())
 	{
-		UMaterialExpressionFunctionOutput* FunctionOutput = Cast<UMaterialExpressionFunctionOutput>(FunctionExpressions[Index]);
+		UMaterialExpressionFunctionOutput* FunctionOutput = Cast<UMaterialExpressionFunctionOutput>(CurrentExpression);
 		if (FunctionOutput && FunctionOutput->A.Expression)
 		{
 			if (FunctionOutput->A.Expression->MaterialExpressionEditorX > FunctionOutput->MaterialExpressionEditorX)
@@ -13473,7 +13623,7 @@ bool UMaterialFunction::HasFlippedCoordinates() const
 bool UMaterialFunction::SetParameterValueEditorOnly(const FName& ParameterName, const FMaterialParameterMetadata& Meta)
 {
 	bool bResult = false;
-	for (UMaterialExpression* Expression : FunctionExpressions)
+	for (UMaterialExpression* Expression : GetExpressions())
 	{
 		if (Expression->SetParameterValue(ParameterName, Meta))
 		{
@@ -13489,15 +13639,11 @@ bool UMaterialFunction::SetParameterValueEditorOnly(const FName& ParameterName,
 
 				for (UMaterialFunctionInterface* Function : Functions)
 				{
-					const TArray<TObjectPtr<UMaterialExpression>>* ExpressionPtr = Function->GetFunctionExpressions();
-					if (ExpressionPtr)
+					for (const TObjectPtr<UMaterialExpression>& FunctionExpression : Function->GetExpressions())
 					{
-						for (const TObjectPtr<UMaterialExpression>& FunctionExpression : *ExpressionPtr)
+						if (FunctionExpression->SetParameterValue(ParameterName, Meta))
 						{
-							if (FunctionExpression->SetParameterValue(ParameterName, Meta))
-							{
-								bResult = true;
-							}
+							bResult = true;
 						}
 					}
 				}
@@ -13581,39 +13727,41 @@ bool UMaterialFunction::IsUsingNewHLSLGenerator() const
 void UMaterialFunction::CreateExecutionFlowExpressions()
 {
 	const bool bUsingControlFlow = IsUsingControlFlow();
+	UMaterialFunctionEditorOnlyData* EditorOnly = GetEditorOnlyData();
 	if (bUsingControlFlow)
 	{
-		if (!ExpressionExecBegin)
+		if (!EditorOnly->ExpressionCollection.ExpressionExecBegin)
 		{
-			ExpressionExecBegin = NewObject<UMaterialExpressionExecBegin>(this);
-			ExpressionExecBegin->Function = this;
-			FunctionExpressions.Add(ExpressionExecBegin);
+			EditorOnly->ExpressionCollection.ExpressionExecBegin = NewObject<UMaterialExpressionExecBegin>(this);
+			EditorOnly->ExpressionCollection.ExpressionExecBegin->Function = this;
+			EditorOnly->ExpressionCollection.AddExpression(EditorOnly->ExpressionCollection.ExpressionExecBegin);
 		}
 
-		if (!ExpressionExecEnd)
+		if (!EditorOnly->ExpressionCollection.ExpressionExecEnd)
 		{
-			ExpressionExecEnd = NewObject<UMaterialExpressionExecEnd>(this);
-			ExpressionExecEnd->Function = this;
-			FunctionExpressions.Add(ExpressionExecEnd);
+			EditorOnly->ExpressionCollection.ExpressionExecEnd = NewObject<UMaterialExpressionExecEnd>(this);
+			EditorOnly->ExpressionCollection.ExpressionExecEnd->Function = this;
+			EditorOnly->ExpressionCollection.AddExpression(EditorOnly->ExpressionCollection.ExpressionExecEnd);
 		}
 	}
 
 	if (EditorMaterial)
 	{
+		UMaterialEditorOnlyData* EditorMaterialEditorOnly = EditorMaterial->GetEditorOnlyData();
 		EditorMaterial->bEnableExecWire = bUsingControlFlow;
-		EditorMaterial->ExpressionExecBegin = ExpressionExecBegin;
-		EditorMaterial->ExpressionExecEnd = ExpressionExecEnd;
-		if (ExpressionExecBegin)
+		EditorMaterialEditorOnly->ExpressionCollection.ExpressionExecBegin = EditorOnly->ExpressionCollection.ExpressionExecBegin;
+		EditorMaterialEditorOnly->ExpressionCollection.ExpressionExecEnd = EditorOnly->ExpressionCollection.ExpressionExecEnd;
+		if (EditorOnly->ExpressionCollection.ExpressionExecBegin)
 		{
-			EditorMaterial->Expressions.AddUnique(ExpressionExecBegin);
-			ExpressionExecBegin->Function = nullptr;
-			ExpressionExecBegin->Material = EditorMaterial;
+			EditorMaterialEditorOnly->ExpressionCollection.AddExpression(EditorOnly->ExpressionCollection.ExpressionExecBegin);
+			EditorOnly->ExpressionCollection.ExpressionExecBegin->Function = nullptr;
+			EditorOnly->ExpressionCollection.ExpressionExecBegin->Material = EditorMaterial;
 		}
-		if (ExpressionExecEnd)
+		if (EditorOnly->ExpressionCollection.ExpressionExecEnd)
 		{
-			EditorMaterial->Expressions.AddUnique(ExpressionExecEnd);
-			ExpressionExecEnd->Function = nullptr;
-			ExpressionExecEnd->Material = EditorMaterial;
+			EditorMaterialEditorOnly->ExpressionCollection.AddExpression(EditorOnly->ExpressionCollection.ExpressionExecEnd);
+			EditorOnly->ExpressionCollection.ExpressionExecEnd->Function = nullptr;
+			EditorOnly->ExpressionCollection.ExpressionExecEnd->Material = EditorMaterial;
 		}
 	}
 }
@@ -13632,10 +13780,32 @@ UMaterialFunctionInstance::UMaterialFunctionInstance(const FObjectInitializer& O
 #endif
 }
 
+UMaterialFunction* UMaterialFunctionInstance::GetBaseFunction(FMFRecursionGuard RecursionGuard)
+{
+	if (!Parent || RecursionGuard.Contains(this))
+	{
+		return nullptr;
+	}
+
+	RecursionGuard.Set(this);
+	return Parent->GetBaseFunction(RecursionGuard);
+}
+
+const UMaterialFunction* UMaterialFunctionInstance::GetBaseFunction(FMFRecursionGuard RecursionGuard) const
+{
+	if (!Parent || RecursionGuard.Contains(this))
+	{
+		return nullptr;
+	}
+
+	RecursionGuard.Set(this);
+	return Parent->GetBaseFunction(RecursionGuard);
+}
+
 #if WITH_EDITOR
 void UMaterialFunctionInstance::UpdateParameterSet()
 {
-	if (UMaterialFunction* BaseFunction = Cast<UMaterialFunction>(GetBaseFunction()))
+	if (UMaterialFunction* BaseFunction = GetBaseFunction())
 	{
 		TArray<UMaterialFunctionInterface*> Functions;
 		BaseFunction->GetDependentFunctions(Functions);
@@ -13644,7 +13814,7 @@ void UMaterialFunctionInstance::UpdateParameterSet()
 		// Loop through all contained parameters and update names as needed
 		for (UMaterialFunctionInterface* Function : Functions)
 		{
-			for (UMaterialExpression* FunctionExpression : *Function->GetFunctionExpressions())
+			for (UMaterialExpression* FunctionExpression : Function->GetExpressions())
 			{
 				if (const UMaterialExpressionScalarParameter* ScalarParameter = Cast<const UMaterialExpressionScalarParameter>(FunctionExpression))
 				{
@@ -13739,8 +13909,8 @@ void UMaterialFunctionInstance::OverrideMaterialInstanceParameterValues(UMateria
 
 	// Static parameters
 	FStaticParameterSet StaticParametersOverride = Instance->GetStaticParameters();
-	StaticParametersOverride.StaticSwitchParameters = StaticSwitchParameterValues;
-	StaticParametersOverride.StaticComponentMaskParameters = StaticComponentMaskParameterValues;
+	StaticParametersOverride.EditorOnly.StaticSwitchParameters = StaticSwitchParameterValues;
+	StaticParametersOverride.EditorOnly.StaticComponentMaskParameters = StaticComponentMaskParameterValues;
 	Instance->UpdateStaticPermutation(StaticParametersOverride);
 }
 
@@ -13910,16 +14080,16 @@ bool UMaterialFunctionInstance::GetParameterOverrideValue(EMaterialParameterType
 }
 
 ///////////////////////////////////////////////////////////////////////////////
-// FMaterialLayersFunctions::ID
+// FMaterialLayersFunctionsID
 ///////////////////////////////////////////////////////////////////////////////
 
-bool FMaterialLayersFunctions::ID::operator==(const ID& Reference) const
+bool FMaterialLayersFunctionsID::operator==(const FMaterialLayersFunctionsID& Reference) const
 {
 	return LayerIDs == Reference.LayerIDs && BlendIDs == Reference.BlendIDs && LayerStates == Reference.LayerStates;
 }
 
 
-void FMaterialLayersFunctions::ID::SerializeForDDC(FArchive& Ar)
+void FMaterialLayersFunctionsID::SerializeForDDC(FArchive& Ar)
 {
 	Ar << LayerIDs;
 	Ar << BlendIDs;
@@ -13927,7 +14097,7 @@ void FMaterialLayersFunctions::ID::SerializeForDDC(FArchive& Ar)
 }
 
 
-void FMaterialLayersFunctions::ID::UpdateHash(FSHA1& HashState) const
+void FMaterialLayersFunctionsID::UpdateHash(FSHA1& HashState) const
 {
 	for (const FGuid &Guid : LayerIDs)
 	{
@@ -13941,7 +14111,7 @@ void FMaterialLayersFunctions::ID::UpdateHash(FSHA1& HashState) const
 }
 
 
-void FMaterialLayersFunctions::ID::AppendKeyString(FString& KeyString) const
+void FMaterialLayersFunctionsID::AppendKeyString(FString& KeyString) const
 {
 	for (const FGuid &Guid : LayerIDs)
 	{
@@ -13965,9 +14135,9 @@ void FMaterialLayersFunctions::ID::AppendKeyString(FString& KeyString) const
 const FGuid FMaterialLayersFunctions::BackgroundGuid(2u, 0u, 0u, 0u);
 
 #if WITH_EDITOR
-const FMaterialLayersFunctions::ID FMaterialLayersFunctions::GetID() const
+const FMaterialLayersFunctionsID FMaterialLayersFunctionsRuntimeData::GetID(const FMaterialLayersFunctionsEditorOnlyData& EditorOnly) const
 {
-	FMaterialLayersFunctions::ID Result;
+	FMaterialLayersFunctionsID Result;
 
 	// Store the layer IDs in following format - stateID per function
 	Result.LayerIDs.SetNum(Layers.Num());
@@ -14002,7 +14172,7 @@ const FMaterialLayersFunctions::ID FMaterialLayersFunctions::GetID() const
 	}
 
 	// Store the states copy
-	Result.LayerStates = LayerStates;
+	Result.LayerStates = EditorOnly.LayerStates;
 
 	return Result;
 }
@@ -14026,17 +14196,46 @@ void FMaterialLayersFunctions::PostSerialize(const FArchive& Ar)
 #if WITH_EDITORONLY_DATA
 	if (Ar.IsLoading())
 	{
-		if (LayerGuids.Num() != Layers.Num() ||
-			LayerLinkStates.Num() != Layers.Num())
+		if (LayerStates_DEPRECATED.Num() > 0)
+		{
+			EditorOnly.LayerStates = MoveTemp(LayerStates_DEPRECATED);
+		}
+		if (LayerNames_DEPRECATED.Num() > 0)
+		{
+			EditorOnly.LayerNames = MoveTemp(LayerNames_DEPRECATED);
+		}
+		if (RestrictToLayerRelatives_DEPRECATED.Num() > 0)
+		{
+			EditorOnly.RestrictToLayerRelatives = MoveTemp(RestrictToLayerRelatives_DEPRECATED);
+		}
+		if (RestrictToBlendRelatives_DEPRECATED.Num() > 0)
+		{
+			EditorOnly.RestrictToBlendRelatives = MoveTemp(RestrictToBlendRelatives_DEPRECATED);
+		}
+		if (LayerGuids_DEPRECATED.Num() > 0)
+		{
+			EditorOnly.LayerGuids = MoveTemp(LayerGuids_DEPRECATED);
+		}
+		if (LayerLinkStates_DEPRECATED.Num() > 0)
+		{
+			EditorOnly.LayerLinkStates = MoveTemp(LayerLinkStates_DEPRECATED);
+		}
+		if (DeletedParentLayerGuids_DEPRECATED.Num() > 0)
+		{
+			EditorOnly.DeletedParentLayerGuids = MoveTemp(DeletedParentLayerGuids_DEPRECATED);
+		}
+
+		if (EditorOnly.LayerGuids.Num() != Layers.Num() ||
+			EditorOnly.LayerLinkStates.Num() != Layers.Num())
 		{
 			const int32 NumLayers = Layers.Num();
-			LayerGuids.Empty(NumLayers);
-			LayerLinkStates.Empty(NumLayers);
+			EditorOnly.LayerGuids.Empty(NumLayers);
+			EditorOnly.LayerLinkStates.Empty(NumLayers);
 
 			if (NumLayers > 0)
 			{
-				LayerGuids.Add(BackgroundGuid);
-				LayerLinkStates.Add(EMaterialLayerLinkState::Uninitialized);
+				EditorOnly.LayerGuids.Add(BackgroundGuid);
+				EditorOnly.LayerLinkStates.Add(EMaterialLayerLinkState::Uninitialized);
 
 				for (int32 i = 1; i < NumLayers; ++i)
 				{
@@ -14045,32 +14244,34 @@ void FMaterialLayersFunctions::PostSerialize(const FArchive& Ar)
 					// But it's possible *this* material may not actually be saved in that case
 					// If that happens, need to ensure that the guids remain consistent if this material is loaded again;
 					// otherwise they will no longer match the guids that were saved into the child material
-					LayerGuids.Add(FGuid(3u, 0u, 0u, i));
-					LayerLinkStates.Add(EMaterialLayerLinkState::Uninitialized);
+					EditorOnly.LayerGuids.Add(FGuid(3u, 0u, 0u, i));
+					EditorOnly.LayerLinkStates.Add(EMaterialLayerLinkState::Uninitialized);
 				}
 			}
 		}
 	}
 #endif // WITH_EDITORONLY_DATA
 }
-
+#if WITH_EDITOR
 int32 FMaterialLayersFunctions::AppendBlendedLayer()
 {
 	const int32 LayerIndex = Layers.AddDefaulted();
 	Blends.AddDefaulted();
-#if WITH_EDITORONLY_DATA
-	LayerStates.Add(true);
+	EditorOnly.LayerStates.Add(true);
 	FText LayerName = FText::Format(LOCTEXT("LayerPrefix", "Layer {0}"), Layers.Num() - 1);
-	LayerNames.Add(LayerName);
-	RestrictToLayerRelatives.Add(false);
-	RestrictToBlendRelatives.Add(false);
-	LayerGuids.Add(FGuid::NewGuid());
-	LayerLinkStates.Add(EMaterialLayerLinkState::NotFromParent);
-#endif // WITH_EDITORONLY_DATA
+	EditorOnly.LayerNames.Add(LayerName);
+	EditorOnly.RestrictToLayerRelatives.Add(false);
+	EditorOnly.RestrictToBlendRelatives.Add(false);
+	EditorOnly.LayerGuids.Add(FGuid::NewGuid());
+	EditorOnly.LayerLinkStates.Add(EMaterialLayerLinkState::NotFromParent);
 	return LayerIndex;
 }
 
-int32 FMaterialLayersFunctions::AddLayerCopy(const FMaterialLayersFunctions& Source, int32 SourceLayerIndex, bool bVisible, EMaterialLayerLinkState LinkState)
+int32 FMaterialLayersFunctions::AddLayerCopy(const FMaterialLayersFunctionsRuntimeData& Source,
+	const FMaterialLayersFunctionsEditorOnlyData& SourceEditorOnly,
+	int32 SourceLayerIndex,
+	bool bVisible,
+	EMaterialLayerLinkState LinkState)
 {
 	check(LinkState != EMaterialLayerLinkState::Uninitialized);
 	const int32 LayerIndex = Layers.Num();
@@ -14081,35 +14282,35 @@ int32 FMaterialLayersFunctions::AddLayerCopy(const FMaterialLayersFunctions& Sou
 		Blends.Add(Source.Blends[SourceLayerIndex - 1]);
 	}
 	
-#if WITH_EDITOR
-	LayerStates.Add(bVisible);
-	LayerNames.Add(Source.LayerNames[SourceLayerIndex]);
-	RestrictToLayerRelatives.Add(Source.RestrictToLayerRelatives[SourceLayerIndex]);
+	EditorOnly.LayerStates.Add(bVisible);
+	EditorOnly.LayerNames.Add(SourceEditorOnly.LayerNames[SourceLayerIndex]);
+	EditorOnly.RestrictToLayerRelatives.Add(SourceEditorOnly.RestrictToLayerRelatives[SourceLayerIndex]);
 	if (LayerIndex > 0)
 	{
-		RestrictToBlendRelatives.Add(Source.RestrictToBlendRelatives[SourceLayerIndex - 1]);
+		EditorOnly.RestrictToBlendRelatives.Add(SourceEditorOnly.RestrictToBlendRelatives[SourceLayerIndex - 1]);
 	}
-	LayerGuids.Add(Source.LayerGuids[SourceLayerIndex]);
-	LayerLinkStates.Add(LinkState);
-#endif
+	EditorOnly.LayerGuids.Add(SourceEditorOnly.LayerGuids[SourceLayerIndex]);
+	EditorOnly.LayerLinkStates.Add(LinkState);
 	return LayerIndex;
 }
 
-void FMaterialLayersFunctions::InsertLayerCopy(const FMaterialLayersFunctions& Source, int32 SourceLayerIndex, EMaterialLayerLinkState LinkState, int32 LayerIndex)
+void FMaterialLayersFunctions::InsertLayerCopy(const FMaterialLayersFunctionsRuntimeData& Source,
+	const FMaterialLayersFunctionsEditorOnlyData& SourceEditorOnly,
+	int32 SourceLayerIndex,
+	EMaterialLayerLinkState LinkState,
+	int32 LayerIndex)
 {
 	check(LinkState != EMaterialLayerLinkState::Uninitialized);
 	check(LayerIndex > 0);
 	Layers.Insert(Source.Layers[SourceLayerIndex], LayerIndex);
 	Blends.Insert(Source.Blends[SourceLayerIndex - 1], LayerIndex - 1);
 	
-#if WITH_EDITOR
-	LayerStates.Insert(Source.LayerStates[SourceLayerIndex], LayerIndex);
-	LayerNames.Insert(Source.LayerNames[SourceLayerIndex], LayerIndex);
-	RestrictToLayerRelatives.Insert(Source.RestrictToLayerRelatives[SourceLayerIndex], LayerIndex);
-	RestrictToBlendRelatives.Insert(Source.RestrictToBlendRelatives[SourceLayerIndex - 1], LayerIndex - 1);
-	LayerGuids.Insert(Source.LayerGuids[SourceLayerIndex], LayerIndex);
-	LayerLinkStates.Insert(LinkState, LayerIndex);
-#endif
+	EditorOnly.LayerStates.Insert(SourceEditorOnly.LayerStates[SourceLayerIndex], LayerIndex);
+	EditorOnly.LayerNames.Insert(SourceEditorOnly.LayerNames[SourceLayerIndex], LayerIndex);
+	EditorOnly.RestrictToLayerRelatives.Insert(SourceEditorOnly.RestrictToLayerRelatives[SourceLayerIndex], LayerIndex);
+	EditorOnly.RestrictToBlendRelatives.Insert(SourceEditorOnly.RestrictToBlendRelatives[SourceLayerIndex - 1], LayerIndex - 1);
+	EditorOnly.LayerGuids.Insert(SourceEditorOnly.LayerGuids[SourceLayerIndex], LayerIndex);
+	EditorOnly.LayerLinkStates.Insert(LinkState, LayerIndex);
 }
 
 void FMaterialLayersFunctions::RemoveBlendedLayerAt(int32 Index)
@@ -14120,24 +14321,25 @@ void FMaterialLayersFunctions::RemoveBlendedLayerAt(int32 Index)
 		Layers.RemoveAt(Index);
 		Blends.RemoveAt(Index - 1);
 		
-#if WITH_EDITOR
-		check(LayerStates.IsValidIndex(Index) && LayerNames.IsValidIndex(Index) && RestrictToLayerRelatives.IsValidIndex(Index) && RestrictToBlendRelatives.IsValidIndex(Index - 1));
+		check(EditorOnly.LayerStates.IsValidIndex(Index) &&
+			EditorOnly.LayerNames.IsValidIndex(Index) &&
+			EditorOnly.RestrictToLayerRelatives.IsValidIndex(Index) &&
+			EditorOnly.RestrictToBlendRelatives.IsValidIndex(Index - 1));
 
-		if (LayerLinkStates[Index] != EMaterialLayerLinkState::NotFromParent)
+		if (EditorOnly.LayerLinkStates[Index] != EMaterialLayerLinkState::NotFromParent)
 		{
 			// Save the parent guid as explicitly deleted, so it's not added back
-			const FGuid& LayerGuid = LayerGuids[Index];
-			check(!DeletedParentLayerGuids.Contains(LayerGuid));
-			DeletedParentLayerGuids.Add(LayerGuid);
+			const FGuid& LayerGuid = EditorOnly.LayerGuids[Index];
+			check(!EditorOnly.DeletedParentLayerGuids.Contains(LayerGuid));
+			EditorOnly.DeletedParentLayerGuids.Add(LayerGuid);
 		}
 
-		LayerStates.RemoveAt(Index);
-		LayerNames.RemoveAt(Index);
-		RestrictToLayerRelatives.RemoveAt(Index);
-		RestrictToBlendRelatives.RemoveAt(Index - 1);
-		LayerGuids.RemoveAt(Index);
-		LayerLinkStates.RemoveAt(Index);
-#endif //WITH_EDITOR
+		EditorOnly.LayerStates.RemoveAt(Index);
+		EditorOnly.LayerNames.RemoveAt(Index);
+		EditorOnly.RestrictToLayerRelatives.RemoveAt(Index);
+		EditorOnly.RestrictToBlendRelatives.RemoveAt(Index - 1);
+		EditorOnly.LayerGuids.RemoveAt(Index);
+		EditorOnly.LayerLinkStates.RemoveAt(Index);
 	}
 }
 
@@ -14149,56 +14351,53 @@ void FMaterialLayersFunctions::MoveBlendedLayer(int32 SrcLayerIndex, int32 DstLa
 	{
 		Layers.Swap(SrcLayerIndex, DstLayerIndex);
 		Blends.Swap(SrcLayerIndex - 1, DstLayerIndex - 1);
-#if WITH_EDITOR
-		LayerStates.Swap(SrcLayerIndex, DstLayerIndex);
-		LayerNames.Swap(SrcLayerIndex, DstLayerIndex);
-		RestrictToLayerRelatives.Swap(SrcLayerIndex, DstLayerIndex);
-		RestrictToBlendRelatives.Swap(SrcLayerIndex - 1, DstLayerIndex - 1);
-		LayerGuids.Swap(SrcLayerIndex, DstLayerIndex);
-		LayerLinkStates.Swap(SrcLayerIndex, DstLayerIndex);
-#endif //WITH_EDITOR
+		EditorOnly.LayerStates.Swap(SrcLayerIndex, DstLayerIndex);
+		EditorOnly.LayerNames.Swap(SrcLayerIndex, DstLayerIndex);
+		EditorOnly.RestrictToLayerRelatives.Swap(SrcLayerIndex, DstLayerIndex);
+		EditorOnly.RestrictToBlendRelatives.Swap(SrcLayerIndex - 1, DstLayerIndex - 1);
+		EditorOnly.LayerGuids.Swap(SrcLayerIndex, DstLayerIndex);
+		EditorOnly.LayerLinkStates.Swap(SrcLayerIndex, DstLayerIndex);
 	}
 }
 
-#if WITH_EDITOR
 void FMaterialLayersFunctions::UnlinkLayerFromParent(int32 Index)
 {
-	if (LayerLinkStates[Index] == EMaterialLayerLinkState::LinkedToParent)
+	if (EditorOnly.LayerLinkStates[Index] == EMaterialLayerLinkState::LinkedToParent)
 	{
-		LayerLinkStates[Index] = EMaterialLayerLinkState::UnlinkedFromParent;
+		EditorOnly.LayerLinkStates[Index] = EMaterialLayerLinkState::UnlinkedFromParent;
 	}
 }
 
 bool FMaterialLayersFunctions::IsLayerLinkedToParent(int32 Index) const
 {
-	if (LayerLinkStates.IsValidIndex(Index))
+	if (EditorOnly.LayerLinkStates.IsValidIndex(Index))
 	{
-		return LayerLinkStates[Index] == EMaterialLayerLinkState::LinkedToParent;
+		return EditorOnly.LayerLinkStates[Index] == EMaterialLayerLinkState::LinkedToParent;
 	}
 	return false;
 }
 
 void FMaterialLayersFunctions::RelinkLayersToParent()
 {
-	for (int32 Index = 0; Index < LayerLinkStates.Num(); ++Index)
+	for (int32 Index = 0; Index < EditorOnly.LayerLinkStates.Num(); ++Index)
 	{
-		if (LayerLinkStates[Index] == EMaterialLayerLinkState::UnlinkedFromParent)
+		if (EditorOnly.LayerLinkStates[Index] == EMaterialLayerLinkState::UnlinkedFromParent)
 		{
-			LayerLinkStates[Index] = EMaterialLayerLinkState::LinkedToParent;
+			EditorOnly.LayerLinkStates[Index] = EMaterialLayerLinkState::LinkedToParent;
 		}
 	}
-	DeletedParentLayerGuids.Empty();
+	EditorOnly.DeletedParentLayerGuids.Empty();
 }
 
 bool FMaterialLayersFunctions::HasAnyUnlinkedLayers() const
 {
-	if (DeletedParentLayerGuids.Num() > 0)
+	if (EditorOnly.DeletedParentLayerGuids.Num() > 0)
 	{
 		return true;
 	}
-	for (int32 Index = 0; Index < LayerLinkStates.Num(); ++Index)
+	for (int32 Index = 0; Index < EditorOnly.LayerLinkStates.Num(); ++Index)
 	{
-		if (LayerLinkStates[Index] == EMaterialLayerLinkState::UnlinkedFromParent)
+		if (EditorOnly.LayerLinkStates[Index] == EMaterialLayerLinkState::UnlinkedFromParent)
 		{
 			return true;
 		}
@@ -14206,7 +14405,7 @@ bool FMaterialLayersFunctions::HasAnyUnlinkedLayers() const
 	return false;
 }
 
-void FMaterialLayersFunctions::LinkAllLayersToParent()
+void FMaterialLayersFunctionsEditorOnlyData::LinkAllLayersToParent()
 {
 	for (int32 Index = 0; Index < LayerLinkStates.Num(); ++Index)
 	{
@@ -14214,37 +14413,44 @@ void FMaterialLayersFunctions::LinkAllLayersToParent()
 	}
 }
 
-bool FMaterialLayersFunctions::MatchesParent(const FMaterialLayersFunctions& Parent) const
+bool FMaterialLayersFunctions::MatchesParent(const FMaterialLayersFunctionsRuntimeData& Runtime,
+	const FMaterialLayersFunctionsEditorOnlyData& EditorOnly,
+	const FMaterialLayersFunctionsRuntimeData& ParentRuntime,
+	const FMaterialLayersFunctionsEditorOnlyData& ParentEditorOnly)
 {
-	if (Layers.Num() != Parent.Layers.Num())
+	if (Runtime.Layers.Num() != ParentRuntime.Layers.Num())
 	{
 		return false;
 	}
 
-	for (int32 LayerIndex = 0; LayerIndex < Layers.Num(); ++LayerIndex)
+	for (int32 LayerIndex = 0; LayerIndex < Runtime.Layers.Num(); ++LayerIndex)
 	{
-		const EMaterialLayerLinkState LinkState = LayerLinkStates[LayerIndex];
+		const EMaterialLayerLinkState LinkState = EditorOnly.LayerLinkStates[LayerIndex];
 		if (LinkState != EMaterialLayerLinkState::LinkedToParent)
 		{
 			return false;
 		}
 
-		const FGuid& LayerGuid = LayerGuids[LayerIndex];
-		const int32 ParentLayerIndex = Parent.LayerGuids.Find(LayerGuid);
+		const FGuid& LayerGuid = EditorOnly.LayerGuids[LayerIndex];
+		const int32 ParentLayerIndex = ParentEditorOnly.LayerGuids.Find(LayerGuid);
 		if (ParentLayerIndex != LayerIndex)
 		{
 			return false;
 		}
 
-		if (LayerStates[LayerIndex] != Parent.LayerStates[ParentLayerIndex])
+		// Possible for LayerStates arrays to be empty, if this is cooked data
+		// We assume all cooked layers are visible
+		const bool bLayerVisible = (EditorOnly.LayerStates.Num() > 0) ? EditorOnly.LayerStates[LayerIndex] : true;
+		const bool bParentLayerVisible = (ParentEditorOnly.LayerStates.Num() > 0) ? ParentEditorOnly.LayerStates[ParentLayerIndex] : true;
+		if (bLayerVisible != bParentLayerVisible)
 		{
 			return false;
 		}
-		if (Layers[LayerIndex] != Parent.Layers[ParentLayerIndex])
+		if (Runtime.Layers[LayerIndex] != ParentRuntime.Layers[ParentLayerIndex])
 		{
 			return false;
 		}
-		if (LayerIndex > 0 && Blends[LayerIndex - 1] != Parent.Blends[ParentLayerIndex - 1])
+		if (LayerIndex > 0 && Runtime.Blends[LayerIndex - 1] != ParentRuntime.Blends[ParentLayerIndex - 1])
 		{
 			return false;
 		}
@@ -14253,10 +14459,14 @@ bool FMaterialLayersFunctions::MatchesParent(const FMaterialLayersFunctions& Par
 	return true;
 }
 
-bool FMaterialLayersFunctions::ResolveParent(const FMaterialLayersFunctions& Parent, TArray<int32>& OutRemapLayerIndices)
+bool FMaterialLayersFunctions::ResolveParent(const FMaterialLayersFunctionsRuntimeData& ParentRuntime,
+	const FMaterialLayersFunctionsEditorOnlyData& ParentEditorOnly,
+	FMaterialLayersFunctionsRuntimeData& Runtime,
+	FMaterialLayersFunctionsEditorOnlyData& EditorOnly,
+	TArray<int32>& OutRemapLayerIndices)
 {
-	check(LayerGuids.Num() == Layers.Num());
-	check(LayerLinkStates.Num() == Layers.Num());
+	check(EditorOnly.LayerGuids.Num() == Runtime.Layers.Num());
+	check(EditorOnly.LayerLinkStates.Num() == Runtime.Layers.Num());
 
 	FMaterialLayersFunctions ResolvedLayers;
 	TArray<int32> ParentLayerIndices;
@@ -14264,12 +14474,12 @@ bool FMaterialLayersFunctions::ResolveParent(const FMaterialLayersFunctions& Par
 	ResolvedLayers.Empty();
 
 	bool bHasUninitializedLinks = false;
-	for (int32 LayerIndex = 0; LayerIndex < Layers.Num(); ++LayerIndex)
+	for (int32 LayerIndex = 0; LayerIndex < Runtime.Layers.Num(); ++LayerIndex)
 	{
-		const FText& LayerName = LayerNames[LayerIndex];
-		const FGuid& LayerGuid = LayerGuids[LayerIndex];
-		const bool bLayerVisible = LayerStates[LayerIndex];
-		const EMaterialLayerLinkState LinkState = LayerLinkStates[LayerIndex];
+		const FText& LayerName = EditorOnly.LayerNames[LayerIndex];
+		const FGuid& LayerGuid = EditorOnly.LayerGuids[LayerIndex];
+		const bool bLayerVisible = EditorOnly.LayerStates[LayerIndex];
+		const EMaterialLayerLinkState LinkState = EditorOnly.LayerLinkStates[LayerIndex];
 
 		int32 ParentLayerIndex = INDEX_NONE;
 		if (LinkState == EMaterialLayerLinkState::Uninitialized)
@@ -14278,17 +14488,17 @@ bool FMaterialLayersFunctions::ResolveParent(const FMaterialLayersFunctions& Par
 			if (LayerIndex == 0)
 			{
 				// Base layer must match against base layer
-				if (Parent.Layers.Num() > 0)
+				if (ParentRuntime.Layers.Num() > 0)
 				{
 					ParentLayerIndex = 0;
 				}
 			}
 			else
 			{
-				for (int32 CheckLayerIndex = 1; CheckLayerIndex < Parent.Layers.Num(); ++CheckLayerIndex)
+				for (int32 CheckLayerIndex = 1; CheckLayerIndex < ParentRuntime.Layers.Num(); ++CheckLayerIndex)
 				{
 					// check if name matches, and if we haven't already linked to this parent layer
-					if (LayerName.EqualTo(Parent.LayerNames[CheckLayerIndex]) &&
+					if (LayerName.EqualTo(ParentEditorOnly.LayerNames[CheckLayerIndex]) &&
 						!ParentLayerIndices.Contains(CheckLayerIndex))
 					{
 						ParentLayerIndex = CheckLayerIndex;
@@ -14301,23 +14511,23 @@ bool FMaterialLayersFunctions::ResolveParent(const FMaterialLayersFunctions& Par
 			if (ParentLayerIndex == INDEX_NONE)
 			{
 				// Didn't find layer in the parent, assume it's local to this material
-				ResolvedLayerIndex = ResolvedLayers.AddLayerCopy(*this, LayerIndex, bLayerVisible, EMaterialLayerLinkState::NotFromParent);
+				ResolvedLayerIndex = ResolvedLayers.AddLayerCopy(Runtime, EditorOnly, LayerIndex, bLayerVisible, EMaterialLayerLinkState::NotFromParent);
 				ParentLayerIndices.Add(INDEX_NONE);
 			}
 			else
 			{
 				// See if we match layer in parent
-				if (Layers[LayerIndex] == Parent.Layers[ParentLayerIndex] &&
-					(LayerIndex == 0 || Blends[LayerIndex - 1] == Parent.Blends[ParentLayerIndex - 1]))
+				if (Runtime.Layers[LayerIndex] == ParentRuntime.Layers[ParentLayerIndex] &&
+					(LayerIndex == 0 || Runtime.Blends[LayerIndex - 1] == ParentRuntime.Blends[ParentLayerIndex - 1]))
 				{
 					// Parent layer matches, so link to parent
-					ResolvedLayerIndex = ResolvedLayers.AddLayerCopy(Parent, ParentLayerIndex, bLayerVisible, EMaterialLayerLinkState::LinkedToParent);
+					ResolvedLayerIndex = ResolvedLayers.AddLayerCopy(ParentRuntime, ParentEditorOnly, ParentLayerIndex, bLayerVisible, EMaterialLayerLinkState::LinkedToParent);
 				}
 				else
 				{
 					// Parent layer does NOT match, so make the child overriden
-					ResolvedLayerIndex = ResolvedLayers.AddLayerCopy(*this, LayerIndex, bLayerVisible, EMaterialLayerLinkState::UnlinkedFromParent);
-					ResolvedLayers.LayerGuids[ResolvedLayerIndex] = Parent.LayerGuids[ParentLayerIndex]; // Still need to match guid to parent
+					ResolvedLayerIndex = ResolvedLayers.AddLayerCopy(Runtime, EditorOnly, LayerIndex, bLayerVisible, EMaterialLayerLinkState::UnlinkedFromParent);
+					ResolvedLayers.EditorOnly.LayerGuids[ResolvedLayerIndex] = ParentEditorOnly.LayerGuids[ParentLayerIndex]; // Still need to match guid to parent
 				}
 
 				check(!ParentLayerIndices.Contains(ParentLayerIndex));
@@ -14330,11 +14540,11 @@ bool FMaterialLayersFunctions::ResolveParent(const FMaterialLayersFunctions& Par
 		else if (LinkState == EMaterialLayerLinkState::LinkedToParent)
 		{
 			check(LayerGuid.IsValid());
-			ParentLayerIndex = Parent.LayerGuids.Find(LayerGuid);
+			ParentLayerIndex = ParentEditorOnly.LayerGuids.Find(LayerGuid);
 			if (ParentLayerIndex != INDEX_NONE)
 			{
 				// Layer comes from parent
-				ResolvedLayers.AddLayerCopy(Parent, ParentLayerIndex, bLayerVisible, EMaterialLayerLinkState::LinkedToParent);
+				ResolvedLayers.AddLayerCopy(ParentRuntime, ParentEditorOnly, ParentLayerIndex, bLayerVisible, EMaterialLayerLinkState::LinkedToParent);
 				check(!ParentLayerIndices.Contains(ParentLayerIndex));
 				ParentLayerIndices.Add(ParentLayerIndex);
 			}
@@ -14349,12 +14559,12 @@ bool FMaterialLayersFunctions::ResolveParent(const FMaterialLayersFunctions& Par
 			if (LinkState == EMaterialLayerLinkState::UnlinkedFromParent)
 			{
 				// If we are unlinked from parent, track the layer index we were previously linked to
-				ParentLayerIndex = Parent.LayerGuids.Find(LayerGuid);
+				ParentLayerIndex = ParentEditorOnly.LayerGuids.Find(LayerGuid);
 			}
 			check(ParentLayerIndex == INDEX_NONE || !ParentLayerIndices.Contains(ParentLayerIndex));
 
 			// Update the link state, depending on if we can find this layer in the parent
-			ResolvedLayers.AddLayerCopy(*this, LayerIndex, bLayerVisible, (ParentLayerIndex == INDEX_NONE) ? EMaterialLayerLinkState::NotFromParent : EMaterialLayerLinkState::UnlinkedFromParent);
+			ResolvedLayers.AddLayerCopy(Runtime, EditorOnly, LayerIndex, bLayerVisible, (ParentLayerIndex == INDEX_NONE) ? EMaterialLayerLinkState::NotFromParent : EMaterialLayerLinkState::UnlinkedFromParent);
 			ParentLayerIndices.Add(ParentLayerIndex);
 		}
 	}
@@ -14362,7 +14572,7 @@ bool FMaterialLayersFunctions::ResolveParent(const FMaterialLayersFunctions& Par
 	check(ResolvedLayers.Layers.Num() == ParentLayerIndices.Num());
 
 	// See if parent has any added layers
-	for (int32 ParentLayerIndex = 1; ParentLayerIndex < Parent.Layers.Num(); ++ParentLayerIndex)
+	for (int32 ParentLayerIndex = 1; ParentLayerIndex < ParentRuntime.Layers.Num(); ++ParentLayerIndex)
 	{
 		if (ParentLayerIndices.Contains(ParentLayerIndex))
 		{
@@ -14370,11 +14580,11 @@ bool FMaterialLayersFunctions::ResolveParent(const FMaterialLayersFunctions& Par
 			continue;
 		}
 
-		const FGuid& ParentLayerGuid = Parent.LayerGuids[ParentLayerIndex];
-		if (DeletedParentLayerGuids.Contains(ParentLayerGuid))
+		const FGuid& ParentLayerGuid = ParentEditorOnly.LayerGuids[ParentLayerIndex];
+		if (EditorOnly.DeletedParentLayerGuids.Contains(ParentLayerGuid))
 		{
 			// Parent layer was previously explicitly overiden/deleted
-			ResolvedLayers.DeletedParentLayerGuids.Add(ParentLayerGuid);
+			ResolvedLayers.EditorOnly.DeletedParentLayerGuids.Add(ParentLayerGuid);
 			continue;
 		}
 
@@ -14385,7 +14595,7 @@ bool FMaterialLayersFunctions::ResolveParent(const FMaterialLayersFunctions& Par
 			// or if this layer was explicitly deleted from child (and should therefore remain deleted).
 			// In order to avoid needlessly changing legacy materials, we assume the layer was explicitly deleted, so we keep it deleted here
 			// If desired, the relink functionality in the editor should allow it to be brought back
-			ResolvedLayers.DeletedParentLayerGuids.Add(ParentLayerGuid);
+			ResolvedLayers.EditorOnly.DeletedParentLayerGuids.Add(ParentLayerGuid);
 			continue;
 		}
 
@@ -14408,17 +14618,17 @@ bool FMaterialLayersFunctions::ResolveParent(const FMaterialLayersFunctions& Par
 		}
 
 		ParentLayerIndices.Insert(ParentLayerIndex, InsertLayerIndex + 1);
-		ResolvedLayers.InsertLayerCopy(Parent, ParentLayerIndex, EMaterialLayerLinkState::LinkedToParent, InsertLayerIndex + 1);
+		ResolvedLayers.InsertLayerCopy(ParentRuntime, ParentEditorOnly, ParentLayerIndex, EMaterialLayerLinkState::LinkedToParent, InsertLayerIndex + 1);
 	}
 
-	bool bUpdatedLayerIndices = Layers.Num() != ResolvedLayers.Layers.Num() ||
-		DeletedParentLayerGuids.Num() != ResolvedLayers.DeletedParentLayerGuids.Num();
+	bool bUpdatedLayerIndices = Runtime.Layers.Num() != ResolvedLayers.Layers.Num() ||
+		EditorOnly.DeletedParentLayerGuids.Num() != ResolvedLayers.EditorOnly.DeletedParentLayerGuids.Num();
 
-	OutRemapLayerIndices.SetNumUninitialized(Layers.Num());
-	for (int32 PrevLayerIndex = 0; PrevLayerIndex < Layers.Num(); ++PrevLayerIndex)
+	OutRemapLayerIndices.SetNumUninitialized(Runtime.Layers.Num());
+	for (int32 PrevLayerIndex = 0; PrevLayerIndex < Runtime.Layers.Num(); ++PrevLayerIndex)
 	{
-		const FGuid& LayerGuid = LayerGuids[PrevLayerIndex];
-		const int32 ResolvedLayerIndex = ResolvedLayers.LayerGuids.Find(LayerGuid);
+		const FGuid& LayerGuid = EditorOnly.LayerGuids[PrevLayerIndex];
+		const int32 ResolvedLayerIndex = ResolvedLayers.EditorOnly.LayerGuids.Find(LayerGuid);
 		OutRemapLayerIndices[PrevLayerIndex] = ResolvedLayerIndex;
 
 		if (PrevLayerIndex != ResolvedLayerIndex)
@@ -14427,10 +14637,24 @@ bool FMaterialLayersFunctions::ResolveParent(const FMaterialLayersFunctions& Par
 		}
 	}
 
-	*this = MoveTemp(ResolvedLayers);
+	Runtime = MoveTemp(static_cast<FMaterialLayersFunctionsRuntimeData&>(ResolvedLayers));
+	EditorOnly = MoveTemp(ResolvedLayers.EditorOnly);
 
 	return bUpdatedLayerIndices;
 }
+
+void FMaterialLayersFunctions::Validate(const FMaterialLayersFunctionsRuntimeData& Runtime, const FMaterialLayersFunctionsEditorOnlyData& EditorOnly)
+{
+	if (Runtime.Layers.Num() > 0)
+	{
+		check(Runtime.Blends.Num() == Runtime.Layers.Num() - 1);
+		check(Runtime.Layers.Num() == EditorOnly.LayerStates.Num());
+		check(Runtime.Layers.Num() == EditorOnly.LayerNames.Num());
+		check(Runtime.Layers.Num() == EditorOnly.LayerGuids.Num());
+		check(Runtime.Layers.Num() == EditorOnly.LayerLinkStates.Num());
+	}
+}
+
 #endif // WITH_EDITOR
 
 ///////////////////////////////////////////////////////////////////////////////
@@ -14761,8 +14985,8 @@ void UMaterialExpressionMaterialFunctionCall::GetExpressionToolTip(TArray<FStrin
 {
 	if (MaterialFunction)
 	{
-		const FString* Description = MaterialFunction->GetDescription();
-		ConvertToMultilineToolTip(Description ? *Description : TEXT(""), 40, OutToolTip);
+		const FString& Description = MaterialFunction->GetDescription();
+		ConvertToMultilineToolTip(Description, 40, OutToolTip);
 	}
 }
 
@@ -14843,12 +15067,12 @@ bool UMaterialExpressionMaterialFunctionCall::SetMaterialFunctionEx(
 			}
 
 			// Fixup any references that the material or material inputs had to the function's outputs, maintaining links with the same output name
-			FixupReferencingExpressions(FunctionOutputs, OriginalOutputs, Material->Expressions, MaterialInputs, true);
+			FixupReferencingExpressions(FunctionOutputs, OriginalOutputs, Material->GetExpressions(), MaterialInputs, true);
 		}
 		else if (Function)
 		{
 			// Fixup any references that the material function had to the function's outputs, maintaining links with the same output name
-			FixupReferencingExpressions(FunctionOutputs, OriginalOutputs, Function->FunctionExpressions, MaterialInputs, true);
+			FixupReferencingExpressions(FunctionOutputs, OriginalOutputs, Function->GetExpressions(), MaterialInputs, true);
 		}
 	}
 
@@ -14914,12 +15138,12 @@ void UMaterialExpressionMaterialFunctionCall::UpdateFromFunctionResource(bool bR
 			}
 			
 			// Fixup any references that the material or material inputs had to the function's outputs
-			FixupReferencingExpressions(FunctionOutputs, OriginalOutputs, Material->Expressions, MaterialInputs, false);
+			FixupReferencingExpressions(FunctionOutputs, OriginalOutputs, Material->GetExpressions(), MaterialInputs, false);
 		}
 		else if (Function)
 		{
 			// Fixup any references that the material function had to the function's outputs
-			FixupReferencingExpressions(FunctionOutputs, OriginalOutputs, Function->FunctionExpressions, MaterialInputs, false);
+			FixupReferencingExpressions(FunctionOutputs, OriginalOutputs, Function->GetExpressions(), MaterialInputs, false);
 		}
 	}
 
@@ -15001,7 +15225,7 @@ static void FixupReferencingInputs(
 void UMaterialExpressionMaterialFunctionCall::FixupReferencingExpressions(
 	const TArray<FFunctionExpressionOutput>& NewOutputs,
 	const TArray<FFunctionExpressionOutput>& OriginalOutputs,
-	TArray<UMaterialExpression*>& Expressions, 
+	TConstArrayView<TObjectPtr<UMaterialExpression>> Expressions,
 	TArray<FExpressionInput*>& MaterialInputs,
 	bool bMatchByName)
 {
@@ -15150,9 +15374,9 @@ void UMaterialExpressionFunctionInput::PostEditChangeProperty(FPropertyChangedEv
 	{
 		if (Material)
 		{
-			for (int32 ExpressionIndex = 0; ExpressionIndex < Material->Expressions.Num(); ExpressionIndex++)
+			for (UMaterialExpression* Expression : Material->GetExpressions())
 			{
-				UMaterialExpressionFunctionInput* OtherFunctionInput = Cast<UMaterialExpressionFunctionInput>(Material->Expressions[ExpressionIndex]);
+				UMaterialExpressionFunctionInput* OtherFunctionInput = Cast<UMaterialExpressionFunctionInput>(Expression);
 				if (OtherFunctionInput && OtherFunctionInput != this && OtherFunctionInput->InputName == InputName)
 				{
 					FMessageDialog::Open( EAppMsgType::Ok, NSLOCTEXT("UnrealEd", "Error_InputNamesMustBeUnique", "Function input names must be unique"));
@@ -15352,9 +15576,9 @@ void UMaterialExpressionFunctionInput::ValidateName()
 			}
 
 			bResultNameIndexValid = true;
-			for (int32 ExpressionIndex = 0; ExpressionIndex < Material->Expressions.Num(); ExpressionIndex++)
+			for (UMaterialExpression* Expression : Material->GetExpressions())
 			{
-				UMaterialExpressionFunctionInput* OtherFunctionInput = Cast<UMaterialExpressionFunctionInput>(Material->Expressions[ExpressionIndex]);
+				UMaterialExpressionFunctionInput* OtherFunctionInput = Cast<UMaterialExpressionFunctionInput>(Expression);
 				if (OtherFunctionInput && OtherFunctionInput != this && OtherFunctionInput->InputName == PotentialInputName)
 				{
 					bResultNameIndexValid = false;
@@ -15491,9 +15715,9 @@ void UMaterialExpressionFunctionOutput::PostEditChangeProperty(FPropertyChangedE
 	{
 		if (Material)
 		{
-			for (int32 ExpressionIndex = 0; ExpressionIndex < Material->Expressions.Num(); ExpressionIndex++)
+			for (UMaterialExpression* Expression : Material->GetExpressions())
 			{
-				UMaterialExpressionFunctionOutput* OtherFunctionOutput = Cast<UMaterialExpressionFunctionOutput>(Material->Expressions[ExpressionIndex]);
+				UMaterialExpressionFunctionOutput* OtherFunctionOutput = Cast<UMaterialExpressionFunctionOutput>(Expression);
 				if (OtherFunctionOutput && OtherFunctionOutput != this && OtherFunctionOutput->OutputName == OutputName)
 				{
 					FMessageDialog::Open( EAppMsgType::Ok, NSLOCTEXT("UnrealEd", "Error_OutputNamesMustBeUnique", "Function output names must be unique"));
@@ -15559,9 +15783,9 @@ void UMaterialExpressionFunctionOutput::ValidateName()
 			}
 
 			bResultNameIndexValid = true;
-			for (int32 ExpressionIndex = 0; ExpressionIndex < Material->Expressions.Num(); ExpressionIndex++)
+			for (UMaterialExpression* Expression : Material->GetExpressions())
 			{
-				UMaterialExpressionFunctionOutput* OtherFunctionOutput = Cast<UMaterialExpressionFunctionOutput>(Material->Expressions[ExpressionIndex]);
+				UMaterialExpressionFunctionOutput* OtherFunctionOutput = Cast<UMaterialExpressionFunctionOutput>(Expression);
 				if (OtherFunctionOutput && OtherFunctionOutput != this && OtherFunctionOutput->OutputName == PotentialOutputName)
 				{
 					bResultNameIndexValid = false;
@@ -16529,30 +16753,17 @@ UMaterialExpressionNamedRerouteBase::UMaterialExpressionNamedRerouteBase(const F
 
 }
 
-UMaterialExpressionNamedRerouteDeclaration* UMaterialExpressionNamedRerouteBase::FindDeclarationInArray(const FGuid& VariableGuid, const TArray<UMaterialExpression*>& Expressions) const
-{
-	for (UMaterialExpression* Expression : Expressions)
-	{
-		auto* Declaration = Cast<UMaterialExpressionNamedRerouteDeclaration>(Expression);
-		if (Declaration && this != Declaration && Declaration->VariableGuid == VariableGuid)
-		{
-			return Declaration;
-		}
-	}
-	return nullptr;
-}
-
 UMaterialExpressionNamedRerouteDeclaration* UMaterialExpressionNamedRerouteBase::FindDeclarationInMaterial(const FGuid& VariableGuid) const
 {
 	UMaterialExpressionNamedRerouteDeclaration* Declaration = nullptr;
 #if WITH_EDITORONLY_DATA
 	if (Material)
 	{
-		Declaration = FindDeclarationInArray(VariableGuid, Material->Expressions);
+		Declaration = FindDeclarationInArray(VariableGuid, Material->GetExpressions());
 	}
 	else if (Function) // Material should always be valid, but just in case also check Function
 	{
-		Declaration = FindDeclarationInArray(VariableGuid, Function->FunctionExpressions);
+		Declaration = FindDeclarationInArray(VariableGuid, Function->GetExpressions());
 	}
 #endif // WITH_EDITORONLY_DATA
 	return Declaration;
@@ -16663,7 +16874,7 @@ void UMaterialExpressionNamedRerouteDeclaration::SetEditableName(const FString&
 	// Refresh usage names
 	if (Material || Function)
 	{
-		auto& Expressions = Material ? Material->Expressions : Function->FunctionExpressions;
+		auto Expressions = Material ? Material->GetExpressions() : Function->GetExpressions();
 		for (UMaterialExpression* Expression : Expressions)
 		{
 			auto* Usage = Cast<UMaterialExpressionNamedRerouteUsage>(Expression);
@@ -16734,7 +16945,7 @@ void UMaterialExpressionNamedRerouteDeclaration::MakeNameUnique()
 #if WITH_EDITORONLY_DATA
 	if (Material || Function)
 	{
-		auto& Expressions = Material ? Material->Expressions : Function->FunctionExpressions;
+		auto Expressions = Material ? Material->GetExpressions() : Function->GetExpressions();
 
 		int32 NameIndex = 1;
 		bool bResultNameIndexValid = true;
@@ -19704,6 +19915,7 @@ bool UMaterialExpressionCurveAtlasRowParameter::GetParameterValue(FMaterialParam
 	OutMeta.ExpressionGuid = ExpressionGUID;
 	OutMeta.Group = Group;
 	OutMeta.SortPriority = SortPriority;
+	OutMeta.AssetPath = GetAssetPathName();
 	return true;
 }
 
@@ -23099,7 +23311,7 @@ UMaterialExpressionExecEnd::UMaterialExpressionExecEnd(const FObjectInitializer&
 #if WITH_EDITOR
 int32 UMaterialExpressionExecEnd::Compile(class FMaterialCompiler* Compiler, int32 OutputIndex)
 {
-	const int32 Attributes = Material->MaterialAttributes.Compile(Compiler);
+	const int32 Attributes = Material->GetEditorOnlyData()->MaterialAttributes.Compile(Compiler);
 	return Compiler->ReturnMaterialAttributes(Attributes);
 }
 #endif
diff --git a/Engine/Source/Runtime/Engine/Private/Materials/MaterialHLSLGenerator.cpp b/Engine/Source/Runtime/Engine/Private/Materials/MaterialHLSLGenerator.cpp
--- a/Engine/Source/Runtime/Engine/Private/Materials/MaterialHLSLGenerator.cpp
+++ b/Engine/Source/Runtime/Engine/Private/Materials/MaterialHLSLGenerator.cpp
@@ -78,7 +78,7 @@ bool FMaterialHLSLGenerator::InternalGenerate()
 	bool bResult = false;
 	if (TargetMaterial->IsUsingControlFlow())
 	{
-		UMaterialExpression* BaseExpression = TargetMaterial->ExpressionExecBegin;
+		UMaterialExpression* BaseExpression = TargetMaterial->GetExpressionExecBegin();
 		if (!BaseExpression)
 		{
 			bResult = Error(TEXT("Missing ExpressionExecBegin"));
@@ -671,9 +671,9 @@ const UE::HLSLTree::FExpression* FMaterialHLSLGenerator::GenerateFunctionCall(UE
 
 		if (HLSLFunction)
 		{
-			UMaterialFunction* BaseMaterialFunction = Cast<UMaterialFunction>(MaterialFunction->GetBaseFunction());
+			UMaterialFunction* BaseMaterialFunction = MaterialFunction->GetBaseFunction();
 			FunctionCallStack.Add(FunctionCall);
-			GenerateStatements(HLSLFunction->GetRootScope(), BaseMaterialFunction->ExpressionExecBegin);
+			GenerateStatements(HLSLFunction->GetRootScope(), BaseMaterialFunction->GetExpressionExecBegin());
 			verify(FunctionCallStack.Pop() == FunctionCall);
 			check(FunctionCall->bGeneratedResult);
 		}
diff --git a/Engine/Source/Runtime/Engine/Private/Materials/MaterialHLSLTree.cpp b/Engine/Source/Runtime/Engine/Private/Materials/MaterialHLSLTree.cpp
--- a/Engine/Source/Runtime/Engine/Private/Materials/MaterialHLSLTree.cpp
+++ b/Engine/Source/Runtime/Engine/Private/Materials/MaterialHLSLTree.cpp
@@ -454,7 +454,7 @@ bool FExpressionParameter::PrepareValue(FEmitContext& Context, FEmitScope& Scope
 	{
 		if (!ParameterInfo.Name.IsNone())
 		{
-			EmitData.CachedExpressionData->Parameters.AddParameter(ParameterInfo, ParameterMeta);
+			EmitData.CachedExpressionData->AddParameter(ParameterInfo, ParameterMeta);
 		}
 
 		UObject* ReferencedTexture = ParameterMeta.Value.AsTextureObject();
@@ -507,7 +507,7 @@ void FExpressionParameter::EmitValuePreshader(FEmitContext& Context, FEmitScope&
 			switch (ParameterType)
 			{
 			case EMaterialParameterType::StaticSwitch:
-				for (const FStaticSwitchParameter& Parameter : EmitMaterialData.StaticParameters->StaticSwitchParameters)
+				for (const FStaticSwitchParameter& Parameter : EmitMaterialData.StaticParameters->EditorOnly.StaticSwitchParameters)
 				{
 					if (Parameter.ParameterInfo == ParameterInfo)
 					{
@@ -517,7 +517,7 @@ void FExpressionParameter::EmitValuePreshader(FEmitContext& Context, FEmitScope&
 				}
 				break;
 			case EMaterialParameterType::StaticComponentMask:
-				for (const FStaticComponentMaskParameter& Parameter : EmitMaterialData.StaticParameters->StaticComponentMaskParameters)
+				for (const FStaticComponentMaskParameter& Parameter : EmitMaterialData.StaticParameters->EditorOnly.StaticComponentMaskParameters)
 				{
 					if (Parameter.ParameterInfo == ParameterInfo)
 					{
@@ -1119,7 +1119,9 @@ bool FExpressionMaterialLayers::PrepareValue(FEmitContext& Context, FEmitScope&
 			{
 				// TODO(?) - Layers for MIs are currently duplicated here and in FStaticParameterSet
 				EmitMaterialData.CachedExpressionData->bHasMaterialLayers = true;
-				EmitMaterialData.CachedExpressionData->MaterialLayers = *MaterialLayers;
+				EmitMaterialData.CachedExpressionData->MaterialLayers = MaterialLayers.GetRuntime();
+				EmitMaterialData.CachedExpressionData->EditorOnlyData->MaterialLayers = MaterialLayers.EditorOnly;
+				FMaterialLayersFunctions::Validate(EmitMaterialData.CachedExpressionData->MaterialLayers, EmitMaterialData.CachedExpressionData->EditorOnlyData->MaterialLayers);
 			}
 		}
 	}
