Add 'Optional' EditorOnly data for UMaterialInterface and UMaterialFunctionInterface

diff --git a/Engine/Source/Runtime/Engine/Private/Materials/MaterialInstance.cpp b/Engine/Source/Runtime/Engine/Private/Materials/MaterialInstance.cpp
--- a/Engine/Source/Runtime/Engine/Private/Materials/MaterialInstance.cpp
+++ b/Engine/Source/Runtime/Engine/Private/Materials/MaterialInstance.cpp
@@ -482,27 +482,35 @@ void UMaterialInstance::SwapLayerParameterIndices(int32 OriginalIndex, int32 New
 {
 	if (OriginalIndex != NewIndex)
 	{
+		UMaterialInstanceEditorOnlyData* EditorOnly = GetEditorOnlyData();
 		SwapLayerParameterIndicesArray(ScalarParameterValues, OriginalIndex, NewIndex);
 		SwapLayerParameterIndicesArray(VectorParameterValues, OriginalIndex, NewIndex);
 		SwapLayerParameterIndicesArray(DoubleVectorParameterValues, OriginalIndex, NewIndex);
 		SwapLayerParameterIndicesArray(TextureParameterValues, OriginalIndex, NewIndex);
 		SwapLayerParameterIndicesArray(RuntimeVirtualTextureParameterValues, OriginalIndex, NewIndex);
 		SwapLayerParameterIndicesArray(FontParameterValues, OriginalIndex, NewIndex);
-		SwapLayerParameterIndicesArray(StaticParameters.StaticSwitchParameters, OriginalIndex, NewIndex);
-		SwapLayerParameterIndicesArray(StaticParameters.StaticComponentMaskParameters, OriginalIndex, NewIndex);
+		if (EditorOnly)
+		{
+			SwapLayerParameterIndicesArray(EditorOnly->StaticParameters.StaticSwitchParameters, OriginalIndex, NewIndex);
+			SwapLayerParameterIndicesArray(EditorOnly->StaticParameters.StaticComponentMaskParameters, OriginalIndex, NewIndex);
+		}
 	}
 }
 
 void UMaterialInstance::RemoveLayerParameterIndex(int32 Index)
 {
+	UMaterialInstanceEditorOnlyData* EditorOnly = GetEditorOnlyData();
 	RemoveLayerParameterIndicesArray(ScalarParameterValues, Index);
 	RemoveLayerParameterIndicesArray(VectorParameterValues, Index);
 	RemoveLayerParameterIndicesArray(DoubleVectorParameterValues, Index);
 	RemoveLayerParameterIndicesArray(TextureParameterValues, Index);
 	RemoveLayerParameterIndicesArray(RuntimeVirtualTextureParameterValues, Index);
 	RemoveLayerParameterIndicesArray(FontParameterValues, Index);
-	RemoveLayerParameterIndicesArray(StaticParameters.StaticSwitchParameters, Index);
-	RemoveLayerParameterIndicesArray(StaticParameters.StaticComponentMaskParameters, Index);
+	if (EditorOnly)
+	{
+		RemoveLayerParameterIndicesArray(EditorOnly->StaticParameters.StaticSwitchParameters, Index);
+		RemoveLayerParameterIndicesArray(EditorOnly->StaticParameters.StaticComponentMaskParameters, Index);
+	}
 }
 #endif // WITH_EDITOR
 
@@ -511,7 +519,8 @@ bool UMaterialInstance::UpdateParameters()
 	bool bDirty = false;
 
 #if WITH_EDITOR
-	if(IsTemplate(RF_ClassDefaultObject)==false)
+	UMaterialInstanceEditorOnlyData* EditorOnly = GetEditorOnlyData();
+	if(IsTemplate(RF_ClassDefaultObject)==false && EditorOnly)
 	{
 		// Get a pointer to the parent material.
 		UMaterial* ParentMaterial = NULL;
@@ -550,27 +559,37 @@ bool UMaterialInstance::UpdateParameters()
 			bDirty = UpdateParameterSet<FFontParameterValue, UMaterialExpressionFontSampleParameter>(FontParameterValues, ParentMaterial) || bDirty;
 
 			// Static switch parameters
-			bDirty = UpdateParameterSet<FStaticSwitchParameter, UMaterialExpressionStaticBoolParameter>(StaticParameters.StaticSwitchParameters, ParentMaterial) || bDirty;
+			bDirty = UpdateParameterSet<FStaticSwitchParameter, UMaterialExpressionStaticBoolParameter>(EditorOnly->StaticParameters.StaticSwitchParameters, ParentMaterial) || bDirty;
 
 			// Static component mask parameters
-			bDirty = UpdateParameterSet<FStaticComponentMaskParameter, UMaterialExpressionStaticComponentMaskParameter>(StaticParameters.StaticComponentMaskParameters, ParentMaterial) || bDirty;
+			bDirty = UpdateParameterSet<FStaticComponentMaskParameter, UMaterialExpressionStaticComponentMaskParameter>(EditorOnly->StaticParameters.StaticComponentMaskParameters, ParentMaterial) || bDirty;
 
 			// Custom parameters
 			PRAGMA_DISABLE_DEPRECATION_WARNINGS
-			for (const auto& CustomParameterSetUpdater : CustomParameterSetUpdaters)
+			if (CustomParameterSetUpdaters.Num() > 0)
 			{
-				bDirty |= CustomParameterSetUpdater.Execute(StaticParameters, ParentMaterial);
+				FStaticParameterSet LocalStaticParameters = GetStaticParameters();
+				for (const auto& CustomParameterSetUpdater : CustomParameterSetUpdaters)
+				{
+					bDirty |= CustomParameterSetUpdater.Execute(LocalStaticParameters, ParentMaterial);
+				}
+				StaticParametersRuntime = LocalStaticParameters.GetRuntime();
+				EditorOnly->StaticParameters = LocalStaticParameters.EditorOnly;
 			}
 			PRAGMA_ENABLE_DEPRECATION_WARNINGS
 		}
 
-		if (StaticParameters.bHasMaterialLayers && Parent)
+		if (StaticParametersRuntime.bHasMaterialLayers && Parent)
 		{
 			FMaterialLayersFunctions ParentLayers;
 			if (Parent->GetMaterialLayers(ParentLayers))
 			{
 				TArray<int32> RemapLayerIndices;
-				if (StaticParameters.MaterialLayers.ResolveParent(ParentLayers, RemapLayerIndices))
+				if (FMaterialLayersFunctions::ResolveParent(ParentLayers,
+					ParentLayers.EditorOnly,
+					StaticParametersRuntime.MaterialLayers,
+					EditorOnly->StaticParameters.MaterialLayers,
+					RemapLayerIndices))
 				{
 					RemapLayerParameterIndicesArray(ScalarParameterValues, RemapLayerIndices);
 					RemapLayerParameterIndicesArray(VectorParameterValues, RemapLayerIndices);
@@ -578,8 +597,8 @@ bool UMaterialInstance::UpdateParameters()
 					RemapLayerParameterIndicesArray(TextureParameterValues, RemapLayerIndices);
 					RemapLayerParameterIndicesArray(RuntimeVirtualTextureParameterValues, RemapLayerIndices);
 					RemapLayerParameterIndicesArray(FontParameterValues, RemapLayerIndices);
-					RemapLayerParameterIndicesArray(StaticParameters.StaticSwitchParameters, RemapLayerIndices);
-					RemapLayerParameterIndicesArray(StaticParameters.StaticComponentMaskParameters, RemapLayerIndices);
+					RemapLayerParameterIndicesArray(EditorOnly->StaticParameters.StaticSwitchParameters, RemapLayerIndices);
+					RemapLayerParameterIndicesArray(EditorOnly->StaticParameters.StaticComponentMaskParameters, RemapLayerIndices);
 					bDirty = true;
 				}
 			}
@@ -855,8 +874,12 @@ bool UMaterialInstance::GetParameterOverrideValue(EMaterialParameterType Type, c
 	case EMaterialParameterType::RuntimeVirtualTexture: bResult = GameThread_GetParameterValue(RuntimeVirtualTextureParameterValues, ParameterInfo, OutResult); break;
 	case EMaterialParameterType::Font: bResult = GameThread_GetParameterValue(FontParameterValues, ParameterInfo, OutResult); break;
 #if WITH_EDITORONLY_DATA
-	case EMaterialParameterType::StaticSwitch: bResult = GameThread_GetParameterValue(StaticParameters.StaticSwitchParameters, ParameterInfo, OutResult); break;
-	case EMaterialParameterType::StaticComponentMask: bResult = GameThread_GetParameterValue(StaticParameters.StaticComponentMaskParameters, ParameterInfo, OutResult); break;
+	case EMaterialParameterType::StaticSwitch:
+		bResult = GameThread_GetParameterValue(GetEditorOnlyData()->StaticParameters.StaticSwitchParameters, ParameterInfo, OutResult);
+		break;
+	case EMaterialParameterType::StaticComponentMask:
+		bResult = GameThread_GetParameterValue(GetEditorOnlyData()->StaticParameters.StaticComponentMaskParameters, ParameterInfo, OutResult);
+		break;
 #endif // WITH_EDITORONLY_DATA
 	default: checkNoEntry(); break;
 	}
@@ -871,7 +894,7 @@ bool UMaterialInstance::GetParameterValue(EMaterialParameterType Type, const FMe
 	bool bResult = false;
 	if (EnumHasAnyFlags(Flags, EMaterialGetParameterValueFlags::CheckNonOverrides))
 	{
-		bResult = InstanceChain.GetCachedExpressionData().Parameters.GetParameterValue(Type, ParameterInfo, OutResult);
+		bResult = InstanceChain.GetCachedExpressionData().GetParameterValue(Type, ParameterInfo, OutResult);
 	}
 
 	const bool bCheckInstanceOverrides = EnumHasAnyFlags(Flags, EMaterialGetParameterValueFlags::CheckInstanceOverrides);
@@ -1681,7 +1704,10 @@ void UMaterialInstance::SetStaticSwitchParameterValueEditorOnly(const FMaterialP
 {
 	check(GIsEditor || IsRunningCommandlet());
 
-	for (FStaticSwitchParameter& StaticSwitches : StaticParameters.StaticSwitchParameters)
+	UMaterialInstanceEditorOnlyData* EditorOnly = GetEditorOnlyData();
+	check(EditorOnly);
+
+	for (FStaticSwitchParameter& StaticSwitches : EditorOnly->StaticParameters.StaticSwitchParameters)
 	{
 		if (StaticSwitches.ParameterInfo == ParameterInfo)
 		{
@@ -1691,7 +1717,7 @@ void UMaterialInstance::SetStaticSwitchParameterValueEditorOnly(const FMaterialP
 		}
 	}
 
-	new(StaticParameters.StaticSwitchParameters) FStaticSwitchParameter(ParameterInfo, Value, true, FGuid());
+	new(EditorOnly->StaticParameters.StaticSwitchParameters) FStaticSwitchParameter(ParameterInfo, Value, true, FGuid());
 }
 
 void UMaterialInstance::GetStaticParameterValues(FStaticParameterSet& OutStaticParameters)
@@ -1719,9 +1745,21 @@ void UMaterialInstance::GetStaticParameterValues(FStaticParameterSet& OutStaticP
 		}
 	}
 
-	OutStaticParameters.TerrainLayerWeightParameters = StaticParameters.TerrainLayerWeightParameters;
+	UMaterialInstanceEditorOnlyData* EditorOnly = GetEditorOnlyData();
+	if (EditorOnly)
+	{
+		OutStaticParameters.EditorOnly.TerrainLayerWeightParameters = EditorOnly->StaticParameters.TerrainLayerWeightParameters;
+	}
 
-	OutStaticParameters.bHasMaterialLayers = GetMaterialLayers(OutStaticParameters.MaterialLayers);
+	FMaterialLayersFunctions MaterialLayers;
+	OutStaticParameters.bHasMaterialLayers = GetMaterialLayers(MaterialLayers);
+	if (OutStaticParameters.bHasMaterialLayers)
+	{
+		OutStaticParameters.MaterialLayers = MoveTemp(MaterialLayers.GetRuntime());
+		OutStaticParameters.EditorOnly.MaterialLayers = MoveTemp(MaterialLayers.EditorOnly);
+	}
+
+	OutStaticParameters.Validate();
 
 	// Custom parameters.
 	PRAGMA_DISABLE_DEPRECATION_WARNINGS
@@ -1759,7 +1797,7 @@ void UMaterialInstance::GetAllParametersOfType(EMaterialParameterType Type, TMap
 	GetMaterialInheritanceChain(InstanceChain);
 
 	OutParameters.Reset();
-	InstanceChain.GetCachedExpressionData().Parameters.GetAllParametersOfType(Type, OutParameters);
+	InstanceChain.GetCachedExpressionData().GetAllParametersOfType(Type, OutParameters);
 
 	TArray<int32, TInlineAllocator<16>> LayerIndexRemap;
 	LayerIndexRemap.Empty(GetCachedInstanceData().ParentLayerIndexRemap.Num());
@@ -1775,7 +1813,6 @@ void UMaterialInstance::GetAllParametersOfType(EMaterialParameterType Type, TMap
 	for (int32 Index = 0; Index < InstanceChain.MaterialInstances.Num(); ++Index)
 	{
 		const UMaterialInstance* Instance = InstanceChain.MaterialInstances[Index];
-
 		const bool bSetOverride = (Index == 0); // Only set the override flag for parameters overriden by the current material (always at slot0)
 		switch (Type)
 		{
@@ -1786,8 +1823,12 @@ void UMaterialInstance::GetAllParametersOfType(EMaterialParameterType Type, TMap
 		case EMaterialParameterType::RuntimeVirtualTexture: GameThread_ApplyParameterOverrides(Instance->RuntimeVirtualTextureParameterValues, LayerIndexRemap, bSetOverride, OverridenParameters, OutParameters); break;
 		case EMaterialParameterType::Font: GameThread_ApplyParameterOverrides(Instance->FontParameterValues, LayerIndexRemap, bSetOverride, OverridenParameters, OutParameters); break;
 #if WITH_EDITORONLY_DATA
-		case EMaterialParameterType::StaticSwitch: GameThread_ApplyParameterOverrides(Instance->StaticParameters.StaticSwitchParameters, LayerIndexRemap, bSetOverride, OverridenParameters, OutParameters); break;
-		case EMaterialParameterType::StaticComponentMask: GameThread_ApplyParameterOverrides(Instance->StaticParameters.StaticComponentMaskParameters, LayerIndexRemap, bSetOverride, OverridenParameters, OutParameters); break;
+		case EMaterialParameterType::StaticSwitch:
+			GameThread_ApplyParameterOverrides(Instance->GetEditorOnlyData()->StaticParameters.StaticSwitchParameters, LayerIndexRemap, bSetOverride, OverridenParameters, OutParameters);
+			break;
+		case EMaterialParameterType::StaticComponentMask:
+			GameThread_ApplyParameterOverrides(Instance->GetEditorOnlyData()->StaticParameters.StaticComponentMaskParameters, LayerIndexRemap, bSetOverride, OverridenParameters, OutParameters);
+			break;
 #endif // WITH_EDITORONLY_DATA
 		default: checkNoEntry();
 		}
@@ -1805,9 +1846,9 @@ bool UMaterialInstance::IterateDependentFunctions(TFunctionRef<bool(UMaterialFun
 {
 	// Important that local function references are listed first so that traversing for a parameter
 	// value we always hit the highest material in the hierarchy that can give us a valid value
-	if (StaticParameters.bHasMaterialLayers)
+	if (StaticParametersRuntime.bHasMaterialLayers)
 	{
-		for (UMaterialFunctionInterface* Layer : StaticParameters.MaterialLayers.Layers)
+		for (UMaterialFunctionInterface* Layer : StaticParametersRuntime.MaterialLayers.Layers)
 		{
 			if (Layer)
 			{
@@ -1822,7 +1863,7 @@ bool UMaterialInstance::IterateDependentFunctions(TFunctionRef<bool(UMaterialFun
 			}
 		}
 
-		for (UMaterialFunctionInterface* Blend : StaticParameters.MaterialLayers.Blends)
+		for (UMaterialFunctionInterface* Blend : StaticParametersRuntime.MaterialLayers.Blends)
 		{
 			if (Blend)
 			{
@@ -1863,8 +1904,10 @@ void UMaterialInstance::InitStaticPermutation(EMaterialShaderPrecompileMode Prec
 {
 	UpdateOverridableBaseProperties();
 
+#if WITH_EDITORONLY_DATA
 	// Update bHasStaticPermutationResource in case the parent was not found
-	bHasStaticPermutationResource = (!StaticParameters.IsEmpty() || HasOverridenBaseProperties()) && Parent;
+	bHasStaticPermutationResource = (HasStaticParameters() || HasOverridenBaseProperties()) && Parent;
+#endif // WITH_EDITORONLY_DATA
 
 	FMaterialResourceDeferredDeletionArray ResourcesToFree;
 
@@ -2186,9 +2229,17 @@ void UMaterialInstance::CacheGivenTypesForCooking(EShaderPlatform ShaderPlatform
 
 bool UMaterialInstance::GetMaterialLayers(FMaterialLayersFunctions& OutLayers, TMicRecursionGuard RecursionGuard) const
 {
-	if (StaticParameters.bHasMaterialLayers)
+	if (StaticParametersRuntime.bHasMaterialLayers)
 	{
-		OutLayers = StaticParameters.MaterialLayers;
+		OutLayers.GetRuntime() = StaticParametersRuntime.MaterialLayers;
+#if WITH_EDITORONLY_DATA
+		const UMaterialInstanceEditorOnlyData* EditorOnly = GetEditorOnlyData();
+		if (EditorOnly)
+		{
+			OutLayers.EditorOnly = EditorOnly->StaticParameters.MaterialLayers;
+		}
+		OutLayers.Validate();
+#endif // WITH_EDITORONLY_DATA
 		return true;
 	}
 
@@ -2239,8 +2290,13 @@ bool UMaterialInstance::IsComplete() const
 #if WITH_EDITOR
 bool UMaterialInstance::SetMaterialLayers(const FMaterialLayersFunctions& LayersValue)
 {
+	UMaterialInstanceEditorOnlyData* EditorOnly = GetEditorOnlyData();
+	check(EditorOnly);
+
 	bool bUpdatedLayers = false;
-	if (!StaticParameters.bHasMaterialLayers || StaticParameters.MaterialLayers != LayersValue)
+	if (!StaticParametersRuntime.bHasMaterialLayers ||
+		StaticParametersRuntime.MaterialLayers != LayersValue.GetRuntime() ||
+		EditorOnly->StaticParameters.MaterialLayers != LayersValue.EditorOnly)
 	{
 		bool bMatchesParentLayers = false;
 		if (Parent)
@@ -2254,16 +2310,19 @@ bool UMaterialInstance::SetMaterialLayers(const FMaterialLayersFunctions& Layers
 
 		if (bMatchesParentLayers)
 		{
-			bUpdatedLayers = StaticParameters.bHasMaterialLayers; // if we previously had layers, but are now clearing them to match parent
-			StaticParameters.bHasMaterialLayers = false;
-			StaticParameters.MaterialLayers.Empty();
+			bUpdatedLayers = StaticParametersRuntime.bHasMaterialLayers; // if we previously had layers, but are now clearing them to match parent
+			StaticParametersRuntime.bHasMaterialLayers = false;
+			StaticParametersRuntime.MaterialLayers.Empty();
+			EditorOnly->StaticParameters.MaterialLayers.Empty();
 		}
 		else
 		{
 			bUpdatedLayers = true;
-			StaticParameters.bHasMaterialLayers = true;
-			StaticParameters.MaterialLayers = LayersValue;
+			StaticParametersRuntime.bHasMaterialLayers = true;
+			StaticParametersRuntime.MaterialLayers = LayersValue.GetRuntime();
+			EditorOnly->StaticParameters.MaterialLayers = LayersValue.EditorOnly;
 		}
+		FStaticParameterSet::Validate(StaticParametersRuntime, EditorOnly->StaticParameters);
 	}
 	return bUpdatedLayers;
 }
@@ -2366,8 +2425,9 @@ void UMaterialInstance::Serialize(FArchive& Ar)
 #endif
 
 	Super::Serialize(Ar);
-		
+
 #if WITH_EDITOR
+	UMaterialInstanceEditorOnlyData* EditorOnly = GetEditorOnlyData();
 	if (Ar.CustomVer(FRenderingObjectVersion::GUID) < FRenderingObjectVersion::MaterialAttributeLayerParameters)
 	{
 		// Material attribute layers parameter refactor fix-up
@@ -2391,13 +2451,20 @@ void UMaterialInstance::Serialize(FArchive& Ar)
 
 	if (Ar.CustomVer(FUE5ReleaseStreamObjectVersion::GUID) < FUE5ReleaseStreamObjectVersion::MaterialLayerStacksAreNotParameters)
 	{
-		StaticParameters.UpdateLegacyMaterialLayersData();
+		StaticParameters_DEPRECATED.UpdateLegacyMaterialLayersData();
 	}
 
 	if (Ar.CustomVer(FFortniteMainBranchObjectVersion::GUID) < FFortniteMainBranchObjectVersion::TerrainLayerWeightsAreNotParameters)
 	{
-		StaticParameters.UpdateLegacyTerrainLayerWeightData();
+		StaticParameters_DEPRECATED.UpdateLegacyTerrainLayerWeightData();
+	}
+
+	if (!StaticParameters_DEPRECATED.IsEmpty())
+	{
+		StaticParametersRuntime = MoveTemp(StaticParameters_DEPRECATED.GetRuntime());
+		EditorOnly->StaticParameters = MoveTemp(StaticParameters_DEPRECATED.EditorOnly);
 	}
+
 #endif // WITH_EDITOR
 
 	bool bSavedCachedData = false;
@@ -2451,7 +2518,9 @@ void UMaterialInstance::Serialize(FArchive& Ar)
 #if WITH_EDITOR
 			if (Ar.CustomVer(FRenderingObjectVersion::GUID) < FRenderingObjectVersion::MaterialAttributeLayerParameters)
 			{
-				StaticParameters.SerializeLegacy(Ar);
+				StaticParameters_DEPRECATED.SerializeLegacy(Ar);
+				StaticParametersRuntime = MoveTemp(StaticParameters_DEPRECATED.GetRuntime());
+				EditorOnly->StaticParameters = MoveTemp(StaticParameters_DEPRECATED.EditorOnly);
 			}
 
 			static_assert(!STORE_ONLY_ACTIVE_SHADERMAPS, "Only discard unused SMs in cooked build");
@@ -2478,12 +2547,12 @@ void UMaterialInstance::Serialize(FArchive& Ar)
 			FMaterialShaderMapId LegacyId;
 			LegacyId.Serialize(Ar, bLoadedByCookedMaterial);
 
-			StaticParameters.StaticSwitchParameters = LegacyId.GetStaticSwitchParameters();
-			StaticParameters.StaticComponentMaskParameters = LegacyId.GetStaticComponentMaskParameters();
-			StaticParameters.TerrainLayerWeightParameters = LegacyId.GetTerrainLayerWeightParameters();
+			EditorOnly->StaticParameters.StaticSwitchParameters = LegacyId.GetStaticSwitchParameters();
+			EditorOnly->StaticParameters.StaticComponentMaskParameters = LegacyId.GetStaticComponentMaskParameters();
+			EditorOnly->StaticParameters.TerrainLayerWeightParameters = LegacyId.GetTerrainLayerWeightParameters();
 
-			TrimToOverriddenOnly(StaticParameters.StaticSwitchParameters);
-			TrimToOverriddenOnly(StaticParameters.StaticComponentMaskParameters);
+			TrimToOverriddenOnly(EditorOnly->StaticParameters.StaticSwitchParameters);
+			TrimToOverriddenOnly(EditorOnly->StaticParameters.StaticComponentMaskParameters);
 		}
 #endif // WITH_EDITOR
 	}
@@ -2637,17 +2706,18 @@ void UMaterialInstance::PostLoad()
 		}
 	}
 
+#if WITH_EDITORONLY_DATA
 	// And any material layers parameter's functions
-	if (StaticParameters.bHasMaterialLayers)
+	if (StaticParametersRuntime.bHasMaterialLayers)
 	{
-		for (UMaterialFunctionInterface* Dependency : StaticParameters.MaterialLayers.Layers)
+		for (UMaterialFunctionInterface* Dependency : StaticParametersRuntime.MaterialLayers.Layers)
 		{
 			if (Dependency)
 			{
 				Dependency->ConditionalPostLoad();
 			}
 		}
-		for (UMaterialFunctionInterface* Dependency : StaticParameters.MaterialLayers.Blends)
+		for (UMaterialFunctionInterface* Dependency : StaticParametersRuntime.MaterialLayers.Blends)
 		{
 			if (Dependency)
 			{
@@ -2655,6 +2725,7 @@ void UMaterialInstance::PostLoad()
 			}
 		}
 	}
+#endif // WITH_EDITORONLY_DATA
 
 #if WITH_EDITOR
 	UpdateCachedData();
@@ -2931,7 +3002,9 @@ void FMaterialInstanceParameterUpdateContext::SetBasePropertyOverrides(const FMa
 void FMaterialInstanceParameterUpdateContext::SetMaterialLayers(const FMaterialLayersFunctions& InValue)
 {
 	StaticParameters.bHasMaterialLayers = true;
-	StaticParameters.MaterialLayers = InValue;
+	StaticParameters.MaterialLayers = InValue.GetRuntime();
+	StaticParameters.EditorOnly.MaterialLayers = InValue.EditorOnly;
+	StaticParameters.Validate();
 }
 #endif // WITH_EDITORONLY_DATA
 
@@ -3261,9 +3334,17 @@ void UMaterialInstance::ClearParameterValuesInternal(EMaterialInstanceClearParam
 
 	if (EnumHasAnyFlags(Flags, EMaterialInstanceClearParameterFlag::Static))
 	{
-		StaticParameters.Empty();
+		StaticParametersRuntime.Empty();
+#if WITH_EDITORONLY_DATA
+		UMaterialInstanceEditorOnlyData* EditorOnly = GetEditorOnlyData();
+		if (EditorOnly)
+		{
+			EditorOnly->StaticParameters.Empty();
+		}
+#endif // WITH_EDITORONLY_DATA
 	}
 
+
 	if (Resource && bUpdateResource)
 	{
 		FMaterialInstanceResource* InResource = Resource;
@@ -3284,10 +3365,11 @@ void UMaterialInstance::ClearParameterValuesInternal(EMaterialInstanceClearParam
 #if WITH_EDITOR
 void UMaterialInstance::UpdateStaticPermutation(const FStaticParameterSet& NewParameters, FMaterialInstanceBasePropertyOverrides& NewBasePropertyOverrides, const bool bForceStaticPermutationUpdate /*= false*/, FMaterialUpdateContext* MaterialUpdateContext)
 {
+	UMaterialInstanceEditorOnlyData* EditorOnly = GetEditorOnlyData();
 	FStaticParameterSet CompareParameters = NewParameters;
 
-	TrimToOverriddenOnly(CompareParameters.StaticSwitchParameters);
-	TrimToOverriddenOnly(CompareParameters.StaticComponentMaskParameters);
+	TrimToOverriddenOnly(CompareParameters.EditorOnly.StaticSwitchParameters);
+	TrimToOverriddenOnly(CompareParameters.EditorOnly.StaticComponentMaskParameters);
 
 	// Check to see if the material layers being assigned match values from the parent
 	if (CompareParameters.bHasMaterialLayers && Parent)
@@ -3295,7 +3377,11 @@ void UMaterialInstance::UpdateStaticPermutation(const FStaticParameterSet& NewPa
 		FMaterialLayersFunctions ParentLayers;
 		if (Parent->GetMaterialLayers(ParentLayers))
 		{
-			if (CompareParameters.MaterialLayers.MatchesParent(ParentLayers))
+			if (FMaterialLayersFunctions::MatchesParent(
+				CompareParameters.MaterialLayers,
+				CompareParameters.EditorOnly.MaterialLayers,
+				ParentLayers,
+				ParentLayers.EditorOnly))
 			{
 				CompareParameters.bHasMaterialLayers = false;
 				CompareParameters.MaterialLayers.Empty();
@@ -3303,7 +3389,8 @@ void UMaterialInstance::UpdateStaticPermutation(const FStaticParameterSet& NewPa
 		}
 	}
 
-	const bool bParamsHaveChanged = StaticParameters != CompareParameters;
+	const FStaticParameterSet CurrentParameters = GetStaticParameters();
+	const bool bParamsHaveChanged = CurrentParameters != CompareParameters;
 	const bool bBasePropertyOverridesHaveChanged = BasePropertyOverrides != NewBasePropertyOverrides;
 
 	BasePropertyOverrides = NewBasePropertyOverrides;
@@ -3321,7 +3408,8 @@ void UMaterialInstance::UpdateStaticPermutation(const FStaticParameterSet& NewPa
 		FlushRenderingCommands();
 
 		bHasStaticPermutationResource = bWantsStaticPermutationResource;
-		StaticParameters = CompareParameters;
+		StaticParametersRuntime = CompareParameters.GetRuntime();
+		EditorOnly->StaticParameters = CompareParameters.EditorOnly;
 
 		UpdateCachedData();
 		CacheResourceShadersForRendering(EMaterialShaderPrecompileMode::None);
@@ -3378,7 +3466,7 @@ void UMaterialInstance::UpdateStaticPermutation(const FStaticParameterSet& NewPa
 void UMaterialInstance::UpdateStaticPermutation(FMaterialUpdateContext* MaterialUpdateContext)
 {
 	// Force the update, since we aren't technically changing anything
-	UpdateStaticPermutation(StaticParameters, BasePropertyOverrides, true, MaterialUpdateContext);
+	UpdateStaticPermutation(GetStaticParameters(), BasePropertyOverrides, true, MaterialUpdateContext);
 }
 
 void UMaterialInstance::UpdateParameterNames()
@@ -3416,7 +3504,7 @@ void UMaterialInstance::PostEditChangeProperty(FPropertyChangedEvent& PropertyCh
 
 	// Force UpdateStaticPermutation when change type is Redirected as this probably means a Material or MaterialInstance parent asset was deleted.
 	const bool bForceStaticPermutationUpdate = PropertyChangedEvent.ChangeType == EPropertyChangeType::Redirected;
-	UpdateStaticPermutation(StaticParameters, BasePropertyOverrides, bForceStaticPermutationUpdate);
+	UpdateStaticPermutation(GetStaticParameters(), BasePropertyOverrides, bForceStaticPermutationUpdate);
 
 	if (PropertyChangedEvent.ChangeType == EPropertyChangeType::ValueSet || PropertyChangedEvent.ChangeType == EPropertyChangeType::ArrayClear || PropertyChangedEvent.ChangeType == EPropertyChangeType::ArrayRemove || PropertyChangedEvent.ChangeType == EPropertyChangeType::Unspecified || PropertyChangedEvent.ChangeType == EPropertyChangeType::Duplicate)
 	{
@@ -3872,17 +3960,48 @@ bool UMaterialInstance::IsPropertyActive(EMaterialProperty InProperty) const
 	return Material ? Material->IsPropertyActiveInDerived(InProperty, this) : false;
 }
 
+bool UMaterialInstance::HasStaticParameters() const
+{
+	if (!StaticParametersRuntime.IsEmpty())
+	{
+		return true;
+	}
+#if WITH_EDITOR
+	const UMaterialInstanceEditorOnlyData* EditorOnly = GetEditorOnlyData();
+	if (EditorOnly && !EditorOnly->StaticParameters.IsEmpty())
+	{
+		return true;
+	}
+#endif // WITH_EDITOR
+	return false;
+}
+
+FStaticParameterSet UMaterialInstance::GetStaticParameters() const
+{
+	FStaticParameterSet Result;
+	Result.GetRuntime() = StaticParametersRuntime;
+#if WITH_EDITORONLY_DATA
+	const UMaterialInstanceEditorOnlyData* EditorOnly = GetEditorOnlyData();
+	if (EditorOnly)
+	{
+		Result.EditorOnly = EditorOnly->StaticParameters;
+	}
+#endif // WITH_EDITORONLY_DATA
+	return Result;
+}
+
 #if WITH_EDITOR
 int32 UMaterialInstance::CompilePropertyEx( class FMaterialCompiler* Compiler, const FGuid& AttributeID )
 {
 	return Parent ? Parent->CompilePropertyEx(Compiler, AttributeID) : INDEX_NONE;
 }
-#endif // WITH_EDITOR
 
-const FStaticParameterSet& UMaterialInstance::GetStaticParameters() const
+const FStaticParameterSetEditorOnlyData& UMaterialInstance::GetEditorOnlyStaticParameters() const
 {
-	return StaticParameters;
+	const UMaterialInstanceEditorOnlyData* EditorOnly = GetEditorOnlyData();
+	return EditorOnly->StaticParameters;
 }
+#endif // WITH_EDITOR
 
 void UMaterialInstance::GetLightingGuidChain(bool bIncludeTextures, TArray<FGuid>& OutGuids) const
 {
@@ -3958,10 +4077,12 @@ bool UMaterialInstance::Equivalent(const UMaterialInstance* CompareTo) const
 		return false;
 	}
 
-	if (!StaticParameters.Equivalent(CompareTo->StaticParameters))
+	const FStaticParameterSet LocalStaticParameters = GetStaticParameters();
+	if (!LocalStaticParameters.Equivalent(CompareTo->GetStaticParameters()))
 	{
 		return false;
 	}
+
 	return true;
 }
 
@@ -4243,58 +4364,6 @@ void UMaterialInstance::CopyMaterialUniformParametersInternal(UMaterialInterface
 #endif
 }
 
-#if WITH_EDITOR
-
-void FindCollectionExpressionRecursive(TArray<FGuid>& OutGuidList, const TArray<TObjectPtr<UMaterialExpression>>& InMaterialExpression)
-{
-	for (int32 ExpressionIndex = 0; ExpressionIndex < InMaterialExpression.Num(); ExpressionIndex++)
-	{
-		if (const UMaterialExpressionCollectionParameter* CollectionPtr = Cast<UMaterialExpressionCollectionParameter>(InMaterialExpression[ExpressionIndex]))
-		{
-			if (CollectionPtr->Collection)
-			{
-				OutGuidList.Add(CollectionPtr->Collection->StateId);
-			}
-			return;
-		}
-		else if (const UMaterialExpressionMaterialFunctionCall* MaterialFunctionCall = Cast<UMaterialExpressionMaterialFunctionCall>(InMaterialExpression[ExpressionIndex]))
-		{
-			if (const TArray<TObjectPtr<UMaterialExpression>>* FunctionExpressions = MaterialFunctionCall->MaterialFunction ? MaterialFunctionCall->MaterialFunction->GetFunctionExpressions() : nullptr)
-			{
-				FindCollectionExpressionRecursive(OutGuidList, *FunctionExpressions);
-			}
-		}
-		else if (const UMaterialExpressionMaterialAttributeLayers* MaterialLayers = Cast<UMaterialExpressionMaterialAttributeLayers>(InMaterialExpression[ExpressionIndex]))
-		{
-			const TArray<UMaterialFunctionInterface*>& Layers = MaterialLayers->GetLayers();
-			const TArray<UMaterialFunctionInterface*>& Blends = MaterialLayers->GetBlends();
-
-			for (const UMaterialFunctionInterface* Layer : Layers)
-			{
-				if (Layer)
-				{
-					if (const TArray<TObjectPtr<UMaterialExpression>>* FunctionExpressions = Layer->GetFunctionExpressions())
-					{
-						FindCollectionExpressionRecursive(OutGuidList, *FunctionExpressions);
-					}
-				}
-			}
-
-			for (const UMaterialFunctionInterface* Blend : Blends)
-			{
-				if (Blend)
-				{
-					if (const TArray<TObjectPtr<UMaterialExpression>>* FunctionExpressions = Blend->GetFunctionExpressions())
-					{
-						FindCollectionExpressionRecursive(OutGuidList, *FunctionExpressions);
-					}
-				}
-			}
-		}
-	}
-}
-#endif // WITH_EDITOR
-
 #if WITH_EDITORONLY_DATA
 PRAGMA_DISABLE_DEPRECATION_WARNINGS
 UMaterialInstance::FCustomStaticParametersGetterDelegate UMaterialInstance::CustomStaticParametersGetters;
diff --git a/Engine/Source/Runtime/Engine/Private/Materials/MaterialInstanceConstant.cpp b/Engine/Source/Runtime/Engine/Private/Materials/MaterialInstanceConstant.cpp
--- a/Engine/Source/Runtime/Engine/Private/Materials/MaterialInstanceConstant.cpp
+++ b/Engine/Source/Runtime/Engine/Private/Materials/MaterialInstanceConstant.cpp
@@ -106,7 +106,7 @@ void UMaterialInstanceConstant::CopyMaterialUniformParametersEditorOnly(UMateria
 			FStaticParameterSet MyParamSet;
 			GetStaticParameterValues(MyParamSet);
 
-			MyParamSet.StaticSwitchParameters = SourceParamSet.StaticSwitchParameters;
+			MyParamSet.EditorOnly.StaticSwitchParameters = SourceParamSet.EditorOnly.StaticSwitchParameters;
 
 			UpdateStaticPermutation(MyParamSet);
 
@@ -165,10 +165,10 @@ void FMaterialInstanceCachedData::InitializeForConstant(const FMaterialLayersFun
 	for (int32 LayerIndex = 0; LayerIndex < NumLayers; ++LayerIndex)
 	{
 		int32 ParentLayerIndex = INDEX_NONE;
-		if (ParentLayers && Layers->LayerLinkStates[LayerIndex] == EMaterialLayerLinkState::LinkedToParent)
+		if (ParentLayers && Layers->EditorOnly.LayerLinkStates[LayerIndex] == EMaterialLayerLinkState::LinkedToParent)
 		{
-			const FGuid& LayerGuid = Layers->LayerGuids[LayerIndex];
-			ParentLayerIndex = ParentLayers->LayerGuids.Find(LayerGuid);
+			const FGuid& LayerGuid = Layers->EditorOnly.LayerGuids[LayerIndex];
+			ParentLayerIndex = ParentLayers->EditorOnly.LayerGuids.Find(LayerGuid);
 		}
 		ParentLayerIndexRemap.Add(ParentLayerIndex);
 	}
@@ -211,18 +211,20 @@ void UMaterialInstanceConstant::UpdateCachedData()
 		if (LocalStaticParameters.bHasMaterialLayers)
 		{
 			UMaterial* BaseMaterial = GetMaterial();
+
+			FMaterialLayersFunctions MaterialLayers;
+			LocalStaticParameters.GetMaterialLayers(MaterialLayers);
 			if (bUsingNewHLSLGenerator)
 			{
 				LocalCachedTree.Reset(new FMaterialCachedHLSLTree());
-				LocalCachedTree->GenerateTree(BaseMaterial, &LocalStaticParameters.MaterialLayers, nullptr);
+				LocalCachedTree->GenerateTree(BaseMaterial, &MaterialLayers, nullptr);
 			}
 			else
 			{
 				FMaterialCachedExpressionContext Context;
-				Context.LayerOverrides = &LocalStaticParameters.MaterialLayers;
+				Context.LayerOverrides = &MaterialLayers;
 				LocalCachedExpressionData.Reset(new FMaterialCachedExpressionData());
-				LocalCachedExpressionData->Reset();
-				LocalCachedExpressionData->UpdateForExpressions(Context, BaseMaterial->Expressions, GlobalParameter, INDEX_NONE);
+				LocalCachedExpressionData->UpdateForExpressions(Context, BaseMaterial->GetExpressions(), GlobalParameter, INDEX_NONE);
 			}
 		}
 
@@ -232,11 +234,14 @@ void UMaterialInstanceConstant::UpdateCachedData()
 		{
 			check(!LocalCachedExpressionData);
 			LocalCachedExpressionData.Reset(new FMaterialCachedExpressionData());
-			LocalCachedExpressionData->Reset();
 			LocalCachedExpressionData->UpdateForCachedHLSLTree(GetCachedHLSLTree(), &LocalStaticParameters);
 		}
 
 		CachedExpressionData = MoveTemp(LocalCachedExpressionData);
+		if (CachedExpressionData)
+		{
+			EditorOnlyData->CachedExpressionData = CachedExpressionData->EditorOnlyData;
+		}
 
 		FObjectCacheEventSink::NotifyReferencedTextureChanged_Concurrent(this);
 	}
diff --git a/Engine/Source/Runtime/Engine/Private/Materials/MaterialInterface.cpp b/Engine/Source/Runtime/Engine/Private/Materials/MaterialInterface.cpp
--- a/Engine/Source/Runtime/Engine/Private/Materials/MaterialInterface.cpp
+++ b/Engine/Source/Runtime/Engine/Private/Materials/MaterialInterface.cpp
@@ -165,7 +165,17 @@ void UMaterialInterface::Serialize(FArchive& Ar)
 		{
 			CachedExpressionData.Reset(new FMaterialCachedExpressionData());
 			bLoadedCachedExpressionData = true;
+#if WITH_EDITORONLY_DATA
+			EditorOnlyData->CachedExpressionData = CachedExpressionData->EditorOnlyData;
+#endif
+		}
+		else
+		{
+#if WITH_EDITOR
+			CachedExpressionData->Validate();
+#endif
 		}
+
 		check(CachedExpressionData);
 		UScriptStruct* Struct = FMaterialCachedExpressionData::StaticStruct();
 		Struct->SerializeTaggedProperties(Ar, (uint8*)CachedExpressionData.Get(), Struct, nullptr);
@@ -191,7 +201,7 @@ void UMaterialInterface::PostLoad()
 	{
 		TextureStreamingData.Empty();
 	}
-#endif
+#endif // WITH_EDITORONLY_DATA
 }
 
 const FMaterialCachedExpressionData& UMaterialInterface::GetCachedExpressionData(TMicRecursionGuard) const
@@ -599,6 +609,14 @@ void UMaterialInterface::AddReferencedObjects(UObject* InThis, FReferenceCollect
 	Super::AddReferencedObjects(This, Collector);
 }
 
+void UMaterialInterface::PostInitProperties()
+{
+#if WITH_EDITORONLY_DATA
+	EditorOnlyData = CreateEditorOnlyData();
+#endif
+	Super::PostInitProperties();
+}
+
 void UMaterialInterface::PostDuplicate(bool bDuplicateForPIE)
 {
 	Super::PostDuplicate(bDuplicateForPIE);
@@ -878,14 +896,14 @@ bool UMaterialInterface::GetStaticComponentMaskParameterDefaultValue(const FHash
 void UMaterialInterface::GetAllParametersOfType(EMaterialParameterType Type, TMap<FMaterialParameterInfo, FMaterialParameterMetadata>& OutParameters) const
 {
 	OutParameters.Reset();
-	GetCachedExpressionData().Parameters.GetAllParametersOfType(Type, OutParameters);
+	GetCachedExpressionData().GetAllParametersOfType(Type, OutParameters);
 }
 
 void UMaterialInterface::GetAllParameterInfoOfType(EMaterialParameterType Type, TArray<FMaterialParameterInfo>& OutParameterInfo, TArray<FGuid>& OutParameterIds) const
 {
 	OutParameterInfo.Reset();
 	OutParameterIds.Reset();
-	GetCachedExpressionData().Parameters.GetAllParameterInfoOfType(Type, OutParameterInfo, OutParameterIds);
+	GetCachedExpressionData().GetAllParameterInfoOfType(Type, OutParameterInfo, OutParameterIds);
 }
 
 void UMaterialInterface::GetAllScalarParameterInfo(TArray<FMaterialParameterInfo>& OutParameterInfo, TArray<FGuid>& OutParameterIds) const
@@ -1344,3 +1362,55 @@ void UMaterialInterface::RemoveUserDataOfClass(TSubclassOf<UAssetUserData> InUse
 	}
 }
 
+#if WITH_EDITORONLY_DATA
+const UClass* UMaterialInterface::GetEditorOnlyDataClass() const
+{
+	return UMaterialInterfaceEditorOnlyData::StaticClass();
+}
+
+#include "Materials/MaterialInstanceDynamic.h"
+
+UMaterialInterfaceEditorOnlyData* UMaterialInterface::CreateEditorOnlyData()
+{
+	const UClass* EditorOnlyClass = GetEditorOnlyDataClass();
+	check(EditorOnlyClass);
+	check(EditorOnlyClass->HasAllClassFlags(CLASS_Optional));
+
+	const FString EditorOnlyName = FString::Printf(TEXT("%sEditorOnlyData"), *GetName());
+	const EObjectFlags EditorOnlyFlags = GetMaskedFlags(RF_PropagateToSubObjects);
+	return NewObject<UMaterialInterfaceEditorOnlyData>(this, EditorOnlyClass, *EditorOnlyName, EditorOnlyFlags);
+}
+#endif // WITH_EDITORONLY_DATA
+
+UMaterialInterfaceEditorOnlyData::UMaterialInterfaceEditorOnlyData()
+{
+}
+
+UMaterialInterfaceEditorOnlyData::UMaterialInterfaceEditorOnlyData(FVTableHelper& Helper)
+{
+}
+
+UMaterialInterfaceEditorOnlyData::~UMaterialInterfaceEditorOnlyData()
+{
+}
+
+void UMaterialInterfaceEditorOnlyData::Serialize(FArchive& Ar)
+{
+	Super::Serialize(Ar);
+
+	bool bSavedCachedExpressionData = false;
+	if ((Ar.IsCooking() || (FPlatformProperties::RequiresCookedData() && Ar.IsSaving() && (Ar.GetPortFlags() & PPF_Duplicate))) && (bool)CachedExpressionData)
+	{
+		bSavedCachedExpressionData = true;
+	}
+
+	Ar << bSavedCachedExpressionData;
+
+	if (bSavedCachedExpressionData)
+	{
+		check(CachedExpressionData);
+		UScriptStruct* Struct = FMaterialCachedExpressionEditorOnlyData::StaticStruct();
+		Struct->SerializeTaggedProperties(Ar, (uint8*)CachedExpressionData.Get(), Struct, nullptr);
+		bLoadedCachedExpressionData = true;
+	}
+}
diff --git a/Engine/Source/Runtime/Engine/Private/Materials/MaterialShader.cpp b/Engine/Source/Runtime/Engine/Private/Materials/MaterialShader.cpp
--- a/Engine/Source/Runtime/Engine/Private/Materials/MaterialShader.cpp
+++ b/Engine/Source/Runtime/Engine/Private/Materials/MaterialShader.cpp
@@ -249,39 +249,50 @@ void UpdateMaterialShaderCompilingStats(const FMaterial* Material)
 	}
 }
 
-FStaticParameterSet::FStaticParameterSet(const FStaticParameterSet& InValue)
-	: StaticSwitchParameters(InValue.StaticSwitchParameters)
-	, StaticComponentMaskParameters(InValue.StaticComponentMaskParameters)
-	, TerrainLayerWeightParameters(InValue.TerrainLayerWeightParameters)
-	, bHasMaterialLayers(InValue.bHasMaterialLayers)
-{
-	if (bHasMaterialLayers)
-	{
-		MaterialLayers = InValue.MaterialLayers;
-	}
-}
+FStaticParameterSet::FStaticParameterSet(const FStaticParameterSet& InValue) = default;
 
 FStaticParameterSet& FStaticParameterSet::operator=(const FStaticParameterSet& InValue)
 {
-	StaticSwitchParameters = InValue.StaticSwitchParameters;
-	StaticComponentMaskParameters = InValue.StaticComponentMaskParameters;
-	TerrainLayerWeightParameters = InValue.TerrainLayerWeightParameters;
+#if WITH_EDITORONLY_DATA
+	EditorOnly.StaticSwitchParameters = InValue.EditorOnly.StaticSwitchParameters;
+	EditorOnly.StaticComponentMaskParameters = InValue.EditorOnly.StaticComponentMaskParameters;
+	EditorOnly.TerrainLayerWeightParameters = InValue.EditorOnly.TerrainLayerWeightParameters;
+#endif // WITH_EDITORONLY_DATA
 	MaterialLayers = InValue.MaterialLayers;
 	bHasMaterialLayers = InValue.bHasMaterialLayers;
 	if (bHasMaterialLayers)
 	{
 		MaterialLayers = InValue.MaterialLayers;
+#if WITH_EDITORONLY_DATA
+		EditorOnly.MaterialLayers = InValue.EditorOnly.MaterialLayers;
+#endif
 	}
 	return *this;
 }
 
-void FStaticParameterSet::Empty()
+void FStaticParameterSet::PostSerialize(const FArchive& Ar)
 {
-	StaticSwitchParameters.Empty();
-	StaticComponentMaskParameters.Empty();
-	TerrainLayerWeightParameters.Empty();
-	MaterialLayers.Empty();
-	bHasMaterialLayers = false;
+#if WITH_EDITORONLY_DATA
+	if (StaticSwitchParameters_DEPRECATED.Num() > 0)
+	{
+		EditorOnly.StaticSwitchParameters = MoveTemp(StaticSwitchParameters_DEPRECATED);
+	}
+	if (StaticComponentMaskParameters_DEPRECATED.Num() > 0)
+	{
+		EditorOnly.StaticComponentMaskParameters = MoveTemp(StaticComponentMaskParameters_DEPRECATED);
+	}
+	if (TerrainLayerWeightParameters_DEPRECATED.Num() > 0)
+	{
+		EditorOnly.TerrainLayerWeightParameters = MoveTemp(TerrainLayerWeightParameters_DEPRECATED);
+	}
+
+	if (!MaterialLayers_DEPRECATED.IsEmpty())
+	{
+		MaterialLayers = MoveTemp(MaterialLayers_DEPRECATED.GetRuntime());
+		EditorOnly.MaterialLayers = MoveTemp(MaterialLayers_DEPRECATED.EditorOnly);
+	}
+	MaterialLayers_DEPRECATED.Empty();
+#endif // WITH_EDITORONLY_DATA
 }
 
 #if WITH_EDITOR
@@ -296,9 +307,9 @@ void FStaticParameterSet::SerializeLegacy(FArchive& Ar)
 	Ar.UsingCustomVersion(FReleaseObjectVersion::GUID);
 	Ar.UsingCustomVersion(FUE5ReleaseStreamObjectVersion::GUID);
 
-	Ar << StaticSwitchParameters;
-	Ar << StaticComponentMaskParameters;
-	Ar << TerrainLayerWeightParameters;
+	Ar << EditorOnly.StaticSwitchParameters;
+	Ar << EditorOnly.StaticComponentMaskParameters;
+	Ar << EditorOnly.TerrainLayerWeightParameters;
 
 	if (Ar.CustomVer(FReleaseObjectVersion::GUID) >= FReleaseObjectVersion::MaterialLayersParameterSerializationRefactor)
 	{
@@ -309,7 +320,8 @@ void FStaticParameterSet::SerializeLegacy(FArchive& Ar)
 			if (MaterialLayersParameters_DEPRECATED.Num() > 0)
 			{
 				bHasMaterialLayers = true;
-				MaterialLayers = MoveTemp(MaterialLayersParameters_DEPRECATED[0].Value);
+				MaterialLayers = MoveTemp(MaterialLayersParameters_DEPRECATED[0].Value.GetRuntime());
+				EditorOnly.MaterialLayers = MoveTemp(MaterialLayersParameters_DEPRECATED[0].Value.EditorOnly);
 				MaterialLayersParameters_DEPRECATED.Empty();
 			}
 			PRAGMA_ENABLE_DEPRECATION_WARNINGS
@@ -320,9 +332,9 @@ void FStaticParameterSet::SerializeLegacy(FArchive& Ar)
 void FStaticParameterSet::UpdateLegacyTerrainLayerWeightData()
 {
 	int32 ParameterIndex = 0;
-	while (ParameterIndex < TerrainLayerWeightParameters.Num())
+	while (ParameterIndex < EditorOnly.TerrainLayerWeightParameters.Num())
 	{
-		FStaticTerrainLayerWeightParameter& TerrainParameter = TerrainLayerWeightParameters[ParameterIndex];
+		FStaticTerrainLayerWeightParameter& TerrainParameter = EditorOnly.TerrainLayerWeightParameters[ParameterIndex];
 		if (TerrainParameter.bOverride_DEPRECATED)
 		{
 			TerrainParameter.LayerName = TerrainParameter.ParameterInfo_DEPRECATED.Name;
@@ -331,7 +343,7 @@ void FStaticParameterSet::UpdateLegacyTerrainLayerWeightData()
 		else
 		{
 			// Remove any parameters that didn't have bOverride set
-			TerrainLayerWeightParameters.RemoveAt(ParameterIndex);
+			EditorOnly.TerrainLayerWeightParameters.RemoveAt(ParameterIndex);
 		}
 	}
 }
@@ -342,70 +354,111 @@ void FStaticParameterSet::UpdateLegacyMaterialLayersData()
 	if (MaterialLayersParameters_DEPRECATED.Num() > 0)
 	{
 		bHasMaterialLayers = true;
-		MaterialLayers = MoveTemp(MaterialLayersParameters_DEPRECATED[0].Value);
+		MaterialLayers = MoveTemp(MaterialLayersParameters_DEPRECATED[0].Value.GetRuntime());
+		EditorOnly.MaterialLayers = MoveTemp(MaterialLayersParameters_DEPRECATED[0].Value.EditorOnly);
 		MaterialLayersParameters_DEPRECATED.Empty();
 	}
 	PRAGMA_ENABLE_DEPRECATION_WARNINGS
 }
+
+bool FStaticParameterSet::GetMaterialLayers(FMaterialLayersFunctions& OutMaterialLayers) const
+{
+	if (bHasMaterialLayers)
+	{
+		OutMaterialLayers.GetRuntime() = MaterialLayers;
+		OutMaterialLayers.EditorOnly = EditorOnly.MaterialLayers;
+		return true;
+	}
+	return false;
+}
+
+void FStaticParameterSet::Validate(const FStaticParameterSetRuntimeData& Runtime, const FStaticParameterSetEditorOnlyData& EditorOnly)
+{
+	FMaterialLayersFunctions::Validate(Runtime.MaterialLayers, EditorOnly.MaterialLayers);
+}
+
 #endif // WITH_EDITOR
 
 bool FStaticParameterSet::operator==(const FStaticParameterSet& ReferenceSet) const
 {
-	if (StaticSwitchParameters.Num() != ReferenceSet.StaticSwitchParameters.Num()
-		|| StaticComponentMaskParameters.Num() != ReferenceSet.StaticComponentMaskParameters.Num()
-		|| TerrainLayerWeightParameters.Num() != ReferenceSet.TerrainLayerWeightParameters.Num()
-		|| bHasMaterialLayers != ReferenceSet.bHasMaterialLayers)
+	if (bHasMaterialLayers != ReferenceSet.bHasMaterialLayers)
 	{
 		return false;
 	}
-
-	if (StaticSwitchParameters != ReferenceSet.StaticSwitchParameters)
+#if WITH_EDITORONLY_DATA
+	if (EditorOnly.StaticSwitchParameters.Num() != ReferenceSet.EditorOnly.StaticSwitchParameters.Num()
+		|| EditorOnly.StaticComponentMaskParameters.Num() != ReferenceSet.EditorOnly.StaticComponentMaskParameters.Num()
+		|| EditorOnly.TerrainLayerWeightParameters.Num() != ReferenceSet.EditorOnly.TerrainLayerWeightParameters.Num())
 	{
 		return false;
 	}
 
-	if (StaticComponentMaskParameters != ReferenceSet.StaticComponentMaskParameters)
+	if (EditorOnly.StaticSwitchParameters != ReferenceSet.EditorOnly.StaticSwitchParameters)
 	{
 		return false;
 	}
 
-	if (TerrainLayerWeightParameters != ReferenceSet.TerrainLayerWeightParameters)
+	if (EditorOnly.StaticComponentMaskParameters != ReferenceSet.EditorOnly.StaticComponentMaskParameters)
 	{
 		return false;
 	}
 
-	if (bHasMaterialLayers && MaterialLayers != ReferenceSet.MaterialLayers)
+	if (EditorOnly.TerrainLayerWeightParameters != ReferenceSet.EditorOnly.TerrainLayerWeightParameters)
 	{
 		return false;
 	}
+#endif // WITH_EDITORONLY_DATA
+
+	if (bHasMaterialLayers)
+	{
+		if (MaterialLayers != ReferenceSet.MaterialLayers)
+		{
+			return false;
+		}
+#if WITH_EDITORONLY_DATA
+		if (EditorOnly.MaterialLayers != ReferenceSet.EditorOnly.MaterialLayers)
+		{
+			return false;
+		}
+#endif // WITH_EDITORONLY_DATA
+	}
 
 	return true;
 }
 
 void FStaticParameterSet::SortForEquivalent()
 {
-	StaticSwitchParameters.Sort([](const FStaticSwitchParameter& A, const FStaticSwitchParameter& B) { return B.ExpressionGUID < A.ExpressionGUID; });
-	StaticComponentMaskParameters.Sort([](const FStaticComponentMaskParameter& A, const FStaticComponentMaskParameter& B) { return B.ExpressionGUID < A.ExpressionGUID; });
-	TerrainLayerWeightParameters.Sort([](const FStaticTerrainLayerWeightParameter& A, const FStaticTerrainLayerWeightParameter& B) { return B.LayerName.LexicalLess(A.LayerName); });
+#if WITH_EDITORONLY_DATA
+	EditorOnly.StaticSwitchParameters.Sort([](const FStaticSwitchParameter& A, const FStaticSwitchParameter& B) { return B.ExpressionGUID < A.ExpressionGUID; });
+	EditorOnly.StaticComponentMaskParameters.Sort([](const FStaticComponentMaskParameter& A, const FStaticComponentMaskParameter& B) { return B.ExpressionGUID < A.ExpressionGUID; });
+	EditorOnly.TerrainLayerWeightParameters.Sort([](const FStaticTerrainLayerWeightParameter& A, const FStaticTerrainLayerWeightParameter& B) { return B.LayerName.LexicalLess(A.LayerName); });
+#endif // WITH_EDITORONLY_DATA
 }
 
 bool FStaticParameterSet::Equivalent(const FStaticParameterSet& ReferenceSet) const
 {
-	if (StaticSwitchParameters.Num() == ReferenceSet.StaticSwitchParameters.Num()
-		&& StaticComponentMaskParameters.Num() == ReferenceSet.StaticComponentMaskParameters.Num()
-		&& TerrainLayerWeightParameters.Num() == ReferenceSet.TerrainLayerWeightParameters.Num()
-		&& bHasMaterialLayers == ReferenceSet.bHasMaterialLayers)
+	if (bHasMaterialLayers != ReferenceSet.bHasMaterialLayers)
 	{
-		// this is not ideal, but it is easy to code up
-		FStaticParameterSet Temp1 = *this;
-		FStaticParameterSet Temp2 = ReferenceSet;
-		Temp1.SortForEquivalent();
-		Temp2.SortForEquivalent();
-		bool bResult = (Temp1 == Temp2);
-		ensure(!bResult || (*this) == ReferenceSet); // if this never fires, then we really didn't need to sort did we?
-		return bResult;
+		return false;
 	}
-	return false;
+
+#if WITH_EDITORONLY_DATA
+	if (EditorOnly.StaticSwitchParameters.Num() != ReferenceSet.EditorOnly.StaticSwitchParameters.Num()
+		|| EditorOnly.StaticComponentMaskParameters.Num() != ReferenceSet.EditorOnly.StaticComponentMaskParameters.Num()
+		|| EditorOnly.TerrainLayerWeightParameters.Num() != ReferenceSet.EditorOnly.TerrainLayerWeightParameters.Num())
+	{
+		return false;
+	}
+#endif // WITH_EDITORONLY_DATA
+
+	// this is not ideal, but it is easy to code up
+	FStaticParameterSet Temp1 = *this;
+	FStaticParameterSet Temp2 = ReferenceSet;
+	Temp1.SortForEquivalent();
+	Temp2.SortForEquivalent();
+	bool bResult = (Temp1 == Temp2);
+	ensure(!bResult || (*this) == ReferenceSet); // if this never fires, then we really didn't need to sort did we?
+	return bResult;
 }
 
 #if WITH_EDITORONLY_DATA
@@ -425,21 +478,21 @@ void FStaticParameterSet::AddParametersOfType(EMaterialParameterType Type, const
 	switch (Type)
 	{
 	case EMaterialParameterType::StaticSwitch:
-		StaticSwitchParameters.Empty(Values.Num());
+		EditorOnly.StaticSwitchParameters.Empty(Values.Num());
 		for (const auto& It : Values)
 		{
 			const FMaterialParameterMetadata& Meta = It.Value;
 			check(Meta.Value.Type == Type);
-			StaticSwitchParameters.Emplace(It.Key, Meta.Value.AsStaticSwitch(), Meta.bOverride, Meta.ExpressionGuid);
+			EditorOnly.StaticSwitchParameters.Emplace(It.Key, Meta.Value.AsStaticSwitch(), Meta.bOverride, Meta.ExpressionGuid);
 		}
 		break;
 	case EMaterialParameterType::StaticComponentMask:
-		StaticComponentMaskParameters.Empty(Values.Num());
+		EditorOnly.StaticComponentMaskParameters.Empty(Values.Num());
 		for (const auto& It : Values)
 		{
 			const FMaterialParameterMetadata& Meta = It.Value;
 			check(Meta.Value.Type == Type);
-			StaticComponentMaskParameters.Emplace(It.Key,
+			EditorOnly.StaticComponentMaskParameters.Emplace(It.Key,
 				Meta.Value.Bool[0],
 				Meta.Value.Bool[1],
 				Meta.Value.Bool[2],
@@ -454,7 +507,7 @@ void FStaticParameterSet::AddParametersOfType(EMaterialParameterType Type, const
 
 void FStaticParameterSet::SetStaticSwitchParameterValue(const FMaterialParameterInfo& ParameterInfo, const FGuid& ExpressionGuid, bool Value)
 {
-	for (FStaticSwitchParameter& Parameter : StaticSwitchParameters)
+	for (FStaticSwitchParameter& Parameter : EditorOnly.StaticSwitchParameters)
 	{
 		if (Parameter.ParameterInfo == ParameterInfo)
 		{
@@ -464,12 +517,12 @@ void FStaticParameterSet::SetStaticSwitchParameterValue(const FMaterialParameter
 		}
 	}
 
-	new(StaticSwitchParameters) FStaticSwitchParameter(ParameterInfo, Value, true, ExpressionGuid);
+	new(EditorOnly.StaticSwitchParameters) FStaticSwitchParameter(ParameterInfo, Value, true, ExpressionGuid);
 }
 
 void FStaticParameterSet::SetStaticComponentMaskParameterValue(const FMaterialParameterInfo& ParameterInfo, const FGuid& ExpressionGuid, bool R, bool G, bool B, bool A)
 {
-	for (FStaticComponentMaskParameter& Parameter : StaticComponentMaskParameters)
+	for (FStaticComponentMaskParameter& Parameter : EditorOnly.StaticComponentMaskParameters)
 	{
 		if (Parameter.ParameterInfo == ParameterInfo)
 		{
@@ -482,10 +535,11 @@ void FStaticParameterSet::SetStaticComponentMaskParameterValue(const FMaterialPa
 		}
 	}
 
-	new(StaticComponentMaskParameters) FStaticComponentMaskParameter(ParameterInfo, R, G, B, A, true, ExpressionGuid);
+	new(EditorOnly.StaticComponentMaskParameters) FStaticComponentMaskParameter(ParameterInfo, R, G, B, A, true, ExpressionGuid);
 }
 #endif // WITH_EDITORONLY_DATA
 
+
 void FMaterialShaderMapId::Serialize(FArchive& Ar, bool bLoadedByCookedMaterial)
 {
 	SCOPED_LOADTIMER(FMaterialShaderMapId_Serialize);
@@ -895,12 +949,12 @@ void FMaterialShaderMapId::UpdateFromParameterSet(const FStaticParameterSet& Sta
 		}
 	};
 
-	StaticSwitchParameters = StaticParameters.StaticSwitchParameters;
-	StaticComponentMaskParameters = StaticParameters.StaticComponentMaskParameters;
-	TerrainLayerWeightParameters = StaticParameters.TerrainLayerWeightParameters;
+	StaticSwitchParameters = StaticParameters.EditorOnly.StaticSwitchParameters;
+	StaticComponentMaskParameters = StaticParameters.EditorOnly.StaticComponentMaskParameters;
+	TerrainLayerWeightParameters = StaticParameters.EditorOnly.TerrainLayerWeightParameters;
 	if (StaticParameters.bHasMaterialLayers)
 	{
-		MaterialLayersId = StaticParameters.MaterialLayers.GetID();
+		MaterialLayersId = StaticParameters.MaterialLayers.GetID(StaticParameters.EditorOnly.MaterialLayers);
 	}
 
 	// Sort the arrays by parameter name, ensure the ID is not influenced by the order
diff --git a/Engine/Source/Runtime/Engine/Private/Materials/MaterialShared.cpp b/Engine/Source/Runtime/Engine/Private/Materials/MaterialShared.cpp
--- a/Engine/Source/Runtime/Engine/Private/Materials/MaterialShared.cpp
+++ b/Engine/Source/Runtime/Engine/Private/Materials/MaterialShared.cpp
@@ -361,21 +361,17 @@ void FExpressionInput::Connect( int32 InOutputIndex, class UMaterialExpression*
 {
 	InExpression->ConnectExpression(this, InOutputIndex);
 }
-#endif // WITH_EDITOR
 
 FExpressionInput FExpressionInput::GetTracedInput() const
 {
-#if WITH_EDITORONLY_DATA
 	if (Expression != nullptr && Expression->IsA(UMaterialExpressionRerouteBase::StaticClass()))
 	{
 		UMaterialExpressionRerouteBase* Reroute = CastChecked<UMaterialExpressionRerouteBase>(Expression);
 		return Reroute->TraceInputsToRealInput();
 	}
-#endif
 	return *this;
 }
 
-#if WITH_EDITOR
 int32 FExpressionExecOutput::Compile(class FMaterialCompiler* Compiler) const
 {
 	if (Expression)
@@ -455,12 +451,8 @@ static bool SerializeExpressionInput(FArchive& Ar, FExpressionInput& Input)
 		return false;
 	}
 
-#if WITH_EDITORONLY_DATA
-	if (!Ar.IsFilterEditorOnly())
-	{
-		Ar << Input.Expression;
-	}
-#endif
+	Ar << Input.Expression;
+
 	Ar << Input.OutputIndex;
 	if (Ar.CustomVer(FFrameworkObjectVersion::GUID) >= FFrameworkObjectVersion::PinsStoreFName)
 	{
@@ -473,29 +465,11 @@ static bool SerializeExpressionInput(FArchive& Ar, FExpressionInput& Input)
 		Input.InputName = *InputNameStr;
 	}
 
-#if WITH_EDITORONLY_DATA
 	Ar << Input.Mask;
 	Ar << Input.MaskR;
 	Ar << Input.MaskG;
 	Ar << Input.MaskB;
 	Ar << Input.MaskA;
-#else
-	int32 Temp = 0;
-	Ar << Temp << Temp << Temp << Temp << Temp;
-#endif
-
-	// Some expressions may have been stripped when cooking and Expression can be null after loading
-	// so make sure we keep the information about the connected node in cooked packages
-	if ( Ar.IsFilterEditorOnly() )
-	{
-#if WITH_EDITORONLY_DATA
-		if (Ar.IsSaving())
-		{
-			Input.ExpressionName = Input.Expression ? Input.Expression->GetFName() : NAME_None;
-		}
-#endif // WITH_EDITORONLY_DATA
-		Ar << Input.ExpressionName;
-	}
 
 	return true;
 }
@@ -505,17 +479,10 @@ static bool SerializeMaterialInput(FArchive& Ar, FMaterialInput<InputType>& Inpu
 {
 	if (SerializeExpressionInput(Ar, Input))
 	{
-#if WITH_EDITORONLY_DATA
 		bool bUseConstantValue = Input.UseConstant;
 		Ar << bUseConstantValue;
 		Input.UseConstant = bUseConstantValue;
 		Ar << Input.Constant;
-#else
-		bool bTemp = false;
-		Ar << bTemp;
-		InputType TempType;
-		Ar << TempType;
-#endif
 		return true;
 	}
 	else
@@ -1386,7 +1353,7 @@ void FMaterial::DiscardShaderMap()
 }
 
 EMaterialDomain FMaterialResource::GetMaterialDomain() const { return Material->MaterialDomain; }
-bool FMaterialResource::IsTangentSpaceNormal() const { return Material->bTangentSpaceNormal || (!Material->Normal.IsConnected() && !Material->bUseMaterialAttributes); }
+bool FMaterialResource::IsTangentSpaceNormal() const { return Material->bTangentSpaceNormal || (!Material->IsPropertyConnected(MP_Normal) && !Material->bUseMaterialAttributes); }
 bool FMaterialResource::ShouldGenerateSphericalParticleNormals() const { return Material->bGenerateSphericalParticleNormals; }
 bool FMaterialResource::ShouldDisableDepthTest() const { return Material->bDisableDepthTest; }
 bool FMaterialResource::ShouldWriteOnlyAlpha() const { return Material->bWriteOnlyAlpha; }
@@ -1402,7 +1369,9 @@ bool FMaterialResource::IsVolumetricPrimitive() const { return Material->Materia
 bool FMaterialResource::IsSpecialEngineMaterial() const { return Material->bUsedAsSpecialEngineMaterial; }
 bool FMaterialResource::HasVertexPositionOffsetConnected() const { return Material->HasVertexPositionOffsetConnected(); }
 bool FMaterialResource::HasPixelDepthOffsetConnected() const { return Material->HasPixelDepthOffsetConnected(); }
-bool FMaterialResource::HasMaterialAttributesConnected() const { return (Material->bUseMaterialAttributes && Material->MaterialAttributes.IsConnected()); }
+#if WITH_EDITOR
+bool FMaterialResource::HasMaterialAttributesConnected() const { return (Material->bUseMaterialAttributes && Material->GetEditorOnlyData()->MaterialAttributes.IsConnected()); }
+#endif
 EMaterialShadingRate FMaterialResource::GetShadingRate() const { return Material->ShadingRate; }
 FString FMaterialResource::GetBaseMaterialPathName() const { return Material->GetPathName(); }
 FString FMaterialResource::GetDebugName() const
@@ -1595,9 +1564,7 @@ bool FMaterialResource::IsUsingPlanarForwardReflections() const
 
 bool FMaterialResource::IsNonmetal() const
 {
-	return !Material->bUseMaterialAttributes ?
-			(!Material->Metallic.IsConnected() && !Material->Specular.IsConnected()) :
-			!(Material->MaterialAttributes.IsConnected(MP_Specular) || Material->MaterialAttributes.IsConnected(MP_Metallic));
+	return (!Material->IsPropertyConnected(MP_Metallic) && !Material->IsPropertyConnected(MP_Specular));
 }
 
 bool FMaterialResource::UseLmDirectionality() const
@@ -1743,7 +1710,8 @@ bool FMaterialResource::HasSpecularConnected() const
 
 bool FMaterialResource::HasMetallicConnected() const
 {
-	return HasMaterialAttributesConnected() || Material->HasMetallicConnected();
+	// HasMaterialAttributesConnected() should be captured by HasMetallicConnected() now
+	return Material->HasMetallicConnected();
 }
 
 bool FMaterialResource::HasEmissiveColorConnected() const
@@ -1787,9 +1755,10 @@ bool FMaterialResource::HasMaterialPropertyConnected(EMaterialProperty In) const
 			if (!CachedStrataMaterialInfo.IsValid())
 			{
 				check(Material->HasStrataFrontMaterialConnected());
-				if (Material->FrontMaterial.Expression->IsResultStrataMaterial(Material->FrontMaterial.OutputIndex))
+				const UMaterialEditorOnlyData* MaterialEditorOnly = Material->GetEditorOnlyData();
+				if (MaterialEditorOnly->FrontMaterial.Expression->IsResultStrataMaterial(MaterialEditorOnly->FrontMaterial.OutputIndex))
 				{
-					Material->FrontMaterial.Expression->GatherStrataMaterialInfo(CachedStrataMaterialInfo, Material->FrontMaterial.OutputIndex);
+					MaterialEditorOnly->FrontMaterial.Expression->GatherStrataMaterialInfo(CachedStrataMaterialInfo, MaterialEditorOnly->FrontMaterial.OutputIndex);
 				}
 			}
 		}
@@ -1808,7 +1777,7 @@ bool FMaterialResource::HasMaterialPropertyConnected(EMaterialProperty In) const
 		case MP_Normal: 			return Material->HasNormalConnected();
 		case MP_Roughness: 			return Material->HasRoughnessConnected();
 		case MP_Specular: 			return Material->HasSpecularConnected();
-		case MP_Metallic: 			return HasMaterialAttributesConnected() || Material->HasMetallicConnected();
+		case MP_Metallic: 			return Material->HasMetallicConnected(); // HasMetallicConnected() should be properly capturing HasMaterialAttributesConnected() as well
 		case MP_Anisotropy: 		return Material->HasAnisotropyConnected();
 		case MP_AmbientOcclusion: 	return Material->HasAmbientOcclusionConnected();
 		}
@@ -4606,7 +4575,6 @@ void DoMaterialAttributeReorder(FExpressionInput* Input, const FPackageFileVersi
 			if (Input->OutputIndex == 11 || Input->OutputIndex == 12)
 			{
 				Input->Expression = nullptr;
-				Input->ExpressionName = NAME_None;
 			}
 			else if (Input->OutputIndex >= 13)
 			{
diff --git a/Engine/Source/Runtime/Engine/Public/MaterialCachedData.h b/Engine/Source/Runtime/Engine/Public/MaterialCachedData.h
--- a/Engine/Source/Runtime/Engine/Public/MaterialCachedData.h
+++ b/Engine/Source/Runtime/Engine/Public/MaterialCachedData.h
@@ -3,6 +3,7 @@
 #pragma once
 
 #include "CoreMinimal.h"
+#include "Templates/SharedPointer.h"
 #include "MaterialTypes.h"
 #include "Materials/MaterialLayersFunctions.h"
 #include "Misc/Guid.h"
@@ -73,10 +74,11 @@
 
 	FMaterialCachedParameterEditorInfo() = default;
 	FMaterialCachedParameterEditorInfo(const FGuid& InGuid) : ExpressionGuid(InGuid) {}
-	FMaterialCachedParameterEditorInfo(const FGuid& InGuid, const FString& InDescription, const FName& InGroup, int32 InSortPriority)
+	FMaterialCachedParameterEditorInfo(const FGuid& InGuid, const FString& InDescription, const FName& InGroup, int32 InSortPriority, int32 InAssetIndex)
 		: Description(InDescription)
 		, Group(InGroup)
 		, SortPriority(InSortPriority)
+		, AssetIndex(InAssetIndex)
 		, ExpressionGuid(InGuid)
 	{}
 
@@ -90,5 +92,8 @@
 	int32 SortPriority = 0;
 
 	UPROPERTY()
+	int32 AssetIndex = INDEX_NONE;
+
+	UPROPERTY()
 	FGuid ExpressionGuid;
 };
@@ -100,88 +105,40 @@
 
 	ENGINE_API static const FMaterialCachedParameterEntry EmptyData;
 
-	void Reset();
-
 	// This is used to map FMaterialParameterInfos to indices, which are then used to index various TArrays containing values for each type of parameter
 	// (ExpressionGuids and Overrides, along with ScalarValues, VectorValues, etc)
 	UPROPERTY()
 	TSet<FMaterialParameterInfo> ParameterInfoSet;
+};
 
-#if WITH_EDITORONLY_DATA
-	// Editor-only information for each parameter
+USTRUCT()
+struct FMaterialCachedParameterEditorEntry
+{
+	GENERATED_USTRUCT_BODY()
+
 	UPROPERTY()
 	TArray<FMaterialCachedParameterEditorInfo> EditorInfo;
-#endif // WITH_EDITORONLY_DATA
+};
 
+struct FMaterialCachedExpressionContext
+{
+	const UMaterialFunctionInterface* CurrentFunction = nullptr;
+	const FMaterialLayersFunctions* LayerOverrides = nullptr;
+	bool bUpdateFunctionExpressions = true;
 };
 
 USTRUCT()
-struct FMaterialCachedParameters
+struct FMaterialCachedExpressionEditorOnlyData
 {
 	GENERATED_USTRUCT_BODY()
 
-#if WITH_EDITORONLY_DATA
-	inline const FMaterialCachedParameterEntry& GetParameterTypeEntry(EMaterialParameterType Type) const { return Type >= EMaterialParameterType::NumRuntime ? EditorOnlyEntries[static_cast<int32>(Type) - NumMaterialRuntimeParameterTypes] : RuntimeEntries[static_cast<int32>(Type)]; }
-	inline FMaterialCachedParameterEntry& GetParameterTypeEntry(EMaterialParameterType Type) { return Type >= EMaterialParameterType::NumRuntime ? EditorOnlyEntries[static_cast<int32>(Type) - NumMaterialRuntimeParameterTypes] : RuntimeEntries[static_cast<int32>(Type)]; }
-#else
-	inline const FMaterialCachedParameterEntry& GetParameterTypeEntry(EMaterialParameterType Type) const { return Type >= EMaterialParameterType::NumRuntime ? FMaterialCachedParameterEntry::EmptyData : RuntimeEntries[static_cast<int32>(Type)]; }
-#endif
-
-#if WITH_EDITOR
-	void AddParameter(const FMaterialParameterInfo& ParameterInfo, const FMaterialParameterMetadata& ParameterMeta, UObject*& OutReferencedTexture);
-	inline void AddParameter(const FMaterialParameterInfo& ParameterInfo, const FMaterialParameterMetadata& ParameterMeta)
-	{
-		UObject* UnusedReferencedTexture = nullptr;
-		AddParameter(ParameterInfo, ParameterMeta, UnusedReferencedTexture);
-	}
-#endif // WITH_EDITOR
-
-	inline int32 GetNumParameters(EMaterialParameterType Type) const { return GetParameterTypeEntry(Type).ParameterInfoSet.Num(); }
-	int32 FindParameterIndex(EMaterialParameterType Type, const FMemoryImageMaterialParameterInfo& HashedParameterInfo) const;
-	bool GetParameterValue(EMaterialParameterType Type, const FMemoryImageMaterialParameterInfo& ParameterInfo, FMaterialParameterMetadata& OutResult) const;
-	void GetParameterValueByIndex(EMaterialParameterType Type, int32 ParameterIndex, FMaterialParameterMetadata& OutResult) const;
-#if WITH_EDITORONLY_DATA
-	const FGuid& GetExpressionGuid(EMaterialParameterType Type, int32 Index) const;
-#endif
-	void GetAllParametersOfType(EMaterialParameterType Type, TMap<FMaterialParameterInfo, FMaterialParameterMetadata>& OutParameters) const;
-	void GetAllParameterInfoOfType(EMaterialParameterType Type, TArray<FMaterialParameterInfo>& OutParameterInfo, TArray<FGuid>& OutParameterIds) const;
-	void GetAllGlobalParametersOfType(EMaterialParameterType Type, TMap<FMaterialParameterInfo, FMaterialParameterMetadata>& OutParameters) const;
-	void GetAllGlobalParameterInfoOfType(EMaterialParameterType Type, TArray<FMaterialParameterInfo>& OutParameterInfo, TArray<FGuid>& OutParameterIds) const;
-	void Reset();
+	ENGINE_API static const FMaterialCachedExpressionEditorOnlyData EmptyData;
 
 	UPROPERTY()
-	FMaterialCachedParameterEntry RuntimeEntries[NumMaterialRuntimeParameterTypes];
-
-	UPROPERTY()
-	TArray<int32> ScalarPrimitiveDataIndexValues;
-
-	UPROPERTY()
-	TArray<int32> VectorPrimitiveDataIndexValues;
-
-	UPROPERTY()
-	TArray<float> ScalarValues;
-
-	UPROPERTY()
-	TArray<FLinearColor> VectorValues;
-
-	UPROPERTY()
-	TArray<FVector4d> DoubleVectorValues;
-
-	UPROPERTY()
-	TArray<TSoftObjectPtr<UTexture>> TextureValues;
-
-	UPROPERTY()
-	TArray<TSoftObjectPtr<UFont>> FontValues;
-
-	UPROPERTY()
-	TArray<int32> FontPageValues;
-
-	UPROPERTY()
-	TArray<TSoftObjectPtr<URuntimeVirtualTexture>> RuntimeVirtualTextureValues;
-
-#if WITH_EDITORONLY_DATA
-	UPROPERTY()
 	FMaterialCachedParameterEntry EditorOnlyEntries[NumMaterialEditorOnlyParameterTypes];
 
 	UPROPERTY()
+	FMaterialCachedParameterEditorEntry EditorEntries[NumMaterialParameterTypes];
+
+	UPROPERTY()
 	TArray<bool> StaticSwitchValues;
@@ -206,14 +163,15 @@
 
 	UPROPERTY()
 	TArray<FParameterChannelNames> TextureChannelNameValues;
-#endif // WITH_EDITORONLY_DATA
+
+	UPROPERTY()
+	FMaterialLayersFunctionsEditorOnlyData MaterialLayers;
+
+	UPROPERTY()
+	TArray<FString> AssetPaths;
+
+	UPROPERTY()
+	TArray<FName> LandscapeLayerNames;
 };
 
-struct FMaterialCachedExpressionContext
-{
-	const UMaterialFunctionInterface* CurrentFunction = nullptr;
-	const FMaterialLayersFunctions* LayerOverrides = nullptr;
-	bool bUpdateFunctionExpressions = true;
-};
-
 USTRUCT()
@@ -223,23 +181,62 @@
 	
 	ENGINE_API static const FMaterialCachedExpressionData EmptyData;
 
-	FMaterialCachedExpressionData()
-		: bHasMaterialLayers(false)
-		, bHasRuntimeVirtualTextureOutput(false)
-		, bHasSceneColor(false)
-		, bHasPerInstanceCustomData(false)
-		, bHasPerInstanceRandom(false)
-		, bHasVertexInterpolator(false)
-	{}
+	ENGINE_API FMaterialCachedExpressionData();
 
 #if WITH_EDITOR
-	void UpdateForExpressions(const FMaterialCachedExpressionContext& Context, const TArray<TObjectPtr<UMaterialExpression>>& Expressions, EMaterialParameterAssociation Association, int32 ParameterIndex);
+	void UpdateForExpressions(const FMaterialCachedExpressionContext& Context, TConstArrayView<TObjectPtr<UMaterialExpression>> Expressions, EMaterialParameterAssociation Association, int32 ParameterIndex);
 	void UpdateForFunction(const FMaterialCachedExpressionContext& Context, UMaterialFunctionInterface* Function, EMaterialParameterAssociation Association, int32 ParameterIndex);
 	void UpdateForLayerFunctions(const FMaterialCachedExpressionContext& Context, const FMaterialLayersFunctions& LayerFunctions);
 	ENGINE_API void UpdateForCachedHLSLTree(const FMaterialCachedHLSLTree& CachedTree, const FStaticParameterSet* StaticParameters);
+	void Validate();
+
+	void AddParameter(const FMaterialParameterInfo& ParameterInfo, const FMaterialParameterMetadata& ParameterMeta, UObject*& OutReferencedTexture);
+	inline void AddParameter(const FMaterialParameterInfo& ParameterInfo, const FMaterialParameterMetadata& ParameterMeta)
+	{
+		UObject* UnusedReferencedTexture = nullptr;
+		AddParameter(ParameterInfo, ParameterMeta, UnusedReferencedTexture);
+	}
 #endif // WITH_EDITOR
 
-	void Reset();
+#if WITH_EDITORONLY_DATA
+	inline FMaterialCachedParameterEntry& GetParameterTypeEntry(EMaterialParameterType Type)
+	{
+		const int32 TypeIndex = (int32)Type;
+		if (TypeIndex < NumMaterialRuntimeParameterTypes)
+		{
+			return RuntimeEntries[TypeIndex];
+		}
+		check(EditorOnlyData);
+		return EditorOnlyData->EditorOnlyEntries[TypeIndex - NumMaterialRuntimeParameterTypes];
+	}
+	inline FMaterialCachedParameterEditorEntry& GetParameterEditorTypeEntry(EMaterialParameterType Type)
+	{
+		check(EditorOnlyData);
+		return EditorOnlyData->EditorEntries[(int32)Type];
+	}
 
+	inline const FMaterialCachedParameterEntry& GetParameterTypeEntry(EMaterialParameterType Type) const
+	{
+		return const_cast<FMaterialCachedExpressionData*>(this)->GetParameterTypeEntry(Type);
+	}
+	inline const FMaterialCachedParameterEditorEntry& GetParameterEditorTypeEntry(EMaterialParameterType Type) const
+	{
+		check(EditorOnlyData);
+		return EditorOnlyData->EditorEntries[(int32)Type];
+	}
+#else
+	inline const FMaterialCachedParameterEntry& GetParameterTypeEntry(EMaterialParameterType Type) const { return Type >= EMaterialParameterType::NumRuntime ? FMaterialCachedParameterEntry::EmptyData : RuntimeEntries[static_cast<int32>(Type)]; }
+#endif
+
+	inline int32 GetNumParameters(EMaterialParameterType Type) const { return GetParameterTypeEntry(Type).ParameterInfoSet.Num(); }
+	int32 FindParameterIndex(EMaterialParameterType Type, const FMemoryImageMaterialParameterInfo& HashedParameterInfo) const;
+	bool GetParameterValue(EMaterialParameterType Type, const FMemoryImageMaterialParameterInfo& ParameterInfo, FMaterialParameterMetadata& OutResult) const;
+	void GetParameterValueByIndex(EMaterialParameterType Type, int32 ParameterIndex, FMaterialParameterMetadata& OutResult) const;
+	const FGuid& GetExpressionGuid(EMaterialParameterType Type, int32 Index) const;
+	void GetAllParametersOfType(EMaterialParameterType Type, TMap<FMaterialParameterInfo, FMaterialParameterMetadata>& OutParameters) const;
+	void GetAllParameterInfoOfType(EMaterialParameterType Type, TArray<FMaterialParameterInfo>& OutParameterInfo, TArray<FGuid>& OutParameterIds) const;
+	void GetAllGlobalParametersOfType(EMaterialParameterType Type, TMap<FMaterialParameterInfo, FMaterialParameterMetadata>& OutParameters) const;
+	void GetAllGlobalParameterInfoOfType(EMaterialParameterType Type, TArray<FMaterialParameterInfo>& OutParameterInfo, TArray<FGuid>& OutParameterIds) const;
+
 	ENGINE_API void AddReferencedObjects(FReferenceCollector& Collector);
 
@@ -249,18 +246,49 @@
 	/** Returns an array of the guids of parameter collections used. */
 	void AppendReferencedParameterCollectionIdsTo(TArray<FGuid>& OutIds) const;
 
-	bool IsMaterialAttributePropertyConnected(EMaterialProperty Property) const
+	bool IsPropertyConnected(EMaterialProperty Property) const
 	{
-		return ((MaterialAttributesPropertyConnectedBitmask >> (uint32)Property) & 0x1) != 0;
+		return ((PropertyConnectedBitmask >> (uint32)Property) & 0x1) != 0;
 	}
 
-	void SetMaterialAttributePropertyConnected(EMaterialProperty Property, bool bIsConnected)
+	void SetPropertyConnected(EMaterialProperty Property)
 	{
-		MaterialAttributesPropertyConnectedBitmask = bIsConnected ? MaterialAttributesPropertyConnectedBitmask | (1 << (uint32)Property) : MaterialAttributesPropertyConnectedBitmask & ~(1 << (uint32)Property);
+		PropertyConnectedBitmask |= (1 << (uint32)Property);
 	}
 
+#if WITH_EDITORONLY_DATA
+	TSharedPtr<FMaterialCachedExpressionEditorOnlyData> EditorOnlyData;
+#endif // WITH_EDITORONLY_DATA
+
 	UPROPERTY()
-	FMaterialCachedParameters Parameters;
+	FMaterialCachedParameterEntry RuntimeEntries[NumMaterialRuntimeParameterTypes];
 
+	UPROPERTY()
+	TArray<int32> ScalarPrimitiveDataIndexValues;
+
+	UPROPERTY()
+	TArray<int32> VectorPrimitiveDataIndexValues;
+
+	UPROPERTY()
+	TArray<float> ScalarValues;
+
+	UPROPERTY()
+	TArray<FLinearColor> VectorValues;
+
+	UPROPERTY()
+	TArray<FVector4d> DoubleVectorValues;
+
+	UPROPERTY()
+	TArray<TSoftObjectPtr<UTexture>> TextureValues;
+
+	UPROPERTY()
+	TArray<TSoftObjectPtr<UFont>> FontValues;
+
+	UPROPERTY()
+	TArray<int32> FontPageValues;
+
+	UPROPERTY()
+	TArray<TSoftObjectPtr<URuntimeVirtualTexture>> RuntimeVirtualTextureValues;
+
 	/** Array of all texture referenced by this material */
 	UPROPERTY()
@@ -275,10 +303,10 @@
 	TArray<FMaterialParameterCollectionInfo> ParameterCollectionInfos;
 
 	UPROPERTY()
-	FMaterialLayersFunctions MaterialLayers;
-
-	UPROPERTY()
 	TArray<TObjectPtr<ULandscapeGrassType>> GrassTypes;
 
 	UPROPERTY()
+	FMaterialLayersFunctionsRuntimeData MaterialLayers;
+
+	UPROPERTY()
 	TArray<FName> DynamicParameterNames;
@@ -306,11 +334,6 @@
 
 	/** Each bit corresponds to EMaterialProperty connection status. */
 	UPROPERTY()
-	uint32 MaterialAttributesPropertyConnectedBitmask = 0;
-
-#if WITH_EDITORONLY_DATA
-	UPROPERTY()
-	TArray<FName> LandscapeLayerNames;
-#endif // WITH_EDITORONLY_DATA
+	uint32 PropertyConnectedBitmask = 0;
 };
 
diff --git a/Engine/Source/Runtime/Engine/Public/MaterialExpressionIO.h b/Engine/Source/Runtime/Engine/Public/MaterialExpressionIO.h
--- a/Engine/Source/Runtime/Engine/Public/MaterialExpressionIO.h
+++ b/Engine/Source/Runtime/Engine/Public/MaterialExpressionIO.h
@@ -30,13 +30,11 @@ struct FValue;
 //@warning: FExpressionInput is mirrored in MaterialExpression.h and manually "subclassed" in Material.h (FMaterialInput)
 struct FExpressionInput
 {
-#if WITH_EDITORONLY_DATA
 	/** 
 	 * Material expression that this input is connected to, or NULL if not connected. 
 	 * If you want to be safe when checking against dangling Reroute nodes, please use GetTracedInput before accessing this property.
 	*/
 	class UMaterialExpression*	Expression;
-#endif
 
 	/** 
 	 * Index into Expression's outputs array that this input is connected to.
@@ -50,32 +48,26 @@ struct FExpressionInput
 	 */
 	FName						InputName;
 
-#if WITH_EDITORONLY_DATA
 	int32						Mask,
 								MaskR,
 								MaskG,
 								MaskB,
 								MaskA;
-#endif
-
-	/** Material expression name that this input is connected to, or None if not connected. Used only in cooked builds */
-	FName				ExpressionName;
 
 	FExpressionInput()
 		: OutputIndex(0)
-#if WITH_EDITORONLY_DATA
 		, Mask(0)
 		, MaskR(0)
 		, MaskG(0)
 		, MaskB(0)
 		, MaskA(0)
-#endif
 	{
-#if WITH_EDITORONLY_DATA
 		Expression = nullptr;
-#endif
 	}
 
+	/** ICPPStructOps interface */
+	ENGINE_API bool Serialize(FArchive& Ar);
+
 #if WITH_EDITOR
 	ENGINE_API int32 Compile(class FMaterialCompiler* Compiler);
 	ENGINE_API const UE::HLSLTree::FExpression* TryAcquireHLSLExpression(FMaterialHLSLGenerator& Generator, UE::HLSLTree::FScope& Scope, int32 InputIndex) const;
@@ -87,7 +79,6 @@ struct FExpressionInput
 	ENGINE_API const UE::HLSLTree::FExpression* AcquireHLSLExpression(FMaterialHLSLGenerator& Generator, UE::HLSLTree::FScope& Scope) const;
 	ENGINE_API const UE::HLSLTree::FExpression* AcquireHLSLExpressionOrConstant(FMaterialHLSLGenerator& Generator, UE::HLSLTree::FScope& Scope, const UE::Shader::FValue& ConstantValue) const;
 	ENGINE_API const UE::HLSLTree::FExpression* AcquireHLSLExpressionOrExternalInput(FMaterialHLSLGenerator& Generator, UE::HLSLTree::FScope& Scope, UE::HLSLTree::Material::EExternalInput Input) const;
-#endif // WITH_EDITOR
 
 	/**
 	 * Tests if the input has a material expression connected to it
@@ -96,20 +87,14 @@ struct FExpressionInput
 	 */
 	bool IsConnected() const 
 	{ 
-#if WITH_EDITORONLY_DATA
 		return (nullptr != Expression);
-#else
-		return ExpressionName != NAME_None;
-#endif // WITH_EDITORONLY_DATA
 	}
 
-#if WITH_EDITOR
+	/** Is the input connected to a constant value */
+	bool IsConstant() const { return false; }
+
 	/** Connects output of InExpression to this input */
 	ENGINE_API void Connect( int32 InOutputIndex, class UMaterialExpression* InExpression );
-#endif // WITH_EDITOR
-
-	/** ICPPStructOps interface */
-	ENGINE_API bool Serialize(FArchive& Ar);
 
 	/** If this input goes through reroute nodes or other paths that should not affect code, trace back on the input chain.*/
 	ENGINE_API FExpressionInput GetTracedInput() const;
@@ -117,14 +102,13 @@ struct FExpressionInput
 	/** Helper for setting component mask. */
 	ENGINE_API void SetMask(int32 UseMask, int32 R, int32 G, int32 B, int32 A)
 	{
-#if WITH_EDITORONLY_DATA
 		Mask = UseMask;
 		MaskR = R;
 		MaskG = G;
 		MaskB = B;
 		MaskA = A;
-#endif
 	}
+#endif // WITH_EDITOR
 };
 
 template<>
@@ -144,45 +128,37 @@ struct TStructOpsTypeTraits<FExpressionInput>
 struct FExpressionOutput
 {
 	FName	OutputName;
-#if WITH_EDITORONLY_DATA
 	int32	Mask,
 		MaskR,
 		MaskG,
 		MaskB,
 		MaskA;
-#endif
 
 	FExpressionOutput(int32 InMask = 0, int32 InMaskR = 0, int32 InMaskG = 0, int32 InMaskB = 0, int32 InMaskA = 0)
-#if WITH_EDITORONLY_DATA
 		: Mask(InMask)
 		, MaskR(InMaskR)
 		, MaskG(InMaskG)
 		, MaskB(InMaskB)
 		, MaskA(InMaskA)
-#endif
 	{}
 
 	FExpressionOutput(FName InOutputName, int32 InMask = 0, int32 InMaskR = 0, int32 InMaskG = 0, int32 InMaskB = 0, int32 InMaskA = 0)
 		: OutputName(InOutputName)
-#if WITH_EDITORONLY_DATA
 		, Mask(InMask)
 		, MaskR(InMaskR)
 		, MaskG(InMaskG)
 		, MaskB(InMaskB)
 		, MaskA(InMaskA)
-#endif
 	{}
 
 	/** Helper for setting component mask. */
 	ENGINE_API void SetMask(int32 UseMask, int32 R, int32 G, int32 B, int32 A)
 	{
-#if WITH_EDITORONLY_DATA
 		Mask = UseMask;
 		MaskR = R;
 		MaskG = G;
 		MaskB = B;
 		MaskA = A;
-#endif
 	}
 };
 
@@ -194,16 +170,16 @@ template<class InputType> struct FMaterialInput : FExpressionInput
 {
 	FMaterialInput()
 	{
-#if WITH_EDITORONLY_DATA
 		UseConstant = 0;
 		Constant = InputType(0);
-#endif
 	}
 
-#if WITH_EDITORONLY_DATA
+#if WITH_EDITOR
+	bool IsConstant() const { return UseConstant; }
+#endif
+
 	uint32	UseConstant : 1;
 	InputType	Constant;
-#endif
 };
 
 struct FColorMaterialInput : FMaterialInput<FColor>
@@ -331,13 +307,13 @@ struct FMaterialAttributesInput : FExpressionInput
 
 #if WITH_EDITOR
 	ENGINE_API int32 CompileWithDefault(class FMaterialCompiler* Compiler, const FGuid& AttributeID);
-#endif  // WITH_EDITOR
-	ENGINE_API bool IsConnected(EMaterialProperty Property) { return ((PropertyConnectedBitmask >> (uint32)Property) & 0x1) != 0; }
-	ENGINE_API bool IsConnected() const { return FExpressionInput::IsConnected(); }
-	ENGINE_API void SetConnectedProperty(EMaterialProperty Property, bool bIsConnected) 
+	inline bool IsConnected(EMaterialProperty Property) { return ((PropertyConnectedBitmask >> (uint32)Property) & 0x1) != 0; }
+	inline bool IsConnected() const { return FExpressionInput::IsConnected(); }
+	inline void SetConnectedProperty(EMaterialProperty Property, bool bIsConnected)
 	{
 		PropertyConnectedBitmask = bIsConnected ? PropertyConnectedBitmask | (1 << (uint32)Property) : PropertyConnectedBitmask & ~(1 << (uint32)Property);
 	}
+#endif  // WITH_EDITOR
 
 	/** ICPPStructOps interface */
 	ENGINE_API bool Serialize(FArchive& Ar);
diff --git a/Engine/Source/Runtime/Engine/Public/MaterialHLSLTree.h b/Engine/Source/Runtime/Engine/Public/MaterialHLSLTree.h
--- a/Engine/Source/Runtime/Engine/Public/MaterialHLSLTree.h
+++ b/Engine/Source/Runtime/Engine/Public/MaterialHLSLTree.h
@@ -326,12 +326,12 @@ public:
 class FExpressionMaterialLayers : public FExpressionForward
 {
 public:
-	FExpressionMaterialLayers(const FExpression* InExpression, const FMaterialLayersFunctions* InMaterialLayers)
+	FExpressionMaterialLayers(const FExpression* InExpression, const FMaterialLayersFunctions& InMaterialLayers)
 		: FExpressionForward(InExpression)
 		, MaterialLayers(InMaterialLayers)
 	{}
 
-	const FMaterialLayersFunctions* MaterialLayers;
+	FMaterialLayersFunctions MaterialLayers;
 
 	virtual bool PrepareValue(FEmitContext& Context, FEmitScope& Scope, const FRequestedType& RequestedType, FPrepareValueResult& OutResult) const override;
 };
diff --git a/Engine/Source/Runtime/Engine/Public/MaterialShared.h b/Engine/Source/Runtime/Engine/Public/MaterialShared.h
--- a/Engine/Source/Runtime/Engine/Public/MaterialShared.h
+++ b/Engine/Source/Runtime/Engine/Public/MaterialShared.h
@@ -1839,7 +1839,6 @@ public:
 	virtual FString GetAssetName() const { return GetFriendlyName(); }
 	virtual bool HasVertexPositionOffsetConnected() const { return false; }
 	virtual bool HasPixelDepthOffsetConnected() const { return false; }
-	virtual bool HasMaterialAttributesConnected() const { return false; }
 	virtual uint32 GetMaterialDecalResponse() const { return 0; }
 	virtual bool HasBaseColorConnected() const { return false; }
 	virtual bool HasNormalConnected() const { return false; }
@@ -1871,6 +1870,11 @@ public:
 	virtual bool IsPersistent() const = 0;
 	virtual UMaterialInterface* GetMaterialInterface() const { return NULL; }
 
+#if WITH_EDITOR
+	UE_DEPRECATED(5.1, "Should instead check individual properties")
+	virtual bool HasMaterialAttributesConnected() const { return false; }
+#endif
+
 	ENGINE_API const FMaterialCachedExpressionData& GetCachedExpressionData() const;
 
 	/** Is the material required to be complete? Default materials, special engine materials, etc */
@@ -2871,7 +2875,9 @@ protected:
 
 	ENGINE_API virtual bool HasVertexPositionOffsetConnected() const override;
 	ENGINE_API virtual bool HasPixelDepthOffsetConnected() const override;
+#if WITH_EDITOR
 	ENGINE_API virtual bool HasMaterialAttributesConnected() const override;
+#endif
 	/** Useful for debugging. */
 	ENGINE_API virtual FString GetBaseMaterialPathName() const override;
 	ENGINE_API virtual FString GetDebugName() const override;
diff --git a/Engine/Source/Runtime/Engine/Public/MaterialTypes.h b/Engine/Source/Runtime/Engine/Public/MaterialTypes.h
--- a/Engine/Source/Runtime/Engine/Public/MaterialTypes.h
+++ b/Engine/Source/Runtime/Engine/Public/MaterialTypes.h
@@ -410,6 +410,9 @@ struct FMaterialParameterMetadata
 	/** Description of the parameter, typically taken from the 'Desc' field of the parameter's UMaterialExpression */
 	FString Description;
 
+	/** Path to the asset (UMaterial or UMaterialFunction) that introduced this parameter */
+	FString AssetPath;
+
 	/** Name of the parameter's group */
 	FName Group;
 
diff --git a/Engine/Source/Runtime/Engine/Public/StaticParameterSet.h b/Engine/Source/Runtime/Engine/Public/StaticParameterSet.h
--- a/Engine/Source/Runtime/Engine/Public/StaticParameterSet.h
+++ b/Engine/Source/Runtime/Engine/Public/StaticParameterSet.h
@@ -15,6 +15,7 @@
 #include "StaticParameterSet.generated.h"
 
 class FSHA1;
+struct FStaticParameterSet;
 
 /**
 Base parameter properties
@@ -375,11 +376,40 @@ struct FStaticMaterialLayersParameter : public FStaticParameterBase
 #endif // WITH_EDITOR
 };
 
-/** Contains all the information needed to identify a single permutation of static parameters. */
 USTRUCT()
-struct FStaticParameterSet
+struct FStaticParameterSetRuntimeData
 {
-	GENERATED_USTRUCT_BODY();
+	GENERATED_BODY()
+
+	FStaticParameterSetRuntimeData() : bHasMaterialLayers(false) {}
+	FStaticParameterSetRuntimeData(const FStaticParameterSetRuntimeData& Rhs) = default;
+	FStaticParameterSetRuntimeData(const FStaticParameterSet& Rhs) = delete;
+	FStaticParameterSetRuntimeData& operator=(const FStaticParameterSetRuntimeData& Rhs) = default;
+	FStaticParameterSetRuntimeData& operator=(const FStaticParameterSet& Rhs) = delete;
+
+	/** Material layers for this set */
+	UPROPERTY()
+	FMaterialLayersFunctionsRuntimeData MaterialLayers;
+
+	UPROPERTY()
+	uint8 bHasMaterialLayers : 1;
+
+	void Empty()
+	{
+		MaterialLayers.Empty();
+		bHasMaterialLayers = false;
+	}
+
+	bool IsEmpty() const
+	{
+		return !bHasMaterialLayers;
+	}
+};
+
+USTRUCT()
+struct FStaticParameterSetEditorOnlyData
+{
+	GENERATED_BODY()
 
 	/** An array of static switch parameters in this set */
 	UPROPERTY()
@@ -395,14 +425,44 @@ struct FStaticParameterSet
 
 	/** Material layers for this set */
 	UPROPERTY()
-	FMaterialLayersFunctions MaterialLayers;
+	FMaterialLayersFunctionsEditorOnlyData MaterialLayers;
+
+#if WITH_EDITOR
+	void Empty()
+	{
+		StaticSwitchParameters.Empty();
+		StaticComponentMaskParameters.Empty();
+		TerrainLayerWeightParameters.Empty();
+		MaterialLayers.Empty();
+	}
 
+	bool IsEmpty() const
+	{
+		return StaticSwitchParameters.Num() == 0 && StaticComponentMaskParameters.Num() == 0 && TerrainLayerWeightParameters.Num() == 0;
+	}
+#endif // WITH_EDITOR
+};
+
+/** Contains all the information needed to identify a single permutation of static parameters. */
+USTRUCT()
+struct FStaticParameterSet : public FStaticParameterSetRuntimeData
+{
+	GENERATED_USTRUCT_BODY();
+
+#if WITH_EDITORONLY_DATA
 	UPROPERTY()
-	uint8 bHasMaterialLayers : 1;
+	FStaticParameterSetEditorOnlyData EditorOnly;
+#endif // WITH_EDITORONLY_DATA
 
-	FStaticParameterSet() : bHasMaterialLayers(false) {}
+	FStaticParameterSet() = default;
 	ENGINE_API FStaticParameterSet(const FStaticParameterSet& InValue);
+	FStaticParameterSet(const FStaticParameterSetRuntimeData&) = delete;
+
 	ENGINE_API FStaticParameterSet& operator=(const FStaticParameterSet& InValue);
+	FStaticParameterSet& operator=(const FStaticParameterSetRuntimeData& InValue) = delete;
+
+	FStaticParameterSetRuntimeData& GetRuntime() { return *this; }
+	const FStaticParameterSetRuntimeData& GetRuntime() const { return *this; }
 
 	/** 
 	* Checks if this set contains any parameters
@@ -411,17 +471,37 @@ struct FStaticParameterSet
 	*/
 	bool IsEmpty() const
 	{
-		return StaticSwitchParameters.Num() == 0 && StaticComponentMaskParameters.Num() == 0 && TerrainLayerWeightParameters.Num() == 0 && !bHasMaterialLayers;
+		if (!FStaticParameterSetRuntimeData::IsEmpty()) return false;
+#if WITH_EDITOR
+		if (!EditorOnly.IsEmpty()) return false;
+#endif
+		return true;
+	}
+
+	void Empty()
+	{
+		FStaticParameterSetRuntimeData::Empty();
+#if WITH_EDITOR
+		EditorOnly.Empty();
+#endif
 	}
 
-	void Empty();
+	bool GetMaterialLayers(FMaterialLayersFunctions& OutMaterialLayers) const;
 
 #if WITH_EDITOR
 	void SerializeLegacy(FArchive& Ar);
 	void UpdateLegacyTerrainLayerWeightData();
 	void UpdateLegacyMaterialLayersData();
+
+	static void Validate(const FStaticParameterSetRuntimeData& Runtime, const FStaticParameterSetEditorOnlyData& EditorOnly);
+	void Validate()
+	{
+		Validate(GetRuntime(), EditorOnly);
+	}
 #endif // WITH_EDITOR
 
+	void PostSerialize(const FArchive& Ar);
+
 	/** 
 	* Tests this set against another for equality
 	* 
@@ -448,6 +528,18 @@ private:
 	UPROPERTY()
 	TArray<FStaticMaterialLayersParameter> MaterialLayersParameters_DEPRECATED;
 	PRAGMA_ENABLE_DEPRECATION_WARNINGS
+
+	UPROPERTY()
+	FMaterialLayersFunctions MaterialLayers_DEPRECATED;
+
+	UPROPERTY()
+	TArray<FStaticSwitchParameter> StaticSwitchParameters_DEPRECATED;
+
+	UPROPERTY()
+	TArray<FStaticComponentMaskParameter> StaticComponentMaskParameters_DEPRECATED;
+
+	UPROPERTY()
+	TArray<FStaticTerrainLayerWeightParameter> TerrainLayerWeightParameters_DEPRECATED;
 #endif // WITH_EDITORONLY_DATA
 
 	void SortForEquivalent();
@@ -457,3 +549,14 @@ private:
 	void SetStaticComponentMaskParameterValue(const FMaterialParameterInfo& ParameterInfo, const FGuid& ExpressionGuid, bool R, bool G, bool B, bool A);
 #endif // WITH_EDITORONLY_DATA
 };
+
+bool operator==(const FStaticParameterSet&, const FStaticParameterSetRuntimeData&) = delete;
+bool operator==(const FStaticParameterSetRuntimeData&, const FStaticParameterSet&) = delete;
+bool operator!=(const FStaticParameterSet&, const FStaticParameterSetRuntimeData&) = delete;
+bool operator!=(const FStaticParameterSetRuntimeData&, const FStaticParameterSet&) = delete;
+
+template<>
+struct TStructOpsTypeTraits<FStaticParameterSet> : TStructOpsTypeTraitsBase2<FStaticParameterSet>
+{
+	enum { WithPostSerialize = true };
+};
diff --git a/Engine/Source/Runtime/Landscape/Private/LandscapeEdit.cpp b/Engine/Source/Runtime/Landscape/Private/LandscapeEdit.cpp
--- a/Engine/Source/Runtime/Landscape/Private/LandscapeEdit.cpp
+++ b/Engine/Source/Runtime/Landscape/Private/LandscapeEdit.cpp
@@ -201,7 +201,7 @@ ULandscapeMaterialInstanceConstant* ALandscapeProxy::GetLayerThumbnailMIC(UMater
 	MaterialInstance->GetStaticParameterValues(StaticParameters);
 
 	// Customize that material instance to only enable our terrain layer's weightmap : 
-	StaticParameters.TerrainLayerWeightParameters.Add(FStaticTerrainLayerWeightParameter(LayerName, /*InWeightmapIndex = */0, /*bInWeightBasedBlend = */false));
+	StaticParameters.EditorOnly.TerrainLayerWeightParameters.Add(FStaticTerrainLayerWeightParameter(LayerName, /*InWeightmapIndex = */0, /*bInWeightBasedBlend = */false));
 
 	// Don't recreate the render state of everything, only update the materials context
 	{
@@ -226,7 +226,7 @@ bool ULandscapeComponent::ValidateCombinationMaterial(UMaterialInstanceConstant*
 		return false;
 	}
 
-	const TArray<FStaticTerrainLayerWeightParameter>& TerrainLayerWeightParameters = InCombinationMaterial->GetStaticParameters().TerrainLayerWeightParameters;
+	const TArray<FStaticTerrainLayerWeightParameter>& TerrainLayerWeightParameters = InCombinationMaterial->GetEditorOnlyStaticParameters().TerrainLayerWeightParameters;
 	const TArray<FWeightmapLayerAllocationInfo>& ComponentWeightmapAllocations = GetWeightmapLayerAllocations();
 
 	if (TerrainLayerWeightParameters.Num() != ComponentWeightmapAllocations.Num())
@@ -361,7 +361,7 @@ UMaterialInstanceConstant* ULandscapeComponent::GetCombinationMaterial(FMaterial
 				{
 					FName LayerName = Allocation.GetLayerName();
 					check(LayerName != NAME_None);
-					StaticParameters.TerrainLayerWeightParameters.Add(FStaticTerrainLayerWeightParameter(LayerName, Allocation.WeightmapTextureIndex, !Allocation.LayerInfo->bNoWeightBlend));
+					StaticParameters.EditorOnly.TerrainLayerWeightParameters.Add(FStaticTerrainLayerWeightParameter(LayerName, Allocation.WeightmapTextureIndex, !Allocation.LayerInfo->bNoWeightBlend));
 				}
 			}
 
@@ -2494,9 +2494,13 @@ const TArray<FName>& ALandscapeProxy::GetLayersFromMaterial(UMaterialInterface*
 {
 	if (MaterialInterface)
 	{
-		return MaterialInterface->GetCachedExpressionData().LandscapeLayerNames;
+		const FMaterialCachedExpressionData& CachedExpressionData = MaterialInterface->GetCachedExpressionData();
+		if (CachedExpressionData.EditorOnlyData)
+		{
+			return CachedExpressionData.EditorOnlyData->LandscapeLayerNames;
+		}
 	}
-	return FMaterialCachedExpressionData::EmptyData.LandscapeLayerNames;
+	return FMaterialCachedExpressionEditorOnlyData::EmptyData.LandscapeLayerNames;
 }
 
 const TArray<FName>& ALandscapeProxy::GetLayersFromMaterial() const
diff --git a/Engine/Source/Runtime/Landscape/Private/LandscapeGrass.cpp b/Engine/Source/Runtime/Landscape/Private/LandscapeGrass.cpp
--- a/Engine/Source/Runtime/Landscape/Private/LandscapeGrass.cpp
+++ b/Engine/Source/Runtime/Landscape/Private/LandscapeGrass.cpp
@@ -904,7 +904,7 @@ public:
 };
 
 FLandscapeComponentGrassData::FLandscapeComponentGrassData(ULandscapeComponent* Component)
-	: RotationForWPO(Component->GetLandscapeMaterial()->GetMaterial()->WorldPositionOffset.IsConnected() ? Component->GetComponentTransform().GetRotation() : FQuat(0, 0, 0, 0))
+	: RotationForWPO(Component->GetLandscapeMaterial()->GetMaterial()->IsPropertyConnected(MP_WorldPositionOffset) ? Component->GetComponentTransform().GetRotation() : FQuat(0, 0, 0, 0))
 	, bIsDirty(false)
 {
 	UMaterialInterface* Material = Component->GetLandscapeMaterial();
@@ -956,7 +956,7 @@ bool ULandscapeComponent::IsGrassMapOutdated() const
 			return true;
 		}
 
-		FQuat RotationForWPO = GetLandscapeMaterial()->GetMaterial()->WorldPositionOffset.IsConnected() ? GetComponentTransform().GetRotation() : FQuat(0, 0, 0, 0);
+		FQuat RotationForWPO = GetLandscapeMaterial()->GetMaterial()->IsPropertyConnected(MP_WorldPositionOffset) ? GetComponentTransform().GetRotation() : FQuat(0, 0, 0, 0);
 		if (GrassData->RotationForWPO != RotationForWPO)
 		{
 			return true;
diff --git a/Engine/Source/Runtime/Landscape/Private/LandscapeLight.cpp b/Engine/Source/Runtime/Landscape/Private/LandscapeLight.cpp
--- a/Engine/Source/Runtime/Landscape/Private/LandscapeLight.cpp
+++ b/Engine/Source/Runtime/Landscape/Private/LandscapeLight.cpp
@@ -427,7 +427,7 @@ void FLandscapeStaticLightingMesh::GetHeightmapData(int32 InLOD, int32 GeometryL
 	check(Info);
 
 	bool bUseRenderedWPO = LandscapeComponent->GetLandscapeProxy()->bUseMaterialPositionOffsetInStaticLighting &&
-	                       LandscapeComponent->GetLandscapeMaterial()->GetMaterial()->WorldPositionOffset.IsConnected();
+		LandscapeComponent->GetLandscapeMaterial()->GetMaterial()->IsPropertyConnected(MP_WorldPositionOffset);
 
 	HeightData.Empty(FMath::Square(NumVertices));
 	HeightData.AddUninitialized(FMath::Square(NumVertices));
diff --git a/Engine/Source/Runtime/Landscape/Private/LandscapeRender.cpp b/Engine/Source/Runtime/Landscape/Private/LandscapeRender.cpp
--- a/Engine/Source/Runtime/Landscape/Private/LandscapeRender.cpp
+++ b/Engine/Source/Runtime/Landscape/Private/LandscapeRender.cpp
@@ -3338,10 +3338,8 @@ void ULandscapeMaterialInstanceConstant::UpdateCachedTextureStreaming()
 	const UMaterial* Material = GetMaterial();
 	if (Material)
 	{
-		int32 NumExpressions = Material->Expressions.Num();
-		for (int32 ExpressionIndex = 0; ExpressionIndex < NumExpressions; ExpressionIndex++)
+		for (UMaterialExpression* Expression : Material->GetExpressions())
 		{
-			UMaterialExpression* Expression = Material->Expressions[ExpressionIndex];
 			UMaterialExpressionTextureSample* TextureSample = Cast<UMaterialExpressionTextureSample>(Expression);
 
 			// TODO: This is only works for direct Coordinate Texture Sample cases
