Blueprint support for FInputDeviceId and FPlatformUserId.

- Expose Platform User IDs and Input Device IDs to blueprints and other UPROPERTY uses allowing users to call functions from blueprints.
- Add events on the game instance that you can hook onto in blueprints
- Add accessors to the PC

diff --git a/Engine/Source/Runtime/CoreUObject/Public/UObject/NoExportTypes.h b/Engine/Source/Runtime/CoreUObject/Public/UObject/NoExportTypes.h
--- a/Engine/Source/Runtime/CoreUObject/Public/UObject/NoExportTypes.h
+++ b/Engine/Source/Runtime/CoreUObject/Public/UObject/NoExportTypes.h
@@ -2005,6 +2005,77 @@ struct FAutomationExecutionEntry
 	FDateTime Timestamp;
 };
 
+
+/**
+ * Represents a single input device such as a gamepad, keyboard, or mouse.
+ *
+ * Has a globally unique identifier.
+ * 
+ * Opaque struct for the FInputDeviceId struct defined in CoreMiscDefines.h
+ */
+USTRUCT(noexport, BlueprintType)
+struct FInputDeviceId
+{
+	GENERATED_BODY()
+private:
+	
+	UPROPERTY(VisibleAnywhere, Category = "PlatformInputDevice")
+	int32 InternalId = -1;
+};
+
+/**
+ * Handle that defines a local user on this platform.
+ * This used to be just a typedef int32 that was used interchangeably as ControllerId and LocalUserIndex.
+ * Moving forward these will be allocated by the platform application layer.
+ *
+ * Opaque struct for the FPlatformUserId struct defined in CoreMiscDefines.h
+ */
+USTRUCT(noexport, BlueprintType)
+struct FPlatformUserId
+{
+	GENERATED_BODY()
+private:
+	
+	UPROPERTY(VisibleAnywhere, Category = "PlatformInputDevice")
+	int32 InternalId = -1;
+};
+
+/**
+ * Represents the connection status of a given FInputDeviceId
+ */
+UENUM(BlueprintType)
+enum class EInputDeviceConnectionState : uint8
+{
+	/** This is not a valid input device */
+	Invalid,
+
+	/** It is not known if this device is connected or not */
+	Unknown,
+
+	/** Device is definitely connected */
+	Disconnected,
+
+	/** Definitely connected and powered on */
+	Connected
+};
+
+/**
+ * Data about an input device's current state
+ */
+USTRUCT(noexport, BlueprintType)
+struct FPlatformInputDeviceState
+{
+	GENERATED_BODY()
+
+	/** The platform user that this input device belongs to */
+	UPROPERTY(VisibleAnywhere, Category = "PlatformInputDevice")
+	FPlatformUserId OwningPlatformUser = PLATFORMUSERID_NONE;
+
+	/** The connection state of this input device */
+	UPROPERTY(VisibleAnywhere, Category = "PlatformInputDevice")
+	EInputDeviceConnectionState ConnectionState = EInputDeviceConnectionState::Invalid;
+};
+
 /** Enum used by DataValidation plugin to see if an asset has been validated for correctness (mirrored in UObjectGlobals.h)*/
 UENUM(BlueprintType)
 enum class EDataValidationResult : uint8
diff --git a/Engine/Source/Runtime/Engine/Classes/Engine/GameInstance.h b/Engine/Source/Runtime/Engine/Classes/Engine/GameInstance.h
--- a/Engine/Source/Runtime/Engine/Classes/Engine/GameInstance.h
+++ b/Engine/Source/Runtime/Engine/Classes/Engine/GameInstance.h
@@ -59,6 +59,10 @@ DECLARE_MULTICAST_DELEGATE_ThreeParams(FOnPreClientTravel, const FString& /*Pend
 typedef FOnPreClientTravel::FDelegate FOnPreClientTravelDelegate;
 DECLARE_DYNAMIC_MULTICAST_DELEGATE_TwoParams(FOnPawnControllerChanged, APawn*, Pawn, AController*, Controller);
 
+DECLARE_DYNAMIC_MULTICAST_DELEGATE_ThreeParams(FOnUserInputDeviceConnectionChange, EInputDeviceConnectionState, NewConnectionState, FPlatformUserId, PlatformUserId, FInputDeviceId, InputDeviceId);
+DECLARE_DYNAMIC_MULTICAST_DELEGATE_ThreeParams(FOnUserInputDevicePairingChange, FInputDeviceId, InputDeviceId, FPlatformUserId, NewUserPlatformId, FPlatformUserId, OldUserPlatformId);
+
+
 #if WITH_EDITOR
 
 // The result of a UGameInstance PIE operation
@@ -138,6 +142,8 @@ struct FGameInstancePIEParameters
 
 #endif
 
+enum class EInputDeviceConnectionState : uint8;
+
 DECLARE_EVENT_OneParam(UGameInstance, FOnLocalPlayerEvent, ULocalPlayer*);
 
 /**
@@ -224,6 +230,38 @@ public:
 	UFUNCTION(BlueprintImplementableEvent, meta=(DisplayName = "Shutdown"))
 	void ReceiveShutdown();
 
+	/**
+	 * Callback for handling an Input Device's connection state change.
+	 *
+	 * @param NewConnectionState	The new connection state of this device
+	 * @param FPlatformUserId		The User ID whose input device has changed
+	 * @param FInputDeviceId		The Input Device ID that has changed connection
+	 * @see IPlatformInputDeviceMapper
+	 */
+	virtual void HandleOnInputDeviceConnectionChange(EInputDeviceConnectionState NewConnectionState, FPlatformUserId PlatformUserId, FInputDeviceId InputDeviceId);
+
+	/** 
+	 * Callback for when an input device connection state has changed (a new gamepad was connected or disconnected)
+	 */
+	UPROPERTY(BlueprintAssignable, DisplayName=OnInputDeviceConnectionChange)
+	FOnUserInputDeviceConnectionChange OnInputDeviceConnectionChange;
+	
+	/**
+	 * Callback for handling an Input Device pairing change.
+	 * 
+	 * @param FInputDeviceId	Input device ID
+	 * @param FPlatformUserId	The NewUserPlatformId
+	 * @param FPlatformUserId	The OldUserPlatformId
+	 * @see IPlatformInputDeviceMapper
+	 */
+	virtual void HandleOnInputDevicePairingChange(FInputDeviceId InputDeviceId, FPlatformUserId NewUserPlatformId, FPlatformUserId OldUserPlatformId);
+
+	/**
+	 * Callback when an input device has changed pairings (the owning platform user has changed for that device)
+	 */
+	UPROPERTY(BlueprintAssignable, DisplayName=OnUserInputDevicePairingChange)
+	FOnUserInputDevicePairingChange OnUserInputDevicePairingChange;
+	
 	/** Opportunity for blueprints to handle network errors. */
 	UFUNCTION(BlueprintImplementableEvent, meta=(DisplayName = "NetworkError"))
 	void HandleNetworkError(ENetworkFailure::Type FailureType, bool bIsServer);
diff --git a/Engine/Source/Runtime/Engine/Classes/GameFramework/PlatformInputDeviceMapperLibrary.h b/Engine/Source/Runtime/Engine/Classes/GameFramework/PlatformInputDeviceMapperLibrary.h
new file mode 100644
--- /dev/null
+++ b/Engine/Source/Runtime/Engine/Classes/GameFramework/PlatformInputDeviceMapperLibrary.h
@@ -0,0 +1,113 @@
+// Copyright Epic Games, Inc. All Rights Reserved.
+
+#pragma once
+
+#include "CoreMinimal.h"
+#include "Kismet/BlueprintFunctionLibrary.h"
+
+#include "PlatformInputDeviceMapperLibrary.generated.h"
+
+enum class EInputDeviceConnectionState : uint8;
+
+/**
+ * A static BP library that exposes the functionality of IPlatformInputDeviceMapper to blueprints.
+ *
+ * @see IPlatformInputDeviceMapper
+ * @note Keep any function comments up to date with those in GenericPlatformInputDeviceMapper.h!
+ */
+UCLASS()
+class ENGINE_API UPlatformInputDeviceMapperLibrary : public UBlueprintFunctionLibrary
+{
+	GENERATED_BODY()
+
+public:
+
+	/**
+	 * Populates the OutInputDevices array with any InputDeviceID's that are mapped to the given platform user
+	 *
+	 * @param UserId				The Platform User to gather the input devices of.
+	 * @param OutInputDevices		Array of input device ID's that will be populated with the mapped devices.
+	 * @return						The number of mapped devices, INDEX_NONE if the user was not found.
+	 */
+	UFUNCTION(BlueprintCallable, Category = "PlatformInputDevice")
+	static int32 GetAllInputDevicesForUser(const FPlatformUserId UserId, TArray<FInputDeviceId>& OutInputDevices);
+
+	/**
+	 * Get all mapped input devices on this platform regardless of their connection state.
+	 * 
+	 * @param OutInputDevices	Array of input devices to populate
+	 * @return					The number of connected input devices
+	 */
+	UFUNCTION(BlueprintCallable, Category = "PlatformInputDevice")
+	static int32 GetAllInputDevices(TArray<FInputDeviceId>& OutInputDevices);
+
+	/**
+	 * Gather all currently connected input devices
+	 * 
+	 * @param OutInputDevices	Array of input devices to populate
+	 * @return					The number of connected input devices
+	 */
+	UFUNCTION(BlueprintCallable, Category = "PlatformInputDevice")
+	static int32 GetAllConnectedInputDevices(TArray<FInputDeviceId>& OutInputDevices);
+
+	/**
+	 * Get all currently active platform ids, anyone who has a mapped input device
+	 *
+	 * @param OutUsers		Array that will be populated with the platform users.
+	 * @return				The number of active platform users
+	 */
+	UFUNCTION(BlueprintCallable, Category = "PlatformInputDevice")
+	static int32 GetAllActiveUsers(TArray<FPlatformUserId>& OutUsers);
+
+	/**
+	 * Returns the platform user id that is being used for unmapped input devices.
+	 * Will be PLATFORMUSERID_NONE if platform does not support this (this is the default behavior)
+	 */
+	UFUNCTION(BlueprintPure, Category = "PlatformInputDevice")
+	static FPlatformUserId GetUserForUnpairedInputDevices();
+
+	/** Returns true if the given Platform User Id is the user for unpaired input devices on this platform. */
+	UFUNCTION(BlueprintPure, Category = "PlatformInputDevice")
+	static bool IsUnpairedUserId(const FPlatformUserId PlatformId);
+
+	/** Returns true if the given input device is mapped to the unpaired platform user id. */
+	UFUNCTION(BlueprintPure, Category = "PlatformInputDevice")
+	static bool IsInputDeviceMappedToUnpairedUser(const FInputDeviceId InputDevice);
+
+	/** Returns the default device id used for things like keyboard/mouse input */
+	UFUNCTION(BlueprintPure, Category = "PlatformInputDevice")
+	static FInputDeviceId GetDefaultInputDevice();
+
+	/** Returns the platform user attached to this input device, or PLATFORMUSERID_NONE if invalid */
+	UFUNCTION(BlueprintCallable, Category = "PlatformInputDevice")
+	static FPlatformUserId GetUserForInputDevice(FInputDeviceId DeviceId);
+
+	/** Returns the primary input device used by a specific player, or INPUTDEVICEID_NONE if invalid */
+	UFUNCTION(BlueprintCallable, Category = "PlatformInputDevice")
+	static FInputDeviceId GetPrimaryInputDeviceForUser(FPlatformUserId UserId);
+
+	/**
+	 * Gets the connection state of the given input device.
+	 * 
+	 * @param DeviceId		The device to get the connection state of
+	 * @return				The connection state of the given device. EInputDeviceConnectionState::Unknown if the device is not mapped
+	 */
+	UFUNCTION(BlueprintCallable, Category = "PlatformInputDevice")
+	static EInputDeviceConnectionState GetInputDeviceConnectionState(const FInputDeviceId DeviceId);
+
+	/**
+	 * Check if the given input device is valid
+	 * 
+	 * @return	True if the given input device is valid (it has been assigned an ID by the PlatformInputDeviceMapper)
+	 */
+	UFUNCTION(BlueprintPure, Category = "PlatformInputDevice")
+	static bool IsValidInputDevice(FInputDeviceId DeviceId);
+
+	/**
+	 * Check if the given platform ID is valid
+	 * 
+	 * @return	True if the platform ID is valid (it has been assigned by the PlatformInputDeviceMapper)
+	 */
+	UFUNCTION(BlueprintPure, Category = "PlatformInputDevice")
+	static bool IsValidPlatformId(FPlatformUserId UserId);
+};
\ No newline at end of file
diff --git a/Engine/Source/Runtime/Engine/Classes/GameFramework/PlayerController.h b/Engine/Source/Runtime/Engine/Classes/GameFramework/PlayerController.h
--- a/Engine/Source/Runtime/Engine/Classes/GameFramework/PlayerController.h
+++ b/Engine/Source/Runtime/Engine/Classes/GameFramework/PlayerController.h
@@ -1675,6 +1675,13 @@ public:
 	/** Returns the ULocalPlayer for this controller if it exists, or null otherwise */
 	class ULocalPlayer* GetLocalPlayer() const;
 
+	/**
+	 * Returns the platform user that is assigned to this Player Controller's Local Player.
+	 * If there is no local player, then this will return PLATFORMUSERID_NONE
+	 */
+	UFUNCTION(BlueprintCallable, BlueprintPure, Category = "Game|Player")
+	FPlatformUserId GetPlatformUserId() const;
+
 	/**
 	 * Called client-side to smoothly interpolate received TargetViewRotation (result is in BlendedTargetViewRotation)
 	 * @param TargetPawn   is the pawn which is the current ViewTarget
diff --git a/Engine/Source/Runtime/Engine/Classes/Kismet/KismetStringLibrary.h b/Engine/Source/Runtime/Engine/Classes/Kismet/KismetStringLibrary.h
--- a/Engine/Source/Runtime/Engine/Classes/Kismet/KismetStringLibrary.h
+++ b/Engine/Source/Runtime/Engine/Classes/Kismet/KismetStringLibrary.h
@@ -68,6 +68,14 @@ class ENGINE_API UKismetStringLibrary : public UBlueprintFunctionLibrary
 	UFUNCTION(BlueprintPure, meta=(DisplayName = "To String (Object)", CompactNodeTitle = "->", BlueprintAutocast), Category="Utilities|String")
 	static FString Conv_ObjectToString(class UObject* InObj);
 
+	/** Converts a InputDeviceId value to a string */
+	UFUNCTION(BlueprintPure, meta=(DisplayName = "To String (InputDeviceId)", CompactNodeTitle = "->", BlueprintAutocast), Category="Utilities|String")
+	static FString Conv_InputDeviceIdToString(FInputDeviceId InDeviceId);
+
+	/** Converts a PlatformUserId value to a string */
+	UFUNCTION(BlueprintPure, meta=(DisplayName = "To String (PlatformUserId)", CompactNodeTitle = "->", BlueprintAutocast), Category="Utilities|String")
+	static FString Conv_PlatformUserIdToString(FPlatformUserId InPlatformUserId);
+
 	/** Converts a linear color value to a string, in the form '(R=,G=,B=,A=)' */
 	UFUNCTION(BlueprintPure, meta=(DisplayName = "To String (LinearColor)", CompactNodeTitle = "->", BlueprintAutocast), Category="Utilities|String")
 	static FString Conv_ColorToString(FLinearColor InColor);
diff --git a/Engine/Source/Runtime/Engine/Private/GameInstance.cpp b/Engine/Source/Runtime/Engine/Private/GameInstance.cpp
--- a/Engine/Source/Runtime/Engine/Private/GameInstance.cpp
+++ b/Engine/Source/Runtime/Engine/Private/GameInstance.cpp
@@ -30,6 +30,7 @@
 #include "Misc/PackageName.h"
 #include "Net/ReplayPlaylistTracker.h"
 #include "ReplaySubsystem.h"
+#include "GenericPlatform/GenericPlatformInputDeviceMapper.h"
 
 #if WITH_EDITOR
 #include "Settings/LevelEditorPlaySettings.h"
@@ -97,6 +98,10 @@ void UGameInstance::Init()
 
 		FNetDelegates::OnReceivedNetworkEncryptionToken.BindUObject(this, &ThisClass::ReceivedNetworkEncryptionToken);
 		FNetDelegates::OnReceivedNetworkEncryptionAck.BindUObject(this, &ThisClass::ReceivedNetworkEncryptionAck);
+
+		IPlatformInputDeviceMapper& PlatformInputMapper = IPlatformInputDeviceMapper::Get();
+		PlatformInputMapper.GetOnInputDeviceConnectionChange().AddUObject(this, &UGameInstance::HandleOnInputDeviceConnectionChange);
+		PlatformInputMapper.GetOnInputDevicePairingChange().AddUObject(this, &UGameInstance::HandleOnInputDevicePairingChange);
 	}
 
 	SubsystemCollection.Initialize(this);
@@ -142,10 +147,24 @@ void UGameInstance::Shutdown()
 	FNetDelegates::OnReceivedNetworkEncryptionToken.Unbind();
 	FNetDelegates::OnReceivedNetworkEncryptionAck.Unbind();
 
+	IPlatformInputDeviceMapper& PlatformInputMapper = IPlatformInputDeviceMapper::Get();
+	PlatformInputMapper.GetOnInputDeviceConnectionChange().RemoveAll(this);
+	PlatformInputMapper.GetOnInputDevicePairingChange().RemoveAll(this);
+
 	// Clear the world context pointer to prevent further access.
 	WorldContext = nullptr;
 }
 
+void UGameInstance::HandleOnInputDeviceConnectionChange(EInputDeviceConnectionState NewConnectionState, FPlatformUserId PlatformUserId, FInputDeviceId InputDeviceId)
+{
+	OnInputDeviceConnectionChange.Broadcast(NewConnectionState, PlatformUserId, InputDeviceId);
+}
+
+void UGameInstance::HandleOnInputDevicePairingChange(FInputDeviceId InputDeviceId, FPlatformUserId NewUserPlatformId, FPlatformUserId OldUserPlatformId)
+{
+	OnUserInputDevicePairingChange.Broadcast(InputDeviceId, NewUserPlatformId, OldUserPlatformId);
+}
+
 void UGameInstance::InitializeStandalone(const FName InPackageName, UPackage* InWorldPackage)
 {
 	// Creates the world context. This should be the only WorldContext that ever gets created for this GameInstance.
diff --git a/Engine/Source/Runtime/Engine/Private/KismetStringLibrary.cpp b/Engine/Source/Runtime/Engine/Private/KismetStringLibrary.cpp
--- a/Engine/Source/Runtime/Engine/Private/KismetStringLibrary.cpp
+++ b/Engine/Source/Runtime/Engine/Private/KismetStringLibrary.cpp
@@ -118,7 +118,17 @@ FString UKismetStringLibrary::Conv_TransformToString(const FTransform& InTrans)
 
 FString UKismetStringLibrary::Conv_ObjectToString(class UObject* InObj)
 {
-	return (InObj != NULL) ? InObj->GetName() : FString(TEXT("None"));
+	return (InObj != nullptr) ? InObj->GetName() : FString(TEXT("None"));
+}
+
+FString UKismetStringLibrary::Conv_InputDeviceIdToString(FInputDeviceId InDeviceId)
+{
+	return FString::Printf(TEXT("%d"), InDeviceId.GetId());
+}
+
+FString UKismetStringLibrary::Conv_PlatformUserIdToString(FPlatformUserId InPlatformUserId)
+{
+	return FString::Printf(TEXT("%d"), InPlatformUserId.GetInternalId());
 }
 
 FString UKismetStringLibrary::Conv_ColorToString(FLinearColor C)
diff --git a/Engine/Source/Runtime/Engine/Private/PlatformInputDeviceMapperLibrary.cpp b/Engine/Source/Runtime/Engine/Private/PlatformInputDeviceMapperLibrary.cpp
new file mode 100644
--- /dev/null
+++ b/Engine/Source/Runtime/Engine/Private/PlatformInputDeviceMapperLibrary.cpp
@@ -0,0 +1,69 @@
+// Copyright Epic Games, Inc. All Rights Reserved.
+
+#include "GameFramework/PlatformInputDeviceMapperLibrary.h"
+#include "GenericPlatform/GenericPlatformInputDeviceMapper.h"
+
+int32 UPlatformInputDeviceMapperLibrary::GetAllInputDevicesForUser(const FPlatformUserId UserId, TArray<FInputDeviceId>& OutInputDevices)
+{
+	return IPlatformInputDeviceMapper::Get().GetAllInputDevicesForUser(UserId, OutInputDevices);
+}
+
+int32 UPlatformInputDeviceMapperLibrary::GetAllInputDevices(TArray<FInputDeviceId>& OutInputDevices)
+{
+	return IPlatformInputDeviceMapper::Get().GetAllInputDevices(OutInputDevices);
+}
+
+int32 UPlatformInputDeviceMapperLibrary::GetAllConnectedInputDevices(TArray<FInputDeviceId>& OutInputDevices)
+{
+	return IPlatformInputDeviceMapper::Get().GetAllConnectedInputDevices(OutInputDevices);
+}
+
+int32 UPlatformInputDeviceMapperLibrary::GetAllActiveUsers(TArray<FPlatformUserId>& OutUsers)
+{
+	return IPlatformInputDeviceMapper::Get().GetAllActiveUsers(OutUsers);
+}
+
+FPlatformUserId UPlatformInputDeviceMapperLibrary::GetUserForUnpairedInputDevices()
+{
+	return IPlatformInputDeviceMapper::Get().GetUserForUnpairedInputDevices();
+}
+
+bool UPlatformInputDeviceMapperLibrary::IsUnpairedUserId(const FPlatformUserId PlatformId)
+{
+	return IPlatformInputDeviceMapper::Get().IsUnpairedUserId(PlatformId);
+}
+
+bool UPlatformInputDeviceMapperLibrary::IsInputDeviceMappedToUnpairedUser(const FInputDeviceId InputDevice)
+{
+	return IPlatformInputDeviceMapper::Get().IsInputDeviceMappedToUnpairedUser(InputDevice);
+}
+
+FInputDeviceId UPlatformInputDeviceMapperLibrary::GetDefaultInputDevice()
+{
+	return IPlatformInputDeviceMapper::Get().GetDefaultInputDevice();
+}
+
+FPlatformUserId UPlatformInputDeviceMapperLibrary::GetUserForInputDevice(FInputDeviceId DeviceId)
+{
+	return IPlatformInputDeviceMapper::Get().GetUserForInputDevice(DeviceId);
+}
+
+FInputDeviceId UPlatformInputDeviceMapperLibrary::GetPrimaryInputDeviceForUser(FPlatformUserId UserId)
+{
+	return IPlatformInputDeviceMapper::Get().GetPrimaryInputDeviceForUser(UserId);
+}
+
+EInputDeviceConnectionState UPlatformInputDeviceMapperLibrary::GetInputDeviceConnectionState(const FInputDeviceId DeviceId)
+{
+	return IPlatformInputDeviceMapper::Get().GetInputDeviceConnectionState(DeviceId);
+}
+
+bool UPlatformInputDeviceMapperLibrary::IsValidInputDevice(FInputDeviceId DeviceId)
+{
+	return DeviceId.IsValid();
+}
+
+bool UPlatformInputDeviceMapperLibrary::IsValidPlatformId(FPlatformUserId UserId)
+{
+	return UserId.IsValid();
+}
\ No newline at end of file
diff --git a/Engine/Source/Runtime/Engine/Private/PlayerController.cpp b/Engine/Source/Runtime/Engine/Private/PlayerController.cpp
--- a/Engine/Source/Runtime/Engine/Private/PlayerController.cpp
+++ b/Engine/Source/Runtime/Engine/Private/PlayerController.cpp
@@ -4694,6 +4694,15 @@ ULocalPlayer* APlayerController::GetLocalPlayer() const
 	return Cast<ULocalPlayer>(Player);
 }
 
+FPlatformUserId APlayerController::GetPlatformUserId() const
+{
+	if (const ULocalPlayer* LocalPlayer = GetLocalPlayer())
+	{
+		return LocalPlayer->GetPlatformUserId();
+	}
+	return PLATFORMUSERID_NONE;
+}
+
 bool APlayerController::IsInViewportClient(UGameViewportClient* ViewportClient) const
 {
 	const ULocalPlayer* LocalPlayer = GetLocalPlayer();
