Fix screenshot canvases not being FCanvas::Allow_DeleteOnRender.

diff --git a/Engine/Source/Runtime/Engine/Private/UnrealClient.cpp b/Engine/Source/Runtime/Engine/Private/UnrealClient.cpp
--- a/Engine/Source/Runtime/Engine/Private/UnrealClient.cpp
+++ b/Engine/Source/Runtime/Engine/Private/UnrealClient.cpp
@@ -1335,6 +1335,7 @@ void FViewport::HighResScreenshot()
 		Canvas.Flush_GameThread();
 
 		// Draw the debug canvas
+		DummyViewport->GetDebugCanvas()->SetAllowedModes(FCanvas::Allow_DeleteOnRender);
 		DummyViewport->GetDebugCanvas()->Flush_GameThread(true);
 
 		HighResScreenshotEndFrame(DummyViewport);
