PR #9081: Avoid animations being always dirtied (Contributed by [REDACTED_PERSON])

diff --git a/Engine/Source/Runtime/Engine/Private/Animation/AnimSequence.cpp b/Engine/Source/Runtime/Engine/Private/Animation/AnimSequence.cpp
--- a/Engine/Source/Runtime/Engine/Private/Animation/AnimSequence.cpp
+++ b/Engine/Source/Runtime/Engine/Private/Animation/AnimSequence.cpp
@@ -5301,6 +5301,8 @@ void UAnimSequence::OnModelModified(const EAnimDataModelNotifyType& NotifyType,
 		}
 	};
 
+	bool bShouldMarkPackageDirty = true;
+
 	PRAGMA_DISABLE_DEPRECATION_WARNINGS
 	switch (NotifyType)
 	{
@@ -5432,6 +5434,13 @@ void UAnimSequence::OnModelModified(const EAnimDataModelNotifyType& NotifyType,
 		{
 			const FCurveRenamedPayload& TypedPayload = Payload.GetPayload<FCurveRenamedPayload>();
 			UpdateCompressedCurveName(TypedPayload.Identifier.InternalName.UID, TypedPayload.NewIdentifier.InternalName);
+			if (TypedPayload.Identifier.InternalName.DisplayName == TypedPayload.NewIdentifier.InternalName.DisplayName
+				&& TypedPayload.Identifier.InternalName.UID == SmartName::MaxUID && TypedPayload.NewIdentifier.InternalName.UID != SmartName::MaxUID)
+			{
+				// We are renaming a Curve with a FSmartName::Invalid UID, this means it was just loaded and we are retrieving
+				// its actual UID from its DisplayName. This should *not* mark the package dirty.
+				bShouldMarkPackageDirty = false;
+			}
 			break;
 		}
 		default:
@@ -5441,8 +5450,11 @@ void UAnimSequence::OnModelModified(const EAnimDataModelNotifyType& NotifyType,
 		}
 	}
 	PRAGMA_ENABLE_DEPRECATION_WARNINGS
-	
-	MarkPackageDirty();
+
+	if (bShouldMarkPackageDirty)
+	{
+		MarkPackageDirty();
+	}
 }
 
 
