Fixed off by one bug in DisableAllReplicatedPropertiesOfClass() which did not properly disable all properties

diff --git a/Engine/Source/Runtime/Engine/Private/UnrealNetwork.cpp b/Engine/Source/Runtime/Engine/Private/UnrealNetwork.cpp
--- a/Engine/Source/Runtime/Engine/Private/UnrealNetwork.cpp
+++ b/Engine/Source/Runtime/Engine/Private/UnrealNetwork.cpp
@@ -186,7 +186,7 @@ void ResetReplicatedLifetimeProperty(const NetworkingPrivate::FRepPropertyDescri
 void DisableAllReplicatedPropertiesOfClass(const NetworkingPrivate::FRepClassDescriptor& ClassDescriptor, EFieldIteratorFlags::SuperClassFlags SuperClassBehavior, TArray<FLifetimeProperty>& OutLifetimeProps)
 {
 	const int32 StartIndex = (EFieldIteratorFlags::IncludeSuper == SuperClassBehavior) ? 0 : ClassDescriptor.StartRepIndex;
-	for (int32 RepIndex = StartIndex; RepIndex < ClassDescriptor.EndRepIndex; ++RepIndex)
+	for (int32 RepIndex = StartIndex; RepIndex <= ClassDescriptor.EndRepIndex; ++RepIndex)
 	{
 		FLifetimeProperty* RegisteredPropertyPtr = OutLifetimeProps.FindByPredicate([&RepIndex](const FLifetimeProperty& Var) { return Var.RepIndex == RepIndex; });
 
diff --git a/Engine/Source/Runtime/GameplayTasks/Private/GameplayTasksComponent.cpp b/Engine/Source/Runtime/GameplayTasks/Private/GameplayTasksComponent.cpp
--- a/Engine/Source/Runtime/GameplayTasks/Private/GameplayTasksComponent.cpp
+++ b/Engine/Source/Runtime/GameplayTasks/Private/GameplayTasksComponent.cpp
@@ -167,7 +167,7 @@ void UGameplayTasksComponent::OnTaskEnded(UGameplayTask& Task)
 void UGameplayTasksComponent::GetLifetimeReplicatedProps(TArray< FLifetimeProperty >& OutLifetimeProps) const
 {
 	// Intentionally not calling super: We do not want to replicate bActive which controls ticking. We sometimes need to tick on client predictively.
-	DISABLE_ALL_CLASS_REPLICATED_PROPERTIES(Super, EFieldIteratorFlags::IncludeSuper);
+	DISABLE_ALL_CLASS_REPLICATED_PROPERTIES_FAST(Super, EFieldIteratorFlags::IncludeSuper);
 
 	const FDoRepLifetimeParams Params{ COND_SkipOwner, REPNOTIFY_Always, true };
 	DOREPLIFETIME_WITH_PARAMS_FAST(UGameplayTasksComponent, SimulatedTasks, Params);
