Optimize copying of final pose to output by removing redundant reference pose copy. If results exist, we can traverse the bone index array in increasing order and fill in any gaps with the reference pose.

diff --git a/Engine/Source/Runtime/Engine/Private/Components/SkeletalMeshComponent.cpp b/Engine/Source/Runtime/Engine/Private/Components/SkeletalMeshComponent.cpp
--- a/Engine/Source/Runtime/Engine/Private/Components/SkeletalMeshComponent.cpp
+++ b/Engine/Source/Runtime/Engine/Private/Components/SkeletalMeshComponent.cpp
@@ -472,20 +472,38 @@ void USkeletalMeshComponent::UpdateClothTickRegisteredState()
 
 void USkeletalMeshComponent::FinalizePoseEvaluationResult(const USkeletalMesh* InMesh, TArray<FTransform>& OutBoneSpaceTransforms, FVector& OutRootBoneTranslation, FCompactPose& InFinalPose) const
 {
-	OutBoneSpaceTransforms = InMesh->GetRefSkeleton().GetRefBonePose();
-
+	const TArray<FTransform>& RefBonePose = InMesh->GetRefSkeleton().GetRefBonePose();
 	if(InFinalPose.IsValid() && InFinalPose.GetNumBones() > 0)
 	{
 		InFinalPose.NormalizeRotations();
 
-		for(const FCompactPoseBoneIndex BoneIndex : InFinalPose.ForEachBoneIndex())
+		// Bone index array is in increasing order, so we can fill gaps with reference pose
+		const auto FillReferencePose = [&](int32 BeginIndex, int32 EndIndex)
 		{
-			FMeshPoseBoneIndex MeshPoseIndex = InFinalPose.GetBoneContainer().MakeMeshPoseIndex(BoneIndex);
-			OutBoneSpaceTransforms[MeshPoseIndex.GetInt()] = InFinalPose[BoneIndex];
+			for (int32 MeshPoseIndex = BeginIndex; MeshPoseIndex < EndIndex; ++MeshPoseIndex)
+			{
+				OutBoneSpaceTransforms[MeshPoseIndex] = RefBonePose[MeshPoseIndex];
+			}
+		};
+
+		int32 LastMeshPoseIndex = 0;
+		const int32 BoneCount = RefBonePose.Num();
+		OutBoneSpaceTransforms.SetNum(BoneCount);
+		for (const FCompactPoseBoneIndex BoneIndex : InFinalPose.ForEachBoneIndex())
+		{
+			const int32 MeshPoseIndex = InFinalPose.GetBoneContainer().MakeMeshPoseIndex(BoneIndex).GetInt();
+			FillReferencePose(LastMeshPoseIndex, MeshPoseIndex);
+			OutBoneSpaceTransforms[MeshPoseIndex] = InFinalPose[BoneIndex];
+			LastMeshPoseIndex = MeshPoseIndex + 1;
 		}
+		FillReferencePose(LastMeshPoseIndex, BoneCount);
+	}
+	else
+	{
+		OutBoneSpaceTransforms = RefBonePose;
 	}
 
-	OutRootBoneTranslation = OutBoneSpaceTransforms[0].GetTranslation() - InMesh->GetRefSkeleton().GetRefBonePose()[0].GetTranslation();
+	OutRootBoneTranslation = OutBoneSpaceTransforms[0].GetTranslation() - RefBonePose[0].GetTranslation();
 }
 
 void USkeletalMeshComponent::FinalizeAttributeEvaluationResults(const FBoneContainer& BoneContainer, const UE::Anim::FHeapAttributeContainer& FinalContainer,
