World Partition: added wp.Runtime.DrawWorldPartitionIndex console command to select the current world partition to debug display.

#rnx

diff --git a/Engine/Source/Runtime/Engine/Private/WorldPartition/WorldPartitionSubsystem.cpp b/Engine/Source/Runtime/Engine/Private/WorldPartition/WorldPartitionSubsystem.cpp
--- a/Engine/Source/Runtime/Engine/Private/WorldPartition/WorldPartitionSubsystem.cpp
+++ b/Engine/Source/Runtime/Engine/Private/WorldPartition/WorldPartitionSubsystem.cpp
@@ -15,6 +15,18 @@
 extern int32 GBlockOnSlowStreaming;
 static const FName NAME_WorldPartitionRuntimeHash("WorldPartitionRuntimeHash");
 
+static int32 GDrawWorldPartitionIndex = 0;
+static FAutoConsoleCommand CVarDrawWorldPartitionIndex(
+	TEXT("wp.Runtime.DrawWorldPartitionIndex"),
+	TEXT("Sets the index of the wanted world partition to display debug draw."),
+	FConsoleCommandWithArgsDelegate::CreateLambda([](const TArray<FString>& Args)
+	{
+		if (Args.Num() == 1)
+		{
+			GDrawWorldPartitionIndex = FCString::Atoi(*Args[0]);
+		}
+	}));
+
 static int32 GDrawRuntimeHash3D = 0;
 static FAutoConsoleCommand CVarDrawRuntimeHash3D(
 	TEXT("wp.Runtime.ToggleDrawRuntimeHash3D"),
@@ -283,7 +295,7 @@ void UWorldPartitionSubsystem::UpdateStreamingState()
 void UWorldPartitionSubsystem::Draw(UCanvas* Canvas, class APlayerController* PC)
 {
 	TRACE_CPUPROFILER_EVENT_SCOPE(UWorldPartitionSubsystem::Draw);
-	if (!Canvas || !Canvas->SceneView)
+	if (!Canvas || !Canvas->SceneView || !RegisteredWorldPartitions.IsValidIndex(GDrawWorldPartitionIndex))
 	{
 		return;
 	}
@@ -295,6 +307,7 @@ void UWorldPartitionSubsystem::Draw(UCanvas* Canvas, class APlayerController* PC
 		return;
 	}
 
+	UWorldPartition* WorldPartition = RegisteredWorldPartitions[GDrawWorldPartitionIndex];
 	const FVector2D CanvasTopLeftPadding(10.f, 10.f);
 
 	FVector2D CurrentOffset(CanvasTopLeftPadding);
@@ -306,27 +319,10 @@ void UWorldPartitionSubsystem::Draw(UCanvas* Canvas, class APlayerController* PC
 		const FVector2D CanvasMinimumSize(100.f, 100.f);
 		const FVector2D CanvasMaxScreenSize = FVector2D::Max(MaxScreenRatio*FVector2D(Canvas->ClipX, Canvas->ClipY) - CanvasBottomRightPadding - CurrentOffset, CanvasMinimumSize);
 
-		FVector2D TotalFootprint(ForceInitToZero);
-		for (UWorldPartition* RegisteredWorldPartition : RegisteredWorldPartitions)
-		{
-			FVector2D Footprint = RegisteredWorldPartition->GetDrawRuntimeHash2DDesiredFootprint(CanvasMaxScreenSize);
-			TotalFootprint.X += Footprint.X;
-		}
-
-		if (TotalFootprint.X > 0.f)
-		{
-			for (UWorldPartition* RegisteredWorldPartition : RegisteredWorldPartitions)
-			{
-				FVector2D Footprint = RegisteredWorldPartition->GetDrawRuntimeHash2DDesiredFootprint(CanvasMaxScreenSize);
-				float FootprintRatio = Footprint.X / TotalFootprint.X;
-				FVector2D PartitionCanvasSize = FVector2D(CanvasMaxScreenSize.X * FootprintRatio, CanvasMaxScreenSize.Y);
-				RegisteredWorldPartition->DrawRuntimeHash2D(Canvas, PartitionCanvasSize, CurrentOffset);
-				CurrentOffset.X += PartitionCanvasSize.X;
-			}
-			
-			CurrentOffset.X = CanvasBottomRightPadding.X;
-			CurrentOffset.Y += CanvasMaxScreenSize.Y;
-		}
+		FVector2D PartitionCanvasSize = FVector2D(CanvasMaxScreenSize.X, CanvasMaxScreenSize.Y);
+		WorldPartition->DrawRuntimeHash2D(Canvas, PartitionCanvasSize, CurrentOffset);
+		CurrentOffset.X = CanvasBottomRightPadding.X;
+		CurrentOffset.Y += CanvasMaxScreenSize.Y;
 	}
 	
 	if (GDrawStreamingPerfs || GDrawRuntimeHash2D)
@@ -350,27 +346,24 @@ void UWorldPartitionSubsystem::Draw(UCanvas* Canvas, class APlayerController* PC
 
 		{
 			FString StatusText;
-			if (const UWorldPartition* WorldPartition = GetWorldPartition())
+			EWorldPartitionStreamingPerformance StreamingPerformance = WorldPartition->GetStreamingPerformance();
+			switch (StreamingPerformance)
 			{
-				EWorldPartitionStreamingPerformance StreamingPerformance = WorldPartition->GetStreamingPerformance();
-				switch (StreamingPerformance)
-				{
-				case EWorldPartitionStreamingPerformance::Good:
-					StatusText = TEXT("Good");
-					break;
-				case EWorldPartitionStreamingPerformance::Slow:
-					StatusText = TEXT("Slow");
-					break;
-				case EWorldPartitionStreamingPerformance::Critical:
-					StatusText = TEXT("Critical");
-					break;
-				default:
-					StatusText = TEXT("Unknown");
-					break;
-				}
-				const FString Text = FString::Printf(TEXT("Streaming Performance: %s (Blocking %s)"), *StatusText, GBlockOnSlowStreaming ? TEXT("Enabled") : TEXT("Disabled"));
-				FWorldPartitionDebugHelper::DrawText(Canvas, Text, GEngine->GetSmallFont(), FColor::White, CurrentOffset);
+			case EWorldPartitionStreamingPerformance::Good:
+				StatusText = TEXT("Good");
+				break;
+			case EWorldPartitionStreamingPerformance::Slow:
+				StatusText = TEXT("Slow");
+				break;
+			case EWorldPartitionStreamingPerformance::Critical:
+				StatusText = TEXT("Critical");
+				break;
+			default:
+				StatusText = TEXT("Unknown");
+				break;
 			}
+			const FString Text = FString::Printf(TEXT("Streaming Performance: %s (Blocking %s)"), *StatusText, GBlockOnSlowStreaming ? TEXT("Enabled") : TEXT("Disabled"));
+			FWorldPartitionDebugHelper::DrawText(Canvas, Text, GEngine->GetSmallFont(), FColor::White, CurrentOffset);
 		}
 	}
 
@@ -378,52 +371,43 @@ void UWorldPartitionSubsystem::Draw(UCanvas* Canvas, class APlayerController* PC
 	{
 		TRACE_CPUPROFILER_EVENT_SCOPE(UWorldPartitionSubsystem::DrawStreamingSources);
 
-		if (const UWorldPartition* WorldPartition = GetWorldPartition())
+		const TArray<FWorldPartitionStreamingSource>& StreamingSources = WorldPartition->GetStreamingSources();
+		if (StreamingSources.Num() > 0)
 		{
-			const TArray<FWorldPartitionStreamingSource>& StreamingSources = WorldPartition->GetStreamingSources();
-			if (StreamingSources.Num() > 0)
+			FString Title(TEXT("Streaming Sources"));
+			FWorldPartitionDebugHelper::DrawText(Canvas, Title, GEngine->GetSmallFont(), FColor::Yellow, CurrentOffset);
+
+			FVector2D Pos = CurrentOffset;
+			float MaxTextWidth = 0;
+			for (const FWorldPartitionStreamingSource& StreamingSource : StreamingSources)
 			{
-				FString Title(TEXT("Streaming Sources"));
-				FWorldPartitionDebugHelper::DrawText(Canvas, Title, GEngine->GetSmallFont(), FColor::Yellow, CurrentOffset);
-
-				FVector2D Pos = CurrentOffset;
-				float MaxTextWidth = 0;
-				for (const FWorldPartitionStreamingSource& StreamingSource : StreamingSources)
-				{
-					FString StreamingSourceDisplay = StreamingSource.Name.ToString();
-					FWorldPartitionDebugHelper::DrawText(Canvas, StreamingSourceDisplay, GEngine->GetSmallFont(), StreamingSource.GetDebugColor(), Pos, &MaxTextWidth);
-				}
-				Pos = CurrentOffset + FVector2D(MaxTextWidth + 10, 0.f);
-				for (const FWorldPartitionStreamingSource& StreamingSource : StreamingSources)
-				{
-					FWorldPartitionDebugHelper::DrawText(Canvas, *StreamingSource.ToString(), GEngine->GetSmallFont(), FColor::White, Pos);
-				}
-				CurrentOffset.Y = Pos.Y;
+				FString StreamingSourceDisplay = StreamingSource.Name.ToString();
+				FWorldPartitionDebugHelper::DrawText(Canvas, StreamingSourceDisplay, GEngine->GetSmallFont(), StreamingSource.GetDebugColor(), Pos, &MaxTextWidth);
+			}
+			Pos = CurrentOffset + FVector2D(MaxTextWidth + 10, 0.f);
+			for (const FWorldPartitionStreamingSource& StreamingSource : StreamingSources)
+			{
+				FWorldPartitionDebugHelper::DrawText(Canvas, *StreamingSource.ToString(), GEngine->GetSmallFont(), FColor::White, Pos);
 			}
+			CurrentOffset.Y = Pos.Y;
 		}
 	}
 
-	if (UWorldPartition* WorldPartition = GetWorldPartition())
-	{
-		UDataLayerSubsystem* DataLayerSubsystem = WorldPartition->GetWorld()->GetSubsystem<UDataLayerSubsystem>();
+	UDataLayerSubsystem* DataLayerSubsystem = WorldPartition->GetWorld()->GetSubsystem<UDataLayerSubsystem>();
 
-		if (GDrawLegends || GDrawRuntimeHash2D)
-		{
-			// Streaming Status Legend
-			WorldPartition->DrawStreamingStatusLegend(Canvas, CurrentOffset);
-		}
+	if (GDrawLegends || GDrawRuntimeHash2D)
+	{
+		// Streaming Status Legend
+		WorldPartition->DrawStreamingStatusLegend(Canvas, CurrentOffset);
+	}
 
-		if (DataLayerSubsystem && (GDrawDataLayers || GDrawDataLayersLoadTime || GDrawRuntimeHash2D))
-		{
-			DataLayerSubsystem->DrawDataLayersStatus(Canvas, CurrentOffset);
-		}
+	if (DataLayerSubsystem && (GDrawDataLayers || GDrawDataLayersLoadTime || GDrawRuntimeHash2D))
+	{
+		DataLayerSubsystem->DrawDataLayersStatus(Canvas, CurrentOffset);
 	}
 
 	if (GDrawRuntimeCellsDetails)
 	{
-		if (UWorldPartition* Partition = GetWorldPartition())
-		{
-			Partition->DrawRuntimeCellsDetails(Canvas, CurrentOffset);
-		}
+		WorldPartition->DrawRuntimeCellsDetails(Canvas, CurrentOffset);
 	}
 }
\ No newline at end of file
