[Backout] - CL19964485

diff --git a/Engine/Source/Runtime/Engine/Private/Materials/Material.cpp b/Engine/Source/Runtime/Engine/Private/Materials/Material.cpp
--- a/Engine/Source/Runtime/Engine/Private/Materials/Material.cpp
+++ b/Engine/Source/Runtime/Engine/Private/Materials/Material.cpp
@@ -918,6 +918,13 @@ UMaterial::UMaterial(const FObjectInitializer& ObjectInitializer)
 #if WITH_EDITORONLY_DATA
 	DiffuseColor_DEPRECATED.Constant = FColor(128,128,128);
 	SpecularColor_DEPRECATED.Constant = FColor(128,128,128);
+	BaseColor.Constant = FColor(128,128,128);	
+	Metallic.Constant = 0.0f;
+	Specular.Constant = 0.5f;
+	Roughness.Constant = 0.5f;
+	
+	Opacity.Constant = 1.0f;
+	OpacityMask.Constant = 1.0f;
 #endif
 	OpacityMaskClipValue = 0.3333f;
 	bCastDynamicShadowAsMasked = false;
@@ -1338,45 +1345,10 @@ bool UMaterial::IsCompilingOrHadCompileError(ERHIFeatureLevel::Type InFeatureLev
 }
 
 #if WITH_EDITOR
-TConstArrayView<TObjectPtr<UMaterialExpression>> UMaterial::GetExpressions() const
-{
-	return GetEditorOnlyData()->ExpressionCollection.Expressions;
-}
-
-TConstArrayView<TObjectPtr<UMaterialExpressionComment>> UMaterial::GetEditorComments() const
-{
-	return GetEditorOnlyData()->ExpressionCollection.EditorComments;
-}
-
-UMaterialExpressionExecBegin* UMaterial::GetExpressionExecBegin() const
-{
-	return GetEditorOnlyData()->ExpressionCollection.ExpressionExecBegin;
-}
-
-UMaterialExpressionExecEnd* UMaterial::GetExpressionExecEnd() const
-{
-	return GetEditorOnlyData()->ExpressionCollection.ExpressionExecEnd;
-}
-
-const FMaterialExpressionCollection& UMaterial::GetExpressionCollection() const
-{
-	return GetEditorOnlyData()->ExpressionCollection;
-}
-
-FMaterialExpressionCollection& UMaterial::GetExpressionCollection()
-{
-	return GetEditorOnlyData()->ExpressionCollection;
-}
-
-void UMaterial::AssignExpressionCollection(const FMaterialExpressionCollection& InCollection)
-{
-	GetEditorOnlyData()->ExpressionCollection = InCollection;
-}
-
 bool UMaterial::SetParameterValueEditorOnly(const FName& ParameterName, const FMaterialParameterMetadata& Meta)
 {
 	bool bResult = false;
-	for (UMaterialExpression* Expression : GetExpressions())
+	for (TObjectPtr<UMaterialExpression>& Expression : Expressions)
 	{
 		if (Expression && Expression->SetParameterValue(ParameterName, Meta, EMaterialExpressionSetParameterValueFlags::SendPostEditChangeProperty))
 		{
@@ -1392,11 +1364,15 @@ bool UMaterial::SetParameterValueEditorOnly(const FName& ParameterName, const FM
 
 				for (UMaterialFunctionInterface* Function : Functions)
 				{
-					for (const TObjectPtr<UMaterialExpression>& FunctionExpression : Function->GetExpressions())
+					const TArray<TObjectPtr<UMaterialExpression>>* ExpressionPtr = Function->GetFunctionExpressions();
+					if (ExpressionPtr)
 					{
-						if (FunctionExpression && FunctionExpression->SetParameterValue(ParameterName, Meta, EMaterialExpressionSetParameterValueFlags::SendPostEditChangeProperty))
+						for (const TObjectPtr<UMaterialExpression>& FunctionExpression : *ExpressionPtr)
 						{
-							bResult = true;
+							if (FunctionExpression && FunctionExpression->SetParameterValue(ParameterName, Meta, EMaterialExpressionSetParameterValueFlags::SendPostEditChangeProperty))
+							{
+								bResult = true;
+							}
 						}
 					}
 				}
@@ -1775,7 +1751,7 @@ void UMaterial::FixupMaterialUsageAfterLoad()
 #if WITH_EDITORONLY_DATA
 bool UMaterial::IterateDependentFunctions(TFunctionRef<bool(UMaterialFunctionInterface*)> Predicate) const
 {
-	for (UMaterialExpression* Expression : GetExpressions())
+	for (UMaterialExpression* Expression : Expressions)
 	{
 		if (UMaterialExpressionMaterialFunctionCall* FunctionCall = Cast<UMaterialExpressionMaterialFunctionCall>(Expression))
 		{
@@ -1906,6 +1882,8 @@ void UMaterial::UpdateCachedExpressionData()
 	}
 
 	FMaterialCachedExpressionData* LocalCachedExpressionData = new FMaterialCachedExpressionData();
+	LocalCachedExpressionData->Reset();
+
 	FMaterialCachedHLSLTree* LocalCachedTree = nullptr;
 	if (IsUsingNewHLSLGenerator())
 	{
@@ -1915,34 +1893,18 @@ void UMaterial::UpdateCachedExpressionData()
 	}
 	else
 	{
-		if (!bUseMaterialAttributes)
-		{
-			for (int32 PropertyIndex = 0; PropertyIndex < MP_MAX; ++PropertyIndex)
-			{
-				const EMaterialProperty Property = (EMaterialProperty)PropertyIndex;
-				const FExpressionInput* Input = GetExpressionInputForProperty(Property);
-				if (Input && Input->IsConnected())
-				{
-					LocalCachedExpressionData->SetPropertyConnected(Property);
-				}
-			}
-		}
-
 		FMaterialCachedExpressionContext Context;
-		LocalCachedExpressionData->UpdateForExpressions(Context, GetExpressions(), EMaterialParameterAssociation::GlobalParameter, -1);
+		LocalCachedExpressionData->UpdateForExpressions(Context, Expressions, EMaterialParameterAssociation::GlobalParameter, -1);
 	}
 
 	if (LocalCachedExpressionData->bHasMaterialLayers)
 	{
 		// Set all layers as linked to parent (there is no parent for base UMaterials)
-		LocalCachedExpressionData->EditorOnlyData->MaterialLayers.LinkAllLayersToParent();
+		LocalCachedExpressionData->MaterialLayers.LinkAllLayersToParent();
 	}
 
-	LocalCachedExpressionData->Validate();
-
 	CachedExpressionData.Reset(LocalCachedExpressionData);
 	CachedHLSLTree.Reset(LocalCachedTree);
-	EditorOnlyData->CachedExpressionData = LocalCachedExpressionData->EditorOnlyData;
 
 	FObjectCacheEventSink::NotifyReferencedTextureChanged_Concurrent(this);
 }
@@ -1952,7 +1914,7 @@ bool UMaterial::GetParameterValue(EMaterialParameterType Type, const FMemoryImag
 {
 	if (EnumHasAnyFlags(Flags, EMaterialGetParameterValueFlags::CheckNonOverrides) && CachedExpressionData)
 	{
-		return CachedExpressionData->GetParameterValue(Type, ParameterInfo, OutResult);
+		return CachedExpressionData->Parameters.GetParameterValue(Type, ParameterInfo, OutResult);
 	}
 
 	return false;
@@ -1962,14 +1924,7 @@ bool UMaterial::GetMaterialLayers(FMaterialLayersFunctions& OutLayers, TMicRecur
 {
 	if (CachedExpressionData && CachedExpressionData->bHasMaterialLayers)
 	{
-		OutLayers.GetRuntime() = CachedExpressionData->MaterialLayers;
-#if WITH_EDITORONLY_DATA
-		if (CachedExpressionData->EditorOnlyData)
-		{
-			OutLayers.EditorOnly = CachedExpressionData->EditorOnlyData->MaterialLayers;
-		}
-		OutLayers.Validate();
-#endif // WITH_EDITORONLY_DATA
+		OutLayers = CachedExpressionData->MaterialLayers;
 		return true;
 	}
 	return false;
@@ -2068,7 +2023,7 @@ static FAutoConsoleVariable GCompileMaterialsForShaderFormatCVar(
 void UMaterial::GetForceRecompileTextureIdsHash(FSHAHash &TextureReferencesHash)
 {
 	TArray<UTexture*> ForceRecompileTextures;
-	for (const UMaterialExpression *MaterialExpression : GetExpressions())
+	for (const UMaterialExpression *MaterialExpression : Expressions)
 	{
 		if (MaterialExpression == nullptr)
 		{
@@ -2103,7 +2058,7 @@ void UMaterial::GetForceRecompileTextureIdsHash(FSHAHash &TextureReferencesHash)
 
 bool UMaterial::IsTextureForceRecompileCacheRessource(UTexture *Texture)
 {
-	for (const UMaterialExpression *MaterialExpression : GetExpressions())
+	for (const UMaterialExpression *MaterialExpression : Expressions)
 	{
 		if (MaterialExpression == nullptr)
 		{
@@ -2427,22 +2382,18 @@ void UMaterial::ReleaseResourcesAndMutateDDCKey(const FGuid& TransformationId)
 bool UMaterial::AttemptInsertNewGroupName(const FString & InNewName)
 {
 #if WITH_EDITOR
-	UMaterialEditorOnlyData* LocalData = GetEditorOnlyData();
-	if (LocalData)
+	FParameterGroupData* ParameterGroupDataElement = ParameterGroupData.FindByPredicate([&InNewName](const FParameterGroupData& DataElement)
 	{
-		FParameterGroupData* ParameterGroupDataElement = LocalData->ParameterGroupData.FindByPredicate([&InNewName](const FParameterGroupData& DataElement)
-		{
-			return InNewName == DataElement.GroupName;
-		});
+		return InNewName == DataElement.GroupName;
+	});
 
-		if (ParameterGroupDataElement == nullptr)
-		{
-			FParameterGroupData NewGroupData;
-			NewGroupData.GroupName = InNewName;
-			NewGroupData.GroupSortPriority = 0;
-			LocalData->ParameterGroupData.Add(NewGroupData);
-			return true;
-		}
+	if (ParameterGroupDataElement == nullptr)
+	{
+		FParameterGroupData NewGroupData;
+		NewGroupData.GroupName = InNewName;
+		NewGroupData.GroupSortPriority = 0;
+		ParameterGroupData.Add(NewGroupData);
+		return true;
 	}
 #endif
 	return false;
@@ -2599,13 +2550,12 @@ void UMaterial::Serialize(FArchive& Ar)
 #if WITH_EDITORONLY_DATA
 	if (MaterialDomain == MD_Volume && Ar.IsLoading() && Ar.CustomVer(FRenderingObjectVersion::GUID) < FRenderingObjectVersion::VolumeExtinctionBecomesRGB)
 	{
-		UMaterialEditorOnlyData* EditorOnly = GetEditorOnlyData();
-		if (EditorOnly && EditorOnly->Opacity.IsConnected()) // Base material input cannot have default values so we only deal with connected expression
+		if (Opacity.IsConnected()) // Base material input cannot have default values so we only deal with connected expression
 		{
 			// Change expression output from the Opacity to SubSurfaceColor that is now representing RGB extinction. Leave opacity connected as it is unused now anyway
-			EditorOnly->SubsurfaceColor.Connect(EditorOnly->Opacity.OutputIndex, EditorOnly->Opacity.Expression);
+			SubsurfaceColor.Connect(Opacity.OutputIndex, Opacity.Expression);
 			// Now disconnect Opacity
-			EditorOnly->Opacity.Expression = nullptr;
+			Opacity.Expression = nullptr;
 
 			// Now force the material to recompile and we use a hash of the original StateId.
 			// This is to avoid having different StateId each time we load the material and to not forever recompile it,i.e. use a cached version.
@@ -2641,9 +2591,8 @@ void UMaterial::BackwardsCompatibilityInputConversion()
 #if WITH_EDITOR
 	if( ShadingModel != MSM_Unlit )
 	{
-		UMaterialEditorOnlyData* EditorOnly = GetEditorOnlyData();
 		bool bIsDS = DiffuseColor_DEPRECATED.IsConnected() || SpecularColor_DEPRECATED.IsConnected();
-		bool bIsBMS = EditorOnly->BaseColor.IsConnected() || EditorOnly->Metallic.IsConnected() || EditorOnly->Specular.IsConnected();
+		bool bIsBMS = BaseColor.IsConnected() || Metallic.IsConnected() || Specular.IsConnected();
 
 		if( bIsDS && !bIsBMS )
 		{
@@ -2652,7 +2601,7 @@ void UMaterial::BackwardsCompatibilityInputConversion()
 			check( GConvertFromDiffSpecMaterialFunction );
 
 			UMaterialExpressionMaterialFunctionCall* FunctionExpression = NewObject<UMaterialExpressionMaterialFunctionCall>(this);
-			EditorOnly->ExpressionCollection.Expressions.Add( FunctionExpression );
+			Expressions.Add( FunctionExpression );
 
 			FunctionExpression->MaterialExpressionEditorX += 200;
 
@@ -2669,9 +2618,9 @@ void UMaterial::BackwardsCompatibilityInputConversion()
 				FunctionExpression->GetInput(1)->Connect( SpecularColor_DEPRECATED.OutputIndex, SpecularColor_DEPRECATED.Expression );
 			}
 
-			EditorOnly->BaseColor.Connect( 0, FunctionExpression );
-			EditorOnly->Metallic.Connect( 1, FunctionExpression );
-			EditorOnly->Specular.Connect( 2, FunctionExpression );
+			BaseColor.Connect( 0, FunctionExpression );
+			Metallic.Connect( 1, FunctionExpression );
+			Specular.Connect( 2, FunctionExpression );
 		}
 	}
 #endif // WITH_EDITOR
@@ -2686,67 +2635,65 @@ void UMaterial::BackwardsCompatibilityVirtualTextureOutputConversion()
 		// Change this guid if you change the conversion code below
 		static FGuid BackwardsCompatibilityVirtualTextureOutputConversionGuid(TEXT("BABD7074-001F-4FC2-BDE5-3A0C436F4414"));
 
-		UMaterialEditorOnlyData* EditorOnly = GetEditorOnlyData();
-
 		MaterialDomain = MD_Surface;
 
 		if (!bUseMaterialAttributes)
 		{
 			// Create a new UMaterialExpressionRuntimeVirtualTextureOutput node and route the old material attribute output to it.
 			UMaterialExpressionRuntimeVirtualTextureOutput* OutputExpression = NewObject<UMaterialExpressionRuntimeVirtualTextureOutput>(this);
-			EditorOnly->ExpressionCollection.Expressions.Add(OutputExpression);
+			Expressions.Add(OutputExpression);
 
 			OutputExpression->MaterialExpressionEditorX = EditorX;
 			OutputExpression->MaterialExpressionEditorY = EditorY - 300;
 
-			if (EditorOnly->BaseColor.IsConnected())
+			if (BaseColor.IsConnected())
 			{
-				OutputExpression->GetInput(0)->Connect(EditorOnly->BaseColor.OutputIndex, EditorOnly->BaseColor.Expression);
+				OutputExpression->GetInput(0)->Connect(BaseColor.OutputIndex, BaseColor.Expression);
 			}
-			if (EditorOnly->Specular.IsConnected())
+			if (Specular.IsConnected())
 			{
-				OutputExpression->GetInput(1)->Connect(EditorOnly->Specular.OutputIndex, EditorOnly->Specular.Expression);
+				OutputExpression->GetInput(1)->Connect(Specular.OutputIndex, Specular.Expression);
 			}
-			if (EditorOnly->Roughness.IsConnected())
+			if (Roughness.IsConnected())
 			{
-				OutputExpression->GetInput(2)->Connect(EditorOnly->Roughness.OutputIndex, EditorOnly->Roughness.Expression);
+				OutputExpression->GetInput(2)->Connect(Roughness.OutputIndex, Roughness.Expression);
 			}
-			if (EditorOnly->Normal.IsConnected())
+			if (Normal.IsConnected())
 			{
 				if (bTangentSpaceNormal)
 				{
-					OutputExpression->GetInput(3)->Connect(EditorOnly->Normal.OutputIndex, EditorOnly->Normal.Expression);
+					OutputExpression->GetInput(3)->Connect(Normal.OutputIndex, Normal.Expression);
 				}
 				else
 				{
 					// Apply the tangent space to world transform that would be applied in the material output.
 					UMaterialExpressionTransform* TransformExpression = NewObject<UMaterialExpressionTransform>(this);
-					EditorOnly->ExpressionCollection.Expressions.Add(TransformExpression);
+					Expressions.Add(TransformExpression);
 
 					TransformExpression->MaterialExpressionEditorX = EditorX - 300;
 					TransformExpression->MaterialExpressionEditorY = EditorY - 300;
 					TransformExpression->TransformSourceType = TRANSFORMSOURCE_Tangent;
 					TransformExpression->TransformType = TRANSFORM_World;
-					TransformExpression->Input.Connect(EditorOnly->Normal.OutputIndex, EditorOnly->Normal.Expression);
+					TransformExpression->Input.Connect(Normal.OutputIndex, Normal.Expression);
 
 					OutputExpression->GetInput(3)->Connect(0, TransformExpression);
 				}
 			}
-			if (EditorOnly->Opacity.IsConnected())
+			if (Opacity.IsConnected())
 			{
-				OutputExpression->GetInput(5)->Connect(EditorOnly->Opacity.OutputIndex, EditorOnly->Opacity.Expression);
+				OutputExpression->GetInput(5)->Connect(Opacity.OutputIndex, Opacity.Expression);
 			}
 			if (BlendMode != BLEND_Opaque)
 			{
 				// Full alpha blend modes were mostly/always used with MD_RuntimeVirtualTexture to allow pin connections.
 				// But we will assume the intention for any associated MD_Surface output is opaque or alpha mask and force convert here.
-				if (EditorOnly->Opacity.IsConnected())
+				if (Opacity.IsConnected())
 				{
-					EditorOnly->OpacityMask.Connect(EditorOnly->Opacity.OutputIndex, EditorOnly->Opacity.Expression);
-					EditorOnly->Opacity.Expression = nullptr;
+					OpacityMask.Connect(Opacity.OutputIndex, Opacity.Expression);
+					Opacity.Expression = nullptr;
 				}
-				BlendMode = EditorOnly->OpacityMask.IsConnected() ? BLEND_Masked : BLEND_Opaque;
-				bCanMaskedBeAssumedOpaque = !EditorOnly->OpacityMask.Expression && !(EditorOnly->OpacityMask.UseConstant && EditorOnly->OpacityMask.Constant < 0.999f);
+				BlendMode = OpacityMask.IsConnected() ? BLEND_Masked : BLEND_Opaque;
+				bCanMaskedBeAssumedOpaque = !OpacityMask.Expression && !(OpacityMask.UseConstant && OpacityMask.Constant < 0.999f);
 			}
 		}
 
@@ -2763,8 +2710,6 @@ void UMaterial::BackwardsCompatibilityDecalConversion()
 	{
 		// Change this guid if you change the conversion code below
 		static FGuid BackwardsCompatibilityDecalConversionGuid(TEXT("352069F8-1B8C-406A-9B88-6946BCDF2C10"));
-		
-		UMaterialEditorOnlyData* EditorOnly = GetEditorOnlyData();
 
 		GMaterialsThatNeedDecalFix.Clear(this);
 
@@ -2785,56 +2730,56 @@ void UMaterial::BackwardsCompatibilityDecalConversion()
 		// Disconnect outputs according to old DBuffer blend mode.
 		if (DecalBlendMode == DBM_DBuffer_Normal || DecalBlendMode == DBM_DBuffer_Roughness || DecalBlendMode == DBM_DBuffer_NormalRoughness)
 		{
-			EditorOnly->BaseColor.Expression = nullptr;
+			BaseColor.Expression = nullptr;
 		}
 		if (DecalBlendMode == DBM_DBuffer_Color || DecalBlendMode == DBM_DBuffer_Roughness || DecalBlendMode == DBM_DBuffer_ColorRoughness || DecalBlendMode == DBM_AlphaComposite)
 		{
-			EditorOnly->Normal.Expression = nullptr;
+			Normal.Expression = nullptr;
 		}
 		if (DecalBlendMode == DBM_DBuffer_Color || DecalBlendMode == DBM_DBuffer_Normal || DecalBlendMode == DBM_DBuffer_ColorNormal)
 		{
-			EditorOnly->Roughness.Expression = EditorOnly->Specular.Expression = EditorOnly->Metallic.Expression = nullptr;
+			Roughness.Expression = Specular.Expression = Metallic.Expression = nullptr;
 		}
 
 		// Previously translucent decals used default values in all unconnected attributes (except for normal).
 		// For backwards compatibility we connect those attributes with defaults.
 		if (DecalBlendMode == DBM_Translucent || DecalBlendMode == DBM_AlphaComposite || DecalBlendMode == DBM_Stain)
 		{
-			if (!EditorOnly->BaseColor.IsConnected() || !EditorOnly->Metallic.IsConnected())
+			if (!BaseColor.IsConnected() || !Metallic.IsConnected())
 			{
 				UMaterialExpressionConstant* Expression = NewObject<UMaterialExpressionConstant>(this);
-				EditorOnly->ExpressionCollection.Expressions.Add(Expression);
+				Expressions.Add(Expression);
 
 				Expression->MaterialExpressionEditorX = EditorX - 100;
 				Expression->MaterialExpressionEditorY = EditorY - 120;
 				Expression->R = 0.f;
 
-				if (!EditorOnly->BaseColor.IsConnected())
+				if (!BaseColor.IsConnected())
 				{
-					EditorOnly->BaseColor.Connect(0, Expression);
+					BaseColor.Connect(0, Expression);
 				}
-				if (!EditorOnly->Metallic.IsConnected())
+				if (!Metallic.IsConnected())
 				{
-					EditorOnly->Metallic.Connect(0, Expression);
+					Metallic.Connect(0, Expression);
 				}
 			}
 
-			if (!EditorOnly->Roughness.IsConnected() || !EditorOnly->Specular.IsConnected())
+			if (!Roughness.IsConnected() || !Specular.IsConnected())
 			{
 				UMaterialExpressionConstant* Expression = NewObject<UMaterialExpressionConstant>(this);
-				EditorOnly->ExpressionCollection.Expressions.Add(Expression);
+				Expressions.Add(Expression);
 
 				Expression->MaterialExpressionEditorX = EditorX - 100;
 				Expression->MaterialExpressionEditorY = EditorY - 60;
 				Expression->R = .5f;
 
-				if (!EditorOnly->Roughness.IsConnected())
+				if (!Roughness.IsConnected())
 				{
-					EditorOnly->Roughness.Connect(0, Expression);
+					Roughness.Connect(0, Expression);
 				}
-				if (!EditorOnly->Specular.IsConnected())
+				if (!Specular.IsConnected())
 				{
-					EditorOnly->Specular.Connect(0, Expression);
+					Specular.Connect(0, Expression);
 				}
 			}
 		}
@@ -2918,8 +2863,6 @@ void UMaterial::ConvertMaterialToStrataMaterial()
 			}
 		}
 	};
-	
-	UMaterialEditorOnlyData* EditorOnly = GetEditorOnlyData();
 
 	bool bCustomNodesGathered = false;
 	UMaterialExpressionThinTranslucentMaterialOutput* ThinTranslucentOutput = nullptr;
@@ -2931,6 +2874,7 @@ void UMaterial::ConvertMaterialToStrataMaterial()
 		if (!bCustomNodesGathered)
 		{
 			bCustomNodesGathered = true;
+
 			TArray<class UMaterialExpressionCustomOutput*> CustomOutputExpressions;
 			GetAllCustomOutputExpressions(CustomOutputExpressions);
 			for (UMaterialExpressionCustomOutput* Expression : CustomOutputExpressions)
@@ -2976,10 +2920,10 @@ void UMaterial::ConvertMaterialToStrataMaterial()
 	bool bRelinkCustomOutputNodes = false;
 	UMaterialExpressionStrataLegacyConversion* ConvertNode = nullptr;
 	// Connect all the legacy pin into the conversion node
-	if (bUseMaterialAttributes && EditorOnly->MaterialAttributes.Expression && !EditorOnly->MaterialAttributes.Expression->IsResultStrataMaterial(EditorOnly->MaterialAttributes.OutputIndex)) // M_Rifle cause issues there
+	if (bUseMaterialAttributes && MaterialAttributes.Expression && !MaterialAttributes.Expression->IsResultStrataMaterial(MaterialAttributes.OutputIndex)) // M_Rifle cause issues there
 	{
 		UMaterialExpressionBreakMaterialAttributes* BreakMatAtt = NewObject<UMaterialExpressionBreakMaterialAttributes>(this);
-		MoveConnectionTo(EditorOnly->MaterialAttributes, BreakMatAtt, 0);
+		MoveConnectionTo(MaterialAttributes, BreakMatAtt, 0);
 
 		ConvertNode = NewObject<UMaterialExpressionStrataLegacyConversion>(this);
 		ConvertNode->BaseColor.Connect(0, BreakMatAtt);
@@ -3003,15 +2947,15 @@ void UMaterial::ConvertMaterialToStrataMaterial()
 		// * Forward inputs to the root node (Do not reconnect the Opacity as we handle the opacity by internally within the conversion node)
 		// * Forward masked opacity only if blend mode is set to masked, as certain material (e.g., FlattenVT requires to have no OpacityMask plugged)
 		bUseMaterialAttributes = false;
-		EditorOnly->FrontMaterial.Connect(0, ConvertNode);
+		FrontMaterial.Connect(0, ConvertNode);
 		if (BlendMode == BLEND_Masked)
 		{
-			EditorOnly->OpacityMask.Connect(7, BreakMatAtt);
+			OpacityMask.Connect(7, BreakMatAtt);
 		}
-		EditorOnly->WorldPositionOffset.Connect(10, BreakMatAtt);
-		EditorOnly->AmbientOcclusion.Connect(14, BreakMatAtt);
-		EditorOnly->PixelDepthOffset.Connect(24, BreakMatAtt);
-		EditorOnly->Refraction.Connect(15, BreakMatAtt);
+		WorldPositionOffset.Connect(10, BreakMatAtt);
+		AmbientOcclusion.Connect(14, BreakMatAtt);
+		PixelDepthOffset.Connect(24, BreakMatAtt);
+		Refraction.Connect(15, BreakMatAtt);
 
 		if (ShadingModel == MSM_FromMaterialExpression)
 		{
@@ -3037,7 +2981,7 @@ void UMaterial::ConvertMaterialToStrataMaterial()
 		ConvertLegacyToStrataBlendMode();
 		bInvalidateShader = true;
 	}
-	else if (!bUseMaterialAttributes && !EditorOnly->FrontMaterial.IsConnected())
+	else if (!bUseMaterialAttributes && !FrontMaterial.IsConnected())
 	{
 		// STRATA_TODO for material conversion
 		//  - WorldPositionOffset can remain on the end point node
@@ -3048,18 +2992,18 @@ void UMaterial::ConvertMaterialToStrataMaterial()
 		{
 			ConvertNode = NewObject<UMaterialExpressionStrataLegacyConversion>(this);
 			ConvertNode->SubsurfaceProfile = bRequireNoSubsurfaceProfile ? nullptr : SubsurfaceProfile;
-			MoveConnectionTo(EditorOnly->BaseColor, ConvertNode, 0);
-			MoveConnectionTo(EditorOnly->Metallic, ConvertNode, 1);
-			MoveConnectionTo(EditorOnly->Specular, ConvertNode, 2);
-			MoveConnectionTo(EditorOnly->Roughness, ConvertNode, 3);
-			MoveConnectionTo(EditorOnly->Anisotropy, ConvertNode, 4);
-			MoveConnectionTo(EditorOnly->EmissiveColor, ConvertNode, 5);
-			CopyConnectionTo(EditorOnly->Normal, ConvertNode, 6);
-			MoveConnectionTo(EditorOnly->Tangent, ConvertNode, 7);
-			MoveConnectionTo(EditorOnly->SubsurfaceColor, ConvertNode, 8);
-			MoveConnectionTo(EditorOnly->ClearCoat, ConvertNode, 9);
-			MoveConnectionTo(EditorOnly->ClearCoatRoughness, ConvertNode, 10);
-			MoveConnectionTo(EditorOnly->Opacity, ConvertNode, 11);
+			MoveConnectionTo(BaseColor, ConvertNode, 0);
+			MoveConnectionTo(Metallic, ConvertNode, 1);
+			MoveConnectionTo(Specular, ConvertNode, 2);
+			MoveConnectionTo(Roughness, ConvertNode, 3);
+			MoveConnectionTo(Anisotropy, ConvertNode, 4);
+			MoveConnectionTo(EmissiveColor, ConvertNode, 5);
+			CopyConnectionTo(Normal, ConvertNode, 6);
+			MoveConnectionTo(Tangent, ConvertNode, 7);
+			MoveConnectionTo(SubsurfaceColor, ConvertNode, 8);
+			MoveConnectionTo(ClearCoat, ConvertNode, 9);
+			MoveConnectionTo(ClearCoatRoughness, ConvertNode, 10);
+			MoveConnectionTo(Opacity, ConvertNode, 11);
 			bRelinkCustomOutputNodes = true;
 			
 			// Shading Model
@@ -3070,10 +3014,10 @@ void UMaterial::ConvertMaterialToStrataMaterial()
 			// rebuilding the final Shading model (see RebuildShadingModelField())
 			if (ShadingModel == MSM_FromMaterialExpression)
 			{
-				check(EditorOnly->ShadingModelFromMaterialExpression.IsConnected());
+				check(ShadingModelFromMaterialExpression.IsConnected());
 
 				// Reconnect the shading model expression
-				MoveConnectionTo(EditorOnly->ShadingModelFromMaterialExpression, ConvertNode, 18);
+				MoveConnectionTo(ShadingModelFromMaterialExpression, ConvertNode, 18);
 
 				// Store strata shading model of the converted material. 
 				GatherCustomNodes();
@@ -3100,19 +3044,19 @@ void UMaterial::ConvertMaterialToStrataMaterial()
 				check(ConvertNode->ConvertedStrataMaterialInfo.CountShadingModels() == 1);
 			}
 
-			EditorOnly->FrontMaterial.Connect(0, ConvertNode);
+			FrontMaterial.Connect(0, ConvertNode);
 			bInvalidateShader = true;
 		}
 		else if (MaterialDomain == MD_Volume)
 		{
 			UMaterialExpressionStrataVolumetricFogCloudBSDF* VolBSDF = NewObject<UMaterialExpressionStrataVolumetricFogCloudBSDF>(this);
-			MoveConnectionTo(EditorOnly->BaseColor, VolBSDF, 0);		// Albedo
-			MoveConnectionTo(EditorOnly->SubsurfaceColor, VolBSDF, 1);	// Extinction
-			MoveConnectionTo(EditorOnly->EmissiveColor, VolBSDF, 2);	// EmissiveColor
-			MoveConnectionTo(EditorOnly->AmbientOcclusion, VolBSDF, 3);	// AmbientOcclusion
+			MoveConnectionTo(BaseColor, VolBSDF, 0);		// Albedo
+			MoveConnectionTo(SubsurfaceColor, VolBSDF, 1);	// Extinction
+			MoveConnectionTo(EmissiveColor, VolBSDF, 2);	// EmissiveColor
+			MoveConnectionTo(AmbientOcclusion, VolBSDF, 3);	// AmbientOcclusion
 
 			// STRATA_TODO remove the VolumetricAdvancedOutput node and add the input onto FogCloudBSDF even if only used by the cloud renderer?
-			EditorOnly->FrontMaterial.Connect(0, VolBSDF);
+			FrontMaterial.Connect(0, VolBSDF);
 			bInvalidateShader = true;
 		}
 		else if (MaterialDomain == MD_LightFunction)
@@ -3124,9 +3068,9 @@ void UMaterial::ConvertMaterialToStrataMaterial()
 
 			// Only Emissive & Opacity are valid input for PostProcess material
 			UMaterialExpressionStrataLightFunction* LightFunctionNode = NewObject<UMaterialExpressionStrataLightFunction>(this);
-			MoveConnectionTo(EditorOnly->EmissiveColor, LightFunctionNode, 0);
+			MoveConnectionTo(EmissiveColor, LightFunctionNode, 0);
 
-			EditorOnly->FrontMaterial.Connect(0, LightFunctionNode);
+			FrontMaterial.Connect(0, LightFunctionNode);
 			bInvalidateShader = true;
 		}
 		else if (MaterialDomain == MD_PostProcess)
@@ -3146,23 +3090,23 @@ void UMaterial::ConvertMaterialToStrataMaterial()
 			{
 				// If the blending mode was translucent, we must multiply the color by the opacity before the output node
 				UMaterialExpressionMultiply* ColorMultiplyOpacityNode = NewObject<UMaterialExpressionMultiply>(this);
-				CopyConnectionTo(EditorOnly->Opacity, ColorMultiplyOpacityNode, 0);
-				MoveConnectionTo(EditorOnly->EmissiveColor, ColorMultiplyOpacityNode, 1);
+				CopyConnectionTo(Opacity, ColorMultiplyOpacityNode, 0);
+				MoveConnectionTo(EmissiveColor, ColorMultiplyOpacityNode, 1);
 
 				PostProcNode->Color.Connect(0, ColorMultiplyOpacityNode);
-				MoveConnectionTo(EditorOnly->Opacity, PostProcNode, 1);
+				MoveConnectionTo(Opacity, PostProcNode, 1);
 			}
 			else
 			{
 				// Only Emissive & Opacity are valid input for PostProcess material
-				MoveConnectionTo(EditorOnly->EmissiveColor, PostProcNode, 0);
+				MoveConnectionTo(EmissiveColor, PostProcNode, 0);
 				if (GetBlendMode() != BLEND_Additive) // Legacy material additive was ignoring opacity
 				{
-					MoveConnectionTo(EditorOnly->Opacity, PostProcNode, 1);
+					MoveConnectionTo(Opacity, PostProcNode, 1);
 				}
 			}
 
-			EditorOnly->FrontMaterial.Connect(0, PostProcNode);
+			FrontMaterial.Connect(0, PostProcNode);
 			bInvalidateShader = true;
 		}
 		else if (MaterialDomain == MD_DeferredDecal)
@@ -3173,18 +3117,18 @@ void UMaterial::ConvertMaterialToStrataMaterial()
 			ShadingModels.AddShadingModel(MSM_DefaultLit);
 
 			ConvertNode = NewObject<UMaterialExpressionStrataLegacyConversion>(this);
-			MoveConnectionTo(EditorOnly->BaseColor, ConvertNode, 0);
-			MoveConnectionTo(EditorOnly->Metallic, ConvertNode, 1);
-			MoveConnectionTo(EditorOnly->Specular, ConvertNode, 2);
-			MoveConnectionTo(EditorOnly->Roughness, ConvertNode, 3);
-			MoveConnectionTo(EditorOnly->Anisotropy, ConvertNode, 4);
-			MoveConnectionTo(EditorOnly->EmissiveColor, ConvertNode, 5);
-			CopyConnectionTo(EditorOnly->Normal, ConvertNode, 6);
-			MoveConnectionTo(EditorOnly->Tangent, ConvertNode, 7);
-			MoveConnectionTo(EditorOnly->SubsurfaceColor, ConvertNode, 8);
-			MoveConnectionTo(EditorOnly->ClearCoat, ConvertNode, 9);
-			MoveConnectionTo(EditorOnly->ClearCoatRoughness, ConvertNode, 10);
-			MoveConnectionTo(EditorOnly->Opacity, ConvertNode, 11);
+			MoveConnectionTo(BaseColor, ConvertNode, 0);
+			MoveConnectionTo(Metallic, ConvertNode, 1);
+			MoveConnectionTo(Specular, ConvertNode, 2);
+			MoveConnectionTo(Roughness, ConvertNode, 3);
+			MoveConnectionTo(Anisotropy, ConvertNode, 4);
+			MoveConnectionTo(EmissiveColor, ConvertNode, 5);
+			CopyConnectionTo(Normal, ConvertNode, 6);
+			MoveConnectionTo(Tangent, ConvertNode, 7);
+			MoveConnectionTo(SubsurfaceColor, ConvertNode, 8);
+			MoveConnectionTo(ClearCoat, ConvertNode, 9);
+			MoveConnectionTo(ClearCoatRoughness, ConvertNode, 10);
+			MoveConnectionTo(Opacity, ConvertNode, 11);
 
 			// Add constant for the Unlit shading model
 			UMaterialExpressionConstant* ShadingModelNode = NewObject<UMaterialExpressionConstant>(this);
@@ -3199,7 +3143,7 @@ void UMaterial::ConvertMaterialToStrataMaterial()
 			UMaterialExpressionStrataConvertToDecal* ConvertToDecalNode= NewObject<UMaterialExpressionStrataConvertToDecal>(this);
 			ConvertToDecalNode->DecalMaterial.Connect(0, ConvertNode);
 
-			EditorOnly->FrontMaterial.Connect(0, ConvertToDecalNode);
+			FrontMaterial.Connect(0, ConvertToDecalNode);
 			bInvalidateShader = true;
 		}
 
@@ -3247,15 +3191,6 @@ void UMaterial::ConvertMaterialToStrataMaterial()
 
 TMap<FGuid, UMaterialInterface*> LightingGuidFixupMap;
 
-template<typename InputType>
-static void MoveExpressionInput(InputType& Src, InputType& OutDst)
-{
-	if (Src.IsConnected() || Src.IsConstant())
-	{
-		OutDst = MoveTemp(Src);
-	}
-}
-
 void UMaterial::PostLoad()
 {
 	LLM_SCOPE(ELLMTag::Materials);
@@ -3285,94 +3220,35 @@ void UMaterial::PostLoad()
 	const int32 RenderObjVer = GetLinkerCustomVersion(FRenderingObjectVersion::GUID);
 	const int32 UE5MainVer = GetLinkerCustomVersion(FUE5MainStreamObjectVersion::GUID);
 
-	UMaterialEditorOnlyData* EditorOnly = GetEditorOnlyData();
-
-	MoveExpressionInput(BaseColor_DEPRECATED, EditorOnly->BaseColor);
-	MoveExpressionInput(Metallic_DEPRECATED, EditorOnly->Metallic);
-	MoveExpressionInput(Specular_DEPRECATED, EditorOnly->Specular);
-	MoveExpressionInput(Roughness_DEPRECATED, EditorOnly->Roughness);
-	MoveExpressionInput(Anisotropy_DEPRECATED, EditorOnly->Anisotropy);
-	MoveExpressionInput(Normal_DEPRECATED, EditorOnly->Normal);
-	MoveExpressionInput(Tangent_DEPRECATED, EditorOnly->Tangent);
-	MoveExpressionInput(EmissiveColor_DEPRECATED, EditorOnly->EmissiveColor);
-	MoveExpressionInput(Tangent_DEPRECATED, EditorOnly->Tangent);
-	MoveExpressionInput(Opacity_DEPRECATED, EditorOnly->Opacity);
-	MoveExpressionInput(OpacityMask_DEPRECATED, EditorOnly->OpacityMask);
-	MoveExpressionInput(WorldPositionOffset_DEPRECATED, EditorOnly->WorldPositionOffset);
-	MoveExpressionInput(SubsurfaceColor_DEPRECATED, EditorOnly->SubsurfaceColor);
-	MoveExpressionInput(ClearCoat_DEPRECATED, EditorOnly->ClearCoat);
-	MoveExpressionInput(ClearCoatRoughness_DEPRECATED, EditorOnly->ClearCoatRoughness);
-	MoveExpressionInput(AmbientOcclusion_DEPRECATED, EditorOnly->AmbientOcclusion);
-	MoveExpressionInput(Refraction_DEPRECATED, EditorOnly->Refraction);
-	MoveExpressionInput(MaterialAttributes_DEPRECATED, EditorOnly->MaterialAttributes);
-	MoveExpressionInput(PixelDepthOffset_DEPRECATED, EditorOnly->PixelDepthOffset);
-	MoveExpressionInput(ShadingModelFromMaterialExpression_DEPRECATED, EditorOnly->ShadingModelFromMaterialExpression);
-	MoveExpressionInput(FrontMaterial_DEPRECATED, EditorOnly->FrontMaterial);
-	for (int32 i = 0; i < 8; ++i)
-	{
-		MoveExpressionInput(CustomizedUVs_DEPRECATED[i], EditorOnly->CustomizedUVs[i]);
-	}
-
 	DoMaterialAttributeReorder(&DiffuseColor_DEPRECATED, UEVer, RenderObjVer, UE5MainVer);
 	DoMaterialAttributeReorder(&SpecularColor_DEPRECATED, UEVer, RenderObjVer, UE5MainVer);
-	DoMaterialAttributeReorder(&EditorOnly->BaseColor, UEVer, RenderObjVer, UE5MainVer);
-	DoMaterialAttributeReorder(&EditorOnly->Metallic, UEVer, RenderObjVer, UE5MainVer);
-	DoMaterialAttributeReorder(&EditorOnly->Specular, UEVer, RenderObjVer, UE5MainVer);
-	DoMaterialAttributeReorder(&EditorOnly->Roughness, UEVer, RenderObjVer, UE5MainVer);
-	DoMaterialAttributeReorder(&EditorOnly->Anisotropy, UEVer, RenderObjVer, UE5MainVer);
-	DoMaterialAttributeReorder(&EditorOnly->Normal, UEVer, RenderObjVer, UE5MainVer);
-	DoMaterialAttributeReorder(&EditorOnly->Tangent, UEVer, RenderObjVer, UE5MainVer);
-	DoMaterialAttributeReorder(&EditorOnly->EmissiveColor, UEVer, RenderObjVer, UE5MainVer);
-	DoMaterialAttributeReorder(&EditorOnly->Opacity, UEVer, RenderObjVer, UE5MainVer);
-	DoMaterialAttributeReorder(&EditorOnly->OpacityMask, UEVer, RenderObjVer, UE5MainVer);
-	DoMaterialAttributeReorder(&EditorOnly->WorldPositionOffset, UEVer, RenderObjVer, UE5MainVer);
-	DoMaterialAttributeReorder(&EditorOnly->SubsurfaceColor, UEVer, RenderObjVer, UE5MainVer);
-	DoMaterialAttributeReorder(&EditorOnly->ClearCoat, UEVer, RenderObjVer, UE5MainVer);
-	DoMaterialAttributeReorder(&EditorOnly->ClearCoatRoughness, UEVer, RenderObjVer, UE5MainVer);
-	DoMaterialAttributeReorder(&EditorOnly->AmbientOcclusion, UEVer, RenderObjVer, UE5MainVer);
-	DoMaterialAttributeReorder(&EditorOnly->Refraction, UEVer, RenderObjVer, UE5MainVer);
-	DoMaterialAttributeReorder(&EditorOnly->CustomizedUVs[0], UEVer, RenderObjVer, UE5MainVer);
-	DoMaterialAttributeReorder(&EditorOnly->CustomizedUVs[1], UEVer, RenderObjVer, UE5MainVer);
-	DoMaterialAttributeReorder(&EditorOnly->CustomizedUVs[2], UEVer, RenderObjVer, UE5MainVer);
-	DoMaterialAttributeReorder(&EditorOnly->CustomizedUVs[3], UEVer, RenderObjVer, UE5MainVer);
-	DoMaterialAttributeReorder(&EditorOnly->CustomizedUVs[4], UEVer, RenderObjVer, UE5MainVer);
-	DoMaterialAttributeReorder(&EditorOnly->CustomizedUVs[5], UEVer, RenderObjVer, UE5MainVer);
-	DoMaterialAttributeReorder(&EditorOnly->CustomizedUVs[6], UEVer, RenderObjVer, UE5MainVer);
-	DoMaterialAttributeReorder(&EditorOnly->CustomizedUVs[7], UEVer, RenderObjVer, UE5MainVer);
-	DoMaterialAttributeReorder(&EditorOnly->PixelDepthOffset, UEVer, RenderObjVer, UE5MainVer);
-	DoMaterialAttributeReorder(&EditorOnly->ShadingModelFromMaterialExpression, UEVer, RenderObjVer, UE5MainVer);
-	DoMaterialAttributeReorder(&EditorOnly->FrontMaterial, UEVer, RenderObjVer, UE5MainVer);
-
-	if (EditorOnly && ParameterGroupData_DEPRECATED.Num() > 0)
-	{
-		ensure(EditorOnly->ParameterGroupData.Num() == 0);
-		EditorOnly->ParameterGroupData = MoveTemp(ParameterGroupData_DEPRECATED);
-	}
-
-	if (EditorOnly && Expressions_DEPRECATED.Num() > 0)
-	{
-		ensure(EditorOnly->ExpressionCollection.Expressions.Num() == 0);
-		EditorOnly->ExpressionCollection.Expressions = MoveTemp(Expressions_DEPRECATED);
-	}
-
-	if (EditorOnly && EditorComments_DEPRECATED.Num() > 0)
-	{
-		ensure(EditorOnly->ExpressionCollection.EditorComments.Num() == 0);
-		EditorOnly->ExpressionCollection.EditorComments = MoveTemp(EditorComments_DEPRECATED);
-	}
-
-	if (EditorOnly && ExpressionExecBegin_DEPRECATED)
-	{
-		ensure(!EditorOnly->ExpressionCollection.ExpressionExecBegin);
-		EditorOnly->ExpressionCollection.ExpressionExecBegin = MoveTemp(ExpressionExecBegin_DEPRECATED);
-	}
-
-	if (EditorOnly && ExpressionExecEnd_DEPRECATED)
-	{
-		ensure(!EditorOnly->ExpressionCollection.ExpressionExecEnd);
-		EditorOnly->ExpressionCollection.ExpressionExecEnd = MoveTemp(ExpressionExecEnd_DEPRECATED);
-	}
-
+	DoMaterialAttributeReorder(&BaseColor, UEVer, RenderObjVer, UE5MainVer);
+	DoMaterialAttributeReorder(&Metallic, UEVer, RenderObjVer, UE5MainVer);
+	DoMaterialAttributeReorder(&Specular, UEVer, RenderObjVer, UE5MainVer);
+	DoMaterialAttributeReorder(&Roughness, UEVer, RenderObjVer, UE5MainVer);
+	DoMaterialAttributeReorder(&Anisotropy, UEVer, RenderObjVer, UE5MainVer);
+	DoMaterialAttributeReorder(&Normal, UEVer, RenderObjVer, UE5MainVer);
+	DoMaterialAttributeReorder(&Tangent, UEVer, RenderObjVer, UE5MainVer);
+	DoMaterialAttributeReorder(&EmissiveColor, UEVer, RenderObjVer, UE5MainVer);
+	DoMaterialAttributeReorder(&Opacity, UEVer, RenderObjVer, UE5MainVer);
+	DoMaterialAttributeReorder(&OpacityMask, UEVer, RenderObjVer, UE5MainVer);
+	DoMaterialAttributeReorder(&WorldPositionOffset, UEVer, RenderObjVer, UE5MainVer);
+	DoMaterialAttributeReorder(&SubsurfaceColor, UEVer, RenderObjVer, UE5MainVer);
+	DoMaterialAttributeReorder(&ClearCoat, UEVer, RenderObjVer, UE5MainVer);
+	DoMaterialAttributeReorder(&ClearCoatRoughness, UEVer, RenderObjVer, UE5MainVer);
+	DoMaterialAttributeReorder(&AmbientOcclusion, UEVer, RenderObjVer, UE5MainVer);
+	DoMaterialAttributeReorder(&Refraction, UEVer, RenderObjVer, UE5MainVer);
+	DoMaterialAttributeReorder(&CustomizedUVs[0], UEVer, RenderObjVer, UE5MainVer);
+	DoMaterialAttributeReorder(&CustomizedUVs[1], UEVer, RenderObjVer, UE5MainVer);
+	DoMaterialAttributeReorder(&CustomizedUVs[2], UEVer, RenderObjVer, UE5MainVer);
+	DoMaterialAttributeReorder(&CustomizedUVs[3], UEVer, RenderObjVer, UE5MainVer);
+	DoMaterialAttributeReorder(&CustomizedUVs[4], UEVer, RenderObjVer, UE5MainVer);
+	DoMaterialAttributeReorder(&CustomizedUVs[5], UEVer, RenderObjVer, UE5MainVer);
+	DoMaterialAttributeReorder(&CustomizedUVs[6], UEVer, RenderObjVer, UE5MainVer);
+	DoMaterialAttributeReorder(&CustomizedUVs[7], UEVer, RenderObjVer, UE5MainVer);
+	DoMaterialAttributeReorder(&PixelDepthOffset, UEVer, RenderObjVer, UE5MainVer);
+	DoMaterialAttributeReorder(&ShadingModelFromMaterialExpression, UEVer, RenderObjVer, UE5MainVer);
+	DoMaterialAttributeReorder(&FrontMaterial, UEVer, RenderObjVer, UE5MainVer);
 #endif // WITH_EDITORONLY_DATA
 
 	if (!IsDefaultMaterial())
@@ -3394,11 +3270,11 @@ void UMaterial::PostLoad()
 #if WITH_EDITORONLY_DATA
 	// Ensure expressions have been postloaded before we use them for compiling
 	// Any UObjects used by material compilation must be postloaded here
-	for (UMaterialExpression* Expression : GetExpressions())
+	for (int32 ExpressionIndex = 0; ExpressionIndex < Expressions.Num(); ExpressionIndex++)
 	{
-		if (Expression)
+		if (Expressions[ExpressionIndex])
 		{
-			Expression->ConditionalPostLoad();
+			Expressions[ExpressionIndex]->ConditionalPostLoad();
 		}
 	}
 #endif // WITH_EDITORONLY_DATA
@@ -3445,7 +3321,7 @@ void UMaterial::PostLoad()
 	if (GIsEditor)
 	{
 		// Clean up any removed material expression classes	
-		if (EditorOnly->ExpressionCollection.Expressions.Remove(nullptr) != 0)
+		if (Expressions.Remove(NULL) != 0)
 		{
 			// Force this material to recompile because its expressions have changed
 			// We're not providing a deterministic transformation guid because there could be many different ways expression
@@ -3470,9 +3346,10 @@ void UMaterial::PostLoad()
 	if ( GMaterialsThatNeedSamplerFixup.Get( this ) )
 	{
 		GMaterialsThatNeedSamplerFixup.Clear( this );
-		for (UMaterialExpression* Expression : GetExpressions())
+		const int32 ExpressionCount = Expressions.Num();
+		for ( int32 ExpressionIndex = 0; ExpressionIndex < ExpressionCount; ++ExpressionIndex )
 		{
-			UMaterialExpressionTextureBase* TextureExpression = Cast<UMaterialExpressionTextureBase>(Expression);
+			UMaterialExpressionTextureBase* TextureExpression = Cast<UMaterialExpressionTextureBase>(Expressions[ExpressionIndex]);
 			if ( TextureExpression && TextureExpression->Texture )
 			{
 				switch( TextureExpression->Texture->CompressionSettings )
@@ -3580,27 +3457,27 @@ void UMaterial::PostLoad()
 	if (GMaterialsThatNeedExpressionsFlipped.Get(this))
 	{
 		GMaterialsThatNeedExpressionsFlipped.Clear(this);
-		FlipExpressionPositions(GetExpressions(), EditorOnly->ExpressionCollection.EditorComments, true, this);
+		FlipExpressionPositions(Expressions, EditorComments, true, this);
 	}
 	else if (GMaterialsThatNeedCoordinateCheck.Get(this))
 	{
 		GMaterialsThatNeedCoordinateCheck.Clear(this);
 		if (HasFlippedCoordinates())
 		{
-			FlipExpressionPositions(GetExpressions(), EditorOnly->ExpressionCollection.EditorComments, false, this);
+			FlipExpressionPositions(Expressions, EditorComments, false, this);
 		}
-		FixCommentPositions(EditorOnly->ExpressionCollection.EditorComments);
+		FixCommentPositions(EditorComments);
 	}
 	else if (GMaterialsThatNeedCommentFix.Get(this))
 	{
 		GMaterialsThatNeedCommentFix.Clear(this);
-		FixCommentPositions(EditorOnly->ExpressionCollection.EditorComments);
+		FixCommentPositions(EditorComments);
 	}
 
 	if (GMaterialsThatNeedFeatureLevelSM6Fix.Get(this))
 	{
 		GMaterialsThatNeedFeatureLevelSM6Fix.Clear(this);
-		if (FixFeatureLevelNodesForSM6(EditorOnly->ExpressionCollection.Expressions))
+		if (FixFeatureLevelNodesForSM6(Expressions))
 		{
 			// Change this guid if you change the conversion logic.
 			static FGuid BackwardsCompatibilityFeatureLevelSM6ConversionGuid(TEXT("FC75DED7-2FB3-463B-B56C-8295871A340C"));
@@ -3675,64 +3552,59 @@ void UMaterial::GetShaderTypes(EShaderPlatform ShaderPlatform, const ITargetPlat
 }
 #endif // WITH_EDITOR
 
-bool UMaterial::IsPropertyConnected(EMaterialProperty Property) const
-{
-	return GetCachedExpressionData().IsPropertyConnected(Property);
-}
-
 bool UMaterial::HasBaseColorConnected() const
 {
-	return IsPropertyConnected(MP_BaseColor);
+	return BaseColor.IsConnected() || GetCachedExpressionData().IsMaterialAttributePropertyConnected(MP_BaseColor);
 }
 
 bool UMaterial::HasRoughnessConnected() const
 {
-	return IsPropertyConnected(MP_Roughness);
+	return Roughness.IsConnected() || GetCachedExpressionData().IsMaterialAttributePropertyConnected(MP_Roughness);
 }
 
 bool UMaterial::HasAmbientOcclusionConnected() const
 {
-	return IsPropertyConnected(MP_AmbientOcclusion);
+	return AmbientOcclusion.IsConnected() || GetCachedExpressionData().IsMaterialAttributePropertyConnected(MP_AmbientOcclusion);
 }
 
 bool UMaterial::HasNormalConnected() const
 {
-	return IsPropertyConnected(MP_Normal);
+	return Normal.IsConnected() || GetCachedExpressionData().IsMaterialAttributePropertyConnected(MP_Normal);
 }
 
 bool UMaterial::HasSpecularConnected() const
 {
-	return IsPropertyConnected(MP_Specular);
+	return Specular.IsConnected() || GetCachedExpressionData().IsMaterialAttributePropertyConnected(MP_Specular);
 }
 
 bool UMaterial::HasMetallicConnected() const
 {
-	return IsPropertyConnected(MP_Metallic);
+	return Metallic.IsConnected() || GetCachedExpressionData().IsMaterialAttributePropertyConnected(MP_Metallic);
 }
 
 bool UMaterial::HasEmissiveColorConnected() const
 {
-	return IsPropertyConnected(MP_EmissiveColor);
+	return EmissiveColor.IsConnected() || GetCachedExpressionData().IsMaterialAttributePropertyConnected(MP_EmissiveColor);
 }
 
 bool UMaterial::HasAnisotropyConnected() const
 {
-	return IsPropertyConnected(MP_Anisotropy);
+	return Anisotropy.IsConnected() || GetCachedExpressionData().IsMaterialAttributePropertyConnected(MP_Anisotropy);
 }
 
 bool UMaterial::HasStrataFrontMaterialConnected() const
 {
-	return IsPropertyConnected(MP_FrontMaterial);
+	return FrontMaterial.IsConnected();
 }
 
 bool UMaterial::HasVertexPositionOffsetConnected() const
 {
-	return IsPropertyConnected(MP_WorldPositionOffset);
+	return WorldPositionOffset.IsConnected() || GetCachedExpressionData().IsMaterialAttributePropertyConnected(MP_WorldPositionOffset);
 }
 
 bool UMaterial::HasPixelDepthOffsetConnected() const
 {
-	return IsPropertyConnected(MP_PixelDepthOffset);
+	return PixelDepthOffset.IsConnected() || GetCachedExpressionData().IsMaterialAttributePropertyConnected(MP_PixelDepthOffset);
 }
 
 void UMaterial::PropagateDataToMaterialProxy()
@@ -3813,7 +3685,6 @@ bool UMaterial::CanEditChange(const FProperty* InProperty) const
 {
 	if (InProperty)
 	{
-		const UMaterialEditorOnlyData* EditorOnly = GetEditorOnlyData();
 		FString PropertyName = InProperty->GetName();
 
 		if (PropertyName == GET_MEMBER_NAME_STRING_CHECKED(UMaterial, PhysMaterial) || PropertyName == GET_MEMBER_NAME_STRING_CHECKED(UMaterial, PhysMaterialMask))
@@ -3906,7 +3777,7 @@ bool UMaterial::CanEditChange(const FProperty* InProperty) const
 		}
 		else if (PropertyName == GET_MEMBER_NAME_STRING_CHECKED(UMaterial, RefractionDepthBias))
 		{
-			return EditorOnly->Refraction.IsConnected();
+			return Refraction.IsConnected();
 		}
 	
 		if (PropertyName == GET_MEMBER_NAME_STRING_CHECKED(UMaterial, TranslucencyPass)
@@ -3963,19 +3834,18 @@ void UMaterial::CreateExecutionFlowExpressions()
 {
 	if (IsUsingControlFlow())
 	{
-		UMaterialEditorOnlyData* EditorOnly = GetEditorOnlyData();
-		if (!EditorOnly->ExpressionCollection.ExpressionExecBegin)
+		if (!ExpressionExecBegin)
 		{
-			EditorOnly->ExpressionCollection.ExpressionExecBegin = NewObject<UMaterialExpressionExecBegin>(this);
-			EditorOnly->ExpressionCollection.ExpressionExecBegin->Material = this;
-			EditorOnly->ExpressionCollection.Expressions.Add(EditorOnly->ExpressionCollection.ExpressionExecBegin);
+			ExpressionExecBegin = NewObject<UMaterialExpressionExecBegin>(this);
+			ExpressionExecBegin->Material = this;
+			Expressions.Add(ExpressionExecBegin);
 		}
 
-		if (!EditorOnly->ExpressionCollection.ExpressionExecEnd)
+		if (!ExpressionExecEnd)
 		{
-			EditorOnly->ExpressionCollection.ExpressionExecEnd = NewObject<UMaterialExpressionExecEnd>(this);
-			EditorOnly->ExpressionCollection.ExpressionExecEnd->Material = this;
-			EditorOnly->ExpressionCollection.Expressions.Add(EditorOnly->ExpressionCollection.ExpressionExecEnd);
+			ExpressionExecEnd = NewObject<UMaterialExpressionExecEnd>(this);
+			ExpressionExecEnd->Material = this;
+			Expressions.Add(ExpressionExecEnd);
 		}
 	}
 }
@@ -4004,30 +3874,28 @@ void UMaterial::PostEditChangePropertyInternal(FPropertyChangedEvent& PropertyCh
 	//Cancel any current compilation jobs that are in flight for this material.
 	CancelOutstandingCompilation();
 
-	const UMaterialEditorOnlyData* EditorOnly = GetEditorOnlyData();
-
 	// check for distortion in material 
 	{
 		bUsesDistortion = false;
 		// check for a distortion value
-		if (EditorOnly->Refraction.Expression
-			|| (EditorOnly->Refraction.UseConstant && FMath::Abs(EditorOnly->Refraction.Constant - 1.0f) >= UE_KINDA_SMALL_NUMBER))
+		if (Refraction.Expression
+			|| (Refraction.UseConstant && FMath::Abs(Refraction.Constant - 1.0f) >= UE_KINDA_SMALL_NUMBER))
 		{
 			bUsesDistortion = true;
 		}
 
 		// check the material attributes for refraction expressions as well
-		if (EditorOnly->MaterialAttributes.Expression)
+		if (MaterialAttributes.Expression)
 		{
 			// handle make attribute expressions
-			UMaterialExpressionMakeMaterialAttributes * MakeAttributeExpression = Cast<UMaterialExpressionMakeMaterialAttributes>(EditorOnly->MaterialAttributes.Expression);
+			UMaterialExpressionMakeMaterialAttributes * MakeAttributeExpression = Cast<UMaterialExpressionMakeMaterialAttributes>(MaterialAttributes.Expression);
 			if (MakeAttributeExpression && MakeAttributeExpression->Refraction.Expression)
 			{
 				bUsesDistortion = true;
 			}
 
 			// handle set attribute expressions
-			UMaterialExpressionSetMaterialAttributes * SetAttributeExpression = Cast<UMaterialExpressionSetMaterialAttributes>(EditorOnly->MaterialAttributes.Expression);
+			UMaterialExpressionSetMaterialAttributes * SetAttributeExpression = Cast<UMaterialExpressionSetMaterialAttributes>(MaterialAttributes.Expression);
 			if (SetAttributeExpression)
 			{
 				for (int32 Index = 0; Index < SetAttributeExpression->Inputs.Num(); Index++)
@@ -4044,7 +3912,7 @@ void UMaterial::PostEditChangePropertyInternal(FPropertyChangedEvent& PropertyCh
 	}
 
 	//If we can be sure this material would be the same opaque as it is masked then allow it to be assumed opaque.
-	bCanMaskedBeAssumedOpaque = !EditorOnly->OpacityMask.Expression && !(EditorOnly->OpacityMask.UseConstant && EditorOnly->OpacityMask.Constant < 0.999f) && !bUseMaterialAttributes;
+	bCanMaskedBeAssumedOpaque = !OpacityMask.Expression && !(OpacityMask.UseConstant && OpacityMask.Constant < 0.999f) && !bUseMaterialAttributes;
 
 	// If the strata blending mode is changed, make sure we update the legacy blend mode too for the shaders to be filtered to passes correctly at runtime.
 	if (Engine_IsStrataEnabled())
@@ -4173,9 +4041,10 @@ bool UMaterial::IsDynamicParameter(const UMaterialExpression* Expression)
 void UMaterial::BuildEditorParameterList()
 {
 	EditorParameters.Empty();
-	for(UMaterialExpression* Expression : GetExpressions())
+
+	for(int32 MaterialExpressionIndex = 0 ; MaterialExpressionIndex < Expressions.Num() ; ++MaterialExpressionIndex)
 	{
-		AddExpressionParameter(Expression, EditorParameters);
+		AddExpressionParameter(Expressions[MaterialExpressionIndex], EditorParameters);
 	}
 }
 
@@ -4201,15 +4070,15 @@ bool UMaterial::HasDuplicateParameters(const UMaterialExpression* Expression)
 	return false;
 }
 
-bool UMaterial::HasDuplicateDynamicParameters(const UMaterialExpression* InExpression)
+bool UMaterial::HasDuplicateDynamicParameters(const UMaterialExpression* Expression)
 {
-	const UMaterialExpressionDynamicParameter* DynParam = Cast<UMaterialExpressionDynamicParameter>(InExpression);
+	const UMaterialExpressionDynamicParameter* DynParam = Cast<UMaterialExpressionDynamicParameter>(Expression);
 	if (DynParam)
 	{
-		for (UMaterialExpression* Expression : GetExpressions())
+		for (int32 ExpIndex = 0; ExpIndex < Expressions.Num(); ExpIndex++)
 		{
-			UMaterialExpressionDynamicParameter* CheckDynParam = Cast<UMaterialExpressionDynamicParameter>(Expression);
-			if (CheckDynParam != InExpression)
+			UMaterialExpressionDynamicParameter* CheckDynParam = Cast<UMaterialExpressionDynamicParameter>(Expressions[ExpIndex]);
+			if (CheckDynParam != Expression)
 			{
 				return true;
 			}
@@ -4218,14 +4087,14 @@ bool UMaterial::HasDuplicateDynamicParameters(const UMaterialExpression* InExpre
 	return false;
 }
 
-void UMaterial::UpdateExpressionDynamicParameters(const UMaterialExpression* InExpression)
+void UMaterial::UpdateExpressionDynamicParameters(const UMaterialExpression* Expression)
 {
-	const UMaterialExpressionDynamicParameter* DynParam = Cast<UMaterialExpressionDynamicParameter>(InExpression);
+	const UMaterialExpressionDynamicParameter* DynParam = Cast<UMaterialExpressionDynamicParameter>(Expression);
 	if (DynParam)
 	{
-		for (UMaterialExpression* Expression : GetExpressions())
+		for (int32 ExpIndex = 0; ExpIndex < Expressions.Num(); ExpIndex++)
 		{
-			UMaterialExpressionDynamicParameter* CheckParam = Cast<UMaterialExpressionDynamicParameter>(Expression);
+			UMaterialExpressionDynamicParameter* CheckParam = Cast<UMaterialExpressionDynamicParameter>(Expressions[ExpIndex]);
 			if (CheckParam && CheckParam->CopyDynamicParameterProperties(DynParam))
 			{
 				CheckParam->GraphNode->ReconstructNode();
@@ -4263,6 +4132,8 @@ void UMaterial::PropagateExpressionParameterChanges(const FName& ParameterName,
 
 void UMaterial::UpdateExpressionParameterName(UMaterialExpression* Expression)
 {
+	FName ExpressionName;
+
 	for(TMap<FName, TArray<UMaterialExpression*> >::TIterator Iter(EditorParameters); Iter; ++Iter)
 	{
 		if(Iter.Value().Remove(Expression) > 0)
@@ -4282,15 +4153,13 @@ void UMaterial::RebuildShadingModelField()
 {
 	ShadingModels.ClearShadingModels();
 	const bool bStrataEnabled = Engine_IsStrataEnabled();
-	UMaterialEditorOnlyData* EditorOnly = GetEditorOnlyData();
-
-	if (bStrataEnabled && EditorOnly->FrontMaterial.IsConnected())
+	if (bStrataEnabled && FrontMaterial.IsConnected())
 	{
 		FStrataMaterialInfo StrataMaterialInfo;
-		check(EditorOnly->FrontMaterial.Expression);
-		if (EditorOnly->FrontMaterial.Expression->IsResultStrataMaterial(EditorOnly->FrontMaterial.OutputIndex))
+		check(this->FrontMaterial.Expression);
+		if (this->FrontMaterial.Expression->IsResultStrataMaterial(this->FrontMaterial.OutputIndex))
 		{
-			EditorOnly->FrontMaterial.Expression->GatherStrataMaterialInfo(StrataMaterialInfo, EditorOnly->FrontMaterial.OutputIndex);
+			this->FrontMaterial.Expression->GatherStrataMaterialInfo(StrataMaterialInfo, this->FrontMaterial.OutputIndex);
 			this->CachedConnectedInputs = StrataMaterialInfo.GetPropertyConnected();
 		}
 
@@ -5079,33 +4948,32 @@ static void SetMaterialInputDescription(FStrataMaterialInput& Input, bool bHidde
 
 bool UMaterial::GetExpressionInputDescription(EMaterialProperty InProperty, FMaterialInputDescription& OutDescription)
 {
-	UMaterialEditorOnlyData* EditorOnly = GetEditorOnlyData();
 	switch (InProperty)
 	{
-	case MP_EmissiveColor: SetMaterialInputDescription(EditorOnly->EmissiveColor, false, OutDescription); return true;
-	case MP_Opacity: SetMaterialInputDescription(EditorOnly->Opacity, false, OutDescription); return true;
-	case MP_OpacityMask: SetMaterialInputDescription(EditorOnly->OpacityMask, false, OutDescription); return true;
-	case MP_BaseColor: SetMaterialInputDescription(EditorOnly->BaseColor, false, OutDescription); return true;
-	case MP_Metallic: SetMaterialInputDescription(EditorOnly->Metallic, false, OutDescription); return true;
-	case MP_Specular: SetMaterialInputDescription(EditorOnly->Specular, false, OutDescription); return true;
-	case MP_Roughness: SetMaterialInputDescription(EditorOnly->Roughness, false, OutDescription); return true;
-	case MP_Anisotropy: SetMaterialInputDescription(EditorOnly->Anisotropy, false, OutDescription); return true;
-	case MP_Normal: SetMaterialInputDescription(EditorOnly->Normal, false, OutDescription); return true;
-	case MP_Tangent: SetMaterialInputDescription(EditorOnly->Tangent, false, OutDescription); return true;
-	case MP_WorldPositionOffset: SetMaterialInputDescription(EditorOnly->WorldPositionOffset, false, OutDescription); return true;
-	case MP_SubsurfaceColor: SetMaterialInputDescription(EditorOnly->SubsurfaceColor, false, OutDescription); return true;
-	case MP_CustomData0: SetMaterialInputDescription(EditorOnly->ClearCoat, false, OutDescription); return true;
-	case MP_CustomData1: SetMaterialInputDescription(EditorOnly->ClearCoatRoughness, false, OutDescription); return true;
-	case MP_AmbientOcclusion: SetMaterialInputDescription(EditorOnly->AmbientOcclusion, false, OutDescription); return true;
-	case MP_Refraction: SetMaterialInputDescription(EditorOnly->Refraction, false, OutDescription); return true;
-	case MP_MaterialAttributes: SetMaterialInputDescription(EditorOnly->MaterialAttributes, false, OutDescription); return true;
-	case MP_PixelDepthOffset: SetMaterialInputDescription(EditorOnly->PixelDepthOffset, false, OutDescription); return true;
-	case MP_ShadingModel: SetMaterialInputDescription(EditorOnly->ShadingModelFromMaterialExpression, false, OutDescription); return true;
-	case MP_FrontMaterial: SetMaterialInputDescription(EditorOnly->FrontMaterial, false, OutDescription); return true;
+	case MP_EmissiveColor: SetMaterialInputDescription(EmissiveColor, false, OutDescription); return true;
+	case MP_Opacity: SetMaterialInputDescription(Opacity, false, OutDescription); return true;
+	case MP_OpacityMask: SetMaterialInputDescription(OpacityMask, false, OutDescription); return true;
+	case MP_BaseColor: SetMaterialInputDescription(BaseColor, false, OutDescription); return true;
+	case MP_Metallic: SetMaterialInputDescription(Metallic, false, OutDescription); return true;
+	case MP_Specular: SetMaterialInputDescription(Specular, false, OutDescription); return true;
+	case MP_Roughness: SetMaterialInputDescription(Roughness, false, OutDescription); return true;
+	case MP_Anisotropy: SetMaterialInputDescription(Anisotropy, false, OutDescription); return true;
+	case MP_Normal: SetMaterialInputDescription(Normal, false, OutDescription); return true;
+	case MP_Tangent: SetMaterialInputDescription(Tangent, false, OutDescription); return true;
+	case MP_WorldPositionOffset: SetMaterialInputDescription(WorldPositionOffset, false, OutDescription); return true;
+	case MP_SubsurfaceColor: SetMaterialInputDescription(SubsurfaceColor, false, OutDescription); return true;
+	case MP_CustomData0: SetMaterialInputDescription(ClearCoat, false, OutDescription); return true;
+	case MP_CustomData1: SetMaterialInputDescription(ClearCoatRoughness, false, OutDescription); return true;
+	case MP_AmbientOcclusion: SetMaterialInputDescription(AmbientOcclusion, false, OutDescription); return true;
+	case MP_Refraction: SetMaterialInputDescription(Refraction, false, OutDescription); return true;
+	case MP_MaterialAttributes: SetMaterialInputDescription(MaterialAttributes, false, OutDescription); return true;
+	case MP_PixelDepthOffset: SetMaterialInputDescription(PixelDepthOffset, false, OutDescription); return true;
+	case MP_ShadingModel: SetMaterialInputDescription(ShadingModelFromMaterialExpression, false, OutDescription); return true;
+	case MP_FrontMaterial: SetMaterialInputDescription(FrontMaterial, false, OutDescription); return true;
 	default:
 		if (InProperty >= MP_CustomizedUVs0 && InProperty <= MP_CustomizedUVs7)
 		{
-			SetMaterialInputDescription(EditorOnly->CustomizedUVs[InProperty - MP_CustomizedUVs0], false, OutDescription);
+			SetMaterialInputDescription(CustomizedUVs[InProperty - MP_CustomizedUVs0], false, OutDescription);
 			return true;
 		}
 		return false;
@@ -5117,7 +4985,7 @@ bool UMaterial::GetExpressionInputDescription(EMaterialProperty InProperty, FMat
 #if WITH_EDITORONLY_DATA
 void UMaterial::GetAllFunctionOutputExpressions(TArray<class UMaterialExpressionFunctionOutput*>& OutFunctionOutputs) const
 {
-	for (UMaterialExpression* Expression : GetExpressions())
+	for (UMaterialExpression* Expression : Expressions)
 	{
 		UMaterialExpressionFunctionOutput* FunctionOutput = Cast<UMaterialExpressionFunctionOutput>(Expression);
 		if (FunctionOutput)
@@ -5129,7 +4997,7 @@ void UMaterial::GetAllFunctionOutputExpressions(TArray<class UMaterialExpression
 
 void UMaterial::GetAllCustomOutputExpressions(TArray<class UMaterialExpressionCustomOutput*>& OutCustomOutputs) const
 {
-	for (UMaterialExpression* Expression : GetExpressions())
+	for (UMaterialExpression* Expression : Expressions)
 	{
 		UMaterialExpressionCustomOutput* CustomOutput = Cast<UMaterialExpressionCustomOutput>(Expression);
 		if (CustomOutput)
@@ -5141,7 +5009,7 @@ void UMaterial::GetAllCustomOutputExpressions(TArray<class UMaterialExpressionCu
 
 void UMaterial::GetAllExpressionsForCustomInterpolators(TArray<class UMaterialExpression*>& OutExpressions) const
 {
-	for (UMaterialExpression* Expression : GetExpressions())
+	for (UMaterialExpression* Expression : Expressions)
 	{
 		if (Expression &&
 			(Expression->IsA(UMaterialExpressionVertexInterpolator::StaticClass()) ||
@@ -5283,19 +5151,16 @@ bool UMaterial::GetExpressionsInPropertyChain(EMaterialProperty InProperty,
 
 bool UMaterial::GetGroupSortPriority(const FString& InGroupName, int32& OutSortPriority) const
 {
-	const UMaterialEditorOnlyData* LocalData = GetEditorOnlyData();
-	if (LocalData)
+	const FParameterGroupData* ParameterGroupDataElement = ParameterGroupData.FindByPredicate([&InGroupName](const FParameterGroupData& DataElement)
 	{
-		const FParameterGroupData* ParameterGroupDataElement = LocalData->ParameterGroupData.FindByPredicate([&InGroupName](const FParameterGroupData& DataElement)
-		{
-			return InGroupName == DataElement.GroupName;
-		});
-		if (ParameterGroupDataElement != nullptr)
-		{
-			OutSortPriority = ParameterGroupDataElement->GroupSortPriority;
-			return true;
-		}
+		return InGroupName == DataElement.GroupName;
+	});
+	if (ParameterGroupDataElement != nullptr)
+	{
+		OutSortPriority = ParameterGroupDataElement->GroupSortPriority;
+		return true;
 	}
+
 	return false;
 }
 
@@ -5454,9 +5319,9 @@ bool UMaterial::RecursiveGetExpressionChain(UMaterialExpression* InExpression, T
 						{
 							bool bUseInputA = StaticSwitchExp->DefaultValue;
 							FName StaticSwitchExpName = StaticSwitchExp->ParameterName;
-							for (int32 CheckIdx = 0; CheckIdx < InStaticParameterSet->EditorOnly.StaticSwitchParameters.Num(); CheckIdx++)
+							for (int32 CheckIdx = 0; CheckIdx < InStaticParameterSet->StaticSwitchParameters.Num(); CheckIdx++)
 							{
-								FStaticSwitchParameter& SwitchParam = InStaticParameterSet->EditorOnly.StaticSwitchParameters[CheckIdx];
+								FStaticSwitchParameter& SwitchParam = InStaticParameterSet->StaticSwitchParameters[CheckIdx];
 								if (SwitchParam.ParameterInfo.Name == StaticSwitchExpName)
 								{
 									// Found it...
@@ -5558,49 +5423,48 @@ void UMaterial::RecursiveUpdateRealtimePreview( UMaterialExpression* InExpressio
 int32 UMaterial::CompilePropertyEx( FMaterialCompiler* Compiler, const FGuid& AttributeID )
 {
 	const EMaterialProperty Property = FMaterialAttributeDefinitionMap::GetProperty(AttributeID);
-	UMaterialEditorOnlyData* EditorOnly = GetEditorOnlyData();
 
 	if (IsUsingControlFlow())
 	{
-		check(EditorOnly->ExpressionCollection.ExpressionExecBegin);
-		return EditorOnly->ExpressionCollection.ExpressionExecBegin->Compile(Compiler, UMaterialExpression::CompileExecutionOutputIndex);
+		check(ExpressionExecBegin);
+		return ExpressionExecBegin->Compile(Compiler, UMaterialExpression::CompileExecutionOutputIndex);
 	}
 
 	if( bUseMaterialAttributes && MP_DiffuseColor != Property && MP_SpecularColor != Property )
 	{
-		return EditorOnly->MaterialAttributes.CompileWithDefault(Compiler, AttributeID);
+		return MaterialAttributes.CompileWithDefault(Compiler, AttributeID);
 	}
 
 	switch (Property)
 	{
-		case MP_Opacity:				return EditorOnly->Opacity.CompileWithDefault(Compiler, Property);
-		case MP_OpacityMask:			return EditorOnly->OpacityMask.CompileWithDefault(Compiler, Property);
-		case MP_Metallic:				return EditorOnly->Metallic.CompileWithDefault(Compiler, Property);
-		case MP_Specular:				return EditorOnly->Specular.CompileWithDefault(Compiler, Property);
-		case MP_Roughness:				return EditorOnly->Roughness.CompileWithDefault(Compiler, Property);
-		case MP_Anisotropy:				return EditorOnly->Anisotropy.CompileWithDefault(Compiler, Property);
-		case MP_CustomData0:			return EditorOnly->ClearCoat.CompileWithDefault(Compiler, Property);
-		case MP_CustomData1:			return EditorOnly->ClearCoatRoughness.CompileWithDefault(Compiler, Property);
-		case MP_AmbientOcclusion:		return EditorOnly->AmbientOcclusion.CompileWithDefault(Compiler, Property);
-		case MP_Refraction:				return EditorOnly->Refraction.CompileWithDefault(Compiler, Property);
-		case MP_EmissiveColor:			return EditorOnly->EmissiveColor.CompileWithDefault(Compiler, Property);
-		case MP_BaseColor:				return EditorOnly->BaseColor.CompileWithDefault(Compiler, Property);
-		case MP_SubsurfaceColor:		return EditorOnly->SubsurfaceColor.CompileWithDefault(Compiler, Property);
-		case MP_Normal:					return EditorOnly->Normal.CompileWithDefault(Compiler, Property);
-		case MP_Tangent:				return EditorOnly->Tangent.CompileWithDefault(Compiler, Property);
-		case MP_WorldPositionOffset:	return EditorOnly->WorldPositionOffset.CompileWithDefault(Compiler, Property);
-		case MP_PixelDepthOffset:		return EditorOnly->PixelDepthOffset.CompileWithDefault(Compiler, Property);
-		case MP_ShadingModel:			return EditorOnly->ShadingModelFromMaterialExpression.CompileWithDefault(Compiler, Property);
-		case MP_FrontMaterial:			return EditorOnly->FrontMaterial.CompileWithDefault(Compiler, Property);
+		case MP_Opacity:				return Opacity.CompileWithDefault(Compiler, Property);
+		case MP_OpacityMask:			return OpacityMask.CompileWithDefault(Compiler, Property);
+		case MP_Metallic:				return Metallic.CompileWithDefault(Compiler, Property);
+		case MP_Specular:				return Specular.CompileWithDefault(Compiler, Property);
+		case MP_Roughness:				return Roughness.CompileWithDefault(Compiler, Property);
+		case MP_Anisotropy:				return Anisotropy.CompileWithDefault(Compiler, Property);
+		case MP_CustomData0:			return ClearCoat.CompileWithDefault(Compiler, Property);
+		case MP_CustomData1:			return ClearCoatRoughness.CompileWithDefault(Compiler, Property);
+		case MP_AmbientOcclusion:		return AmbientOcclusion.CompileWithDefault(Compiler, Property);
+		case MP_Refraction:				return Refraction.CompileWithDefault(Compiler, Property);
+		case MP_EmissiveColor:			return EmissiveColor.CompileWithDefault(Compiler, Property);
+		case MP_BaseColor:				return BaseColor.CompileWithDefault(Compiler, Property);
+		case MP_SubsurfaceColor:		return SubsurfaceColor.CompileWithDefault(Compiler, Property);
+		case MP_Normal:					return Normal.CompileWithDefault(Compiler, Property);
+		case MP_Tangent:				return Tangent.CompileWithDefault(Compiler, Property);
+		case MP_WorldPositionOffset:	return WorldPositionOffset.CompileWithDefault(Compiler, Property);
+		case MP_PixelDepthOffset:		return PixelDepthOffset.CompileWithDefault(Compiler, Property);
+		case MP_ShadingModel:			return ShadingModelFromMaterialExpression.CompileWithDefault(Compiler, Property);
+		case MP_FrontMaterial:			return FrontMaterial.CompileWithDefault(Compiler, Property);
 
 		default:
 			if (Property >= MP_CustomizedUVs0 && Property <= MP_CustomizedUVs7)
 			{
 				const int32 TextureCoordinateIndex = Property - MP_CustomizedUVs0;
 
-				if (TextureCoordinateIndex < NumCustomizedUVs && EditorOnly->CustomizedUVs[TextureCoordinateIndex].Expression)
+				if (TextureCoordinateIndex < NumCustomizedUVs && CustomizedUVs[TextureCoordinateIndex].Expression)
 				{
-					return EditorOnly->CustomizedUVs[TextureCoordinateIndex].CompileWithDefault(Compiler, Property);
+					return CustomizedUVs[TextureCoordinateIndex].CompileWithDefault(Compiler, Property);
 				}
 				else
 				{
@@ -6116,7 +5980,7 @@ bool UMaterial::IsPropertyActiveInEditor(EMaterialProperty InProperty) const
 		ShadingModels,
 		TranslucencyLightingMode,
 		BlendableOutputAlpha,
-		GetEditorOnlyData()->Refraction.IsConnected(),
+		Refraction.IsConnected(),
 		IsShadingModelFromMaterialExpression(),
 		IsTranslucencyWritingVelocity(),
 		IsPropertySupported(InProperty));
@@ -6125,11 +5989,6 @@ bool UMaterial::IsPropertyActiveInEditor(EMaterialProperty InProperty) const
 
 bool UMaterial::IsPropertyActiveInDerived(EMaterialProperty InProperty, const UMaterialInterface* DerivedMaterial) const
 {
-#if WITH_EDITOR
-	// Make sure our cached data is consistent
-	ensureMsgf(!(GetEditorOnlyData()->Refraction.IsConnected() && !GetCachedExpressionData().IsPropertyConnected(MP_Refraction)),
-		TEXT("GetCachedExpressionData() says refraction isn't connected, but GetEditorOnlyData() says it is"));
-#endif
 	return IsPropertyActive_Internal(InProperty,
 		MaterialDomain,
 		DerivedMaterial->GetBlendMode(),
@@ -6137,14 +5996,14 @@ bool UMaterial::IsPropertyActiveInDerived(EMaterialProperty InProperty, const UM
 		DerivedMaterial->GetShadingModels(),
 		TranslucencyLightingMode,
 		BlendableOutputAlpha,
-		GetCachedExpressionData().IsPropertyConnected(MP_Refraction),
+		Refraction.IsConnected(),
 		DerivedMaterial->IsShadingModelFromMaterialExpression(),
 		IsTranslucencyWritingVelocity(),
 		IsPropertySupported(InProperty));
 }
 
 #if WITH_EDITORONLY_DATA
-void UMaterial::FlipExpressionPositions(TConstArrayView<TObjectPtr<UMaterialExpression>> Expressions, TConstArrayView<TObjectPtr<UMaterialExpressionComment>> Comments, bool bScaleCoords, UMaterial* InMaterial)
+void UMaterial::FlipExpressionPositions(const TArray<UMaterialExpression*>& Expressions, const TArray<UMaterialExpressionComment*>& Comments, bool bScaleCoords, UMaterial* InMaterial)
 {
 	// Rough estimate of average increase in node size for the new editor
 	const float PosScaling = bScaleCoords ? 1.25f : 1.0f;
@@ -6175,7 +6034,7 @@ void UMaterial::FlipExpressionPositions(TConstArrayView<TObjectPtr<UMaterialExpr
 	}
 }
 
-void UMaterial::FixCommentPositions(TConstArrayView<TObjectPtr<UMaterialExpressionComment>> Comments)
+void UMaterial::FixCommentPositions(const TArray<UMaterialExpressionComment*>& Comments)
 {
 	// equivalent to 1/1.25 * 0.25 to get the amount that should have been used when first flipping
 	const float SizeScaling = 0.2f;
@@ -6250,15 +6109,4 @@ void UMaterial::GetLightingGuidChain(bool bIncludeTextures, TArray<FGuid>& OutGu
 #endif
 }
 
-UMaterialEditorOnlyData::UMaterialEditorOnlyData()
-{
-	BaseColor.Constant = FColor(128, 128, 128);
-	Metallic.Constant = 0.0f;
-	Specular.Constant = 0.5f;
-	Roughness.Constant = 0.5f;
-
-	Opacity.Constant = 1.0f;
-	OpacityMask.Constant = 1.0f;
-}
-
 #undef LOCTEXT_NAMESPACE
diff --git a/Engine/Source/Runtime/Engine/Private/Materials/MaterialCachedData.cpp b/Engine/Source/Runtime/Engine/Private/Materials/MaterialCachedData.cpp
--- a/Engine/Source/Runtime/Engine/Private/Materials/MaterialCachedData.cpp
+++ b/Engine/Source/Runtime/Engine/Private/Materials/MaterialCachedData.cpp
@@ -45,30 +45,38 @@
 
 const FMaterialCachedParameterEntry FMaterialCachedParameterEntry::EmptyData{};
 const FMaterialCachedExpressionData FMaterialCachedExpressionData::EmptyData{};
-const FMaterialCachedExpressionEditorOnlyData FMaterialCachedExpressionEditorOnlyData::EmptyData{};
 
-static_assert((uint32)(EMaterialProperty::MP_MAX)-1 <= (8 * sizeof(FMaterialCachedExpressionData::PropertyConnectedBitmask)), "PropertyConnectedBitmask cannot contain entire EMaterialProperty enumeration.");
-
-FMaterialCachedExpressionData::FMaterialCachedExpressionData()
-	: bHasMaterialLayers(false)
-	, bHasRuntimeVirtualTextureOutput(false)
-	, bHasSceneColor(false)
-	, bHasPerInstanceCustomData(false)
-	, bHasPerInstanceRandom(false)
-	, bHasVertexInterpolator(false)
+void FMaterialCachedExpressionData::Reset()
 {
+	Parameters.Reset();
+	ReferencedTextures.Reset();
+	FunctionInfos.Reset();
+	ParameterCollectionInfos.Reset();
+	GrassTypes.Reset();
+	DynamicParameterNames.Reset();
+	QualityLevelsUsed.Reset();
 	QualityLevelsUsed.AddDefaulted(EMaterialQualityLevel::Num);
+	bHasMaterialLayers = false;
+	bHasRuntimeVirtualTextureOutput = false;
+	bHasSceneColor = false;
+	bHasPerInstanceCustomData = false;
+	bHasPerInstanceRandom = false;
+	bHasVertexInterpolator = false;
+	MaterialAttributesPropertyConnectedBitmask = 0;
+
 #if WITH_EDITORONLY_DATA
-	EditorOnlyData = MakeShared<FMaterialCachedExpressionEditorOnlyData>();
+	LandscapeLayerNames.Reset();
 #endif // WITH_EDITORONLY_DATA
+
+	static_assert((uint32)(EMaterialProperty::MP_MAX)-1 <= (8 * sizeof(MaterialAttributesPropertyConnectedBitmask)), "MaterialAttributesPropertyConnectedBitmask cannot contain entire EMaterialProperty enumeration.");
 }
 
 void FMaterialCachedExpressionData::AddReferencedObjects(FReferenceCollector& Collector)
 {
 	Collector.AddReferencedObjects(ReferencedTextures);
-	Collector.AddReferencedObjects(GrassTypes);
 	Collector.AddReferencedObjects(MaterialLayers.Layers);
 	Collector.AddReferencedObjects(MaterialLayers.Blends);
+	Collector.AddReferencedObjects(GrassTypes);
 	for (FMaterialFunctionInfo& FunctionInfo : FunctionInfos)
 	{
 		Collector.AddReferencedObject(FunctionInfo.Function);
@@ -98,31 +106,28 @@ void FMaterialCachedExpressionData::AppendReferencedParameterCollectionIdsTo(TAr
 }
 
 #if WITH_EDITOR
-static int32 TryAddParameter(FMaterialCachedExpressionData& CachedData,
+static int32 TryAddParameter(FMaterialCachedParameters& CachedParameters,
 	EMaterialParameterType Type,
 	const FMaterialParameterInfo& ParameterInfo,
 	const FMaterialCachedParameterEditorInfo& InEditorInfo)
 {
-	check(CachedData.EditorOnlyData);
-	FMaterialCachedParameterEntry& Entry = CachedData.GetParameterTypeEntry(Type);
-	FMaterialCachedParameterEditorEntry& EditorEntry = CachedData.EditorOnlyData->EditorEntries[(int32)Type];
-
+	FMaterialCachedParameterEntry& Entry = CachedParameters.GetParameterTypeEntry(Type);
 	FSetElementId ElementId = Entry.ParameterInfoSet.FindId(ParameterInfo);
 	int32 Index = INDEX_NONE;
 	if (!ElementId.IsValidId())
 	{
 		ElementId = Entry.ParameterInfoSet.Add(ParameterInfo);
 		Index = ElementId.AsInteger();
-		EditorEntry.EditorInfo.Insert(InEditorInfo, Index);
+		Entry.EditorInfo.Insert(InEditorInfo, Index);
 		// should be valid as long as we don't ever remove elements from ParameterInfoSet
-		check(Entry.ParameterInfoSet.Num() == EditorEntry.EditorInfo.Num());
+		check(Entry.ParameterInfoSet.Num() == Entry.EditorInfo.Num());
 		return Index;
 	}
 
 	// Update any editor values that haven't been set yet
 	// TODO still need to do this??
 	Index = ElementId.AsInteger();
-	FMaterialCachedParameterEditorInfo& EditorInfo = EditorEntry.EditorInfo[Index];
+	FMaterialCachedParameterEditorInfo& EditorInfo = Entry.EditorInfo[Index];
 	if (!EditorInfo.ExpressionGuid.IsValid())
 	{
 		EditorInfo.ExpressionGuid = InEditorInfo.ExpressionGuid;
@@ -141,16 +146,9 @@ static int32 TryAddParameter(FMaterialCachedExpressionData& CachedData,
 	return INDEX_NONE;
 }
 
-void FMaterialCachedExpressionData::AddParameter(const FMaterialParameterInfo& ParameterInfo, const FMaterialParameterMetadata& ParameterMeta, UObject*& OutReferencedTexture)
+void FMaterialCachedParameters::AddParameter(const FMaterialParameterInfo& ParameterInfo, const FMaterialParameterMetadata& ParameterMeta, UObject*& OutReferencedTexture)
 {
-	check(EditorOnlyData);
-	int32 AssetIndex = INDEX_NONE;
-	if (!ParameterMeta.AssetPath.IsEmpty())
-	{
-		AssetIndex = EditorOnlyData->AssetPaths.AddUnique(ParameterMeta.AssetPath);
-	}
-
-	const FMaterialCachedParameterEditorInfo EditorInfo(ParameterMeta.ExpressionGuid, ParameterMeta.Description, ParameterMeta.Group, ParameterMeta.SortPriority, AssetIndex);
+	const FMaterialCachedParameterEditorInfo EditorInfo(ParameterMeta.ExpressionGuid, ParameterMeta.Description, ParameterMeta.Group, ParameterMeta.SortPriority);
 	const int32 Index = TryAddParameter(*this, ParameterMeta.Value.Type, ParameterInfo, EditorInfo);
 	if (Index != INDEX_NONE)
 	{
@@ -158,24 +156,24 @@ void FMaterialCachedExpressionData::AddParameter(const FMaterialParameterInfo& P
 		{
 		case EMaterialParameterType::Scalar:
 			ScalarValues.Insert(ParameterMeta.Value.AsScalar(), Index);
-			EditorOnlyData->ScalarMinMaxValues.Insert(FVector2D(ParameterMeta.ScalarMin, ParameterMeta.ScalarMax), Index);
+			ScalarMinMaxValues.Insert(FVector2D(ParameterMeta.ScalarMin, ParameterMeta.ScalarMax), Index);
 			ScalarPrimitiveDataIndexValues.Insert(ParameterMeta.PrimitiveDataIndex, Index);
 			if (ParameterMeta.bUsedAsAtlasPosition)
 			{
-				EditorOnlyData->ScalarCurveValues.Insert(ParameterMeta.ScalarCurve.Get(), Index);
-				EditorOnlyData->ScalarCurveAtlasValues.Insert(ParameterMeta.ScalarAtlas.Get(), Index);
+				ScalarCurveValues.Insert(ParameterMeta.ScalarCurve.Get(), Index);
+				ScalarCurveAtlasValues.Insert(ParameterMeta.ScalarAtlas.Get(), Index);
 				OutReferencedTexture = ParameterMeta.ScalarAtlas.Get();
 			}
 			else
 			{
-				EditorOnlyData->ScalarCurveValues.Insert(nullptr, Index);
-				EditorOnlyData->ScalarCurveAtlasValues.Insert(nullptr, Index);
+				ScalarCurveValues.Insert(nullptr, Index);
+				ScalarCurveAtlasValues.Insert(nullptr, Index);
 			}
 			break;
 		case EMaterialParameterType::Vector:
 			VectorValues.Insert(ParameterMeta.Value.AsLinearColor(), Index);
-			EditorOnlyData->VectorChannelNameValues.Insert(ParameterMeta.ChannelNames, Index);
-			EditorOnlyData->VectorUsedAsChannelMaskValues.Insert(ParameterMeta.bUsedAsChannelMask, Index);
+			VectorChannelNameValues.Insert(ParameterMeta.ChannelNames, Index);
+			VectorUsedAsChannelMaskValues.Insert(ParameterMeta.bUsedAsChannelMask, Index);
 			VectorPrimitiveDataIndexValues.Insert(ParameterMeta.PrimitiveDataIndex, Index);
 			break;
 		case EMaterialParameterType::DoubleVector:
@@ -183,7 +181,7 @@ void FMaterialCachedExpressionData::AddParameter(const FMaterialParameterInfo& P
 			break;
 		case EMaterialParameterType::Texture:
 			TextureValues.Insert(ParameterMeta.Value.Texture, Index);
-			EditorOnlyData->TextureChannelNameValues.Insert(ParameterMeta.ChannelNames, Index);
+			TextureChannelNameValues.Insert(ParameterMeta.ChannelNames, Index);
 			OutReferencedTexture = ParameterMeta.Value.Texture;
 			break;
 		case EMaterialParameterType::Font:
@@ -199,10 +197,10 @@ void FMaterialCachedExpressionData::AddParameter(const FMaterialParameterInfo& P
 			OutReferencedTexture = ParameterMeta.Value.RuntimeVirtualTexture;
 			break;
 		case EMaterialParameterType::StaticSwitch:
-			EditorOnlyData->StaticSwitchValues.Insert(ParameterMeta.Value.AsStaticSwitch(), Index);
+			StaticSwitchValues.Insert(ParameterMeta.Value.AsStaticSwitch(), Index);
 			break;
 		case EMaterialParameterType::StaticComponentMask:
-			EditorOnlyData->StaticComponentMaskValues.Insert(ParameterMeta.Value.AsStaticComponentMask(), Index);
+			StaticComponentMaskValues.Insert(ParameterMeta.Value.AsStaticComponentMask(), Index);
 			break;
 		default:
 			checkNoEntry();
@@ -227,7 +225,11 @@ void FMaterialCachedExpressionData::UpdateForFunction(const FMaterialCachedExpre
 	FMaterialCachedExpressionData* Self = this;
 	auto ProcessFunction = [Self, &LocalContext, Association, ParameterIndex](UMaterialFunctionInterface* InFunction) -> bool
 	{
-		Self->UpdateForExpressions(LocalContext, InFunction->GetExpressions(), Association, ParameterIndex);
+		const TArray<TObjectPtr<UMaterialExpression>>* FunctionExpressions = InFunction->GetFunctionExpressions();
+		if (FunctionExpressions)
+		{
+			Self->UpdateForExpressions(LocalContext, *FunctionExpressions, Association, ParameterIndex);
+		}
 
 		FMaterialFunctionInfo NewFunctionInfo;
 		NewFunctionInfo.Function = InFunction;
@@ -254,10 +256,8 @@ void FMaterialCachedExpressionData::UpdateForLayerFunctions(const FMaterialCache
 	}
 }
 
-void FMaterialCachedExpressionData::UpdateForExpressions(const FMaterialCachedExpressionContext& Context, TConstArrayView<TObjectPtr<UMaterialExpression>> Expressions, EMaterialParameterAssociation Association, int32 ParameterIndex)
+void FMaterialCachedExpressionData::UpdateForExpressions(const FMaterialCachedExpressionContext& Context, const TArray<TObjectPtr<UMaterialExpression>>& Expressions, EMaterialParameterAssociation Association, int32 ParameterIndex)
 {
-	check(EditorOnlyData);
-
 	for (UMaterialExpression* Expression : Expressions)
 	{
 		if (!Expression)
@@ -287,7 +287,7 @@ void FMaterialCachedExpressionData::UpdateForExpressions(const FMaterialCachedEx
 			}
 
 			const FMaterialParameterInfo ParameterInfo(ParameterName, Association, ParameterIndex);
-			AddParameter(ParameterInfo, ParameterMeta, ReferencedTexture);
+			Parameters.AddParameter(ParameterInfo, ParameterMeta, ReferencedTexture);
 		}
 
 		// We first try to extract the referenced texture from the parameter value, that way we'll also get the proper texture in case value is overriden by a function instance
@@ -303,7 +303,7 @@ void FMaterialCachedExpressionData::UpdateForExpressions(const FMaterialCachedEx
 			ReferencedTextures.AddUnique(ReferencedTexture);
 		}
 
-		Expression->GetLandscapeLayerNames(EditorOnlyData->LandscapeLayerNames);
+		Expression->GetLandscapeLayerNames(LandscapeLayerNames);
 
 		if (UMaterialExpressionCollectionParameter* ExpressionCollectionParameter = Cast<UMaterialExpressionCollectionParameter>(Expression))
 		{
@@ -381,9 +381,7 @@ void FMaterialCachedExpressionData::UpdateForExpressions(const FMaterialCachedEx
 
 				// TODO(?) - Layers for MIs are currently duplicated here and in FStaticParameterSet
 				bHasMaterialLayers = true;
-				MaterialLayers = Layers.GetRuntime();
-				EditorOnlyData->MaterialLayers = Layers.EditorOnly;
-				FMaterialLayersFunctions::Validate(MaterialLayers, EditorOnlyData->MaterialLayers);
+				MaterialLayers = Layers;
 				LayersExpression->RebuildLayerGraph(false);
 			}
 		}
@@ -400,25 +398,30 @@ void FMaterialCachedExpressionData::UpdateForExpressions(const FMaterialCachedEx
 		}
 		else if (UMaterialExpressionSetMaterialAttributes* SetMatAttributes = Cast<UMaterialExpressionSetMaterialAttributes>(Expression))
 		{
-			for (int32 PinIndex = 0; PinIndex < SetMatAttributes->AttributeSetTypes.Num(); ++PinIndex)
+			for (int PinIndex = 0; PinIndex < SetMatAttributes->AttributeSetTypes.Num(); ++PinIndex)
 			{
 				// For this material attribute pin do we have something connected?
 				const FGuid& Guid = SetMatAttributes->AttributeSetTypes[PinIndex];
 				const FExpressionInput& AttributeInput = SetMatAttributes->Inputs[PinIndex + 1];
 				const EMaterialProperty MaterialProperty = FMaterialAttributeDefinitionMap::GetProperty(Guid);
-				if (AttributeInput.Expression)
+
+				// Only set the material property if it hasn't been set yet.  We want to specifically avoid a Set Material Attributes node which doesn't have a 
+				// attribute set from disabling the attribute from a different Set Material Attributes node which does have it enabled.
+				if (!IsMaterialAttributePropertyConnected(MaterialProperty))
 				{
-					SetPropertyConnected(MaterialProperty);
+					SetMaterialAttributePropertyConnected(MaterialProperty, AttributeInput.Expression ? true : false);
 				}
 			}
-		}
-		else if (UMaterialExpressionMakeMaterialAttributes* MakeMatAttributes = Cast<UMaterialExpressionMakeMaterialAttributes>(Expression))
-		{
+			}
+			else if (UMaterialExpressionMakeMaterialAttributes* MakeMatAttributes = Cast<UMaterialExpressionMakeMaterialAttributes>(Expression))
+			{
+			// Only set the material property if it hasn't been set yet.  We want to specifically avoid a Set Material Attributes node which doesn't have a 
+			// attribute set from disabling the attribute from a different Set Material Attributes node which does have it enabled.
 			auto SetMatAttributeConditionally = [&](EMaterialProperty InMaterialProperty, bool InIsConnected)
 			{
-				if (InIsConnected)
+				if (!IsMaterialAttributePropertyConnected(InMaterialProperty))
 				{
-					SetPropertyConnected(InMaterialProperty);
+					SetMaterialAttributePropertyConnected(InMaterialProperty, InIsConnected);
 				}
 			};
 
@@ -484,7 +487,7 @@ void PrepareHLSLTree(UE::HLSLTree::FEmitContext& EmitContext,
 				const EMaterialProperty Property = FMaterialAttributeDefinitionMap::GetProperty(AttributeID);
 				if (CachedTree.IsAttributeUsed(EmitContext, *EmitResultScope, ResultType, Property))
 				{
-					CachedData.SetPropertyConnected(Property);
+					CachedData.SetMaterialAttributePropertyConnected(Property, true);
 				}
 			}
 		}
@@ -511,60 +514,88 @@ void FMaterialCachedExpressionData::UpdateForCachedHLSLTree(const FMaterialCache
 	::Private::PrepareHLSLTree(EmitContext, CachedTree, *this, SF_Vertex);
 }
 
-void FMaterialCachedExpressionData::Validate()
+#endif // WITH_EDITOR
+
+void FMaterialCachedParameterEntry::Reset()
+{
+	ParameterInfoSet.Reset();
+#if WITH_EDITORONLY_DATA
+	EditorInfo.Reset();
+#endif
+}
+
+void FMaterialCachedParameters::Reset()
 {
-	if (EditorOnlyData)
+	for (int32 i = 0; i < NumMaterialRuntimeParameterTypes; ++i)
 	{
-		for (int32 TypeIndex = 0; TypeIndex < NumMaterialParameterTypes; ++TypeIndex)
-		{
-			const FMaterialCachedParameterEditorEntry& EditorEntry = EditorOnlyData->EditorEntries[TypeIndex];
-			const FMaterialCachedParameterEntry& Entry = GetParameterTypeEntry((EMaterialParameterType)TypeIndex);
-			check(EditorEntry.EditorInfo.Num() == Entry.ParameterInfoSet.Num());
-		}
-		FMaterialLayersFunctions::Validate(MaterialLayers, EditorOnlyData->MaterialLayers);
+		RuntimeEntries[i].Reset();
 	}
-}
+#if WITH_EDITORONLY_DATA
+	for (int32 i = 0; i < NumMaterialEditorOnlyParameterTypes; ++i)
+	{
+		EditorOnlyEntries[i].Reset();
+	}
+#endif
+
+	ScalarPrimitiveDataIndexValues.Reset();
+	VectorPrimitiveDataIndexValues.Reset();
+	ScalarValues.Reset();
+	VectorValues.Reset();
+	DoubleVectorValues.Reset();
+	TextureValues.Reset();
+	FontValues.Reset();
+	FontPageValues.Reset();
+	RuntimeVirtualTextureValues.Reset();
 
-#endif // WITH_EDITOR
+#if WITH_EDITORONLY_DATA
+	StaticSwitchValues.Reset();
+	StaticComponentMaskValues.Reset();
+	ScalarMinMaxValues.Reset();
+	ScalarCurveValues.Reset();
+	ScalarCurveAtlasValues.Reset();
+	VectorChannelNameValues.Reset();
+	VectorUsedAsChannelMaskValues.Reset();
+	TextureChannelNameValues.Reset();
+#endif // WITH_EDITORONLY_DATA
+}
 
-int32 FMaterialCachedExpressionData::FindParameterIndex(EMaterialParameterType Type, const FMemoryImageMaterialParameterInfo& ParameterInfo) const
+int32 FMaterialCachedParameters::FindParameterIndex(EMaterialParameterType Type, const FMemoryImageMaterialParameterInfo& ParameterInfo) const
 {
 	const FMaterialCachedParameterEntry& Entry = GetParameterTypeEntry(Type);
 	const FSetElementId ElementId = Entry.ParameterInfoSet.FindId(FMaterialParameterInfo(ParameterInfo));
 	return ElementId.AsInteger();
 }
 
-void FMaterialCachedExpressionData::GetParameterValueByIndex(EMaterialParameterType Type, int32 ParameterIndex, FMaterialParameterMetadata& OutResult) const
+void FMaterialCachedParameters::GetParameterValueByIndex(EMaterialParameterType Type, int32 ParameterIndex, FMaterialParameterMetadata& OutResult) const
 {
+	const FMaterialCachedParameterEntry& Entry = GetParameterTypeEntry(Type);
+	
 #if WITH_EDITORONLY_DATA
-	if (EditorOnlyData)
+	const bool bIsEditorOnlyDataStripped = Entry.EditorInfo.Num() == 0;
+	if (!bIsEditorOnlyDataStripped)
 	{
-		const FMaterialCachedParameterEditorEntry& EditorEntry = EditorOnlyData->EditorEntries[(int32)Type];
-		const FMaterialCachedParameterEditorInfo& EditorInfo = EditorEntry.EditorInfo[ParameterIndex];
+		const FMaterialCachedParameterEditorInfo& EditorInfo = Entry.EditorInfo[ParameterIndex];
 		OutResult.ExpressionGuid = EditorInfo.ExpressionGuid;
 		OutResult.Description = EditorInfo.Description;
 		OutResult.Group = EditorInfo.Group;
 		OutResult.SortPriority = EditorInfo.SortPriority;
-		if (EditorInfo.AssetIndex != INDEX_NONE)
-		{
-			OutResult.AssetPath = EditorOnlyData->AssetPaths[EditorInfo.AssetIndex];
-		}
 	}
-#endif // WITH_EDITORONLY_DATA
+#endif
 
 	switch (Type)
 	{
 	case EMaterialParameterType::Scalar:
 		OutResult.Value = ScalarValues[ParameterIndex];
 		OutResult.PrimitiveDataIndex = ScalarPrimitiveDataIndexValues[ParameterIndex];
+
 #if WITH_EDITORONLY_DATA
-		if (EditorOnlyData)
+		if (!bIsEditorOnlyDataStripped)
 		{
-			OutResult.ScalarMin = EditorOnlyData->ScalarMinMaxValues[ParameterIndex].X;
-			OutResult.ScalarMax = EditorOnlyData->ScalarMinMaxValues[ParameterIndex].Y;
+			OutResult.ScalarMin = ScalarMinMaxValues[ParameterIndex].X;
+			OutResult.ScalarMax = ScalarMinMaxValues[ParameterIndex].Y;
 			{
-				const TSoftObjectPtr<UCurveLinearColor>& Curve = EditorOnlyData->ScalarCurveValues[ParameterIndex];
-				const TSoftObjectPtr<UCurveLinearColorAtlas>& Atlas = EditorOnlyData->ScalarCurveAtlasValues[ParameterIndex];
+				const TSoftObjectPtr<UCurveLinearColor>& Curve = ScalarCurveValues[ParameterIndex];
+				const TSoftObjectPtr<UCurveLinearColorAtlas>& Atlas = ScalarCurveAtlasValues[ParameterIndex];
 				if (!Curve.IsNull() && !Atlas.IsNull())
 				{
 					OutResult.ScalarCurve = Curve;
@@ -578,11 +609,12 @@ void FMaterialCachedExpressionData::GetParameterValueByIndex(EMaterialParameterT
 	case EMaterialParameterType::Vector:
 		OutResult.Value = VectorValues[ParameterIndex];
 		OutResult.PrimitiveDataIndex = VectorPrimitiveDataIndexValues[ParameterIndex];
+
 #if  WITH_EDITORONLY_DATA
-		if (EditorOnlyData)
+		if (!bIsEditorOnlyDataStripped)
 		{
-			OutResult.ChannelNames = EditorOnlyData->VectorChannelNameValues[ParameterIndex];
-			OutResult.bUsedAsChannelMask = EditorOnlyData->VectorUsedAsChannelMaskValues[ParameterIndex];
+			OutResult.ChannelNames = VectorChannelNameValues[ParameterIndex];
+			OutResult.bUsedAsChannelMask = VectorUsedAsChannelMaskValues[ParameterIndex];
 		}
 #endif // WITH_EDITORONLY_DATA
 		break;
@@ -592,9 +624,9 @@ void FMaterialCachedExpressionData::GetParameterValueByIndex(EMaterialParameterT
 	case EMaterialParameterType::Texture:
 		OutResult.Value = TextureValues[ParameterIndex].LoadSynchronous();
 #if WITH_EDITORONLY_DATA
-		if (EditorOnlyData)
+		if (!bIsEditorOnlyDataStripped)
 		{
-			OutResult.ChannelNames = EditorOnlyData->TextureChannelNameValues[ParameterIndex];
+			OutResult.ChannelNames = TextureChannelNameValues[ParameterIndex];
 		}
 #endif // WITH_EDITORONLY_DATA
 		break;
@@ -606,15 +638,15 @@ void FMaterialCachedExpressionData::GetParameterValueByIndex(EMaterialParameterT
 		break;
 #if WITH_EDITORONLY_DATA
 	case EMaterialParameterType::StaticSwitch:
-		if (EditorOnlyData)
+		if (!bIsEditorOnlyDataStripped)
 		{
-			OutResult.Value = EditorOnlyData->StaticSwitchValues[ParameterIndex];
+			OutResult.Value = StaticSwitchValues[ParameterIndex];
 		}
 		break;
 	case EMaterialParameterType::StaticComponentMask:
-		if (EditorOnlyData)
+		if (!bIsEditorOnlyDataStripped)
 		{
-			OutResult.Value = EditorOnlyData->StaticComponentMaskValues[ParameterIndex];
+			OutResult.Value = StaticComponentMaskValues[ParameterIndex];
 		}
 		break;
 #endif // WITH_EDITORONLY_DATA
@@ -624,7 +656,7 @@ void FMaterialCachedExpressionData::GetParameterValueByIndex(EMaterialParameterT
 	}
 }
 
-bool FMaterialCachedExpressionData::GetParameterValue(EMaterialParameterType Type, const FMemoryImageMaterialParameterInfo& ParameterInfo, FMaterialParameterMetadata& OutResult) const
+bool FMaterialCachedParameters::GetParameterValue(EMaterialParameterType Type, const FMemoryImageMaterialParameterInfo& ParameterInfo, FMaterialParameterMetadata& OutResult) const
 {
 	const int32 Index = FindParameterIndex(Type, ParameterInfo);
 	if (Index != INDEX_NONE)
@@ -636,19 +668,15 @@ bool FMaterialCachedExpressionData::GetParameterValue(EMaterialParameterType Typ
 	return false;
 }
 
-const FGuid& FMaterialCachedExpressionData::GetExpressionGuid(EMaterialParameterType Type, int32 Index) const
-{
 #if WITH_EDITORONLY_DATA
-	if (EditorOnlyData)
-	{
-		return EditorOnlyData->EditorEntries[(int32)Type].EditorInfo[Index].ExpressionGuid;
-	}
-#endif // WITH_EDITORONLY_DATA
-	static const FGuid EmptyGuid;
-	return EmptyGuid;
+const FGuid& FMaterialCachedParameters::GetExpressionGuid(EMaterialParameterType Type, int32 Index) const
+{
+	const FMaterialCachedParameterEntry& Entry = GetParameterTypeEntry(Type);
+	return Entry.EditorInfo[Index].ExpressionGuid;
 }
+#endif // WITH_EDITORONLY_DATA
 
-void FMaterialCachedExpressionData::GetAllParametersOfType(EMaterialParameterType Type, TMap<FMaterialParameterInfo, FMaterialParameterMetadata>& OutParameters) const
+void FMaterialCachedParameters::GetAllParametersOfType(EMaterialParameterType Type, TMap<FMaterialParameterInfo, FMaterialParameterMetadata>& OutParameters) const
 {
 	const FMaterialCachedParameterEntry& Entry = GetParameterTypeEntry(Type);
 	const int32 NumParameters = Entry.ParameterInfoSet.Num();
@@ -662,7 +690,7 @@ void FMaterialCachedExpressionData::GetAllParametersOfType(EMaterialParameterTyp
 	}
 }
 
-void FMaterialCachedExpressionData::GetAllParameterInfoOfType(EMaterialParameterType Type, TArray<FMaterialParameterInfo>& OutParameterInfo, TArray<FGuid>& OutParameterIds) const
+void FMaterialCachedParameters::GetAllParameterInfoOfType(EMaterialParameterType Type, TArray<FMaterialParameterInfo>& OutParameterInfo, TArray<FGuid>& OutParameterIds) const
 {
 	const FMaterialCachedParameterEntry& Entry = GetParameterTypeEntry(Type);
 	const int32 NumParameters = Entry.ParameterInfoSet.Num();
@@ -671,13 +699,22 @@ void FMaterialCachedExpressionData::GetAllParameterInfoOfType(EMaterialParameter
 
 	for (TSet<FMaterialParameterInfo>::TConstIterator It(Entry.ParameterInfoSet); It; ++It)
 	{
-		const int32 ParameterIndex = It.GetId().AsInteger();
 		OutParameterInfo.Add(*It);
-		OutParameterIds.Add(GetExpressionGuid(Type, ParameterIndex));
+#if WITH_EDITORONLY_DATA
+		// cooked materials can strip out expression guids
+		if (Entry.EditorInfo.Num() != 0)
+		{
+			OutParameterIds.Add(Entry.EditorInfo[It.GetId().AsInteger()].ExpressionGuid);
+		}
+		else
+#endif
+		{
+			OutParameterIds.Add(FGuid());
+		}
 	}
 }
 
-void FMaterialCachedExpressionData::GetAllGlobalParametersOfType(EMaterialParameterType Type, TMap<FMaterialParameterInfo, FMaterialParameterMetadata>& OutParameters) const
+void FMaterialCachedParameters::GetAllGlobalParametersOfType(EMaterialParameterType Type, TMap<FMaterialParameterInfo, FMaterialParameterMetadata>& OutParameters) const
 {
 	const FMaterialCachedParameterEntry& Entry = GetParameterTypeEntry(Type);
 	const int32 NumParameters = Entry.ParameterInfoSet.Num();
@@ -697,7 +734,7 @@ void FMaterialCachedExpressionData::GetAllGlobalParametersOfType(EMaterialParame
 	}
 }
 
-void FMaterialCachedExpressionData::GetAllGlobalParameterInfoOfType(EMaterialParameterType Type, TArray<FMaterialParameterInfo>& OutParameterInfo, TArray<FGuid>& OutParameterIds) const
+void FMaterialCachedParameters::GetAllGlobalParameterInfoOfType(EMaterialParameterType Type, TArray<FMaterialParameterInfo>& OutParameterInfo, TArray<FGuid>& OutParameterIds) const
 {
 	const FMaterialCachedParameterEntry& Entry = GetParameterTypeEntry(Type);
 	const int32 NumParameters = Entry.ParameterInfoSet.Num();
@@ -709,9 +746,18 @@ void FMaterialCachedExpressionData::GetAllGlobalParameterInfoOfType(EMaterialPar
 		const FMaterialParameterInfo& ParameterInfo = *It;
 		if (ParameterInfo.Association == GlobalParameter)
 		{
-			const int32 ParameterIndex = It.GetId().AsInteger();
-			OutParameterInfo.Add(*It);
-			OutParameterIds.Add(GetExpressionGuid(Type, ParameterIndex));
+			OutParameterInfo.Add(ParameterInfo);
+#if WITH_EDITORONLY_DATA
+			// cooked materials can strip out expression guids
+			if (Entry.EditorInfo.Num() != 0)
+			{
+				OutParameterIds.Add(Entry.EditorInfo[It.GetId().AsInteger()].ExpressionGuid);
+			}
+			else
+#endif
+			{
+				OutParameterIds.Add(FGuid());
+			}
 		}
 	}
 }
diff --git a/Engine/Source/Runtime/Engine/Private/Materials/MaterialCachedHLSLTree.cpp b/Engine/Source/Runtime/Engine/Private/Materials/MaterialCachedHLSLTree.cpp
--- a/Engine/Source/Runtime/Engine/Private/Materials/MaterialCachedHLSLTree.cpp
+++ b/Engine/Source/Runtime/Engine/Private/Materials/MaterialCachedHLSLTree.cpp
@@ -54,7 +54,7 @@ bool FMaterialCachedHLSLTree::GenerateTree(UMaterial* Material, const FMaterialL
 {
 	const EMaterialShadingModel DefaultShadingModel = Material->GetShadingModels().GetFirstShadingModel();
 
-	for (UMaterialExpression* Expression : Material->GetExpressions())
+	for (UMaterialExpression* Expression : Material->Expressions)
 	{
 		UMaterialExpressionCustomOutput* CustomOutput = Cast<UMaterialExpressionCustomOutput>(Expression);
 		// We don't want anything with HasCustomSourceOutput() here (VertexInterpolators)
diff --git a/Engine/Source/Runtime/Engine/Private/Materials/MaterialExpressionHLSL.cpp b/Engine/Source/Runtime/Engine/Private/Materials/MaterialExpressionHLSL.cpp
--- a/Engine/Source/Runtime/Engine/Private/Materials/MaterialExpressionHLSL.cpp
+++ b/Engine/Source/Runtime/Engine/Private/Materials/MaterialExpressionHLSL.cpp
@@ -1627,7 +1627,7 @@ bool UMaterialExpressionMaterialAttributeLayers::GenerateHLSLExpression(FMateria
 	{
 		UMaterialFunctionInterface* LayerFunction = MaterialLayers.Layers[LayerIndex];
 		const FExpression* LayerExpression = nullptr;
-		if (LayerFunction && MaterialLayers.EditorOnly.LayerStates[LayerIndex])
+		if (LayerFunction && MaterialLayers.LayerStates[LayerIndex])
 		{
 			const EMaterialFunctionUsage Usage = LayerFunction->GetMaterialFunctionUsage();
 			if (Usage != EMaterialFunctionUsage::MaterialLayer)
@@ -1670,7 +1670,7 @@ bool UMaterialExpressionMaterialAttributeLayers::GenerateHLSLExpression(FMateria
 			return Generator.Errorf(TEXT("Invalid number of layers (%d) and blends (%d)"), MaterialLayers.Layers.Num(), MaterialLayers.Blends.Num());
 		}
 
-		if (MaterialLayers.Layers[LayerIndex] && MaterialLayers.EditorOnly.LayerStates[LayerIndex])
+		if (MaterialLayers.Layers[LayerIndex] && MaterialLayers.LayerStates[LayerIndex])
 		{
 			const FExpression* LayerExpression = LayerExpressions[LayerIndex];
 			if (!LayerExpression)
@@ -1709,7 +1709,7 @@ bool UMaterialExpressionMaterialAttributeLayers::GenerateHLSLExpression(FMateria
 		}
 	}
 
-	OutExpression = Generator.GetTree().NewExpression<Material::FExpressionMaterialLayers>(BottomLayerExpression, MaterialLayers);
+	OutExpression = Generator.GetTree().NewExpression<Material::FExpressionMaterialLayers>(BottomLayerExpression, &MaterialLayers);
 	return true;
 }
 
