[Backout] - CL19964485

diff --git a/Engine/Source/Runtime/Engine/Private/Materials/MaterialExpressions.cpp b/Engine/Source/Runtime/Engine/Private/Materials/MaterialExpressions.cpp
--- a/Engine/Source/Runtime/Engine/Private/Materials/MaterialExpressions.cpp
+++ b/Engine/Source/Runtime/Engine/Private/Materials/MaterialExpressions.cpp
@@ -504,7 +504,7 @@ void ValidateParameterNameInternal(class UMaterialExpression* ExpressionToValida
 
 				bFoundValidName = true;
 
-				for (const UMaterialExpression* Expression : OwningMaterial->GetExpressions())
+				for (const UMaterialExpression* Expression : OwningMaterial->Expressions)
 				{
 					if (Expression != nullptr && Expression->HasAParameterName())
 					{
@@ -528,7 +528,7 @@ void ValidateParameterNameInternal(class UMaterialExpression* ExpressionToValida
 		if (bAllowDuplicateName)
 		{
 			// Check for any matching values
-			for (UMaterialExpression* Expression : OwningMaterial->GetExpressions())
+			for (UMaterialExpression* Expression : OwningMaterial->Expressions)
 			{
 				if (Expression != nullptr && Expression->HasAParameterName())
 				{
@@ -657,34 +657,6 @@ int32 CompileStrataBlendFunction(FMaterialCompiler* Compiler, const int32 A, con
 	return INDEX_NONE;
 }
 
-void FMaterialExpressionCollection::AddExpression(UMaterialExpression* InExpression)
-{
-	Expressions.AddUnique(InExpression);
-}
-
-void FMaterialExpressionCollection::RemoveExpression(UMaterialExpression* InExpression)
-{
-	Expressions.Remove(InExpression);
-}
-
-void FMaterialExpressionCollection::AddComment(UMaterialExpressionComment* InExpression)
-{
-	EditorComments.AddUnique(InExpression);
-}
-
-void FMaterialExpressionCollection::RemoveComment(UMaterialExpressionComment* InExpression)
-{
-	EditorComments.Remove(InExpression);
-}
-
-void FMaterialExpressionCollection::Empty()
-{
-	Expressions.Empty();
-	EditorComments.Empty();
-	ExpressionExecBegin = nullptr;
-	ExpressionExecEnd = nullptr;
-}
-
 #if WITH_EDITOR
 void UMaterialExpression::InitializeNumExecutionInputs(TArrayView<UMaterialExpression*> Expressions)
 {
@@ -781,7 +753,7 @@ void UMaterialExpression::CopyMaterialExpressions(const TArray<UMaterialExpressi
 			SrcToDestMap.Add( SrcExpression, NewExpression );
 
 			// Add to list of material expressions associated with the copy buffer.
-			Material->GetExpressionCollection().AddExpression( NewExpression );
+			Material->Expressions.Add( NewExpression );
 
 			// There can be only one default mesh paint texture.
 			UMaterialExpressionTextureBase* TextureSample = Cast<UMaterialExpressionTextureBase>( NewExpression );
@@ -846,7 +818,7 @@ void UMaterialExpression::CopyMaterialExpressions(const TArray<UMaterialExpressi
 		NewComment->Material = Material;
 
 		// Add reference to the material
-		Material->GetExpressionCollection().AddComment(NewComment);
+		Material->EditorComments.Add(NewComment);
 
 		// Add to the output array.
 		OutNewComments.Add(NewComment);
@@ -1646,10 +1618,12 @@ void UMaterialExpressionTextureBase::PostEditChangeProperty(FPropertyChangedEven
 		const FName PropertyName = PropertyChangedEvent.Property->GetFName();
 		if (PropertyName == FName(TEXT("IsDefaultMeshpaintTexture")))
 		{
+			const TArray<UMaterialExpression*>& Expressions = this->Material->GetMaterial()->Expressions;
+
 			// Check for other defaulted textures in THIS material (does not search sub levels ie functions etc, as these are ignored in the texture painter). 
-			for (UMaterialExpression* Expression : this->Material->GetMaterial()->GetExpressions())
+			for (auto ItExpressions = Expressions.CreateConstIterator(); ItExpressions; ItExpressions++)
 			{
-				UMaterialExpressionTextureBase* TextureSample = Cast<UMaterialExpressionTextureBase>(Expression);
+				UMaterialExpressionTextureBase* TextureSample = Cast<UMaterialExpressionTextureBase>(*ItExpressions);
 				if (TextureSample != nullptr && TextureSample != this)
 				{
 					if(TextureSample->IsDefaultMeshpaintTexture)
@@ -5990,15 +5964,15 @@ UMaterialExpressionMakeMaterialAttributes::UMaterialExpressionMakeMaterialAttrib
 void UMaterialExpressionMakeMaterialAttributes::Serialize(FStructuredArchive::FRecord Record)
 {
 	Super::Serialize(Record);
-#if WITH_EDITOR
 	FArchive& UnderlyingArchive = Record.GetUnderlyingArchive();
+
 	UnderlyingArchive.UsingCustomVersion(FRenderingObjectVersion::GUID);
+	
 	if (UnderlyingArchive.CustomVer(FRenderingObjectVersion::GUID) < FRenderingObjectVersion::FixedLegacyMaterialAttributeNodeTypes)
 	{
 		// Update the legacy masks else fail on vec3 to vec2 conversion
 		Refraction.SetMask(1, 1, 1, 0, 0);
 	}
-#endif // WITH_EDITOR
 }
 
 #if WITH_EDITOR
@@ -6968,9 +6942,8 @@ UMaterialExpressionMaterialAttributeLayers::UMaterialExpressionMaterialAttribute
 	};
 	static FConstructorStatics ConstructorStatics;
 
-#if WITH_EDITOR
 	DefaultLayers.AddDefaultBackgroundLayer();
-#endif
+
 #if WITH_EDITORONLY_DATA
 	MenuCategories.Add(ConstructorStatics.NAME_MaterialAttributes);
 	MenuCategories.Add(ConstructorStatics.NAME_Parameters);
@@ -7977,14 +7950,14 @@ void UMaterialExpressionVectorParameter::ApplyChannelNames()
 void UMaterialExpressionVectorParameter::ValidateParameterName(const bool bAllowDuplicateName)
 {
 	bool bOverrideDuplicateBehavior = false;
-	TConstArrayView<TObjectPtr<UMaterialExpression>> Expressions;
+	TArray<UMaterialExpression*> Expressions;
 	if (Material)
 	{
-		Expressions = Material->GetExpressions();
+		Expressions = Material->Expressions;
 	}
 	else if (Function)
 	{
-		Expressions = Function->GetExpressions();
+		Expressions = Function->FunctionExpressions;
 	}
 
 	for (UMaterialExpression* Expression : Expressions)
@@ -8343,14 +8316,14 @@ void UMaterialExpressionScalarParameter::PostEditChangeProperty(FPropertyChanged
 void UMaterialExpressionScalarParameter::ValidateParameterName(const bool bAllowDuplicateName)
 {
 	bool bOverrideDuplicateBehavior = false;
-	TConstArrayView<TObjectPtr<UMaterialExpression>> Expressions;
+	TArray<UMaterialExpression*> Expressions;
 	if (Material)
 	{
-		Expressions = Material->GetExpressions();
+		Expressions = Material->Expressions;
 	}
 	else if (Function)
 	{
-		Expressions = Function->GetExpressions();
+		Expressions = Function->FunctionExpressions;
 	}
 
 	for (UMaterialExpression* Expression : Expressions)
@@ -9084,6 +9057,10 @@ void UMaterialExpressionFeatureLevelSwitch::Serialize(FStructuredArchive::FRecor
 	FArchive& UnderlyingArchive = Record.GetUnderlyingArchive();
 	UnderlyingArchive.UsingCustomVersion(FRenderingObjectVersion::GUID);
 
+#if !WITH_EDITORONLY_DATA
+	FExpressionInput Inputs[ERHIFeatureLevel::Num];
+#endif
+
 	if (UnderlyingArchive.IsLoading() && UnderlyingArchive.UEVer() < VER_UE4_RENAME_SM3_TO_ES3_1)
 	{
 		// Copy the ES2 input to SM3 (since SM3 will now become ES3_1 and we don't want broken content)
@@ -9572,9 +9549,9 @@ void UMaterialExpressionDynamicParameter::PostLoad()
 void UMaterialExpressionDynamicParameter::UpdateDynamicParameterProperties()
 {
 	check(Material);
-	for (UMaterialExpression* Expression : Material->GetExpressions())
+	for (int32 ExpIndex = 0; ExpIndex < Material->Expressions.Num(); ExpIndex++)
 	{
-		const UMaterialExpressionDynamicParameter* DynParam = Cast<UMaterialExpressionDynamicParameter>(Expression);
+		const UMaterialExpressionDynamicParameter* DynParam = Cast<UMaterialExpressionDynamicParameter>(Material->Expressions[ExpIndex]);
 		if (CopyDynamicParameterProperties(DynParam))
 		{
 			break;
@@ -11172,12 +11149,12 @@ void UMaterialExpressionPinBase::DeleteReroutePins()
 		if (Reroute.Expression)
 		{
 			Reroute.Expression->Modify();
-			Material->GetExpressionCollection().RemoveExpression(Reroute.Expression);
+			Material->Expressions.Remove(Reroute.Expression);
 			Reroute.Expression->MarkAsGarbage();
 		}
 		else
 		{
-			Material->GetExpressionCollection().RemoveExpression(Reroute.Expression);
+			Material->Expressions.Remove(Reroute.Expression);
 		}
 	}
 	ReroutePins.Empty();
@@ -11216,7 +11193,7 @@ void UMaterialExpressionPinBase::PostEditChangeProperty(FPropertyChangedEvent& P
 			AddedReroute.Expression->Material = Material;
 			AddedReroute.Name = AddedReroute.Name.IsNone() ? FName(FString::Printf(TEXT("Pin %u"), AddedRerouteIndex + 1)) : AddedReroute.Name;
 
-			Material->GetExpressionCollection().AddExpression(AddedReroute.Expression);
+			Material->Expressions.Add(AddedReroute.Expression);
 		}
 		else if (PropertyChangedEvent.ChangeType == EPropertyChangeType::ArrayRemove)
 		{
@@ -11226,12 +11203,12 @@ void UMaterialExpressionPinBase::PostEditChangeProperty(FPropertyChangedEvent& P
 			if (RemovedReroute)
 			{
 				RemovedReroute->Modify();
-				Material->GetExpressionCollection().RemoveExpression(RemovedReroute);
+				Material->Expressions.Remove(RemovedReroute);
 				RemovedReroute->MarkAsGarbage();
 			}
 			else
 			{
-				Material->GetExpressionCollection().RemoveExpression(RemovedReroute);
+				Material->Expressions.Remove(RemovedReroute);
 			}
 		}
 		else if (PropertyChangedEvent.ChangeType == EPropertyChangeType::ArrayClear)
@@ -11241,12 +11218,12 @@ void UMaterialExpressionPinBase::PostEditChangeProperty(FPropertyChangedEvent& P
 				if (RemovedReroute)
 				{
 					RemovedReroute->Modify();
-					Material->GetExpressionCollection().RemoveExpression(RemovedReroute);
+					Material->Expressions.Remove(RemovedReroute);
 					RemovedReroute->MarkAsGarbage();
 				}
 				else
 				{
-					Material->GetExpressionCollection().RemoveExpression(RemovedReroute);
+					Material->Expressions.Remove(RemovedReroute);
 				}
 			}
 		}
@@ -12624,9 +12601,6 @@ UMaterialFunctionInterface::UMaterialFunctionInterface(const FObjectInitializer&
 
 void UMaterialFunctionInterface::PostInitProperties()
 {
-#if WITH_EDITORONLY_DATA
-	EditorOnlyData = CreateEditorOnlyData();
-#endif
 	Super::PostInitProperties();
 
 	// Initialize StateId to something unique, in case this is a new function
@@ -12636,6 +12610,7 @@ void UMaterialFunctionInterface::PostInitProperties()
 void UMaterialFunctionInterface::PostLoad()
 {
 	Super::PostLoad();
+
 	if (!StateId.IsValid())
 	{
 		StateId = FGuid::NewGuid();
@@ -12659,58 +12634,7 @@ void UMaterialFunctionInterface::GetAssetRegistryTags(TArray<FAssetRegistryTag>&
 #endif
 }
 
-UMaterialFunctionInterface* UMaterialFunctionInterface::GetBaseFunctionInterface()
-{
-	return GetBaseFunction();
-}
-
-const UMaterialFunctionInterface* UMaterialFunctionInterface::GetBaseFunctionInterface() const
-{
-	return GetBaseFunction();
-}
-
 #if WITH_EDITOR
-
-TConstArrayView<TObjectPtr<UMaterialExpression>> UMaterialFunctionInterface::GetExpressions() const
-{
-	const UMaterialFunction* BaseFunction = GetBaseFunction();
-	if (BaseFunction)
-	{
-		return BaseFunction->GetExpressions();
-	}
-	return TConstArrayView<TObjectPtr<UMaterialExpression>>();
-}
-
-const FString& UMaterialFunctionInterface::GetDescription() const
-{
-	const UMaterialFunction* BaseFunction = GetBaseFunction();
-	if (BaseFunction)
-	{
-		return BaseFunction->Description;
-	}
-	static const FString EmptyString;
-	return EmptyString;
-}
-
-bool UMaterialFunctionInterface::GetReentrantFlag() const
-{
-	const UMaterialFunction* BaseFunction = GetBaseFunction();
-	if (BaseFunction)
-	{
-		return BaseFunction->GetReentrantFlag();
-	}
-	return false;
-}
-
-void UMaterialFunctionInterface::SetReentrantFlag(bool bIsReentrant)
-{
-	UMaterialFunction* BaseFunction = GetBaseFunction();
-	if (BaseFunction)
-	{
-		BaseFunction->SetReentrantFlag(bIsReentrant);
-	}
-}
-
 void UMaterialFunctionInterface::PostEditChangeProperty(FPropertyChangedEvent& PropertyChangedEvent)
 {
 	if (false) // temporary to unblock people, tracked as UE-109810
@@ -12862,24 +12786,6 @@ bool UMaterialFunctionInterface::OverrideNamedStaticComponentMaskParameter(const
 }
 #endif // WITH_EDITOR
 
-#if WITH_EDITORONLY_DATA
-const UClass* UMaterialFunctionInterface::GetEditorOnlyDataClass() const
-{
-	return UMaterialFunctionInterfaceEditorOnlyData::StaticClass();
-}
-
-UMaterialFunctionInterfaceEditorOnlyData* UMaterialFunctionInterface::CreateEditorOnlyData()
-{
-	const UClass* EditorOnlyClass = GetEditorOnlyDataClass();
-	check(EditorOnlyClass);
-	check(EditorOnlyClass->HasAllClassFlags(CLASS_Optional));
-
-	const FString EditorOnlyName = FString::Printf(TEXT("%sEditorOnlyData"), *GetName());
-	const EObjectFlags EditorOnlyFlags = GetMaskedFlags(RF_PropagateToSubObjects);
-	return NewObject<UMaterialFunctionInterfaceEditorOnlyData>(this, EditorOnlyClass, *EditorOnlyName, EditorOnlyFlags);
-}
-#endif // WITH_EDITORONLY_DATA
-
 ///////////////////////////////////////////////////////////////////////////////
 // UMaterialFunctionMaterialLayer
 ///////////////////////////////////////////////////////////////////////////////
@@ -12935,13 +12841,13 @@ UMaterialInterface* UMaterialFunction::GetPreviewMaterial()
 	{
 		PreviewMaterial = NewObject<UMaterial>(this, NAME_None, RF_Transient | RF_Public);
 
-		PreviewMaterial->AssignExpressionCollection(GetExpressionCollection());
+		PreviewMaterial->Expressions = FunctionExpressions;
 
 		//Find the first output expression and use that. 
-		for (UMaterialExpression* Expression : GetExpressions())
+		for( int32 i=0; i < FunctionExpressions.Num() ; ++i)
 		{
-			UMaterialExpressionFunctionOutput* Output = Cast<UMaterialExpressionFunctionOutput>(Expression);
-			if (Output)
+			UMaterialExpressionFunctionOutput* Output = Cast<UMaterialExpressionFunctionOutput>(FunctionExpressions[i]);
+			if( Output )
 			{
 				Output->ConnectToPreviewMaterial(PreviewMaterial, 0);
 			}
@@ -12959,8 +12865,9 @@ void UMaterialFunction::UpdateInputOutputTypes()
 	CombinedInputTypes = 0;
 	CombinedOutputTypes = 0;
 
-	for (UMaterialExpression* CurrentExpression : GetExpressions())
+	for (int32 ExpressionIndex = 0; ExpressionIndex < FunctionExpressions.Num(); ExpressionIndex++)
 	{
+		UMaterialExpression* CurrentExpression = FunctionExpressions[ExpressionIndex];
 		UMaterialExpressionFunctionOutput* OutputExpression = Cast<UMaterialExpressionFunctionOutput>(CurrentExpression);
 		UMaterialExpressionFunctionInput* InputExpression = Cast<UMaterialExpressionFunctionInput>(CurrentExpression);
 
@@ -12979,7 +12886,7 @@ void UMaterialFunction::UpdateInputOutputTypes()
 void UMaterialFunction::UpdateDependentFunctionCandidates()
 {
 	DependentFunctionExpressionCandidates.Reset();
-	for (UMaterialExpression* CurrentExpression : GetExpressions())
+	for (UMaterialExpression* CurrentExpression : FunctionExpressions)
 	{
 		if (UMaterialExpressionMaterialFunctionCall* MaterialFunctionExpression = Cast<UMaterialExpressionMaterialFunctionCall>(CurrentExpression))
 		{
@@ -13072,38 +12979,12 @@ void UMaterialFunction::PostLoad()
 	Super::PostLoad();
 
 #if WITH_EDITORONLY_DATA
-	UMaterialFunctionEditorOnlyData* EditorOnly = GetEditorOnlyData();
-
-	if (EditorOnly && FunctionExpressions_DEPRECATED.Num() > 0)
-	{
-		ensure(EditorOnly->ExpressionCollection.Expressions.Num() == 0);
-		EditorOnly->ExpressionCollection.Expressions = MoveTemp(FunctionExpressions_DEPRECATED);
-	}
-
-	if (EditorOnly && FunctionEditorComments_DEPRECATED.Num() > 0)
-	{
-		ensure(EditorOnly->ExpressionCollection.EditorComments.Num() == 0);
-		EditorOnly->ExpressionCollection.EditorComments = MoveTemp(FunctionEditorComments_DEPRECATED);
-	}
-
-	if (EditorOnly && ExpressionExecBegin_DEPRECATED)
-	{
-		ensure(!EditorOnly->ExpressionCollection.ExpressionExecBegin);
-		EditorOnly->ExpressionCollection.ExpressionExecBegin = MoveTemp(ExpressionExecBegin_DEPRECATED);
-	}
-
-	if (EditorOnly && ExpressionExecEnd_DEPRECATED)
-	{
-		ensure(!EditorOnly->ExpressionCollection.ExpressionExecEnd);
-		EditorOnly->ExpressionCollection.ExpressionExecEnd = MoveTemp(ExpressionExecEnd_DEPRECATED);
-	}
-
-	for (UMaterialExpression* Expression : EditorOnly->ExpressionCollection.Expressions)
+	for (int32 ExpressionIndex = 0; ExpressionIndex < FunctionExpressions.Num(); ExpressionIndex++)
 	{
 		// Expressions whose type was removed can be nullptr
-		if (Expression)
+		if (FunctionExpressions[ExpressionIndex])
 		{
-			Expression->ConditionalPostLoad();
+			FunctionExpressions[ExpressionIndex]->ConditionalPostLoad();
 		}
 	}
 #endif // WITH_EDITORONLY_DATA
@@ -13118,7 +12999,7 @@ void UMaterialFunction::PostLoad()
 	if (GIsEditor)
 	{
 		// Clean up any removed material expression classes	
-		if (EditorOnly->ExpressionCollection.Expressions.Remove(nullptr) != 0)
+		if (FunctionExpressions.Remove(nullptr) != 0)
 		{
 			// Force this function to recompile because its expressions have changed
 			// Warning: any content taking this path will recompile every load until saved!
@@ -13130,35 +13011,36 @@ void UMaterialFunction::PostLoad()
 	if (GMaterialFunctionsThatNeedExpressionsFlipped.Get(this))
 	{
 		GMaterialFunctionsThatNeedExpressionsFlipped.Clear(this);
-		UMaterial::FlipExpressionPositions(EditorOnly->ExpressionCollection.Expressions, EditorOnly->ExpressionCollection.EditorComments, true);
+		UMaterial::FlipExpressionPositions(FunctionExpressions, FunctionEditorComments, true);
 	}
 	else if (GMaterialFunctionsThatNeedCoordinateCheck.Get(this))
 	{
 		GMaterialFunctionsThatNeedCoordinateCheck.Clear(this);
 		if (HasFlippedCoordinates())
 		{
-			UMaterial::FlipExpressionPositions(EditorOnly->ExpressionCollection.Expressions, EditorOnly->ExpressionCollection.EditorComments, false);
+			UMaterial::FlipExpressionPositions(FunctionExpressions, FunctionEditorComments, false);
 		}
-		UMaterial::FixCommentPositions(EditorOnly->ExpressionCollection.EditorComments);
+		UMaterial::FixCommentPositions(FunctionEditorComments);
 	}
 	else if (GMaterialFunctionsThatNeedCommentFix.Get(this))
 	{
 		GMaterialFunctionsThatNeedCommentFix.Clear(this);
-		UMaterial::FixCommentPositions(EditorOnly->ExpressionCollection.EditorComments);
+		UMaterial::FixCommentPositions(FunctionEditorComments);
 	}
 
 	if (GMaterialFunctionsThatNeedFeatureLevelSM6Fix.Get(this))
 	{
 		GMaterialFunctionsThatNeedFeatureLevelSM6Fix.Clear(this);
-		UMaterial::FixFeatureLevelNodesForSM6(EditorOnly->ExpressionCollection.Expressions);
+		UMaterial::FixFeatureLevelNodesForSM6(FunctionExpressions);
 	}
 
 	if (GMaterialFunctionsThatNeedSamplerFixup.Get(this))
 	{
 		GMaterialFunctionsThatNeedSamplerFixup.Clear(this);
-		for (UMaterialExpression* Expression : EditorOnly->ExpressionCollection.Expressions)
+		const int32 ExpressionCount = FunctionExpressions.Num();
+		for (int32 ExpressionIndex = 0; ExpressionIndex < ExpressionCount; ++ExpressionIndex)
 		{
-			UMaterialExpressionTextureBase* TextureExpression = Cast<UMaterialExpressionTextureBase>(Expression);
+			UMaterialExpressionTextureBase* TextureExpression = Cast<UMaterialExpressionTextureBase>(FunctionExpressions[ExpressionIndex]);
 			if (TextureExpression && TextureExpression->Texture)
 			{
 				switch (TextureExpression->Texture->CompressionSettings)
@@ -13187,49 +13069,13 @@ void UMaterialFunction::PostLoad()
 #endif // #if WITH_EDITOR
 }
 
-#if WITH_EDITORONLY_DATA
-TConstArrayView<TObjectPtr<UMaterialExpression>> UMaterialFunction::GetExpressions() const
-{
-	return GetEditorOnlyData()->ExpressionCollection.Expressions;
-}
-
-TConstArrayView<TObjectPtr<UMaterialExpressionComment>> UMaterialFunction::GetEditorComments() const
-{
-	return GetEditorOnlyData()->ExpressionCollection.EditorComments;
-}
-
-UMaterialExpressionExecBegin* UMaterialFunction::GetExpressionExecBegin() const
-{
-	return GetEditorOnlyData()->ExpressionCollection.ExpressionExecBegin;
-}
-
-UMaterialExpressionExecEnd* UMaterialFunction::GetExpressionExecEnd() const
-{
-	return GetEditorOnlyData()->ExpressionCollection.ExpressionExecEnd;
-}
-
-const FMaterialExpressionCollection& UMaterialFunction::GetExpressionCollection() const
-{
-	return GetEditorOnlyData()->ExpressionCollection;
-}
-
-FMaterialExpressionCollection& UMaterialFunction::GetExpressionCollection()
-{
-	return GetEditorOnlyData()->ExpressionCollection;
-}
-
-void UMaterialFunction::AssignExpressionCollection(const FMaterialExpressionCollection& InCollection)
-{
-	GetEditorOnlyData()->ExpressionCollection = InCollection;
-}
-#endif // WITH_EDITORONLY_DATA
-
 #if WITH_EDITOR
 
 void UMaterialFunction::UpdateFromFunctionResource()
 {
-	for (UMaterialExpression* CurrentExpression : GetExpressions())
+	for (int32 ExpressionIndex = 0; ExpressionIndex < FunctionExpressions.Num(); ExpressionIndex++)
 	{
+		UMaterialExpression* CurrentExpression = FunctionExpressions[ExpressionIndex];
 		UMaterialExpressionMaterialFunctionCall* MaterialFunctionExpression = Cast<UMaterialExpressionMaterialFunctionCall>(CurrentExpression);
 		if (MaterialFunctionExpression)
 		{
@@ -13240,8 +13086,9 @@ void UMaterialFunction::UpdateFromFunctionResource()
 
 void UMaterialFunction::GetInputsAndOutputs(TArray<FFunctionExpressionInput>& OutInputs, TArray<FFunctionExpressionOutput>& OutOutputs) const
 {
-	for (UMaterialExpression* CurrentExpression : GetExpressions())
+	for (int32 ExpressionIndex = 0; ExpressionIndex < FunctionExpressions.Num(); ExpressionIndex++)
 	{
+		UMaterialExpression* CurrentExpression = FunctionExpressions[ExpressionIndex];
 		UMaterialExpressionFunctionOutput* OutputExpression = Cast<UMaterialExpressionFunctionOutput>(CurrentExpression);
 		UMaterialExpressionFunctionInput* InputExpression = Cast<UMaterialExpressionFunctionInput>(CurrentExpression);
 
@@ -13367,7 +13214,7 @@ bool UMaterialFunction::ValidateFunctionUsage(FMaterialCompiler* Compiler, const
 	if (GetMaterialFunctionUsage() == EMaterialFunctionUsage::MaterialLayer)
 	{
 		// Material layers must have a single MA input and output only
-		for (UMaterialExpression* Expression : GetExpressions())
+		for (UMaterialExpression* Expression : FunctionExpressions)
 		{
 			if (UMaterialExpressionFunctionInput* InputExpression = Cast<UMaterialExpressionFunctionInput>(Expression))
 			{
@@ -13403,7 +13250,7 @@ bool UMaterialFunction::ValidateFunctionUsage(FMaterialCompiler* Compiler, const
 	else if (GetMaterialFunctionUsage() == EMaterialFunctionUsage::MaterialLayerBlend)
 	{
 		// Material layer blends can have up to two MA inputs and single MA output only
-		for (UMaterialExpression* Expression : GetExpressions())
+		for (UMaterialExpression* Expression : FunctionExpressions)
 		{
 			if (UMaterialExpressionFunctionInput* InputExpression = Cast<UMaterialExpressionFunctionInput>(Expression))
 			{
@@ -13465,8 +13312,9 @@ int32 UMaterialFunction::Compile(FMaterialCompiler* Compiler, const FFunctionExp
 void UMaterialFunction::LinkIntoCaller(const TArray<FFunctionExpressionInput>& CallerInputs)
 {
 	// Go through all the function's input expressions and hook their inputs up to the corresponding expression in the material being compiled.
-	for (UMaterialExpression* CurrentExpression : GetExpressions())
+	for (int32 ExpressionIndex = 0; ExpressionIndex < FunctionExpressions.Num(); ExpressionIndex++)
 	{
+		UMaterialExpression* CurrentExpression = FunctionExpressions[ExpressionIndex];
 		UMaterialExpressionFunctionInput* InputExpression = Cast<UMaterialExpressionFunctionInput>(CurrentExpression);
 
 		if (InputExpression)
@@ -13499,8 +13347,9 @@ void UMaterialFunction::LinkIntoCaller(const TArray<FFunctionExpressionInput>& C
 
 void UMaterialFunction::UnlinkFromCaller()
 {
-	for (UMaterialExpression* CurrentExpression : GetExpressions())
+	for (int32 ExpressionIndex = 0; ExpressionIndex < FunctionExpressions.Num(); ExpressionIndex++)
 	{
+		UMaterialExpression* CurrentExpression = FunctionExpressions[ExpressionIndex];
 		UMaterialExpressionFunctionInput* InputExpression = Cast<UMaterialExpressionFunctionInput>(CurrentExpression);
 
 		if (InputExpression)
@@ -13539,8 +13388,9 @@ bool UMaterialFunction::IsDependent(UMaterialFunctionInterface* OtherFunction)
 #endif
 
 	bool bIsDependent = false;
-	for (UMaterialExpression* CurrentExpression : GetExpressions())
+	for (int32 ExpressionIndex = 0; ExpressionIndex < FunctionExpressions.Num(); ExpressionIndex++)
 	{
+		UMaterialExpression* CurrentExpression = FunctionExpressions[ExpressionIndex];
 		UMaterialExpressionMaterialFunctionCall* MaterialFunctionExpression = Cast<UMaterialExpressionMaterialFunctionCall>(CurrentExpression);
 		if (MaterialFunctionExpression && MaterialFunctionExpression->MaterialFunction)
 		{
@@ -13600,9 +13450,9 @@ bool UMaterialFunction::HasFlippedCoordinates() const
 	uint32 ReversedInputCount = 0;
 	uint32 StandardInputCount = 0;
 
-	for (UMaterialExpression* CurrentExpression : GetExpressions())
+	for (int32 Index = 0; Index < FunctionExpressions.Num(); ++Index)
 	{
-		UMaterialExpressionFunctionOutput* FunctionOutput = Cast<UMaterialExpressionFunctionOutput>(CurrentExpression);
+		UMaterialExpressionFunctionOutput* FunctionOutput = Cast<UMaterialExpressionFunctionOutput>(FunctionExpressions[Index]);
 		if (FunctionOutput && FunctionOutput->A.Expression)
 		{
 			if (FunctionOutput->A.Expression->MaterialExpressionEditorX > FunctionOutput->MaterialExpressionEditorX)
@@ -13623,7 +13473,7 @@ bool UMaterialFunction::HasFlippedCoordinates() const
 bool UMaterialFunction::SetParameterValueEditorOnly(const FName& ParameterName, const FMaterialParameterMetadata& Meta)
 {
 	bool bResult = false;
-	for (UMaterialExpression* Expression : GetExpressions())
+	for (UMaterialExpression* Expression : FunctionExpressions)
 	{
 		if (Expression->SetParameterValue(ParameterName, Meta))
 		{
@@ -13639,11 +13489,15 @@ bool UMaterialFunction::SetParameterValueEditorOnly(const FName& ParameterName,
 
 				for (UMaterialFunctionInterface* Function : Functions)
 				{
-					for (const TObjectPtr<UMaterialExpression>& FunctionExpression : Function->GetExpressions())
+					const TArray<TObjectPtr<UMaterialExpression>>* ExpressionPtr = Function->GetFunctionExpressions();
+					if (ExpressionPtr)
 					{
-						if (FunctionExpression->SetParameterValue(ParameterName, Meta))
+						for (const TObjectPtr<UMaterialExpression>& FunctionExpression : *ExpressionPtr)
 						{
-							bResult = true;
+							if (FunctionExpression->SetParameterValue(ParameterName, Meta))
+							{
+								bResult = true;
+							}
 						}
 					}
 				}
@@ -13727,41 +13581,39 @@ bool UMaterialFunction::IsUsingNewHLSLGenerator() const
 void UMaterialFunction::CreateExecutionFlowExpressions()
 {
 	const bool bUsingControlFlow = IsUsingControlFlow();
-	UMaterialFunctionEditorOnlyData* EditorOnly = GetEditorOnlyData();
 	if (bUsingControlFlow)
 	{
-		if (!EditorOnly->ExpressionCollection.ExpressionExecBegin)
+		if (!ExpressionExecBegin)
 		{
-			EditorOnly->ExpressionCollection.ExpressionExecBegin = NewObject<UMaterialExpressionExecBegin>(this);
-			EditorOnly->ExpressionCollection.ExpressionExecBegin->Function = this;
-			EditorOnly->ExpressionCollection.AddExpression(EditorOnly->ExpressionCollection.ExpressionExecBegin);
+			ExpressionExecBegin = NewObject<UMaterialExpressionExecBegin>(this);
+			ExpressionExecBegin->Function = this;
+			FunctionExpressions.Add(ExpressionExecBegin);
 		}
 
-		if (!EditorOnly->ExpressionCollection.ExpressionExecEnd)
+		if (!ExpressionExecEnd)
 		{
-			EditorOnly->ExpressionCollection.ExpressionExecEnd = NewObject<UMaterialExpressionExecEnd>(this);
-			EditorOnly->ExpressionCollection.ExpressionExecEnd->Function = this;
-			EditorOnly->ExpressionCollection.AddExpression(EditorOnly->ExpressionCollection.ExpressionExecEnd);
+			ExpressionExecEnd = NewObject<UMaterialExpressionExecEnd>(this);
+			ExpressionExecEnd->Function = this;
+			FunctionExpressions.Add(ExpressionExecEnd);
 		}
 	}
 
 	if (EditorMaterial)
 	{
-		UMaterialEditorOnlyData* EditorMaterialEditorOnly = EditorMaterial->GetEditorOnlyData();
 		EditorMaterial->bEnableExecWire = bUsingControlFlow;
-		EditorMaterialEditorOnly->ExpressionCollection.ExpressionExecBegin = EditorOnly->ExpressionCollection.ExpressionExecBegin;
-		EditorMaterialEditorOnly->ExpressionCollection.ExpressionExecEnd = EditorOnly->ExpressionCollection.ExpressionExecEnd;
-		if (EditorOnly->ExpressionCollection.ExpressionExecBegin)
+		EditorMaterial->ExpressionExecBegin = ExpressionExecBegin;
+		EditorMaterial->ExpressionExecEnd = ExpressionExecEnd;
+		if (ExpressionExecBegin)
 		{
-			EditorMaterialEditorOnly->ExpressionCollection.AddExpression(EditorOnly->ExpressionCollection.ExpressionExecBegin);
-			EditorOnly->ExpressionCollection.ExpressionExecBegin->Function = nullptr;
-			EditorOnly->ExpressionCollection.ExpressionExecBegin->Material = EditorMaterial;
+			EditorMaterial->Expressions.AddUnique(ExpressionExecBegin);
+			ExpressionExecBegin->Function = nullptr;
+			ExpressionExecBegin->Material = EditorMaterial;
 		}
-		if (EditorOnly->ExpressionCollection.ExpressionExecEnd)
+		if (ExpressionExecEnd)
 		{
-			EditorMaterialEditorOnly->ExpressionCollection.AddExpression(EditorOnly->ExpressionCollection.ExpressionExecEnd);
-			EditorOnly->ExpressionCollection.ExpressionExecEnd->Function = nullptr;
-			EditorOnly->ExpressionCollection.ExpressionExecEnd->Material = EditorMaterial;
+			EditorMaterial->Expressions.AddUnique(ExpressionExecEnd);
+			ExpressionExecEnd->Function = nullptr;
+			ExpressionExecEnd->Material = EditorMaterial;
 		}
 	}
 }
@@ -13780,32 +13632,10 @@ UMaterialFunctionInstance::UMaterialFunctionInstance(const FObjectInitializer& O
 #endif
 }
 
-UMaterialFunction* UMaterialFunctionInstance::GetBaseFunction(FMFRecursionGuard RecursionGuard)
-{
-	if (!Parent || RecursionGuard.Contains(this))
-	{
-		return nullptr;
-	}
-
-	RecursionGuard.Set(this);
-	return Parent->GetBaseFunction(RecursionGuard);
-}
-
-const UMaterialFunction* UMaterialFunctionInstance::GetBaseFunction(FMFRecursionGuard RecursionGuard) const
-{
-	if (!Parent || RecursionGuard.Contains(this))
-	{
-		return nullptr;
-	}
-
-	RecursionGuard.Set(this);
-	return Parent->GetBaseFunction(RecursionGuard);
-}
-
 #if WITH_EDITOR
 void UMaterialFunctionInstance::UpdateParameterSet()
 {
-	if (UMaterialFunction* BaseFunction = GetBaseFunction())
+	if (UMaterialFunction* BaseFunction = Cast<UMaterialFunction>(GetBaseFunction()))
 	{
 		TArray<UMaterialFunctionInterface*> Functions;
 		BaseFunction->GetDependentFunctions(Functions);
@@ -13814,7 +13644,7 @@ void UMaterialFunctionInstance::UpdateParameterSet()
 		// Loop through all contained parameters and update names as needed
 		for (UMaterialFunctionInterface* Function : Functions)
 		{
-			for (UMaterialExpression* FunctionExpression : Function->GetExpressions())
+			for (UMaterialExpression* FunctionExpression : *Function->GetFunctionExpressions())
 			{
 				if (const UMaterialExpressionScalarParameter* ScalarParameter = Cast<const UMaterialExpressionScalarParameter>(FunctionExpression))
 				{
@@ -13909,8 +13739,8 @@ void UMaterialFunctionInstance::OverrideMaterialInstanceParameterValues(UMateria
 
 	// Static parameters
 	FStaticParameterSet StaticParametersOverride = Instance->GetStaticParameters();
-	StaticParametersOverride.EditorOnly.StaticSwitchParameters = StaticSwitchParameterValues;
-	StaticParametersOverride.EditorOnly.StaticComponentMaskParameters = StaticComponentMaskParameterValues;
+	StaticParametersOverride.StaticSwitchParameters = StaticSwitchParameterValues;
+	StaticParametersOverride.StaticComponentMaskParameters = StaticComponentMaskParameterValues;
 	Instance->UpdateStaticPermutation(StaticParametersOverride);
 }
 
@@ -14080,16 +13910,16 @@ bool UMaterialFunctionInstance::GetParameterOverrideValue(EMaterialParameterType
 }
 
 ///////////////////////////////////////////////////////////////////////////////
-// FMaterialLayersFunctionsID
+// FMaterialLayersFunctions::ID
 ///////////////////////////////////////////////////////////////////////////////
 
-bool FMaterialLayersFunctionsID::operator==(const FMaterialLayersFunctionsID& Reference) const
+bool FMaterialLayersFunctions::ID::operator==(const ID& Reference) const
 {
 	return LayerIDs == Reference.LayerIDs && BlendIDs == Reference.BlendIDs && LayerStates == Reference.LayerStates;
 }
 
 
-void FMaterialLayersFunctionsID::SerializeForDDC(FArchive& Ar)
+void FMaterialLayersFunctions::ID::SerializeForDDC(FArchive& Ar)
 {
 	Ar << LayerIDs;
 	Ar << BlendIDs;
@@ -14097,7 +13927,7 @@ void FMaterialLayersFunctionsID::SerializeForDDC(FArchive& Ar)
 }
 
 
-void FMaterialLayersFunctionsID::UpdateHash(FSHA1& HashState) const
+void FMaterialLayersFunctions::ID::UpdateHash(FSHA1& HashState) const
 {
 	for (const FGuid &Guid : LayerIDs)
 	{
@@ -14111,7 +13941,7 @@ void FMaterialLayersFunctionsID::UpdateHash(FSHA1& HashState) const
 }
 
 
-void FMaterialLayersFunctionsID::AppendKeyString(FString& KeyString) const
+void FMaterialLayersFunctions::ID::AppendKeyString(FString& KeyString) const
 {
 	for (const FGuid &Guid : LayerIDs)
 	{
@@ -14135,9 +13965,9 @@ void FMaterialLayersFunctionsID::AppendKeyString(FString& KeyString) const
 const FGuid FMaterialLayersFunctions::BackgroundGuid(2u, 0u, 0u, 0u);
 
 #if WITH_EDITOR
-const FMaterialLayersFunctionsID FMaterialLayersFunctionsRuntimeData::GetID(const FMaterialLayersFunctionsEditorOnlyData& EditorOnly) const
+const FMaterialLayersFunctions::ID FMaterialLayersFunctions::GetID() const
 {
-	FMaterialLayersFunctionsID Result;
+	FMaterialLayersFunctions::ID Result;
 
 	// Store the layer IDs in following format - stateID per function
 	Result.LayerIDs.SetNum(Layers.Num());
@@ -14172,7 +14002,7 @@ const FMaterialLayersFunctionsID FMaterialLayersFunctionsRuntimeData::GetID(cons
 	}
 
 	// Store the states copy
-	Result.LayerStates = EditorOnly.LayerStates;
+	Result.LayerStates = LayerStates;
 
 	return Result;
 }
@@ -14196,46 +14026,17 @@ void FMaterialLayersFunctions::PostSerialize(const FArchive& Ar)
 #if WITH_EDITORONLY_DATA
 	if (Ar.IsLoading())
 	{
-		if (LayerStates_DEPRECATED.Num() > 0)
-		{
-			EditorOnly.LayerStates = MoveTemp(LayerStates_DEPRECATED);
-		}
-		if (LayerNames_DEPRECATED.Num() > 0)
-		{
-			EditorOnly.LayerNames = MoveTemp(LayerNames_DEPRECATED);
-		}
-		if (RestrictToLayerRelatives_DEPRECATED.Num() > 0)
-		{
-			EditorOnly.RestrictToLayerRelatives = MoveTemp(RestrictToLayerRelatives_DEPRECATED);
-		}
-		if (RestrictToBlendRelatives_DEPRECATED.Num() > 0)
-		{
-			EditorOnly.RestrictToBlendRelatives = MoveTemp(RestrictToBlendRelatives_DEPRECATED);
-		}
-		if (LayerGuids_DEPRECATED.Num() > 0)
-		{
-			EditorOnly.LayerGuids = MoveTemp(LayerGuids_DEPRECATED);
-		}
-		if (LayerLinkStates_DEPRECATED.Num() > 0)
-		{
-			EditorOnly.LayerLinkStates = MoveTemp(LayerLinkStates_DEPRECATED);
-		}
-		if (DeletedParentLayerGuids_DEPRECATED.Num() > 0)
-		{
-			EditorOnly.DeletedParentLayerGuids = MoveTemp(DeletedParentLayerGuids_DEPRECATED);
-		}
-
-		if (EditorOnly.LayerGuids.Num() != Layers.Num() ||
-			EditorOnly.LayerLinkStates.Num() != Layers.Num())
+		if (LayerGuids.Num() != Layers.Num() ||
+			LayerLinkStates.Num() != Layers.Num())
 		{
 			const int32 NumLayers = Layers.Num();
-			EditorOnly.LayerGuids.Empty(NumLayers);
-			EditorOnly.LayerLinkStates.Empty(NumLayers);
+			LayerGuids.Empty(NumLayers);
+			LayerLinkStates.Empty(NumLayers);
 
 			if (NumLayers > 0)
 			{
-				EditorOnly.LayerGuids.Add(BackgroundGuid);
-				EditorOnly.LayerLinkStates.Add(EMaterialLayerLinkState::Uninitialized);
+				LayerGuids.Add(BackgroundGuid);
+				LayerLinkStates.Add(EMaterialLayerLinkState::Uninitialized);
 
 				for (int32 i = 1; i < NumLayers; ++i)
 				{
@@ -14244,34 +14045,32 @@ void FMaterialLayersFunctions::PostSerialize(const FArchive& Ar)
 					// But it's possible *this* material may not actually be saved in that case
 					// If that happens, need to ensure that the guids remain consistent if this material is loaded again;
 					// otherwise they will no longer match the guids that were saved into the child material
-					EditorOnly.LayerGuids.Add(FGuid(3u, 0u, 0u, i));
-					EditorOnly.LayerLinkStates.Add(EMaterialLayerLinkState::Uninitialized);
+					LayerGuids.Add(FGuid(3u, 0u, 0u, i));
+					LayerLinkStates.Add(EMaterialLayerLinkState::Uninitialized);
 				}
 			}
 		}
 	}
 #endif // WITH_EDITORONLY_DATA
 }
-#if WITH_EDITOR
+
 int32 FMaterialLayersFunctions::AppendBlendedLayer()
 {
 	const int32 LayerIndex = Layers.AddDefaulted();
 	Blends.AddDefaulted();
-	EditorOnly.LayerStates.Add(true);
+#if WITH_EDITORONLY_DATA
+	LayerStates.Add(true);
 	FText LayerName = FText::Format(LOCTEXT("LayerPrefix", "Layer {0}"), Layers.Num() - 1);
-	EditorOnly.LayerNames.Add(LayerName);
-	EditorOnly.RestrictToLayerRelatives.Add(false);
-	EditorOnly.RestrictToBlendRelatives.Add(false);
-	EditorOnly.LayerGuids.Add(FGuid::NewGuid());
-	EditorOnly.LayerLinkStates.Add(EMaterialLayerLinkState::NotFromParent);
+	LayerNames.Add(LayerName);
+	RestrictToLayerRelatives.Add(false);
+	RestrictToBlendRelatives.Add(false);
+	LayerGuids.Add(FGuid::NewGuid());
+	LayerLinkStates.Add(EMaterialLayerLinkState::NotFromParent);
+#endif // WITH_EDITORONLY_DATA
 	return LayerIndex;
 }
 
-int32 FMaterialLayersFunctions::AddLayerCopy(const FMaterialLayersFunctionsRuntimeData& Source,
-	const FMaterialLayersFunctionsEditorOnlyData& SourceEditorOnly,
-	int32 SourceLayerIndex,
-	bool bVisible,
-	EMaterialLayerLinkState LinkState)
+int32 FMaterialLayersFunctions::AddLayerCopy(const FMaterialLayersFunctions& Source, int32 SourceLayerIndex, bool bVisible, EMaterialLayerLinkState LinkState)
 {
 	check(LinkState != EMaterialLayerLinkState::Uninitialized);
 	const int32 LayerIndex = Layers.Num();
@@ -14282,35 +14081,35 @@ int32 FMaterialLayersFunctions::AddLayerCopy(const FMaterialLayersFunctionsRunti
 		Blends.Add(Source.Blends[SourceLayerIndex - 1]);
 	}
 	
-	EditorOnly.LayerStates.Add(bVisible);
-	EditorOnly.LayerNames.Add(SourceEditorOnly.LayerNames[SourceLayerIndex]);
-	EditorOnly.RestrictToLayerRelatives.Add(SourceEditorOnly.RestrictToLayerRelatives[SourceLayerIndex]);
+#if WITH_EDITOR
+	LayerStates.Add(bVisible);
+	LayerNames.Add(Source.LayerNames[SourceLayerIndex]);
+	RestrictToLayerRelatives.Add(Source.RestrictToLayerRelatives[SourceLayerIndex]);
 	if (LayerIndex > 0)
 	{
-		EditorOnly.RestrictToBlendRelatives.Add(SourceEditorOnly.RestrictToBlendRelatives[SourceLayerIndex - 1]);
+		RestrictToBlendRelatives.Add(Source.RestrictToBlendRelatives[SourceLayerIndex - 1]);
 	}
-	EditorOnly.LayerGuids.Add(SourceEditorOnly.LayerGuids[SourceLayerIndex]);
-	EditorOnly.LayerLinkStates.Add(LinkState);
+	LayerGuids.Add(Source.LayerGuids[SourceLayerIndex]);
+	LayerLinkStates.Add(LinkState);
+#endif
 	return LayerIndex;
 }
 
-void FMaterialLayersFunctions::InsertLayerCopy(const FMaterialLayersFunctionsRuntimeData& Source,
-	const FMaterialLayersFunctionsEditorOnlyData& SourceEditorOnly,
-	int32 SourceLayerIndex,
-	EMaterialLayerLinkState LinkState,
-	int32 LayerIndex)
+void FMaterialLayersFunctions::InsertLayerCopy(const FMaterialLayersFunctions& Source, int32 SourceLayerIndex, EMaterialLayerLinkState LinkState, int32 LayerIndex)
 {
 	check(LinkState != EMaterialLayerLinkState::Uninitialized);
 	check(LayerIndex > 0);
 	Layers.Insert(Source.Layers[SourceLayerIndex], LayerIndex);
 	Blends.Insert(Source.Blends[SourceLayerIndex - 1], LayerIndex - 1);
 	
-	EditorOnly.LayerStates.Insert(SourceEditorOnly.LayerStates[SourceLayerIndex], LayerIndex);
-	EditorOnly.LayerNames.Insert(SourceEditorOnly.LayerNames[SourceLayerIndex], LayerIndex);
-	EditorOnly.RestrictToLayerRelatives.Insert(SourceEditorOnly.RestrictToLayerRelatives[SourceLayerIndex], LayerIndex);
-	EditorOnly.RestrictToBlendRelatives.Insert(SourceEditorOnly.RestrictToBlendRelatives[SourceLayerIndex - 1], LayerIndex - 1);
-	EditorOnly.LayerGuids.Insert(SourceEditorOnly.LayerGuids[SourceLayerIndex], LayerIndex);
-	EditorOnly.LayerLinkStates.Insert(LinkState, LayerIndex);
+#if WITH_EDITOR
+	LayerStates.Insert(Source.LayerStates[SourceLayerIndex], LayerIndex);
+	LayerNames.Insert(Source.LayerNames[SourceLayerIndex], LayerIndex);
+	RestrictToLayerRelatives.Insert(Source.RestrictToLayerRelatives[SourceLayerIndex], LayerIndex);
+	RestrictToBlendRelatives.Insert(Source.RestrictToBlendRelatives[SourceLayerIndex - 1], LayerIndex - 1);
+	LayerGuids.Insert(Source.LayerGuids[SourceLayerIndex], LayerIndex);
+	LayerLinkStates.Insert(LinkState, LayerIndex);
+#endif
 }
 
 void FMaterialLayersFunctions::RemoveBlendedLayerAt(int32 Index)
@@ -14321,25 +14120,24 @@ void FMaterialLayersFunctions::RemoveBlendedLayerAt(int32 Index)
 		Layers.RemoveAt(Index);
 		Blends.RemoveAt(Index - 1);
 		
-		check(EditorOnly.LayerStates.IsValidIndex(Index) &&
-			EditorOnly.LayerNames.IsValidIndex(Index) &&
-			EditorOnly.RestrictToLayerRelatives.IsValidIndex(Index) &&
-			EditorOnly.RestrictToBlendRelatives.IsValidIndex(Index - 1));
+#if WITH_EDITOR
+		check(LayerStates.IsValidIndex(Index) && LayerNames.IsValidIndex(Index) && RestrictToLayerRelatives.IsValidIndex(Index) && RestrictToBlendRelatives.IsValidIndex(Index - 1));
 
-		if (EditorOnly.LayerLinkStates[Index] != EMaterialLayerLinkState::NotFromParent)
+		if (LayerLinkStates[Index] != EMaterialLayerLinkState::NotFromParent)
 		{
 			// Save the parent guid as explicitly deleted, so it's not added back
-			const FGuid& LayerGuid = EditorOnly.LayerGuids[Index];
-			check(!EditorOnly.DeletedParentLayerGuids.Contains(LayerGuid));
-			EditorOnly.DeletedParentLayerGuids.Add(LayerGuid);
+			const FGuid& LayerGuid = LayerGuids[Index];
+			check(!DeletedParentLayerGuids.Contains(LayerGuid));
+			DeletedParentLayerGuids.Add(LayerGuid);
 		}
 
-		EditorOnly.LayerStates.RemoveAt(Index);
-		EditorOnly.LayerNames.RemoveAt(Index);
-		EditorOnly.RestrictToLayerRelatives.RemoveAt(Index);
-		EditorOnly.RestrictToBlendRelatives.RemoveAt(Index - 1);
-		EditorOnly.LayerGuids.RemoveAt(Index);
-		EditorOnly.LayerLinkStates.RemoveAt(Index);
+		LayerStates.RemoveAt(Index);
+		LayerNames.RemoveAt(Index);
+		RestrictToLayerRelatives.RemoveAt(Index);
+		RestrictToBlendRelatives.RemoveAt(Index - 1);
+		LayerGuids.RemoveAt(Index);
+		LayerLinkStates.RemoveAt(Index);
+#endif //WITH_EDITOR
 	}
 }
 
@@ -14351,53 +14149,56 @@ void FMaterialLayersFunctions::MoveBlendedLayer(int32 SrcLayerIndex, int32 DstLa
 	{
 		Layers.Swap(SrcLayerIndex, DstLayerIndex);
 		Blends.Swap(SrcLayerIndex - 1, DstLayerIndex - 1);
-		EditorOnly.LayerStates.Swap(SrcLayerIndex, DstLayerIndex);
-		EditorOnly.LayerNames.Swap(SrcLayerIndex, DstLayerIndex);
-		EditorOnly.RestrictToLayerRelatives.Swap(SrcLayerIndex, DstLayerIndex);
-		EditorOnly.RestrictToBlendRelatives.Swap(SrcLayerIndex - 1, DstLayerIndex - 1);
-		EditorOnly.LayerGuids.Swap(SrcLayerIndex, DstLayerIndex);
-		EditorOnly.LayerLinkStates.Swap(SrcLayerIndex, DstLayerIndex);
+#if WITH_EDITOR
+		LayerStates.Swap(SrcLayerIndex, DstLayerIndex);
+		LayerNames.Swap(SrcLayerIndex, DstLayerIndex);
+		RestrictToLayerRelatives.Swap(SrcLayerIndex, DstLayerIndex);
+		RestrictToBlendRelatives.Swap(SrcLayerIndex - 1, DstLayerIndex - 1);
+		LayerGuids.Swap(SrcLayerIndex, DstLayerIndex);
+		LayerLinkStates.Swap(SrcLayerIndex, DstLayerIndex);
+#endif //WITH_EDITOR
 	}
 }
 
+#if WITH_EDITOR
 void FMaterialLayersFunctions::UnlinkLayerFromParent(int32 Index)
 {
-	if (EditorOnly.LayerLinkStates[Index] == EMaterialLayerLinkState::LinkedToParent)
+	if (LayerLinkStates[Index] == EMaterialLayerLinkState::LinkedToParent)
 	{
-		EditorOnly.LayerLinkStates[Index] = EMaterialLayerLinkState::UnlinkedFromParent;
+		LayerLinkStates[Index] = EMaterialLayerLinkState::UnlinkedFromParent;
 	}
 }
 
 bool FMaterialLayersFunctions::IsLayerLinkedToParent(int32 Index) const
 {
-	if (EditorOnly.LayerLinkStates.IsValidIndex(Index))
+	if (LayerLinkStates.IsValidIndex(Index))
 	{
-		return EditorOnly.LayerLinkStates[Index] == EMaterialLayerLinkState::LinkedToParent;
+		return LayerLinkStates[Index] == EMaterialLayerLinkState::LinkedToParent;
 	}
 	return false;
 }
 
 void FMaterialLayersFunctions::RelinkLayersToParent()
 {
-	for (int32 Index = 0; Index < EditorOnly.LayerLinkStates.Num(); ++Index)
+	for (int32 Index = 0; Index < LayerLinkStates.Num(); ++Index)
 	{
-		if (EditorOnly.LayerLinkStates[Index] == EMaterialLayerLinkState::UnlinkedFromParent)
+		if (LayerLinkStates[Index] == EMaterialLayerLinkState::UnlinkedFromParent)
 		{
-			EditorOnly.LayerLinkStates[Index] = EMaterialLayerLinkState::LinkedToParent;
+			LayerLinkStates[Index] = EMaterialLayerLinkState::LinkedToParent;
 		}
 	}
-	EditorOnly.DeletedParentLayerGuids.Empty();
+	DeletedParentLayerGuids.Empty();
 }
 
 bool FMaterialLayersFunctions::HasAnyUnlinkedLayers() const
 {
-	if (EditorOnly.DeletedParentLayerGuids.Num() > 0)
+	if (DeletedParentLayerGuids.Num() > 0)
 	{
 		return true;
 	}
-	for (int32 Index = 0; Index < EditorOnly.LayerLinkStates.Num(); ++Index)
+	for (int32 Index = 0; Index < LayerLinkStates.Num(); ++Index)
 	{
-		if (EditorOnly.LayerLinkStates[Index] == EMaterialLayerLinkState::UnlinkedFromParent)
+		if (LayerLinkStates[Index] == EMaterialLayerLinkState::UnlinkedFromParent)
 		{
 			return true;
 		}
@@ -14405,7 +14206,7 @@ bool FMaterialLayersFunctions::HasAnyUnlinkedLayers() const
 	return false;
 }
 
-void FMaterialLayersFunctionsEditorOnlyData::LinkAllLayersToParent()
+void FMaterialLayersFunctions::LinkAllLayersToParent()
 {
 	for (int32 Index = 0; Index < LayerLinkStates.Num(); ++Index)
 	{
@@ -14413,44 +14214,37 @@ void FMaterialLayersFunctionsEditorOnlyData::LinkAllLayersToParent()
 	}
 }
 
-bool FMaterialLayersFunctions::MatchesParent(const FMaterialLayersFunctionsRuntimeData& Runtime,
-	const FMaterialLayersFunctionsEditorOnlyData& EditorOnly,
-	const FMaterialLayersFunctionsRuntimeData& ParentRuntime,
-	const FMaterialLayersFunctionsEditorOnlyData& ParentEditorOnly)
+bool FMaterialLayersFunctions::MatchesParent(const FMaterialLayersFunctions& Parent) const
 {
-	if (Runtime.Layers.Num() != ParentRuntime.Layers.Num())
+	if (Layers.Num() != Parent.Layers.Num())
 	{
 		return false;
 	}
 
-	for (int32 LayerIndex = 0; LayerIndex < Runtime.Layers.Num(); ++LayerIndex)
+	for (int32 LayerIndex = 0; LayerIndex < Layers.Num(); ++LayerIndex)
 	{
-		const EMaterialLayerLinkState LinkState = EditorOnly.LayerLinkStates[LayerIndex];
+		const EMaterialLayerLinkState LinkState = LayerLinkStates[LayerIndex];
 		if (LinkState != EMaterialLayerLinkState::LinkedToParent)
 		{
 			return false;
 		}
 
-		const FGuid& LayerGuid = EditorOnly.LayerGuids[LayerIndex];
-		const int32 ParentLayerIndex = ParentEditorOnly.LayerGuids.Find(LayerGuid);
+		const FGuid& LayerGuid = LayerGuids[LayerIndex];
+		const int32 ParentLayerIndex = Parent.LayerGuids.Find(LayerGuid);
 		if (ParentLayerIndex != LayerIndex)
 		{
 			return false;
 		}
 
-		// Possible for LayerStates arrays to be empty, if this is cooked data
-		// We assume all cooked layers are visible
-		const bool bLayerVisible = (EditorOnly.LayerStates.Num() > 0) ? EditorOnly.LayerStates[LayerIndex] : true;
-		const bool bParentLayerVisible = (ParentEditorOnly.LayerStates.Num() > 0) ? ParentEditorOnly.LayerStates[ParentLayerIndex] : true;
-		if (bLayerVisible != bParentLayerVisible)
+		if (LayerStates[LayerIndex] != Parent.LayerStates[ParentLayerIndex])
 		{
 			return false;
 		}
-		if (Runtime.Layers[LayerIndex] != ParentRuntime.Layers[ParentLayerIndex])
+		if (Layers[LayerIndex] != Parent.Layers[ParentLayerIndex])
 		{
 			return false;
 		}
-		if (LayerIndex > 0 && Runtime.Blends[LayerIndex - 1] != ParentRuntime.Blends[ParentLayerIndex - 1])
+		if (LayerIndex > 0 && Blends[LayerIndex - 1] != Parent.Blends[ParentLayerIndex - 1])
 		{
 			return false;
 		}
@@ -14459,14 +14253,10 @@ bool FMaterialLayersFunctions::MatchesParent(const FMaterialLayersFunctionsRunti
 	return true;
 }
 
-bool FMaterialLayersFunctions::ResolveParent(const FMaterialLayersFunctionsRuntimeData& ParentRuntime,
-	const FMaterialLayersFunctionsEditorOnlyData& ParentEditorOnly,
-	FMaterialLayersFunctionsRuntimeData& Runtime,
-	FMaterialLayersFunctionsEditorOnlyData& EditorOnly,
-	TArray<int32>& OutRemapLayerIndices)
+bool FMaterialLayersFunctions::ResolveParent(const FMaterialLayersFunctions& Parent, TArray<int32>& OutRemapLayerIndices)
 {
-	check(EditorOnly.LayerGuids.Num() == Runtime.Layers.Num());
-	check(EditorOnly.LayerLinkStates.Num() == Runtime.Layers.Num());
+	check(LayerGuids.Num() == Layers.Num());
+	check(LayerLinkStates.Num() == Layers.Num());
 
 	FMaterialLayersFunctions ResolvedLayers;
 	TArray<int32> ParentLayerIndices;
@@ -14474,12 +14264,12 @@ bool FMaterialLayersFunctions::ResolveParent(const FMaterialLayersFunctionsRunti
 	ResolvedLayers.Empty();
 
 	bool bHasUninitializedLinks = false;
-	for (int32 LayerIndex = 0; LayerIndex < Runtime.Layers.Num(); ++LayerIndex)
+	for (int32 LayerIndex = 0; LayerIndex < Layers.Num(); ++LayerIndex)
 	{
-		const FText& LayerName = EditorOnly.LayerNames[LayerIndex];
-		const FGuid& LayerGuid = EditorOnly.LayerGuids[LayerIndex];
-		const bool bLayerVisible = EditorOnly.LayerStates[LayerIndex];
-		const EMaterialLayerLinkState LinkState = EditorOnly.LayerLinkStates[LayerIndex];
+		const FText& LayerName = LayerNames[LayerIndex];
+		const FGuid& LayerGuid = LayerGuids[LayerIndex];
+		const bool bLayerVisible = LayerStates[LayerIndex];
+		const EMaterialLayerLinkState LinkState = LayerLinkStates[LayerIndex];
 
 		int32 ParentLayerIndex = INDEX_NONE;
 		if (LinkState == EMaterialLayerLinkState::Uninitialized)
@@ -14488,17 +14278,17 @@ bool FMaterialLayersFunctions::ResolveParent(const FMaterialLayersFunctionsRunti
 			if (LayerIndex == 0)
 			{
 				// Base layer must match against base layer
-				if (ParentRuntime.Layers.Num() > 0)
+				if (Parent.Layers.Num() > 0)
 				{
 					ParentLayerIndex = 0;
 				}
 			}
 			else
 			{
-				for (int32 CheckLayerIndex = 1; CheckLayerIndex < ParentRuntime.Layers.Num(); ++CheckLayerIndex)
+				for (int32 CheckLayerIndex = 1; CheckLayerIndex < Parent.Layers.Num(); ++CheckLayerIndex)
 				{
 					// check if name matches, and if we haven't already linked to this parent layer
-					if (LayerName.EqualTo(ParentEditorOnly.LayerNames[CheckLayerIndex]) &&
+					if (LayerName.EqualTo(Parent.LayerNames[CheckLayerIndex]) &&
 						!ParentLayerIndices.Contains(CheckLayerIndex))
 					{
 						ParentLayerIndex = CheckLayerIndex;
@@ -14511,23 +14301,23 @@ bool FMaterialLayersFunctions::ResolveParent(const FMaterialLayersFunctionsRunti
 			if (ParentLayerIndex == INDEX_NONE)
 			{
 				// Didn't find layer in the parent, assume it's local to this material
-				ResolvedLayerIndex = ResolvedLayers.AddLayerCopy(Runtime, EditorOnly, LayerIndex, bLayerVisible, EMaterialLayerLinkState::NotFromParent);
+				ResolvedLayerIndex = ResolvedLayers.AddLayerCopy(*this, LayerIndex, bLayerVisible, EMaterialLayerLinkState::NotFromParent);
 				ParentLayerIndices.Add(INDEX_NONE);
 			}
 			else
 			{
 				// See if we match layer in parent
-				if (Runtime.Layers[LayerIndex] == ParentRuntime.Layers[ParentLayerIndex] &&
-					(LayerIndex == 0 || Runtime.Blends[LayerIndex - 1] == ParentRuntime.Blends[ParentLayerIndex - 1]))
+				if (Layers[LayerIndex] == Parent.Layers[ParentLayerIndex] &&
+					(LayerIndex == 0 || Blends[LayerIndex - 1] == Parent.Blends[ParentLayerIndex - 1]))
 				{
 					// Parent layer matches, so link to parent
-					ResolvedLayerIndex = ResolvedLayers.AddLayerCopy(ParentRuntime, ParentEditorOnly, ParentLayerIndex, bLayerVisible, EMaterialLayerLinkState::LinkedToParent);
+					ResolvedLayerIndex = ResolvedLayers.AddLayerCopy(Parent, ParentLayerIndex, bLayerVisible, EMaterialLayerLinkState::LinkedToParent);
 				}
 				else
 				{
 					// Parent layer does NOT match, so make the child overriden
-					ResolvedLayerIndex = ResolvedLayers.AddLayerCopy(Runtime, EditorOnly, LayerIndex, bLayerVisible, EMaterialLayerLinkState::UnlinkedFromParent);
-					ResolvedLayers.EditorOnly.LayerGuids[ResolvedLayerIndex] = ParentEditorOnly.LayerGuids[ParentLayerIndex]; // Still need to match guid to parent
+					ResolvedLayerIndex = ResolvedLayers.AddLayerCopy(*this, LayerIndex, bLayerVisible, EMaterialLayerLinkState::UnlinkedFromParent);
+					ResolvedLayers.LayerGuids[ResolvedLayerIndex] = Parent.LayerGuids[ParentLayerIndex]; // Still need to match guid to parent
 				}
 
 				check(!ParentLayerIndices.Contains(ParentLayerIndex));
@@ -14540,11 +14330,11 @@ bool FMaterialLayersFunctions::ResolveParent(const FMaterialLayersFunctionsRunti
 		else if (LinkState == EMaterialLayerLinkState::LinkedToParent)
 		{
 			check(LayerGuid.IsValid());
-			ParentLayerIndex = ParentEditorOnly.LayerGuids.Find(LayerGuid);
+			ParentLayerIndex = Parent.LayerGuids.Find(LayerGuid);
 			if (ParentLayerIndex != INDEX_NONE)
 			{
 				// Layer comes from parent
-				ResolvedLayers.AddLayerCopy(ParentRuntime, ParentEditorOnly, ParentLayerIndex, bLayerVisible, EMaterialLayerLinkState::LinkedToParent);
+				ResolvedLayers.AddLayerCopy(Parent, ParentLayerIndex, bLayerVisible, EMaterialLayerLinkState::LinkedToParent);
 				check(!ParentLayerIndices.Contains(ParentLayerIndex));
 				ParentLayerIndices.Add(ParentLayerIndex);
 			}
@@ -14559,12 +14349,12 @@ bool FMaterialLayersFunctions::ResolveParent(const FMaterialLayersFunctionsRunti
 			if (LinkState == EMaterialLayerLinkState::UnlinkedFromParent)
 			{
 				// If we are unlinked from parent, track the layer index we were previously linked to
-				ParentLayerIndex = ParentEditorOnly.LayerGuids.Find(LayerGuid);
+				ParentLayerIndex = Parent.LayerGuids.Find(LayerGuid);
 			}
 			check(ParentLayerIndex == INDEX_NONE || !ParentLayerIndices.Contains(ParentLayerIndex));
 
 			// Update the link state, depending on if we can find this layer in the parent
-			ResolvedLayers.AddLayerCopy(Runtime, EditorOnly, LayerIndex, bLayerVisible, (ParentLayerIndex == INDEX_NONE) ? EMaterialLayerLinkState::NotFromParent : EMaterialLayerLinkState::UnlinkedFromParent);
+			ResolvedLayers.AddLayerCopy(*this, LayerIndex, bLayerVisible, (ParentLayerIndex == INDEX_NONE) ? EMaterialLayerLinkState::NotFromParent : EMaterialLayerLinkState::UnlinkedFromParent);
 			ParentLayerIndices.Add(ParentLayerIndex);
 		}
 	}
@@ -14572,7 +14362,7 @@ bool FMaterialLayersFunctions::ResolveParent(const FMaterialLayersFunctionsRunti
 	check(ResolvedLayers.Layers.Num() == ParentLayerIndices.Num());
 
 	// See if parent has any added layers
-	for (int32 ParentLayerIndex = 1; ParentLayerIndex < ParentRuntime.Layers.Num(); ++ParentLayerIndex)
+	for (int32 ParentLayerIndex = 1; ParentLayerIndex < Parent.Layers.Num(); ++ParentLayerIndex)
 	{
 		if (ParentLayerIndices.Contains(ParentLayerIndex))
 		{
@@ -14580,11 +14370,11 @@ bool FMaterialLayersFunctions::ResolveParent(const FMaterialLayersFunctionsRunti
 			continue;
 		}
 
-		const FGuid& ParentLayerGuid = ParentEditorOnly.LayerGuids[ParentLayerIndex];
-		if (EditorOnly.DeletedParentLayerGuids.Contains(ParentLayerGuid))
+		const FGuid& ParentLayerGuid = Parent.LayerGuids[ParentLayerIndex];
+		if (DeletedParentLayerGuids.Contains(ParentLayerGuid))
 		{
 			// Parent layer was previously explicitly overiden/deleted
-			ResolvedLayers.EditorOnly.DeletedParentLayerGuids.Add(ParentLayerGuid);
+			ResolvedLayers.DeletedParentLayerGuids.Add(ParentLayerGuid);
 			continue;
 		}
 
@@ -14595,7 +14385,7 @@ bool FMaterialLayersFunctions::ResolveParent(const FMaterialLayersFunctionsRunti
 			// or if this layer was explicitly deleted from child (and should therefore remain deleted).
 			// In order to avoid needlessly changing legacy materials, we assume the layer was explicitly deleted, so we keep it deleted here
 			// If desired, the relink functionality in the editor should allow it to be brought back
-			ResolvedLayers.EditorOnly.DeletedParentLayerGuids.Add(ParentLayerGuid);
+			ResolvedLayers.DeletedParentLayerGuids.Add(ParentLayerGuid);
 			continue;
 		}
 
@@ -14618,17 +14408,17 @@ bool FMaterialLayersFunctions::ResolveParent(const FMaterialLayersFunctionsRunti
 		}
 
 		ParentLayerIndices.Insert(ParentLayerIndex, InsertLayerIndex + 1);
-		ResolvedLayers.InsertLayerCopy(ParentRuntime, ParentEditorOnly, ParentLayerIndex, EMaterialLayerLinkState::LinkedToParent, InsertLayerIndex + 1);
+		ResolvedLayers.InsertLayerCopy(Parent, ParentLayerIndex, EMaterialLayerLinkState::LinkedToParent, InsertLayerIndex + 1);
 	}
 
-	bool bUpdatedLayerIndices = Runtime.Layers.Num() != ResolvedLayers.Layers.Num() ||
-		EditorOnly.DeletedParentLayerGuids.Num() != ResolvedLayers.EditorOnly.DeletedParentLayerGuids.Num();
+	bool bUpdatedLayerIndices = Layers.Num() != ResolvedLayers.Layers.Num() ||
+		DeletedParentLayerGuids.Num() != ResolvedLayers.DeletedParentLayerGuids.Num();
 
-	OutRemapLayerIndices.SetNumUninitialized(Runtime.Layers.Num());
-	for (int32 PrevLayerIndex = 0; PrevLayerIndex < Runtime.Layers.Num(); ++PrevLayerIndex)
+	OutRemapLayerIndices.SetNumUninitialized(Layers.Num());
+	for (int32 PrevLayerIndex = 0; PrevLayerIndex < Layers.Num(); ++PrevLayerIndex)
 	{
-		const FGuid& LayerGuid = EditorOnly.LayerGuids[PrevLayerIndex];
-		const int32 ResolvedLayerIndex = ResolvedLayers.EditorOnly.LayerGuids.Find(LayerGuid);
+		const FGuid& LayerGuid = LayerGuids[PrevLayerIndex];
+		const int32 ResolvedLayerIndex = ResolvedLayers.LayerGuids.Find(LayerGuid);
 		OutRemapLayerIndices[PrevLayerIndex] = ResolvedLayerIndex;
 
 		if (PrevLayerIndex != ResolvedLayerIndex)
@@ -14637,24 +14427,10 @@ bool FMaterialLayersFunctions::ResolveParent(const FMaterialLayersFunctionsRunti
 		}
 	}
 
-	Runtime = MoveTemp(static_cast<FMaterialLayersFunctionsRuntimeData&>(ResolvedLayers));
-	EditorOnly = MoveTemp(ResolvedLayers.EditorOnly);
+	*this = MoveTemp(ResolvedLayers);
 
 	return bUpdatedLayerIndices;
 }
-
-void FMaterialLayersFunctions::Validate(const FMaterialLayersFunctionsRuntimeData& Runtime, const FMaterialLayersFunctionsEditorOnlyData& EditorOnly)
-{
-	if (Runtime.Layers.Num() > 0)
-	{
-		check(Runtime.Blends.Num() == Runtime.Layers.Num() - 1);
-		check(Runtime.Layers.Num() == EditorOnly.LayerStates.Num());
-		check(Runtime.Layers.Num() == EditorOnly.LayerNames.Num());
-		check(Runtime.Layers.Num() == EditorOnly.LayerGuids.Num());
-		check(Runtime.Layers.Num() == EditorOnly.LayerLinkStates.Num());
-	}
-}
-
 #endif // WITH_EDITOR
 
 ///////////////////////////////////////////////////////////////////////////////
@@ -14985,8 +14761,8 @@ void UMaterialExpressionMaterialFunctionCall::GetExpressionToolTip(TArray<FStrin
 {
 	if (MaterialFunction)
 	{
-		const FString& Description = MaterialFunction->GetDescription();
-		ConvertToMultilineToolTip(Description, 40, OutToolTip);
+		const FString* Description = MaterialFunction->GetDescription();
+		ConvertToMultilineToolTip(Description ? *Description : TEXT(""), 40, OutToolTip);
 	}
 }
 
@@ -15067,12 +14843,12 @@ bool UMaterialExpressionMaterialFunctionCall::SetMaterialFunctionEx(
 			}
 
 			// Fixup any references that the material or material inputs had to the function's outputs, maintaining links with the same output name
-			FixupReferencingExpressions(FunctionOutputs, OriginalOutputs, Material->GetExpressions(), MaterialInputs, true);
+			FixupReferencingExpressions(FunctionOutputs, OriginalOutputs, Material->Expressions, MaterialInputs, true);
 		}
 		else if (Function)
 		{
 			// Fixup any references that the material function had to the function's outputs, maintaining links with the same output name
-			FixupReferencingExpressions(FunctionOutputs, OriginalOutputs, Function->GetExpressions(), MaterialInputs, true);
+			FixupReferencingExpressions(FunctionOutputs, OriginalOutputs, Function->FunctionExpressions, MaterialInputs, true);
 		}
 	}
 
@@ -15138,12 +14914,12 @@ void UMaterialExpressionMaterialFunctionCall::UpdateFromFunctionResource(bool bR
 			}
 			
 			// Fixup any references that the material or material inputs had to the function's outputs
-			FixupReferencingExpressions(FunctionOutputs, OriginalOutputs, Material->GetExpressions(), MaterialInputs, false);
+			FixupReferencingExpressions(FunctionOutputs, OriginalOutputs, Material->Expressions, MaterialInputs, false);
 		}
 		else if (Function)
 		{
 			// Fixup any references that the material function had to the function's outputs
-			FixupReferencingExpressions(FunctionOutputs, OriginalOutputs, Function->GetExpressions(), MaterialInputs, false);
+			FixupReferencingExpressions(FunctionOutputs, OriginalOutputs, Function->FunctionExpressions, MaterialInputs, false);
 		}
 	}
 
@@ -15225,7 +15001,7 @@ static void FixupReferencingInputs(
 void UMaterialExpressionMaterialFunctionCall::FixupReferencingExpressions(
 	const TArray<FFunctionExpressionOutput>& NewOutputs,
 	const TArray<FFunctionExpressionOutput>& OriginalOutputs,
-	TConstArrayView<TObjectPtr<UMaterialExpression>> Expressions,
+	TArray<UMaterialExpression*>& Expressions, 
 	TArray<FExpressionInput*>& MaterialInputs,
 	bool bMatchByName)
 {
@@ -15374,9 +15150,9 @@ void UMaterialExpressionFunctionInput::PostEditChangeProperty(FPropertyChangedEv
 	{
 		if (Material)
 		{
-			for (UMaterialExpression* Expression : Material->GetExpressions())
+			for (int32 ExpressionIndex = 0; ExpressionIndex < Material->Expressions.Num(); ExpressionIndex++)
 			{
-				UMaterialExpressionFunctionInput* OtherFunctionInput = Cast<UMaterialExpressionFunctionInput>(Expression);
+				UMaterialExpressionFunctionInput* OtherFunctionInput = Cast<UMaterialExpressionFunctionInput>(Material->Expressions[ExpressionIndex]);
 				if (OtherFunctionInput && OtherFunctionInput != this && OtherFunctionInput->InputName == InputName)
 				{
 					FMessageDialog::Open( EAppMsgType::Ok, NSLOCTEXT("UnrealEd", "Error_InputNamesMustBeUnique", "Function input names must be unique"));
@@ -15576,9 +15352,9 @@ void UMaterialExpressionFunctionInput::ValidateName()
 			}
 
 			bResultNameIndexValid = true;
-			for (UMaterialExpression* Expression : Material->GetExpressions())
+			for (int32 ExpressionIndex = 0; ExpressionIndex < Material->Expressions.Num(); ExpressionIndex++)
 			{
-				UMaterialExpressionFunctionInput* OtherFunctionInput = Cast<UMaterialExpressionFunctionInput>(Expression);
+				UMaterialExpressionFunctionInput* OtherFunctionInput = Cast<UMaterialExpressionFunctionInput>(Material->Expressions[ExpressionIndex]);
 				if (OtherFunctionInput && OtherFunctionInput != this && OtherFunctionInput->InputName == PotentialInputName)
 				{
 					bResultNameIndexValid = false;
@@ -15715,9 +15491,9 @@ void UMaterialExpressionFunctionOutput::PostEditChangeProperty(FPropertyChangedE
 	{
 		if (Material)
 		{
-			for (UMaterialExpression* Expression : Material->GetExpressions())
+			for (int32 ExpressionIndex = 0; ExpressionIndex < Material->Expressions.Num(); ExpressionIndex++)
 			{
-				UMaterialExpressionFunctionOutput* OtherFunctionOutput = Cast<UMaterialExpressionFunctionOutput>(Expression);
+				UMaterialExpressionFunctionOutput* OtherFunctionOutput = Cast<UMaterialExpressionFunctionOutput>(Material->Expressions[ExpressionIndex]);
 				if (OtherFunctionOutput && OtherFunctionOutput != this && OtherFunctionOutput->OutputName == OutputName)
 				{
 					FMessageDialog::Open( EAppMsgType::Ok, NSLOCTEXT("UnrealEd", "Error_OutputNamesMustBeUnique", "Function output names must be unique"));
@@ -15783,9 +15559,9 @@ void UMaterialExpressionFunctionOutput::ValidateName()
 			}
 
 			bResultNameIndexValid = true;
-			for (UMaterialExpression* Expression : Material->GetExpressions())
+			for (int32 ExpressionIndex = 0; ExpressionIndex < Material->Expressions.Num(); ExpressionIndex++)
 			{
-				UMaterialExpressionFunctionOutput* OtherFunctionOutput = Cast<UMaterialExpressionFunctionOutput>(Expression);
+				UMaterialExpressionFunctionOutput* OtherFunctionOutput = Cast<UMaterialExpressionFunctionOutput>(Material->Expressions[ExpressionIndex]);
 				if (OtherFunctionOutput && OtherFunctionOutput != this && OtherFunctionOutput->OutputName == PotentialOutputName)
 				{
 					bResultNameIndexValid = false;
@@ -16753,17 +16529,30 @@ UMaterialExpressionNamedRerouteBase::UMaterialExpressionNamedRerouteBase(const F
 
 }
 
+UMaterialExpressionNamedRerouteDeclaration* UMaterialExpressionNamedRerouteBase::FindDeclarationInArray(const FGuid& VariableGuid, const TArray<UMaterialExpression*>& Expressions) const
+{
+	for (UMaterialExpression* Expression : Expressions)
+	{
+		auto* Declaration = Cast<UMaterialExpressionNamedRerouteDeclaration>(Expression);
+		if (Declaration && this != Declaration && Declaration->VariableGuid == VariableGuid)
+		{
+			return Declaration;
+		}
+	}
+	return nullptr;
+}
+
 UMaterialExpressionNamedRerouteDeclaration* UMaterialExpressionNamedRerouteBase::FindDeclarationInMaterial(const FGuid& VariableGuid) const
 {
 	UMaterialExpressionNamedRerouteDeclaration* Declaration = nullptr;
 #if WITH_EDITORONLY_DATA
 	if (Material)
 	{
-		Declaration = FindDeclarationInArray(VariableGuid, Material->GetExpressions());
+		Declaration = FindDeclarationInArray(VariableGuid, Material->Expressions);
 	}
 	else if (Function) // Material should always be valid, but just in case also check Function
 	{
-		Declaration = FindDeclarationInArray(VariableGuid, Function->GetExpressions());
+		Declaration = FindDeclarationInArray(VariableGuid, Function->FunctionExpressions);
 	}
 #endif // WITH_EDITORONLY_DATA
 	return Declaration;
@@ -16874,7 +16663,7 @@ void UMaterialExpressionNamedRerouteDeclaration::SetEditableName(const FString&
 	// Refresh usage names
 	if (Material || Function)
 	{
-		auto Expressions = Material ? Material->GetExpressions() : Function->GetExpressions();
+		auto& Expressions = Material ? Material->Expressions : Function->FunctionExpressions;
 		for (UMaterialExpression* Expression : Expressions)
 		{
 			auto* Usage = Cast<UMaterialExpressionNamedRerouteUsage>(Expression);
@@ -16945,7 +16734,7 @@ void UMaterialExpressionNamedRerouteDeclaration::MakeNameUnique()
 #if WITH_EDITORONLY_DATA
 	if (Material || Function)
 	{
-		auto Expressions = Material ? Material->GetExpressions() : Function->GetExpressions();
+		auto& Expressions = Material ? Material->Expressions : Function->FunctionExpressions;
 
 		int32 NameIndex = 1;
 		bool bResultNameIndexValid = true;
@@ -19915,7 +19704,6 @@ bool UMaterialExpressionCurveAtlasRowParameter::GetParameterValue(FMaterialParam
 	OutMeta.ExpressionGuid = ExpressionGUID;
 	OutMeta.Group = Group;
 	OutMeta.SortPriority = SortPriority;
-	OutMeta.AssetPath = GetAssetPathName();
 	return true;
 }
 
@@ -23311,7 +23099,7 @@ UMaterialExpressionExecEnd::UMaterialExpressionExecEnd(const FObjectInitializer&
 #if WITH_EDITOR
 int32 UMaterialExpressionExecEnd::Compile(class FMaterialCompiler* Compiler, int32 OutputIndex)
 {
-	const int32 Attributes = Material->GetEditorOnlyData()->MaterialAttributes.Compile(Compiler);
+	const int32 Attributes = Material->MaterialAttributes.Compile(Compiler);
 	return Compiler->ReturnMaterialAttributes(Attributes);
 }
 #endif
diff --git a/Engine/Source/Runtime/Engine/Private/Materials/MaterialHLSLGenerator.cpp b/Engine/Source/Runtime/Engine/Private/Materials/MaterialHLSLGenerator.cpp
--- a/Engine/Source/Runtime/Engine/Private/Materials/MaterialHLSLGenerator.cpp
+++ b/Engine/Source/Runtime/Engine/Private/Materials/MaterialHLSLGenerator.cpp
@@ -78,7 +78,7 @@ bool FMaterialHLSLGenerator::InternalGenerate()
 	bool bResult = false;
 	if (TargetMaterial->IsUsingControlFlow())
 	{
-		UMaterialExpression* BaseExpression = TargetMaterial->GetExpressionExecBegin();
+		UMaterialExpression* BaseExpression = TargetMaterial->ExpressionExecBegin;
 		if (!BaseExpression)
 		{
 			bResult = Error(TEXT("Missing ExpressionExecBegin"));
@@ -671,9 +671,9 @@ const UE::HLSLTree::FExpression* FMaterialHLSLGenerator::GenerateFunctionCall(UE
 
 		if (HLSLFunction)
 		{
-			UMaterialFunction* BaseMaterialFunction = MaterialFunction->GetBaseFunction();
+			UMaterialFunction* BaseMaterialFunction = Cast<UMaterialFunction>(MaterialFunction->GetBaseFunction());
 			FunctionCallStack.Add(FunctionCall);
-			GenerateStatements(HLSLFunction->GetRootScope(), BaseMaterialFunction->GetExpressionExecBegin());
+			GenerateStatements(HLSLFunction->GetRootScope(), BaseMaterialFunction->ExpressionExecBegin);
 			verify(FunctionCallStack.Pop() == FunctionCall);
 			check(FunctionCall->bGeneratedResult);
 		}
diff --git a/Engine/Source/Runtime/Engine/Private/Materials/MaterialHLSLTree.cpp b/Engine/Source/Runtime/Engine/Private/Materials/MaterialHLSLTree.cpp
--- a/Engine/Source/Runtime/Engine/Private/Materials/MaterialHLSLTree.cpp
+++ b/Engine/Source/Runtime/Engine/Private/Materials/MaterialHLSLTree.cpp
@@ -454,7 +454,7 @@ bool FExpressionParameter::PrepareValue(FEmitContext& Context, FEmitScope& Scope
 	{
 		if (!ParameterInfo.Name.IsNone())
 		{
-			EmitData.CachedExpressionData->AddParameter(ParameterInfo, ParameterMeta);
+			EmitData.CachedExpressionData->Parameters.AddParameter(ParameterInfo, ParameterMeta);
 		}
 
 		UObject* ReferencedTexture = ParameterMeta.Value.AsTextureObject();
@@ -507,7 +507,7 @@ void FExpressionParameter::EmitValuePreshader(FEmitContext& Context, FEmitScope&
 			switch (ParameterType)
 			{
 			case EMaterialParameterType::StaticSwitch:
-				for (const FStaticSwitchParameter& Parameter : EmitMaterialData.StaticParameters->EditorOnly.StaticSwitchParameters)
+				for (const FStaticSwitchParameter& Parameter : EmitMaterialData.StaticParameters->StaticSwitchParameters)
 				{
 					if (Parameter.ParameterInfo == ParameterInfo)
 					{
@@ -517,7 +517,7 @@ void FExpressionParameter::EmitValuePreshader(FEmitContext& Context, FEmitScope&
 				}
 				break;
 			case EMaterialParameterType::StaticComponentMask:
-				for (const FStaticComponentMaskParameter& Parameter : EmitMaterialData.StaticParameters->EditorOnly.StaticComponentMaskParameters)
+				for (const FStaticComponentMaskParameter& Parameter : EmitMaterialData.StaticParameters->StaticComponentMaskParameters)
 				{
 					if (Parameter.ParameterInfo == ParameterInfo)
 					{
@@ -1119,9 +1119,7 @@ bool FExpressionMaterialLayers::PrepareValue(FEmitContext& Context, FEmitScope&
 			{
 				// TODO(?) - Layers for MIs are currently duplicated here and in FStaticParameterSet
 				EmitMaterialData.CachedExpressionData->bHasMaterialLayers = true;
-				EmitMaterialData.CachedExpressionData->MaterialLayers = MaterialLayers.GetRuntime();
-				EmitMaterialData.CachedExpressionData->EditorOnlyData->MaterialLayers = MaterialLayers.EditorOnly;
-				FMaterialLayersFunctions::Validate(EmitMaterialData.CachedExpressionData->MaterialLayers, EmitMaterialData.CachedExpressionData->EditorOnlyData->MaterialLayers);
+				EmitMaterialData.CachedExpressionData->MaterialLayers = *MaterialLayers;
 			}
 		}
 	}
diff --git a/Engine/Source/Runtime/Engine/Private/Materials/MaterialInstance.cpp b/Engine/Source/Runtime/Engine/Private/Materials/MaterialInstance.cpp
--- a/Engine/Source/Runtime/Engine/Private/Materials/MaterialInstance.cpp
+++ b/Engine/Source/Runtime/Engine/Private/Materials/MaterialInstance.cpp
@@ -482,35 +482,27 @@ void UMaterialInstance::SwapLayerParameterIndices(int32 OriginalIndex, int32 New
 {
 	if (OriginalIndex != NewIndex)
 	{
-		UMaterialInstanceEditorOnlyData* EditorOnly = GetEditorOnlyData();
 		SwapLayerParameterIndicesArray(ScalarParameterValues, OriginalIndex, NewIndex);
 		SwapLayerParameterIndicesArray(VectorParameterValues, OriginalIndex, NewIndex);
 		SwapLayerParameterIndicesArray(DoubleVectorParameterValues, OriginalIndex, NewIndex);
 		SwapLayerParameterIndicesArray(TextureParameterValues, OriginalIndex, NewIndex);
 		SwapLayerParameterIndicesArray(RuntimeVirtualTextureParameterValues, OriginalIndex, NewIndex);
 		SwapLayerParameterIndicesArray(FontParameterValues, OriginalIndex, NewIndex);
-		if (EditorOnly)
-		{
-			SwapLayerParameterIndicesArray(EditorOnly->StaticParameters.StaticSwitchParameters, OriginalIndex, NewIndex);
-			SwapLayerParameterIndicesArray(EditorOnly->StaticParameters.StaticComponentMaskParameters, OriginalIndex, NewIndex);
-		}
+		SwapLayerParameterIndicesArray(StaticParameters.StaticSwitchParameters, OriginalIndex, NewIndex);
+		SwapLayerParameterIndicesArray(StaticParameters.StaticComponentMaskParameters, OriginalIndex, NewIndex);
 	}
 }
 
 void UMaterialInstance::RemoveLayerParameterIndex(int32 Index)
 {
-	UMaterialInstanceEditorOnlyData* EditorOnly = GetEditorOnlyData();
 	RemoveLayerParameterIndicesArray(ScalarParameterValues, Index);
 	RemoveLayerParameterIndicesArray(VectorParameterValues, Index);
 	RemoveLayerParameterIndicesArray(DoubleVectorParameterValues, Index);
 	RemoveLayerParameterIndicesArray(TextureParameterValues, Index);
 	RemoveLayerParameterIndicesArray(RuntimeVirtualTextureParameterValues, Index);
 	RemoveLayerParameterIndicesArray(FontParameterValues, Index);
-	if (EditorOnly)
-	{
-		RemoveLayerParameterIndicesArray(EditorOnly->StaticParameters.StaticSwitchParameters, Index);
-		RemoveLayerParameterIndicesArray(EditorOnly->StaticParameters.StaticComponentMaskParameters, Index);
-	}
+	RemoveLayerParameterIndicesArray(StaticParameters.StaticSwitchParameters, Index);
+	RemoveLayerParameterIndicesArray(StaticParameters.StaticComponentMaskParameters, Index);
 }
 #endif // WITH_EDITOR
 
@@ -519,8 +511,7 @@ bool UMaterialInstance::UpdateParameters()
 	bool bDirty = false;
 
 #if WITH_EDITOR
-	UMaterialInstanceEditorOnlyData* EditorOnly = GetEditorOnlyData();
-	if(IsTemplate(RF_ClassDefaultObject)==false && EditorOnly)
+	if(IsTemplate(RF_ClassDefaultObject)==false)
 	{
 		// Get a pointer to the parent material.
 		UMaterial* ParentMaterial = NULL;
@@ -559,37 +550,27 @@ bool UMaterialInstance::UpdateParameters()
 			bDirty = UpdateParameterSet<FFontParameterValue, UMaterialExpressionFontSampleParameter>(FontParameterValues, ParentMaterial) || bDirty;
 
 			// Static switch parameters
-			bDirty = UpdateParameterSet<FStaticSwitchParameter, UMaterialExpressionStaticBoolParameter>(EditorOnly->StaticParameters.StaticSwitchParameters, ParentMaterial) || bDirty;
+			bDirty = UpdateParameterSet<FStaticSwitchParameter, UMaterialExpressionStaticBoolParameter>(StaticParameters.StaticSwitchParameters, ParentMaterial) || bDirty;
 
 			// Static component mask parameters
-			bDirty = UpdateParameterSet<FStaticComponentMaskParameter, UMaterialExpressionStaticComponentMaskParameter>(EditorOnly->StaticParameters.StaticComponentMaskParameters, ParentMaterial) || bDirty;
+			bDirty = UpdateParameterSet<FStaticComponentMaskParameter, UMaterialExpressionStaticComponentMaskParameter>(StaticParameters.StaticComponentMaskParameters, ParentMaterial) || bDirty;
 
 			// Custom parameters
 			PRAGMA_DISABLE_DEPRECATION_WARNINGS
-			if (CustomParameterSetUpdaters.Num() > 0)
+			for (const auto& CustomParameterSetUpdater : CustomParameterSetUpdaters)
 			{
-				FStaticParameterSet LocalStaticParameters = GetStaticParameters();
-				for (const auto& CustomParameterSetUpdater : CustomParameterSetUpdaters)
-				{
-					bDirty |= CustomParameterSetUpdater.Execute(LocalStaticParameters, ParentMaterial);
-				}
-				StaticParametersRuntime = LocalStaticParameters.GetRuntime();
-				EditorOnly->StaticParameters = LocalStaticParameters.EditorOnly;
+				bDirty |= CustomParameterSetUpdater.Execute(StaticParameters, ParentMaterial);
 			}
 			PRAGMA_ENABLE_DEPRECATION_WARNINGS
 		}
 
-		if (StaticParametersRuntime.bHasMaterialLayers && Parent)
+		if (StaticParameters.bHasMaterialLayers && Parent)
 		{
 			FMaterialLayersFunctions ParentLayers;
 			if (Parent->GetMaterialLayers(ParentLayers))
 			{
 				TArray<int32> RemapLayerIndices;
-				if (FMaterialLayersFunctions::ResolveParent(ParentLayers,
-					ParentLayers.EditorOnly,
-					StaticParametersRuntime.MaterialLayers,
-					EditorOnly->StaticParameters.MaterialLayers,
-					RemapLayerIndices))
+				if (StaticParameters.MaterialLayers.ResolveParent(ParentLayers, RemapLayerIndices))
 				{
 					RemapLayerParameterIndicesArray(ScalarParameterValues, RemapLayerIndices);
 					RemapLayerParameterIndicesArray(VectorParameterValues, RemapLayerIndices);
@@ -597,8 +578,8 @@ bool UMaterialInstance::UpdateParameters()
 					RemapLayerParameterIndicesArray(TextureParameterValues, RemapLayerIndices);
 					RemapLayerParameterIndicesArray(RuntimeVirtualTextureParameterValues, RemapLayerIndices);
 					RemapLayerParameterIndicesArray(FontParameterValues, RemapLayerIndices);
-					RemapLayerParameterIndicesArray(EditorOnly->StaticParameters.StaticSwitchParameters, RemapLayerIndices);
-					RemapLayerParameterIndicesArray(EditorOnly->StaticParameters.StaticComponentMaskParameters, RemapLayerIndices);
+					RemapLayerParameterIndicesArray(StaticParameters.StaticSwitchParameters, RemapLayerIndices);
+					RemapLayerParameterIndicesArray(StaticParameters.StaticComponentMaskParameters, RemapLayerIndices);
 					bDirty = true;
 				}
 			}
@@ -874,12 +855,8 @@ bool UMaterialInstance::GetParameterOverrideValue(EMaterialParameterType Type, c
 	case EMaterialParameterType::RuntimeVirtualTexture: bResult = GameThread_GetParameterValue(RuntimeVirtualTextureParameterValues, ParameterInfo, OutResult); break;
 	case EMaterialParameterType::Font: bResult = GameThread_GetParameterValue(FontParameterValues, ParameterInfo, OutResult); break;
 #if WITH_EDITORONLY_DATA
-	case EMaterialParameterType::StaticSwitch:
-		bResult = GameThread_GetParameterValue(GetEditorOnlyData()->StaticParameters.StaticSwitchParameters, ParameterInfo, OutResult);
-		break;
-	case EMaterialParameterType::StaticComponentMask:
-		bResult = GameThread_GetParameterValue(GetEditorOnlyData()->StaticParameters.StaticComponentMaskParameters, ParameterInfo, OutResult);
-		break;
+	case EMaterialParameterType::StaticSwitch: bResult = GameThread_GetParameterValue(StaticParameters.StaticSwitchParameters, ParameterInfo, OutResult); break;
+	case EMaterialParameterType::StaticComponentMask: bResult = GameThread_GetParameterValue(StaticParameters.StaticComponentMaskParameters, ParameterInfo, OutResult); break;
 #endif // WITH_EDITORONLY_DATA
 	default: checkNoEntry(); break;
 	}
@@ -894,7 +871,7 @@ bool UMaterialInstance::GetParameterValue(EMaterialParameterType Type, const FMe
 	bool bResult = false;
 	if (EnumHasAnyFlags(Flags, EMaterialGetParameterValueFlags::CheckNonOverrides))
 	{
-		bResult = InstanceChain.GetCachedExpressionData().GetParameterValue(Type, ParameterInfo, OutResult);
+		bResult = InstanceChain.GetCachedExpressionData().Parameters.GetParameterValue(Type, ParameterInfo, OutResult);
 	}
 
 	const bool bCheckInstanceOverrides = EnumHasAnyFlags(Flags, EMaterialGetParameterValueFlags::CheckInstanceOverrides);
@@ -1704,10 +1681,7 @@ void UMaterialInstance::SetStaticSwitchParameterValueEditorOnly(const FMaterialP
 {
 	check(GIsEditor || IsRunningCommandlet());
 
-	UMaterialInstanceEditorOnlyData* EditorOnly = GetEditorOnlyData();
-	check(EditorOnly);
-
-	for (FStaticSwitchParameter& StaticSwitches : EditorOnly->StaticParameters.StaticSwitchParameters)
+	for (FStaticSwitchParameter& StaticSwitches : StaticParameters.StaticSwitchParameters)
 	{
 		if (StaticSwitches.ParameterInfo == ParameterInfo)
 		{
@@ -1717,7 +1691,7 @@ void UMaterialInstance::SetStaticSwitchParameterValueEditorOnly(const FMaterialP
 		}
 	}
 
-	new(EditorOnly->StaticParameters.StaticSwitchParameters) FStaticSwitchParameter(ParameterInfo, Value, true, FGuid());
+	new(StaticParameters.StaticSwitchParameters) FStaticSwitchParameter(ParameterInfo, Value, true, FGuid());
 }
 
 void UMaterialInstance::GetStaticParameterValues(FStaticParameterSet& OutStaticParameters)
@@ -1745,21 +1719,9 @@ void UMaterialInstance::GetStaticParameterValues(FStaticParameterSet& OutStaticP
 		}
 	}
 
-	UMaterialInstanceEditorOnlyData* EditorOnly = GetEditorOnlyData();
-	if (EditorOnly)
-	{
-		OutStaticParameters.EditorOnly.TerrainLayerWeightParameters = EditorOnly->StaticParameters.TerrainLayerWeightParameters;
-	}
+	OutStaticParameters.TerrainLayerWeightParameters = StaticParameters.TerrainLayerWeightParameters;
 
-	FMaterialLayersFunctions MaterialLayers;
-	OutStaticParameters.bHasMaterialLayers = GetMaterialLayers(MaterialLayers);
-	if (OutStaticParameters.bHasMaterialLayers)
-	{
-		OutStaticParameters.MaterialLayers = MoveTemp(MaterialLayers.GetRuntime());
-		OutStaticParameters.EditorOnly.MaterialLayers = MoveTemp(MaterialLayers.EditorOnly);
-	}
-
-	OutStaticParameters.Validate();
+	OutStaticParameters.bHasMaterialLayers = GetMaterialLayers(OutStaticParameters.MaterialLayers);
 
 	// Custom parameters.
 	PRAGMA_DISABLE_DEPRECATION_WARNINGS
@@ -1797,7 +1759,7 @@ void UMaterialInstance::GetAllParametersOfType(EMaterialParameterType Type, TMap
 	GetMaterialInheritanceChain(InstanceChain);
 
 	OutParameters.Reset();
-	InstanceChain.GetCachedExpressionData().GetAllParametersOfType(Type, OutParameters);
+	InstanceChain.GetCachedExpressionData().Parameters.GetAllParametersOfType(Type, OutParameters);
 
 	TArray<int32, TInlineAllocator<16>> LayerIndexRemap;
 	LayerIndexRemap.Empty(GetCachedInstanceData().ParentLayerIndexRemap.Num());
@@ -1813,6 +1775,7 @@ void UMaterialInstance::GetAllParametersOfType(EMaterialParameterType Type, TMap
 	for (int32 Index = 0; Index < InstanceChain.MaterialInstances.Num(); ++Index)
 	{
 		const UMaterialInstance* Instance = InstanceChain.MaterialInstances[Index];
+
 		const bool bSetOverride = (Index == 0); // Only set the override flag for parameters overriden by the current material (always at slot0)
 		switch (Type)
 		{
@@ -1823,12 +1786,8 @@ void UMaterialInstance::GetAllParametersOfType(EMaterialParameterType Type, TMap
 		case EMaterialParameterType::RuntimeVirtualTexture: GameThread_ApplyParameterOverrides(Instance->RuntimeVirtualTextureParameterValues, LayerIndexRemap, bSetOverride, OverridenParameters, OutParameters); break;
 		case EMaterialParameterType::Font: GameThread_ApplyParameterOverrides(Instance->FontParameterValues, LayerIndexRemap, bSetOverride, OverridenParameters, OutParameters); break;
 #if WITH_EDITORONLY_DATA
-		case EMaterialParameterType::StaticSwitch:
-			GameThread_ApplyParameterOverrides(Instance->GetEditorOnlyData()->StaticParameters.StaticSwitchParameters, LayerIndexRemap, bSetOverride, OverridenParameters, OutParameters);
-			break;
-		case EMaterialParameterType::StaticComponentMask:
-			GameThread_ApplyParameterOverrides(Instance->GetEditorOnlyData()->StaticParameters.StaticComponentMaskParameters, LayerIndexRemap, bSetOverride, OverridenParameters, OutParameters);
-			break;
+		case EMaterialParameterType::StaticSwitch: GameThread_ApplyParameterOverrides(Instance->StaticParameters.StaticSwitchParameters, LayerIndexRemap, bSetOverride, OverridenParameters, OutParameters); break;
+		case EMaterialParameterType::StaticComponentMask: GameThread_ApplyParameterOverrides(Instance->StaticParameters.StaticComponentMaskParameters, LayerIndexRemap, bSetOverride, OverridenParameters, OutParameters); break;
 #endif // WITH_EDITORONLY_DATA
 		default: checkNoEntry();
 		}
@@ -1846,9 +1805,9 @@ bool UMaterialInstance::IterateDependentFunctions(TFunctionRef<bool(UMaterialFun
 {
 	// Important that local function references are listed first so that traversing for a parameter
 	// value we always hit the highest material in the hierarchy that can give us a valid value
-	if (StaticParametersRuntime.bHasMaterialLayers)
+	if (StaticParameters.bHasMaterialLayers)
 	{
-		for (UMaterialFunctionInterface* Layer : StaticParametersRuntime.MaterialLayers.Layers)
+		for (UMaterialFunctionInterface* Layer : StaticParameters.MaterialLayers.Layers)
 		{
 			if (Layer)
 			{
@@ -1863,7 +1822,7 @@ bool UMaterialInstance::IterateDependentFunctions(TFunctionRef<bool(UMaterialFun
 			}
 		}
 
-		for (UMaterialFunctionInterface* Blend : StaticParametersRuntime.MaterialLayers.Blends)
+		for (UMaterialFunctionInterface* Blend : StaticParameters.MaterialLayers.Blends)
 		{
 			if (Blend)
 			{
@@ -1904,10 +1863,8 @@ void UMaterialInstance::InitStaticPermutation(EMaterialShaderPrecompileMode Prec
 {
 	UpdateOverridableBaseProperties();
 
-#if WITH_EDITORONLY_DATA
 	// Update bHasStaticPermutationResource in case the parent was not found
-	bHasStaticPermutationResource = (HasStaticParameters() || HasOverridenBaseProperties()) && Parent;
-#endif // WITH_EDITORONLY_DATA
+	bHasStaticPermutationResource = (!StaticParameters.IsEmpty() || HasOverridenBaseProperties()) && Parent;
 
 	FMaterialResourceDeferredDeletionArray ResourcesToFree;
 
@@ -2229,17 +2186,9 @@ void UMaterialInstance::CacheGivenTypesForCooking(EShaderPlatform ShaderPlatform
 
 bool UMaterialInstance::GetMaterialLayers(FMaterialLayersFunctions& OutLayers, TMicRecursionGuard RecursionGuard) const
 {
-	if (StaticParametersRuntime.bHasMaterialLayers)
+	if (StaticParameters.bHasMaterialLayers)
 	{
-		OutLayers.GetRuntime() = StaticParametersRuntime.MaterialLayers;
-#if WITH_EDITORONLY_DATA
-		const UMaterialInstanceEditorOnlyData* EditorOnly = GetEditorOnlyData();
-		if (EditorOnly)
-		{
-			OutLayers.EditorOnly = EditorOnly->StaticParameters.MaterialLayers;
-		}
-		OutLayers.Validate();
-#endif // WITH_EDITORONLY_DATA
+		OutLayers = StaticParameters.MaterialLayers;
 		return true;
 	}
 
@@ -2290,13 +2239,8 @@ bool UMaterialInstance::IsComplete() const
 #if WITH_EDITOR
 bool UMaterialInstance::SetMaterialLayers(const FMaterialLayersFunctions& LayersValue)
 {
-	UMaterialInstanceEditorOnlyData* EditorOnly = GetEditorOnlyData();
-	check(EditorOnly);
-
 	bool bUpdatedLayers = false;
-	if (!StaticParametersRuntime.bHasMaterialLayers ||
-		StaticParametersRuntime.MaterialLayers != LayersValue.GetRuntime() ||
-		EditorOnly->StaticParameters.MaterialLayers != LayersValue.EditorOnly)
+	if (!StaticParameters.bHasMaterialLayers || StaticParameters.MaterialLayers != LayersValue)
 	{
 		bool bMatchesParentLayers = false;
 		if (Parent)
@@ -2310,19 +2254,16 @@ bool UMaterialInstance::SetMaterialLayers(const FMaterialLayersFunctions& Layers
 
 		if (bMatchesParentLayers)
 		{
-			bUpdatedLayers = StaticParametersRuntime.bHasMaterialLayers; // if we previously had layers, but are now clearing them to match parent
-			StaticParametersRuntime.bHasMaterialLayers = false;
-			StaticParametersRuntime.MaterialLayers.Empty();
-			EditorOnly->StaticParameters.MaterialLayers.Empty();
+			bUpdatedLayers = StaticParameters.bHasMaterialLayers; // if we previously had layers, but are now clearing them to match parent
+			StaticParameters.bHasMaterialLayers = false;
+			StaticParameters.MaterialLayers.Empty();
 		}
 		else
 		{
 			bUpdatedLayers = true;
-			StaticParametersRuntime.bHasMaterialLayers = true;
-			StaticParametersRuntime.MaterialLayers = LayersValue.GetRuntime();
-			EditorOnly->StaticParameters.MaterialLayers = LayersValue.EditorOnly;
+			StaticParameters.bHasMaterialLayers = true;
+			StaticParameters.MaterialLayers = LayersValue;
 		}
-		FStaticParameterSet::Validate(StaticParametersRuntime, EditorOnly->StaticParameters);
 	}
 	return bUpdatedLayers;
 }
@@ -2425,9 +2366,8 @@ void UMaterialInstance::Serialize(FArchive& Ar)
 #endif
 
 	Super::Serialize(Ar);
-
+		
 #if WITH_EDITOR
-	UMaterialInstanceEditorOnlyData* EditorOnly = GetEditorOnlyData();
 	if (Ar.CustomVer(FRenderingObjectVersion::GUID) < FRenderingObjectVersion::MaterialAttributeLayerParameters)
 	{
 		// Material attribute layers parameter refactor fix-up
@@ -2451,20 +2391,13 @@ void UMaterialInstance::Serialize(FArchive& Ar)
 
 	if (Ar.CustomVer(FUE5ReleaseStreamObjectVersion::GUID) < FUE5ReleaseStreamObjectVersion::MaterialLayerStacksAreNotParameters)
 	{
-		StaticParameters_DEPRECATED.UpdateLegacyMaterialLayersData();
+		StaticParameters.UpdateLegacyMaterialLayersData();
 	}
 
 	if (Ar.CustomVer(FFortniteMainBranchObjectVersion::GUID) < FFortniteMainBranchObjectVersion::TerrainLayerWeightsAreNotParameters)
 	{
-		StaticParameters_DEPRECATED.UpdateLegacyTerrainLayerWeightData();
-	}
-
-	if (!StaticParameters_DEPRECATED.IsEmpty())
-	{
-		StaticParametersRuntime = MoveTemp(StaticParameters_DEPRECATED.GetRuntime());
-		EditorOnly->StaticParameters = MoveTemp(StaticParameters_DEPRECATED.EditorOnly);
+		StaticParameters.UpdateLegacyTerrainLayerWeightData();
 	}
-
 #endif // WITH_EDITOR
 
 	bool bSavedCachedData = false;
@@ -2518,9 +2451,7 @@ void UMaterialInstance::Serialize(FArchive& Ar)
 #if WITH_EDITOR
 			if (Ar.CustomVer(FRenderingObjectVersion::GUID) < FRenderingObjectVersion::MaterialAttributeLayerParameters)
 			{
-				StaticParameters_DEPRECATED.SerializeLegacy(Ar);
-				StaticParametersRuntime = MoveTemp(StaticParameters_DEPRECATED.GetRuntime());
-				EditorOnly->StaticParameters = MoveTemp(StaticParameters_DEPRECATED.EditorOnly);
+				StaticParameters.SerializeLegacy(Ar);
 			}
 
 			static_assert(!STORE_ONLY_ACTIVE_SHADERMAPS, "Only discard unused SMs in cooked build");
@@ -2547,12 +2478,12 @@ void UMaterialInstance::Serialize(FArchive& Ar)
 			FMaterialShaderMapId LegacyId;
 			LegacyId.Serialize(Ar, bLoadedByCookedMaterial);
 
-			EditorOnly->StaticParameters.StaticSwitchParameters = LegacyId.GetStaticSwitchParameters();
-			EditorOnly->StaticParameters.StaticComponentMaskParameters = LegacyId.GetStaticComponentMaskParameters();
-			EditorOnly->StaticParameters.TerrainLayerWeightParameters = LegacyId.GetTerrainLayerWeightParameters();
+			StaticParameters.StaticSwitchParameters = LegacyId.GetStaticSwitchParameters();
+			StaticParameters.StaticComponentMaskParameters = LegacyId.GetStaticComponentMaskParameters();
+			StaticParameters.TerrainLayerWeightParameters = LegacyId.GetTerrainLayerWeightParameters();
 
-			TrimToOverriddenOnly(EditorOnly->StaticParameters.StaticSwitchParameters);
-			TrimToOverriddenOnly(EditorOnly->StaticParameters.StaticComponentMaskParameters);
+			TrimToOverriddenOnly(StaticParameters.StaticSwitchParameters);
+			TrimToOverriddenOnly(StaticParameters.StaticComponentMaskParameters);
 		}
 #endif // WITH_EDITOR
 	}
@@ -2706,18 +2637,17 @@ void UMaterialInstance::PostLoad()
 		}
 	}
 
-#if WITH_EDITORONLY_DATA
 	// And any material layers parameter's functions
-	if (StaticParametersRuntime.bHasMaterialLayers)
+	if (StaticParameters.bHasMaterialLayers)
 	{
-		for (UMaterialFunctionInterface* Dependency : StaticParametersRuntime.MaterialLayers.Layers)
+		for (UMaterialFunctionInterface* Dependency : StaticParameters.MaterialLayers.Layers)
 		{
 			if (Dependency)
 			{
 				Dependency->ConditionalPostLoad();
 			}
 		}
-		for (UMaterialFunctionInterface* Dependency : StaticParametersRuntime.MaterialLayers.Blends)
+		for (UMaterialFunctionInterface* Dependency : StaticParameters.MaterialLayers.Blends)
 		{
 			if (Dependency)
 			{
@@ -2725,7 +2655,6 @@ void UMaterialInstance::PostLoad()
 			}
 		}
 	}
-#endif // WITH_EDITORONLY_DATA
 
 #if WITH_EDITOR
 	UpdateCachedData();
@@ -3002,9 +2931,7 @@ void FMaterialInstanceParameterUpdateContext::SetBasePropertyOverrides(const FMa
 void FMaterialInstanceParameterUpdateContext::SetMaterialLayers(const FMaterialLayersFunctions& InValue)
 {
 	StaticParameters.bHasMaterialLayers = true;
-	StaticParameters.MaterialLayers = InValue.GetRuntime();
-	StaticParameters.EditorOnly.MaterialLayers = InValue.EditorOnly;
-	StaticParameters.Validate();
+	StaticParameters.MaterialLayers = InValue;
 }
 #endif // WITH_EDITORONLY_DATA
 
@@ -3334,17 +3261,9 @@ void UMaterialInstance::ClearParameterValuesInternal(EMaterialInstanceClearParam
 
 	if (EnumHasAnyFlags(Flags, EMaterialInstanceClearParameterFlag::Static))
 	{
-		StaticParametersRuntime.Empty();
-#if WITH_EDITORONLY_DATA
-		UMaterialInstanceEditorOnlyData* EditorOnly = GetEditorOnlyData();
-		if (EditorOnly)
-		{
-			EditorOnly->StaticParameters.Empty();
-		}
-#endif // WITH_EDITORONLY_DATA
+		StaticParameters.Empty();
 	}
 
-
 	if (Resource && bUpdateResource)
 	{
 		FMaterialInstanceResource* InResource = Resource;
@@ -3365,11 +3284,10 @@ void UMaterialInstance::ClearParameterValuesInternal(EMaterialInstanceClearParam
 #if WITH_EDITOR
 void UMaterialInstance::UpdateStaticPermutation(const FStaticParameterSet& NewParameters, FMaterialInstanceBasePropertyOverrides& NewBasePropertyOverrides, const bool bForceStaticPermutationUpdate /*= false*/, FMaterialUpdateContext* MaterialUpdateContext)
 {
-	UMaterialInstanceEditorOnlyData* EditorOnly = GetEditorOnlyData();
 	FStaticParameterSet CompareParameters = NewParameters;
 
-	TrimToOverriddenOnly(CompareParameters.EditorOnly.StaticSwitchParameters);
-	TrimToOverriddenOnly(CompareParameters.EditorOnly.StaticComponentMaskParameters);
+	TrimToOverriddenOnly(CompareParameters.StaticSwitchParameters);
+	TrimToOverriddenOnly(CompareParameters.StaticComponentMaskParameters);
 
 	// Check to see if the material layers being assigned match values from the parent
 	if (CompareParameters.bHasMaterialLayers && Parent)
@@ -3377,11 +3295,7 @@ void UMaterialInstance::UpdateStaticPermutation(const FStaticParameterSet& NewPa
 		FMaterialLayersFunctions ParentLayers;
 		if (Parent->GetMaterialLayers(ParentLayers))
 		{
-			if (FMaterialLayersFunctions::MatchesParent(
-				CompareParameters.MaterialLayers,
-				CompareParameters.EditorOnly.MaterialLayers,
-				ParentLayers,
-				ParentLayers.EditorOnly))
+			if (CompareParameters.MaterialLayers.MatchesParent(ParentLayers))
 			{
 				CompareParameters.bHasMaterialLayers = false;
 				CompareParameters.MaterialLayers.Empty();
@@ -3389,8 +3303,7 @@ void UMaterialInstance::UpdateStaticPermutation(const FStaticParameterSet& NewPa
 		}
 	}
 
-	const FStaticParameterSet CurrentParameters = GetStaticParameters();
-	const bool bParamsHaveChanged = CurrentParameters != CompareParameters;
+	const bool bParamsHaveChanged = StaticParameters != CompareParameters;
 	const bool bBasePropertyOverridesHaveChanged = BasePropertyOverrides != NewBasePropertyOverrides;
 
 	BasePropertyOverrides = NewBasePropertyOverrides;
@@ -3408,8 +3321,7 @@ void UMaterialInstance::UpdateStaticPermutation(const FStaticParameterSet& NewPa
 		FlushRenderingCommands();
 
 		bHasStaticPermutationResource = bWantsStaticPermutationResource;
-		StaticParametersRuntime = CompareParameters.GetRuntime();
-		EditorOnly->StaticParameters = CompareParameters.EditorOnly;
+		StaticParameters = CompareParameters;
 
 		UpdateCachedData();
 		CacheResourceShadersForRendering(EMaterialShaderPrecompileMode::None);
@@ -3466,7 +3378,7 @@ void UMaterialInstance::UpdateStaticPermutation(const FStaticParameterSet& NewPa
 void UMaterialInstance::UpdateStaticPermutation(FMaterialUpdateContext* MaterialUpdateContext)
 {
 	// Force the update, since we aren't technically changing anything
-	UpdateStaticPermutation(GetStaticParameters(), BasePropertyOverrides, true, MaterialUpdateContext);
+	UpdateStaticPermutation(StaticParameters, BasePropertyOverrides, true, MaterialUpdateContext);
 }
 
 void UMaterialInstance::UpdateParameterNames()
@@ -3504,7 +3416,7 @@ void UMaterialInstance::PostEditChangeProperty(FPropertyChangedEvent& PropertyCh
 
 	// Force UpdateStaticPermutation when change type is Redirected as this probably means a Material or MaterialInstance parent asset was deleted.
 	const bool bForceStaticPermutationUpdate = PropertyChangedEvent.ChangeType == EPropertyChangeType::Redirected;
-	UpdateStaticPermutation(GetStaticParameters(), BasePropertyOverrides, bForceStaticPermutationUpdate);
+	UpdateStaticPermutation(StaticParameters, BasePropertyOverrides, bForceStaticPermutationUpdate);
 
 	if (PropertyChangedEvent.ChangeType == EPropertyChangeType::ValueSet || PropertyChangedEvent.ChangeType == EPropertyChangeType::ArrayClear || PropertyChangedEvent.ChangeType == EPropertyChangeType::ArrayRemove || PropertyChangedEvent.ChangeType == EPropertyChangeType::Unspecified || PropertyChangedEvent.ChangeType == EPropertyChangeType::Duplicate)
 	{
@@ -3960,48 +3872,17 @@ bool UMaterialInstance::IsPropertyActive(EMaterialProperty InProperty) const
 	return Material ? Material->IsPropertyActiveInDerived(InProperty, this) : false;
 }
 
-bool UMaterialInstance::HasStaticParameters() const
-{
-	if (!StaticParametersRuntime.IsEmpty())
-	{
-		return true;
-	}
-#if WITH_EDITOR
-	const UMaterialInstanceEditorOnlyData* EditorOnly = GetEditorOnlyData();
-	if (EditorOnly && !EditorOnly->StaticParameters.IsEmpty())
-	{
-		return true;
-	}
-#endif // WITH_EDITOR
-	return false;
-}
-
-FStaticParameterSet UMaterialInstance::GetStaticParameters() const
-{
-	FStaticParameterSet Result;
-	Result.GetRuntime() = StaticParametersRuntime;
-#if WITH_EDITORONLY_DATA
-	const UMaterialInstanceEditorOnlyData* EditorOnly = GetEditorOnlyData();
-	if (EditorOnly)
-	{
-		Result.EditorOnly = EditorOnly->StaticParameters;
-	}
-#endif // WITH_EDITORONLY_DATA
-	return Result;
-}
-
 #if WITH_EDITOR
 int32 UMaterialInstance::CompilePropertyEx( class FMaterialCompiler* Compiler, const FGuid& AttributeID )
 {
 	return Parent ? Parent->CompilePropertyEx(Compiler, AttributeID) : INDEX_NONE;
 }
+#endif // WITH_EDITOR
 
-const FStaticParameterSetEditorOnlyData& UMaterialInstance::GetEditorOnlyStaticParameters() const
+const FStaticParameterSet& UMaterialInstance::GetStaticParameters() const
 {
-	const UMaterialInstanceEditorOnlyData* EditorOnly = GetEditorOnlyData();
-	return EditorOnly->StaticParameters;
+	return StaticParameters;
 }
-#endif // WITH_EDITOR
 
 void UMaterialInstance::GetLightingGuidChain(bool bIncludeTextures, TArray<FGuid>& OutGuids) const
 {
@@ -4077,12 +3958,10 @@ bool UMaterialInstance::Equivalent(const UMaterialInstance* CompareTo) const
 		return false;
 	}
 
-	const FStaticParameterSet LocalStaticParameters = GetStaticParameters();
-	if (!LocalStaticParameters.Equivalent(CompareTo->GetStaticParameters()))
+	if (!StaticParameters.Equivalent(CompareTo->StaticParameters))
 	{
 		return false;
 	}
-
 	return true;
 }
 
@@ -4364,6 +4243,58 @@ void UMaterialInstance::CopyMaterialUniformParametersInternal(UMaterialInterface
 #endif
 }
 
+#if WITH_EDITOR
+
+void FindCollectionExpressionRecursive(TArray<FGuid>& OutGuidList, const TArray<TObjectPtr<UMaterialExpression>>& InMaterialExpression)
+{
+	for (int32 ExpressionIndex = 0; ExpressionIndex < InMaterialExpression.Num(); ExpressionIndex++)
+	{
+		if (const UMaterialExpressionCollectionParameter* CollectionPtr = Cast<UMaterialExpressionCollectionParameter>(InMaterialExpression[ExpressionIndex]))
+		{
+			if (CollectionPtr->Collection)
+			{
+				OutGuidList.Add(CollectionPtr->Collection->StateId);
+			}
+			return;
+		}
+		else if (const UMaterialExpressionMaterialFunctionCall* MaterialFunctionCall = Cast<UMaterialExpressionMaterialFunctionCall>(InMaterialExpression[ExpressionIndex]))
+		{
+			if (const TArray<TObjectPtr<UMaterialExpression>>* FunctionExpressions = MaterialFunctionCall->MaterialFunction ? MaterialFunctionCall->MaterialFunction->GetFunctionExpressions() : nullptr)
+			{
+				FindCollectionExpressionRecursive(OutGuidList, *FunctionExpressions);
+			}
+		}
+		else if (const UMaterialExpressionMaterialAttributeLayers* MaterialLayers = Cast<UMaterialExpressionMaterialAttributeLayers>(InMaterialExpression[ExpressionIndex]))
+		{
+			const TArray<UMaterialFunctionInterface*>& Layers = MaterialLayers->GetLayers();
+			const TArray<UMaterialFunctionInterface*>& Blends = MaterialLayers->GetBlends();
+
+			for (const UMaterialFunctionInterface* Layer : Layers)
+			{
+				if (Layer)
+				{
+					if (const TArray<TObjectPtr<UMaterialExpression>>* FunctionExpressions = Layer->GetFunctionExpressions())
+					{
+						FindCollectionExpressionRecursive(OutGuidList, *FunctionExpressions);
+					}
+				}
+			}
+
+			for (const UMaterialFunctionInterface* Blend : Blends)
+			{
+				if (Blend)
+				{
+					if (const TArray<TObjectPtr<UMaterialExpression>>* FunctionExpressions = Blend->GetFunctionExpressions())
+					{
+						FindCollectionExpressionRecursive(OutGuidList, *FunctionExpressions);
+					}
+				}
+			}
+		}
+	}
+}
+#endif // WITH_EDITOR
+
 #if WITH_EDITORONLY_DATA
 PRAGMA_DISABLE_DEPRECATION_WARNINGS
 UMaterialInstance::FCustomStaticParametersGetterDelegate UMaterialInstance::CustomStaticParametersGetters;
