[Backout] - CL19964485

diff --git a/Engine/Source/Runtime/Engine/Private/Materials/MaterialInstanceConstant.cpp b/Engine/Source/Runtime/Engine/Private/Materials/MaterialInstanceConstant.cpp
--- a/Engine/Source/Runtime/Engine/Private/Materials/MaterialInstanceConstant.cpp
+++ b/Engine/Source/Runtime/Engine/Private/Materials/MaterialInstanceConstant.cpp
@@ -106,7 +106,7 @@ void UMaterialInstanceConstant::CopyMaterialUniformParametersEditorOnly(UMateria
 			FStaticParameterSet MyParamSet;
 			GetStaticParameterValues(MyParamSet);
 
-			MyParamSet.EditorOnly.StaticSwitchParameters = SourceParamSet.EditorOnly.StaticSwitchParameters;
+			MyParamSet.StaticSwitchParameters = SourceParamSet.StaticSwitchParameters;
 
 			UpdateStaticPermutation(MyParamSet);
 
@@ -165,10 +165,10 @@ void FMaterialInstanceCachedData::InitializeForConstant(const FMaterialLayersFun
 	for (int32 LayerIndex = 0; LayerIndex < NumLayers; ++LayerIndex)
 	{
 		int32 ParentLayerIndex = INDEX_NONE;
-		if (ParentLayers && Layers->EditorOnly.LayerLinkStates[LayerIndex] == EMaterialLayerLinkState::LinkedToParent)
+		if (ParentLayers && Layers->LayerLinkStates[LayerIndex] == EMaterialLayerLinkState::LinkedToParent)
 		{
-			const FGuid& LayerGuid = Layers->EditorOnly.LayerGuids[LayerIndex];
-			ParentLayerIndex = ParentLayers->EditorOnly.LayerGuids.Find(LayerGuid);
+			const FGuid& LayerGuid = Layers->LayerGuids[LayerIndex];
+			ParentLayerIndex = ParentLayers->LayerGuids.Find(LayerGuid);
 		}
 		ParentLayerIndexRemap.Add(ParentLayerIndex);
 	}
@@ -211,20 +211,18 @@ void UMaterialInstanceConstant::UpdateCachedData()
 		if (LocalStaticParameters.bHasMaterialLayers)
 		{
 			UMaterial* BaseMaterial = GetMaterial();
-
-			FMaterialLayersFunctions MaterialLayers;
-			LocalStaticParameters.GetMaterialLayers(MaterialLayers);
 			if (bUsingNewHLSLGenerator)
 			{
 				LocalCachedTree.Reset(new FMaterialCachedHLSLTree());
-				LocalCachedTree->GenerateTree(BaseMaterial, &MaterialLayers, nullptr);
+				LocalCachedTree->GenerateTree(BaseMaterial, &LocalStaticParameters.MaterialLayers, nullptr);
 			}
 			else
 			{
 				FMaterialCachedExpressionContext Context;
-				Context.LayerOverrides = &MaterialLayers;
+				Context.LayerOverrides = &LocalStaticParameters.MaterialLayers;
 				LocalCachedExpressionData.Reset(new FMaterialCachedExpressionData());
-				LocalCachedExpressionData->UpdateForExpressions(Context, BaseMaterial->GetExpressions(), GlobalParameter, INDEX_NONE);
+				LocalCachedExpressionData->Reset();
+				LocalCachedExpressionData->UpdateForExpressions(Context, BaseMaterial->Expressions, GlobalParameter, INDEX_NONE);
 			}
 		}
 
@@ -234,14 +232,11 @@ void UMaterialInstanceConstant::UpdateCachedData()
 		{
 			check(!LocalCachedExpressionData);
 			LocalCachedExpressionData.Reset(new FMaterialCachedExpressionData());
+			LocalCachedExpressionData->Reset();
 			LocalCachedExpressionData->UpdateForCachedHLSLTree(GetCachedHLSLTree(), &LocalStaticParameters);
 		}
 
 		CachedExpressionData = MoveTemp(LocalCachedExpressionData);
-		if (CachedExpressionData)
-		{
-			EditorOnlyData->CachedExpressionData = CachedExpressionData->EditorOnlyData;
-		}
 
 		FObjectCacheEventSink::NotifyReferencedTextureChanged_Concurrent(this);
 	}
diff --git a/Engine/Source/Runtime/Engine/Private/Materials/MaterialInterface.cpp b/Engine/Source/Runtime/Engine/Private/Materials/MaterialInterface.cpp
--- a/Engine/Source/Runtime/Engine/Private/Materials/MaterialInterface.cpp
+++ b/Engine/Source/Runtime/Engine/Private/Materials/MaterialInterface.cpp
@@ -165,17 +165,7 @@ void UMaterialInterface::Serialize(FArchive& Ar)
 		{
 			CachedExpressionData.Reset(new FMaterialCachedExpressionData());
 			bLoadedCachedExpressionData = true;
-#if WITH_EDITORONLY_DATA
-			EditorOnlyData->CachedExpressionData = CachedExpressionData->EditorOnlyData;
-#endif
-		}
-		else
-		{
-#if WITH_EDITOR
-			CachedExpressionData->Validate();
-#endif
 		}
-
 		check(CachedExpressionData);
 		UScriptStruct* Struct = FMaterialCachedExpressionData::StaticStruct();
 		Struct->SerializeTaggedProperties(Ar, (uint8*)CachedExpressionData.Get(), Struct, nullptr);
@@ -201,7 +191,7 @@ void UMaterialInterface::PostLoad()
 	{
 		TextureStreamingData.Empty();
 	}
-#endif // WITH_EDITORONLY_DATA
+#endif
 }
 
 const FMaterialCachedExpressionData& UMaterialInterface::GetCachedExpressionData(TMicRecursionGuard) const
@@ -609,14 +599,6 @@ void UMaterialInterface::AddReferencedObjects(UObject* InThis, FReferenceCollect
 	Super::AddReferencedObjects(This, Collector);
 }
 
-void UMaterialInterface::PostInitProperties()
-{
-#if WITH_EDITORONLY_DATA
-	EditorOnlyData = CreateEditorOnlyData();
-#endif
-	Super::PostInitProperties();
-}
-
 void UMaterialInterface::PostDuplicate(bool bDuplicateForPIE)
 {
 	Super::PostDuplicate(bDuplicateForPIE);
@@ -896,14 +878,14 @@ bool UMaterialInterface::GetStaticComponentMaskParameterDefaultValue(const FHash
 void UMaterialInterface::GetAllParametersOfType(EMaterialParameterType Type, TMap<FMaterialParameterInfo, FMaterialParameterMetadata>& OutParameters) const
 {
 	OutParameters.Reset();
-	GetCachedExpressionData().GetAllParametersOfType(Type, OutParameters);
+	GetCachedExpressionData().Parameters.GetAllParametersOfType(Type, OutParameters);
 }
 
 void UMaterialInterface::GetAllParameterInfoOfType(EMaterialParameterType Type, TArray<FMaterialParameterInfo>& OutParameterInfo, TArray<FGuid>& OutParameterIds) const
 {
 	OutParameterInfo.Reset();
 	OutParameterIds.Reset();
-	GetCachedExpressionData().GetAllParameterInfoOfType(Type, OutParameterInfo, OutParameterIds);
+	GetCachedExpressionData().Parameters.GetAllParameterInfoOfType(Type, OutParameterInfo, OutParameterIds);
 }
 
 void UMaterialInterface::GetAllScalarParameterInfo(TArray<FMaterialParameterInfo>& OutParameterInfo, TArray<FGuid>& OutParameterIds) const
@@ -1362,55 +1344,3 @@ void UMaterialInterface::RemoveUserDataOfClass(TSubclassOf<UAssetUserData> InUse
 	}
 }
 
-#if WITH_EDITORONLY_DATA
-const UClass* UMaterialInterface::GetEditorOnlyDataClass() const
-{
-	return UMaterialInterfaceEditorOnlyData::StaticClass();
-}
-
-#include "Materials/MaterialInstanceDynamic.h"
-
-UMaterialInterfaceEditorOnlyData* UMaterialInterface::CreateEditorOnlyData()
-{
-	const UClass* EditorOnlyClass = GetEditorOnlyDataClass();
-	check(EditorOnlyClass);
-	check(EditorOnlyClass->HasAllClassFlags(CLASS_Optional));
-
-	const FString EditorOnlyName = FString::Printf(TEXT("%sEditorOnlyData"), *GetName());
-	const EObjectFlags EditorOnlyFlags = GetMaskedFlags(RF_PropagateToSubObjects);
-	return NewObject<UMaterialInterfaceEditorOnlyData>(this, EditorOnlyClass, *EditorOnlyName, EditorOnlyFlags);
-}
-#endif // WITH_EDITORONLY_DATA
-
-UMaterialInterfaceEditorOnlyData::UMaterialInterfaceEditorOnlyData()
-{
-}
-
-UMaterialInterfaceEditorOnlyData::UMaterialInterfaceEditorOnlyData(FVTableHelper& Helper)
-{
-}
-
-UMaterialInterfaceEditorOnlyData::~UMaterialInterfaceEditorOnlyData()
-{
-}
-
-void UMaterialInterfaceEditorOnlyData::Serialize(FArchive& Ar)
-{
-	Super::Serialize(Ar);
-
-	bool bSavedCachedExpressionData = false;
-	if ((Ar.IsCooking() || (FPlatformProperties::RequiresCookedData() && Ar.IsSaving() && (Ar.GetPortFlags() & PPF_Duplicate))) && (bool)CachedExpressionData)
-	{
-		bSavedCachedExpressionData = true;
-	}
-
-	Ar << bSavedCachedExpressionData;
-
-	if (bSavedCachedExpressionData)
-	{
-		check(CachedExpressionData);
-		UScriptStruct* Struct = FMaterialCachedExpressionEditorOnlyData::StaticStruct();
-		Struct->SerializeTaggedProperties(Ar, (uint8*)CachedExpressionData.Get(), Struct, nullptr);
-		bLoadedCachedExpressionData = true;
-	}
-}
diff --git a/Engine/Source/Runtime/Engine/Private/Materials/MaterialShader.cpp b/Engine/Source/Runtime/Engine/Private/Materials/MaterialShader.cpp
--- a/Engine/Source/Runtime/Engine/Private/Materials/MaterialShader.cpp
+++ b/Engine/Source/Runtime/Engine/Private/Materials/MaterialShader.cpp
@@ -249,50 +249,39 @@ void UpdateMaterialShaderCompilingStats(const FMaterial* Material)
 	}
 }
 
-FStaticParameterSet::FStaticParameterSet(const FStaticParameterSet& InValue) = default;
+FStaticParameterSet::FStaticParameterSet(const FStaticParameterSet& InValue)
+	: StaticSwitchParameters(InValue.StaticSwitchParameters)
+	, StaticComponentMaskParameters(InValue.StaticComponentMaskParameters)
+	, TerrainLayerWeightParameters(InValue.TerrainLayerWeightParameters)
+	, bHasMaterialLayers(InValue.bHasMaterialLayers)
+{
+	if (bHasMaterialLayers)
+	{
+		MaterialLayers = InValue.MaterialLayers;
+	}
+}
 
 FStaticParameterSet& FStaticParameterSet::operator=(const FStaticParameterSet& InValue)
 {
-#if WITH_EDITORONLY_DATA
-	EditorOnly.StaticSwitchParameters = InValue.EditorOnly.StaticSwitchParameters;
-	EditorOnly.StaticComponentMaskParameters = InValue.EditorOnly.StaticComponentMaskParameters;
-	EditorOnly.TerrainLayerWeightParameters = InValue.EditorOnly.TerrainLayerWeightParameters;
-#endif // WITH_EDITORONLY_DATA
+	StaticSwitchParameters = InValue.StaticSwitchParameters;
+	StaticComponentMaskParameters = InValue.StaticComponentMaskParameters;
+	TerrainLayerWeightParameters = InValue.TerrainLayerWeightParameters;
 	MaterialLayers = InValue.MaterialLayers;
 	bHasMaterialLayers = InValue.bHasMaterialLayers;
 	if (bHasMaterialLayers)
 	{
 		MaterialLayers = InValue.MaterialLayers;
-#if WITH_EDITORONLY_DATA
-		EditorOnly.MaterialLayers = InValue.EditorOnly.MaterialLayers;
-#endif
 	}
 	return *this;
 }
 
-void FStaticParameterSet::PostSerialize(const FArchive& Ar)
+void FStaticParameterSet::Empty()
 {
-#if WITH_EDITORONLY_DATA
-	if (StaticSwitchParameters_DEPRECATED.Num() > 0)
-	{
-		EditorOnly.StaticSwitchParameters = MoveTemp(StaticSwitchParameters_DEPRECATED);
-	}
-	if (StaticComponentMaskParameters_DEPRECATED.Num() > 0)
-	{
-		EditorOnly.StaticComponentMaskParameters = MoveTemp(StaticComponentMaskParameters_DEPRECATED);
-	}
-	if (TerrainLayerWeightParameters_DEPRECATED.Num() > 0)
-	{
-		EditorOnly.TerrainLayerWeightParameters = MoveTemp(TerrainLayerWeightParameters_DEPRECATED);
-	}
-
-	if (!MaterialLayers_DEPRECATED.IsEmpty())
-	{
-		MaterialLayers = MoveTemp(MaterialLayers_DEPRECATED.GetRuntime());
-		EditorOnly.MaterialLayers = MoveTemp(MaterialLayers_DEPRECATED.EditorOnly);
-	}
-	MaterialLayers_DEPRECATED.Empty();
-#endif // WITH_EDITORONLY_DATA
+	StaticSwitchParameters.Empty();
+	StaticComponentMaskParameters.Empty();
+	TerrainLayerWeightParameters.Empty();
+	MaterialLayers.Empty();
+	bHasMaterialLayers = false;
 }
 
 #if WITH_EDITOR
@@ -307,9 +296,9 @@ void FStaticParameterSet::SerializeLegacy(FArchive& Ar)
 	Ar.UsingCustomVersion(FReleaseObjectVersion::GUID);
 	Ar.UsingCustomVersion(FUE5ReleaseStreamObjectVersion::GUID);
 
-	Ar << EditorOnly.StaticSwitchParameters;
-	Ar << EditorOnly.StaticComponentMaskParameters;
-	Ar << EditorOnly.TerrainLayerWeightParameters;
+	Ar << StaticSwitchParameters;
+	Ar << StaticComponentMaskParameters;
+	Ar << TerrainLayerWeightParameters;
 
 	if (Ar.CustomVer(FReleaseObjectVersion::GUID) >= FReleaseObjectVersion::MaterialLayersParameterSerializationRefactor)
 	{
@@ -320,8 +309,7 @@ void FStaticParameterSet::SerializeLegacy(FArchive& Ar)
 			if (MaterialLayersParameters_DEPRECATED.Num() > 0)
 			{
 				bHasMaterialLayers = true;
-				MaterialLayers = MoveTemp(MaterialLayersParameters_DEPRECATED[0].Value.GetRuntime());
-				EditorOnly.MaterialLayers = MoveTemp(MaterialLayersParameters_DEPRECATED[0].Value.EditorOnly);
+				MaterialLayers = MoveTemp(MaterialLayersParameters_DEPRECATED[0].Value);
 				MaterialLayersParameters_DEPRECATED.Empty();
 			}
 			PRAGMA_ENABLE_DEPRECATION_WARNINGS
@@ -332,9 +320,9 @@ void FStaticParameterSet::SerializeLegacy(FArchive& Ar)
 void FStaticParameterSet::UpdateLegacyTerrainLayerWeightData()
 {
 	int32 ParameterIndex = 0;
-	while (ParameterIndex < EditorOnly.TerrainLayerWeightParameters.Num())
+	while (ParameterIndex < TerrainLayerWeightParameters.Num())
 	{
-		FStaticTerrainLayerWeightParameter& TerrainParameter = EditorOnly.TerrainLayerWeightParameters[ParameterIndex];
+		FStaticTerrainLayerWeightParameter& TerrainParameter = TerrainLayerWeightParameters[ParameterIndex];
 		if (TerrainParameter.bOverride_DEPRECATED)
 		{
 			TerrainParameter.LayerName = TerrainParameter.ParameterInfo_DEPRECATED.Name;
@@ -343,7 +331,7 @@ void FStaticParameterSet::UpdateLegacyTerrainLayerWeightData()
 		else
 		{
 			// Remove any parameters that didn't have bOverride set
-			EditorOnly.TerrainLayerWeightParameters.RemoveAt(ParameterIndex);
+			TerrainLayerWeightParameters.RemoveAt(ParameterIndex);
 		}
 	}
 }
@@ -354,73 +342,41 @@ void FStaticParameterSet::UpdateLegacyMaterialLayersData()
 	if (MaterialLayersParameters_DEPRECATED.Num() > 0)
 	{
 		bHasMaterialLayers = true;
-		MaterialLayers = MoveTemp(MaterialLayersParameters_DEPRECATED[0].Value.GetRuntime());
-		EditorOnly.MaterialLayers = MoveTemp(MaterialLayersParameters_DEPRECATED[0].Value.EditorOnly);
+		MaterialLayers = MoveTemp(MaterialLayersParameters_DEPRECATED[0].Value);
 		MaterialLayersParameters_DEPRECATED.Empty();
 	}
 	PRAGMA_ENABLE_DEPRECATION_WARNINGS
 }
-
-bool FStaticParameterSet::GetMaterialLayers(FMaterialLayersFunctions& OutMaterialLayers) const
-{
-	if (bHasMaterialLayers)
-	{
-		OutMaterialLayers.GetRuntime() = MaterialLayers;
-		OutMaterialLayers.EditorOnly = EditorOnly.MaterialLayers;
-		return true;
-	}
-	return false;
-}
-
-void FStaticParameterSet::Validate(const FStaticParameterSetRuntimeData& Runtime, const FStaticParameterSetEditorOnlyData& EditorOnly)
-{
-	FMaterialLayersFunctions::Validate(Runtime.MaterialLayers, EditorOnly.MaterialLayers);
-}
-
 #endif // WITH_EDITOR
 
 bool FStaticParameterSet::operator==(const FStaticParameterSet& ReferenceSet) const
 {
-	if (bHasMaterialLayers != ReferenceSet.bHasMaterialLayers)
-	{
-		return false;
-	}
-#if WITH_EDITORONLY_DATA
-	if (EditorOnly.StaticSwitchParameters.Num() != ReferenceSet.EditorOnly.StaticSwitchParameters.Num()
-		|| EditorOnly.StaticComponentMaskParameters.Num() != ReferenceSet.EditorOnly.StaticComponentMaskParameters.Num()
-		|| EditorOnly.TerrainLayerWeightParameters.Num() != ReferenceSet.EditorOnly.TerrainLayerWeightParameters.Num())
+	if (StaticSwitchParameters.Num() != ReferenceSet.StaticSwitchParameters.Num()
+		|| StaticComponentMaskParameters.Num() != ReferenceSet.StaticComponentMaskParameters.Num()
+		|| TerrainLayerWeightParameters.Num() != ReferenceSet.TerrainLayerWeightParameters.Num()
+		|| bHasMaterialLayers != ReferenceSet.bHasMaterialLayers)
 	{
 		return false;
 	}
 
-	if (EditorOnly.StaticSwitchParameters != ReferenceSet.EditorOnly.StaticSwitchParameters)
+	if (StaticSwitchParameters != ReferenceSet.StaticSwitchParameters)
 	{
 		return false;
 	}
 
-	if (EditorOnly.StaticComponentMaskParameters != ReferenceSet.EditorOnly.StaticComponentMaskParameters)
+	if (StaticComponentMaskParameters != ReferenceSet.StaticComponentMaskParameters)
 	{
 		return false;
 	}
 
-	if (EditorOnly.TerrainLayerWeightParameters != ReferenceSet.EditorOnly.TerrainLayerWeightParameters)
+	if (TerrainLayerWeightParameters != ReferenceSet.TerrainLayerWeightParameters)
 	{
 		return false;
 	}
-#endif // WITH_EDITORONLY_DATA
 
-	if (bHasMaterialLayers)
+	if (bHasMaterialLayers && MaterialLayers != ReferenceSet.MaterialLayers)
 	{
-		if (MaterialLayers != ReferenceSet.MaterialLayers)
-		{
-			return false;
-		}
-#if WITH_EDITORONLY_DATA
-		if (EditorOnly.MaterialLayers != ReferenceSet.EditorOnly.MaterialLayers)
-		{
-			return false;
-		}
-#endif // WITH_EDITORONLY_DATA
+		return false;
 	}
 
 	return true;
@@ -428,37 +384,28 @@ bool FStaticParameterSet::operator==(const FStaticParameterSet& ReferenceSet) co
 
 void FStaticParameterSet::SortForEquivalent()
 {
-#if WITH_EDITORONLY_DATA
-	EditorOnly.StaticSwitchParameters.Sort([](const FStaticSwitchParameter& A, const FStaticSwitchParameter& B) { return B.ExpressionGUID < A.ExpressionGUID; });
-	EditorOnly.StaticComponentMaskParameters.Sort([](const FStaticComponentMaskParameter& A, const FStaticComponentMaskParameter& B) { return B.ExpressionGUID < A.ExpressionGUID; });
-	EditorOnly.TerrainLayerWeightParameters.Sort([](const FStaticTerrainLayerWeightParameter& A, const FStaticTerrainLayerWeightParameter& B) { return B.LayerName.LexicalLess(A.LayerName); });
-#endif // WITH_EDITORONLY_DATA
+	StaticSwitchParameters.Sort([](const FStaticSwitchParameter& A, const FStaticSwitchParameter& B) { return B.ExpressionGUID < A.ExpressionGUID; });
+	StaticComponentMaskParameters.Sort([](const FStaticComponentMaskParameter& A, const FStaticComponentMaskParameter& B) { return B.ExpressionGUID < A.ExpressionGUID; });
+	TerrainLayerWeightParameters.Sort([](const FStaticTerrainLayerWeightParameter& A, const FStaticTerrainLayerWeightParameter& B) { return B.LayerName.LexicalLess(A.LayerName); });
 }
 
 bool FStaticParameterSet::Equivalent(const FStaticParameterSet& ReferenceSet) const
 {
-	if (bHasMaterialLayers != ReferenceSet.bHasMaterialLayers)
+	if (StaticSwitchParameters.Num() == ReferenceSet.StaticSwitchParameters.Num()
+		&& StaticComponentMaskParameters.Num() == ReferenceSet.StaticComponentMaskParameters.Num()
+		&& TerrainLayerWeightParameters.Num() == ReferenceSet.TerrainLayerWeightParameters.Num()
+		&& bHasMaterialLayers == ReferenceSet.bHasMaterialLayers)
 	{
-		return false;
+		// this is not ideal, but it is easy to code up
+		FStaticParameterSet Temp1 = *this;
+		FStaticParameterSet Temp2 = ReferenceSet;
+		Temp1.SortForEquivalent();
+		Temp2.SortForEquivalent();
+		bool bResult = (Temp1 == Temp2);
+		ensure(!bResult || (*this) == ReferenceSet); // if this never fires, then we really didn't need to sort did we?
+		return bResult;
 	}
-
-#if WITH_EDITORONLY_DATA
-	if (EditorOnly.StaticSwitchParameters.Num() != ReferenceSet.EditorOnly.StaticSwitchParameters.Num()
-		|| EditorOnly.StaticComponentMaskParameters.Num() != ReferenceSet.EditorOnly.StaticComponentMaskParameters.Num()
-		|| EditorOnly.TerrainLayerWeightParameters.Num() != ReferenceSet.EditorOnly.TerrainLayerWeightParameters.Num())
-	{
-		return false;
-	}
-#endif // WITH_EDITORONLY_DATA
-
-	// this is not ideal, but it is easy to code up
-	FStaticParameterSet Temp1 = *this;
-	FStaticParameterSet Temp2 = ReferenceSet;
-	Temp1.SortForEquivalent();
-	Temp2.SortForEquivalent();
-	bool bResult = (Temp1 == Temp2);
-	ensure(!bResult || (*this) == ReferenceSet); // if this never fires, then we really didn't need to sort did we?
-	return bResult;
+	return false;
 }
 
 #if WITH_EDITORONLY_DATA
@@ -478,21 +425,21 @@ void FStaticParameterSet::AddParametersOfType(EMaterialParameterType Type, const
 	switch (Type)
 	{
 	case EMaterialParameterType::StaticSwitch:
-		EditorOnly.StaticSwitchParameters.Empty(Values.Num());
+		StaticSwitchParameters.Empty(Values.Num());
 		for (const auto& It : Values)
 		{
 			const FMaterialParameterMetadata& Meta = It.Value;
 			check(Meta.Value.Type == Type);
-			EditorOnly.StaticSwitchParameters.Emplace(It.Key, Meta.Value.AsStaticSwitch(), Meta.bOverride, Meta.ExpressionGuid);
+			StaticSwitchParameters.Emplace(It.Key, Meta.Value.AsStaticSwitch(), Meta.bOverride, Meta.ExpressionGuid);
 		}
 		break;
 	case EMaterialParameterType::StaticComponentMask:
-		EditorOnly.StaticComponentMaskParameters.Empty(Values.Num());
+		StaticComponentMaskParameters.Empty(Values.Num());
 		for (const auto& It : Values)
 		{
 			const FMaterialParameterMetadata& Meta = It.Value;
 			check(Meta.Value.Type == Type);
-			EditorOnly.StaticComponentMaskParameters.Emplace(It.Key,
+			StaticComponentMaskParameters.Emplace(It.Key,
 				Meta.Value.Bool[0],
 				Meta.Value.Bool[1],
 				Meta.Value.Bool[2],
@@ -507,7 +454,7 @@ void FStaticParameterSet::AddParametersOfType(EMaterialParameterType Type, const
 
 void FStaticParameterSet::SetStaticSwitchParameterValue(const FMaterialParameterInfo& ParameterInfo, const FGuid& ExpressionGuid, bool Value)
 {
-	for (FStaticSwitchParameter& Parameter : EditorOnly.StaticSwitchParameters)
+	for (FStaticSwitchParameter& Parameter : StaticSwitchParameters)
 	{
 		if (Parameter.ParameterInfo == ParameterInfo)
 		{
@@ -517,12 +464,12 @@ void FStaticParameterSet::SetStaticSwitchParameterValue(const FMaterialParameter
 		}
 	}
 
-	new(EditorOnly.StaticSwitchParameters) FStaticSwitchParameter(ParameterInfo, Value, true, ExpressionGuid);
+	new(StaticSwitchParameters) FStaticSwitchParameter(ParameterInfo, Value, true, ExpressionGuid);
 }
 
 void FStaticParameterSet::SetStaticComponentMaskParameterValue(const FMaterialParameterInfo& ParameterInfo, const FGuid& ExpressionGuid, bool R, bool G, bool B, bool A)
 {
-	for (FStaticComponentMaskParameter& Parameter : EditorOnly.StaticComponentMaskParameters)
+	for (FStaticComponentMaskParameter& Parameter : StaticComponentMaskParameters)
 	{
 		if (Parameter.ParameterInfo == ParameterInfo)
 		{
@@ -535,11 +482,10 @@ void FStaticParameterSet::SetStaticComponentMaskParameterValue(const FMaterialPa
 		}
 	}
 
-	new(EditorOnly.StaticComponentMaskParameters) FStaticComponentMaskParameter(ParameterInfo, R, G, B, A, true, ExpressionGuid);
+	new(StaticComponentMaskParameters) FStaticComponentMaskParameter(ParameterInfo, R, G, B, A, true, ExpressionGuid);
 }
 #endif // WITH_EDITORONLY_DATA
 
-
 void FMaterialShaderMapId::Serialize(FArchive& Ar, bool bLoadedByCookedMaterial)
 {
 	SCOPED_LOADTIMER(FMaterialShaderMapId_Serialize);
@@ -949,12 +895,12 @@ void FMaterialShaderMapId::UpdateFromParameterSet(const FStaticParameterSet& Sta
 		}
 	};
 
-	StaticSwitchParameters = StaticParameters.EditorOnly.StaticSwitchParameters;
-	StaticComponentMaskParameters = StaticParameters.EditorOnly.StaticComponentMaskParameters;
-	TerrainLayerWeightParameters = StaticParameters.EditorOnly.TerrainLayerWeightParameters;
+	StaticSwitchParameters = StaticParameters.StaticSwitchParameters;
+	StaticComponentMaskParameters = StaticParameters.StaticComponentMaskParameters;
+	TerrainLayerWeightParameters = StaticParameters.TerrainLayerWeightParameters;
 	if (StaticParameters.bHasMaterialLayers)
 	{
-		MaterialLayersId = StaticParameters.MaterialLayers.GetID(StaticParameters.EditorOnly.MaterialLayers);
+		MaterialLayersId = StaticParameters.MaterialLayers.GetID();
 	}
 
 	// Sort the arrays by parameter name, ensure the ID is not influenced by the order
diff --git a/Engine/Source/Runtime/Engine/Private/Materials/MaterialShared.cpp b/Engine/Source/Runtime/Engine/Private/Materials/MaterialShared.cpp
--- a/Engine/Source/Runtime/Engine/Private/Materials/MaterialShared.cpp
+++ b/Engine/Source/Runtime/Engine/Private/Materials/MaterialShared.cpp
@@ -361,17 +361,21 @@ void FExpressionInput::Connect( int32 InOutputIndex, class UMaterialExpression*
 {
 	InExpression->ConnectExpression(this, InOutputIndex);
 }
+#endif // WITH_EDITOR
 
 FExpressionInput FExpressionInput::GetTracedInput() const
 {
+#if WITH_EDITORONLY_DATA
 	if (Expression != nullptr && Expression->IsA(UMaterialExpressionRerouteBase::StaticClass()))
 	{
 		UMaterialExpressionRerouteBase* Reroute = CastChecked<UMaterialExpressionRerouteBase>(Expression);
 		return Reroute->TraceInputsToRealInput();
 	}
+#endif
 	return *this;
 }
 
+#if WITH_EDITOR
 int32 FExpressionExecOutput::Compile(class FMaterialCompiler* Compiler) const
 {
 	if (Expression)
@@ -451,8 +455,12 @@ static bool SerializeExpressionInput(FArchive& Ar, FExpressionInput& Input)
 		return false;
 	}
 
-	Ar << Input.Expression;
-
+#if WITH_EDITORONLY_DATA
+	if (!Ar.IsFilterEditorOnly())
+	{
+		Ar << Input.Expression;
+	}
+#endif
 	Ar << Input.OutputIndex;
 	if (Ar.CustomVer(FFrameworkObjectVersion::GUID) >= FFrameworkObjectVersion::PinsStoreFName)
 	{
@@ -465,11 +473,29 @@ static bool SerializeExpressionInput(FArchive& Ar, FExpressionInput& Input)
 		Input.InputName = *InputNameStr;
 	}
 
+#if WITH_EDITORONLY_DATA
 	Ar << Input.Mask;
 	Ar << Input.MaskR;
 	Ar << Input.MaskG;
 	Ar << Input.MaskB;
 	Ar << Input.MaskA;
+#else
+	int32 Temp = 0;
+	Ar << Temp << Temp << Temp << Temp << Temp;
+#endif
+
+	// Some expressions may have been stripped when cooking and Expression can be null after loading
+	// so make sure we keep the information about the connected node in cooked packages
+	if ( Ar.IsFilterEditorOnly() )
+	{
+#if WITH_EDITORONLY_DATA
+		if (Ar.IsSaving())
+		{
+			Input.ExpressionName = Input.Expression ? Input.Expression->GetFName() : NAME_None;
+		}
+#endif // WITH_EDITORONLY_DATA
+		Ar << Input.ExpressionName;
+	}
 
 	return true;
 }
@@ -479,10 +505,17 @@ static bool SerializeMaterialInput(FArchive& Ar, FMaterialInput<InputType>& Inpu
 {
 	if (SerializeExpressionInput(Ar, Input))
 	{
+#if WITH_EDITORONLY_DATA
 		bool bUseConstantValue = Input.UseConstant;
 		Ar << bUseConstantValue;
 		Input.UseConstant = bUseConstantValue;
 		Ar << Input.Constant;
+#else
+		bool bTemp = false;
+		Ar << bTemp;
+		InputType TempType;
+		Ar << TempType;
+#endif
 		return true;
 	}
 	else
@@ -1353,7 +1386,7 @@ void FMaterial::DiscardShaderMap()
 }
 
 EMaterialDomain FMaterialResource::GetMaterialDomain() const { return Material->MaterialDomain; }
-bool FMaterialResource::IsTangentSpaceNormal() const { return Material->bTangentSpaceNormal || (!Material->IsPropertyConnected(MP_Normal) && !Material->bUseMaterialAttributes); }
+bool FMaterialResource::IsTangentSpaceNormal() const { return Material->bTangentSpaceNormal || (!Material->Normal.IsConnected() && !Material->bUseMaterialAttributes); }
 bool FMaterialResource::ShouldGenerateSphericalParticleNormals() const { return Material->bGenerateSphericalParticleNormals; }
 bool FMaterialResource::ShouldDisableDepthTest() const { return Material->bDisableDepthTest; }
 bool FMaterialResource::ShouldWriteOnlyAlpha() const { return Material->bWriteOnlyAlpha; }
@@ -1369,9 +1402,7 @@ bool FMaterialResource::IsVolumetricPrimitive() const { return Material->Materia
 bool FMaterialResource::IsSpecialEngineMaterial() const { return Material->bUsedAsSpecialEngineMaterial; }
 bool FMaterialResource::HasVertexPositionOffsetConnected() const { return Material->HasVertexPositionOffsetConnected(); }
 bool FMaterialResource::HasPixelDepthOffsetConnected() const { return Material->HasPixelDepthOffsetConnected(); }
-#if WITH_EDITOR
-bool FMaterialResource::HasMaterialAttributesConnected() const { return (Material->bUseMaterialAttributes && Material->GetEditorOnlyData()->MaterialAttributes.IsConnected()); }
-#endif
+bool FMaterialResource::HasMaterialAttributesConnected() const { return (Material->bUseMaterialAttributes && Material->MaterialAttributes.IsConnected()); }
 EMaterialShadingRate FMaterialResource::GetShadingRate() const { return Material->ShadingRate; }
 FString FMaterialResource::GetBaseMaterialPathName() const { return Material->GetPathName(); }
 FString FMaterialResource::GetDebugName() const
@@ -1564,7 +1595,9 @@ bool FMaterialResource::IsUsingPlanarForwardReflections() const
 
 bool FMaterialResource::IsNonmetal() const
 {
-	return (!Material->IsPropertyConnected(MP_Metallic) && !Material->IsPropertyConnected(MP_Specular));
+	return !Material->bUseMaterialAttributes ?
+			(!Material->Metallic.IsConnected() && !Material->Specular.IsConnected()) :
+			!(Material->MaterialAttributes.IsConnected(MP_Specular) || Material->MaterialAttributes.IsConnected(MP_Metallic));
 }
 
 bool FMaterialResource::UseLmDirectionality() const
@@ -1710,8 +1743,7 @@ bool FMaterialResource::HasSpecularConnected() const
 
 bool FMaterialResource::HasMetallicConnected() const
 {
-	// HasMaterialAttributesConnected() should be captured by HasMetallicConnected() now
-	return Material->HasMetallicConnected();
+	return HasMaterialAttributesConnected() || Material->HasMetallicConnected();
 }
 
 bool FMaterialResource::HasEmissiveColorConnected() const
@@ -1755,10 +1787,9 @@ bool FMaterialResource::HasMaterialPropertyConnected(EMaterialProperty In) const
 			if (!CachedStrataMaterialInfo.IsValid())
 			{
 				check(Material->HasStrataFrontMaterialConnected());
-				const UMaterialEditorOnlyData* MaterialEditorOnly = Material->GetEditorOnlyData();
-				if (MaterialEditorOnly->FrontMaterial.Expression->IsResultStrataMaterial(MaterialEditorOnly->FrontMaterial.OutputIndex))
+				if (Material->FrontMaterial.Expression->IsResultStrataMaterial(Material->FrontMaterial.OutputIndex))
 				{
-					MaterialEditorOnly->FrontMaterial.Expression->GatherStrataMaterialInfo(CachedStrataMaterialInfo, MaterialEditorOnly->FrontMaterial.OutputIndex);
+					Material->FrontMaterial.Expression->GatherStrataMaterialInfo(CachedStrataMaterialInfo, Material->FrontMaterial.OutputIndex);
 				}
 			}
 		}
@@ -1777,7 +1808,7 @@ bool FMaterialResource::HasMaterialPropertyConnected(EMaterialProperty In) const
 		case MP_Normal: 			return Material->HasNormalConnected();
 		case MP_Roughness: 			return Material->HasRoughnessConnected();
 		case MP_Specular: 			return Material->HasSpecularConnected();
-		case MP_Metallic: 			return Material->HasMetallicConnected(); // HasMetallicConnected() should be properly capturing HasMaterialAttributesConnected() as well
+		case MP_Metallic: 			return HasMaterialAttributesConnected() || Material->HasMetallicConnected();
 		case MP_Anisotropy: 		return Material->HasAnisotropyConnected();
 		case MP_AmbientOcclusion: 	return Material->HasAmbientOcclusionConnected();
 		}
@@ -4575,6 +4606,7 @@ void DoMaterialAttributeReorder(FExpressionInput* Input, const FPackageFileVersi
 			if (Input->OutputIndex == 11 || Input->OutputIndex == 12)
 			{
 				Input->Expression = nullptr;
+				Input->ExpressionName = NAME_None;
 			}
 			else if (Input->OutputIndex >= 13)
 			{
diff --git a/Engine/Source/Runtime/Engine/Public/MaterialCachedData.h b/Engine/Source/Runtime/Engine/Public/MaterialCachedData.h
--- a/Engine/Source/Runtime/Engine/Public/MaterialCachedData.h
+++ b/Engine/Source/Runtime/Engine/Public/MaterialCachedData.h
@@ -3,7 +3,6 @@
 #pragma once
 
 #include "CoreMinimal.h"
-#include "Templates/SharedPointer.h"
 #include "MaterialTypes.h"
 #include "Materials/MaterialLayersFunctions.h"
 #include "Misc/Guid.h"
@@ -74,11 +73,10 @@
 
 	FMaterialCachedParameterEditorInfo() = default;
 	FMaterialCachedParameterEditorInfo(const FGuid& InGuid) : ExpressionGuid(InGuid) {}
-	FMaterialCachedParameterEditorInfo(const FGuid& InGuid, const FString& InDescription, const FName& InGroup, int32 InSortPriority, int32 InAssetIndex)
+	FMaterialCachedParameterEditorInfo(const FGuid& InGuid, const FString& InDescription, const FName& InGroup, int32 InSortPriority)
 		: Description(InDescription)
 		, Group(InGroup)
 		, SortPriority(InSortPriority)
-		, AssetIndex(InAssetIndex)
 		, ExpressionGuid(InGuid)
 	{}
 
@@ -92,8 +90,5 @@
 	int32 SortPriority = 0;
 
 	UPROPERTY()
-	int32 AssetIndex = INDEX_NONE;
-
-	UPROPERTY()
 	FGuid ExpressionGuid;
 };
@@ -105,40 +100,88 @@
 
 	ENGINE_API static const FMaterialCachedParameterEntry EmptyData;
 
+	void Reset();
+
 	// This is used to map FMaterialParameterInfos to indices, which are then used to index various TArrays containing values for each type of parameter
 	// (ExpressionGuids and Overrides, along with ScalarValues, VectorValues, etc)
 	UPROPERTY()
 	TSet<FMaterialParameterInfo> ParameterInfoSet;
-};
 
-USTRUCT()
-struct FMaterialCachedParameterEditorEntry
-{
-	GENERATED_USTRUCT_BODY()
-
+#if WITH_EDITORONLY_DATA
+	// Editor-only information for each parameter
 	UPROPERTY()
 	TArray<FMaterialCachedParameterEditorInfo> EditorInfo;
-};
+#endif // WITH_EDITORONLY_DATA
 
-struct FMaterialCachedExpressionContext
-{
-	const UMaterialFunctionInterface* CurrentFunction = nullptr;
-	const FMaterialLayersFunctions* LayerOverrides = nullptr;
-	bool bUpdateFunctionExpressions = true;
 };
 
 USTRUCT()
-struct FMaterialCachedExpressionEditorOnlyData
+struct FMaterialCachedParameters
 {
 	GENERATED_USTRUCT_BODY()
 
-	ENGINE_API static const FMaterialCachedExpressionEditorOnlyData EmptyData;
+#if WITH_EDITORONLY_DATA
+	inline const FMaterialCachedParameterEntry& GetParameterTypeEntry(EMaterialParameterType Type) const { return Type >= EMaterialParameterType::NumRuntime ? EditorOnlyEntries[static_cast<int32>(Type) - NumMaterialRuntimeParameterTypes] : RuntimeEntries[static_cast<int32>(Type)]; }
+	inline FMaterialCachedParameterEntry& GetParameterTypeEntry(EMaterialParameterType Type) { return Type >= EMaterialParameterType::NumRuntime ? EditorOnlyEntries[static_cast<int32>(Type) - NumMaterialRuntimeParameterTypes] : RuntimeEntries[static_cast<int32>(Type)]; }
+#else
+	inline const FMaterialCachedParameterEntry& GetParameterTypeEntry(EMaterialParameterType Type) const { return Type >= EMaterialParameterType::NumRuntime ? FMaterialCachedParameterEntry::EmptyData : RuntimeEntries[static_cast<int32>(Type)]; }
+#endif
+
+#if WITH_EDITOR
+	void AddParameter(const FMaterialParameterInfo& ParameterInfo, const FMaterialParameterMetadata& ParameterMeta, UObject*& OutReferencedTexture);
+	inline void AddParameter(const FMaterialParameterInfo& ParameterInfo, const FMaterialParameterMetadata& ParameterMeta)
+	{
+		UObject* UnusedReferencedTexture = nullptr;
+		AddParameter(ParameterInfo, ParameterMeta, UnusedReferencedTexture);
+	}
+#endif // WITH_EDITOR
+
+	inline int32 GetNumParameters(EMaterialParameterType Type) const { return GetParameterTypeEntry(Type).ParameterInfoSet.Num(); }
+	int32 FindParameterIndex(EMaterialParameterType Type, const FMemoryImageMaterialParameterInfo& HashedParameterInfo) const;
+	bool GetParameterValue(EMaterialParameterType Type, const FMemoryImageMaterialParameterInfo& ParameterInfo, FMaterialParameterMetadata& OutResult) const;
+	void GetParameterValueByIndex(EMaterialParameterType Type, int32 ParameterIndex, FMaterialParameterMetadata& OutResult) const;
+#if WITH_EDITORONLY_DATA
+	const FGuid& GetExpressionGuid(EMaterialParameterType Type, int32 Index) const;
+#endif
+	void GetAllParametersOfType(EMaterialParameterType Type, TMap<FMaterialParameterInfo, FMaterialParameterMetadata>& OutParameters) const;
+	void GetAllParameterInfoOfType(EMaterialParameterType Type, TArray<FMaterialParameterInfo>& OutParameterInfo, TArray<FGuid>& OutParameterIds) const;
+	void GetAllGlobalParametersOfType(EMaterialParameterType Type, TMap<FMaterialParameterInfo, FMaterialParameterMetadata>& OutParameters) const;
+	void GetAllGlobalParameterInfoOfType(EMaterialParameterType Type, TArray<FMaterialParameterInfo>& OutParameterInfo, TArray<FGuid>& OutParameterIds) const;
+	void Reset();
 
 	UPROPERTY()
+	FMaterialCachedParameterEntry RuntimeEntries[NumMaterialRuntimeParameterTypes];
+
+	UPROPERTY()
+	TArray<int32> ScalarPrimitiveDataIndexValues;
+
+	UPROPERTY()
+	TArray<int32> VectorPrimitiveDataIndexValues;
+
+	UPROPERTY()
+	TArray<float> ScalarValues;
+
+	UPROPERTY()
+	TArray<FLinearColor> VectorValues;
+
+	UPROPERTY()
+	TArray<FVector4d> DoubleVectorValues;
+
+	UPROPERTY()
+	TArray<TSoftObjectPtr<UTexture>> TextureValues;
+
+	UPROPERTY()
+	TArray<TSoftObjectPtr<UFont>> FontValues;
+
+	UPROPERTY()
+	TArray<int32> FontPageValues;
+
+	UPROPERTY()
+	TArray<TSoftObjectPtr<URuntimeVirtualTexture>> RuntimeVirtualTextureValues;
+
+#if WITH_EDITORONLY_DATA
+	UPROPERTY()
 	FMaterialCachedParameterEntry EditorOnlyEntries[NumMaterialEditorOnlyParameterTypes];
 
 	UPROPERTY()
-	FMaterialCachedParameterEditorEntry EditorEntries[NumMaterialParameterTypes];
-
-	UPROPERTY()
 	TArray<bool> StaticSwitchValues;
@@ -163,15 +206,14 @@
 
 	UPROPERTY()
 	TArray<FParameterChannelNames> TextureChannelNameValues;
-
-	UPROPERTY()
-	FMaterialLayersFunctionsEditorOnlyData MaterialLayers;
-
-	UPROPERTY()
-	TArray<FString> AssetPaths;
-
-	UPROPERTY()
-	TArray<FName> LandscapeLayerNames;
+#endif // WITH_EDITORONLY_DATA
 };
 
+struct FMaterialCachedExpressionContext
+{
+	const UMaterialFunctionInterface* CurrentFunction = nullptr;
+	const FMaterialLayersFunctions* LayerOverrides = nullptr;
+	bool bUpdateFunctionExpressions = true;
+};
+
 USTRUCT()
@@ -181,62 +223,23 @@
 	
 	ENGINE_API static const FMaterialCachedExpressionData EmptyData;
 
-	ENGINE_API FMaterialCachedExpressionData();
+	FMaterialCachedExpressionData()
+		: bHasMaterialLayers(false)
+		, bHasRuntimeVirtualTextureOutput(false)
+		, bHasSceneColor(false)
+		, bHasPerInstanceCustomData(false)
+		, bHasPerInstanceRandom(false)
+		, bHasVertexInterpolator(false)
+	{}
 
 #if WITH_EDITOR
-	void UpdateForExpressions(const FMaterialCachedExpressionContext& Context, TConstArrayView<TObjectPtr<UMaterialExpression>> Expressions, EMaterialParameterAssociation Association, int32 ParameterIndex);
+	void UpdateForExpressions(const FMaterialCachedExpressionContext& Context, const TArray<TObjectPtr<UMaterialExpression>>& Expressions, EMaterialParameterAssociation Association, int32 ParameterIndex);
 	void UpdateForFunction(const FMaterialCachedExpressionContext& Context, UMaterialFunctionInterface* Function, EMaterialParameterAssociation Association, int32 ParameterIndex);
 	void UpdateForLayerFunctions(const FMaterialCachedExpressionContext& Context, const FMaterialLayersFunctions& LayerFunctions);
 	ENGINE_API void UpdateForCachedHLSLTree(const FMaterialCachedHLSLTree& CachedTree, const FStaticParameterSet* StaticParameters);
-	void Validate();
-
-	void AddParameter(const FMaterialParameterInfo& ParameterInfo, const FMaterialParameterMetadata& ParameterMeta, UObject*& OutReferencedTexture);
-	inline void AddParameter(const FMaterialParameterInfo& ParameterInfo, const FMaterialParameterMetadata& ParameterMeta)
-	{
-		UObject* UnusedReferencedTexture = nullptr;
-		AddParameter(ParameterInfo, ParameterMeta, UnusedReferencedTexture);
-	}
 #endif // WITH_EDITOR
 
-#if WITH_EDITORONLY_DATA
-	inline FMaterialCachedParameterEntry& GetParameterTypeEntry(EMaterialParameterType Type)
-	{
-		const int32 TypeIndex = (int32)Type;
-		if (TypeIndex < NumMaterialRuntimeParameterTypes)
-		{
-			return RuntimeEntries[TypeIndex];
-		}
-		check(EditorOnlyData);
-		return EditorOnlyData->EditorOnlyEntries[TypeIndex - NumMaterialRuntimeParameterTypes];
-	}
-	inline FMaterialCachedParameterEditorEntry& GetParameterEditorTypeEntry(EMaterialParameterType Type)
-	{
-		check(EditorOnlyData);
-		return EditorOnlyData->EditorEntries[(int32)Type];
-	}
+	void Reset();
 
-	inline const FMaterialCachedParameterEntry& GetParameterTypeEntry(EMaterialParameterType Type) const
-	{
-		return const_cast<FMaterialCachedExpressionData*>(this)->GetParameterTypeEntry(Type);
-	}
-	inline const FMaterialCachedParameterEditorEntry& GetParameterEditorTypeEntry(EMaterialParameterType Type) const
-	{
-		check(EditorOnlyData);
-		return EditorOnlyData->EditorEntries[(int32)Type];
-	}
-#else
-	inline const FMaterialCachedParameterEntry& GetParameterTypeEntry(EMaterialParameterType Type) const { return Type >= EMaterialParameterType::NumRuntime ? FMaterialCachedParameterEntry::EmptyData : RuntimeEntries[static_cast<int32>(Type)]; }
-#endif
-
-	inline int32 GetNumParameters(EMaterialParameterType Type) const { return GetParameterTypeEntry(Type).ParameterInfoSet.Num(); }
-	int32 FindParameterIndex(EMaterialParameterType Type, const FMemoryImageMaterialParameterInfo& HashedParameterInfo) const;
-	bool GetParameterValue(EMaterialParameterType Type, const FMemoryImageMaterialParameterInfo& ParameterInfo, FMaterialParameterMetadata& OutResult) const;
-	void GetParameterValueByIndex(EMaterialParameterType Type, int32 ParameterIndex, FMaterialParameterMetadata& OutResult) const;
-	const FGuid& GetExpressionGuid(EMaterialParameterType Type, int32 Index) const;
-	void GetAllParametersOfType(EMaterialParameterType Type, TMap<FMaterialParameterInfo, FMaterialParameterMetadata>& OutParameters) const;
-	void GetAllParameterInfoOfType(EMaterialParameterType Type, TArray<FMaterialParameterInfo>& OutParameterInfo, TArray<FGuid>& OutParameterIds) const;
-	void GetAllGlobalParametersOfType(EMaterialParameterType Type, TMap<FMaterialParameterInfo, FMaterialParameterMetadata>& OutParameters) const;
-	void GetAllGlobalParameterInfoOfType(EMaterialParameterType Type, TArray<FMaterialParameterInfo>& OutParameterInfo, TArray<FGuid>& OutParameterIds) const;
-
 	ENGINE_API void AddReferencedObjects(FReferenceCollector& Collector);
 
@@ -246,49 +249,18 @@
 	/** Returns an array of the guids of parameter collections used. */
 	void AppendReferencedParameterCollectionIdsTo(TArray<FGuid>& OutIds) const;
 
-	bool IsPropertyConnected(EMaterialProperty Property) const
+	bool IsMaterialAttributePropertyConnected(EMaterialProperty Property) const
 	{
-		return ((PropertyConnectedBitmask >> (uint32)Property) & 0x1) != 0;
+		return ((MaterialAttributesPropertyConnectedBitmask >> (uint32)Property) & 0x1) != 0;
 	}
 
-	void SetPropertyConnected(EMaterialProperty Property)
+	void SetMaterialAttributePropertyConnected(EMaterialProperty Property, bool bIsConnected)
 	{
-		PropertyConnectedBitmask |= (1 << (uint32)Property);
+		MaterialAttributesPropertyConnectedBitmask = bIsConnected ? MaterialAttributesPropertyConnectedBitmask | (1 << (uint32)Property) : MaterialAttributesPropertyConnectedBitmask & ~(1 << (uint32)Property);
 	}
 
-#if WITH_EDITORONLY_DATA
-	TSharedPtr<FMaterialCachedExpressionEditorOnlyData> EditorOnlyData;
-#endif // WITH_EDITORONLY_DATA
+	UPROPERTY()
+	FMaterialCachedParameters Parameters;
 
-	UPROPERTY()
-	FMaterialCachedParameterEntry RuntimeEntries[NumMaterialRuntimeParameterTypes];
-
-	UPROPERTY()
-	TArray<int32> ScalarPrimitiveDataIndexValues;
-
-	UPROPERTY()
-	TArray<int32> VectorPrimitiveDataIndexValues;
-
-	UPROPERTY()
-	TArray<float> ScalarValues;
-
-	UPROPERTY()
-	TArray<FLinearColor> VectorValues;
-
-	UPROPERTY()
-	TArray<FVector4d> DoubleVectorValues;
-
-	UPROPERTY()
-	TArray<TSoftObjectPtr<UTexture>> TextureValues;
-
-	UPROPERTY()
-	TArray<TSoftObjectPtr<UFont>> FontValues;
-
-	UPROPERTY()
-	TArray<int32> FontPageValues;
-
-	UPROPERTY()
-	TArray<TSoftObjectPtr<URuntimeVirtualTexture>> RuntimeVirtualTextureValues;
-
 	/** Array of all texture referenced by this material */
 	UPROPERTY()
@@ -303,10 +275,10 @@
 	TArray<FMaterialParameterCollectionInfo> ParameterCollectionInfos;
 
 	UPROPERTY()
+	FMaterialLayersFunctions MaterialLayers;
+
+	UPROPERTY()
 	TArray<TObjectPtr<ULandscapeGrassType>> GrassTypes;
 
 	UPROPERTY()
-	FMaterialLayersFunctionsRuntimeData MaterialLayers;
-
-	UPROPERTY()
 	TArray<FName> DynamicParameterNames;
@@ -334,6 +306,11 @@
 
 	/** Each bit corresponds to EMaterialProperty connection status. */
 	UPROPERTY()
-	uint32 PropertyConnectedBitmask = 0;
+	uint32 MaterialAttributesPropertyConnectedBitmask = 0;
+
+#if WITH_EDITORONLY_DATA
+	UPROPERTY()
+	TArray<FName> LandscapeLayerNames;
+#endif // WITH_EDITORONLY_DATA
 };
 
diff --git a/Engine/Source/Runtime/Engine/Public/MaterialExpressionIO.h b/Engine/Source/Runtime/Engine/Public/MaterialExpressionIO.h
--- a/Engine/Source/Runtime/Engine/Public/MaterialExpressionIO.h
+++ b/Engine/Source/Runtime/Engine/Public/MaterialExpressionIO.h
@@ -30,11 +30,13 @@ struct FValue;
 //@warning: FExpressionInput is mirrored in MaterialExpression.h and manually "subclassed" in Material.h (FMaterialInput)
 struct FExpressionInput
 {
+#if WITH_EDITORONLY_DATA
 	/** 
 	 * Material expression that this input is connected to, or NULL if not connected. 
 	 * If you want to be safe when checking against dangling Reroute nodes, please use GetTracedInput before accessing this property.
 	*/
 	class UMaterialExpression*	Expression;
+#endif
 
 	/** 
 	 * Index into Expression's outputs array that this input is connected to.
@@ -48,26 +50,32 @@ struct FExpressionInput
 	 */
 	FName						InputName;
 
+#if WITH_EDITORONLY_DATA
 	int32						Mask,
 								MaskR,
 								MaskG,
 								MaskB,
 								MaskA;
+#endif
+
+	/** Material expression name that this input is connected to, or None if not connected. Used only in cooked builds */
+	FName				ExpressionName;
 
 	FExpressionInput()
 		: OutputIndex(0)
+#if WITH_EDITORONLY_DATA
 		, Mask(0)
 		, MaskR(0)
 		, MaskG(0)
 		, MaskB(0)
 		, MaskA(0)
+#endif
 	{
+#if WITH_EDITORONLY_DATA
 		Expression = nullptr;
+#endif
 	}
 
-	/** ICPPStructOps interface */
-	ENGINE_API bool Serialize(FArchive& Ar);
-
 #if WITH_EDITOR
 	ENGINE_API int32 Compile(class FMaterialCompiler* Compiler);
 	ENGINE_API const UE::HLSLTree::FExpression* TryAcquireHLSLExpression(FMaterialHLSLGenerator& Generator, UE::HLSLTree::FScope& Scope, int32 InputIndex) const;
@@ -79,6 +87,7 @@ struct FExpressionInput
 	ENGINE_API const UE::HLSLTree::FExpression* AcquireHLSLExpression(FMaterialHLSLGenerator& Generator, UE::HLSLTree::FScope& Scope) const;
 	ENGINE_API const UE::HLSLTree::FExpression* AcquireHLSLExpressionOrConstant(FMaterialHLSLGenerator& Generator, UE::HLSLTree::FScope& Scope, const UE::Shader::FValue& ConstantValue) const;
 	ENGINE_API const UE::HLSLTree::FExpression* AcquireHLSLExpressionOrExternalInput(FMaterialHLSLGenerator& Generator, UE::HLSLTree::FScope& Scope, UE::HLSLTree::Material::EExternalInput Input) const;
+#endif // WITH_EDITOR
 
 	/**
 	 * Tests if the input has a material expression connected to it
@@ -87,14 +96,20 @@ struct FExpressionInput
 	 */
 	bool IsConnected() const 
 	{ 
+#if WITH_EDITORONLY_DATA
 		return (nullptr != Expression);
+#else
+		return ExpressionName != NAME_None;
+#endif // WITH_EDITORONLY_DATA
 	}
 
-	/** Is the input connected to a constant value */
-	bool IsConstant() const { return false; }
-
+#if WITH_EDITOR
 	/** Connects output of InExpression to this input */
 	ENGINE_API void Connect( int32 InOutputIndex, class UMaterialExpression* InExpression );
+#endif // WITH_EDITOR
+
+	/** ICPPStructOps interface */
+	ENGINE_API bool Serialize(FArchive& Ar);
 
 	/** If this input goes through reroute nodes or other paths that should not affect code, trace back on the input chain.*/
 	ENGINE_API FExpressionInput GetTracedInput() const;
@@ -102,13 +117,14 @@ struct FExpressionInput
 	/** Helper for setting component mask. */
 	ENGINE_API void SetMask(int32 UseMask, int32 R, int32 G, int32 B, int32 A)
 	{
+#if WITH_EDITORONLY_DATA
 		Mask = UseMask;
 		MaskR = R;
 		MaskG = G;
 		MaskB = B;
 		MaskA = A;
+#endif
 	}
-#endif // WITH_EDITOR
 };
 
 template<>
@@ -128,37 +144,45 @@ struct TStructOpsTypeTraits<FExpressionInput>
 struct FExpressionOutput
 {
 	FName	OutputName;
+#if WITH_EDITORONLY_DATA
 	int32	Mask,
 		MaskR,
 		MaskG,
 		MaskB,
 		MaskA;
+#endif
 
 	FExpressionOutput(int32 InMask = 0, int32 InMaskR = 0, int32 InMaskG = 0, int32 InMaskB = 0, int32 InMaskA = 0)
+#if WITH_EDITORONLY_DATA
 		: Mask(InMask)
 		, MaskR(InMaskR)
 		, MaskG(InMaskG)
 		, MaskB(InMaskB)
 		, MaskA(InMaskA)
+#endif
 	{}
 
 	FExpressionOutput(FName InOutputName, int32 InMask = 0, int32 InMaskR = 0, int32 InMaskG = 0, int32 InMaskB = 0, int32 InMaskA = 0)
 		: OutputName(InOutputName)
+#if WITH_EDITORONLY_DATA
 		, Mask(InMask)
 		, MaskR(InMaskR)
 		, MaskG(InMaskG)
 		, MaskB(InMaskB)
 		, MaskA(InMaskA)
+#endif
 	{}
 
 	/** Helper for setting component mask. */
 	ENGINE_API void SetMask(int32 UseMask, int32 R, int32 G, int32 B, int32 A)
 	{
+#if WITH_EDITORONLY_DATA
 		Mask = UseMask;
 		MaskR = R;
 		MaskG = G;
 		MaskB = B;
 		MaskA = A;
+#endif
 	}
 };
 
@@ -170,16 +194,16 @@ template<class InputType> struct FMaterialInput : FExpressionInput
 {
 	FMaterialInput()
 	{
+#if WITH_EDITORONLY_DATA
 		UseConstant = 0;
 		Constant = InputType(0);
-	}
-
-#if WITH_EDITOR
-	bool IsConstant() const { return UseConstant; }
 #endif
+	}
 
+#if WITH_EDITORONLY_DATA
 	uint32	UseConstant : 1;
 	InputType	Constant;
+#endif
 };
 
 struct FColorMaterialInput : FMaterialInput<FColor>
@@ -307,13 +331,13 @@ struct FMaterialAttributesInput : FExpressionInput
 
 #if WITH_EDITOR
 	ENGINE_API int32 CompileWithDefault(class FMaterialCompiler* Compiler, const FGuid& AttributeID);
-	inline bool IsConnected(EMaterialProperty Property) { return ((PropertyConnectedBitmask >> (uint32)Property) & 0x1) != 0; }
-	inline bool IsConnected() const { return FExpressionInput::IsConnected(); }
-	inline void SetConnectedProperty(EMaterialProperty Property, bool bIsConnected)
+#endif  // WITH_EDITOR
+	ENGINE_API bool IsConnected(EMaterialProperty Property) { return ((PropertyConnectedBitmask >> (uint32)Property) & 0x1) != 0; }
+	ENGINE_API bool IsConnected() const { return FExpressionInput::IsConnected(); }
+	ENGINE_API void SetConnectedProperty(EMaterialProperty Property, bool bIsConnected) 
 	{
 		PropertyConnectedBitmask = bIsConnected ? PropertyConnectedBitmask | (1 << (uint32)Property) : PropertyConnectedBitmask & ~(1 << (uint32)Property);
 	}
-#endif  // WITH_EDITOR
 
 	/** ICPPStructOps interface */
 	ENGINE_API bool Serialize(FArchive& Ar);
diff --git a/Engine/Source/Runtime/Engine/Public/MaterialHLSLTree.h b/Engine/Source/Runtime/Engine/Public/MaterialHLSLTree.h
--- a/Engine/Source/Runtime/Engine/Public/MaterialHLSLTree.h
+++ b/Engine/Source/Runtime/Engine/Public/MaterialHLSLTree.h
@@ -326,12 +326,12 @@ public:
 class FExpressionMaterialLayers : public FExpressionForward
 {
 public:
-	FExpressionMaterialLayers(const FExpression* InExpression, const FMaterialLayersFunctions& InMaterialLayers)
+	FExpressionMaterialLayers(const FExpression* InExpression, const FMaterialLayersFunctions* InMaterialLayers)
 		: FExpressionForward(InExpression)
 		, MaterialLayers(InMaterialLayers)
 	{}
 
-	FMaterialLayersFunctions MaterialLayers;
+	const FMaterialLayersFunctions* MaterialLayers;
 
 	virtual bool PrepareValue(FEmitContext& Context, FEmitScope& Scope, const FRequestedType& RequestedType, FPrepareValueResult& OutResult) const override;
 };
diff --git a/Engine/Source/Runtime/Engine/Public/MaterialTypes.h b/Engine/Source/Runtime/Engine/Public/MaterialTypes.h
--- a/Engine/Source/Runtime/Engine/Public/MaterialTypes.h
+++ b/Engine/Source/Runtime/Engine/Public/MaterialTypes.h
@@ -410,9 +410,6 @@ struct FMaterialParameterMetadata
 	/** Description of the parameter, typically taken from the 'Desc' field of the parameter's UMaterialExpression */
 	FString Description;
 
-	/** Path to the asset (UMaterial or UMaterialFunction) that introduced this parameter */
-	FString AssetPath;
-
 	/** Name of the parameter's group */
 	FName Group;
 
diff --git a/Engine/Source/Runtime/Engine/Public/StaticParameterSet.h b/Engine/Source/Runtime/Engine/Public/StaticParameterSet.h
--- a/Engine/Source/Runtime/Engine/Public/StaticParameterSet.h
+++ b/Engine/Source/Runtime/Engine/Public/StaticParameterSet.h
@@ -15,7 +15,6 @@
 #include "StaticParameterSet.generated.h"
 
 class FSHA1;
-struct FStaticParameterSet;
 
 /**
 Base parameter properties
@@ -376,40 +375,11 @@ struct FStaticMaterialLayersParameter : public FStaticParameterBase
 #endif // WITH_EDITOR
 };
 
+/** Contains all the information needed to identify a single permutation of static parameters. */
 USTRUCT()
-struct FStaticParameterSetRuntimeData
-{
-	GENERATED_BODY()
-
-	FStaticParameterSetRuntimeData() : bHasMaterialLayers(false) {}
-	FStaticParameterSetRuntimeData(const FStaticParameterSetRuntimeData& Rhs) = default;
-	FStaticParameterSetRuntimeData(const FStaticParameterSet& Rhs) = delete;
-	FStaticParameterSetRuntimeData& operator=(const FStaticParameterSetRuntimeData& Rhs) = default;
-	FStaticParameterSetRuntimeData& operator=(const FStaticParameterSet& Rhs) = delete;
-
-	/** Material layers for this set */
-	UPROPERTY()
-	FMaterialLayersFunctionsRuntimeData MaterialLayers;
-
-	UPROPERTY()
-	uint8 bHasMaterialLayers : 1;
-
-	void Empty()
-	{
-		MaterialLayers.Empty();
-		bHasMaterialLayers = false;
-	}
-
-	bool IsEmpty() const
-	{
-		return !bHasMaterialLayers;
-	}
-};
-
-USTRUCT()
-struct FStaticParameterSetEditorOnlyData
+struct FStaticParameterSet
 {
-	GENERATED_BODY()
+	GENERATED_USTRUCT_BODY();
 
 	/** An array of static switch parameters in this set */
 	UPROPERTY()
@@ -425,44 +395,14 @@ struct FStaticParameterSetEditorOnlyData
 
 	/** Material layers for this set */
 	UPROPERTY()
-	FMaterialLayersFunctionsEditorOnlyData MaterialLayers;
-
-#if WITH_EDITOR
-	void Empty()
-	{
-		StaticSwitchParameters.Empty();
-		StaticComponentMaskParameters.Empty();
-		TerrainLayerWeightParameters.Empty();
-		MaterialLayers.Empty();
-	}
+	FMaterialLayersFunctions MaterialLayers;
 
-	bool IsEmpty() const
-	{
-		return StaticSwitchParameters.Num() == 0 && StaticComponentMaskParameters.Num() == 0 && TerrainLayerWeightParameters.Num() == 0;
-	}
-#endif // WITH_EDITOR
-};
-
-/** Contains all the information needed to identify a single permutation of static parameters. */
-USTRUCT()
-struct FStaticParameterSet : public FStaticParameterSetRuntimeData
-{
-	GENERATED_USTRUCT_BODY();
-
-#if WITH_EDITORONLY_DATA
 	UPROPERTY()
-	FStaticParameterSetEditorOnlyData EditorOnly;
-#endif // WITH_EDITORONLY_DATA
+	uint8 bHasMaterialLayers : 1;
 
-	FStaticParameterSet() = default;
+	FStaticParameterSet() : bHasMaterialLayers(false) {}
 	ENGINE_API FStaticParameterSet(const FStaticParameterSet& InValue);
-	FStaticParameterSet(const FStaticParameterSetRuntimeData&) = delete;
-
 	ENGINE_API FStaticParameterSet& operator=(const FStaticParameterSet& InValue);
-	FStaticParameterSet& operator=(const FStaticParameterSetRuntimeData& InValue) = delete;
-
-	FStaticParameterSetRuntimeData& GetRuntime() { return *this; }
-	const FStaticParameterSetRuntimeData& GetRuntime() const { return *this; }
 
 	/** 
 	* Checks if this set contains any parameters
@@ -471,37 +411,17 @@ struct FStaticParameterSet : public FStaticParameterSetRuntimeData
 	*/
 	bool IsEmpty() const
 	{
-		if (!FStaticParameterSetRuntimeData::IsEmpty()) return false;
-#if WITH_EDITOR
-		if (!EditorOnly.IsEmpty()) return false;
-#endif
-		return true;
-	}
-
-	void Empty()
-	{
-		FStaticParameterSetRuntimeData::Empty();
-#if WITH_EDITOR
-		EditorOnly.Empty();
-#endif
+		return StaticSwitchParameters.Num() == 0 && StaticComponentMaskParameters.Num() == 0 && TerrainLayerWeightParameters.Num() == 0 && !bHasMaterialLayers;
 	}
 
-	bool GetMaterialLayers(FMaterialLayersFunctions& OutMaterialLayers) const;
+	void Empty();
 
 #if WITH_EDITOR
 	void SerializeLegacy(FArchive& Ar);
 	void UpdateLegacyTerrainLayerWeightData();
 	void UpdateLegacyMaterialLayersData();
-
-	static void Validate(const FStaticParameterSetRuntimeData& Runtime, const FStaticParameterSetEditorOnlyData& EditorOnly);
-	void Validate()
-	{
-		Validate(GetRuntime(), EditorOnly);
-	}
 #endif // WITH_EDITOR
 
-	void PostSerialize(const FArchive& Ar);
-
 	/** 
 	* Tests this set against another for equality
 	* 
@@ -528,18 +448,6 @@ private:
 	UPROPERTY()
 	TArray<FStaticMaterialLayersParameter> MaterialLayersParameters_DEPRECATED;
 	PRAGMA_ENABLE_DEPRECATION_WARNINGS
-
-	UPROPERTY()
-	FMaterialLayersFunctions MaterialLayers_DEPRECATED;
-
-	UPROPERTY()
-	TArray<FStaticSwitchParameter> StaticSwitchParameters_DEPRECATED;
-
-	UPROPERTY()
-	TArray<FStaticComponentMaskParameter> StaticComponentMaskParameters_DEPRECATED;
-
-	UPROPERTY()
-	TArray<FStaticTerrainLayerWeightParameter> TerrainLayerWeightParameters_DEPRECATED;
 #endif // WITH_EDITORONLY_DATA
 
 	void SortForEquivalent();
@@ -549,14 +457,3 @@ private:
 	void SetStaticComponentMaskParameterValue(const FMaterialParameterInfo& ParameterInfo, const FGuid& ExpressionGuid, bool R, bool G, bool B, bool A);
 #endif // WITH_EDITORONLY_DATA
 };
-
-bool operator==(const FStaticParameterSet&, const FStaticParameterSetRuntimeData&) = delete;
-bool operator==(const FStaticParameterSetRuntimeData&, const FStaticParameterSet&) = delete;
-bool operator!=(const FStaticParameterSet&, const FStaticParameterSetRuntimeData&) = delete;
-bool operator!=(const FStaticParameterSetRuntimeData&, const FStaticParameterSet&) = delete;
-
-template<>
-struct TStructOpsTypeTraits<FStaticParameterSet> : TStructOpsTypeTraitsBase2<FStaticParameterSet>
-{
-	enum { WithPostSerialize = true };
-};
diff --git a/Engine/Source/Runtime/Landscape/Private/LandscapeEdit.cpp b/Engine/Source/Runtime/Landscape/Private/LandscapeEdit.cpp
--- a/Engine/Source/Runtime/Landscape/Private/LandscapeEdit.cpp
+++ b/Engine/Source/Runtime/Landscape/Private/LandscapeEdit.cpp
@@ -201,7 +201,7 @@ ULandscapeMaterialInstanceConstant* ALandscapeProxy::GetLayerThumbnailMIC(UMater
 	MaterialInstance->GetStaticParameterValues(StaticParameters);
 
 	// Customize that material instance to only enable our terrain layer's weightmap : 
-	StaticParameters.EditorOnly.TerrainLayerWeightParameters.Add(FStaticTerrainLayerWeightParameter(LayerName, /*InWeightmapIndex = */0, /*bInWeightBasedBlend = */false));
+	StaticParameters.TerrainLayerWeightParameters.Add(FStaticTerrainLayerWeightParameter(LayerName, /*InWeightmapIndex = */0, /*bInWeightBasedBlend = */false));
 
 	// Don't recreate the render state of everything, only update the materials context
 	{
@@ -226,7 +226,7 @@ bool ULandscapeComponent::ValidateCombinationMaterial(UMaterialInstanceConstant*
 		return false;
 	}
 
-	const TArray<FStaticTerrainLayerWeightParameter>& TerrainLayerWeightParameters = InCombinationMaterial->GetEditorOnlyStaticParameters().TerrainLayerWeightParameters;
+	const TArray<FStaticTerrainLayerWeightParameter>& TerrainLayerWeightParameters = InCombinationMaterial->GetStaticParameters().TerrainLayerWeightParameters;
 	const TArray<FWeightmapLayerAllocationInfo>& ComponentWeightmapAllocations = GetWeightmapLayerAllocations();
 
 	if (TerrainLayerWeightParameters.Num() != ComponentWeightmapAllocations.Num())
@@ -361,7 +361,7 @@ UMaterialInstanceConstant* ULandscapeComponent::GetCombinationMaterial(FMaterial
 				{
 					FName LayerName = Allocation.GetLayerName();
 					check(LayerName != NAME_None);
-					StaticParameters.EditorOnly.TerrainLayerWeightParameters.Add(FStaticTerrainLayerWeightParameter(LayerName, Allocation.WeightmapTextureIndex, !Allocation.LayerInfo->bNoWeightBlend));
+					StaticParameters.TerrainLayerWeightParameters.Add(FStaticTerrainLayerWeightParameter(LayerName, Allocation.WeightmapTextureIndex, !Allocation.LayerInfo->bNoWeightBlend));
 				}
 			}
 
@@ -2494,13 +2494,9 @@ const TArray<FName>& ALandscapeProxy::GetLayersFromMaterial(UMaterialInterface*
 {
 	if (MaterialInterface)
 	{
-		const FMaterialCachedExpressionData& CachedExpressionData = MaterialInterface->GetCachedExpressionData();
-		if (CachedExpressionData.EditorOnlyData)
-		{
-			return CachedExpressionData.EditorOnlyData->LandscapeLayerNames;
-		}
+		return MaterialInterface->GetCachedExpressionData().LandscapeLayerNames;
 	}
-	return FMaterialCachedExpressionEditorOnlyData::EmptyData.LandscapeLayerNames;
+	return FMaterialCachedExpressionData::EmptyData.LandscapeLayerNames;
 }
 
 const TArray<FName>& ALandscapeProxy::GetLayersFromMaterial() const
diff --git a/Engine/Source/Runtime/Landscape/Private/LandscapeGrass.cpp b/Engine/Source/Runtime/Landscape/Private/LandscapeGrass.cpp
--- a/Engine/Source/Runtime/Landscape/Private/LandscapeGrass.cpp
+++ b/Engine/Source/Runtime/Landscape/Private/LandscapeGrass.cpp
@@ -904,7 +904,7 @@ public:
 };
 
 FLandscapeComponentGrassData::FLandscapeComponentGrassData(ULandscapeComponent* Component)
-	: RotationForWPO(Component->GetLandscapeMaterial()->GetMaterial()->IsPropertyConnected(MP_WorldPositionOffset) ? Component->GetComponentTransform().GetRotation() : FQuat(0, 0, 0, 0))
+	: RotationForWPO(Component->GetLandscapeMaterial()->GetMaterial()->WorldPositionOffset.IsConnected() ? Component->GetComponentTransform().GetRotation() : FQuat(0, 0, 0, 0))
 	, bIsDirty(false)
 {
 	UMaterialInterface* Material = Component->GetLandscapeMaterial();
@@ -956,7 +956,7 @@ bool ULandscapeComponent::IsGrassMapOutdated() const
 			return true;
 		}
 
-		FQuat RotationForWPO = GetLandscapeMaterial()->GetMaterial()->IsPropertyConnected(MP_WorldPositionOffset) ? GetComponentTransform().GetRotation() : FQuat(0, 0, 0, 0);
+		FQuat RotationForWPO = GetLandscapeMaterial()->GetMaterial()->WorldPositionOffset.IsConnected() ? GetComponentTransform().GetRotation() : FQuat(0, 0, 0, 0);
 		if (GrassData->RotationForWPO != RotationForWPO)
 		{
 			return true;
diff --git a/Engine/Source/Runtime/Landscape/Private/LandscapeLight.cpp b/Engine/Source/Runtime/Landscape/Private/LandscapeLight.cpp
--- a/Engine/Source/Runtime/Landscape/Private/LandscapeLight.cpp
+++ b/Engine/Source/Runtime/Landscape/Private/LandscapeLight.cpp
@@ -427,7 +427,7 @@ void FLandscapeStaticLightingMesh::GetHeightmapData(int32 InLOD, int32 GeometryL
 	check(Info);
 
 	bool bUseRenderedWPO = LandscapeComponent->GetLandscapeProxy()->bUseMaterialPositionOffsetInStaticLighting &&
-		LandscapeComponent->GetLandscapeMaterial()->GetMaterial()->IsPropertyConnected(MP_WorldPositionOffset);
+	                       LandscapeComponent->GetLandscapeMaterial()->GetMaterial()->WorldPositionOffset.IsConnected();
 
 	HeightData.Empty(FMath::Square(NumVertices));
 	HeightData.AddUninitialized(FMath::Square(NumVertices));
diff --git a/Engine/Source/Runtime/Landscape/Private/LandscapeRender.cpp b/Engine/Source/Runtime/Landscape/Private/LandscapeRender.cpp
--- a/Engine/Source/Runtime/Landscape/Private/LandscapeRender.cpp
+++ b/Engine/Source/Runtime/Landscape/Private/LandscapeRender.cpp
@@ -3338,8 +3338,10 @@ void ULandscapeMaterialInstanceConstant::UpdateCachedTextureStreaming()
 	const UMaterial* Material = GetMaterial();
 	if (Material)
 	{
-		for (UMaterialExpression* Expression : Material->GetExpressions())
+		int32 NumExpressions = Material->Expressions.Num();
+		for (int32 ExpressionIndex = 0; ExpressionIndex < NumExpressions; ExpressionIndex++)
 		{
+			UMaterialExpression* Expression = Material->Expressions[ExpressionIndex];
 			UMaterialExpressionTextureSample* TextureSample = Cast<UMaterialExpressionTextureSample>(Expression);
 
 			// TODO: This is only works for direct Coordinate Texture Sample cases
