Rename HandleOnInputDeviceConnectionChange to HandleInputDeviceConnectionChange because "On" is redundant naming. Same thing for HandleOnInputDevicePairingChange.

#rnx

diff --git a/Engine/Source/Runtime/Engine/Classes/Engine/GameInstance.h b/Engine/Source/Runtime/Engine/Classes/Engine/GameInstance.h
--- a/Engine/Source/Runtime/Engine/Classes/Engine/GameInstance.h
+++ b/Engine/Source/Runtime/Engine/Classes/Engine/GameInstance.h
@@ -238,7 +238,7 @@ public:
 	 * @param FInputDeviceId		The Input Device ID that has changed connection
 	 * @see IPlatformInputDeviceMapper
 	 */
-	virtual void HandleOnInputDeviceConnectionChange(EInputDeviceConnectionState NewConnectionState, FPlatformUserId PlatformUserId, FInputDeviceId InputDeviceId);
+	virtual void HandleInputDeviceConnectionChange(EInputDeviceConnectionState NewConnectionState, FPlatformUserId PlatformUserId, FInputDeviceId InputDeviceId);
 
 	/** 
 	 * Callback for when an input device connection state has changed (a new gamepad was connected or disconnected)
@@ -254,7 +254,7 @@ public:
 	 * @param FPlatformUserId	The OldUserPlatformId
 	 * @see IPlatformInputDeviceMapper
 	 */
-	virtual void HandleOnInputDevicePairingChange(FInputDeviceId InputDeviceId, FPlatformUserId NewUserPlatformId, FPlatformUserId OldUserPlatformId);
+	virtual void HandleInputDevicePairingChange(FInputDeviceId InputDeviceId, FPlatformUserId NewUserPlatformId, FPlatformUserId OldUserPlatformId);
 
 	/**
 	 * Callback when an input device has changed pairings (the owning platform user has changed for that device)
diff --git a/Engine/Source/Runtime/Engine/Private/GameInstance.cpp b/Engine/Source/Runtime/Engine/Private/GameInstance.cpp
--- a/Engine/Source/Runtime/Engine/Private/GameInstance.cpp
+++ b/Engine/Source/Runtime/Engine/Private/GameInstance.cpp
@@ -100,8 +100,8 @@ void UGameInstance::Init()
 		FNetDelegates::OnReceivedNetworkEncryptionAck.BindUObject(this, &ThisClass::ReceivedNetworkEncryptionAck);
 
 		IPlatformInputDeviceMapper& PlatformInputMapper = IPlatformInputDeviceMapper::Get();
-		PlatformInputMapper.GetOnInputDeviceConnectionChange().AddUObject(this, &UGameInstance::HandleOnInputDeviceConnectionChange);
-		PlatformInputMapper.GetOnInputDevicePairingChange().AddUObject(this, &UGameInstance::HandleOnInputDevicePairingChange);
+		PlatformInputMapper.GetOnInputDeviceConnectionChange().AddUObject(this, &UGameInstance::HandleInputDeviceConnectionChange);
+		PlatformInputMapper.GetOnInputDevicePairingChange().AddUObject(this, &UGameInstance::HandleInputDevicePairingChange);
 	}
 
 	SubsystemCollection.Initialize(this);
@@ -155,12 +155,12 @@ void UGameInstance::Shutdown()
 	WorldContext = nullptr;
 }
 
-void UGameInstance::HandleOnInputDeviceConnectionChange(EInputDeviceConnectionState NewConnectionState, FPlatformUserId PlatformUserId, FInputDeviceId InputDeviceId)
+void UGameInstance::HandleInputDeviceConnectionChange(EInputDeviceConnectionState NewConnectionState, FPlatformUserId PlatformUserId, FInputDeviceId InputDeviceId)
 {
 	OnInputDeviceConnectionChange.Broadcast(NewConnectionState, PlatformUserId, InputDeviceId);
 }
 
-void UGameInstance::HandleOnInputDevicePairingChange(FInputDeviceId InputDeviceId, FPlatformUserId NewUserPlatformId, FPlatformUserId OldUserPlatformId)
+void UGameInstance::HandleInputDevicePairingChange(FInputDeviceId InputDeviceId, FPlatformUserId NewUserPlatformId, FPlatformUserId OldUserPlatformId)
 {
 	OnUserInputDevicePairingChange.Broadcast(InputDeviceId, NewUserPlatformId, OldUserPlatformId);
 }
