Add static references to PLATFORMUSERID_NONE and INPUTDEVICEID_NONE to for blueprints. This will be useful for folks when they want to do comparisons between the two or something like that. Add "Equals" comparison operators for the PlatformUserID and InputDeviceID structs too, so that you can easily check if they are the same in blueprints.

#rnx

diff --git a/Engine/Source/Runtime/Engine/Classes/GameFramework/PlatformInputDeviceMapperLibrary.h b/Engine/Source/Runtime/Engine/Classes/GameFramework/PlatformInputDeviceMapperLibrary.h
--- a/Engine/Source/Runtime/Engine/Classes/GameFramework/PlatformInputDeviceMapperLibrary.h
+++ b/Engine/Source/Runtime/Engine/Classes/GameFramework/PlatformInputDeviceMapperLibrary.h
@@ -110,4 +110,28 @@ public:
 	 */
 	UFUNCTION(BlueprintPure, Category = "PlatformInputDevice")
 	static bool IsValidPlatformId(FPlatformUserId UserId);
+
+	/** Static invalid platform user */
+	UFUNCTION(BlueprintPure, meta = (ScriptConstant = "None", ScriptConstantHost = "PlatformUserId"), Category = "PlatformInputDevice")
+	static FPlatformUserId PlatformUserId_None();
+
+	/** Static invalid input device */
+	UFUNCTION(BlueprintPure, meta = (ScriptConstant = "None", ScriptConstantHost = "InputDeviceId"), Category = "PlatformInputDevice")
+	static FInputDeviceId InputDeviceId_None();
+
+	/** Returns true if PlatformUserId A is equal to PlatformUserId B (A == B) */
+	UFUNCTION(BlueprintPure, meta = (DisplayName = "Equal (PlatformUserId)", CompactNodeTitle = "==", ScriptMethod = "Equals", ScriptOperator = "==", Keywords = "== equal"), Category = "PlatformInputDevice")
+	static bool EqualEqual_PlatformUserId(FPlatformUserId A, FPlatformUserId B);
+
+	/** Returns true if PlatformUserId A is not equal to PlatformUserId B (A != B) */
+	UFUNCTION(BlueprintPure, meta = (DisplayName = "Not Equal (PlatformUserId)", CompactNodeTitle = "!=", ScriptMethod = "NotEqual", ScriptOperator = "!=", Keywords = "!= not equal"), Category = "PlatformInputDevice")
+	static bool NotEqual_PlatformUserId(FPlatformUserId A, FPlatformUserId B);
+	
+	/** Returns true if InputDeviceId A is equal to InputDeviceId B (A == B) */
+	UFUNCTION(BlueprintPure, meta = (DisplayName = "Equal (InputDeviceId)", CompactNodeTitle = "==", ScriptMethod = "Equals", ScriptOperator = "==", Keywords = "== equal"), Category = "PlatformInputDevice")
+	static bool EqualEqual_InputDeviceId(FInputDeviceId A, FInputDeviceId B);
+	
+	/** Returns true if InputDeviceId A is not equal to InputDeviceId B (A != B) */
+	UFUNCTION(BlueprintPure, meta = (DisplayName = "Not Equal (InputDeviceId)", CompactNodeTitle = "!=", ScriptMethod = "NotEqual", ScriptOperator = "!=", Keywords = "!= not equal"), Category = "PlatformInputDevice")
+	static bool NotEqual_InputDeviceId(FInputDeviceId A, FInputDeviceId B);
 };
\ No newline at end of file
diff --git a/Engine/Source/Runtime/Engine/Private/PlatformInputDeviceMapperLibrary.cpp b/Engine/Source/Runtime/Engine/Private/PlatformInputDeviceMapperLibrary.cpp
--- a/Engine/Source/Runtime/Engine/Private/PlatformInputDeviceMapperLibrary.cpp
+++ b/Engine/Source/Runtime/Engine/Private/PlatformInputDeviceMapperLibrary.cpp
@@ -66,4 +66,34 @@ bool UPlatformInputDeviceMapperLibrary::IsValidInputDevice(FInputDeviceId Device
 bool UPlatformInputDeviceMapperLibrary::IsValidPlatformId(FPlatformUserId UserId)
 {
 	return UserId.IsValid();
+}
+
+FPlatformUserId UPlatformInputDeviceMapperLibrary::PlatformUserId_None()
+{
+	return PLATFORMUSERID_NONE;
+}
+
+FInputDeviceId UPlatformInputDeviceMapperLibrary::InputDeviceId_None()
+{
+	return INPUTDEVICEID_NONE;
+}
+
+bool UPlatformInputDeviceMapperLibrary::EqualEqual_PlatformUserId(FPlatformUserId A, FPlatformUserId B)
+{
+	return A == B;
+}
+
+bool UPlatformInputDeviceMapperLibrary::NotEqual_PlatformUserId(FPlatformUserId A, FPlatformUserId B)
+{
+	return A != B;
+}
+
+bool UPlatformInputDeviceMapperLibrary::EqualEqual_InputDeviceId(FInputDeviceId A, FInputDeviceId B)
+{
+	return A == B;
+}
+
+bool UPlatformInputDeviceMapperLibrary::NotEqual_InputDeviceId(FInputDeviceId A, FInputDeviceId B)
+{
+	return A != B;
 }
\ No newline at end of file
