Serialize a reference from a ULevel to its corresponding WorldPartitionCell when its a level streamed from a WorldPartition grid.
Allow actors to retrieve their DataLayers outside of editor builds.
Actors now use their cell to retrieve their DataLayers when possible (i.e. when GenerateStreaming was done and the grid is built.)
Removed UActors::GetDataLayerInstances(AWorldDataLayers). Function isn't needed anymore as it was an overload called by the UI and problematic when exiting PIE. Actors on DataLayers not longer had their owning world set and the WOrldDataLayer was inaccessible. No longer need to use the WorldDataLayer as we use the Cell to get the Datalayers in PIE.

#test Editor, PIE, SIE, Cook, Cooked Game, -game

diff --git a/Engine/Source/Editor/DataLayerEditor/Private/DataLayer/DataLayerHierarchy.cpp b/Engine/Source/Editor/DataLayerEditor/Private/DataLayer/DataLayerHierarchy.cpp
--- a/Engine/Source/Editor/DataLayerEditor/Private/DataLayer/DataLayerHierarchy.cpp
+++ b/Engine/Source/Editor/DataLayerEditor/Private/DataLayer/DataLayerHierarchy.cpp
@@ -306,7 +306,7 @@ void FDataLayerHierarchy::OnLevelActorsRemoved(const TArray<AActor*>& InActors)
 				if (Actor != nullptr && Actor->HasDataLayers())
 				{
 					// It is possible here that Actor doesn't have world anymore
-					const TArray<const UDataLayerInstance*> DataLayerInstances = Actor->GetDataLayerInstances(WorldDataLayers);
+					const TArray<const UDataLayerInstance*> DataLayerInstances = Actor->GetDataLayerInstances();
 					EventData.ItemIDs.Reserve(DataLayerInstances.Num());
 					for (const UDataLayerInstance* DataLayerInstance : DataLayerInstances)
 					{
diff --git a/Engine/Source/Runtime/Engine/Classes/Engine/Level.h b/Engine/Source/Runtime/Engine/Classes/Engine/Level.h
--- a/Engine/Source/Runtime/Engine/Classes/Engine/Level.h
+++ b/Engine/Source/Runtime/Engine/Classes/Engine/Level.h
@@ -38,6 +38,8 @@ class FLevelPartitionOperationScope;
 class FRegisterComponentContext;
 class SNotificationItem;
 class UActorFolder;
+class IWorldPartitionCell;
+class UWorldPartitionRuntimeCell;
 
 UINTERFACE()
 class ULevelPartitionInterface : public UInterface
@@ -674,9 +676,6 @@ public:
 	/** Whether the level is partitioned or not. */
     UPROPERTY()
 	uint8										bIsPartitioned : 1;
-	/** Whether the level is a world partition runtime cell. */
-	UPROPERTY()
-	uint8										bIsWorldPartitionRuntimeCell : 1;
 
 	enum class EIncrementalComponentState : uint8
 	{
@@ -795,6 +794,9 @@ private:
 	UPROPERTY()
 	TObjectPtr<AWorldDataLayers> WorldDataLayers;
 
+	UPROPERTY()
+	TSoftObjectPtr<UWorldPartitionRuntimeCell> WorldPartitionRuntimeCell;
+
 	/** Cached level collection that this level is contained in, for faster access than looping through the collections in the world. */
 	FLevelCollection* CachedLevelCollection;
 
@@ -1078,6 +1080,20 @@ public:
 
 	ENGINE_API void SetWorldDataLayers(AWorldDataLayers* NewWorldDataLayers);
 
+	/**
+	 * Returns the RuntimeCell associated with this Level if it is a level representing a cell of a WorldPartition World.
+	 *
+	 * @return		The cell associated with the level.
+	 */
+	ENGINE_API const IWorldPartitionCell* GetWorldPartitionRuntimeCell() const;
+
+	/**
+	 * Returns if the level is a cell from a WorldPartition World.
+	 *
+	 * @return		The cell associated with the level.
+	 */
+	ENGINE_API bool IsWorldPartitionRuntimeCell() const { return !WorldPartitionRuntimeCell.GetUniqueID().IsNull(); }
+
 	/**
 	 * Returns the UWorldPartition for this level.
 	 *
diff --git a/Engine/Source/Runtime/Engine/Classes/GameFramework/Actor.h b/Engine/Source/Runtime/Engine/Classes/GameFramework/Actor.h
--- a/Engine/Source/Runtime/Engine/Classes/GameFramework/Actor.h
+++ b/Engine/Source/Runtime/Engine/Classes/GameFramework/Actor.h
@@ -1174,22 +1174,18 @@ public:
 		return Cast<T>(GetInstigatorController());
 	}
 
-#if WITH_EDITOR
 	//~=============================================================================
 	// DataLayers functions.
+#if WITH_EDITOR
 	bool AddDataLayer(const UDataLayerAsset* DataLayerAsset);
 	bool AddDataLayer(const UDataLayerInstance* DataLayerInstance);
 	bool RemoveDataLayer(const UDataLayerAsset* DataLayerAsset);
 	bool RemoveDataLayer(const UDataLayerInstance* DataLayerInstance);
 	bool RemoveAllDataLayers();
-	bool ContainsDataLayer(const UDataLayerAsset* DataLayerAsset) const;
-	bool ContainsDataLayer(const UDataLayerInstance* DataLayerInstance) const;
 	virtual bool SupportsDataLayer() const { return true; }
-	bool HasDataLayers() const;
+	
 	bool HasValidDataLayers() const;
 	TArray<FName> GetDataLayerInstanceNames() const;
-	TArray<const UDataLayerInstance*> GetDataLayerInstances() const;
-	TArray<const UDataLayerInstance*> GetDataLayerInstances(const AWorldDataLayers* WorldDataLayers) const;
 	bool IsPropertyChangedAffectingDataLayers(FPropertyChangedEvent& PropertyChangedEvent) const;
 	bool IsValidForDataLayer() const;
 	void FixupDataLayers(bool bRevertChangesOnLockedDataLayer = false);
@@ -1225,6 +1221,12 @@ public:
 	//~ End Deprecated
 #endif
 
+	TArray<const UDataLayerInstance*> GetDataLayerInstances() const;
+
+	bool ContainsDataLayer(const UDataLayerAsset* DataLayerAsset) const;
+	bool ContainsDataLayer(const UDataLayerInstance* DataLayerInstance) const;
+	bool HasDataLayers() const;
+
 	//~=============================================================================
 	// General functions.
 
diff --git a/Engine/Source/Runtime/Engine/Private/Actor.cpp b/Engine/Source/Runtime/Engine/Private/Actor.cpp
--- a/Engine/Source/Runtime/Engine/Private/Actor.cpp
+++ b/Engine/Source/Runtime/Engine/Private/Actor.cpp
@@ -61,6 +61,10 @@
 
 #include "WorldPartition/WorldPartition.h"
 #include "WorldPartition/WorldPartitionActorDesc.h"
+#include "WorldPartition/WorldPartitionRuntimeCell.h"
+#include "WorldPartition/DataLayer/DataLayerInstance.h"
+#include "WorldPartition/DataLayer/DataLayerAsset.h"
+#include "WorldPartition/DataLayer/WorldDataLayers.h"
 
 DEFINE_LOG_CATEGORY(LogActor);
 
@@ -5901,4 +5905,74 @@ FArchive& operator<<(FArchive& Ar, AActor::FActorTransactionAnnotationData& Acto
 	return Ar;
 }
 
+//---------------------------------------------------------------------------
+// DataLayers (begin)
+
+TArray<const UDataLayerInstance*> AActor::GetDataLayerInstances() const
+{
+	if (const IWorldPartitionCell* RuntinmeCell = GetLevel()->GetWorldPartitionRuntimeCell())
+	{
+		return RuntinmeCell->GetDataLayerInstances();
+	}
+
+#if WITH_EDITOR
+	TArray<const UDataLayerInstance*> DataLayerInstances;
+	AWorldDataLayers* WorldDataLayers = GetWorld()->GetWorldDataLayers();
+	DataLayerInstances += WorldDataLayers->GetDataLayerInstances(DataLayerAssets);
+
+	PRAGMA_DISABLE_DEPRECATION_WARNINGS
+	DataLayerInstances += WorldDataLayers->GetDataLayerInstances(DataLayers);
+	PRAGMA_ENABLE_DEPRECATION_WARNINGS
+
+	return DataLayerInstances;
+#else
+	return TArray<const UDataLayerInstance*>();
+#endif
+}
+
+bool AActor::ContainsDataLayer(const UDataLayerInstance* DataLayerInstance) const
+{
+	if (const IWorldPartitionCell* RuntinmeCell = GetLevel()->GetWorldPartitionRuntimeCell())
+	{
+		return RuntinmeCell->ContainsDataLayer(DataLayerInstance);
+	}
+
+#if WITH_EDITOR
+	return DataLayerInstance->ContainsActor(this);
+#else
+	return false;
+#endif
+}
+
+bool AActor::ContainsDataLayer(const UDataLayerAsset* DataLayerAsset) const
+{
+	if (const IWorldPartitionCell* RuntinmeCell = GetLevel()->GetWorldPartitionRuntimeCell())
+	{
+		return RuntinmeCell->ContainsDataLayer(DataLayerAsset);
+	}
+
+#if WITH_EDITOR
+	return DataLayerAsset && DataLayerAssets.Contains(DataLayerAsset);
+#else
+	return false;
+#endif
+}
+
+bool AActor::HasDataLayers() const
+{
+	if (const IWorldPartitionCell* RuntinmeCell = GetLevel()->GetWorldPartitionRuntimeCell())
+	{
+		return RuntinmeCell->HasDataLayers();
+	}
+
+#if WITH_EDITOR
+	return DataLayerAssets.Num() > 0 || DataLayers.Num() > 0;
+#else
+	return false;
+#endif
+}
+
+// DataLayers (end)
+//---------------------------------------------------------------------------
+
 #undef LOCTEXT_NAMESPACE
diff --git a/Engine/Source/Runtime/Engine/Private/ActorEditor.cpp b/Engine/Source/Runtime/Engine/Private/ActorEditor.cpp
--- a/Engine/Source/Runtime/Engine/Private/ActorEditor.cpp
+++ b/Engine/Source/Runtime/Engine/Private/ActorEditor.cpp
@@ -1365,21 +1365,6 @@ bool AActor::RemoveAllDataLayers()
 	return false;
 }
 
-bool AActor::ContainsDataLayer(const UDataLayerInstance* DataLayerInstance) const
-{
-	return DataLayerInstance->ContainsActor(this);
-}
-
-bool AActor::ContainsDataLayer(const UDataLayerAsset* DataLayerAsset) const
-{
-	return DataLayerAsset && DataLayerAssets.Contains(DataLayerAsset);
-}
-
-bool AActor::HasDataLayers() const
-{
-	return DataLayerAssets.Num() > 0 || DataLayers.Num() > 0;
-}
-
 bool AActor::HasValidDataLayers() const
 {
 	if (const AWorldDataLayers* WorldDataLayers = GetWorld()->GetWorldDataLayers())
@@ -1419,27 +1404,6 @@ TArray<FName> AActor::GetDataLayerInstanceNames() const
 	return DataLayerInstanceNames;
 }
 
-TArray<const UDataLayerInstance*> AActor::GetDataLayerInstances() const
-{
-	return GetWorld() ? GetDataLayerInstances(GetWorld()->GetWorldDataLayers()) : TArray<const UDataLayerInstance*>();
-}
-
-TArray<const UDataLayerInstance*> AActor::GetDataLayerInstances(const AWorldDataLayers* WorldDataLayers) const
-{
-	TArray<const UDataLayerInstance*> DataLayerInstances;
-	if (WorldDataLayers != nullptr)
-	{
-		DataLayerInstances += WorldDataLayers->GetDataLayerInstances(DataLayerAssets);
-
-		PRAGMA_DISABLE_DEPRECATION_WARNINGS
-		DataLayerInstances += WorldDataLayers->GetDataLayerInstances(DataLayers);
-		PRAGMA_ENABLE_DEPRECATION_WARNINGS
-	}
-
-
-	return DataLayerInstances;
-}
-
 void AActor::FixupDataLayers(bool bRevertChangesOnLockedDataLayer /*= false*/)
 {
 	if (!GetPackage()->HasAnyPackageFlags(PKG_PlayInEditor))
diff --git a/Engine/Source/Runtime/Engine/Private/Level.cpp b/Engine/Source/Runtime/Engine/Private/Level.cpp
--- a/Engine/Source/Runtime/Engine/Private/Level.cpp
+++ b/Engine/Source/Runtime/Engine/Private/Level.cpp
@@ -92,6 +92,7 @@ Level.cpp: Level-related functions
 #include "ObjectTrace.h"
 #include "ProfilingDebugging/TagTrace.h"
 #include "UObject/MetaData.h"
+#include "WorldPartition/WorldPartitionRuntimeCell.h"
 
 #define LOCTEXT_NAMESPACE "ULevel"
 DEFINE_LOG_CATEGORY(LogLevel);
@@ -2655,6 +2656,10 @@ void ULevel::SetWorldDataLayers(AWorldDataLayers* NewWorldDataLayers)
 	WorldDataLayers = NewWorldDataLayers;
 }
 
+const IWorldPartitionCell* ULevel::GetWorldPartitionRuntimeCell() const
+{
+	return WorldPartitionRuntimeCell.Get();
+}
 
 UWorldPartition* ULevel::GetWorldPartition() const
 {
diff --git a/Engine/Source/Runtime/Engine/Private/WorldPartition/WorldPartitionLevelHelper.cpp b/Engine/Source/Runtime/Engine/Private/WorldPartition/WorldPartitionLevelHelper.cpp
--- a/Engine/Source/Runtime/Engine/Private/WorldPartition/WorldPartitionLevelHelper.cpp
+++ b/Engine/Source/Runtime/Engine/Private/WorldPartition/WorldPartitionLevelHelper.cpp
@@ -191,7 +191,7 @@ FString FWorldPartitionLevelHelper::AddActorContainerIDToActorPath(const FActorC
 /**
  * Creates an empty Level used in World Partition
  */
-ULevel* FWorldPartitionLevelHelper::CreateEmptyLevelForRuntimeCell(const UWorld* InWorld, const FString& InWorldAssetName, UPackage* InPackage)
+ULevel* FWorldPartitionLevelHelper::CreateEmptyLevelForRuntimeCell(const UWorldPartitionRuntimeCell* Cell, const UWorld* InWorld, const FString& InWorldAssetName, UPackage* InPackage)
 {
 	// Create or use given package
 	UPackage* CellPackage = nullptr;
@@ -233,8 +233,7 @@ ULevel* FWorldPartitionLevelHelper::CreateEmptyLevelForRuntimeCell(const UWorld*
 	check(NewLevel->Model);
 	check(!NewLevel->bIsVisible);
 
-	// Mark the level as a runtime cell
-	NewLevel->bIsWorldPartitionRuntimeCell = true;
+	NewLevel->WorldPartitionRuntimeCell = Cell;
 	
 	// Mark the level package as fully loaded
 	CellPackage->MarkAsFullyLoaded();
diff --git a/Engine/Source/Runtime/Engine/Private/WorldPartition/WorldPartitionLevelStreamingDynamic.cpp b/Engine/Source/Runtime/Engine/Private/WorldPartition/WorldPartitionLevelStreamingDynamic.cpp
--- a/Engine/Source/Runtime/Engine/Private/WorldPartition/WorldPartitionLevelStreamingDynamic.cpp
+++ b/Engine/Source/Runtime/Engine/Private/WorldPartition/WorldPartitionLevelStreamingDynamic.cpp
@@ -125,9 +125,10 @@ void UWorldPartitionLevelStreamingDynamic::CreateRuntimeLevel()
 	check(RuntimeLevel == nullptr);
 	const UWorld* World = GetWorld();
 	check(World && (World->IsGameWorld() || GetShouldBeVisibleInEditor()));
+	check(StreamingCell != nullptr);
 
 	// Create streaming cell Level package
-	RuntimeLevel = FWorldPartitionLevelHelper::CreateEmptyLevelForRuntimeCell(World, GetWorldAsset().ToString());
+	RuntimeLevel = FWorldPartitionLevelHelper::CreateEmptyLevelForRuntimeCell(StreamingCell.Get(), World, GetWorldAsset().ToString());
 	check(RuntimeLevel);
 
 	UPackage* RuntimeLevelPackage = RuntimeLevel->GetPackage();
diff --git a/Engine/Source/Runtime/Engine/Private/WorldPartition/WorldPartitionRuntimeCell.cpp b/Engine/Source/Runtime/Engine/Private/WorldPartition/WorldPartitionRuntimeCell.cpp
--- a/Engine/Source/Runtime/Engine/Private/WorldPartition/WorldPartitionRuntimeCell.cpp
+++ b/Engine/Source/Runtime/Engine/Private/WorldPartition/WorldPartitionRuntimeCell.cpp
@@ -114,4 +114,32 @@ FLinearColor UWorldPartitionRuntimeCell::GetDebugStreamingPriorityColor() const
 	}
 #endif
 	return FLinearColor::Transparent;
+}
+
+TArray<const UDataLayerInstance*> UWorldPartitionRuntimeCell::GetDataLayerInstances() const
+{
+	if (const AWorldDataLayers* WorldDataLayers = GetWorld()->GetWorldDataLayers())
+	{
+		return WorldDataLayers->GetDataLayerInstances(GetDataLayers());
+	}
+
+	return TArray<const UDataLayerInstance*>();
+}
+
+bool UWorldPartitionRuntimeCell::ContainsDataLayer(const UDataLayerAsset* DataLayerAsset) const
+{
+	if (const AWorldDataLayers* WorldDataLayers = GetWorld()->GetWorldDataLayers())
+	{
+		if (const UDataLayerInstance* DataLayerInstance = WorldDataLayers->GetDataLayerInstance(DataLayerAsset))
+		{
+			return ContainsDataLayer(DataLayerInstance);
+		}
+	}
+
+	return false;
+}
+
+bool UWorldPartitionRuntimeCell::ContainsDataLayer(const UDataLayerInstance* DataLayerInstance) const
+{
+	return GetDataLayers().Contains(DataLayerInstance->GetDataLayerFName());
 }
\ No newline at end of file
diff --git a/Engine/Source/Runtime/Engine/Private/WorldPartition/WorldPartitionRuntimeLevelStreamingCell.cpp b/Engine/Source/Runtime/Engine/Private/WorldPartition/WorldPartitionRuntimeLevelStreamingCell.cpp
--- a/Engine/Source/Runtime/Engine/Private/WorldPartition/WorldPartitionRuntimeLevelStreamingCell.cpp
+++ b/Engine/Source/Runtime/Engine/Private/WorldPartition/WorldPartitionRuntimeLevelStreamingCell.cpp
@@ -197,7 +197,7 @@ bool UWorldPartitionRuntimeLevelStreamingCell::PopulateGeneratedPackageForCook(U
 		verify(FWorldPartitionLevelHelper::LoadActors(OuterWorld, nullptr, Packages, PackageReferencer, [](bool) {}, bLoadAsync, FLinkerInstancingContext()));
 
 		// Create a level and move these actors in it
-		ULevel* NewLevel = FWorldPartitionLevelHelper::CreateEmptyLevelForRuntimeCell(OuterWorld, LevelStreaming->GetWorldAsset().ToString(), InPackage);
+		ULevel* NewLevel = FWorldPartitionLevelHelper::CreateEmptyLevelForRuntimeCell(this, OuterWorld, LevelStreaming->GetWorldAsset().ToString(), InPackage);
 		check(NewLevel->GetPackage() == InPackage);
 		FWorldPartitionLevelHelper::MoveExternalActorsToLevel(Packages, NewLevel);
 
diff --git a/Engine/Source/Runtime/Engine/Public/WorldPartition/WorldPartitionLevelHelper.h b/Engine/Source/Runtime/Engine/Public/WorldPartition/WorldPartitionLevelHelper.h
--- a/Engine/Source/Runtime/Engine/Public/WorldPartition/WorldPartitionLevelHelper.h
+++ b/Engine/Source/Runtime/Engine/Public/WorldPartition/WorldPartitionLevelHelper.h
@@ -29,7 +29,7 @@ public:
 		void RemoveReferences();
 	};
 
-	static ULevel* CreateEmptyLevelForRuntimeCell(const UWorld* InWorld, const FString& InWorldAssetName, UPackage* DestPackage = nullptr);
+	static ULevel* CreateEmptyLevelForRuntimeCell(const UWorldPartitionRuntimeCell* Cell, const UWorld* InWorld, const FString& InWorldAssetName, UPackage* DestPackage = nullptr);
 	static void DuplicateActorFolderToRuntimeCell(ULevel* CellLevel, ULevel* SrcLevel, const FGuid& ActorFolderGuid);
 	static void MoveExternalActorsToLevel(const TArray<FWorldPartitionRuntimeCellObjectMapping>& InChildPackages, ULevel* InLevel);
 	static void RemapLevelSoftObjectPaths(ULevel* InLevel, UWorldPartition* InWorldPartition);
diff --git a/Engine/Source/Runtime/Engine/Public/WorldPartition/WorldPartitionRuntimeCell.h b/Engine/Source/Runtime/Engine/Public/WorldPartition/WorldPartitionRuntimeCell.h
--- a/Engine/Source/Runtime/Engine/Public/WorldPartition/WorldPartitionRuntimeCell.h
+++ b/Engine/Source/Runtime/Engine/Public/WorldPartition/WorldPartitionRuntimeCell.h
@@ -7,10 +7,14 @@
 #include "WorldPartition/WorldPartitionActorDescView.h"
 #include "WorldPartition/WorldPartitionStreamingSource.h"
 #include "WorldPartition/WorldPartitionActorCluster.h"
+#include "WorldPartition/WorldPartitionRuntimeCellInterface.h"
 #include "ProfilingDebugging/ProfilingHelpers.h"
 #include "Algo/AnyOf.h"
 #include "WorldPartitionRuntimeCell.generated.h"
 
+class UDataLayerAsset;
+class UDataLayerInstance;
+
 enum class EWorldPartitionRuntimeCellVisualizeMode
 {
 	StreamingPriority,
@@ -106,7 +110,7 @@ static_assert(EWorldPartitionRuntimeCellState::Unloaded < EWorldPartitionRuntime
  * Represents a PIE/Game streaming cell which points to external actor/data chunk packages
  */
 UCLASS(Abstract, Within = WorldPartition)
-class UWorldPartitionRuntimeCell : public UObject
+class UWorldPartitionRuntimeCell : public UObject, public IWorldPartitionCell
 {
 	GENERATED_UCLASS_BODY()
 
@@ -145,12 +149,17 @@ class UWorldPartitionRuntimeCell : public UObject
 
 	static void DirtyStreamingSourceCacheEpoch() { ++UWorldPartitionRuntimeCell::StreamingSourceCacheEpoch; }
 
-	bool HasDataLayers() const { return !DataLayers.IsEmpty(); }
-	const TArray<FName>& GetDataLayers() const { return DataLayers; }
-	bool HasAnyDataLayer(const TSet<FName>& InDataLayers) const
+	// IRuntimeCell Interface begin
+	virtual TArray<const UDataLayerInstance*> GetDataLayerInstances() const override;
+	virtual bool ContainsDataLayer(const UDataLayerAsset* DataLayerAsset) const override;
+	virtual bool ContainsDataLayer(const UDataLayerInstance* DataLayerInstance) const override;
+	virtual bool HasDataLayers() const override { return !DataLayers.IsEmpty(); }
+	virtual const TArray<FName>& GetDataLayers() const override  { return DataLayers; }
+	virtual bool HasAnyDataLayer(const TSet<FName>& InDataLayers) const override
 	{
 		return Algo::AnyOf(DataLayers, [&InDataLayers](const FName& DataLayer) { return InDataLayers.Contains(DataLayer); });
 	}
+	// IRuntimeCell Interface end
 
 	bool GetBlockOnSlowLoading() const { return bBlockOnSlowLoading; }
 #if WITH_EDITOR
diff --git a/Engine/Source/Runtime/Engine/Public/WorldPartition/WorldPartitionRuntimeCellInterface.h b/Engine/Source/Runtime/Engine/Public/WorldPartition/WorldPartitionRuntimeCellInterface.h
new file mode 100644
--- /dev/null
+++ b/Engine/Source/Runtime/Engine/Public/WorldPartition/WorldPartitionRuntimeCellInterface.h
@@ -0,0 +1,19 @@
+// Copyright Epic Games, Inc. All Rights Reserved.
+#pragma once
+
+#include "CoreMinimal.h"
+
+class UDataLayerInstance;
+class UDataLayerAsset;
+
+class IWorldPartitionCell
+{
+public:
+	virtual TArray<const UDataLayerInstance*> GetDataLayerInstances() const = 0;
+	virtual bool ContainsDataLayer(const UDataLayerAsset* DataLayerAsset) const = 0;
+	virtual bool ContainsDataLayer(const UDataLayerInstance* DataLayerInstance) const = 0;
+	virtual bool HasDataLayers() const = 0;
+	virtual const TArray<FName>& GetDataLayers() const = 0;
+	virtual bool HasAnyDataLayer(const TSet<FName>& InDataLayers) const = 0;
+};
+
diff --git a/Engine/Source/Runtime/Foliage/Private/InstancedFoliage.cpp b/Engine/Source/Runtime/Foliage/Private/InstancedFoliage.cpp
--- a/Engine/Source/Runtime/Foliage/Private/InstancedFoliage.cpp
+++ b/Engine/Source/Runtime/Foliage/Private/InstancedFoliage.cpp
@@ -4308,7 +4308,7 @@ void AInstancedFoliageActor::PostLoad()
 	// We can't check the ActorPartitionSubsystem here because World is not initialized yet. So we fallback on the bIsPartitioned
 	// to know if multiple InstanceFoliageActors is valid or not.
 	// For levels that are World Partition runtime cells. Having multiple IFAs is valid so skip validation.
-	if (OwningLevel && !OwningLevel->bIsPartitioned && !OwningLevel->bIsWorldPartitionRuntimeCell)
+	if (OwningLevel && !OwningLevel->bIsPartitioned && !OwningLevel->IsWorldPartitionRuntimeCell())
 	{
 		if (!OwningLevel->InstancedFoliageActor.IsValid())
 		{
