Do not tick empty levels causing a rerun through all the Tick Groups with no work to do.

diff --git a/Engine/Source/Runtime/Engine/Private/LevelTick.cpp b/Engine/Source/Runtime/Engine/Private/LevelTick.cpp
--- a/Engine/Source/Runtime/Engine/Private/LevelTick.cpp
+++ b/Engine/Source/Runtime/Engine/Private/LevelTick.cpp
@@ -1511,7 +1511,8 @@ void UWorld::Tick( ELevelTick TickType, float DeltaSeconds )
 		FScopedLevelCollectionContextSwitch LevelContext(i, this);
 
 		// If caller wants time update only, or we are paused, skip the rest.
-		if (bDoingActorTicks)
+		const bool bShouldSkipTick = (LevelsToTick.Num() == 0);
+		if (bDoingActorTicks && !bShouldSkipTick)
 		{
 			// Actually tick actors now that context is set up
 			SetupPhysicsTickFunctions(DeltaSeconds);
@@ -1625,7 +1626,7 @@ void UWorld::Tick( ELevelTick TickType, float DeltaSeconds )
 			}
 		}
 
-		if (bDoingActorTicks)
+		if (bDoingActorTicks && !bShouldSkipTick)
 		{
 			SCOPE_CYCLE_COUNTER(STAT_TickTime);
 			{
