Protect against PlayQuantized being called w/ a null sound from C++

diff --git a/Engine/Source/Runtime/Engine/Private/Components/AudioComponent.cpp b/Engine/Source/Runtime/Engine/Private/Components/AudioComponent.cpp
--- a/Engine/Source/Runtime/Engine/Private/Components/AudioComponent.cpp
+++ b/Engine/Source/Runtime/Engine/Private/Components/AudioComponent.cpp
@@ -562,40 +562,53 @@ void UAudioComponent::PlayQueuedQuantizedInternal(const UObject* WorldContextObj
 			InternalRequestData.FadeInDuration = PendingData.FadeDuration;
 			InternalRequestData.FadeVolumeLevel = PendingData.FadeVolume;
 			InternalRequestData.FadeCurve = PendingData.FadeCurve;
-			
-			if (Handle != nullptr)
-			{
-				InternalRequestData.QuantizedRequestData = Handle->GetQuartzSubsystem()->CreateDataDataForSchedulePlaySound(Handle, PendingData.Delegate, PendingData.AnticapatoryBoundary);
-				UGameplayStatics::PrimeSound(Sound);
-			}
 
-			// validate clock existence 
-			if (!Handle)
+			// verify we have a sound to play
+			if(Sound)
 			{
-				UE_LOG(LogAudioQuartz, Warning, TEXT("Attempting to play Quantized Sound without supplying a Clock Handle"));
-				bIsValidCommand = false;
-			}
-			else if (!Handle->DoesClockExist(WorldContextObject))
-			{
-				UE_LOG(LogAudioQuartz, Warning, TEXT("Clock: '%s' Does not exist! Cannot play quantized sound: '%s'"), *InternalRequestData.QuantizedRequestData.ClockName.ToString(), *this->Sound->GetName());
-				bIsValidCommand = false;
-			}
+				// confirm a valid handle
+				if (Handle != nullptr)
+				{
+					InternalRequestData.QuantizedRequestData = Handle->GetQuartzSubsystem()->CreateDataDataForSchedulePlaySound(Handle, PendingData.Delegate, PendingData.AnticapatoryBoundary);
+					UGameplayStatics::PrimeSound(Sound);
+				}
+
+				// validate clock existence 
+				if (!Handle)
+				{
+					UE_LOG(LogAudioQuartz, Warning, TEXT("Attempting to play Quantized Sound without supplying a Clock Handle"));
+					bIsValidCommand = false;
+				}
+				else if (!Handle->DoesClockExist(WorldContextObject))
+				{
+					UE_LOG(LogAudioQuartz, Warning, TEXT("Clock: '%s' Does not exist! Cannot play quantized sound: '%s'"), *InternalRequestData.QuantizedRequestData.ClockName.ToString(), *this->Sound->GetName());
+					bIsValidCommand = false;
+				}
+
+				// was the sound already stopped while we were caching it?
+				if (PendingData.bHasBeenStoppedWhileQueued)
+				{
+					InternalRequestData.QuantizedRequestData.QuantizedCommandPtr->FailedToQueue(InternalRequestData.QuantizedRequestData);
+					UE_LOG(LogAudioQuartz, Verbose, TEXT("Sound (%s) to be played (on Clock: %s) was stopped before being evaluated to play internally"), *InternalRequestData.QuantizedRequestData.ClockName.ToString(), *this->Sound->GetName());
+					bIsValidCommand = false;
+				}
+
 
-			// was the sound already stopped while we were caching it?
-			if (PendingData.bHasBeenStoppedWhileQueued)
+				if (bIsValidCommand)
+				{
+					InternalRequestData.QuantizedRequestData.GameThreadSubscribers.Add(CommandQueuePtr);
+
+					// Now play the quartz command
+					PlayInternal(InternalRequestData);
+				} 
+			}
+			else // this component does not have a Sound to play
 			{
-				InternalRequestData.QuantizedRequestData.QuantizedCommandPtr->FailedToQueue(InternalRequestData.QuantizedRequestData);
-				UE_LOG(LogAudioQuartz, Verbose, TEXT("Sound (%s) to be played (on Clock: %s) was stopped before being evaluated to play internally"), *InternalRequestData.QuantizedRequestData.ClockName.ToString(), *this->Sound->GetName());
-				bIsValidCommand = false;
+				UE_LOG(LogAudioQuartz, Warning, TEXT("Attempting to play Quantized Sound without supplying a valid Sound to play"));
 			}
+			
 
-			if (bIsValidCommand)
-			{
-				InternalRequestData.QuantizedRequestData.GameThreadSubscribers.Add(CommandQueuePtr);
 
-				// Now play the quartz command
-				PlayInternal(InternalRequestData);
-			} 
 
 			// remove the pending quartz command data from the audio component
 			bFoundQuantizedCommand = true;
