World Partition: added possibility to cull 2d streaming cells by height.

#rnx

diff --git a/Engine/Source/Runtime/Engine/Private/WorldPartition/WorldPartitionRuntimeCell.cpp b/Engine/Source/Runtime/Engine/Private/WorldPartition/WorldPartitionRuntimeCell.cpp
--- a/Engine/Source/Runtime/Engine/Private/WorldPartition/WorldPartitionRuntimeCell.cpp
+++ b/Engine/Source/Runtime/Engine/Private/WorldPartition/WorldPartitionRuntimeCell.cpp
@@ -14,6 +14,7 @@ int32 UWorldPartitionRuntimeCell::StreamingSourceCacheEpoch = 0;
 UWorldPartitionRuntimeCell::UWorldPartitionRuntimeCell(const FObjectInitializer& ObjectInitializer)
 : Super(ObjectInitializer)
 , bIsAlwaysLoaded(false)
+, MinMaxZ(-WORLDPARTITION_MAX, WORLDPARTITION_MAX)
 , Priority(0)
 , CachedMinSourcePriority((uint8)EStreamingSourcePriority::Lowest)
 , CachedSourceInfoEpoch(INT_MIN)
@@ -36,6 +37,11 @@ void UWorldPartitionRuntimeCell::SetDataLayers(const TArray<const UDataLayerInst
 	UpdateDebugName();
 }
 
+void UWorldPartitionRuntimeCell::SetMinMaxZ(const FVector2D& InMinMaxZ)
+{
+	MinMaxZ = InMinMaxZ;
+}
+
 void UWorldPartitionRuntimeCell::SetDebugInfo(FIntVector InCoords, FName InGridName)
 {
 	Coords = InCoords;
diff --git a/Engine/Source/Runtime/Engine/Private/WorldPartition/WorldPartitionRuntimeSpatialHash.cpp b/Engine/Source/Runtime/Engine/Private/WorldPartition/WorldPartitionRuntimeSpatialHash.cpp
--- a/Engine/Source/Runtime/Engine/Private/WorldPartition/WorldPartitionRuntimeSpatialHash.cpp
+++ b/Engine/Source/Runtime/Engine/Private/WorldPartition/WorldPartitionRuntimeSpatialHash.cpp
@@ -32,6 +32,7 @@
 #include "RenderUtils.h"
 #include "Algo/Transform.h"
 #include "Algo/Copy.h"
+#include "Math/Range.h"
 #include "UObject/ObjectSaveContext.h"
 #include "Components/LineBatchComponent.h"
 
@@ -130,7 +131,7 @@ int32 FSpatialHashStreamingGrid::GetCellSize(int32 Level) const
 	return GetGridHelper().Levels[Level].CellSize;
 }
 
-void FSpatialHashStreamingGrid::GetCells(const FWorldPartitionStreamingQuerySource& QuerySource, TSet<const UWorldPartitionRuntimeCell*>& OutCells) const
+void FSpatialHashStreamingGrid::GetCells(const FWorldPartitionStreamingQuerySource& QuerySource, TSet<const UWorldPartitionRuntimeCell*>& OutCells, bool bEnableZCulling) const
 {
 	TRACE_CPUPROFILER_EVENT_SCOPE(FSpatialHashStreamingGrid::GetCells_QuerySource);
 
@@ -152,6 +153,7 @@ void FSpatialHashStreamingGrid::GetCells(const FWorldPartitionStreamingQuerySour
 	};
 
 	const FSquare2DGridHelper& Helper = GetGridHelper();
+
 	// Spatial Query
 	if (QuerySource.bSpatialQuery)
 	{
@@ -163,9 +165,12 @@ void FSpatialHashStreamingGrid::GetCells(const FWorldPartitionStreamingQuerySour
 				{
 					for (const UWorldPartitionRuntimeCell* Cell : LayerCell->GridCells)
 					{
-						if (ShouldAddCell(Cell, QuerySource))
+						if (!bEnableZCulling || TRange<double>(Cell->GetMinMaxZ().X, Cell->GetMinMaxZ().Y).Overlaps(TRange<double>(Shape.GetCenter().Z - Shape.GetRadius(), Shape.GetCenter().Z + Shape.GetRadius())))
 						{
-							OutCells.Add(Cell);
+							if (ShouldAddCell(Cell, QuerySource))
+							{
+								OutCells.Add(Cell);
+							}
 						}
 					}
 				}
@@ -187,7 +192,7 @@ void FSpatialHashStreamingGrid::GetCells(const FWorldPartitionStreamingQuerySour
 	}	
 }
 
-void FSpatialHashStreamingGrid::GetCells(const TArray<FWorldPartitionStreamingSource>& Sources, const UDataLayerSubsystem* DataLayerSubsystem, UWorldPartitionRuntimeHash::FStreamingSourceCells& OutActivateCells, UWorldPartitionRuntimeHash::FStreamingSourceCells& OutLoadCells) const
+void FSpatialHashStreamingGrid::GetCells(const TArray<FWorldPartitionStreamingSource>& Sources, const UDataLayerSubsystem* DataLayerSubsystem, UWorldPartitionRuntimeHash::FStreamingSourceCells& OutActivateCells, UWorldPartitionRuntimeHash::FStreamingSourceCells& OutLoadCells, bool bEnableZCulling) const
 {
 	TRACE_CPUPROFILER_EVENT_SCOPE(FSpatialHashStreamingGrid::GetCells);
 	
@@ -213,23 +218,26 @@ void FSpatialHashStreamingGrid::GetCells(const TArray<FWorldPartitionStreamingSo
 					{
 						for (const UWorldPartitionRuntimeCell* Cell : LayerCell->GridCells)
 						{
-							if (!Cell->HasDataLayers() || (DataLayerSubsystem && DataLayerSubsystem->IsAnyDataLayerInEffectiveRuntimeState(Cell->GetDataLayers(), EDataLayerRuntimeState::Activated)))
+							if (!bEnableZCulling || TRange<double>(Cell->GetMinMaxZ().X, Cell->GetMinMaxZ().Y).Overlaps(TRange<double>(Shape.GetCenter().Z - Shape.GetRadius(), Shape.GetCenter().Z + Shape.GetRadius())))
 							{
-								if (Source.TargetState == EStreamingSourceTargetState::Loaded)
+								if (!Cell->HasDataLayers() || (DataLayerSubsystem && DataLayerSubsystem->IsAnyDataLayerInEffectiveRuntimeState(Cell->GetDataLayers(), EDataLayerRuntimeState::Activated)))
 								{
-									OutLoadCells.AddCell(Cell, Info);
+									if (Source.TargetState == EStreamingSourceTargetState::Loaded)
+									{
+										OutLoadCells.AddCell(Cell, Info);
+									}
+									else
+									{
+										check(Source.TargetState == EStreamingSourceTargetState::Activated);
+										OutActivateCells.AddCell(Cell, Info);
+										bAddedActivatedCell = !GRuntimeSpatialHashUseAlignedGridLevels && GRuntimeSpatialHashSnapNonAlignedGridLevelsToLowerLevels;
+									}
 								}
-								else
+								else if (DataLayerSubsystem && DataLayerSubsystem->IsAnyDataLayerInEffectiveRuntimeState(Cell->GetDataLayers(), EDataLayerRuntimeState::Loaded))
 								{
-									check(Source.TargetState == EStreamingSourceTargetState::Activated);
-									OutActivateCells.AddCell(Cell, Info);
-									bAddedActivatedCell = !GRuntimeSpatialHashUseAlignedGridLevels && GRuntimeSpatialHashSnapNonAlignedGridLevelsToLowerLevels;
+									OutLoadCells.AddCell(Cell, Info);
 								}
 							}
-							else if (DataLayerSubsystem && DataLayerSubsystem->IsAnyDataLayerInEffectiveRuntimeState(Cell->GetDataLayers(), EDataLayerRuntimeState::Loaded))
-							{
-								OutLoadCells.AddCell(Cell, Info);
-							}
 						}
 					}
 				}
@@ -991,10 +999,14 @@ bool UWorldPartitionRuntimeSpatialHash::CreateStreamingGrid(const FSpatialHashRu
 					}
 				}
 
+				FVector2D CellMinMaxZ(UE_BIG_NUMBER, -UE_BIG_NUMBER);
 				for (const FActorInstance& ActorInstance : FilteredActors)
 				{
 					const FWorldPartitionActorDescView& ActorDescView = ActorInstance.GetActorDescView();
 					StreamingCell->AddActorToCell(ActorDescView, ActorInstance.ContainerInstance->ID, ActorInstance.ContainerInstance->Transform, ActorInstance.ContainerInstance->Container);
+					
+					CellMinMaxZ.X = FMath::Min(CellMinMaxZ.X, ActorDescView.GetBounds().Min.Z);
+					CellMinMaxZ.Y = FMath::Max(CellMinMaxZ.Y, ActorDescView.GetBounds().Max.Z);
 
 					if (ActorInstance.ContainerInstance->ID.IsMainContainer() && StreamingCell->UnsavedActorsContainer)
 					{
@@ -1014,6 +1026,7 @@ bool UWorldPartitionRuntimeSpatialHash::CreateStreamingGrid(const FSpatialHashRu
 					}
 					UE_LOG(LogWorldPartition, Verbose, TEXT("  Actor : %s (%s) (Container %s)"), *(ActorDescView.GetActorPath().ToString()), *ActorDescView.GetGuid().ToString(EGuidFormats::UniqueObjectGuid), *ActorInstance.ContainerInstance->ID.ToString());
 				}
+				StreamingCell->SetMinMaxZ(CellMinMaxZ);
 
 				if (IsRunningCookCommandlet())
 				{
@@ -1169,7 +1182,7 @@ bool UWorldPartitionRuntimeSpatialHash::GetStreamingCells(const FWorldPartitionS
 	{
 		if (!StreamingGrid.bClientOnlyVisible || bShouldConsiderClientOnlyVisible)
 		{
-			StreamingGrid.GetCells(QuerySource, OutCells);
+			StreamingGrid.GetCells(QuerySource, OutCells, bEnableZCulling);
 		}
 	}
 
@@ -1201,7 +1214,7 @@ bool UWorldPartitionRuntimeSpatialHash::GetStreamingCells(const TArray<FWorldPar
 		{
 			if (!StreamingGrid.bClientOnlyVisible || bShouldConsiderClientOnlyVisible)
 			{
-				StreamingGrid.GetCells(Sources, DataLayerSubsystem, OutActivateCells, OutLoadCells);
+				StreamingGrid.GetCells(Sources, DataLayerSubsystem, OutActivateCells, OutLoadCells, bEnableZCulling);
 			}
 		}
 	}
diff --git a/Engine/Source/Runtime/Engine/Public/WorldPartition/WorldPartitionRuntimeCell.h b/Engine/Source/Runtime/Engine/Public/WorldPartition/WorldPartitionRuntimeCell.h
--- a/Engine/Source/Runtime/Engine/Public/WorldPartition/WorldPartitionRuntimeCell.h
+++ b/Engine/Source/Runtime/Engine/Public/WorldPartition/WorldPartitionRuntimeCell.h
@@ -161,6 +161,8 @@ class UWorldPartitionRuntimeCell : public UObject, public IWorldPartitionCell
 	}
 	// IRuntimeCell Interface end
 
+	const FVector2D& GetMinMaxZ() const { return MinMaxZ; }
+
 	bool GetBlockOnSlowLoading() const { return bBlockOnSlowLoading; }
 #if WITH_EDITOR
 	void SetBlockOnSlowLoading(bool bInBlockOnSlowLoading) { bBlockOnSlowLoading = bInBlockOnSlowLoading; }
@@ -169,6 +171,7 @@ class UWorldPartitionRuntimeCell : public UObject, public IWorldPartitionCell
 	bool GetClientOnlyVisible() const { return bClientOnlyVisible; }
 
 	void SetDataLayers(const TArray<const UDataLayerInstance*>& InDataLayerInstances);
+	void SetMinMaxZ(const FVector2D& InMinMaxZ);
 	void SetDebugInfo(FIntVector InCoords, FName InGridName);
 	virtual void AddActorToCell(const FWorldPartitionActorDescView& ActorDescView, const FActorContainerID& InContainerID, const FTransform& InContainerTransform, const UActorDescContainer* InContainer) PURE_VIRTUAL(UWorldPartitionRuntimeCell::AddActorToCell,);
 	virtual int32 GetActorCount() const PURE_VIRTUAL(UWorldPartitionRuntimeCell::GetActorCount, return 0;);
@@ -202,6 +205,9 @@ private:
 	UPROPERTY()
 	TArray<FName> DataLayers;
 
+	UPROPERTY()
+	FVector2D MinMaxZ;
+
 	// Debug Info
 	UPROPERTY()
 	FIntVector Coords;
diff --git a/Engine/Source/Runtime/Engine/Public/WorldPartition/WorldPartitionRuntimeSpatialHash.h b/Engine/Source/Runtime/Engine/Public/WorldPartition/WorldPartitionRuntimeSpatialHash.h
--- a/Engine/Source/Runtime/Engine/Public/WorldPartition/WorldPartitionRuntimeSpatialHash.h
+++ b/Engine/Source/Runtime/Engine/Public/WorldPartition/WorldPartitionRuntimeSpatialHash.h
@@ -106,7 +106,7 @@ struct FSpatialHashStreamingGrid
 
 	UPROPERTY()
 	bool bBlockOnSlowStreaming;
-		
+
 	UPROPERTY()
 	FLinearColor DebugColor;
 
@@ -135,8 +135,8 @@ struct FSpatialHashStreamingGrid
 
 	// Used by PIE/Game
 	int32 GetCellSize(int32 Level) const;
-	void GetCells(const FWorldPartitionStreamingQuerySource& QuerySource, TSet<const UWorldPartitionRuntimeCell*>& OutCells) const;
-	void GetCells(const TArray<FWorldPartitionStreamingSource>& Sources, const class UDataLayerSubsystem* DataLayerSubsystem, UWorldPartitionRuntimeHash::FStreamingSourceCells& OutActivateCells, UWorldPartitionRuntimeHash::FStreamingSourceCells& OutLoadCells) const;
+	void GetCells(const FWorldPartitionStreamingQuerySource& QuerySource, TSet<const UWorldPartitionRuntimeCell*>& OutCells, bool bEnableZCulling) const;
+	void GetCells(const TArray<FWorldPartitionStreamingSource>& Sources, const class UDataLayerSubsystem* DataLayerSubsystem, UWorldPartitionRuntimeHash::FStreamingSourceCells& OutActivateCells, UWorldPartitionRuntimeHash::FStreamingSourceCells& OutLoadCells, bool bEnableZCulling) const;
 	void GetAlwaysLoadedCells(const UDataLayerSubsystem* DataLayerSubsystem, TSet<const UWorldPartitionRuntimeCell*>& OutActivateCells, TSet<const UWorldPartitionRuntimeCell*>& OutLoadCells) const;
 	void Draw2D(UCanvas* Canvas, UWorld* World, const TArray<FWorldPartitionStreamingSource>& Sources, const FBox& Region, const FBox2D& GridScreenBounds, TFunctionRef<FVector2D(const FVector2D&)> WorldToScreen) const;
 	void Draw3D(UWorld* World, const TArray<FWorldPartitionStreamingSource>& Sources, const FTransform& Transform) const;
@@ -187,7 +187,7 @@ struct FSpatialHashRuntimeGrid
 	/** Should streaming block in situations where cells aren't getting loaded fast enough. */
 	UPROPERTY(EditAnywhere, Category = Settings)
 	bool bBlockOnSlowStreaming;
-		
+
 	UPROPERTY(EditAnywhere, Category=Settings)
 	int32 Priority;
 
@@ -275,6 +275,10 @@ private:
 	TMap<FString, UWorldPartitionRuntimeCell*> PackagesToGenerateForCook;
 #endif
 
+	/** Whether this hash enables Z culling. */
+	UPROPERTY(EditAnywhere, Config, Category = RuntimeSettings)
+	bool bEnableZCulling;
+
 	/** 
 	 * Represents the streaming grids (PIE or Game)
 	 */
