[WP Navmesh] Temporarily disable streaming bounds expansion by agent radius

diff --git a/Engine/Source/Runtime/Engine/Private/Components/SceneComponent.cpp b/Engine/Source/Runtime/Engine/Private/Components/SceneComponent.cpp
--- a/Engine/Source/Runtime/Engine/Private/Components/SceneComponent.cpp
+++ b/Engine/Source/Runtime/Engine/Private/Components/SceneComponent.cpp
@@ -2398,12 +2398,13 @@ TStructOnScope<FActorComponentInstanceData> USceneComponent::GetComponentInstanc
 FBox USceneComponent::GetStreamingBounds() const
 {
 	FBox Box = Bounds.GetBox();
-	
-	if (IsNavigationRelevant())
-	{
-		const FNavDataConfig& Config = FNavigationSystem::GetBiggestSupportedAgent(GetWorld());
-		Box = Box.ExpandBy(Config.AgentRadius);
-	}
+
+	// Temporarily disabled while we resolve why Config.AgentRadius is sometime humongous.
+	// if (IsNavigationRelevant())
+	// {
+	// 	const FNavDataConfig& Config = FNavigationSystem::GetBiggestSupportedAgent(GetWorld());
+	// 	Box = Box.ExpandBy(Config.AgentRadius);
+	// }
 
 	return Box;
 }
