Adding size sanity checking for serialized net export groups in replays

diff --git a/Engine/Source/Runtime/Engine/Private/PackageMapClient.cpp b/Engine/Source/Runtime/Engine/Private/PackageMapClient.cpp
--- a/Engine/Source/Runtime/Engine/Private/PackageMapClient.cpp
+++ b/Engine/Source/Runtime/Engine/Private/PackageMapClient.cpp
@@ -60,6 +60,18 @@ namespace UE
 	namespace Net
 	{
 		extern int32 FilterGuidRemapping;
+
+		int32 MaxSerializedNetGuids = 2048;
+		static FAutoConsoleVariableRef CVarMaxSerializedNetGuids(TEXT("net.MaxSerializedNetGuids"), MaxSerializedNetGuids, TEXT("Maximum number of network guids we would expect to receive in a bunch"));
+
+		int32 MaxSerializedReplayNetGuids = 16 * 1024;
+		static FAutoConsoleVariableRef CVarMaxSerializedReplayNetGuids(TEXT("net.MaxSerializedReplayNetGuids"), MaxSerializedReplayNetGuids, TEXT("Maximum number of network guids we would expect to receive in replay export data."));
+
+		int32 MaxSerializedNetExportGroups = 64 * 1024;
+		static FAutoConsoleVariableRef CVarMaxSerializedNetExportGroups(TEXT("net.MaxSerializedNetExportGroups"), MaxSerializedNetExportGroups, TEXT("Maximum number of network export groups we would expect to receive in a bunch"));
+
+		int32 MaxSerializedNetExportsPerGroup = 128 * 1024;	// Gameplay tags will be exported into a single large group for replays
+		static FAutoConsoleVariableRef CVarMaxSerializedNetExportsPerGroup(TEXT("net.MaxSerializedNetExportsPerGroup"), MaxSerializedNetExportsPerGroup, TEXT("Maximum number of network exports in each group we would expect to receive in a bunch"));
 	};
 };
 
@@ -1159,10 +1171,6 @@ UObject* UPackageMapClient::ResolvePathAndAssignNetGUID( const FNetworkGUID& Net
 //
 //--------------------------------------------------------------------
 
-
-// TODO: This limit might not actually need to be enforced anymore.
-constexpr int32 MAX_GUID_COUNT = 2048;
-
 bool UPackageMapClient::ExportNetGUIDForReplay(FNetworkGUID& NetGUID, UObject* Object, FString& PathName, UObject* ObjOuter)
 {
 	int32 const * const FoundExpectedPacketIdPtr = OverrideAckState->NetGUIDAckStatus.Find(NetGUID);
@@ -1378,9 +1386,9 @@ void UPackageMapClient::ReceiveNetGUIDBunch( FInBunch &InBunch )
 	int32 NumGUIDsInBunch = 0;
 	InBunch << NumGUIDsInBunch;
 
-	if ( NumGUIDsInBunch > MAX_GUID_COUNT )
+	if ( NumGUIDsInBunch > UE::Net::MaxSerializedNetGuids )
 	{
-		UE_LOG( LogNetPackageMap, Error, TEXT( "UPackageMapClient::ReceiveNetGUIDBunch: NumGUIDsInBunch > MAX_GUID_COUNT (%i)" ), NumGUIDsInBunch );
+		UE_LOG( LogNetPackageMap, Error, TEXT( "UPackageMapClient::ReceiveNetGUIDBunch: NumGUIDsInBunch > MaxSerializedNetGuids (%d / %d)" ), NumGUIDsInBunch, UE::Net::MaxSerializedNetGuids);
 		InBunch.SetError();
 		return;
 	}
@@ -1409,7 +1417,7 @@ void UPackageMapClient::ReceiveNetGUIDBunch( FInBunch &InBunch )
 		NumGUIDsRead++;
 	}
 
-	UE_LOG(LogNetPackageMap, Log, TEXT("UPackageMapClient::ReceiveNetGUIDBunch end. BitPos: %d"), InBunch.GetPosBits() );
+	UE_LOG(LogNetPackageMap, Log, TEXT("UPackageMapClient::ReceiveNetGUIDBunch end. BitPos: %d"), InBunch.GetPosBits());
 }
 
 TSharedPtr<FNetFieldExportGroup> UPackageMapClient::GetNetFieldExportGroup(const FString& PathName)
@@ -1417,21 +1425,21 @@ TSharedPtr<FNetFieldExportGroup> UPackageMapClient::GetNetFieldExportGroup(const
 	return GuidCache->NetFieldExportGroupMap.FindRef(PathName);
 }
 
-void UPackageMapClient::AddNetFieldExportGroup( const FString& PathName, TSharedPtr< FNetFieldExportGroup > NewNetFieldExportGroup )
+void UPackageMapClient::AddNetFieldExportGroup(const FString& PathName, TSharedPtr< FNetFieldExportGroup > NewNetFieldExportGroup)
 {
-	check( !GuidCache->NetFieldExportGroupMap.Contains( NewNetFieldExportGroup->PathName ) );
+	check(!GuidCache->NetFieldExportGroupMap.Contains(NewNetFieldExportGroup->PathName));
 
 	NewNetFieldExportGroup->PathNameIndex = ++GuidCache->UniqueNetFieldExportGroupPathIndex;
 
-	check( !GuidCache->NetFieldExportGroupPathToIndex.Contains( NewNetFieldExportGroup->PathName ) );
-	check( !GuidCache->NetFieldExportGroupIndexToGroup.Contains( NewNetFieldExportGroup->PathNameIndex ) );
+	check(!GuidCache->NetFieldExportGroupPathToIndex.Contains(NewNetFieldExportGroup->PathName));
+	check(!GuidCache->NetFieldExportGroupIndexToGroup.Contains(NewNetFieldExportGroup->PathNameIndex));
 
 	GuidCache->NetFieldExportGroupPathToIndex.Add(NewNetFieldExportGroup->PathName, NewNetFieldExportGroup->PathNameIndex);
-	GuidCache->NetFieldExportGroupIndexToGroup.Add(NewNetFieldExportGroup->PathNameIndex, NewNetFieldExportGroup.Get() );
-	GuidCache->NetFieldExportGroupMap.Add( NewNetFieldExportGroup->PathName, NewNetFieldExportGroup );
+	GuidCache->NetFieldExportGroupIndexToGroup.Add(NewNetFieldExportGroup->PathNameIndex, NewNetFieldExportGroup.Get());
+	GuidCache->NetFieldExportGroupMap.Add(NewNetFieldExportGroup->PathName, NewNetFieldExportGroup);
 }
 
-void UPackageMapClient::TrackNetFieldExport( FNetFieldExportGroup* NetFieldExportGroup, const int32 NetFieldExportHandle )
+void UPackageMapClient::TrackNetFieldExport(FNetFieldExportGroup* NetFieldExportGroup, const int32 NetFieldExportHandle)
 {
 	check(Connection->IsInternalAck());
 	check(NetFieldExportGroup);
@@ -1447,23 +1455,23 @@ void UPackageMapClient::TrackNetFieldExport( FNetFieldExportGroup* NetFieldExpor
 
 	NetFieldExportGroup->NetFieldExports[NetFieldExportHandle].bExported = true;
 
-	const uint64 CmdHandle = ( ( uint64 )NetFieldExportGroup->PathNameIndex ) << 32 | ( uint64 )NetFieldExportHandle;
+	const uint64 CmdHandle = ((uint64)NetFieldExportGroup->PathNameIndex) << 32 | (uint64)NetFieldExportHandle;
 
 	// If this cmd hasn't been confirmed as exported, we need to export it for this bunch
-	if ( !OverrideAckState->NetFieldExportAcked.Contains( CmdHandle ) )
+	if (!OverrideAckState->NetFieldExportAcked.Contains(CmdHandle))
 	{
-		NetFieldExports.Add( CmdHandle );		// NOTE - This is a set, so it will only add once
+		NetFieldExports.Add(CmdHandle);		// NOTE - This is a set, so it will only add once
 	}
 }
 
-TSharedPtr< FNetFieldExportGroup > UPackageMapClient::GetNetFieldExportGroupChecked( const FString& PathName ) const
+TSharedPtr< FNetFieldExportGroup > UPackageMapClient::GetNetFieldExportGroupChecked(const FString& PathName) const
 {
-	return GuidCache->NetFieldExportGroupMap.FindChecked( PathName );
+	return GuidCache->NetFieldExportGroupMap.FindChecked(PathName);
 }
 
-void UPackageMapClient::SerializeNetFieldExportGroupMap( FArchive& Ar, bool bClearPendingExports )
+void UPackageMapClient::SerializeNetFieldExportGroupMap(FArchive& Ar, bool bClearPendingExports)
 {
-	if ( Ar.IsSaving() )
+	if (Ar.IsSaving())
 	{
 		if (bClearPendingExports)
 		{
@@ -1475,7 +1483,7 @@ void UPackageMapClient::SerializeNetFieldExportGroupMap( FArchive& Ar, bool bCle
 		Ar << NumNetFieldExportGroups;
 
 		// Save each layout
-		for ( auto It = GuidCache->NetFieldExportGroupMap.CreateIterator(); It; ++It )
+		for (auto It = GuidCache->NetFieldExportGroupMap.CreateIterator(); It; ++It)
 		{
 			// Save out the export group
 			Ar << *It.Value().Get();
@@ -1494,7 +1502,14 @@ void UPackageMapClient::SerializeNetFieldExportGroupMap( FArchive& Ar, bool bCle
 
 		if (Ar.IsError())
 		{
-			UE_LOG(LogNetPackageMap, Warning, TEXT("UPackageMapClient::SerializeNetFieldExportGroupMap - Archive error while reading NumNetFieldExportGroups"));
+			UE_LOG(LogNetPackageMap, Error, TEXT("UPackageMapClient::SerializeNetFieldExportGroupMap - Archive error while reading NumNetFieldExportGroups"));
+			return;
+		}
+
+		if (NumNetFieldExportGroups > (uint32)UE::Net::MaxSerializedNetExportGroups)
+		{
+			Ar.SetError();
+			UE_LOG(LogNetPackageMap, Error, TEXT("UPackageMapClient::SerializeNetFieldExportGroupMap - NumNetFieldExportGroups exceeds MaxSerializedNetExportGroups (%u / %d)"), NumNetFieldExportGroups, UE::Net::MaxSerializedNetExportGroups);
 			return;
 		}
 
@@ -1508,7 +1523,7 @@ void UPackageMapClient::SerializeNetFieldExportGroupMap( FArchive& Ar, bool bCle
 
 			if (Ar.IsError())
 			{
-				UE_LOG(LogNetPackageMap, Warning, TEXT("UPackageMapClient::SerializeNetFieldExportGroupMap - Archive error while loading FNetFieldExportGroup, Index: %u"), i);
+				UE_LOG(LogNetPackageMap, Error, TEXT("UPackageMapClient::SerializeNetFieldExportGroupMap - Archive error while loading FNetFieldExportGroup, Index: %u"), i);
 				return;
 			}
 
@@ -1654,10 +1669,17 @@ void UPackageMapClient::ReceiveNetFieldExportsCompat(FInBunch &InBunch)
 	}
 
 	// Read number of net field exports
-	uint32 NumLayoutCmdExports = 0;
-	InBunch << NumLayoutCmdExports;
+	uint32 NumExportGroups = 0;
+	InBunch << NumExportGroups;
+
+	if (NumExportGroups > (uint32)UE::Net::MaxSerializedNetExportGroups)
+	{
+		UE_LOG(LogNetPackageMap, Error, TEXT("UPackageMapClient::ReceiveNetFieldExportsCompat - NumExportGroups exceeds MaxSerializedNetExportGroups (%u / %d)"), NumExportGroups, UE::Net::MaxSerializedNetExportGroups);
+		InBunch.SetError();
+		return;
+	}
 
-	for (int32 i = 0; i < (int32)NumLayoutCmdExports; i++)
+	for (int32 i = 0; i < (int32)NumExportGroups; i++)
 	{
 		// Read the index that represents the name in the NetFieldExportGroupIndexToPath map
 		uint32 PathNameIndex;
@@ -1665,10 +1687,11 @@ void UPackageMapClient::ReceiveNetFieldExportsCompat(FInBunch &InBunch)
 
 		if (InBunch.IsError())
 		{
-			break;
+			UE_LOG(LogNetPackageMap, Error, TEXT("UPackageMapClient::ReceiveNetFieldExportsCompat - Error serializing export path index."));
+			return;
 		}
 
-		int32 MaxExports = 0;
+		int32 NumExportsInGroup = 0;
 
 		FNetFieldExportGroup* NetFieldExportGroup = nullptr;
 
@@ -1676,13 +1699,21 @@ void UPackageMapClient::ReceiveNetFieldExportsCompat(FInBunch &InBunch)
 		if (InBunch.ReadBit() == 1)
 		{
 			FString PathName;
-
 			InBunch << PathName;
-			InBunch << MaxExports;
 
 			if (InBunch.IsError())
 			{
-				break;
+				UE_LOG(LogNetPackageMap, Error, TEXT("UPackageMapClient::ReceiveNetFieldExportsCompat - Error serializing export path."));
+				return;
+			}
+
+			InBunch << NumExportsInGroup;
+
+			if (NumExportsInGroup > UE::Net::MaxSerializedNetExportsPerGroup)
+			{
+				UE_LOG(LogNetPackageMap, Error, TEXT("UPackageMapClient::ReceiveNetFieldExportsCompat - NumExportsInGroup exceeds MaxSerializedNetExportsPerGroup (%d / %d)"), NumExportsInGroup, UE::Net::MaxSerializedNetExportsPerGroup);
+				InBunch.SetError();
+				return;
 			}
 
 			GEngine->NetworkRemapPath(Connection, PathName, true);
@@ -1696,7 +1727,7 @@ void UPackageMapClient::ReceiveNetFieldExportsCompat(FInBunch &InBunch)
 				NetFieldExportGroup->PathName = PathName;
 				NetFieldExportGroup->PathNameIndex = PathNameIndex;
 
-				NetFieldExportGroup->NetFieldExports.SetNum(MaxExports);
+				NetFieldExportGroup->NetFieldExports.SetNum(NumExportsInGroup);
 
 				GuidCache->NetFieldExportGroupMap.Add(PathName, NewNetFieldExportGroup);
 			}
@@ -1717,7 +1748,7 @@ void UPackageMapClient::ReceiveNetFieldExportsCompat(FInBunch &InBunch)
 
 		if (InBunch.IsError())
 		{
-			break;
+			return;
 		}
 
 		TArray<FNetFieldExport>& NetFieldExportsRef = NetFieldExportGroup->NetFieldExports;
@@ -1733,6 +1764,7 @@ void UPackageMapClient::ReceiveNetFieldExportsCompat(FInBunch &InBunch)
 				NetFieldExport.Handle, NetFieldExportsRef.Num());
 
 			InBunch.SetError();
+			return;
 		}
 	}
 }
@@ -1744,10 +1776,17 @@ void UPackageMapClient::ReceiveNetFieldExports(FArchive& Archive)
 	check(Connection->IsInternalAck());
 
 	// Read number of net field exports
-	uint32 NumLayoutCmdExports = 0;
-	Archive.SerializeIntPacked(NumLayoutCmdExports);
+	uint32 NumNetExports = 0;
+	Archive.SerializeIntPacked(NumNetExports);
 
-	for (int32 i = 0; i < (int32)NumLayoutCmdExports; i++)
+	if (NumNetExports > (uint32)UE::Net::MaxSerializedNetExportGroups)
+	{
+		UE_LOG(LogNetPackageMap, Error, TEXT("UPackageMapClient::ReceiveNetFieldExports - NumNetExports exceeds MaxSerializedNetExportGroups (%u / %d)"), NumNetExports, UE::Net::MaxSerializedNetExportGroups);
+		Archive.SetError();
+		return;
+	}
+
+	for (int32 i = 0; i < (int32)NumNetExports; i++)
 	{
 		uint32 PathNameIndex = 0;
 		uint32 WasExported = 0;
@@ -1759,10 +1798,17 @@ void UPackageMapClient::ReceiveNetFieldExports(FArchive& Archive)
 		if (!!WasExported)
 		{
 			FString PathName;
-			uint32 NumExports = 0;
+			uint32 NumExportsInGroup = 0;
 
 			Archive << PathName;
-			Archive.SerializeIntPacked(NumExports);
+			Archive.SerializeIntPacked(NumExportsInGroup);
+
+			if (NumExportsInGroup > (uint32)UE::Net::MaxSerializedNetExportsPerGroup)
+			{
+				UE_LOG(LogNetPackageMap, Warning, TEXT("UPackageMapClient::ReceiveNetFieldExports - NumExportsInGroup exceeds MaxSerializedNetExportsPerGroup (%u / %d)"), NumExportsInGroup, UE::Net::MaxSerializedNetExportsPerGroup);
+				Archive.SetError();
+				return;
+			}
 
 			GEngine->NetworkRemapPath(Connection, PathName, true);
 
@@ -1774,7 +1820,7 @@ void UPackageMapClient::ReceiveNetFieldExports(FArchive& Archive)
 
 				NetFieldExportGroup->PathName = PathName;
 				NetFieldExportGroup->PathNameIndex = PathNameIndex;
-				NetFieldExportGroup->NetFieldExports.SetNum(NumExports);
+				NetFieldExportGroup->NetFieldExports.SetNum(NumExportsInGroup);
 
 				GuidCache->NetFieldExportGroupMap.Add(PathName, NewNetFieldExportGroup);
 			}
@@ -1838,6 +1884,18 @@ void UPackageMapClient::ReceiveNetExportGUIDs(FArchive& Archive)
 	uint32 NumGUIDs = 0;
 	Archive.SerializeIntPacked(NumGUIDs);
 
+	if (Archive.IsError())
+	{
+		return;
+	}
+
+	if (NumGUIDs > (uint32)UE::Net::MaxSerializedReplayNetGuids)
+	{
+		UE_LOG(LogNetPackageMap, Error, TEXT("UPackageMapClient::ReceiveNetExportGUIDs: NumGUIDs > MaxSerializedReplayNetGuids (%u / %d)"), NumGUIDs, UE::Net::MaxSerializedReplayNetGuids);
+		Archive.SetError();
+		return;
+	}
+
 	if (bIgnoreReceivedExportGUIDs)
 	{
 		// Array Serialization works by first serializing an int32 count, and then serializing
