Removed some debug code I accidentally submitted and fixed the issue it was trying to identify.
I did submit the wrong change when changing the marquee selection and this add the missing stuff.

diff --git a/Engine/Source/Editor/UnrealEd/Private/DragTool_BoxSelect.cpp b/Engine/Source/Editor/UnrealEd/Private/DragTool_BoxSelect.cpp
--- a/Engine/Source/Editor/UnrealEd/Private/DragTool_BoxSelect.cpp
+++ b/Engine/Source/Editor/UnrealEd/Private/DragTool_BoxSelect.cpp
@@ -37,14 +37,11 @@ namespace UE::LevelEditor::Private
 	{
 		if (Actor && (!Viewport || !Actor->IsA(AVolume::StaticClass()) || !Viewport->IsVolumeVisibleInViewport(*Actor)))
 		{
-			if (Actor->GetActorLabel() == FString(TEXT("ChaosAsset_Floor3")))
+			if (FTypedElementHandle ActorHandle = UEngineElementsLibrary::AcquireEditorActorElementHandle(Actor))
 			{
-				if (FTypedElementHandle ActorHandle = UEngineElementsLibrary::AcquireEditorActorElementHandle(Actor))
+				if (TTypedElement<ITypedElementWorldInterface> WorldElement = UTypedElementRegistry::GetInstance()->GetElement<ITypedElementWorldInterface>(ActorHandle))
 				{
-					if (TTypedElement<ITypedElementWorldInterface> WorldElement = UTypedElementRegistry::GetInstance()->GetElement<ITypedElementWorldInterface>(ActorHandle))
-					{
-						return WorldElement.GetSelectionElementsInBox(InBox, SelectionArgs);
-					}
+					return WorldElement.GetSelectionElementsInBox(InBox, SelectionArgs);
 				}
 			}
 		}
diff --git a/Engine/Source/Runtime/Engine/Private/Elements/Interfaces/TypedElementWorldInterface.cpp b/Engine/Source/Runtime/Engine/Private/Elements/Interfaces/TypedElementWorldInterface.cpp
--- a/Engine/Source/Runtime/Engine/Private/Elements/Interfaces/TypedElementWorldInterface.cpp
+++ b/Engine/Source/Runtime/Engine/Private/Elements/Interfaces/TypedElementWorldInterface.cpp
@@ -79,11 +79,14 @@ TArray<FTypedElementHandle> ITypedElementWorldInterface::GetSelectionElementsFro
 	// The default implementation simply iterate the child elements
 
 	TArray<FTypedElementHandle> SelectedSubElements;
-
+	bool bTryToSelectCurrent = true;
 	if (TTypedElement<ITypedElementHierarchyInterface> HierarchyElement = GetRegistry().GetElement<ITypedElementHierarchyInterface>(InElementHandle))
 	{
 		TArray<FTypedElementHandle> ChildElementHandles;
 		HierarchyElement.GetChildElements(ChildElementHandles);
+
+		bTryToSelectCurrent = ChildElementHandles.IsEmpty();
+
 		for (const FTypedElementHandle& ChildHandle : ChildElementHandles)
 		{
 			if (TTypedElement<ITypedElementWorldInterface> ChildWorldElement = GetRegistry().GetElement<ITypedElementWorldInterface>(ChildHandle))
@@ -105,10 +108,10 @@ TArray<FTypedElementHandle> ITypedElementWorldInterface::GetSelectionElementsFro
 				}
 			}
 		}
-
-		// If all the elements selected are the child elements, Select the parent instead if possible.
-		if (!ChildElementHandles.IsEmpty() && ChildElementHandles == SelectedSubElements)
+	
+		if (!bTryToSelectCurrent && ChildElementHandles == SelectedSubElements)
 		{
+			// If all the elements selected are the child elements, add the current also if possible.
 			if (SelectionArgs.SelectionSet)
 			{
 				SelectedSubElements.Add(SelectionArgs.SelectionSet->GetSelectionElement(InElementHandle, SelectionArgs.SelectionMethod));
@@ -119,18 +122,16 @@ TArray<FTypedElementHandle> ITypedElementWorldInterface::GetSelectionElementsFro
 			}
 		}
 	}
-	else
+	
+	if (bTryToSelectCurrent && SelectionFunction(InElementHandle, SelectionArgs))
 	{
-		if (SelectionFunction(InElementHandle, SelectionArgs))
+		if (SelectionArgs.SelectionSet)
 		{
-			if (SelectionArgs.SelectionSet)
-			{
-				SelectedSubElements.Add(SelectionArgs.SelectionSet->GetSelectionElement(InElementHandle, SelectionArgs.SelectionMethod));
-			}
-			else
-			{
-				SelectedSubElements.Add(InElementHandle);
-			}
+			SelectedSubElements.Add(SelectionArgs.SelectionSet->GetSelectionElement(InElementHandle, SelectionArgs.SelectionMethod));
+		}
+		else
+		{
+			SelectedSubElements.Add(InElementHandle);
 		}
 	}
 
