Fix for incorrect requested virtual mip texture streaming size.

diff --git a/Engine/Source/Runtime/Engine/Private/Rendering/Texture2DResource.cpp b/Engine/Source/Runtime/Engine/Private/Rendering/Texture2DResource.cpp
--- a/Engine/Source/Runtime/Engine/Private/Rendering/Texture2DResource.cpp
+++ b/Engine/Source/Runtime/Engine/Private/Rendering/Texture2DResource.cpp
@@ -178,7 +178,7 @@ uint64 FTexture2DResource::GetPlatformMipsSize(uint32 NumMips) const
 		// Must be consistent with the logic in FTexture2DResource::InitRHI
 		if (bUsePartiallyResidentMips && (!CVarReducedMode->GetValueOnRenderThread() || NumMips > State.NumNonStreamingLODs))
 		{
-			return RHICalcVMTexture2DPlatformSize(SizeX, SizeY, PixelFormat, NumMips, State.LODCountToFirstLODIdx(NumMips), 1, CreationFlags | TexCreate_Virtual, TextureAlign);
+			return RHICalcVMTexture2DPlatformSize(SizeX, SizeY, PixelFormat, State.MaxNumLODs, State.LODCountToFirstLODIdx(NumMips), 1, CreationFlags | TexCreate_Virtual, TextureAlign);
 		}
 		else
 		{
