Added time of last received/sent packet and last TickFlush/TickDispatch, to UNetConnection::Close - to narrow down potential causes of some connection timeouts.

diff --git a/Engine/Source/Runtime/Engine/Private/NetConnection.cpp b/Engine/Source/Runtime/Engine/Private/NetConnection.cpp
--- a/Engine/Source/Runtime/Engine/Private/NetConnection.cpp
+++ b/Engine/Source/Runtime/Engine/Private/NetConnection.cpp
@@ -94,6 +94,11 @@ TAutoConsoleVariable<int32> CVarNetEnableCongestionControl(TEXT("net.EnableConge
 static TAutoConsoleVariable<int32> CVarLogUnhandledFaults(TEXT("net.LogUnhandledFaults"), 1,
 	TEXT("Whether or not to warn about unhandled net faults (could be deliberate, depending on implementation). 0 = off, 1 = log once, 2 = log always."));
 
+static int32 GNetCloseTimingDebug = 0;
+
+static FAutoConsoleVariableRef CVarCloseTimingDebug(TEXT("net.CloseTimingDebug"), GNetCloseTimingDebug,
+	TEXT("Logs the last packet send/receive and TickFlush/TickDispatch times, on connection close - for debugging blocked send/recv paths."));
+
 
 extern int32 GNetDormancyValidate;
 extern bool GbNetReuseReplicatorsForDormantObjects;
@@ -936,8 +941,26 @@ void UNetConnection::Close(FNetResult&& CloseReason)
 	{
 		NETWORK_PROFILER(GNetworkProfiler.TrackEvent(TEXT("CLOSE"), *(GetName() + TEXT(" ") + LowLevelGetRemoteAddress()), this));
 
-		UE_LOG(LogNet, Log, TEXT("UNetConnection::Close: %s, Channels: %i, Time: %s"), *Describe(), OpenChannels.Num(),
-				*FDateTime::UtcNow().ToString(TEXT("%Y.%m.%d-%H.%M.%S")));
+		TStringBuilder<2048> CloseStr;
+
+		CloseStr.Append(TEXT("UNetConnection::Close: "));
+		CloseStr.Append(ToCStr(Describe()));
+		CloseStr.Appendf(TEXT(", Channels: %i, Time: "), OpenChannels.Num());
+		CloseStr.Append(ToCStr(FDateTime::UtcNow().ToString(TEXT("%Y.%m.%d-%H.%M.%S"))));
+
+		if (GNetCloseTimingDebug)
+		{
+			const double CurTime = FPlatformTime::Seconds();
+			const double TimeSinceRecv = LastReceiveRealtime != 0.0 ? (CurTime - LastReceiveRealtime) : -1.0;
+			const double TimeSinceSend = PreviousPacketSendTimeInS != 0.0 ? (CurTime - PreviousPacketSendTimeInS) : -1.0;
+			const double TimeSinceTickFlush = LastTime != 0.0 ? (CurTime - LastTime) : -1.0;
+			const double TimeSinceTickDispatch = CurTime - Driver->LastTickDispatchRealtime;
+
+			CloseStr.Appendf(TEXT(", TimeSinceTickDispatch: %f, TimeSinceRecv: %f, TimeSinceTickFlush: %f, TimeSinceSend: %f"),
+								TimeSinceTickDispatch, TimeSinceRecv, TimeSinceTickFlush, TimeSinceSend);
+		}
+
+		UE_LOG(LogNet, Log, TEXT("%s"), ToCStr(CloseStr.ToString()));
 
 		SendCloseReason(MoveTemp(CloseReason));
 
