Remove incorrect assert in FVolumeTextureBulkData::MergeMips

diff --git a/Engine/Source/Runtime/Engine/Private/Rendering/Texture3DResource.cpp b/Engine/Source/Runtime/Engine/Private/Rendering/Texture3DResource.cpp
--- a/Engine/Source/Runtime/Engine/Private/Rendering/Texture3DResource.cpp
+++ b/Engine/Source/Runtime/Engine/Private/Rendering/Texture3DResource.cpp
@@ -36,7 +36,6 @@ void FVolumeTextureBulkData::MergeMips(int32 NumMips)
 	int64 MergedSize = 0;
 	for (int32 MipIndex = FirstMipIdx; MipIndex < NumMips; ++MipIndex)
 	{
-		check( MipSize[MipIndex] != 0 );
 		MergedSize += MipSize[MipIndex];
 	}
 
