World Partition: added cvar wp.Runtime.ForceRuntimeSpatialHashZCulling to allow forcing the behavior of runtime hash cells Z culling at runtime.

#rnx

diff --git a/Engine/Source/Runtime/Engine/Private/WorldPartition/WorldPartitionRuntimeSpatialHash.cpp b/Engine/Source/Runtime/Engine/Private/WorldPartition/WorldPartitionRuntimeSpatialHash.cpp
--- a/Engine/Source/Runtime/Engine/Private/WorldPartition/WorldPartitionRuntimeSpatialHash.cpp
+++ b/Engine/Source/Runtime/Engine/Private/WorldPartition/WorldPartitionRuntimeSpatialHash.cpp
@@ -88,6 +88,23 @@ static FAutoConsoleVariableRef CVarFilterRuntimeSpatialHashGridLevel(
 	TEXT("Used to choose filter a single world partition runtime hash grid level."));
 #endif
 
+static int32 GForceRuntimeSpatialHashZCulling = -1;
+static FAutoConsoleVariableRef CVarForceRuntimeSpatialHashZCulling(
+	TEXT("wp.Runtime.ForceRuntimeSpatialHashZCulling"),
+	GForceRuntimeSpatialHashZCulling,
+	TEXT("Used to force the behavior of the runtime hash cells Z culling. Set to 0 to force off, to 1 to force on and any other value to respect the runtime hash setting."));
+
+static bool GetEffectiveEnableZCulling(bool bEnableZCulling)
+{
+	switch (GForceRuntimeSpatialHashZCulling)
+	{
+	case 0: return false;
+	case 1: return true;
+	}
+
+	return bEnableZCulling;
+}
+
 // ------------------------------------------------------------------------------------------------
 FSpatialHashStreamingGrid::FSpatialHashStreamingGrid()
 	: Origin(ForceInitToZero)
@@ -1182,7 +1199,7 @@ bool UWorldPartitionRuntimeSpatialHash::GetStreamingCells(const FWorldPartitionS
 	{
 		if (!StreamingGrid.bClientOnlyVisible || bShouldConsiderClientOnlyVisible)
 		{
-			StreamingGrid.GetCells(QuerySource, OutCells, bEnableZCulling);
+			StreamingGrid.GetCells(QuerySource, OutCells, GetEffectiveEnableZCulling(bEnableZCulling));
 		}
 	}
 
@@ -1214,7 +1231,7 @@ bool UWorldPartitionRuntimeSpatialHash::GetStreamingCells(const TArray<FWorldPar
 		{
 			if (!StreamingGrid.bClientOnlyVisible || bShouldConsiderClientOnlyVisible)
 			{
-				StreamingGrid.GetCells(Sources, DataLayerSubsystem, OutActivateCells, OutLoadCells, bEnableZCulling);
+				StreamingGrid.GetCells(Sources, DataLayerSubsystem, OutActivateCells, OutLoadCells, GetEffectiveEnableZCulling(bEnableZCulling));
 			}
 		}
 	}
