Fix crash occuring during Hlod generation
Allow WorldPartitionLevelStreamingDynamic created by via ::LoadInEditor to not have an associated UWorldPartitionStreamingCell

diff --git a/Engine/Source/Runtime/Engine/Private/WorldPartition/WorldPartitionLevelStreamingDynamic.cpp b/Engine/Source/Runtime/Engine/Private/WorldPartition/WorldPartitionLevelStreamingDynamic.cpp
--- a/Engine/Source/Runtime/Engine/Private/WorldPartition/WorldPartitionLevelStreamingDynamic.cpp
+++ b/Engine/Source/Runtime/Engine/Private/WorldPartition/WorldPartitionLevelStreamingDynamic.cpp
@@ -125,7 +125,10 @@ void UWorldPartitionLevelStreamingDynamic::CreateRuntimeLevel()
 	check(RuntimeLevel == nullptr);
 	const UWorld* World = GetWorld();
 	check(World && (World->IsGameWorld() || GetShouldBeVisibleInEditor()));
-	check(StreamingCell != nullptr);
+
+	// Make sure we are creating a runtime level for a cell.
+	// Or, we created a transient WPLevelStreamingDynamic via UWorldPartitionLevelStreamingDynamic::LoadInEditor, to load a pack of actors (Hlods generation).
+	check(StreamingCell != nullptr || GetShouldBeVisibleInEditor());
 
 	// Create streaming cell Level package
 	RuntimeLevel = FWorldPartitionLevelHelper::CreateEmptyLevelForRuntimeCell(StreamingCell.Get(), World, GetWorldAsset().ToString());
