Allowing time code information parsed on import to be visible in details panel of the editor.

diff --git a/Engine/Source/Runtime/Engine/Classes/Sound/SoundWave.h b/Engine/Source/Runtime/Engine/Classes/Sound/SoundWave.h
--- a/Engine/Source/Runtime/Engine/Classes/Sound/SoundWave.h
+++ b/Engine/Source/Runtime/Engine/Classes/Sound/SoundWave.h
@@ -749,7 +749,7 @@ public:
 #if WITH_EDITORONLY_DATA
 	
 	/** Information about the time-code from import, if available.  */
-	UPROPERTY()
+	UPROPERTY(VisibleAnywhere, Category = Info)
 	FSoundWaveTimecodeInfo TimecodeInfo;
 
 #endif // WITH_EDITORONLY_DATA
