Fix data loss caused by a failure to look up component template during reinstancing

diff --git a/Engine/Source/Editor/UnrealEd/Private/Kismet2/KismetReinstanceUtilities.cpp b/Engine/Source/Editor/UnrealEd/Private/Kismet2/KismetReinstanceUtilities.cpp
--- a/Engine/Source/Editor/UnrealEd/Private/Kismet2/KismetReinstanceUtilities.cpp
+++ b/Engine/Source/Editor/UnrealEd/Private/Kismet2/KismetReinstanceUtilities.cpp
@@ -1855,6 +1855,8 @@ static void ReplaceObjectHelper(UObject*& OldObject, UClass* OldClass, UObject*&
 
 	FName OldName(OldObject->GetFName());
 
+	UEngine::FCopyPropertiesForUnrelatedObjectsParams Options;
+
 	// If the old object is in this table, we've already renamed it away in a previous iteration. Don't rename it again!
 	if (!OldToNewNameMap.Contains(OldObject))
 	{
@@ -1876,6 +1878,7 @@ static void ReplaceObjectHelper(UObject*& OldObject, UClass* OldClass, UObject*&
 		}
 		else
 		{
+			Options.SourceObjectArchetype = OldObject->GetArchetype(); // Cache the old object's archetype before renaming
 			OldObject->Rename(nullptr, OldObject->GetOuter(), REN_DoNotDirty | REN_DontCreateRedirectors | REN_ForceNoResetLoaders);
 		}
 	}
@@ -1905,7 +1908,6 @@ static void ReplaceObjectHelper(UObject*& OldObject, UClass* OldClass, UObject*&
 	InstancedPropertyUtils::FInstancedPropertyMap InstancedPropertyMap;
 	InstancedPropertyUtils::FArchiveInstancedSubObjCollector  InstancedSubObjCollector(OldObject, InstancedPropertyMap);
 	// Copy property values
-	UEngine::FCopyPropertiesForUnrelatedObjectsParams Options;
 	Options.bNotifyObjectReplacement = true;
 	Options.bDontClearReferenceIfNewerClassExists = true;
 	UEditorEngine::CopyPropertiesForUnrelatedObjects(OldObject, NewUObject, Options);
diff --git a/Engine/Source/Runtime/Engine/Classes/Engine/Engine.h b/Engine/Source/Runtime/Engine/Classes/Engine/Engine.h
--- a/Engine/Source/Runtime/Engine/Classes/Engine/Engine.h
+++ b/Engine/Source/Runtime/Engine/Classes/Engine/Engine.h
@@ -2785,6 +2785,9 @@ public:
 		bool bClearReferences;
 		bool bDontClearReferenceIfNewerClassExists;
 
+		// In cases where the SourceObject will no longer be able to look up its correct Archetype, it can be supplied
+		UObject* SourceObjectArchetype;
+
 		ENGINE_API FCopyPropertiesForUnrelatedObjectsParams();
 		ENGINE_API FCopyPropertiesForUnrelatedObjectsParams(const FCopyPropertiesForUnrelatedObjectsParams&);
 	};
diff --git a/Engine/Source/Runtime/Engine/Private/BlueprintGeneratedClass.cpp b/Engine/Source/Runtime/Engine/Private/BlueprintGeneratedClass.cpp
--- a/Engine/Source/Runtime/Engine/Private/BlueprintGeneratedClass.cpp
+++ b/Engine/Source/Runtime/Engine/Private/BlueprintGeneratedClass.cpp
@@ -985,9 +985,24 @@ UObject* UBlueprintGeneratedClass::FindArchetype(const UClass* ArchetypeClass, c
 				// after having previously serialized an instanced into another package (e.g. map). In
 				// that case, the Blueprint class might contain an SCS node with the same name as the
 				// previously-serialized object, but it might also have been switched to a different type.
-				if (SCSNode->ComponentTemplate && SCSNode->ComponentTemplate->IsA(ArchetypeClass))
+				if (SCSNode->ComponentTemplate)
 				{
-					Archetype = SCSNode->ComponentTemplate;
+					if (SCSNode->ComponentTemplate->IsA(ArchetypeClass))
+					{
+						Archetype = SCSNode->ComponentTemplate;
+					}
+
+					// If the component class is in the process of being reinstanced, archetype lookup
+					// will fail the class check because the SCS template will already have been replaced
+					// when instances try to look up their archetype, so return the already reinstanced one
+					else if (ArchetypeClass->HasAnyClassFlags(CLASS_CompiledFromBlueprint))
+					{
+						const FString ArchetypeClassName = ArchetypeClass->GetName();
+						if (ArchetypeClassName.StartsWith("REINST_") && FStringView(ArchetypeClassName).RightChop(7).StartsWith(SCSNode->ComponentTemplate->GetClass()->GetName()))
+						{
+							Archetype = SCSNode->ComponentTemplate;
+						}
+					}
 				}
 			}
 			else if (UInheritableComponentHandler* ICH = Class->GetInheritableComponentHandler())
diff --git a/Engine/Source/Runtime/Engine/Private/UnrealEngine.cpp b/Engine/Source/Runtime/Engine/Private/UnrealEngine.cpp
--- a/Engine/Source/Runtime/Engine/Private/UnrealEngine.cpp
+++ b/Engine/Source/Runtime/Engine/Private/UnrealEngine.cpp
@@ -15880,6 +15880,8 @@ public:
 		// serialization methods, which is dangerous (the data is not guaranteed to be homogeneous)
 		// in that case, we have to stick with tagged properties only
 		, FCPFUOArchive(FindNativeSuperClass(SrcObject) == FindNativeSuperClass(DstObject))
+		, SourceObject(SrcObject)
+		, SourceArchetype(Params.SourceObjectArchetype)
 		, bSkipCompilerGeneratedDefaults(Params.bSkipCompilerGeneratedDefaults)
 	{
 		ArIgnoreArchetypeRef = true;
@@ -15914,6 +15916,11 @@ public:
 		}
 	}
 
+	virtual UObject* GetArchetypeFromLoader(const UObject* Obj) override
+	{
+		return (Obj == SourceObject && SourceArchetype ? SourceArchetype : nullptr);
+	}
+
 	virtual void MarkScriptSerializationStart(const UObject* Object) override { OpenTaggedDataScope(); }
 	virtual void MarkScriptSerializationEnd(const UObject* Object) override   { CloseTaggedDataScope(); }
 
@@ -15939,6 +15946,8 @@ private:
 		return Class;
 	}
 
+	UObject* SourceObject;
+	UObject* SourceArchetype;
 	bool bSkipCompilerGeneratedDefaults;
 };
 
@@ -15989,6 +15998,7 @@ UEngine::FCopyPropertiesForUnrelatedObjectsParams::FCopyPropertiesForUnrelatedOb
 	, bNotifyObjectReplacement(false)
 	, bClearReferences(true)
 	, bDontClearReferenceIfNewerClassExists(false)
+	, SourceObjectArchetype(nullptr)
 {}
 UEngine::FCopyPropertiesForUnrelatedObjectsParams::FCopyPropertiesForUnrelatedObjectsParams(const FCopyPropertiesForUnrelatedObjectsParams&) = default;
 PRAGMA_ENABLE_DEPRECATION_WARNINGS
