Use TextExportTransient instead of NonPIEDuplicateTransient on VirtualTexture.
The motivation was always that we don't want copy/pasting a RuntimeVirtualTextureComponent to end up binding more than one component to the VirtualTexture. TextExportTransient is the better expression of this.
It fixes the VirtualTexture setting being lost when doing a Save As on the map.
#rnx

diff --git a/Engine/Source/Runtime/Engine/Classes/Components/RuntimeVirtualTextureComponent.h b/Engine/Source/Runtime/Engine/Classes/Components/RuntimeVirtualTextureComponent.h
--- a/Engine/Source/Runtime/Engine/Classes/Components/RuntimeVirtualTextureComponent.h
+++ b/Engine/Source/Runtime/Engine/Classes/Components/RuntimeVirtualTextureComponent.h
@@ -31,7 +31,7 @@ protected:
 	bool bSnapBoundsToLandscape;
 
 	/** The virtual texture object to use. */
-	UPROPERTY(EditAnywhere, BlueprintReadOnly, NonPIEDuplicateTransient, Category = VirtualTexture)
+	UPROPERTY(EditAnywhere, BlueprintReadOnly, TextExportTransient, Category = VirtualTexture)
 	TObjectPtr<URuntimeVirtualTexture> VirtualTexture = nullptr;
 
 	/** Set to true to enable scalability settings for the virtual texture. */
@@ -47,7 +47,7 @@ protected:
 	bool bHidePrimitives = false;
 
 	/** Texture object containing streamed low mips. This can reduce rendering update cost. */
-	UPROPERTY(EditAnywhere, BlueprintReadOnly, NonPIEDuplicateTransient, Category = VirtualTextureBuild)
+	UPROPERTY(EditAnywhere, BlueprintReadOnly, Category = VirtualTextureBuild)
 	TObjectPtr<UVirtualTextureBuilder> StreamingTexture = nullptr;
 
 	/** Number of streaming low mips to build for the virtual texture. */
