Added event-based transitions to anim state machines

diff --git a/Engine/Source/Runtime/Engine/Classes/Animation/AnimInstance.h b/Engine/Source/Runtime/Engine/Classes/Animation/AnimInstance.h
--- a/Engine/Source/Runtime/Engine/Classes/Animation/AnimInstance.h
+++ b/Engine/Source/Runtime/Engine/Classes/Animation/AnimInstance.h
@@ -33,6 +33,8 @@ struct FBakedAnimationStateMachine;
 class FCompilerResultsLog;
 struct FBoneContainer;
 struct FAnimNode_LinkedAnimLayer;
+enum class ETransitionRequestQueueMode : uint8;
+enum class ETransitionRequestOverwriteMode : uint8;
 
 typedef TArray<FTransform> FTransformArrayA2;
 
@@ -1163,6 +1165,26 @@ public:
 	UFUNCTION(BlueprintCallable, Category = "SyncGroup", meta=(BlueprintThreadSafe))
 	FMarkerSyncAnimPosition GetSyncGroupPosition(FName InSyncGroupName) const;
 
+	/** Attempts to queue a transition request, returns true if the request was successful */
+	UFUNCTION(BlueprintCallable, Category="Animation", meta = (BlueprintThreadSafe, Keywords = "Event,Request,Transition"))
+	bool RequestTransitionEvent(const FName EventName, const double RequestTimeout, const ETransitionRequestQueueMode QueueMode, const ETransitionRequestOverwriteMode OverwriteMode);
+
+	/** Removes all queued transition requests with the given event name */
+	UFUNCTION(BlueprintCallable, Category = "Animation", meta = (BlueprintThreadSafe, Keywords = "Event,Request,Transition"))
+	void ClearTransitionEvents(const FName EventName);
+
+	/** Removes all queued transition requests */
+	UFUNCTION(BlueprintCallable, Category = "Animation", meta = (BlueprintThreadSafe, Keywords = "Event,Request,Transition"))
+	void ClearAllTransitionEvents();
+
+	/** Returns whether or not the given event transition request has been queued */
+	UFUNCTION(BlueprintPure, Category = "Transitions", meta = (BlueprintInternalUseOnly = "true", AnimGetter = "true", GetterContext = "Transition", BlueprintThreadSafe, Keywords = "Event,Request,Transition"))
+	bool QueryTransitionEvent(int32 MachineIndex, int32 TransitionIndex, FName EventName);
+
+	/** Behaves like QueryTransitionEvent but additionally marks the event for consumption */
+	UFUNCTION(BlueprintPure, Category = "Transitions", meta = (BlueprintInternalUseOnly = "true", AnimGetter = "true", GetterContext = "Transition", BlueprintThreadSafe, Keywords = "Event,Request,Transition"))
+	bool QueryAndMarkTransitionEvent(int32 MachineIndex, int32 TransitionIndex, FName EventName);
+
 public:
 	//~ Begin UObject Interface
 	virtual void Serialize(FArchive& Ar) override;
diff --git a/Engine/Source/Runtime/Engine/Classes/Animation/AnimNode_StateMachine.h b/Engine/Source/Runtime/Engine/Classes/Animation/AnimNode_StateMachine.h
--- a/Engine/Source/Runtime/Engine/Classes/Animation/AnimNode_StateMachine.h
+++ b/Engine/Source/Runtime/Engine/Classes/Animation/AnimNode_StateMachine.h
@@ -131,6 +131,11 @@ public:
 	UPROPERTY(EditAnywhere, Category=Settings)
 	int32 MaxTransitionsPerFrame;
 
+	// The maximum number of transition requests that can be buffered at any time.
+	// The oldest transition requests are dropped to accommodate for newly created requests.
+	UPROPERTY(EditAnywhere, Category = Settings, meta = (ClampMin = "0"))
+	int32 MaxTransitionsRequests = 32;
+
 	// When the state machine becomes relevant, it is initialized into the Entry state.
 	// It then tries to take any valid transitions to possibly end up in a different state on that same frame.
 	// - if true, that new state starts full weight.
@@ -201,6 +206,14 @@ protected:
 	// Delegates that native code can hook into to handle state exits
 	TArray<FOnGraphStateChanged> OnGraphStatesExited;
 
+	// All alive transition requests that have been queued
+	TArray<FTransitionEvent> QueuedTransitionEvents;
+
+#if WITH_EDITORONLY_DATA
+	// The set of transition requests handled this update
+	TArray<FTransitionEvent> HandledTransitionEvents;
+#endif
+
 private:
 	TArray<FPoseContext*> StateCachedPoses;
 
@@ -287,6 +300,24 @@ protected:
 
 	void LogInertializationRequestError(const FAnimationUpdateContext& Context, int32 PreviousState, int32 NextState);
 
+	/** Queues a new transition request, returns true if the transition request was successfully queued */
+	bool RequestTransitionEvent(const FTransitionEvent& InTransitionEvent);
+
+	/** Removes all queued transition requests with the given event name */
+	void ClearTransitionEvents(const FName& EventName);
+
+	/** Removes all queued transition requests*/
+	void ClearAllTransitionEvents();
+
+	/** Returns whether or not the given event transition request has been queued */
+	bool QueryTransitionEvent(const int32 TransitionIndex, const FName& EventName) const;
+
+	/** Behaves like QueryTransitionEvent but additionally marks the event for consumption */
+	bool QueryAndMarkTransitionEvent(const int32 TransitionIndex, const FName& EventName);
+
+	/** Removes all marked events that are queued */
+	void ConsumeMarkedTransitionEvents();
+
 public:
 	friend struct FAnimInstanceProxy;
 };
diff --git a/Engine/Source/Runtime/Engine/Classes/Animation/AnimStateMachineTypes.h b/Engine/Source/Runtime/Engine/Classes/Animation/AnimStateMachineTypes.h
--- a/Engine/Source/Runtime/Engine/Classes/Animation/AnimStateMachineTypes.h
+++ b/Engine/Source/Runtime/Engine/Classes/Animation/AnimStateMachineTypes.h
@@ -11,6 +11,22 @@
 #include "AnimStateMachineTypes.generated.h"
 
 class UCurveFloat;
+class UAnimStateTransitionNode;
+
+UENUM(BlueprintType)
+enum class ETransitionRequestQueueMode : uint8
+{
+	Shared			UMETA(ToolTip = "Only one transition can handle this request"),
+	Unique			UMETA(ToolTip = "Allows multiple transitions to handle the same request"),
+};
+
+UENUM(BlueprintType)
+enum class ETransitionRequestOverwriteMode : uint8
+{
+	Append			UMETA(ToolTip = "This request is added whether or not another with the same name is already queued"),
+	Ignore			UMETA(ToolTip = "This request is ignored if another request with the same name is already queued"),
+	Overwrite		UMETA(ToolTip = "This request overwrites another request with the same name if one exists")
+};
 
 //@TODO: Document
 UENUM()
@@ -35,6 +51,60 @@ namespace ETransitionLogicType
 	};
 }
 
+struct FTransitionEvent
+{
+	TArray<int32, TInlineAllocator<8>> ConsumedTransitions;
+	FName EventName;
+	double CreationTime;
+	double TimeToLive;
+	ETransitionRequestQueueMode QueueMode;
+	ETransitionRequestOverwriteMode OverwriteMode;
+
+	FTransitionEvent(const FName& InEventName, const double InTimeToLive, const ETransitionRequestQueueMode& InQueueMode, const ETransitionRequestOverwriteMode& InOverwriteMode)
+		: EventName(InEventName)
+		, TimeToLive(InTimeToLive)
+		, QueueMode(InQueueMode)
+		, OverwriteMode(InOverwriteMode)
+	{
+		CreationTime = FPlatformTime::Seconds();
+	}
+
+	bool IsValidRequest() const
+	{
+		return TimeToLive > 0.0;
+	}
+
+	double GetRemainingTime() const
+	{
+		return TimeToLive - (FPlatformTime::Seconds() - CreationTime);
+	}
+
+	bool HasExpired() const
+	{
+		return GetRemainingTime() <= 0.0;
+	}
+
+	bool ToBeConsumed() const
+	{
+		if (QueueMode == ETransitionRequestQueueMode::Shared && ConsumedTransitions.Num() > 0)
+		{
+			return true;
+		}
+		return false;
+	}
+
+	bool HasBeenHandled() const
+	{
+		return ConsumedTransitions.Num() > 0;
+	}
+
+	FString ToDebugString() const
+	{
+		FString HandledByString = *FString::JoinBy(ConsumedTransitions, TEXT(", "), [](const int32& TransitionIndex) { return FString::Printf(TEXT("%d"), TransitionIndex); });
+		return FString::Printf(TEXT("%s (%.2fs) [Handled by: %s]"), *EventName.ToString(), GetRemainingTime(), *HandledByString);
+	}
+};
+
 // This structure represents a baked transition rule inside a state
 USTRUCT()
 struct FAnimationTransitionRule
diff --git a/Engine/Source/Runtime/Engine/Private/Animation/AnimInstance.cpp b/Engine/Source/Runtime/Engine/Private/Animation/AnimInstance.cpp
--- a/Engine/Source/Runtime/Engine/Private/Animation/AnimInstance.cpp
+++ b/Engine/Source/Runtime/Engine/Private/Animation/AnimInstance.cpp
@@ -32,6 +32,7 @@
 #include "UObject/UE5MainStreamObjectVersion.h"
 #include "Animation/AnimSubsystem.h"
 #include "Animation/AnimSubsystem_Tag.h"
+#include "Animation/AnimStateMachineTypes.h"
 #if WITH_EDITOR
 #include "Animation/DebugSkelMeshComponent.h"
 #endif
@@ -3784,4 +3785,29 @@ void UAnimInstance::QueueRootMotionBlend(const FTransform& RootTransform, const
 	RootMotionBlendQueue.Add(FQueuedRootMotionBlend(RootTransform, SlotName, Weight));
 }
 
+bool UAnimInstance::RequestTransitionEvent(const FName EventName, const double RequestTimeout, const ETransitionRequestQueueMode QueueMode, const ETransitionRequestOverwriteMode OverwriteMode)
+{
+	return GetProxyOnAnyThread<FAnimInstanceProxy>().RequestTransitionEvent(EventName, RequestTimeout, QueueMode, OverwriteMode);
+}
+
+void UAnimInstance::ClearTransitionEvents(const FName EventName)
+{
+	GetProxyOnAnyThread<FAnimInstanceProxy>().ClearTransitionEvents(EventName);
+}
+
+void UAnimInstance::ClearAllTransitionEvents()
+{
+	GetProxyOnAnyThread<FAnimInstanceProxy>().ClearAllTransitionEvents();
+}
+
+bool UAnimInstance::QueryTransitionEvent(int32 MachineIndex, int32 TransitionIndex, FName EventName)
+{
+	return GetProxyOnAnyThread<FAnimInstanceProxy>().QueryTransitionEvent(MachineIndex, TransitionIndex, EventName);
+}
+
+bool UAnimInstance::QueryAndMarkTransitionEvent(int32 MachineIndex, int32 TransitionIndex, FName EventName)
+{
+	return GetProxyOnAnyThread<FAnimInstanceProxy>().QueryAndMarkTransitionEvent(MachineIndex, TransitionIndex, EventName);
+}
+
 #undef LOCTEXT_NAMESPACE 
diff --git a/Engine/Source/Runtime/Engine/Private/Animation/AnimInstanceProxy.cpp b/Engine/Source/Runtime/Engine/Private/Animation/AnimInstanceProxy.cpp
--- a/Engine/Source/Runtime/Engine/Private/Animation/AnimInstanceProxy.cpp
+++ b/Engine/Source/Runtime/Engine/Private/Animation/AnimInstanceProxy.cpp
@@ -26,6 +26,7 @@
 #include "Animation/AnimNotifyStateMachineInspectionLibrary.h"
 #include "Animation/MirrorDataTable.h"
 #include "Animation/AnimBlueprintGeneratedClass.h"
+#include "Animation/AnimStateMachineTypes.h"
 
 #define DO_ANIMSTAT_PROCESSING(StatName) DEFINE_STAT(STAT_ ## StatName)
 #include "Animation/AnimMTStats.h"
@@ -2555,6 +2556,69 @@ bool FAnimInstanceProxy::WasAnimNotifyNameTriggeredInStateMachine(int32 MachineI
 	return false;
 }
 
+bool FAnimInstanceProxy::RequestTransitionEvent(const FName& EventName, const double RequestTimeout, const ETransitionRequestQueueMode& QueueMode, const ETransitionRequestOverwriteMode& OverwriteMode)
+{
+	FTransitionEvent NewTransitionEvent(EventName, RequestTimeout, QueueMode, OverwriteMode);
+
+	if (!NewTransitionEvent.IsValidRequest())
+	{
+		return false;
+	}
+	
+	ForEachStateMachine([NewTransitionEvent](FAnimNode_StateMachine& StateMachine)
+	{
+		StateMachine.RequestTransitionEvent(NewTransitionEvent);
+	});
+
+	return true;
+}
+
+bool FAnimInstanceProxy::QueryTransitionEvent(int32 MachineIndex, int32 TransitionIndex, const FName& EventName) const
+{
+	const FAnimNode_StateMachine* StateMachine = GetStateMachineInstance(MachineIndex);
+	if (StateMachine)
+	{
+		return StateMachine->QueryTransitionEvent(TransitionIndex, EventName);
+	}
+	return false;
+}
+
+void FAnimInstanceProxy::ClearTransitionEvents(const FName& EventName)
+{
+	ForEachStateMachine([EventName](FAnimNode_StateMachine& StateMachine)
+	{
+		StateMachine.ClearTransitionEvents(EventName);
+	});
+}
+
+void FAnimInstanceProxy::ClearAllTransitionEvents()
+{
+	ForEachStateMachine([](FAnimNode_StateMachine& StateMachine)
+	{
+		StateMachine.ClearAllTransitionEvents();
+	});
+}
+
+bool FAnimInstanceProxy::QueryAndMarkTransitionEvent(int32 MachineIndex, int32 TransitionIndex, const FName& EventName)
+{
+	if (AnimClassInterface)
+	{
+		const TArray<FStructProperty*>& AnimNodeProperties = AnimClassInterface->GetAnimNodeProperties();
+
+		FStructProperty* Property = AnimNodeProperties[AnimNodeProperties.Num() - 1 - MachineIndex];
+		if (Property && Property->Struct->IsChildOf(FAnimNode_StateMachine::StaticStruct()))
+		{
+			FAnimNode_StateMachine* StateMachine = Property->ContainerPtrToValuePtr<FAnimNode_StateMachine>(AnimInstanceObject);
+			if (StateMachine)
+			{
+				return StateMachine->QueryAndMarkTransitionEvent(TransitionIndex, EventName);
+			}
+		}
+	}
+
+	return false;
+}
+
 const FAnimNode_AssetPlayerBase* FAnimInstanceProxy::GetRelevantAssetPlayerFromState(int32 MachineIndex, int32 StateIndex) const
 {
 	if(const FAnimNode_StateMachine* StateMachine = GetStateMachineInstance(MachineIndex))
@@ -3365,4 +3429,24 @@ void FAnimInstanceProxy::ResetCounterInputProxy(FAnimInstanceProxy* InputProxy)
 	}
 }
 
+void FAnimInstanceProxy::ForEachStateMachine(const TFunctionRef<void(FAnimNode_StateMachine&)>& Functor)
+{
+	if (AnimClassInterface)
+	{
+		const TArray<FStructProperty*>& AnimNodeProperties = AnimClassInterface->GetAnimNodeProperties();
+		for (int32 MachineIndex = 0; MachineIndex < AnimNodeProperties.Num(); MachineIndex++)
+		{
+			FStructProperty* Property = AnimNodeProperties[AnimNodeProperties.Num() - 1 - MachineIndex];
+			if (Property && Property->Struct->IsChildOf(FAnimNode_StateMachine::StaticStruct()))
+			{
+				FAnimNode_StateMachine* StateMachine = Property->ContainerPtrToValuePtr<FAnimNode_StateMachine>(AnimInstanceObject);
+				if (StateMachine)
+				{
+					Functor(*StateMachine);
+				}
+			}
+		}
+	}
+}
+
 #undef LOCTEXT_NAMESPACE
diff --git a/Engine/Source/Runtime/Engine/Private/Animation/AnimNode_StateMachine.cpp b/Engine/Source/Runtime/Engine/Private/Animation/AnimNode_StateMachine.cpp
--- a/Engine/Source/Runtime/Engine/Private/Animation/AnimNode_StateMachine.cpp
+++ b/Engine/Source/Runtime/Engine/Private/Animation/AnimNode_StateMachine.cpp
@@ -12,6 +12,8 @@
 #include "Animation/AnimTrace.h"
 #include "Animation/ActiveStateMachineScope.h"
 #include "Animation/AnimBlueprintGeneratedClass.h"
+#include "Animation/AnimStateMachineTypes.h"
+#include "Logging/LogMacros.h"
 
 #if WITH_EDITORONLY_DATA
 #include "Animation/AnimBlueprint.h"
@@ -19,6 +21,8 @@
 
 #define LOCTEXT_NAMESPACE "AnimNode_StateMachine"
 
+DEFINE_LOG_CATEGORY_STATIC(LogAnimTransitionRequests, NoLogging, All);
+
 DECLARE_CYCLE_STAT(TEXT("StateMachine SetState"), Stat_StateMachineSetState, STATGROUP_Anim);
 
 static const FName DefaultAnimGraphName("AnimGraph");
@@ -257,6 +261,7 @@ void FAnimNode_StateMachine::Initialize_AnyThread(const FAnimationInitializeCont
 			// Reset transition related variables
 			StatesUpdated.Reset();
 			ActiveTransitionArray.Reset();
+			QueuedTransitionEvents.Reset(MaxTransitionsRequests);
 
 			StateCacheBoneCounters.Reset(Machine->States.Num());
 			StateCacheBoneCounters.AddDefaulted(Machine->States.Num());
@@ -371,6 +376,20 @@ void FAnimNode_StateMachine::Update_AnyThread(const FAnimationUpdateContext& Con
 	}
 	UpdateCounter.SynchronizeWith(Context.AnimInstanceProxy->GetUpdateCounter());
 
+	// Remove expired transition requests
+#if WITH_EDITORONLY_DATA
+	HandledTransitionEvents.Reset();
+#endif
+	for (int32 RequestIndex = QueuedTransitionEvents.Num() - 1; RequestIndex >= 0; --RequestIndex)
+	{
+		if (QueuedTransitionEvents[RequestIndex].HasExpired())
+		{
+			UE_LOG(LogAnimTransitionRequests, Verbose, TEXT("'%s' expired (Machine: %s)"), *QueuedTransitionEvents[RequestIndex].EventName.ToString(), *GetMachineDescription()->MachineName.ToString());
+			QueuedTransitionEvents.RemoveAt(RequestIndex, 1, false);
+		}
+	}
+	QueuedTransitionEvents.Shrink();
+
 	const FBakedAnimationStateMachine* Machine = GetMachineDescription();
 	if (Machine != nullptr)
 	{
@@ -519,7 +538,19 @@ void FAnimNode_StateMachine::Update_AnyThread(const FAnimationUpdateContext& Con
 #if ANIM_TRACE_ENABLED
 	TRACE_ANIM_NODE_VALUE(Context, TEXT("Name"), GetMachineDescription()->MachineName);
 	TRACE_ANIM_NODE_VALUE(Context, TEXT("Current State"), GetStateInfo().StateName);
-#endif
+	TRACE_ANIM_NODE_VALUE(
+		Context,
+		TEXT("Queued Transition Requests"), 
+		*FString::JoinBy(QueuedTransitionEvents, TEXT(",\n"), [](const FTransitionEvent& TransitionEvent) { return TransitionEvent.ToDebugString(); })
+	);
+#if WITH_EDITORONLY_DATA
+	TRACE_ANIM_NODE_VALUE(
+		Context,
+		TEXT("Consumed Transition Requests"),
+		*FString::JoinBy(HandledTransitionEvents, TEXT(",\n"), [](const FTransitionEvent& TransitionEvent) { return TransitionEvent.EventName.ToString(); })
+	);
+#endif //WITH_EDITORONLY_DATA
+#endif //ANIM_TRACE_ENABLED
 }
 
 const FAnimNode_AssetPlayerBase* FAnimNode_StateMachine::GetRelevantAssetPlayerFromState(const FAnimInstanceProxy* InAnimInstanceProxy, const FBakedAnimationState& StateInfo) const
@@ -1119,6 +1150,7 @@ void FAnimNode_StateMachine::TransitionToState(const FAnimationUpdateContext& Co
 		}
 	}
 
+	ConsumeMarkedTransitionEvents();
 	SetState(Context, NextState);
 }
 
@@ -1176,6 +1208,115 @@ bool FAnimNode_StateMachine::IsTransitionActive(int32 TransIndex) const
 	return false;
 }
 
+bool FAnimNode_StateMachine::RequestTransitionEvent(const FTransitionEvent& InTransitionEvent)
+{
+	if (!InTransitionEvent.IsValidRequest() || MaxTransitionsRequests <= 0)
+	{
+		return false;
+	}
+	
+	const int32 ExistingEventIndex = QueuedTransitionEvents.IndexOfByPredicate([InTransitionEvent](const FTransitionEvent& Transition)
+	{
+		return Transition.EventName.IsEqual(InTransitionEvent.EventName);
+	});
+
+	if (ExistingEventIndex == INDEX_NONE || InTransitionEvent.OverwriteMode == ETransitionRequestOverwriteMode::Append)
+	{
+		UE_LOG(LogAnimTransitionRequests, Verbose, TEXT("Creating new '%s' request (Machine: %s)"), *InTransitionEvent.EventName.ToString(), *GetMachineDescription()->MachineName.ToString());
+
+		ensure(QueuedTransitionEvents.Num() <= MaxTransitionsRequests);
+		if (QueuedTransitionEvents.Num() == MaxTransitionsRequests)
+		{
+			UE_LOG(LogAnimTransitionRequests, Warning, TEXT("Transition request cap reached, dropping old requests (Machine: %s)"), *InTransitionEvent.EventName.ToString(), *GetMachineDescription()->MachineName.ToString());
+			QueuedTransitionEvents.Pop(false);
+		}
+		QueuedTransitionEvents.Insert(InTransitionEvent, 0);
+
+		return true;
+	}
+	else if (InTransitionEvent.OverwriteMode == ETransitionRequestOverwriteMode::Ignore)
+	{
+		UE_LOG(LogAnimTransitionRequests, Verbose, TEXT("Ignoring new '%s' request (Machine %s, Transition %d)"), *InTransitionEvent.EventName.ToString(), *GetMachineDescription()->MachineName.ToString());
+		return false;
+	}
+	else if (InTransitionEvent.OverwriteMode == ETransitionRequestOverwriteMode::Overwrite)
+	{
+		UE_LOG(LogAnimTransitionRequests, Verbose, TEXT("Overwriting '%s' request (Machine %s, Transition %d)"), *InTransitionEvent.EventName.ToString(), *GetMachineDescription()->MachineName.ToString());
+		QueuedTransitionEvents[ExistingEventIndex] = InTransitionEvent;
+		return true;
+	}
+	else
+	{
+		ensure(false);
+	}
+
+	return false;
+}
+
+void FAnimNode_StateMachine::ClearTransitionEvents(const FName& EventName)
+{
+	for (int32 RequestIndex = QueuedTransitionEvents.Num() - 1; RequestIndex >= 0; --RequestIndex)
+	{
+		if (QueuedTransitionEvents[RequestIndex].EventName.IsEqual(EventName))
+		{
+			UE_LOG(LogAnimTransitionRequests, Verbose, TEXT("Clearing '%s' request (Machine %s)"), *EventName.ToString(), *GetMachineDescription()->MachineName.ToString());
+			QueuedTransitionEvents.RemoveAt(RequestIndex, 1, false);
+		}
+	}
+	QueuedTransitionEvents.Shrink();
+}
+
+void FAnimNode_StateMachine::ClearAllTransitionEvents()
+{
+	UE_LOG(LogAnimTransitionRequests, Verbose, TEXT("Clearing all request (Machine %s)"), *GetMachineDescription()->MachineName.ToString());
+	QueuedTransitionEvents.Reset(MaxTransitionsRequests);
+}
+
+bool FAnimNode_StateMachine::QueryTransitionEvent(const int32 TransitionIndex, const FName& EventName) const
+{
+	// Assumes QueuedTransitionEvents is sorted by request creation time, i.e. index 0 is newest request, and that
+	// ContainsByPredicate returns the first (newest) request
+	return QueuedTransitionEvents.ContainsByPredicate([EventName, TransitionIndex](const FTransitionEvent& Transition)
+	{
+		return Transition.EventName.IsEqual(EventName) && !Transition.ConsumedTransitions.Contains(TransitionIndex);
+	});
+}
+
+bool FAnimNode_StateMachine::QueryAndMarkTransitionEvent(const int32 TransitionIndex, const FName& EventName)
+{
+	// Assumes QueuedTransitionEvents is sorted by request creation time, i.e. index 0 is newest request, and that
+	// IndexOfByPredicate returns the first (newest) request
+	const int32 EventIndex = QueuedTransitionEvents.IndexOfByPredicate([EventName, TransitionIndex](const FTransitionEvent& Transition)
+	{
+		return Transition.EventName.IsEqual(EventName) && !Transition.ConsumedTransitions.Contains(TransitionIndex);
+	});
+
+	if (EventIndex != INDEX_NONE)
+	{
+		UE_LOG(LogAnimTransitionRequests, Verbose, TEXT("Marking '%s' (Machine: %s, Transition: %d)"), *EventName.ToString(), *GetMachineDescription()->MachineName.ToString(), TransitionIndex);
+		ensure(!QueuedTransitionEvents[EventIndex].ConsumedTransitions.Contains(TransitionIndex));
+		QueuedTransitionEvents[EventIndex].ConsumedTransitions.Add(TransitionIndex);
+		return true;
+	}
+	return false;
+}
+
+void FAnimNode_StateMachine::ConsumeMarkedTransitionEvents()
+{
+	for (int32 RequestIndex = QueuedTransitionEvents.Num() - 1; RequestIndex >= 0; --RequestIndex)
+	{
+		if (QueuedTransitionEvents[RequestIndex].ToBeConsumed())
+		{
+			UE_LOG(LogAnimTransitionRequests, Verbose, TEXT("Consuming '%s' (Machine: %s)"), *QueuedTransitionEvents[RequestIndex].EventName.ToString(), *GetMachineDescription()->MachineName.ToString());
+#if WITH_EDITORONLY_DATA
+			HandledTransitionEvents.Add(QueuedTransitionEvents[RequestIndex]);
+#endif
+			QueuedTransitionEvents.RemoveAt(RequestIndex, 1, false);
+		}
+	}
+	QueuedTransitionEvents.Shrink();
+}
+
 void FAnimNode_StateMachine::UpdateState(int32 StateIndex, const FAnimationUpdateContext& Context)
 {
 	if ((StateIndex != INDEX_NONE) && !StatesUpdated.Contains(StateIndex) && !IsAConduitState(StateIndex))
diff --git a/Engine/Source/Runtime/Engine/Public/Animation/AnimInstanceProxy.h b/Engine/Source/Runtime/Engine/Public/Animation/AnimInstanceProxy.h
--- a/Engine/Source/Runtime/Engine/Public/Animation/AnimInstanceProxy.h
+++ b/Engine/Source/Runtime/Engine/Public/Animation/AnimInstanceProxy.h
@@ -45,6 +45,8 @@ struct FCompactPose;
 struct FNodeDebugData;
 struct FPoseContext;
 struct FPoseSnapshot;
+enum class ETransitionRequestQueueMode : uint8;
+enum class ETransitionRequestOverwriteMode : uint8;
 
 namespace UE::Anim
 {
@@ -838,7 +840,22 @@ protected:
 	
 	/** Get whether the given state machine triggered the animation notify with the specified name last tick. */
     bool WasAnimNotifyNameTriggeredInStateMachine(int32 MachineIndex, FName NotifyName);
+
+	/** Attempts to queue a transition request, returns true if successful */
+	bool RequestTransitionEvent(const FName& EventName, const double RequestTimeout, const ETransitionRequestQueueMode& QueueMode, const ETransitionRequestOverwriteMode& OverwriteMode);
 	
+	/** Removes all queued transition requests with the given event name */
+	void ClearTransitionEvents(const FName& EventName);
+
+	/** Removes all queued transition requests */
+	void ClearAllTransitionEvents();
+
+	/** Returns whether or not the given event transition request has been queued */
+	bool QueryTransitionEvent(int32 MachineIndex, int32 TransitionIndex, const FName& EventName) const;
+
+	/** Behaves like QueryTransitionEvent but additionally marks the event for consumption */
+	bool QueryAndMarkTransitionEvent(int32 MachineIndex, int32 TransitionIndex, const FName& EventName);
+
 	// Sets up a native transition delegate between states with PrevStateName and NextStateName, in the state machine with name MachineName.
 	// Note that a transition already has to exist for this to succeed
 	void AddNativeTransitionBinding(const FName& MachineName, const FName& PrevStateName, const FName& NextStateName, const FCanTakeTransition& NativeTransitionDelegate, const FName& TransitionName = NAME_None);
@@ -912,6 +929,10 @@ protected:
 	static void ResetCounterInputProxy(FAnimInstanceProxy* InputProxy);
 
 private:
+
+	/** Executes the provided functor on each valid state machine on this anim instance */
+	void ForEachStateMachine(const TFunctionRef<void(FAnimNode_StateMachine&)>& Functor);
+
 	/** Evaluate the slot when there are blend spaces involved in any of the active anim montages. */
 	void SlotEvaluatePoseWithBlendProfiles(const FName& SlotNodeName, const FAnimationPoseData& SourceAnimationPoseData, float InSourceWeight, FAnimationPoseData& OutBlendedAnimationPoseData, float InBlendWeight);
 
