Fix typo in cvar name.

#rnx

diff --git a/Engine/Source/Runtime/Engine/Private/WorldPartition/WorldPartition.cpp b/Engine/Source/Runtime/Engine/Private/WorldPartition/WorldPartition.cpp
--- a/Engine/Source/Runtime/Engine/Private/WorldPartition/WorldPartition.cpp
+++ b/Engine/Source/Runtime/Engine/Private/WorldPartition/WorldPartition.cpp
@@ -50,7 +50,7 @@
 #include "GameDelegates.h"
 
 static int32 GLoadingRangeBugItGo = 12800;
-static FAutoConsoleVariableRef CVarGDefaultLoadingRangeHLOD0(
+static FAutoConsoleVariableRef CVarLoadingRangeBugItGo(
 	TEXT("wp.Editor.LoadingRangeBugItGo"),
 	GLoadingRangeBugItGo,
 	TEXT("Loading range for BugItGo command."),
