World Partition: added accessors for newly added types in actor descriptors and cleaned class friends.

#rnx

diff --git a/Engine/Source/Runtime/Engine/Private/WorldPartition/DataLayer/DataLayerSubsystem.cpp b/Engine/Source/Runtime/Engine/Private/WorldPartition/DataLayer/DataLayerSubsystem.cpp
--- a/Engine/Source/Runtime/Engine/Private/WorldPartition/DataLayer/DataLayerSubsystem.cpp
+++ b/Engine/Source/Runtime/Engine/Private/WorldPartition/DataLayer/DataLayerSubsystem.cpp
@@ -526,7 +526,7 @@ void UDataLayerSubsystem::OnActorDescContainerInitialized(UActorDescContainer* I
 	{
 		FWorldPartitionActorDesc* ActorDesc = *Iterator;
 		check(ActorDesc->GetContainer() == InActorDescContainer);
-		ActorDesc->DataLayerInstanceNames = FDataLayerUtils::ResolvedDataLayerInstanceNames(ActorDesc, /*WorldDataLayers*/nullptr, WorldDataLayersActorDesc);
+		ActorDesc->SetDataLayerInstanceNames(FDataLayerUtils::ResolvedDataLayerInstanceNames(ActorDesc, /*WorldDataLayers*/nullptr, WorldDataLayersActorDesc));
 	}
 }
 
diff --git a/Engine/Source/Runtime/Engine/Private/WorldPartition/DataLayer/DataLayerUtils.cpp b/Engine/Source/Runtime/Engine/Private/WorldPartition/DataLayer/DataLayerUtils.cpp
--- a/Engine/Source/Runtime/Engine/Private/WorldPartition/DataLayer/DataLayerUtils.cpp
+++ b/Engine/Source/Runtime/Engine/Private/WorldPartition/DataLayer/DataLayerUtils.cpp
@@ -49,7 +49,7 @@ TArray<FName> FDataLayerUtils::ResolvedDataLayerInstanceNames(const FWorldPartit
 		if (WorldDataLayers)
 		{
 			TArray<FName> Result;
-			for (const FName& DataLayerInstanceName : InActorDesc->DataLayers)
+			for (const FName& DataLayerInstanceName : InActorDesc->GetDataLayers())
 			{
 				if (const UDataLayerInstance* DataLayerInstance = WorldDataLayers->GetDataLayerInstance(DataLayerInstanceName))
 				{
@@ -62,7 +62,7 @@ TArray<FName> FDataLayerUtils::ResolvedDataLayerInstanceNames(const FWorldPartit
 		else if (InWorldDataLayersActorDesc)
 		{
 			TArray<FName> Result;
-			for (const FName& DataLayerInstanceName : InActorDesc->DataLayers)
+			for (const FName& DataLayerInstanceName : InActorDesc->GetDataLayers())
 			{
 				if (const FDataLayerInstanceDesc* DataLayerInstanceDesc = InWorldDataLayersActorDesc->GetDataLayerInstanceFromInstanceName(DataLayerInstanceName))
 				{
@@ -78,7 +78,7 @@ TArray<FName> FDataLayerUtils::ResolvedDataLayerInstanceNames(const FWorldPartit
 		if (WorldDataLayers)
 		{
 			TArray<FName> Result;
-			for (const FName& DataLayerAssetPath : InActorDesc->DataLayers)
+			for (const FName& DataLayerAssetPath : InActorDesc->GetDataLayers())
 			{
 				WorldDataLayers->ForEachDataLayer([DataLayerAssetPath, &Result](UDataLayerInstance* DataLayerInstance)
 				{
@@ -98,7 +98,7 @@ TArray<FName> FDataLayerUtils::ResolvedDataLayerInstanceNames(const FWorldPartit
 		else if (InWorldDataLayersActorDesc)
 		{
 			TArray<FName> Result;
-			for (const FName& DataLayerAssetPath : InActorDesc->DataLayers)
+			for (const FName& DataLayerAssetPath : InActorDesc->GetDataLayers())
 			{
 				if (const FDataLayerInstanceDesc* DataLayerInstanceDesc = InWorldDataLayersActorDesc->GetDataLayerInstanceFromAssetPath(DataLayerAssetPath))
 				{
@@ -110,7 +110,7 @@ TArray<FName> FDataLayerUtils::ResolvedDataLayerInstanceNames(const FWorldPartit
 	}
 
 	bIsSuccess = false;
-	return InActorDesc->DataLayers;
+	return InActorDesc->GetDataLayers();
 }
 
 // For performance reasons, this function assumes that InActorDesc's DataLayerInstanceNames was already resolved.
diff --git a/Engine/Source/Runtime/Engine/Public/WorldPartition/WorldPartitionActorDesc.h b/Engine/Source/Runtime/Engine/Public/WorldPartition/WorldPartitionActorDesc.h
--- a/Engine/Source/Runtime/Engine/Public/WorldPartition/WorldPartitionActorDesc.h
+++ b/Engine/Source/Runtime/Engine/Public/WorldPartition/WorldPartitionActorDesc.h
@@ -39,8 +39,6 @@ class ENGINE_API FWorldPartitionActorDesc
 	friend class UWorldPartition;
 	friend struct FWorldPartitionHandleImpl;
 	friend struct FWorldPartitionReferenceImpl;
-	friend class UDataLayerSubsystem;
-	friend class FDataLayerUtils;
 
 public:
 	virtual ~FWorldPartitionActorDesc() {}
@@ -57,7 +55,9 @@ public:
 	inline bool GetLevelBoundsRelevant() const { return bLevelBoundsRelevant; }
 	inline bool GetActorIsHLODRelevant() const { return bActorIsHLODRelevant; }
 	inline FName GetHLODLayer() const { return HLODLayer; }
+	inline const TArray<FName>& GetDataLayers() const { return DataLayers; }
 	inline const TArray<FName>& GetDataLayerInstanceNames() const { return DataLayerInstanceNames; }
+	inline void SetDataLayerInstanceNames(const TArray<FName>& InDataLayerInstanceNames) { DataLayerInstanceNames = InDataLayerInstanceNames; }
 	inline FName GetActorPackage() const { return ActorPackage; }
 	inline FName GetActorPath() const { return ActorPath; }
 	inline FName GetActorLabel() const { return ActorLabel; }
