Remove bSupportsGPUSkinCache from FGenericDataDrivenShaderPlatformInfo since all platforms now support it.
This is now replaced by r.SkinCache.Allow (default to true) in case users still want to disable it and claim the memory back

diff --git a/Engine/Config/VulkanPC/DataDrivenPlatformInfo.ini b/Engine/Config/VulkanPC/DataDrivenPlatformInfo.ini
--- a/Engine/Config/VulkanPC/DataDrivenPlatformInfo.ini
+++ b/Engine/Config/VulkanPC/DataDrivenPlatformInfo.ini
@@ -53,7 +53,6 @@ bSupportsRayTracing = true
 bSupportsRayTracingShaders = false
 bSupportsInlineRayTracing = true
 bSupportsRayTracingIndirectInstanceData = true
-bSupportsGPUSkinCache = true
 bSupportsByteBufferComputeShaders = true
 
 bSupportsGPUScene = true
diff --git a/Engine/Source/Runtime/Engine/Private/GPUSkinCache.cpp b/Engine/Source/Runtime/Engine/Private/GPUSkinCache.cpp
--- a/Engine/Source/Runtime/Engine/Private/GPUSkinCache.cpp
+++ b/Engine/Source/Runtime/Engine/Private/GPUSkinCache.cpp
@@ -66,6 +66,13 @@ static FSkinCacheExecHelper GSkelMeshExecHelper;
 
 static int32 GEnableGPUSkinCacheShaders = 0;
 
+static TAutoConsoleVariable<bool> CVarAllowGPUSkinCache(
+	TEXT("r.SkinCache.Allow"),
+	true,
+	TEXT("Whether or not to allow the GPU skin Cache system to be enabled.\n"),
+	ECVF_RenderThreadSafe | ECVF_ReadOnly
+);
+
 static FAutoConsoleVariableRef CVarEnableGPUSkinCacheShaders(
 	TEXT("r.SkinCache.CompileShaders"),
 	GEnableGPUSkinCacheShaders,
@@ -92,7 +99,7 @@ static TAutoConsoleVariable<int32> CVarEnableGPUSkinCache(
 	1,
 	TEXT("Whether or not to use the GPU compute skinning cache.\n")
 	TEXT("This will perform skinning on a compute job and not skin on the vertex shader.\n")
-	TEXT("Requires r.SkinCache.CompileShaders=1\n")
+	TEXT("Requires r.SkinCache.CompileShaders=1 and r.SkinCache.Allow=1\n")
 	TEXT(" 0: off\n")
 	TEXT(" 1: on(default)\n"),
 	ECVF_RenderThreadSafe
@@ -111,7 +118,7 @@ TAutoConsoleVariable<int32> CVarGPUSkinCacheRecomputeTangents(
 	TEXT("r.SkinCache.RecomputeTangents"),
 	2,
 	TEXT("This option enables recomputing the vertex tangents on the GPU.\n")
-	TEXT("Can be changed at runtime, requires both r.SkinCache.CompileShaders=1 and r.SkinCache.Mode=1\n")
+	TEXT("Can be changed at runtime, requires both r.SkinCache.CompileShaders=1, r.SkinCache.Mode=1 and r.SkinCache.Allow=1\n")
 	TEXT(" 0: off\n")
 	TEXT(" 1: on, forces all skinned object to Recompute Tangents\n")
 	TEXT(" 2: on, only recompute tangents on skinned objects who ticked the Recompute Tangents checkbox(default)\n"),
@@ -190,9 +197,10 @@ static int32 GGPUSkinCacheFlushCounter = 0;
 
 const float MBSize = 1048576.f; // 1024 x 1024 bytes
 
-static inline bool DoesPlatformSupportGPUSkinCache(const FStaticShaderPlatform Platform)
+static inline bool IsGPUSkinCacheAllowed(EShaderPlatform Platform)
 {
-	return FDataDrivenShaderPlatformInfo::GetSupportsGPUSkinCache(Platform);
+	static FShaderPlatformCachedIniValue<bool> PerPlatformCVar(TEXT("r.SkinCache.Allow"));
+	return PerPlatformCVar.Get(Platform);
 }
 
 ENGINE_API bool IsGPUSkinCacheAvailable(EShaderPlatform Platform)
@@ -201,7 +209,7 @@ ENGINE_API bool IsGPUSkinCacheAvailable(EShaderPlatform Platform)
 	// Store in static because it needs to be consistent and available on all threads.
 	static bool bMeshDeformersAvailable = IMeshDeformerProvider::IsAvailable();
 
-	return (bMeshDeformersAvailable || AreSkinCacheShadersEnabled(Platform) != 0) && DoesPlatformSupportGPUSkinCache(Platform);
+	return (bMeshDeformersAvailable || AreSkinCacheShadersEnabled(Platform) != 0) && IsGPUSkinCacheAllowed(Platform);
 }
 
 static inline bool IsGPUSkinCacheEnable(EShaderPlatform Platform)
@@ -274,7 +282,7 @@ ENGINE_API bool DoRecomputeSkinTangentsOnGPU_RT()
 {
 	// currently only implemented and tested on Window SM5 (needs Compute, Atomics, SRV for index buffers, UAV for VertexBuffers)
 	//#todo-gpuskin: Enable on PS4 when SRVs for IB exist
-	return DoesPlatformSupportGPUSkinCache(GMaxRHIShaderPlatform) && GEnableGPUSkinCacheShaders != 0 && (GEnableGPUSkinCache && GSkinCacheRecomputeTangents != 0);
+	return IsGPUSkinCacheAllowed(GMaxRHIShaderPlatform) && GEnableGPUSkinCacheShaders != 0 && (GEnableGPUSkinCache && GSkinCacheRecomputeTangents != 0);
 }
 
 // determine if during DispatchUpdateSkinning caching should occur
diff --git a/Engine/Source/Runtime/RHI/Private/RHI.cpp b/Engine/Source/Runtime/RHI/Private/RHI.cpp
--- a/Engine/Source/Runtime/RHI/Private/RHI.cpp
+++ b/Engine/Source/Runtime/RHI/Private/RHI.cpp
@@ -2407,7 +2407,6 @@ void FGenericDataDrivenShaderPlatformInfo::SetDefaultValues()
 	bSupportsHZBOcclusion = true;
 	bSupportsWaterIndirectDraw = true;
 	bSupportsAsyncPipelineCompilation = true;
-	bSupportsGPUSkinCache = true;
 	bSupportsManualVertexFetch = true;
 	bMobileSupportFetchBindedCustomStencilBuffer = false;
 	bSupportsVolumeTextureAtomics = true;
@@ -2452,7 +2451,6 @@ void FGenericDataDrivenShaderPlatformInfo::ParseDataDrivenShaderInfo(const FConf
 	GET_SECTION_BOOL_HELPER(bSupportsRayTracingIndirectInstanceData);
 	GET_SECTION_BOOL_HELPER(bSupportsPathTracing);
 	GET_SECTION_BOOL_HELPER(bSupportsHighEndRayTracingReflections);
-	GET_SECTION_BOOL_HELPER(bSupportsGPUSkinCache);
 	GET_SECTION_BOOL_HELPER(bSupportsByteBufferComputeShaders);
 	GET_SECTION_BOOL_HELPER(bSupportsGPUScene);
 	GET_SECTION_BOOL_HELPER(bSupportsPrimitiveShaders);
diff --git a/Engine/Source/Runtime/RHI/Public/RHIDefinitions.h b/Engine/Source/Runtime/RHI/Public/RHIDefinitions.h
--- a/Engine/Source/Runtime/RHI/Public/RHIDefinitions.h
+++ b/Engine/Source/Runtime/RHI/Public/RHIDefinitions.h
@@ -360,7 +360,6 @@ class RHI_API FGenericDataDrivenShaderPlatformInfo
 	uint32 bSupportsRayTracingIndirectInstanceData : 1; // Whether instance transforms can be copied from the GPU to the TLAS instances buffer
 	uint32 bSupportsHighEndRayTracingReflections : 1; // Whether fully-featured RT reflections can be used on the platform (with multi-bounce, translucency, etc.)
 	uint32 bSupportsPathTracing : 1; // Whether real-time path tracer is supported on this platform (avoids compiling unnecessary shaders)
-	uint32 bSupportsGPUSkinCache: 1;
 	uint32 bSupportsGPUScene : 1;
 	uint32 bSupportsByteBufferComputeShaders : 1;
 	uint32 bSupportsPrimitiveShaders : 1;
@@ -640,10 +639,11 @@ public:
 		return Infos[Platform].bSupportsRayTracing && Infos[Platform].bSupportsHighEndRayTracingReflections;
 	}
 
+	UE_DEPRECATED(5.1, "This function is no longer in use and will be removed.")
 	static FORCEINLINE_DEBUGGABLE const bool GetSupportsGPUSkinCache(const FStaticShaderPlatform Platform)
 	{
 		check(IsValid(Platform));
-		return Infos[Platform].bSupportsGPUSkinCache;
+		return true;
 	}
 
 	static FORCEINLINE_DEBUGGABLE const bool GetSupportsComputeFramework(const FStaticShaderPlatform Platform)
