Fix for typo in URectLightComponent::SetSourceTexture() and URectLightComponent::SetSourceWidth()

#rnx

diff --git a/Engine/Config/BaseEngine.ini b/Engine/Config/BaseEngine.ini
--- a/Engine/Config/BaseEngine.ini
+++ b/Engine/Config/BaseEngine.ini
@@ -1065,6 +1065,9 @@ DefaultWorldPartitionHLODLayer=/Engine/Maps/Templates/HLODs/HLODLayer_Instanced.
 
 +PropertyRedirects=(OldName="AITask_MoveTo.AIMoveTo.bUseContinuosGoalTracking", NewName="AITask_MoveTo.AIMoveTo.bUseContinuousGoalTracking")
 
++PropertyRedirects=(OldName="RectLightComponent.SetSourceTexture.bNewValue", NewName="NewValue")
++PropertyRedirects=(OldName="RectLightComponent.SetSourceWidth.bNewValue", NewName="NewValue")
+
 [CoreUObject.Metadata]
 ; Note: UnrealHeaderTool should be rerun after making changes to MetadataRedirects to catch any keys specified in class headers
 +MetadataRedirects=(OldKey="K2Protected", NewKey="BlueprintProtected")
diff --git a/Engine/Source/Runtime/Engine/Classes/Components/RectLightComponent.h b/Engine/Source/Runtime/Engine/Classes/Components/RectLightComponent.h
--- a/Engine/Source/Runtime/Engine/Classes/Components/RectLightComponent.h
+++ b/Engine/Source/Runtime/Engine/Classes/Components/RectLightComponent.h
@@ -52,10 +52,10 @@ class ENGINE_API URectLightComponent : public ULocalLightComponent
 	TObjectPtr<class UTexture> SourceTexture;
 
 	UFUNCTION(BlueprintCallable, Category = "Rendering|Lighting")
-	void SetSourceTexture(UTexture* bNewValue);
+	void SetSourceTexture(UTexture* NewValue);
 
 	UFUNCTION(BlueprintCallable, Category="Rendering|Lighting")
-	void SetSourceWidth(float bNewValue);
+	void SetSourceWidth(float NewValue);
 
 	UFUNCTION(BlueprintCallable, Category="Rendering|Lighting")
 	void SetSourceHeight(float NewValue);
