- IWYU fix for no pch build

#rnx

diff --git a/Engine/Source/Runtime/Engine/Public/FXSystem.h b/Engine/Source/Runtime/Engine/Public/FXSystem.h
--- a/Engine/Source/Runtime/Engine/Public/FXSystem.h
+++ b/Engine/Source/Runtime/Engine/Public/FXSystem.h
@@ -10,6 +10,7 @@
 #include "HAL/IConsoleManager.h"
 #include "RenderUtils.h"
 #include "RenderGraphDefinitions.h"
+#include "SceneView.h"
 
 class FCanvas;
 class FGPUSortManager;
