Texture: Optimized batching of streaming mip cache requests

The Horde Storage client does not batch chunk requests, and the hierarchy does not implement propagation of them yet. Use GetValue instead of GetChunks to improve batching and ensure local caching.

#rnx

diff --git a/Engine/Source/Runtime/Engine/Private/TextureDerivedData.cpp b/Engine/Source/Runtime/Engine/Private/TextureDerivedData.cpp
--- a/Engine/Source/Runtime/Engine/Private/TextureDerivedData.cpp
+++ b/Engine/Source/Runtime/Engine/Private/TextureDerivedData.cpp
@@ -1351,11 +1351,14 @@ typedef TArray<uint32> FAsyncVTChunkHandles;
 static bool LoadDerivedStreamingMips(FTexturePlatformData& PlatformData, int32 FirstMipToLoad, FStringView DebugContext, TFunctionRef<void (int32 MipIndex, FSharedBuffer MipData)> Callback)
 {
 	using namespace UE::DerivedData;
-	TArray<FCacheGetChunkRequest, TInlineAllocator<MAX_TEXTURE_MIP_COUNT>> Requests;
+
+	bool bMiss = false;
 
 	TIndirectArray<FTexture2DMipMap>& Mips = PlatformData.Mips;
 	if (PlatformData.DerivedDataKey.IsType<FString>())
 	{
+		TArray<FCacheGetValueRequest, TInlineAllocator<MAX_TEXTURE_MIP_COUNT>> Requests;
+
 		for (int32 MipIndex = FirstMipToLoad; MipIndex < Mips.Num(); ++MipIndex)
 		{
 			const FTexture2DMipMap& Mip = Mips[MipIndex];
@@ -1363,15 +1366,38 @@ static bool LoadDerivedStreamingMips(FTexturePlatformData& PlatformData, int32 F
 			{
 				TStringBuilder<256> MipNameBuilder;
 				MipNameBuilder.Append(DebugContext).Appendf(TEXT(" [MIP %d]"), MipIndex);
-				FCacheGetChunkRequest& Request = Requests.AddDefaulted_GetRef();
+				FCacheGetValueRequest& Request = Requests.AddDefaulted_GetRef();
 				Request.Name = MipNameBuilder;
 				Request.Key = ConvertLegacyCacheKey(PlatformData.GetDerivedDataMipKeyString(MipIndex, Mip));
 				Request.UserData = MipIndex;
 			}
 		}
+
+		if (!Requests.IsEmpty())
+		{
+			COOK_STAT(auto Timer = TextureCookStats::StreamingMipUsageStats.TimeSyncWork());
+			uint64 Size = 0;
+			FRequestOwner BlockingOwner(EPriority::Blocking);
+			GetCache().GetValue(Requests, BlockingOwner, [Callback = MoveTemp(Callback), &Size, &bMiss](FCacheGetValueResponse&& Response)
+			{
+				Size += Response.Value.GetRawSize();
+				if (Response.Status == EStatus::Ok)
+				{
+					Callback(int32(Response.UserData), Response.Value.GetData().Decompress());
+				}
+				else
+				{
+					bMiss = true;
+				}
+			});
+			BlockingOwner.Wait();
+			COOK_STAT(Timer.AddHitOrMiss(!bMiss ? FCookStats::CallStats::EHitOrMiss::Hit : FCookStats::CallStats::EHitOrMiss::Miss, int64(Size)));
+		}
 	}
 	else if (PlatformData.DerivedDataKey.IsType<FCacheKeyProxy>())
 	{
+		TArray<FCacheGetChunkRequest, TInlineAllocator<MAX_TEXTURE_MIP_COUNT>> Requests;
+
 		const FCacheKey& Key = *PlatformData.DerivedDataKey.Get<UE::DerivedData::FCacheKeyProxy>().AsCacheKey();
 		for (int32 MipIndex = FirstMipToLoad; MipIndex < Mips.Num(); ++MipIndex)
 		{
@@ -1387,35 +1413,33 @@ static bool LoadDerivedStreamingMips(FTexturePlatformData& PlatformData, int32 F
 				Request.UserData = MipIndex;
 			}
 		}
+
+		if (!Requests.IsEmpty())
+		{
+			COOK_STAT(auto Timer = TextureCookStats::StreamingMipUsageStats.TimeSyncWork());
+			uint64 Size = 0;
+			FRequestOwner BlockingOwner(EPriority::Blocking);
+			GetCache().GetChunks(Requests, BlockingOwner, [Callback = MoveTemp(Callback), &Size, &bMiss](FCacheGetChunkResponse&& Response)
+			{
+				Size += Response.RawSize;
+				if (Response.Status == EStatus::Ok)
+				{
+					Callback(int32(Response.UserData), MoveTemp(Response.RawData));
+				}
+				else
+				{
+					bMiss = true;
+				}
+			});
+			BlockingOwner.Wait();
+			COOK_STAT(Timer.AddHitOrMiss(!bMiss ? FCookStats::CallStats::EHitOrMiss::Hit : FCookStats::CallStats::EHitOrMiss::Miss, int64(Size)));
+		}
 	}
 	else
 	{
 		UE_LOG(LogTexture, Error, TEXT("Attempting to stream in mips for texture that has not generated a supported derived data key format."));
 	}
 
-	bool bMiss = false;
-
-	if (!Requests.IsEmpty())
-	{
-		COOK_STAT(auto Timer = TextureCookStats::StreamingMipUsageStats.TimeSyncWork());
-		uint64 Size = 0;
-		FRequestOwner BlockingOwner(EPriority::Blocking);
-		GetCache().GetChunks(Requests, BlockingOwner, [Callback = MoveTemp(Callback), &Size, &bMiss](FCacheGetChunkResponse&& Response)
-		{
-			Size += Response.RawSize;
-			if (Response.Status == EStatus::Ok)
-			{
-				Callback(int32(Response.UserData), MoveTemp(Response.RawData));
-			}
-			else
-			{
-				bMiss = true;
-			}
-		});
-		BlockingOwner.Wait();
-		COOK_STAT(Timer.AddHitOrMiss(!bMiss ? FCookStats::CallStats::EHitOrMiss::Hit : FCookStats::CallStats::EHitOrMiss::Miss, int64(Size)));
-	}
-
 	return !bMiss;
 }
 
