Prevent world settings for temporary transition maps from replicating in order to prevent warnings during seamless travel.

diff --git a/Engine/Source/Runtime/Engine/Private/World.cpp b/Engine/Source/Runtime/Engine/Private/World.cpp
--- a/Engine/Source/Runtime/Engine/Private/World.cpp
+++ b/Engine/Source/Runtime/Engine/Private/World.cpp
@@ -6734,14 +6734,8 @@ bool FSeamlessTravelHandler::StartTravel(UWorld* InCurrentWorld, const FURL& InU
 				if (TransitionMap.IsEmpty())
 				{
 					// If a default transition map doesn't exist, create a dummy World to use as the transition
-					if (CurrentWorld->WorldType == EWorldType::PIE)
-					{
-						SetHandlerLoadedData(NULL, UWorld::CreateWorld(EWorldType::PIE, false));
-					}
-					else
-					{
-						SetHandlerLoadedData(NULL, UWorld::CreateWorld(EWorldType::None, false));
-					}
+					EWorldType::Type TransitionWorldType = CurrentWorld->WorldType == EWorldType::PIE ? EWorldType::PIE : EWorldType::None;
+					SetHandlerLoadedData(nullptr, UWorld::CreateWorld(TransitionWorldType, false));
 				}
 				else
 				{
@@ -7294,6 +7288,18 @@ UWorld* FSeamlessTravelHandler::Tick()
 			// call initialize functions on everything that wasn't carried over from the old world
 			LoadedWorld->InitializeActorsForPlay(PendingTravelURL, false);
 
+			// If using an empty temporary transition world, make sure the world settings aren't replicated
+			FString TransitionMap = GetDefault<UGameMapsSettings>()->TransitionMap.GetLongPackageName();
+			if (TransitionMap.IsEmpty() && !bSwitchedToDefaultMap)
+			{
+				AWorldSettings* WorldSettings = LoadedWorld->GetWorldSettings();
+				if (WorldSettings)
+				{
+					WorldSettings->SetReplicates(false);
+					NetDriver->RemoveNetworkActor(WorldSettings);
+				}
+			}
+
 			// calling it after InitializeActorsForPlay has been called to have all potential bounding boxed initialized
 			FNavigationSystem::AddNavigationSystemToWorld(*LoadedWorld, FNavigationSystemRunMode::GameMode);
 
