Adding an option to block resize frame.

diff --git a/Engine/Source/Runtime/Engine/Private/Slate/SceneViewport.cpp b/Engine/Source/Runtime/Engine/Private/Slate/SceneViewport.cpp
--- a/Engine/Source/Runtime/Engine/Private/Slate/SceneViewport.cpp
+++ b/Engine/Source/Runtime/Engine/Private/Slate/SceneViewport.cpp
@@ -24,6 +24,13 @@
 
 DEFINE_LOG_CATEGORY(LogViewport);
 
+static bool bSkipViewportFrameResize = false;
+FAutoConsoleVariableRef CVarSceneSkipResizeFrame(
+	TEXT("Viewport.SkipResizeFrame"),
+	bSkipViewportFrameResize,
+	TEXT("Block the resize frame method"),
+	ECVF_Default);
+
 extern EWindowMode::Type GetWindowModeType(EWindowMode::Type WindowMode);
 
 const FName NAME_SceneViewport = FName(TEXT("SceneViewport"));
@@ -1270,6 +1277,10 @@ void FSceneViewport::PaintDebugCanvas(const FGeometry& AllottedGeometry, FSlateW
 
 void FSceneViewport::ResizeFrame(uint32 NewWindowSizeX, uint32 NewWindowSizeY, EWindowMode::Type NewWindowMode)
 {
+	if (bSkipViewportFrameResize)
+	{
+		return;
+	}
 	// Resizing the window directly is only supported in the game
 	if( FApp::IsGame() && FApp::CanEverRender() && NewWindowSizeX > 0 && NewWindowSizeY > 0 )
 	{		
