Making an au.debug version of stat streaming

diff --git a/Engine/Source/Runtime/Engine/Private/AudioDebug.cpp b/Engine/Source/Runtime/Engine/Private/AudioDebug.cpp
--- a/Engine/Source/Runtime/Engine/Private/AudioDebug.cpp
+++ b/Engine/Source/Runtime/Engine/Private/AudioDebug.cpp
@@ -230,6 +230,7 @@ namespace Audio
 	static bool bDebugMixesForAllViewsEnabled = false;
 	static bool bDebugReverbForAllViewsEnabled = false;
 	static bool bDebugModulationForAllViewsEnabled = false;
+	static bool bDebugStreamingForAllViewsEnabled = false;
 
 	const int32 DebuggerTabWidth = 12;
 
@@ -247,19 +248,23 @@ namespace Audio
 		return World->GetAudioDeviceRaw();
 	}
 
-	struct FAudioStats
+	namespace DebugStatNames
 	{
-		static const uint32 SoundWaves = 0x1;
-		static const uint32 SoundCues = 0x2;
-		static const uint32 Sounds = 0x4;
-		static const uint32 SoundMixes = 0x8;
-		static const uint32 SoundModulation = 0x10;
-		static const uint32 SoundReverb = 0x20;
+		const FName SoundWaves = "SoundWaves";
+		const FName SoundCues = "SoundCues";
+		const FName Sounds = "Sounds";
+		const FName SoundMixes = "SoundMixes";
+		const FName SoundModulation = "SoundModulation";
+		const FName SoundReverb = "SoundReverb";
+		const FName AudioStreaming = "AudioStreaming";
 
 		// TODO: Move to console variables
-		static const uint32 DebugSounds = 0x40;
-		static const uint32 LongSoundNames = 0x80;
+		const FName DebugSounds = "DebugSounds";
+		const FName LongSoundNames = "LogSoundNames";
+	}
 
+	struct FAudioStats
+	{
 		enum class EDisplaySort : uint8
 		{
 			Class,
@@ -315,7 +320,7 @@ namespace Audio
 
 		uint8 DisplayFlags;
 		EDisplaySort DisplaySort;
-		uint32 EnabledStats;
+		TSet<FName> EnabledStats;
 		TArray<FTransform> ListenerTransforms;
 		TArray<FStatSoundInfo> StatSoundInfos;
 		TArray<FStatSoundMix> StatSoundMixes;
@@ -323,7 +328,7 @@ namespace Audio
 		FAudioStats()
 			: DisplayFlags(0)
 			, DisplaySort(EDisplaySort::Name)
-			, EnabledStats(0)
+			, EnabledStats()
 		{
 		}
 	};
@@ -362,10 +367,10 @@ namespace Audio
 
 	struct FAudioStats_AudioThread
 	{
-		uint32 RequestedStats;
+		TSet<FName> RequestedStats;
 
 		FAudioStats_AudioThread()
-			: RequestedStats(0)
+			: RequestedStats()
 		{
 		}
 	};
@@ -598,7 +603,7 @@ namespace Audio
 		}
 	}
 
-	void DebugSoundObject(const TArray<FString>& Args, UWorld* InWorld, const uint32 InStatToEnable, bool& bAllEnabled)
+	void DebugSoundObject(const TArray<FString>& Args, UWorld* InWorld, const FName InStatToEnable, bool& bAllEnabled)
 	{
 		if (Args.Num() > 0)
 		{
@@ -610,13 +615,17 @@ namespace Audio
 
 			bool bEnablementRequest = Args[0].ToBool();
 			bAllEnabled = bAllViews ? bEnablementRequest : false;
+
+			TSet<FName> Stats;
+			Stats.Add(InStatToEnable);
+
 			bEnablementRequest
-				? Audio::FAudioDebugger::SetStats(InStatToEnable, bAllViews ? nullptr : InWorld)
+				? Audio::FAudioDebugger::SetStats(Stats, bAllViews ? nullptr : InWorld)
 				: Audio::FAudioDebugger::ClearStats(InStatToEnable, bAllViews ? nullptr : InWorld);
 		}
 	}
 
-	bool DebugShouldRenderStat(UWorld* World, FCanvas* Canvas, bool bEnablementBool, uint32 InAudioStat, FAudioDevice** OutAudioDevice)
+	bool DebugShouldRenderStat(UWorld* World, FCanvas* Canvas, bool bEnablementBool, const FName InAudioStat, FAudioDevice** OutAudioDevice)
 	{
 		if (!Canvas || !World)
 		{
@@ -637,7 +646,7 @@ namespace Audio
 		if (!bEnablementBool)
 		{
 			FAudioStats* AudioStats = AudioDeviceStats.Find((*OutAudioDevice)->DeviceID);
-			if (!AudioStats || !(AudioStats->EnabledStats & InAudioStat))
+			if (!AudioStats   || !(AudioStats->EnabledStats.Contains(InAudioStat)) )
 			{
 				return false;
 			}
@@ -698,7 +707,7 @@ static FAutoConsoleCommandWithWorldAndArgs GAudioDebugSoundCues
 	TEXT("(Optional) -AllViews: Enables/Disables for all viewports, not just those associated with the current world"),
 	FConsoleCommandWithWorldAndArgsDelegate::CreateStatic([](const TArray<FString>& Args, UWorld* InWorld)
 	{
-		Audio::DebugSoundObject(Args, InWorld, Audio::FAudioStats::SoundCues, Audio::bDebugCuesForAllViewsEnabled);
+		Audio::DebugSoundObject(Args, InWorld, Audio::DebugStatNames::SoundCues, Audio::bDebugCuesForAllViewsEnabled);
 	}), ECVF_Cheat
 );
 
@@ -710,7 +719,7 @@ static FAutoConsoleCommandWithWorldAndArgs GAudioDebugSounds
 	TEXT("(Optional) -AllViews: Enables/Disables for all viewports, not just those associated with the current world"),
 	FConsoleCommandWithWorldAndArgsDelegate::CreateStatic([](const TArray<FString>& Args, UWorld* InWorld)
 	{
-		Audio::DebugSoundObject(Args, InWorld, Audio::FAudioStats::Sounds, Audio::bDebugSoundsForAllViewsEnabled);
+		Audio::DebugSoundObject(Args, InWorld, Audio::DebugStatNames::Sounds, Audio::bDebugSoundsForAllViewsEnabled);
 	}), ECVF_Cheat
 );
 
@@ -722,7 +731,7 @@ static FAutoConsoleCommandWithWorldAndArgs GAudioDebugSoundWaves
 		TEXT("(Optional) -AllViews: Enables/Disables for all viewports, not just those associated with the current world"),
 	FConsoleCommandWithWorldAndArgsDelegate::CreateStatic([](const TArray<FString>& Args, UWorld* InWorld)
 	{
-		Audio::DebugSoundObject(Args, InWorld, Audio::FAudioStats::SoundWaves, Audio::bDebugWavesForAllViewsEnabled);
+		Audio::DebugSoundObject(Args, InWorld, Audio::DebugStatNames::SoundWaves, Audio::bDebugWavesForAllViewsEnabled);
 	}), ECVF_Cheat
 );
 
@@ -734,7 +743,7 @@ static FAutoConsoleCommandWithWorldAndArgs GAudioDebugSoundMixes
 		TEXT("(Optional) -AllViews: Enables/Disables for all viewports, not just those associated with the current world"),
 	FConsoleCommandWithWorldAndArgsDelegate::CreateStatic([](const TArray<FString>& Args, UWorld* InWorld)
 	{
-		Audio::DebugSoundObject(Args, InWorld, Audio::FAudioStats::SoundMixes, Audio::bDebugMixesForAllViewsEnabled);
+		Audio::DebugSoundObject(Args, InWorld, Audio::DebugStatNames::SoundMixes, Audio::bDebugMixesForAllViewsEnabled);
 	}), ECVF_Cheat
 );
 
@@ -746,7 +755,7 @@ static FAutoConsoleCommandWithWorldAndArgs GAudioDebugSoundReverb
 	TEXT("(Optional) -AllViews: Enables/Disables for all viewports, not just those associated with the current world"),
 	FConsoleCommandWithWorldAndArgsDelegate::CreateStatic([](const TArray<FString>& Args, UWorld* InWorld)
 	{
-		Audio::DebugSoundObject(Args, InWorld, Audio::FAudioStats::SoundReverb, Audio::bDebugReverbForAllViewsEnabled);
+		Audio::DebugSoundObject(Args, InWorld, Audio::DebugStatNames::SoundReverb, Audio::bDebugReverbForAllViewsEnabled);
 	}), ECVF_Cheat
 );
 
@@ -758,10 +767,21 @@ static FAutoConsoleCommandWithWorldAndArgs GAudioDebugSoundModulation
 		TEXT("(Optional) -AllViews: Enables/Disables for all viewports, not just those associated with the current world"),
 	FConsoleCommandWithWorldAndArgsDelegate::CreateStatic([](const TArray<FString>& Args, UWorld* InWorld)
 	{
-		Audio::DebugSoundObject(Args, InWorld, Audio::FAudioStats::SoundModulation, Audio::bDebugModulationForAllViewsEnabled);
+		Audio::DebugSoundObject(Args, InWorld, Audio::DebugStatNames::SoundModulation, Audio::bDebugModulationForAllViewsEnabled);
 	}), ECVF_Cheat
 );
 
+static FAutoConsoleCommandWithWorldAndArgs GAudioDebugStreaming
+(
+	TEXT("au.Debug.Streaming"),
+	TEXT("Post Stream Caching information to viewport(s).\n")
+	TEXT("0: Disable, 1: Enable\n")
+	TEXT("(Optional) -AllViews: Enables/Disables for all viewports, not just those associated with the current world"),
+	FConsoleCommandWithWorldAndArgsDelegate::CreateStatic([](const TArray<FString>& Args, UWorld* InWorld)
+		{
+			Audio::DebugSoundObject(Args, InWorld, Audio::DebugStatNames::AudioStreaming, Audio::bDebugStreamingForAllViewsEnabled);
+		}), ECVF_Cheat
+);
 
 namespace Audio
 {
@@ -771,64 +791,44 @@ namespace Audio
 	{
 		WorldRegisteredWithDeviceHandle = FAudioDeviceWorldDelegates::OnWorldRegisteredToAudioDevice.AddLambda([this](const UWorld* InWorld, FDeviceId InDeviceId)
 		{
-			uint32 StatsToSet = 0;
-			uint32 StatsToClear = 0;
+			TSet<FName> StatsToSet;
+
 			if (bDebugSoundsForAllViewsEnabled)
 			{
-				StatsToSet |= FAudioStats::Sounds;
-			}
-			else
-			{
-				StatsToClear |= FAudioStats::Sounds;
+				StatsToSet.Add(DebugStatNames::Sounds);
 			}
 
 			if (bDebugCuesForAllViewsEnabled)
 			{
-				StatsToSet |= FAudioStats::SoundCues;
-			}
-			else
-			{
-				StatsToClear |= FAudioStats::SoundCues;
+				StatsToSet.Add(DebugStatNames::SoundCues);
 			}
 
 			if (bDebugWavesForAllViewsEnabled)
 			{
-				StatsToSet |= FAudioStats::SoundWaves;
-			}
-			else
-			{
-				StatsToClear |= FAudioStats::SoundWaves;
+				StatsToSet.Add(DebugStatNames::SoundWaves);
 			}
 
 			if (bDebugModulationForAllViewsEnabled)
 			{
-				StatsToSet |= FAudioStats::SoundModulation;
-			}
-			else
-			{
-				StatsToClear |= FAudioStats::SoundModulation;
+				StatsToSet.Add(DebugStatNames::SoundModulation);
 			}
 
 			if (bDebugMixesForAllViewsEnabled)
 			{
-				StatsToSet |= FAudioStats::SoundMixes;
-			}
-			else
-			{
-				StatsToClear |= FAudioStats::SoundMixes;
+				StatsToSet.Add(DebugStatNames::SoundMixes);
 			}
 
 			if (bDebugReverbForAllViewsEnabled)
 			{
-				StatsToSet |= FAudioStats::SoundReverb;
+				StatsToSet.Add(DebugStatNames::SoundReverb);
 			}
-			else
+
+			if (bDebugStreamingForAllViewsEnabled)
 			{
-				StatsToClear |= FAudioStats::SoundReverb;
+				StatsToSet.Add(DebugStatNames::AudioStreaming);
 			}
 
 			SetStats(InDeviceId, StatsToSet);
-			ClearStats(InDeviceId, StatsToClear);
 		});
 	}
 
@@ -1328,7 +1328,8 @@ namespace Audio
 			|| bDebugSoundsForAllViewsEnabled
 			|| bDebugMixesForAllViewsEnabled
 			|| bDebugReverbForAllViewsEnabled
-			|| bDebugModulationForAllViewsEnabled;
+			|| bDebugModulationForAllViewsEnabled
+			|| bDebugStreamingForAllViewsEnabled;
 	}
 
 	int32 FAudioDebugger::DrawDebugStatsInternal(UWorld& World, FViewport& Viewport, FCanvas* InCanvas, int32 InY)
@@ -1348,6 +1349,7 @@ namespace Audio
 		Y = RenderStatSounds(&World, &Viewport, Canvas, X, Y);
 		Y = RenderStatCues(&World, &Viewport, Canvas, X, Y);
 		Y = RenderStatWaves(&World, &Viewport, Canvas, X, Y);
+		Y = RenderStatStreaming(&World, &Viewport, Canvas, X, Y, nullptr, nullptr);
 
 		return Y;
 	}
@@ -1410,64 +1412,42 @@ namespace Audio
 			return;
 		}
 
-		uint32 SetStatFlags = 0;
-		uint32 ClearStatFlags = 0;
+		TSet<FName> SetStatFlags;
 
 		if (ViewportClient->IsStatEnabled(TEXT("SoundCues")))
 		{
-			SetStatFlags |= FAudioStats::SoundCues;
-		}
-		else
-		{
-			ClearStatFlags |= FAudioStats::SoundCues;
+			SetStatFlags.Add(DebugStatNames::SoundCues);
 		}
 
 		if (ViewportClient->IsStatEnabled(TEXT("SoundWaves")))
 		{
-			SetStatFlags |= FAudioStats::SoundWaves;
-		}
-		else
-		{
-			ClearStatFlags |= FAudioStats::SoundWaves;
+			SetStatFlags.Add(DebugStatNames::SoundWaves);
 		}
 
 		if (ViewportClient->IsStatEnabled(TEXT("SoundMixes")))
 		{
-			SetStatFlags |= FAudioStats::SoundMixes;
-		}
-		else
-		{
-			ClearStatFlags |= FAudioStats::SoundMixes;
+			SetStatFlags.Add(DebugStatNames::SoundMixes);
 		}
 
 		if (ViewportClient->IsStatEnabled(TEXT("Sounds")))
 		{
 			FAudioStats& Stats = AudioDeviceStats.FindOrAdd(AudioDevice->DeviceID);
-			SetStatFlags |= FAudioStats::Sounds;
+			SetStatFlags.Add(DebugStatNames::Sounds);
 
 			if (Stats.DisplayFlags & static_cast<uint8>(FAudioStats::EDisplayFlags::Debug))
 			{
-				SetStatFlags |= FAudioStats::DebugSounds;
-			}
-			else
-			{
-				ClearStatFlags |= FAudioStats::DebugSounds;
+				SetStatFlags.Add(DebugStatNames::DebugSounds);
 			}
 
 			if (Stats.DisplayFlags & static_cast<uint8>(FAudioStats::EDisplayFlags::Long_Names))
 			{
-				SetStatFlags |= FAudioStats::LongSoundNames;
-			}
-			else
-			{
-				ClearStatFlags |= FAudioStats::LongSoundNames;
+				SetStatFlags.Add(DebugStatNames::LongSoundNames);
 			}
 		}
-		else
+
+		if (ViewportClient->IsStatEnabled(TEXT("Streaming")))
 		{
-			ClearStatFlags |= FAudioStats::Sounds;
-			ClearStatFlags |= FAudioStats::DebugSounds;
-			ClearStatFlags |= FAudioStats::LongSoundNames;
+			SetStatFlags.Add(DebugStatNames::AudioStreaming);
 		}
 
 		DECLARE_CYCLE_STAT(TEXT("FAudioThreadTask.ResolveDesiredStats"), STAT_AudioResolveDesiredStats, STATGROUP_TaskGraphTasks);
@@ -1476,16 +1456,14 @@ namespace Audio
 		if (IsInAudioThread())
 		{
 			FAudioStats_AudioThread& Stats = AudioDeviceStats_AudioThread.FindOrAdd(DeviceID);
-			Stats.RequestedStats |= SetStatFlags;
-			Stats.RequestedStats &= ~ClearStatFlags;
+			Stats.RequestedStats = SetStatFlags;
 		}
 		else
 		{
-			FAudioThread::RunCommandOnAudioThread([SetStatFlags, ClearStatFlags, DeviceID]()
+			FAudioThread::RunCommandOnAudioThread([SetStatFlags, DeviceID]()
 			{
 				FAudioStats_AudioThread& Stats = AudioDeviceStats_AudioThread.FindOrAdd(DeviceID);
-				Stats.RequestedStats |= SetStatFlags;
-				Stats.RequestedStats &= ~ClearStatFlags;
+				Stats.RequestedStats = SetStatFlags;
 			}, GET_STATID(STAT_AudioResolveDesiredStats));
 		}
 	}
@@ -1493,7 +1471,7 @@ namespace Audio
 	int32 FAudioDebugger::RenderStatCues(UWorld* World, FViewport* Viewport, FCanvas* Canvas, int32 X, int32 Y)
 	{
 		FAudioDevice* AudioDevice = nullptr;
-		if (!DebugShouldRenderStat(World, Canvas, bDebugCuesForAllViewsEnabled, FAudioStats::SoundCues, &AudioDevice))
+		if (!DebugShouldRenderStat(World, Canvas, bDebugCuesForAllViewsEnabled, DebugStatNames::SoundCues, &AudioDevice))
 		{
 			return Y;
 		}
@@ -1632,7 +1610,7 @@ namespace Audio
 	int32 FAudioDebugger::RenderStatMixes(UWorld* World, FViewport* Viewport, FCanvas* Canvas, int32 X, int32 Y)
 	{
 		FAudioDevice* AudioDevice = nullptr;
-		if (!DebugShouldRenderStat(World, Canvas, bDebugMixesForAllViewsEnabled, FAudioStats::SoundMixes, &AudioDevice))
+		if (!DebugShouldRenderStat(World, Canvas, bDebugMixesForAllViewsEnabled, DebugStatNames::SoundMixes, &AudioDevice))
 		{
 			return Y;
 		}
@@ -1672,7 +1650,7 @@ namespace Audio
 	int32 FAudioDebugger::RenderStatModulators(UWorld* World, FViewport* Viewport, FCanvas* Canvas, int32 X, int32 Y, const FVector* ViewLocation, const FRotator* ViewRotation)
 	{
 		FAudioDevice* AudioDevice = nullptr;
-		if (!DebugShouldRenderStat(World, Canvas, bDebugModulationForAllViewsEnabled, FAudioStats::SoundModulation, &AudioDevice))
+		if (!DebugShouldRenderStat(World, Canvas, bDebugModulationForAllViewsEnabled, DebugStatNames::SoundModulation, &AudioDevice))
 		{
 			return Y;
 		}
@@ -1702,7 +1680,7 @@ namespace Audio
 	int32 FAudioDebugger::RenderStatReverb(UWorld* World, FViewport* Viewport, FCanvas* Canvas, int32 X, int32 Y)
 	{
 		FAudioDevice* AudioDevice = nullptr;
-		if (!DebugShouldRenderStat(World, Canvas, bDebugReverbForAllViewsEnabled, FAudioStats::SoundReverb, &AudioDevice))
+		if (!DebugShouldRenderStat(World, Canvas, bDebugReverbForAllViewsEnabled, DebugStatNames::SoundReverb, &AudioDevice))
 		{
 			return Y;
 		}
@@ -1784,7 +1762,7 @@ namespace Audio
 	int32 FAudioDebugger::RenderStatSounds(UWorld* World, FViewport* Viewport, FCanvas* Canvas, int32 X, int32 Y)
 	{
 		FAudioDevice* AudioDevice = nullptr;
-		if (!DebugShouldRenderStat(World, Canvas, bDebugSoundsForAllViewsEnabled, FAudioStats::Sounds, &AudioDevice))
+		if (!DebugShouldRenderStat(World, Canvas, bDebugSoundsForAllViewsEnabled, DebugStatNames::Sounds, &AudioDevice))
 		{
 			return Y;
 		}
@@ -2084,7 +2062,7 @@ namespace Audio
 	int32 FAudioDebugger::RenderStatWaves(UWorld* World, FViewport* Viewport, FCanvas* Canvas, int32 X, int32 Y)
 	{
 		FAudioDevice* AudioDevice = nullptr;
-		if (!DebugShouldRenderStat(World, Canvas, bDebugWavesForAllViewsEnabled, FAudioStats::SoundWaves, &AudioDevice))
+		if (!DebugShouldRenderStat(World, Canvas, bDebugWavesForAllViewsEnabled, DebugStatNames::SoundWaves, &AudioDevice))
 		{
 			return Y;
 		}
@@ -2149,13 +2127,24 @@ namespace Audio
 		return Y;
 	}
 
+	int32 FAudioDebugger::RenderStatStreaming(UWorld* World, FViewport* Viewport, FCanvas* Canvas, int32 X, int32 Y, const FVector* ViewLocation /*= nullptr*/, const FRotator* ViewRotation /*= nullptr*/)
+	{
+		FAudioDevice* AudioDevice = nullptr;
+		if (!DebugShouldRenderStat(World, Canvas, bDebugSoundsForAllViewsEnabled, DebugStatNames::AudioStreaming, &AudioDevice))
+		{
+			return Y;
+		}
+
+		return IStreamingManager::Get().GetAudioStreamingManager().RenderStatAudioStreaming(World, Viewport, Canvas, X, Y, ViewLocation, ViewRotation);
+	}
+
 	void FAudioDebugger::RemoveDevice(const FAudioDevice& AudioDevice)
 	{
 		AudioDeviceStats.Remove(AudioDevice.DeviceID);
 		AudioDeviceStats_AudioThread.Remove(AudioDevice.DeviceID);
 	}
 
-	void FAudioDebugger::ClearStats(const uint32 StatsToClear, UWorld* InWorld)
+	void FAudioDebugger::ClearStats(const FName StatsToClear, UWorld* InWorld)
 	{
 		if (!GEngine)
 		{
@@ -2177,30 +2166,30 @@ namespace Audio
 		}
 	}
 
-	void FAudioDebugger::ClearStats(FDeviceId DeviceId, const uint32 StatsToClear)
+	void FAudioDebugger::ClearStats(FDeviceId DeviceId, const FName StatToClear)
 	{
 		if (IsInGameThread())
 		{
 			FAudioStats& Stats = AudioDeviceStats.FindOrAdd(DeviceId);
-			Stats.EnabledStats &= ~StatsToClear;
+			Stats.EnabledStats.Remove(StatToClear);
 		}
 
 		if (!IsInAudioThread())
 		{
 			DECLARE_CYCLE_STAT(TEXT("FAudioThreadTask.ClearStats"), STAT_AudioClearStats, STATGROUP_TaskGraphTasks);
 
-			FAudioThread::RunCommandOnAudioThread([this, DeviceId, StatsToClear]()
+			FAudioThread::RunCommandOnAudioThread([this, DeviceId, StatToClear]()
 			{
-				ClearStats(DeviceId, StatsToClear);
+				ClearStats(DeviceId, StatToClear);
 			}, GET_STATID(STAT_AudioClearStats));
 			return;
 		}
 
 		FAudioStats_AudioThread& Stats = AudioDeviceStats_AudioThread.FindOrAdd(DeviceId);
-		Stats.RequestedStats &= ~StatsToClear;
+		Stats.RequestedStats.Remove(StatToClear);
 	}
 
-	void FAudioDebugger::SetStats(const uint32 StatsToSet, UWorld* InWorld)
+	void FAudioDebugger::SetStats(const TSet<FName> & StatsToSet, UWorld* InWorld)
 	{
 		if (!GEngine)
 		{
@@ -2222,12 +2211,13 @@ namespace Audio
 		}
 	}
 
-	void FAudioDebugger::SetStats(FDeviceId DeviceId, const uint32 StatsToSet)
+	void FAudioDebugger::SetStats(FDeviceId DeviceId, const TSet<FName> & StatsToSet)
 	{
 		if (IsInGameThread())
 		{
 			FAudioStats& Stats = AudioDeviceStats.FindOrAdd(DeviceId);
-			Stats.EnabledStats |= StatsToSet;
+
+			Stats.EnabledStats.Append(StatsToSet);
 		}
 
 		if (!IsInAudioThread())
@@ -2242,7 +2232,8 @@ namespace Audio
 		}
 
 		FAudioStats_AudioThread& Stats = AudioDeviceStats_AudioThread.FindOrAdd(DeviceId);
-		Stats.RequestedStats |= StatsToSet;
+		
+		Stats.RequestedStats.Append(StatsToSet);
 	}
 
 	bool FAudioDebugger::PostStatModulatorHelp(UWorld* World, FCommonViewportClient* ViewportClient, const TCHAR* Stream)
@@ -2283,17 +2274,13 @@ namespace Audio
 		TArray<FAudioStats::FStatSoundInfo> StatSoundInfos;
 		TArray<FAudioStats::FStatSoundMix> StatSoundMixes;
 
-		const uint32 RequestedStats = Stats_AudioThread->RequestedStats;
+		const TSet<FName> RequestedStats = Stats_AudioThread->RequestedStats;
 		TMap<FActiveSound*, int32> ActiveSoundToInfoIndex;
 	
-		const bool bDebug = (RequestedStats & FAudioStats::DebugSounds) != 0;
+		const bool bDebug = RequestedStats.Contains(DebugStatNames::DebugSounds);
 
-		static const uint32 SoundMask =
-			FAudioStats::Sounds |
-			FAudioStats::SoundCues |
-			FAudioStats::SoundMixes |
-			FAudioStats::SoundWaves;
-		if (RequestedStats & SoundMask)
+		if (RequestedStats.Contains(DebugStatNames::Sounds) || RequestedStats.Contains(DebugStatNames::SoundCues) 
+			|| RequestedStats.Contains(DebugStatNames::SoundMixes) || RequestedStats.Contains(DebugStatNames::SoundWaves))
 		{
 			for (FActiveSound* ActiveSound : AudioDevice.GetActiveSounds())
 			{
@@ -2360,7 +2347,7 @@ namespace Audio
 
 					FAudioStats::FStatWaveInstanceInfo WaveInstanceInfo;
 					FSoundSource* Source = WaveInstanceSourceMap.FindRef(WaveInstance);
-					WaveInstanceInfo.Description = Source ? Source->Describe((RequestedStats & FAudioStats::LongSoundNames) != 0) : FString(TEXT("No source"));
+					WaveInstanceInfo.Description = Source ? Source->Describe(RequestedStats.Contains(DebugStatNames::LongSoundNames)) : FString(TEXT("No source"));
 					WaveInstanceInfo.Volume = WaveInstance->GetVolumeWithDistanceAndOcclusionAttenuation() * WaveInstance->GetDynamicVolume();
 					WaveInstanceInfo.InstanceIndex = InstanceIndex;
 					WaveInstanceInfo.WaveInstanceName = *WaveInstance->GetName();
@@ -2372,7 +2359,7 @@ namespace Audio
 			}
 		}
 
-		if (RequestedStats & FAudioStats::SoundMixes)
+		if (RequestedStats.Contains(DebugStatNames::SoundMixes))
 		{
 			if (const FAudioEffectsManager* Effects = AudioDevice.GetEffects())
 			{
diff --git a/Engine/Source/Runtime/Engine/Public/Audio/AudioDebug.h b/Engine/Source/Runtime/Engine/Public/Audio/AudioDebug.h
--- a/Engine/Source/Runtime/Engine/Public/Audio/AudioDebug.h
+++ b/Engine/Source/Runtime/Engine/Public/Audio/AudioDebug.h
@@ -19,8 +19,24 @@ class FSoundSource;
 class USoundWave;
 class UWorld;
 
+
 namespace Audio
 {
+	namespace DebugStatNames
+	{
+		ENGINE_API extern const FName SoundWaves;
+		ENGINE_API extern const FName SoundCues;
+		ENGINE_API extern const FName Sounds;
+		ENGINE_API extern const FName SoundMixes;
+		ENGINE_API extern const FName SoundModulation;
+		ENGINE_API extern const FName SoundReverb;
+		ENGINE_API extern const FName AudioStreaming;
+
+		// TODO: Move to console variables
+		ENGINE_API extern const FName DebugSounds;
+		ENGINE_API extern const FName LongSoundNames;
+	}
+
 	class ENGINE_API FAudioDebugger
 	{
 	public:
@@ -58,6 +74,7 @@ namespace Audio
 		static int32 RenderStatReverb(UWorld* World, FViewport* Viewport, FCanvas* Canvas, int32 X, int32 Y);
 		static int32 RenderStatSounds(UWorld* World, FViewport* Viewport, FCanvas* Canvas, int32 X, int32 Y);
 		static int32 RenderStatWaves(UWorld* World, FViewport* Viewport, FCanvas* Canvas, int32 X, int32 Y);
+		static int32 RenderStatStreaming(UWorld* World, FViewport* Viewport, FCanvas* Canvas, int32 X, int32 Y, const FVector* ViewLocation, const FRotator* ViewRotation);
 		static void RemoveDevice(const FAudioDevice& AudioDevice);
 		static void ResolveDesiredStats(FViewportClient* ViewportClient);
 		static void SendUpdateResultsToGameThread(const FAudioDevice& AudioDevice, const int32 FirstActiveIndex);
@@ -68,8 +85,8 @@ namespace Audio
 		static bool ToggleStatWaves(UWorld* World, FCommonViewportClient* ViewportClient, const TCHAR* Stream);
 		static void UpdateAudibleInactiveSounds(const uint32 FirstIndex, const TArray<FWaveInstance*>& WaveInstances);
 		static void LogSubtitle(const TCHAR* InCmd, USoundWave& InSoundWave);
-		static void ClearStats(const uint32 StatsToToggle, UWorld* InWorld);
-		static void SetStats(const uint32 StatsToToggle, UWorld* InWorld);
+		static void ClearStats(const FName StatsToToggle, UWorld* InWorld);
+		static void SetStats(const TSet<FName>& StatsToToggle, UWorld* InWorld);
 
 		void ClearMutesAndSolos();
 		void DumpActiveSounds() const;
@@ -125,9 +142,9 @@ namespace Audio
 			const FString& Name, const TArray<FName>& Solos, const TArray<FName>& Mutes,
 			bool& bOutIsSoloed, bool& bOutIsMuted, FString& OutReason) const;
 
-		void ClearStats(FDeviceId DeviceId, const uint32 StatsToClear);
+		void ClearStats(FDeviceId DeviceId, FName StatsToClear);
 
-		void SetStats(FDeviceId DeviceId, const uint32 StatsToSet);
+		void SetStats(FDeviceId DeviceId, const TSet<FName>& StatsToSet);
 
 		static bool ToggleStats(UWorld* World, const uint32 StatToToggle);
 		void ToggleStats(FDeviceId DeviceId, const uint32 StatsToToggle);
