FWorldPartitionLevelHelper::FPackageReferencer::RemoveReferences now postpones the call to FWorldPartitionPackageHelper::UnloadPackage and pushes it in a queue that will only be processed in PreGC.
- This fixes the case where UnloadPackage was called on a package of an instance partitioned level which skipped the removal of RF_Standalone on the AWorldDataLayers (OFPA) because it was still loading (EInternalObjectFlags::AsyncLoading).
- This new implementation is also more optimal as we wait for a GC to process packages to unload.

diff --git a/Engine/Source/Runtime/Engine/Private/WorldPartition/WorldPartitionLevelHelper.cpp b/Engine/Source/Runtime/Engine/Private/WorldPartition/WorldPartitionLevelHelper.cpp
--- a/Engine/Source/Runtime/Engine/Private/WorldPartition/WorldPartitionLevelHelper.cpp
+++ b/Engine/Source/Runtime/Engine/Private/WorldPartition/WorldPartitionLevelHelper.cpp
@@ -22,34 +22,62 @@
 #include "Templates/SharedPointer.h"
 #include "ActorFolder.h"
 
-TMap<FName, FWorldPartitionLevelHelper::FPackageReference> FWorldPartitionLevelHelper::PackageReferences;
+FWorldPartitionLevelHelper& FWorldPartitionLevelHelper::Get()
+{
+	static FWorldPartitionLevelHelper Instance;
+	return Instance;
+}
 
-void FWorldPartitionLevelHelper::FPackageReferencer::AddReference(UPackage* InPackage)
+FWorldPartitionLevelHelper::FWorldPartitionLevelHelper()
+{
+	FCoreUObjectDelegates::GetPreGarbageCollectDelegate().AddRaw(this, &FWorldPartitionLevelHelper::PreGarbageCollect);
+}
+
+void FWorldPartitionLevelHelper::PreGarbageCollect()
+{
+	for (TWeakObjectPtr<UPackage>& PackageToUnload : PreGCPackagesToUnload)
+	{
+		check(PackageToUnload.IsValid());
+		FWorldPartitionPackageHelper::UnloadPackage(PackageToUnload.Get());
+	}
+	PreGCPackagesToUnload.Reset();
+}
+
+void FWorldPartitionLevelHelper::AddReference(UPackage* InPackage, FPackageReferencer* InReferencer)
 {
 	check(InPackage);
-	FPackageReference& RefInfo = FWorldPartitionLevelHelper::PackageReferences.FindOrAdd(InPackage->GetFName());
+	FPackageReference& RefInfo = PackageReferences.FindOrAdd(InPackage->GetFName());
 	check(RefInfo.Package == nullptr || RefInfo.Package == InPackage);
 	RefInfo.Package = InPackage;
-	RefInfo.Referencers.Add(this);
+	RefInfo.Referencers.Add(InReferencer);
+	PreGCPackagesToUnload.Remove(InPackage);
 }
 
-void FWorldPartitionLevelHelper::FPackageReferencer::RemoveReferences()
+void FWorldPartitionLevelHelper::RemoveReferences(FPackageReferencer* InReferencer)
 {
-	for (auto It = FWorldPartitionLevelHelper::PackageReferences.CreateIterator(); It; ++It)
+	for (auto It = PackageReferences.CreateIterator(); It; ++It)
 	{
 		FPackageReference& RefInfo = It->Value;
-		RefInfo.Referencers.Remove(this);
+		RefInfo.Referencers.Remove(InReferencer);
 		if (RefInfo.Referencers.Num() == 0)
 		{
-			if (UPackage* Package = RefInfo.Package.Get())
-			{
-				FWorldPartitionPackageHelper::UnloadPackage(Package);
-			}
+			check(RefInfo.Package.IsValid());
+			PreGCPackagesToUnload.Add(RefInfo.Package);
 			It.RemoveCurrent();
 		}
 	}
 }
 
+void FWorldPartitionLevelHelper::FPackageReferencer::AddReference(UPackage* InPackage)
+{
+	FWorldPartitionLevelHelper::Get().AddReference(InPackage, this);
+}
+
+void FWorldPartitionLevelHelper::FPackageReferencer::RemoveReferences()
+{
+	FWorldPartitionLevelHelper::Get().RemoveReferences(this);
+}
+
 
  /**
   * Defaults World's initialization values for World Partition StreamingLevels
@@ -463,4 +491,4 @@ bool FWorldPartitionLevelHelper::LoadActors(UWorld* InOwningWorld, ULevel* InDes
 	return (LoadProgress->NumPendingLoadRequests == 0);
 }
 
-#endif
+#endif
\ No newline at end of file
diff --git a/Engine/Source/Runtime/Engine/Public/WorldPartition/WorldPartitionLevelHelper.h b/Engine/Source/Runtime/Engine/Public/WorldPartition/WorldPartitionLevelHelper.h
--- a/Engine/Source/Runtime/Engine/Public/WorldPartition/WorldPartitionLevelHelper.h
+++ b/Engine/Source/Runtime/Engine/Public/WorldPartition/WorldPartitionLevelHelper.h
@@ -21,6 +21,8 @@ struct FActorContainerID;
 class FWorldPartitionLevelHelper
 {
 public:
+	static FWorldPartitionLevelHelper& Get();
+
 	struct FPackageReferencer
 	{
 		~FPackageReferencer() { RemoveReferences(); }
@@ -38,7 +40,14 @@ public:
 	
 	static FString AddActorContainerIDToActorPath(const FActorContainerID& InContainerID, const FString& InActorPath);
 	static FString AddActorContainerIDToSubPathString(const FActorContainerID& InContainerID, const FString& InSubPathString);
+
 private:
+	FWorldPartitionLevelHelper();
+
+	void AddReference(UPackage* InPackage, FPackageReferencer* InReferencer);
+	void RemoveReferences(FPackageReferencer* InReferencer);
+	void PreGarbageCollect();
+
 	static UWorld::InitializationValues GetWorldInitializationValues();
 
 	struct FPackageReference
@@ -46,7 +55,12 @@ private:
 		TSet<FPackageReferencer*> Referencers;
 		TWeakObjectPtr<UPackage> Package;
 	};
-	static TMap<FName, FPackageReference> PackageReferences;
+
+	friend struct FPackageReferencer;
+
+	TMap<FName, FPackageReference> PackageReferences;
+
+	TSet<TWeakObjectPtr<UPackage>> PreGCPackagesToUnload;
 };
 
 #endif
\ No newline at end of file
