Fixed crash due to incorrect custom data buffer size in ISM proxy when removing instances.

diff --git a/Engine/Source/Runtime/Engine/Private/InstancedStaticMesh.cpp b/Engine/Source/Runtime/Engine/Private/InstancedStaticMesh.cpp
--- a/Engine/Source/Runtime/Engine/Private/InstancedStaticMesh.cpp
+++ b/Engine/Source/Runtime/Engine/Private/InstancedStaticMesh.cpp
@@ -1448,8 +1448,7 @@ void FInstancedStaticMeshSceneProxy::SetupProxy(UInstancedStaticMeshComponent* I
 		// Only allocate if material bound which uses this
 		if (bHasPerInstanceCustomData && InComponent->NumCustomDataFloats > 0)
 		{
-			InstanceCustomData = InComponent->PerInstanceSMCustomData;
-			check(InstanceCustomData.Num() / InComponent->NumCustomDataFloats == InComponent->GetInstanceCount()); // Sanity check on the data packing
+			InstanceCustomData.SetNumZeroed(NumRenderInstances * InComponent->NumCustomDataFloats);
 		}
 		else
 		{
