Fix deprecation warning on CitySample Linux builds

diff --git a/Engine/Source/Runtime/Engine/Public/WorldPartition/DataLayer/WorldDataLayers.h b/Engine/Source/Runtime/Engine/Public/WorldPartition/DataLayer/WorldDataLayers.h
--- a/Engine/Source/Runtime/Engine/Public/WorldPartition/DataLayer/WorldDataLayers.h
+++ b/Engine/Source/Runtime/Engine/Public/WorldPartition/DataLayer/WorldDataLayers.h
@@ -267,7 +267,9 @@ void AWorldDataLayers::OverwriteDataLayerRuntimeStates(const TArray<IdentifierTy
 			ActiveDataLayerNames.Empty(InActiveDataLayers->Num());
 			for (const IdentifierType& DataLayerIdentitier : *InActiveDataLayers)
 			{
+				PRAGMA_DISABLE_DEPRECATION_WARNINGS
 				const UDataLayerInstance* DataLayerInstance = GetDataLayerInstance(DataLayerIdentitier);
+				PRAGMA_ENABLE_DEPRECATION_WARNINGS
 				if (DataLayerInstance && DataLayerInstance->IsRuntime())
 				{
 					ActiveDataLayerNames.Add(DataLayerInstance->GetDataLayerFName());
@@ -282,7 +284,9 @@ void AWorldDataLayers::OverwriteDataLayerRuntimeStates(const TArray<IdentifierTy
 			LoadedDataLayerNames.Empty(InLoadedDataLayers->Num());
 			for (const IdentifierType& DataLayerIdentitier : *InLoadedDataLayers)
 			{
+				PRAGMA_DISABLE_DEPRECATION_WARNINGS
 				const UDataLayerInstance* DataLayerInstance = GetDataLayerInstance(DataLayerIdentitier);
+				PRAGMA_ENABLE_DEPRECATION_WARNINGS
 				if (DataLayerInstance && DataLayerInstance->IsRuntime())
 				{
 					LoadedDataLayerNames.Add(DataLayerInstance->GetDataLayerFName());
