Editgrate //UE5/Main/...[REDACTED_HANDLE]
Fix crashes when duplicating material instances (UE-127418)

diff --git a/Engine/Source/Runtime/Engine/Private/Materials/MaterialInstance.cpp b/Engine/Source/Runtime/Engine/Private/Materials/MaterialInstance.cpp
--- a/Engine/Source/Runtime/Engine/Private/Materials/MaterialInstance.cpp
+++ b/Engine/Source/Runtime/Engine/Private/Materials/MaterialInstance.cpp
@@ -2469,17 +2469,31 @@ void UMaterialInstance::Serialize(FArchive& Ar)
 
 #endif // WITH_EDITOR
 
+	bool bAllowMissingCachedData = false;
 	bool bSavedCachedData = false;
 	if (Ar.CustomVer(FUE5MainStreamObjectVersion::GUID) >= FUE5MainStreamObjectVersion::MaterialSavedCachedData)
 	{
+		// If we have editor data, up-to-date cached data can be regenerated on load
+#if WITH_EDITORONLY_DATA
 		// we want to save the cached data when cooking or duplicating the object in a cooked game
-		if (Ar.IsCooking() || (FPlatformProperties::RequiresCookedData() && Ar.IsSaving() && (Ar.GetPortFlags() & PPF_Duplicate)))
+		const bool bWantToSaveCachedData = Ar.IsCooking();
+#else
+		const bool bDuplicatingObjectInACookedGame = FPlatformProperties::RequiresCookedData() && (Ar.GetPortFlags() & PPF_Duplicate);
+		const bool bWantToSaveCachedData = Ar.IsSaving() && bDuplicatingObjectInACookedGame;
+
+		// Workaround for materials being created by annotation data at runtime not having CachedData
+		if (bDuplicatingObjectInACookedGame)
+		{
+			bAllowMissingCachedData = true;
+		}
+#endif
+		if (bWantToSaveCachedData)
 		{
 			if (CachedData)
 			{
 				bSavedCachedData = true;
 			}
-			else
+			else if (!bAllowMissingCachedData)
 			{
 				// ClassDefault object is expected to be missing cached data, but in all other cases it should have been created when the material was loaded, in PostLoad
 				checkf(HasAllFlags(RF_ClassDefaultObject), TEXT("Trying to save cooked material instance %s, missing CachedExpressionData"), *GetName());
@@ -2497,7 +2511,7 @@ void UMaterialInstance::Serialize(FArchive& Ar)
 #endif // WITH_EDITORONLY_DATA
 
 #if !WITH_EDITORONLY_DATA
-	ensureMsgf(!Ar.IsLoading() || bSavedCachedData, TEXT("MaterialInstance %s must have saved cached data, if editor-only data is not present"), *GetName());
+	ensureMsgf(!Ar.IsLoading() || bSavedCachedData || bAllowMissingCachedData, TEXT("MaterialInstance %s must have saved cached data, if editor-only data is not present"), *GetName());
 #endif
 
 	if (bSavedCachedData)
