Apply waveform transformations to cooked soundwave data, and display in the waveform editor

diff --git a/Engine/Source/Runtime/Engine/Classes/Sound/SoundWave.h b/Engine/Source/Runtime/Engine/Classes/Sound/SoundWave.h
--- a/Engine/Source/Runtime/Engine/Classes/Sound/SoundWave.h
+++ b/Engine/Source/Runtime/Engine/Classes/Sound/SoundWave.h
@@ -25,6 +25,7 @@
 #include "PerPlatformProperties.h"
 #include "ContentStreaming.h"
 #include "IAudioProxyInitializer.h"
+#include "IWaveformTransformation.h"
 #include "SoundWave.generated.h"
 
 class FSoundWaveData;
@@ -807,6 +808,11 @@ public:
 
 	UE::Serialization::FEditorBulkData RawData;
 
+#if WITH_EDITORONLY_DATA
+	UPROPERTY(EditAnywhere, Instanced, Category = "Transformations")
+	TArray<TObjectPtr<class UWaveformTransformationBase>> Transformations;
+#endif
+
 	/** GUID used to uniquely identify this node so it can be found in the DDC */
 	FGuid CompressedDataGuid;
 
@@ -857,6 +863,11 @@ private:
 
 public:
 #endif // WITH_EDITOR
+
+#if WITH_EDITORONLY_DATA
+	TArray<Audio::FTransformationPtr> CreateTransformations() const;
+#endif // WITH_EDITORONLY_DATA
+	
 	virtual void GetResourceSizeEx(FResourceSizeEx& CumulativeResourceSize) override;
 	virtual FName GetExporterName() override;
 	virtual FString GetDesc() override;
diff --git a/Engine/Source/Runtime/Engine/Private/AudioDerivedData.cpp b/Engine/Source/Runtime/Engine/Private/AudioDerivedData.cpp
--- a/Engine/Source/Runtime/Engine/Private/AudioDerivedData.cpp
+++ b/Engine/Source/Runtime/Engine/Private/AudioDerivedData.cpp
@@ -23,6 +23,10 @@
 #include "Sound/SoundSourceBus.h"
 #include "Sound/SoundWave.h"
 #include "Sound/SoundWaveProcedural.h"
+#include "AudioExtensions/Public/IWaveformTransformation.h"
+#include "DSP/FloatArrayMath.h"
+#include "DSP/BufferVectorOperations.h"
+#include "DSP/MultichannelBuffer.h"
 
 DEFINE_LOG_CATEGORY_STATIC(LogAudioDerivedData, Log, All);
 
@@ -1148,35 +1152,22 @@ public:
 /**
 * Function used for resampling a USoundWave's WaveData, which is assumed to be int16 here:
 */
-static void ResampleWaveData(TArray<uint8>& WaveData, size_t& NumBytes, int32 NumChannels, float SourceSampleRate, float DestinationSampleRate)
+static void ResampleWaveData(Audio::FAlignedFloatBuffer& WaveData, int32 NumChannels, float SourceSampleRate, float DestinationSampleRate)
 {
 	double StartTime = FPlatformTime::Seconds();
 
 	// Set up temporary output buffers:
-	Audio::FAlignedFloatBuffer ResamplerInputData;
 	Audio::FAlignedFloatBuffer ResamplerOutputData;
 
-	int32 NumSamples = NumBytes / sizeof(int16);
-
-	check(WaveData.Num() == NumBytes);
-	check(NumSamples == NumBytes / 2);
-
-	// Convert wav data from int16 to float:
-	ResamplerInputData.AddUninitialized(NumSamples);
-	int16* InputData = (int16*) WaveData.GetData();
-
-	for (int32 Index = 0; Index < NumSamples; Index++)
-	{
-		ResamplerInputData[Index] = ((float)InputData[Index]) / 32767.0f;
-	}
-
+	int32 NumSamples = WaveData.Num();
+	
 	// set up converter input params:
 	Audio::FResamplingParameters ResamplerParams = {
 		Audio::EResamplingMethod::BestSinc,
 		NumChannels,
 		SourceSampleRate,
 		DestinationSampleRate,
-		ResamplerInputData
+		WaveData
 	};
 
 	// Allocate enough space in output buffer for the resulting audio:
@@ -1189,37 +1180,9 @@ static void ResampleWaveData(TArray<uint8>& WaveData, size_t& NumBytes, int32 Nu
 	{
 		// resize WaveData buffer and convert back to int16:
 		int32 NumSamplesGenerated = ResamplerResults.OutputFramesGenerated * NumChannels;
-		WaveData.SetNum(NumSamplesGenerated * sizeof(int16));
-		InputData = (int16*) WaveData.GetData();
-
-		// Detect if the output will clip:
-		float MaxValue = 0.0f;
-		for (int32 Index = 0; Index < NumSamplesGenerated; Index++)
-		{
-			const float AbsSample = FMath::Abs(ResamplerOutputData[Index]);
-			if (AbsSample > MaxValue)
-			{
-				MaxValue = AbsSample;
-			}
-		}
-
-		// If the output will clip, normalize it.
-		if (MaxValue > 1.0f)
-		{
-			UE_LOG(LogAudioDerivedData, Display, TEXT("Audio clipped during resampling: This asset will be normalized by a factor of 1/%f. Consider attenuating the above asset."), MaxValue);
-
-			for (int32 Index = 0; Index < NumSamplesGenerated; Index++)
-			{
-				ResamplerOutputData[Index] /= MaxValue;
-			}
-		}
-
-		for (int32 Index = 0; Index < NumSamplesGenerated; Index++)
-		{
-			InputData[Index] = (int16)(ResamplerOutputData[Index] * 32767.0f);
-		}
 
-		NumBytes = NumSamplesGenerated * sizeof(int16);
+		WaveData.SetNum(NumSamplesGenerated);
+		FMemory::Memcpy(WaveData.GetData(), ResamplerOutputData.GetData(), NumSamplesGenerated * sizeof(float));
 	}
 	else
 	{
@@ -1247,6 +1210,8 @@ struct FAudioCookInputs
 	float                       SampleRateOverride = -1.0f;
 	bool                        bIsStreaming;
 	float                       CompressionQualityModifier;
+	
+	TArray<Audio::FTransformationPtr> WaveTransformations;
 
 	// Those are the only refs we keep on the actual USoundWave until
 	// we have a mechanism in place to reference a bulkdata using a 
@@ -1269,6 +1234,7 @@ struct FAudioCookInputs
 		, bIsSoundWaveProcedural(InSoundWave->IsA<USoundWaveProcedural>())
 		, CompressionQuality(InSoundWave->GetCompressionQuality())
 		, CompressionQualityModifier(InCookOverrides ? InCookOverrides->CompressionQualityModifier : 1.0f)
+		, WaveTransformations(InSoundWave->CreateTransformations())
 		, BulkDataCriticalSection(InSoundWave->RawDataCriticalSection)
 		, BulkData(InSoundWave->RawData)
 #endif
@@ -1318,26 +1284,10 @@ static void CookSimpleWave(const FAudioCookInputs& Inputs, TArray<uint8>& Output
 
 	bool bWasLocked = false;
 
-	int32 WaveSampleRate = 0;
-
-	// Don't need to cook source buses
-	if (Inputs.bIsASourceBus)
-	{
-		return;
-	}
-
-	if (Inputs.bIsSoundWaveProcedural)
-	{
-		return;
-	}
-
 	FScopeLock ScopeLock(&Inputs.BulkDataCriticalSection);
 
 	TFuture<FSharedBuffer> FutureBuffer = Inputs.BulkData.GetPayload();
 	
-	
-
-	
 	const uint8* RawWaveData = (const uint8*)FutureBuffer.Get().GetData(); // Will block 
 	int32 RawDataSize = FutureBuffer.Get().GetSize();
 
@@ -1359,55 +1309,94 @@ static void CookSimpleWave(const FAudioCookInputs& Inputs, TArray<uint8>& Output
 	if(!Input.Num())
 	{
 		UE_LOG(LogAudioDerivedData, Warning, TEXT( "Can't cook %s because there is no source LPCM data" ), *Inputs.SoundFullName );
+		return;
 	}
-	else
-	{
-		WaveSampleRate = *WaveInfo.pSamplesPerSec;
+	
+	int32 WaveSampleRate = *WaveInfo.pSamplesPerSec;
+	int32 NumChannels = *WaveInfo.pChannels;
+	int32 NumBytes = Input.Num();
+	int32 NumSamples = NumBytes / sizeof(int16);
 
-		{
-			// Check for a platform resample override here and resample if necessary:
-			if (Inputs.SampleRateOverride > 0 && Inputs.SampleRateOverride != (float)WaveSampleRate)
-			{
-				size_t TotalDataSize = WaveInfo.SampleDataSize;
+	// To float for processing
+	Audio::FAlignedFloatBuffer InputFloatBuffer;
+	InputFloatBuffer.SetNumUninitialized(NumSamples);
+	
+	Audio::ArrayPcm16ToFloat(MakeArrayView((int16*)Input.GetData(), NumSamples), InputFloatBuffer);
 
-				ResampleWaveData(Input, TotalDataSize, *WaveInfo.pChannels, WaveSampleRate, Inputs.SampleRateOverride);
-				WaveSampleRate = Inputs.SampleRateOverride;
-				WaveInfo.SampleDataSize = TotalDataSize;
-			}
-		}
+	// Run any transformations
+	if(Inputs.WaveTransformations.Num())
+	{
+		Audio::FWaveformTransformationWaveInfo TransformationInfo;
 
-		FSoundQualityInfo QualityInfo = { 0 };
-		float ModifiedCompressionQuality = (float)Inputs.CompressionQuality * Inputs.CompressionQualityModifier;
-		if (ModifiedCompressionQuality >= 1.0f)
+		TransformationInfo.Audio = &InputFloatBuffer;
+		TransformationInfo.NumChannels = NumChannels;
+		TransformationInfo.SampleRate = WaveSampleRate;
+		
+		for(const Audio::FTransformationPtr& Transformation : Inputs.WaveTransformations)
 		{
-			QualityInfo.Quality = FMath::FloorToInt(ModifiedCompressionQuality);
-			UE_CLOG(Inputs.CompressionQuality != QualityInfo.Quality, LogAudioDerivedData, Display, TEXT("Compression Quality for %s will be modified from %d to %d."), *Inputs.SoundFullName, Inputs.CompressionQuality, QualityInfo.Quality);
-		}
-		else
-		{
-			QualityInfo.Quality = Inputs.CompressionQuality;
+			Transformation->ProcessAudio(TransformationInfo);
 		}
 
-		QualityInfo.NumChannels = *WaveInfo.pChannels;
-		QualityInfo.SampleRate = WaveSampleRate;
-		QualityInfo.SampleDataSize = Input.Num();
+		WaveSampleRate = TransformationInfo.SampleRate;
+		NumChannels = TransformationInfo.NumChannels;
+		NumSamples = InputFloatBuffer.Num();
+	}
+	
+	// Check for a platform resample override here and resample if necessary:
+	if (Inputs.SampleRateOverride > 0 && Inputs.SampleRateOverride != (float)WaveSampleRate)
+	{
+		ResampleWaveData(InputFloatBuffer, NumChannels, WaveSampleRate, Inputs.SampleRateOverride);
+		
+		WaveSampleRate = Inputs.SampleRateOverride;
+		NumSamples = InputFloatBuffer.Num();
+	}
 
-		QualityInfo.bStreaming = Inputs.bIsStreaming;
-		QualityInfo.DebugName = Inputs.SoundFullName;
+	// Clip Normalize
+	if (const float MaxValue = Audio::ArrayMaxAbsValue(InputFloatBuffer) > 1.0f)
+	{
+		UE_LOG(LogAudioDerivedData, Display, TEXT("Audio clipped during cook: This asset will be normalized by a factor of 1/%f. Consider attenuating the above asset."), MaxValue);
 
-		static const FName NAME_BINKA(TEXT("BINKA"));
-		if (WaveSampleRate > 48000 &&
-			Inputs.BaseFormat == NAME_BINKA)
-		{
-			// We have to do this here because we don't know the name of the wave inside the codec.
-			UE_LOG(LogAudioDerivedData, Warning, TEXT("[%s] High sample rate wave (%d) with Bink Audio - perf waste - high frequencies are discarded by Bink Audio (like most perceptual codecs)."), *Inputs.SoundFullName, WaveSampleRate);
-		}
+		Audio::ArrayMultiplyByConstantInPlace(InputFloatBuffer, 1.f / MaxValue);
+	}
 
-		// Cook the data.
-		if (!Inputs.Compressor->Cook(Inputs.BaseFormat, Input, QualityInfo, OutputBuffer))
-		{
-			UE_LOG(LogAudioDerivedData, Warning, TEXT("Cooking sound failed: %s"), *Inputs.SoundFullName);
-		}
+	// back to PCM
+	NumBytes = NumSamples * sizeof(int16);
+	Input.SetNumUninitialized(NumBytes);
+	
+	Audio::ArrayFloatToPcm16(InputFloatBuffer, MakeArrayView((int16*)Input.GetData(), NumSamples));
+
+	// Compression Quality
+	FSoundQualityInfo QualityInfo = { 0 };
+	float ModifiedCompressionQuality = (float)Inputs.CompressionQuality * Inputs.CompressionQualityModifier;
+	if (ModifiedCompressionQuality >= 1.0f)
+	{
+		QualityInfo.Quality = FMath::FloorToInt(ModifiedCompressionQuality);
+		UE_CLOG(Inputs.CompressionQuality != QualityInfo.Quality, LogAudioDerivedData, Display, TEXT("Compression Quality for %s will be modified from %d to %d."), *Inputs.SoundFullName, Inputs.CompressionQuality, QualityInfo.Quality);
+	}
+	else
+	{
+		QualityInfo.Quality = Inputs.CompressionQuality;
+	}
+
+	QualityInfo.NumChannels = NumChannels;
+	QualityInfo.SampleRate = WaveSampleRate;
+	QualityInfo.SampleDataSize = NumBytes;
+
+	QualityInfo.bStreaming = Inputs.bIsStreaming;
+	QualityInfo.DebugName = Inputs.SoundFullName;
+
+	static const FName NAME_BINKA(TEXT("BINKA"));
+	if (WaveSampleRate > 48000 &&
+		Inputs.BaseFormat == NAME_BINKA)
+	{
+		// We have to do this here because we don't know the name of the wave inside the codec.
+		UE_LOG(LogAudioDerivedData, Warning, TEXT("[%s] High sample rate wave (%d) with Bink Audio - perf waste - high frequencies are discarded by Bink Audio (like most perceptual codecs)."), *Inputs.SoundFullName, WaveSampleRate);
+	}
+
+	// Cook the data.
+	if (!Inputs.Compressor->Cook(Inputs.BaseFormat, Input, QualityInfo, OutputBuffer))
+	{
+		UE_LOG(LogAudioDerivedData, Warning, TEXT("Cooking sound failed: %s"), *Inputs.SoundFullName);
 	}
 }
 
@@ -1426,11 +1415,6 @@ static void CookSurroundWave(const FAudioCookInputs& Inputs,  TArray<uint8>& Out
 	TArray<TArray<uint8> >	SourceBuffers;
 	TArray<int32>			RequiredChannels;
 
-	if (Inputs.bIsSoundWaveProcedural)
-	{
-		return;
-	}
-
 	FScopeLock ScopeLock(&Inputs.BulkDataCriticalSection);
 	// Lock raw wave data.
 	TFuture<FSharedBuffer> FutureBuffer = Inputs.BulkData.GetPayload();
@@ -1490,112 +1474,174 @@ static void CookSurroundWave(const FAudioCookInputs& Inputs,  TArray<uint8>& Out
 		}
 	}
 
-	if (SampleDataSize != 0)
+	if (SampleDataSize == 0)
 	{
-		int32 ChannelCount = 0;
-		// Extract all the info for channels or insert blank data
-		for( i = 0; i < SPEAKER_Count; i++ )
+		UE_LOG(LogAudioDerivedData, Warning, TEXT( "Cooking surround sound failed: %s" ), *Inputs.SoundFullName );
+		return;
+	}
+
+	TArray<FWaveModInfo, TInlineAllocator<SPEAKER_Count>> ChannelInfos;
+	
+	int32 ChannelCount = 0;
+	// Extract all the info for channels
+	for( i = 0; i < SPEAKER_Count; i++ )
+	{
+		FWaveModInfo WaveInfoInner;
+		if( WaveInfoInner.ReadWaveHeader( RawWaveData, Inputs.ChannelSizes[ i ], Inputs.ChannelOffsets[ i ] )
+			&& *WaveInfoInner.pChannels == 1 )
 		{
-			FWaveModInfo WaveInfoInner;
-			if( WaveInfoInner.ReadWaveHeader( RawWaveData, Inputs.ChannelSizes[ i ], Inputs.ChannelOffsets[ i ] )
-				&& *WaveInfoInner.pChannels == 1 )
-			{
-				ChannelCount++;
-				TArray<uint8>& Input = *new (SourceBuffers) TArray<uint8>;
-				Input.AddUninitialized(WaveInfoInner.SampleDataSize);
-				FMemory::Memcpy(Input.GetData(), WaveInfoInner.SampleDataStart, WaveInfoInner.SampleDataSize);
-				SampleDataSize = FMath::Min<uint32>(WaveInfoInner.SampleDataSize, SampleDataSize);
-			}
-			else if (RequiredChannels.Contains(i))
-			{
-				// Add an empty channel for cooking
-				ChannelCount++;
-				TArray<uint8>& Input = *new (SourceBuffers) TArray<uint8>;
-				Input.AddZeroed(SampleDataSize);
-			}
-		}
+			ChannelCount++;
+			ChannelInfos.Add(WaveInfoInner);
 
-		// Only allow the formats that can be played back through
-		if( ChannelCount == 4 || ChannelCount == 6 || ChannelCount == 7 || ChannelCount == 8 )
+			SampleDataSize = FMath::Max<uint32>(WaveInfoInner.SampleDataSize, SampleDataSize);
+		}
+		else if (RequiredChannels.Contains(i))
 		{
-			int32 WaveSampleRate = *WaveInfo.pSamplesPerSec;
+			// Add an empty channel for cooking
+			ChannelCount++;
+			WaveInfoInner.SampleDataSize = 0;
 
-			// Check for a platform resample override here and resample if necessary:
-			if (Inputs.SampleRateOverride > 0 && Inputs.SampleRateOverride != (float)WaveSampleRate)
-			{
-				for (int32 ChannelIndex = 0; ChannelIndex < ChannelCount; ChannelIndex++)
-				{
-					size_t DataSize = SourceBuffers[ChannelIndex].Num();
-					ResampleWaveData(SourceBuffers[ChannelIndex], DataSize, 1, WaveSampleRate, Inputs.SampleRateOverride);
-				}
+			ChannelInfos.Add(WaveInfoInner);
+		}
+	}
 
-				// Since each channel is resampled independently, we may have slightly different sample counts in each channel.
-				// To counter this, we truncate or zero-pad every non-zero channel's buffer to the zeroth channel's length.
-				int32 SizeOfZerothChannel = SourceBuffers[0].Num();
+	// Only allow the formats that can be played back through
+	const bool bChannelCountValidForPlayback = ChannelCount == 4 || ChannelCount == 6 || ChannelCount == 7 || ChannelCount == 8;
+	
+	if( bChannelCountValidForPlayback == false )
+	{
+		UE_LOG(LogAudioDerivedData, Warning, TEXT( "No format available for a %d channel surround sound: %s" ), ChannelCount, *Inputs.SoundFullName );
+		return;
+	}
 
-				for (int32 ChannelIndex = 1; ChannelIndex < ChannelCount; ChannelIndex++)
-				{
-					TArray<uint8>& SourceBuffer = SourceBuffers[ChannelIndex];
+	// copy channels we need, ensuring all channels are the same size
+	for(const FWaveModInfo& ChannelInfo : ChannelInfos)
+	{
+		TArray<uint8>& Input = *new (SourceBuffers) TArray<uint8>;
+		Input.AddZeroed(SampleDataSize);
 
-					if (SourceBuffer.Num() != SizeOfZerothChannel)
-					{	
-						UE_LOG(LogAudioDerivedData, Display, TEXT("Fixing up channel %d from %d to %d samples."), ChannelIndex, SizeOfZerothChannel, SourceBuffer.Num());
-						SourceBuffer.SetNumZeroed(SizeOfZerothChannel);
-					}
-				}
+		if(ChannelInfo.SampleDataSize > 0)
+		{
+			FMemory::Memcpy(Input.GetData(), ChannelInfo.SampleDataStart, ChannelInfo.SampleDataSize);
+		}
+	}
+	
+	int32 WaveSampleRate = *WaveInfo.pSamplesPerSec;
+	int32 NumFrames = SampleDataSize / sizeof(int16);
 
-				WaveSampleRate = Inputs.SampleRateOverride;
-				SampleDataSize = SizeOfZerothChannel;
-			}
+	// bNeedsResample could change if a transformation changes the sample rate
+	bool bNeedsResample = Inputs.SampleRateOverride > 0 && Inputs.SampleRateOverride != (float)WaveSampleRate;
+	
+	const bool bContainsTransformations = Inputs.WaveTransformations.Num() > 0;
+	const bool bNeedsDeinterleave = bNeedsResample || bContainsTransformations;
 
-			UE_LOG(LogAudioDerivedData, Display, TEXT("Cooking %d channels for: %s"), ChannelCount, *Inputs.SoundFullName);
+	if(bNeedsDeinterleave)
+	{
+		// multichannel wav's are stored deinterleaved, but our dsp assumes interleaved
+		Audio::FAlignedFloatBuffer InterleavedFloatBuffer;
 
-			FSoundQualityInfo QualityInfo = { 0 };
+		Audio::FMultichannelBuffer InputMultichannelBuffer;
 
-			float ModifiedCompressionQuality = (float)Inputs.CompressionQuality;
+		Audio::SetMultichannelBufferSize(ChannelCount, NumFrames, InputMultichannelBuffer);
+		
+		// convert to float
+		for (int32 ChannelIndex = 0; ChannelIndex < ChannelCount; ChannelIndex++)
+		{
+			Audio::ArrayPcm16ToFloat(MakeArrayView((const int16*)(SourceBuffers[ChannelIndex].GetData()), NumFrames), InputMultichannelBuffer[ChannelIndex]);
+		}
+	
+		Audio::ArrayInterleave(InputMultichannelBuffer, InterleavedFloatBuffer);
 
-			if (!FMath::IsNearlyEqual(Inputs.CompressionQualityModifier, 1.0f))
+		// run transformations
+		if(bContainsTransformations)
+		{
+			Audio::FWaveformTransformationWaveInfo TransformationInfo;
+
+			TransformationInfo.Audio = &InterleavedFloatBuffer;
+			TransformationInfo.NumChannels = ChannelCount;
+			TransformationInfo.SampleRate = WaveSampleRate;
+		
+			for(const Audio::FTransformationPtr& Transformation : Inputs.WaveTransformations)
 			{
-				ModifiedCompressionQuality = (float)Inputs.CompressionQuality * Inputs.CompressionQualityModifier;
+				Transformation->ProcessAudio(TransformationInfo);
 			}
+
+			WaveSampleRate = TransformationInfo.SampleRate;
+			bNeedsResample = Inputs.SampleRateOverride > 0 && Inputs.SampleRateOverride != (float)WaveSampleRate;
 			
-			if (ModifiedCompressionQuality >= 1.0f)
-			{
-				QualityInfo.Quality = FMath::FloorToInt(ModifiedCompressionQuality);
-			}
-			else
-			{
-				QualityInfo.Quality = Inputs.CompressionQuality;
-			}
+			ChannelCount = TransformationInfo.NumChannels;
+			NumFrames = InterleavedFloatBuffer.Num() / ChannelCount;
+		}
 
-			QualityInfo.NumChannels = ChannelCount;
-			QualityInfo.SampleRate = WaveSampleRate;
-			QualityInfo.SampleDataSize = SampleDataSize;
-			QualityInfo.bStreaming = Inputs.bIsStreaming;
-			QualityInfo.DebugName = Inputs.SoundFullName;
+		if (bNeedsResample)
+		{
+			ResampleWaveData(InterleavedFloatBuffer, ChannelCount, WaveSampleRate, Inputs.SampleRateOverride);
+			
+			WaveSampleRate = Inputs.SampleRateOverride;
+			NumFrames = InterleavedFloatBuffer.Num() / ChannelCount;
+		}
 
-			static const FName NAME_BINKA(TEXT("BINKA"));
-			if (WaveSampleRate > 48000 &&
-				Inputs.BaseFormat == NAME_BINKA)
-			{
-				// We have to do this here because we don't know the name of the wave inside the codec.
-				UE_LOG(LogAudioDerivedData, Warning, TEXT("[%s] High sample rate wave (%d) with Bink Audio - perf waste - high frequencies are discarded by Bink Audio (like most perceptual codecs)."), *Inputs.SoundFullName, WaveSampleRate);
-			}
+		// clip normalize
+		if (const float MaxValue = Audio::ArrayMaxAbsValue(InterleavedFloatBuffer) > 1.0f)
+		{
+			UE_LOG(LogAudioDerivedData, Display, TEXT("Audio clipped during cook: This asset will be normalized by a factor of 1/%f. Consider attenuating the above asset."), MaxValue);
 
-			//@todo tighten up the checking for empty results here
-			if (!Inputs.Compressor->CookSurround(Inputs.BaseFormat, SourceBuffers, QualityInfo, OutputBuffer))
-			{
-				UE_LOG(LogAudioDerivedData, Warning, TEXT("Cooking surround sound failed: %s"), *Inputs.SoundFullName);
-			}
+			Audio::ArrayMultiplyByConstantInPlace(InterleavedFloatBuffer, 1.f / MaxValue);
 		}
-		else
+
+		Audio::ArrayDeinterleave(InterleavedFloatBuffer, InputMultichannelBuffer, ChannelCount);
+
+		SampleDataSize = NumFrames * sizeof(int16);
+
+		// back to PCM
+		for (int32 ChannelIndex = 0; ChannelIndex < ChannelCount; ChannelIndex++)
 		{
-			UE_LOG(LogAudioDerivedData, Warning, TEXT( "No format available for a %d channel surround sound: %s" ), ChannelCount, *Inputs.SoundFullName );
+			TArray<uint8>& PcmBuffer = SourceBuffers[ChannelIndex];
+				
+			PcmBuffer.SetNum(SampleDataSize);
+			
+			Audio::ArrayFloatToPcm16(InputMultichannelBuffer[ChannelIndex], MakeArrayView((int16*)PcmBuffer.GetData(), NumFrames));
 		}
 	}
+
+	UE_LOG(LogAudioDerivedData, Display, TEXT("Cooking %d channels for: %s"), ChannelCount, *Inputs.SoundFullName);
+
+	FSoundQualityInfo QualityInfo = { 0 };
+
+	float ModifiedCompressionQuality = (float)Inputs.CompressionQuality;
+
+	if (!FMath::IsNearlyEqual(Inputs.CompressionQualityModifier, 1.0f))
+	{
+		ModifiedCompressionQuality = (float)Inputs.CompressionQuality * Inputs.CompressionQualityModifier;
+	}
+	
+	if (ModifiedCompressionQuality >= 1.0f)
+	{
+		QualityInfo.Quality = FMath::FloorToInt(ModifiedCompressionQuality);
+	}
 	else
 	{
-		UE_LOG(LogAudioDerivedData, Warning, TEXT( "Cooking surround sound failed: %s" ), *Inputs.SoundFullName );
+		QualityInfo.Quality = Inputs.CompressionQuality;
+	}
+
+	QualityInfo.NumChannels = ChannelCount;
+	QualityInfo.SampleRate = WaveSampleRate;
+	QualityInfo.SampleDataSize = SampleDataSize;
+	QualityInfo.bStreaming = Inputs.bIsStreaming;
+	QualityInfo.DebugName = Inputs.SoundFullName;
+
+	static const FName NAME_BINKA(TEXT("BINKA"));
+	if (WaveSampleRate > 48000 &&
+		Inputs.BaseFormat == NAME_BINKA)
+	{
+		// We have to do this here because we don't know the name of the wave inside the codec.
+		UE_LOG(LogAudioDerivedData, Warning, TEXT("[%s] High sample rate wave (%d) with Bink Audio - perf waste - high frequencies are discarded by Bink Audio (like most perceptual codecs)."), *Inputs.SoundFullName, WaveSampleRate);
+	}
+
+	//@todo tighten up the checking for empty results here
+	if (!Inputs.Compressor->CookSurround(Inputs.BaseFormat, SourceBuffers, QualityInfo, OutputBuffer))
+	{
+		UE_LOG(LogAudioDerivedData, Warning, TEXT("Cooking surround sound failed: %s"), *Inputs.SoundFullName);
 	}
 }
 
@@ -1641,6 +1687,12 @@ bool FDerivedAudioDataCompressor::Build(TArray<uint8>& OutData)
 	Args.Add(TEXT("SoundNodeName"), FText::FromString(CookInputs->SoundName));
 	FAudioStatusMessageContext StatusMessage(FText::Format(NSLOCTEXT("Engine", "BuildingCompressedAudioTaskStatus", "Building compressed audio format {AudioFormat} hash {Hash} wave {SoundNodeName}..."), Args));
 
+	// these types of sounds do not need cooked data
+	if(CookInputs->bIsASourceBus || CookInputs->bIsASourceBus)
+	{
+		return false;
+	}
+	
 	if (!CookInputs->ChannelSizes.Num())
 	{
 		check(!CookInputs->ChannelOffsets.Num());
diff --git a/Engine/Source/Runtime/Engine/Private/SoundWave.cpp b/Engine/Source/Runtime/Engine/Private/SoundWave.cpp
--- a/Engine/Source/Runtime/Engine/Private/SoundWave.cpp
+++ b/Engine/Source/Runtime/Engine/Private/SoundWave.cpp
@@ -2416,6 +2416,7 @@ void USoundWave::PostEditChangeProperty(FPropertyChangedEvent& PropertyChangedEv
 	static const FName SoundAssetCompressionTypeFName = GET_MEMBER_NAME_CHECKED(USoundWave, SoundAssetCompressionType);
 	static const FName LoadingBehaviorFName = GET_MEMBER_NAME_CHECKED(USoundWave, LoadingBehavior);
 	static const FName InitialChunkSizeFName = GET_MEMBER_NAME_CHECKED(USoundWave, InitialChunkSize);
+	static const FName TransformationsFName = GET_MEMBER_NAME_CHECKED(USoundWave, Transformations);
 
 	// force proxy flags to be up to date
 	SoundWaveDataPtr->bIsSeekable = IsSeekable();
@@ -2453,7 +2454,8 @@ void USoundWave::PostEditChangeProperty(FPropertyChangedEvent& PropertyChangedEv
 				|| Name == StreamingFName
 				|| Name == SoundAssetCompressionTypeFName
 				|| Name == LoadingBehaviorFName
-				|| Name == InitialChunkSizeFName)
+				|| Name == InitialChunkSizeFName
+				|| Name == TransformationsFName)
 			{
 				UpdateAsset();
 			}
@@ -2470,8 +2472,27 @@ void USoundWave::PostEditChangeProperty(FPropertyChangedEvent& PropertyChangedEv
 		}
 	}
 }
+
+
 #endif // WITH_EDITOR
 
+#if WITH_EDITORONLY_DATA
+TArray<Audio::FTransformationPtr> USoundWave::CreateTransformations() const
+{
+	TArray<Audio::FTransformationPtr> TransformationPtrs;
+
+	for(UWaveformTransformationBase* Transformation : Transformations)
+	{
+		if(Transformation)
+		{
+			TransformationPtrs.Add(Transformation->CreateTransformation());
+		}
+	}
+	
+	return TransformationPtrs;
+}
+#endif // WITH_EDITORONLY_DATA
+
 void USoundWave::FreeResources(bool bStopSoundsUsingThisResource)
 {
 	check(IsInAudioThread());
diff --git a/Engine/Source/Runtime/SignalProcessing/Private/FloatArrayMath.cpp b/Engine/Source/Runtime/SignalProcessing/Private/FloatArrayMath.cpp
--- a/Engine/Source/Runtime/SignalProcessing/Private/FloatArrayMath.cpp
+++ b/Engine/Source/Runtime/SignalProcessing/Private/FloatArrayMath.cpp
@@ -2382,6 +2382,82 @@ namespace Audio
 			}
 		}
 	}
+	
+	void ArrayInterleave(const TArray<FAlignedFloatBuffer>& InBuffers, FAlignedFloatBuffer& OutBuffer)
+	{
+		if(InBuffers.Num() == 0)
+		{
+			return;
+		}
+		
+		const int32 NumChannels = InBuffers.Num();
+		const int32 NumFrames = InBuffers[0].Num();
+
+		OutBuffer.SetNumUninitialized(NumChannels * NumFrames);
+		
+		TArray<const float*> BufferPtrArray;
+		BufferPtrArray.Reset(NumChannels);
+		
+		for(const FAlignedFloatBuffer& Buffer : InBuffers)
+		{
+			const float* BufferPtr = Buffer.GetData();
+			BufferPtrArray.Add(BufferPtr);
+		}
+		
+		const float** InBufferPtr = BufferPtrArray.GetData();
+		
+		ArrayInterleave(InBufferPtr, OutBuffer.GetData(), NumFrames, NumChannels);
+	}
+
+	void ArrayInterleave(const float** RESTRICT InBuffers, float* RESTRICT OutBuffer, const int32 InFrames, const int32 InChannels)
+	{
+		for(int32 ChannelIdx = 0; ChannelIdx < InChannels; ChannelIdx++)
+		{
+			const float* InBuffer = InBuffers[ChannelIdx];
+			
+			for(int32 SampleIdx = 0; SampleIdx < InFrames; SampleIdx++)
+			{
+				OutBuffer[ChannelIdx + (SampleIdx * InChannels)] = InBuffer[SampleIdx];
+			}
+		}
+	}
+
+	void ArrayDeinterleave(const FAlignedFloatBuffer& InBuffer, TArray<FAlignedFloatBuffer>& OutBuffers, const int32 InChannels)
+	{
+		check(InChannels > 0);
+		
+		const int32 NumFrames = InBuffer.Num() / InChannels;
+
+		TArray<float*> BufferPtrArray;
+		BufferPtrArray.Reset(InChannels);
+
+		OutBuffers.SetNum(InChannels);
+		
+		for(FAlignedFloatBuffer& Buffer : OutBuffers)
+		{
+			Buffer.SetNumUninitialized(NumFrames);
+			
+			float* BufferPtr = Buffer.GetData();
+			BufferPtrArray.Add(BufferPtr);
+		}
+		
+		float** OutBufferPtr = BufferPtrArray.GetData();
+		
+		ArrayDeinterleave(InBuffer.GetData(), OutBufferPtr, NumFrames, InChannels);
+	}
+
+	void ArrayDeinterleave(const float* RESTRICT InBuffer, float** RESTRICT OutBuffers, const int32 InFrames, const int32 InChannels)
+	{
+		for(int32 ChannelIdx = 0; ChannelIdx < InChannels; ChannelIdx++)
+		{
+			float* OutBuffer = OutBuffers[ChannelIdx];
+			
+			for(int32 SampleIdx = 0; SampleIdx < InFrames; SampleIdx++)
+			{
+				OutBuffer[SampleIdx] = InBuffer[ChannelIdx + (SampleIdx * InChannels)];
+			}
+		}
+	}
 
 	FContiguousSparse2DKernelTransform::FContiguousSparse2DKernelTransform(const int32 NumInElements, const int32 NumOutElements)
 	:	NumIn(NumInElements)
diff --git a/Engine/Source/Runtime/SignalProcessing/Public/DSP/FloatArrayMath.h b/Engine/Source/Runtime/SignalProcessing/Public/DSP/FloatArrayMath.h
--- a/Engine/Source/Runtime/SignalProcessing/Public/DSP/FloatArrayMath.h
+++ b/Engine/Source/Runtime/SignalProcessing/Public/DSP/FloatArrayMath.h
@@ -258,6 +258,18 @@ namespace Audio
 	SIGNALPROCESSING_API void ArrayFloatToPcm16(TArrayView<const float> InView, TArrayView<int16> OutView);
 	SIGNALPROCESSING_API void ArrayPcm16ToFloat(TArrayView<const int16> InView, TArrayView<float> OutView);
 
+	/** Interleaves samples from an array of input buffers */
+	SIGNALPROCESSING_API void ArrayInterleave(const TArray<FAlignedFloatBuffer>& InBuffers, FAlignedFloatBuffer& OutBuffer);
+
+	/** Interleaves samples from an array of input buffers */
+	SIGNALPROCESSING_API void ArrayInterleave(const float** RESTRICT InBuffers, float* RESTRICT OutBuffer, const int32 InFrames, const int32 InChannels);
+
+	/** Interleaves samples from an array of input buffers */
+	SIGNALPROCESSING_API void ArrayDeinterleave(const FAlignedFloatBuffer& InBuffer, TArray<FAlignedFloatBuffer>& OutBuffers, const int32 InChannels);
+
+	/** Interleaves samples from an array of input buffers */
+	SIGNALPROCESSING_API void ArrayDeinterleave(const float* RESTRICT InBuffer, float** RESTRICT OutBuffers, const int32 InFrames, const int32 InChannels);
+
 	/** FContiguousSparse2DKernelTransform
 	 *
 	 *  FContiguousSparse2DKernelTransform applies a matrix transformation to an input array. 
