Texture: Optimized batching of streaming virtual texture chunk cache requests

The Horde Storage client does not batch chunk requests, and the hierarchy does not implement propagation of them yet. Use GetValue instead of GetChunks to improve batching and ensure local caching.

#rnx

diff --git a/Engine/Source/Runtime/Engine/Private/TextureDerivedData.cpp b/Engine/Source/Runtime/Engine/Private/TextureDerivedData.cpp
--- a/Engine/Source/Runtime/Engine/Private/TextureDerivedData.cpp
+++ b/Engine/Source/Runtime/Engine/Private/TextureDerivedData.cpp
@@ -1446,14 +1446,14 @@ static bool LoadDerivedStreamingMips(FTexturePlatformData& PlatformData, int32 F
 static bool LoadDerivedStreamingVTChunks(const TArray<FVirtualTextureDataChunk>& Chunks, FStringView DebugContext, TFunctionRef<void (int32 ChunkIndex, FSharedBuffer ChunkData)> Callback)
 {
 	using namespace UE::DerivedData;
-	TArray<FCacheGetChunkRequest> Requests;
+	TArray<FCacheGetValueRequest> Requests;
 
 	for (int32 ChunkIndex = 0; ChunkIndex < Chunks.Num(); ++ChunkIndex)
 	{
 		const FVirtualTextureDataChunk& Chunk = Chunks[ChunkIndex];
 		if (!Chunk.DerivedDataKey.IsEmpty() && !Chunk.BulkData.IsBulkDataLoaded())
 		{
-			FCacheGetChunkRequest& Request = Requests.AddDefaulted_GetRef();
+			FCacheGetValueRequest& Request = Requests.AddDefaulted_GetRef();
 			Request.Name = FSharedString(WriteToString<256>(DebugContext, TEXT(" [Chunk ]"), ChunkIndex, TEXT("]")));
 			Request.Key = ConvertLegacyCacheKey(Chunk.DerivedDataKey);
 			Request.UserData = ChunkIndex;
@@ -1467,12 +1467,12 @@ static bool LoadDerivedStreamingVTChunks(const TArray<FVirtualTextureDataChunk>&
 		COOK_STAT(auto Timer = TextureCookStats::StreamingMipUsageStats.TimeSyncWork());
 		uint64 Size = 0;
 		FRequestOwner BlockingOwner(EPriority::Blocking);
-		GetCache().GetChunks(Requests, BlockingOwner, [Callback = MoveTemp(Callback), &Size, &bMiss](FCacheGetChunkResponse&& Response)
+		GetCache().GetValue(Requests, BlockingOwner, [Callback = MoveTemp(Callback), &Size, &bMiss](FCacheGetValueResponse&& Response)
 		{
-			Size += Response.RawSize;
+			Size += Response.Value.GetRawSize();
 			if (Response.Status == EStatus::Ok)
 			{
-				Callback(int32(Response.UserData), MoveTemp(Response.RawData));
+				Callback(int32(Response.UserData), Response.Value.GetData().Decompress());
 			}
 			else
 			{
