Adding support to modify an actor before other systems have a chance to pickup that a new actor was spawned.
This is useful for saving systems that want to recreate an actor in an expected state before that actor is potentially picked up.
Many systems would like to modify actor data but are required to be notified one frame after that actor was spawned to allow the spawn system to fully create that actor, this reduces the need for frame delaying delegate binds to actor spawning.

diff --git a/Engine/Source/Runtime/Engine/Classes/Engine/World.h b/Engine/Source/Runtime/Engine/Classes/Engine/World.h
--- a/Engine/Source/Runtime/Engine/Classes/Engine/World.h
+++ b/Engine/Source/Runtime/Engine/Classes/Engine/World.h
@@ -682,7 +682,10 @@ public:
 	ESpawnActorNameMode NameMode;
 
 	/* Flags used to describe the spawned actor/object instance. */
-	EObjectFlags ObjectFlags;		
+	EObjectFlags ObjectFlags;
+
+	/* Custom function allowing the caller to specific a function to execute post actor construction but before other systems see this actor spawn. */
+	TFunction<void(AActor*)> CustomPreSpawnInitalization;
 };
 
 /* World actors spawmning helper functions */
diff --git a/Engine/Source/Runtime/Engine/Private/LevelActor.cpp b/Engine/Source/Runtime/Engine/Private/LevelActor.cpp
--- a/Engine/Source/Runtime/Engine/Private/LevelActor.cpp
+++ b/Engine/Source/Runtime/Engine/Private/LevelActor.cpp
@@ -692,6 +692,11 @@ AActor* UWorld::SpawnActor( UClass* Class, FTransform const* UserTransformPtr, c
 		FActorParentComponentSetter::Set(Actor, SpawnParameters.OverrideParentComponent);
 	}
 
+	if (SpawnParameters.CustomPreSpawnInitalization)
+	{
+		SpawnParameters.CustomPreSpawnInitalization(Actor);
+	}
+
 	if ( GUndo )
 	{
 		ModifyLevel( LevelToSpawnIn );
