Chaos - add some cvar comments

diff --git a/Engine/Source/Runtime/Engine/Private/PhysicsEngine/ImmediatePhysics/ImmediatePhysicsChaos/ImmediatePhysicsSimulation_Chaos.cpp b/Engine/Source/Runtime/Engine/Private/PhysicsEngine/ImmediatePhysics/ImmediatePhysicsChaos/ImmediatePhysicsSimulation_Chaos.cpp
--- a/Engine/Source/Runtime/Engine/Private/PhysicsEngine/ImmediatePhysics/ImmediatePhysicsChaos/ImmediatePhysicsSimulation_Chaos.cpp
+++ b/Engine/Source/Runtime/Engine/Private/PhysicsEngine/ImmediatePhysics/ImmediatePhysicsChaos/ImmediatePhysicsSimulation_Chaos.cpp
@@ -149,13 +149,13 @@ bool bChaosImmediate_InertiaConditioningEnabled = true;
 FAutoConsoleVariableRef  CVarChaosImmPhysParticleInertiaConditioningEnabled(TEXT("p.Chaos.ImmPhys.InertiaConditioning.Enabled"), bChaosImmediate_InertiaConditioningEnabled, TEXT("Enable/Disable constraint stabilization through inertia conditioning"));
 
 Chaos::FRealSingle ChaosImmediate_InertiaConditioningDistance = 10;
-FAutoConsoleVariableRef  CVarChaosImmPhysParticleInertiaConditioningDistance(TEXT("p.Chaos.ImmPhys.InertiaConditioning.Distance"), ChaosImmediate_InertiaConditioningDistance, TEXT(""));
+FAutoConsoleVariableRef  CVarChaosImmPhysParticleInertiaConditioningDistance(TEXT("p.Chaos.ImmPhys.InertiaConditioning.Distance"), ChaosImmediate_InertiaConditioningDistance, TEXT("An input to inertia conditioning system. The joint distance error which needs to be stable (generate a low rotation)."));
 
 Chaos::FRealSingle ChaosImmediate_InertiaConditioningRotationRatio = 2;
-FAutoConsoleVariableRef  CVarChaosImmPhysParticleInertiaConditioningRotationRatio(TEXT("p.Chaos.ImmPhys.InertiaConditioning.RotationRatio"), ChaosImmediate_InertiaConditioningRotationRatio, TEXT(""));
+FAutoConsoleVariableRef  CVarChaosImmPhysParticleInertiaConditioningRotationRatio(TEXT("p.Chaos.ImmPhys.InertiaConditioning.RotationRatio"), ChaosImmediate_InertiaConditioningRotationRatio, TEXT("An input to inertia conditioning system. The maximum ratio of joint correction from rotation versus translation"));
 
 Chaos::FRealSingle ChaosImmediate_MaxInvInertiaComponentRatio = 0;
-FAutoConsoleVariableRef  CVarChaosImmPhysInertiaConditioningMaxInvInertiaComponentRatio(TEXT("p.Chaos.ImmPhys.InertiaConditioning.MaxInvInertiaComponentRatio"), ChaosImmediate_MaxInvInertiaComponentRatio, TEXT(""));
+FAutoConsoleVariableRef  CVarChaosImmPhysInertiaConditioningMaxInvInertiaComponentRatio(TEXT("p.Chaos.ImmPhys.InertiaConditioning.MaxInvInertiaComponentRatio"), ChaosImmediate_MaxInvInertiaComponentRatio, TEXT("An input to inertia conditioning system. The largest inertia component must be at least least multiple of the smallest component"));
 
 //
 // Select the solver technique to use until we settle on the final one...
diff --git a/Engine/Source/Runtime/Experimental/Chaos/Private/Chaos/PBDRigidsEvolutionGBF.cpp b/Engine/Source/Runtime/Experimental/Chaos/Private/Chaos/PBDRigidsEvolutionGBF.cpp
--- a/Engine/Source/Runtime/Experimental/Chaos/Private/Chaos/PBDRigidsEvolutionGBF.cpp
+++ b/Engine/Source/Runtime/Experimental/Chaos/Private/Chaos/PBDRigidsEvolutionGBF.cpp
@@ -98,15 +98,15 @@ namespace Chaos
 
 		// The largest joint error we expect to resolve in a moderately stable way
 		FRealSingle ChaosSolverInertiaConditioningDistance = 20;
-		FAutoConsoleVariableRef  CVarChaosSolverInertiaConditioningDistance(TEXT("p.Chaos.Solver.InertiaConditioning.Distance"), ChaosSolverInertiaConditioningDistance, TEXT(""));
+		FAutoConsoleVariableRef  CVarChaosSolverInertiaConditioningDistance(TEXT("p.Chaos.Solver.InertiaConditioning.Distance"), ChaosSolverInertiaConditioningDistance, TEXT("An input to inertia conditioning system. The joint distance error which needs to be stable (generate a low rotation)."));
 
 		// The ratio of joint error correction that comes from particle rotation versus translation
 		FRealSingle ChaosSolverInertiaConditioningRotationRatio = 2;
-		FAutoConsoleVariableRef  CVarChaosSolverInertiaConditioningRotationRatio(TEXT("p.Chaos.Solver.InertiaConditioning.RotationRatio"), ChaosSolverInertiaConditioningRotationRatio, TEXT(""));
+		FAutoConsoleVariableRef  CVarChaosSolverInertiaConditioningRotationRatio(TEXT("p.Chaos.Solver.InertiaConditioning.RotationRatio"), ChaosSolverInertiaConditioningRotationRatio, TEXT("An input to inertia conditioning system. The maximum ratio of joint correction from rotation versus translation"));
 
 		// If > 1, limits the inverse inertia components to be mo more than this multiple of the smallest component. Makes the objects more round. Set to 0 to disable.
 		FRealSingle ChaosSolverMaxInvInertiaComponentRatio = 0;
-		FAutoConsoleVariableRef  CVarChaosSolverInertiaConditioningMaxInvInertiaComponentRatio(TEXT("p.Chaos.Solver.InertiaConditioning.MaxInvInertiaComponentRatio"), ChaosSolverMaxInvInertiaComponentRatio, TEXT(""));
+		FAutoConsoleVariableRef  CVarChaosSolverInertiaConditioningMaxInvInertiaComponentRatio(TEXT("p.Chaos.Solver.InertiaConditioning.MaxInvInertiaComponentRatio"), ChaosSolverMaxInvInertiaComponentRatio, TEXT("An input to inertia conditioning system. The largest inertia component must be at least least multiple of the smallest component"));
 
 		DECLARE_CYCLE_STAT(TEXT("FPBDRigidsEvolutionGBF::AdvanceOneTimeStep"), STAT_Evolution_AdvanceOneTimeStep, STATGROUP_Chaos);
 		DECLARE_CYCLE_STAT(TEXT("FPBDRigidsEvolutionGBF::UnclusterUnions"), STAT_Evolution_UnclusterUnions, STATGROUP_Chaos);
