Fix a comment

diff --git a/Engine/Source/Runtime/Engine/Public/AssetExportTask.h b/Engine/Source/Runtime/Engine/Public/AssetExportTask.h
--- a/Engine/Source/Runtime/Engine/Public/AssetExportTask.h
+++ b/Engine/Source/Runtime/Engine/Public/AssetExportTask.h
@@ -8,7 +8,7 @@
 #include "AssetExportTask.generated.h"
 
 /**
- * Contains data for a group of assets to import
+ * Contains data for a group of assets to export
  */ 
 UCLASS(Transient, BlueprintType)
 class ENGINE_API UAssetExportTask : public UObject
