AssetRegistryIncludes: Change AssetData.h to AssetRegistry/AssetData.h.
#rnx

diff --git a/Engine/Source/Runtime/Engine/Classes/Materials/MaterialLayersFunctions.h b/Engine/Source/Runtime/Engine/Classes/Materials/MaterialLayersFunctions.h
--- a/Engine/Source/Runtime/Engine/Classes/Materials/MaterialLayersFunctions.h
+++ b/Engine/Source/Runtime/Engine/Classes/Materials/MaterialLayersFunctions.h
@@ -3,7 +3,7 @@
 #pragma once
 
 #include "CoreMinimal.h"
-#include "AssetData.h"
+#include "AssetRegistry/AssetData.h"
 #include "MaterialTypes.h"
 #include "MaterialLayersFunctions.generated.h"
 
