Marked various anim instance-related Skeletal Mesh Component functions as UnsafeDuringActorConstruction
#rnx

diff --git a/Engine/Source/Runtime/Engine/Classes/Components/SkeletalMeshComponent.h b/Engine/Source/Runtime/Engine/Classes/Components/SkeletalMeshComponent.h
--- a/Engine/Source/Runtime/Engine/Classes/Components/SkeletalMeshComponent.h
+++ b/Engine/Source/Runtime/Engine/Classes/Components/SkeletalMeshComponent.h
@@ -922,7 +922,7 @@ public:
 	 * Returns the active post process instance is one is available. This is set on the mesh that this
 	 * component is using, and is evaluated immediately after the main instance.
 	 */
-	UFUNCTION(BlueprintCallable, Category = "Components|SkeletalMesh", meta = (Keywords = "AnimBlueprint"))
+	UFUNCTION(BlueprintCallable, Category = "Components|SkeletalMesh", meta = (Keywords = "AnimBlueprint", UnsafeDuringActorConstruction = "true"))
 	UAnimInstance* GetPostProcessInstance() const;
 
 	/** Get the anim instances linked to the main AnimScriptInstance */
@@ -949,8 +949,8 @@ public:
 	 * Returns the a tagged linked instance node. If no linked instances are found or none are tagged with the
 	 * supplied name, this will return NULL.
 	 */
-	UFUNCTION(BlueprintPure, Category = "Components|SkeletalMesh|Animation Blueprint Linking", meta = (Keywords = "AnimBlueprint"))
-		UAnimInstance* GetLinkedAnimGraphInstanceByTag(FName InTag) const;
+	UFUNCTION(BlueprintPure, Category = "Components|SkeletalMesh|Animation Blueprint Linking", meta = (Keywords = "AnimBlueprint", UnsafeDuringActorConstruction = "true"))
+	UAnimInstance* GetLinkedAnimGraphInstanceByTag(FName InTag) const;
 
 	UE_DEPRECATED(4.24, "Function renamed, please use GetLinkedAnimGraphInstancesByTag")
 	void GetSubInstancesByTag(FName InTag, TArray<UAnimInstance*>& OutSubInstances) const
@@ -971,7 +971,7 @@ public:
 	void SetSubInstanceClassByTag(FName InTag, TSubclassOf<UAnimInstance> InClass) { LinkAnimGraphByTag(InTag, InClass); }
 
 	/** Runs through all nodes, attempting to find linked instance by name/tag, then sets the class of each node if the tag matches */
-	UFUNCTION(BlueprintCallable, Category = "Components|SkeletalMesh|Animation Blueprint Linking", meta = (Keywords = "AnimBlueprint"))
+	UFUNCTION(BlueprintCallable, Category = "Components|SkeletalMesh|Animation Blueprint Linking", meta = (Keywords = "AnimBlueprint", UnsafeDuringActorConstruction = "true"))
 	void LinkAnimGraphByTag(FName InTag, TSubclassOf<UAnimInstance> InClass);
 
 	UE_DEPRECATED(4.24, "Function renamed, please use LinkAnimClassLayers")
@@ -983,7 +983,7 @@ public:
 	 * and a separate linked instance is allocated for each layer node.
 	 * If InClass is null, then all layers are reset to their defaults.
 	 */
-	UFUNCTION(BlueprintCallable, Category = "Components|SkeletalMesh|Animation Blueprint Linking")
+	UFUNCTION(BlueprintCallable, Category = "Components|SkeletalMesh|Animation Blueprint Linking", meta = (Keywords = "AnimBlueprint", UnsafeDuringActorConstruction = "true"))
 	void LinkAnimClassLayers(TSubclassOf<UAnimInstance> InClass);
 
 	UE_DEPRECATED(4.24, "Function renamed, please use UnlinkAnimClassLayers")
@@ -994,21 +994,21 @@ public:
 	 * State sharing rules are as with SetLayerOverlay.
 	 * If InClass is null, does nothing.
 	 */
-	UFUNCTION(BlueprintCallable, Category = "Components|SkeletalMesh|Animation Blueprint Linking")
+	UFUNCTION(BlueprintCallable, Category = "Components|SkeletalMesh|Animation Blueprint Linking", meta = (Keywords = "AnimBlueprint", UnsafeDuringActorConstruction = "true"))
 	void UnlinkAnimClassLayers(TSubclassOf<UAnimInstance> InClass);
 
 	UE_DEPRECATED(4.24, "Function renamed, please use GetLinkedLayerInstanceByGroup")
 	UAnimInstance* GetLayerSubInstanceByGroup(FName InGroup) const { return GetLinkedAnimLayerInstanceByGroup(InGroup); }
 
 	/** Gets the layer linked instance corresponding to the specified group */
-	UFUNCTION(BlueprintPure, Category = "Components|SkeletalMesh|Animation Blueprint Linking")
+	UFUNCTION(BlueprintPure, Category = "Components|SkeletalMesh|Animation Blueprint Linking", meta = (Keywords = "AnimBlueprint", UnsafeDuringActorConstruction = "true"))
 	UAnimInstance* GetLinkedAnimLayerInstanceByGroup(FName InGroup) const;
 
 	UE_DEPRECATED(4.24, "Function renamed, please use GetLinkedAnimLayerInstanceByClass")
 	UAnimInstance* GetLayerSubInstanceByClass(TSubclassOf<UAnimInstance> InClass) const { return GetLinkedAnimLayerInstanceByClass(InClass);  }
 
 	/** Gets the first layer linked instance corresponding to the specified class */
-	UFUNCTION(BlueprintPure, Category = "Components|SkeletalMesh|Animation Blueprint Linking")
+	UFUNCTION(BlueprintPure, Category = "Components|SkeletalMesh|Animation Blueprint Linking", meta = (Keywords = "AnimBlueprint", UnsafeDuringActorConstruction = "true"))
 	UAnimInstance* GetLinkedAnimLayerInstanceByClass(TSubclassOf<UAnimInstance> InClass) const;
 
 	/** Calls a function on each of the anim instances that this mesh component hosts, including linked and post-process instances */
