Hide sprite in viewport for WorldPartitionMiniMap and WorldDataLayers actors.

diff --git a/Engine/Source/Runtime/Engine/Private/WorldPartition/WorldPartitionMiniMap.cpp b/Engine/Source/Runtime/Engine/Private/WorldPartition/WorldPartitionMiniMap.cpp
--- a/Engine/Source/Runtime/Engine/Private/WorldPartition/WorldPartitionMiniMap.cpp
+++ b/Engine/Source/Runtime/Engine/Private/WorldPartition/WorldPartitionMiniMap.cpp
@@ -11,7 +11,7 @@
 #define LOCTEXT_NAMESPACE "WorldPartitionEditor"
 
 AWorldPartitionMiniMap::AWorldPartitionMiniMap(const FObjectInitializer& ObjectInitializer)
-	: Super(ObjectInitializer)
+	: Super(ObjectInitializer.DoNotCreateDefaultSubobject(TEXT("Sprite")))
 	, MiniMapWorldBounds(ForceInit)
 	, MiniMapTexture(nullptr)
 {
