Hide sprite in viewport for WorldDataLayers actors.

diff --git a/Engine/Source/Runtime/Engine/Private/WorldPartition/DataLayer/WorldDataLayers.cpp b/Engine/Source/Runtime/Engine/Private/WorldPartition/DataLayer/WorldDataLayers.cpp
--- a/Engine/Source/Runtime/Engine/Private/WorldPartition/DataLayer/WorldDataLayers.cpp
+++ b/Engine/Source/Runtime/Engine/Private/WorldPartition/DataLayer/WorldDataLayers.cpp
@@ -38,7 +38,7 @@ FString JoinDataLayerShortNamesFromInstanceNames(AWorldDataLayers* InWorldDataLa
 }
 
 AWorldDataLayers::AWorldDataLayers(const FObjectInitializer& ObjectInitializer)
-	: Super(ObjectInitializer)
+	: Super(ObjectInitializer.DoNotCreateDefaultSubobject(TEXT("Sprite")))
 #if WITH_EDITORONLY_DATA
 	, bAllowRuntimeDataLayerEditing(true)
 #endif
