Level Instance Pivot now taken into account in PIE/game.
Added Level Instance Pivot Offset in Level's metadata.

diff --git a/Engine/Source/Runtime/Engine/Classes/Engine/Level.h b/Engine/Source/Runtime/Engine/Classes/Engine/Level.h
--- a/Engine/Source/Runtime/Engine/Classes/Engine/Level.h
+++ b/Engine/Source/Runtime/Engine/Classes/Engine/Level.h
@@ -778,6 +778,8 @@ public:
 	ENGINE_API static bool GetIsStreamingDisabledFromPackage(FName LevelPackage);
 	ENGINE_API static bool GetPartitionedLevelCanBeUsedByLevelInstanceFromAsset(const FAssetData& Asset);
 	ENGINE_API static bool GetPartitionedLevelCanBeUsedByLevelInstanceFromPackage(FName LevelPackage);
+	ENGINE_API static FVector GetLevelInstancePivotOffsetFromAsset(const FAssetData& Asset);
+	ENGINE_API static FVector GetLevelInstancePivotOffsetFromPackage(FName LevelPackage);
 	ENGINE_API static const FName LoadAllExternalObjectsTag;
 	ENGINE_API static const FName DontLoadExternalObjectsTag;
 
diff --git a/Engine/Source/Runtime/Engine/Private/Level.cpp b/Engine/Source/Runtime/Engine/Private/Level.cpp
--- a/Engine/Source/Runtime/Engine/Private/Level.cpp
+++ b/Engine/Source/Runtime/Engine/Private/Level.cpp
@@ -2311,6 +2311,32 @@ bool ULevel::GetLevelBoundsFromPackage(FName LevelPackage, FBox& OutLevelBounds)
 	});
 }
 
+FVector ULevel::GetLevelInstancePivotOffsetFromAsset(const FAssetData& Asset)
+{
+	static const FName NAME_LevelInstancePivotOffset(TEXT("LevelInstancePivotOffset"));
+	FString LevelInstancePivotOffsetStr;
+	if (Asset.GetTagValue(NAME_LevelInstancePivotOffset, LevelInstancePivotOffsetStr))
+	{
+		FVector LevelInstancePivotOffset;
+		if (LevelInstancePivotOffset.InitFromCompactString(LevelInstancePivotOffsetStr))
+		{
+			return LevelInstancePivotOffset;
+		}
+	}
+	return FVector::ZeroVector;
+}
+
+FVector ULevel::GetLevelInstancePivotOffsetFromPackage(FName LevelPackage)
+{
+	FVector LevelInstancePivot = FVector::ZeroVector;
+	LevelAssetRegistryHelper::GetLevelInfoFromAssetRegistry(LevelPackage, [&LevelInstancePivot](const FAssetData& Asset)
+	{
+		LevelInstancePivot = GetLevelInstancePivotOffsetFromAsset(Asset);
+		return true;
+	});
+	return LevelInstancePivot;
+}
+
 bool ULevel::GetPromptWhenAddingToLevelOutsideBounds() const
 {
 	return !bIsPartitioned && bPromptWhenAddingToLevelOutsideBounds && GetDefault<ULevelEditorMiscSettings>()->bPromptWhenAddingToLevelOutsideBounds;
diff --git a/Engine/Source/Runtime/Engine/Private/World.cpp b/Engine/Source/Runtime/Engine/Private/World.cpp
--- a/Engine/Source/Runtime/Engine/Private/World.cpp
+++ b/Engine/Source/Runtime/Engine/Private/World.cpp
@@ -8466,6 +8466,16 @@ void UWorld::GetAssetRegistryTags(TArray<FAssetRegistryTag>& OutTags) const
 			static const FName NAME_LevelIsUsingActorFolders(TEXT("LevelIsUsingActorFolders"));
 			OutTags.Add(FAssetRegistryTag(NAME_LevelIsUsingActorFolders, TEXT("1"), FAssetRegistryTag::TT_Hidden));
 		}
+
+		{
+			AWorldSettings* WorldSettings = GetWorldSettings();
+			FVector LevelInstancePivotOffset = WorldSettings ? WorldSettings->LevelInstancePivotOffset : FVector::ZeroVector;
+			if (!LevelInstancePivotOffset.IsNearlyZero())
+			{
+				static const FName NAME_LevelInstancePivotOffset(TEXT("LevelInstancePivotOffset"));
+				OutTags.Add(FAssetRegistryTag(NAME_LevelInstancePivotOffset, LevelInstancePivotOffset.ToCompactString(), FAssetRegistryTag::TT_Hidden));
+			}
+		}
 	}
 
 	// Get the full file path with extension
diff --git a/Engine/Source/Runtime/Engine/Private/WorldPartition/LevelInstance/LevelInstanceActorDesc.cpp b/Engine/Source/Runtime/Engine/Private/WorldPartition/LevelInstance/LevelInstanceActorDesc.cpp
--- a/Engine/Source/Runtime/Engine/Private/WorldPartition/LevelInstance/LevelInstanceActorDesc.cpp
+++ b/Engine/Source/Runtime/Engine/Private/WorldPartition/LevelInstance/LevelInstanceActorDesc.cpp
@@ -184,7 +184,9 @@ bool FLevelInstanceActorDesc::GetContainerInstance(const UActorDescContainer*& O
 	if (LevelInstanceContainer.IsValid())
 	{
 		OutLevelContainer = LevelInstanceContainer.Get();
-		OutLevelTransform = LevelInstanceTransform;
+		// Apply level instance pivot offset
+		FTransform LevelInstancePivotOffsetTransform = FTransform(ULevel::GetLevelInstancePivotOffsetFromPackage(LevelInstanceContainer->GetContainerPackage()));
+		OutLevelTransform = LevelInstancePivotOffsetTransform * LevelInstanceTransform;
 		OutClusterMode = EContainerClusterMode::Partitioned;
 		return true;
 	}
