World Partition Actor Descriptors
- Now includes full path to the base class to support BP classes (none if from a native class).
- FWorldPartitionHelpers actor descriptors iterator functions now supports BP classes.
- UWorldPartitionResaveActorsBuilder now supports filtering actors by BP classes.
- Scene Ouliner now shows the correct base class for unloaded actors instead of the first native class.

#rnx

diff --git a/Engine/Source/Editor/SceneOutliner/Private/ActorDescTreeItem.cpp b/Engine/Source/Editor/SceneOutliner/Private/ActorDescTreeItem.cpp
--- a/Engine/Source/Editor/SceneOutliner/Private/ActorDescTreeItem.cpp
+++ b/Engine/Source/Editor/SceneOutliner/Private/ActorDescTreeItem.cpp
@@ -136,7 +136,7 @@ private:
 		{
 			if (const FWorldPartitionActorDesc* ActorDesc = TreeItem->ActorDescHandle.Get())
 			{
-				return FText::FromName(ActorDesc->GetActorClass()->GetFName());
+				return FText::FromName(ActorDesc->GetDisplayClassName());
 			}
 		}
 
@@ -156,7 +156,7 @@ private:
 		{
 			if (const FWorldPartitionActorDesc* ActorDesc = TreeItem->ActorDescHandle.Get())
 			{
-				const FName IconName = ActorDesc->GetActorClass()->GetFName();
+				const FName IconName = ActorDesc->GetNativeClass();
 
 				const FSlateBrush* CachedBrush = WeakSceneOutliner.Pin()->GetCachedIconForClass(IconName);
 				if (CachedBrush != nullptr)
@@ -165,8 +165,7 @@ private:
 				}
 				else if (IconName != NAME_None)
 				{
-
-					const FSlateBrush* FoundSlateBrush = FSlateIconFinder::FindIconForClass(ActorDesc->GetActorClass()).GetIcon();
+					const FSlateBrush* FoundSlateBrush = FSlateIconFinder::FindIconForClass(ActorDesc->GetActorNativeClass()).GetIcon();
 					WeakSceneOutliner.Pin()->CacheIconForClass(IconName, FoundSlateBrush);
 					return FoundSlateBrush;
 				}
diff --git a/Engine/Source/Editor/SceneOutliner/Private/SceneOutlinerActorInfoColumn.cpp b/Engine/Source/Editor/SceneOutliner/Private/SceneOutlinerActorInfoColumn.cpp
--- a/Engine/Source/Editor/SceneOutliner/Private/SceneOutlinerActorInfoColumn.cpp
+++ b/Engine/Source/Editor/SceneOutliner/Private/SceneOutlinerActorInfoColumn.cpp
@@ -59,14 +59,10 @@ struct FGetInfo
 		}
 		else if (const FActorDescTreeItem* ActorDescItem = Item.CastTo<FActorDescTreeItem>())
 		{
-			const FWorldPartitionActorDesc* ActorDesc = ActorDescItem->ActorDescHandle.Get();
-
-			if (!ActorDesc)
+			if (const FWorldPartitionActorDesc* ActorDesc = ActorDescItem->ActorDescHandle.Get())
 			{
-				return FString();
+				return ActorDesc->GetDisplayClassName().ToString();
 			}
-
-			return ActorDesc->GetActorClass()->GetName();
 		}
 
 		return FString();
diff --git a/Engine/Source/Editor/UnrealEd/Private/WorldPartition/WorldPartitionHLODsBuilder.cpp b/Engine/Source/Editor/UnrealEd/Private/WorldPartition/WorldPartitionHLODsBuilder.cpp
--- a/Engine/Source/Editor/UnrealEd/Private/WorldPartition/WorldPartitionHLODsBuilder.cpp
+++ b/Engine/Source/Editor/UnrealEd/Private/WorldPartition/WorldPartitionHLODsBuilder.cpp
@@ -811,7 +811,7 @@ TArray<TArray<FGuid>> UWorldPartitionHLODsBuilder::GetHLODWorkloads(int32 NumWor
 			if (SubActor.ContainerID.IsMainContainer())
 			{
 				FWorldPartitionActorDesc* SubActorDesc = WorldPartition->GetActorDesc(SubActor.ActorGuid);
-				if (SubActorDesc && SubActorDesc->GetActorClass()->IsChildOf<AWorldPartitionHLOD>())
+				if (SubActorDesc && SubActorDesc->GetActorNativeClass()->IsChildOf<AWorldPartitionHLOD>())
 				{
 					ChildHLODs.Add(SubActor.ActorGuid);
 				}
@@ -893,7 +893,7 @@ bool UWorldPartitionHLODsBuilder::ValidateWorkload(const TArray<FGuid>&Workload)
 			return false;
 		}
 
-		if (!ActorDesc->GetActorClass()->IsChildOf<AWorldPartitionHLOD>())
+		if (!ActorDesc->GetActorNativeClass()->IsChildOf<AWorldPartitionHLOD>())
 		{
 			UE_LOG(LogWorldPartitionHLODsBuilder, Error, TEXT("Unexpected actor guid found in HLOD workload, exiting..."));
 			return false;
@@ -910,7 +910,7 @@ bool UWorldPartitionHLODsBuilder::ValidateWorkload(const TArray<FGuid>&Workload)
 				// Invalid sub actor guid found, this is unexpected when running distributed builds as the build step is always preceeded by the setup step.
 				check(SubActorDesc || !bDistributedBuild);
 
-				if (SubActorDesc && SubActorDesc->GetActorClass()->IsChildOf<AWorldPartitionHLOD>())
+				if (SubActorDesc && SubActorDesc->GetActorNativeClass()->IsChildOf<AWorldPartitionHLOD>())
 				{
 					if (!ProcessedHLOD.Contains(SubActor.ActorGuid))
 					{
diff --git a/Engine/Source/Editor/UnrealEd/Private/WorldPartition/WorldPartitionResaveActorsBuilder.cpp b/Engine/Source/Editor/UnrealEd/Private/WorldPartition/WorldPartitionResaveActorsBuilder.cpp
--- a/Engine/Source/Editor/UnrealEd/Private/WorldPartition/WorldPartitionResaveActorsBuilder.cpp
+++ b/Engine/Source/Editor/UnrealEd/Private/WorldPartition/WorldPartitionResaveActorsBuilder.cpp
@@ -95,7 +95,6 @@ bool UWorldPartitionResaveActorsBuilder::RunInternal(UWorld* World, const FCellI
 	// Actor Class Filter
 	UClass* ActorClass = AActor::StaticClass();
 
-	// @todo_ow: support BP Classes when the ActorDesc have that information
 	if (!ActorClassName.IsEmpty())
 	{
 		if (bSwitchActorPackagingSchemeToReduced)
@@ -104,11 +103,28 @@ bool UWorldPartitionResaveActorsBuilder::RunInternal(UWorld* World, const FCellI
 			return false;
 		}
 
+		// Look for native classes
 		ActorClass = FindObject<UClass>(ANY_PACKAGE, *ActorClassName);
+
 		if (!ActorClass)
 		{
-			UE_LOG(LogWorldPartitionResaveActorsBuilder, Error, TEXT("Failed to find Actor Class: %s."), *ActorClassName);
-			return false;
+			// Look for a fully qualified BP class
+			ActorClass = LoadClass<AActor>(nullptr, *ActorClassName, nullptr, LOAD_None, nullptr);
+
+			if (!ActorClass)
+			{
+				// Look for a package BP
+				if (FPackageName::DoesPackageExist(ActorClassName))
+				{
+					ActorClass = LoadClass<AActor>(nullptr, *FString::Printf(TEXT("%s.%s_C"), *ActorClassName, *FPackageName::GetLongPackageAssetName(ActorClassName)), nullptr, LOAD_None, nullptr);
+				}
+
+				if (!ActorClass)
+				{
+					UE_LOG(LogWorldPartitionResaveActorsBuilder, Error, TEXT("Failed to find Actor Class: %s."), *ActorClassName);
+					return false;
+				}
+			}
 		}
 	}
 
diff --git a/Engine/Source/Editor/UnrealEd/Public/WorldPartition/WorldPartitionResaveActorsBuilder.h b/Engine/Source/Editor/UnrealEd/Public/WorldPartition/WorldPartitionResaveActorsBuilder.h
--- a/Engine/Source/Editor/UnrealEd/Public/WorldPartition/WorldPartitionResaveActorsBuilder.h
+++ b/Engine/Source/Editor/UnrealEd/Public/WorldPartition/WorldPartitionResaveActorsBuilder.h
@@ -5,7 +5,7 @@
 #include "WorldPartition/WorldPartitionBuilder.h"
 #include "WorldPartitionResaveActorsBuilder.generated.h"
 
-// Example Command Line: ProjectName MapName -run=WorldPartitionBuilderCommandlet -SCCProvider=Perforce -Builder=WorldPartitionResaveActorsBuilder [-ActorClass=StaticMeshActor] [-SwitchActorPackagingSchemeToReduced]
+// Example Command Line: ProjectName MapName -run=WorldPartitionBuilderCommandlet -SCCProvider=Perforce -Builder=WorldPartitionResaveActorsBuilder [-ActorClassName=StaticMeshActor] [-SwitchActorPackagingSchemeToReduced]
 
 UCLASS()
 class UWorldPartitionResaveActorsBuilder : public UWorldPartitionBuilder
diff --git a/Engine/Source/Editor/WorldPartitionEditor/Private/WorldPartition/SWorldPartitionEditorGrid2D.cpp b/Engine/Source/Editor/WorldPartitionEditor/Private/WorldPartition/SWorldPartitionEditorGrid2D.cpp
--- a/Engine/Source/Editor/WorldPartitionEditor/Private/WorldPartition/SWorldPartitionEditorGrid2D.cpp
+++ b/Engine/Source/Editor/WorldPartitionEditor/Private/WorldPartition/SWorldPartitionEditorGrid2D.cpp
@@ -466,7 +466,7 @@ uint32 SWorldPartitionEditorGrid2D::PaintActors(const FGeometry& AllottedGeometr
 				float ActorBrightness = ActorDescView.GetIsSpatiallyLoaded() ? 1.0f : 0.3f;
 				FLinearColor ActorColor(ActorBrightness, ActorBrightness, ActorBrightness, ActorColorGradient);
 
-				UClass* ActorClass = ActorDescView.GetActorClass();
+				UClass* ActorClass = ActorDescView.GetActorNativeClass();
 
 				const float SquaredDistanceToPoint = ActorViewBox.ComputeSquaredDistanceToPoint(MouseCursorPos);
 				if ((ActorDescView.GetTag() == FWorldPartitionActorDesc::GlobalTag) || (SquaredDistanceToPoint > 0.0f && SquaredDistanceToPoint <= 2.0f))
diff --git a/Engine/Source/Runtime/Core/Public/UObject/FortniteNCBranchObjectVersion.h b/Engine/Source/Runtime/Core/Public/UObject/FortniteNCBranchObjectVersion.h
--- a/Engine/Source/Runtime/Core/Public/UObject/FortniteNCBranchObjectVersion.h
+++ b/Engine/Source/Runtime/Core/Public/UObject/FortniteNCBranchObjectVersion.h
@@ -24,6 +24,9 @@ struct CORE_API FFortniteNCBranchObjectVersion
 		// Remapped bEvaluateWorldPositionOffset to bEvaluateWorldPositionOffsetInRayTracing
 		RemappedEvaluateWorldPositionOffsetInRayTracing,
 
+		// Serialize native and base class for actor descriptors
+		WorldPartitionActorDescNativeBaseClassSerialization,
+
 		// -----<new versions can be added above this line>-------------------------------------------------
 		VersionPlusOne,
 		LatestVersion = VersionPlusOne - 1
diff --git a/Engine/Source/Runtime/Engine/Private/ActorPartition/ActorPartitionSubsystem.cpp b/Engine/Source/Runtime/Engine/Private/ActorPartition/ActorPartitionSubsystem.cpp
--- a/Engine/Source/Runtime/Engine/Private/ActorPartition/ActorPartitionSubsystem.cpp
+++ b/Engine/Source/Runtime/Engine/Private/ActorPartition/ActorPartitionSubsystem.cpp
@@ -196,7 +196,7 @@ public:
 		bool bUnloadedActorExists = false;
 		auto FindActor = [&FoundActor, &bUnloadedActorExists, InCellCoord, InActorPartitionId, InGridSize, ThisWorld = World](const FWorldPartitionActorDesc* ActorDesc)
 		{
-			check(ActorDesc->GetActorClass()->IsChildOf(InActorPartitionId.GetClass()));
+			check(ActorDesc->GetActorNativeClass()->IsChildOf(InActorPartitionId.GetClass()));
 			FPartitionActorDesc* PartitionActorDesc = (FPartitionActorDesc*)ActorDesc;
 			if ((PartitionActorDesc->GridIndexX == InCellCoord.X) &&
 				(PartitionActorDesc->GridIndexY == InCellCoord.Y) &&
diff --git a/Engine/Source/Runtime/Engine/Private/WorldPartition/HLOD/HLODLayer.cpp b/Engine/Source/Runtime/Engine/Private/WorldPartition/HLOD/HLODLayer.cpp
--- a/Engine/Source/Runtime/Engine/Private/WorldPartition/HLOD/HLODLayer.cpp
+++ b/Engine/Source/Runtime/Engine/Private/WorldPartition/HLOD/HLODLayer.cpp
@@ -67,7 +67,7 @@ UHLODLayer* UHLODLayer::GetHLODLayer(const FWorldPartitionActorDescView& InActor
 	}
 
 	// Only fallback to the default HLODLayer for the first level of HLOD
-	bool bIsHLOD0 = !InActorDesc.GetActorClass()->IsChildOf<AWorldPartitionHLOD>();
+	bool bIsHLOD0 = !InActorDesc.GetActorNativeClass()->IsChildOf<AWorldPartitionHLOD>();
 	if (bIsHLOD0)
 	{
 		// Fallback to the world partition default HLOD layer
diff --git a/Engine/Source/Runtime/Engine/Private/WorldPartition/LevelInstance/LevelInstanceActorDesc.cpp b/Engine/Source/Runtime/Engine/Private/WorldPartition/LevelInstance/LevelInstanceActorDesc.cpp
--- a/Engine/Source/Runtime/Engine/Private/WorldPartition/LevelInstance/LevelInstanceActorDesc.cpp
+++ b/Engine/Source/Runtime/Engine/Private/WorldPartition/LevelInstance/LevelInstanceActorDesc.cpp
@@ -237,7 +237,7 @@ void FLevelInstanceActorDesc::Serialize(FArchive& Ar)
 		const bool bFixupOldVersion = (Ar.CustomVer(FUE5MainStreamObjectVersion::GUID) < FUE5MainStreamObjectVersion::PackedLevelInstanceBoundsFix) && 
 									  (Ar.CustomVer(FUE5ReleaseStreamObjectVersion::GUID) < FUE5ReleaseStreamObjectVersion::PackedLevelInstanceBoundsFix);
 
-		const AActor* CDO = GetActorClass()->GetDefaultObject<AActor>();
+		const AActor* CDO = GetActorNativeClass()->GetDefaultObject<AActor>();
 		const ILevelInstanceInterface* LevelInstanceCDO = CastChecked<ILevelInstanceInterface>(CDO);
 		if (!LevelPackage.IsNone() && (LevelInstanceCDO->IsLoadingEnabled() || bFixupOldVersion))
 		{
diff --git a/Engine/Source/Runtime/Engine/Private/WorldPartition/RuntimeSpatialHash/RuntimeSpatialHashHLOD.cpp b/Engine/Source/Runtime/Engine/Private/WorldPartition/RuntimeSpatialHash/RuntimeSpatialHashHLOD.cpp
--- a/Engine/Source/Runtime/Engine/Private/WorldPartition/RuntimeSpatialHash/RuntimeSpatialHashHLOD.cpp
+++ b/Engine/Source/Runtime/Engine/Private/WorldPartition/RuntimeSpatialHash/RuntimeSpatialHashHLOD.cpp
@@ -266,7 +266,7 @@ static TMap<UHLODLayer*, int32> GatherHLODLayers(UWorldPartition* WorldPartition
 	for (FActorDescList::TIterator<> ActorDescIterator(WorldPartition); ActorDescIterator; ++ActorDescIterator)
 	{
 		const FWorldPartitionActorDesc& ActorDesc = **ActorDescIterator;
-		if (!ActorDesc.GetActorClass()->IsChildOf<AWorldPartitionHLOD>())
+		if (!ActorDesc.GetActorNativeClass()->IsChildOf<AWorldPartitionHLOD>())
 		{
 			if (ActorDesc.GetActorIsHLODRelevant())
 			{
diff --git a/Engine/Source/Runtime/Engine/Private/WorldPartition/WorldPartition.cpp b/Engine/Source/Runtime/Engine/Private/WorldPartition/WorldPartition.cpp
--- a/Engine/Source/Runtime/Engine/Private/WorldPartition/WorldPartition.cpp
+++ b/Engine/Source/Runtime/Engine/Private/WorldPartition/WorldPartition.cpp
@@ -99,7 +99,7 @@ FString GetActorDescDumpString(const FWorldPartitionActorDesc* ActorDesc, const
 	return FString::Printf(
 		TEXT("%s, %s, %s, %s, %s, %.2f, %.2f, %.2f, %.2f, %.2f, %.2f, %s") LINE_TERMINATOR,
 		*ActorDesc->GetGuid().ToString(),
-		*ActorDesc->GetClass().ToString(),
+		*ActorDesc->GetNativeClass().ToString(),
 		*ActorDesc->GetActorName().ToString(),
 		*ActorDesc->GetActorPackage().ToString(),
 		*ActorDesc->GetActorLabel().ToString(),
@@ -490,7 +490,7 @@ void UWorldPartition::Initialize(UWorld* InWorld, const FTransform& InTransform)
 		// Make sure to preload only AWorldDataLayers actor first (ShouldActorBeLoadedByEditorCells requires it)
 		for (UActorDescContainer::TIterator<> ActorDescIterator(this); ActorDescIterator; ++ActorDescIterator)
 		{
-			if (ActorDescIterator->GetActorClass()->IsChildOf<AWorldDataLayers>())
+			if (ActorDescIterator->GetActorNativeClass()->IsChildOf<AWorldDataLayers>())
 			{
 				WorldDataLayersActor = FWorldPartitionReference(this, ActorDescIterator->GetGuid());
 				break;
diff --git a/Engine/Source/Runtime/Engine/Private/WorldPartition/WorldPartitionActorDesc.cpp b/Engine/Source/Runtime/Engine/Private/WorldPartition/WorldPartitionActorDesc.cpp
--- a/Engine/Source/Runtime/Engine/Private/WorldPartition/WorldPartitionActorDesc.cpp
+++ b/Engine/Source/Runtime/Engine/Private/WorldPartition/WorldPartitionActorDesc.cpp
@@ -10,6 +10,7 @@
 #include "Algo/Transform.h"
 #include "Engine/World.h"
 #include "Engine/Level.h"
+#include "Misc/PackageName.h"
 #include "UObject/UE5MainStreamObjectVersion.h"
 #include "UObject/UE5ReleaseStreamObjectVersion.h"
 #include "UObject/FortniteNCBranchObjectVersion.h"
@@ -41,9 +42,17 @@ void FWorldPartitionActorDesc::Init(const AActor* InActor)
 	Guid = InActor->GetActorGuid();
 	check(Guid.IsValid());
 
+	UClass* ActorClass = InActor->GetClass();
+
 	// Get the first native class in the hierarchy
-	ActorClass = GetParentNativeClass(InActor->GetClass());
-	Class = ActorClass->GetFName();
+	ActorNativeClass = GetParentNativeClass(ActorClass);
+	NativeClass = ActorNativeClass->GetFName();
+	
+	// For native class, don't set this
+	if (!ActorClass->IsNative())
+	{
+		BaseClass = *InActor->GetClass()->GetPathName();
+	}
 
 	const FBox StreamingBounds = InActor->GetStreamingBounds();
 	StreamingBounds.GetCenterAndExtents(BoundsLocation, BoundsExtent);
@@ -120,8 +129,8 @@ void FWorldPartitionActorDesc::Init(const FWorldPartitionActorDescInitData& Desc
 {
 	ActorPackage = DescData.PackageName;
 	ActorPath = DescData.ActorPath;
-	ActorClass = DescData.NativeClass;
-	Class = DescData.NativeClass->GetFName();
+	ActorNativeClass = DescData.NativeClass;
+	NativeClass = DescData.NativeClass->GetFName();
 
 	// Serialize actor metadata
 	FMemoryReader MetadataAr(DescData.SerializedData, true);
@@ -140,7 +149,8 @@ void FWorldPartitionActorDesc::Init(const FWorldPartitionActorDescInitData& Desc
 bool FWorldPartitionActorDesc::Equals(const FWorldPartitionActorDesc* Other) const
 {
 	if (Guid == Other->Guid && 
-		Class == Other->Class && 
+		BaseClass == Other->BaseClass && 
+		NativeClass == Other->NativeClass && 
 		ActorPackage == Other->ActorPackage && 
 		ActorPath == Other->ActorPath && 
 		ActorLabel == Other->ActorLabel && 
@@ -214,9 +224,10 @@ void FWorldPartitionActorDesc::TransformInstance(const FString& From, const FStr
 FString FWorldPartitionActorDesc::ToString() const
 {
 	return FString::Printf(
-		TEXT("Guid:%s Class:%s Name:%s SpatiallyLoaded:%s Bounds:%s RuntimeGrid:%s EditorOnly:%s LevelBoundsRelevant:%s HLODRelevant:%s FolderPath:%s FolderGuid:%s Parent:%s"), 
+		TEXT("Guid:%s BaseClass:%s NativeClass:%s Name:%s SpatiallyLoaded:%s Bounds:%s RuntimeGrid:%s EditorOnly:%s LevelBoundsRelevant:%s HLODRelevant:%s FolderPath:%s FolderGuid:%s Parent:%s"), 
 		*Guid.ToString(), 
-		*Class.ToString(), 
+		*BaseClass.ToString(), 
+		*NativeClass.ToString(), 
 		*GetActorName().ToString(),
 		bIsSpatiallyLoaded ? TEXT("true") : TEXT("false"),
 		*GetBounds().ToString(),
@@ -238,7 +249,12 @@ void FWorldPartitionActorDesc::Serialize(FArchive& Ar)
 	Ar.UsingCustomVersion(FUE5ReleaseStreamObjectVersion::GUID);
 	Ar.UsingCustomVersion(FFortniteNCBranchObjectVersion::GUID);
 
-	Ar << Class << Guid;
+	if (Ar.CustomVer(FFortniteNCBranchObjectVersion::GUID) >= FFortniteNCBranchObjectVersion::WorldPartitionActorDescNativeBaseClassSerialization)
+	{
+		Ar << BaseClass;
+	}
+
+	Ar << NativeClass << Guid;
 
 	if(Ar.CustomVer(FUE5ReleaseStreamObjectVersion::GUID) < FUE5ReleaseStreamObjectVersion::LargeWorldCoordinates)
 	{
@@ -338,6 +354,25 @@ FName FWorldPartitionActorDesc::GetActorLabelOrName() const
 	return GetActorLabel().IsNone() ? GetActorName() : GetActorLabel();
 }
 
+FName FWorldPartitionActorDesc::GetDisplayClassName() const
+{
+	if (BaseClass.IsNone())
+	{
+		return NativeClass;
+	}
+
+	int32 Index;
+	const FString BaseClassStr(BaseClass.ToString());
+	if (BaseClassStr.FindLastChar(TCHAR('.'), Index))
+	{
+		FString CleanClassName = BaseClassStr.Mid(Index + 1);
+		CleanClassName.RemoveFromEnd(TEXT("_C"));
+		return *CleanClassName;
+	}
+
+	return BaseClass;
+}
+
 bool FWorldPartitionActorDesc::IsLoaded(bool bEvenIfPendingKill) const
 {
 	if (ActorPtr.IsExplicitlyNull())
diff --git a/Engine/Source/Runtime/Engine/Private/WorldPartition/WorldPartitionActorDescView.cpp b/Engine/Source/Runtime/Engine/Private/WorldPartition/WorldPartitionActorDescView.cpp
--- a/Engine/Source/Runtime/Engine/Private/WorldPartition/WorldPartitionActorDescView.cpp
+++ b/Engine/Source/Runtime/Engine/Private/WorldPartition/WorldPartitionActorDescView.cpp
@@ -37,14 +37,19 @@ const FGuid& FWorldPartitionActorDescView::GetGuid() const
 	return ActorDesc->GetGuid();
 }
 
-FName FWorldPartitionActorDescView::GetClass() const
+FName FWorldPartitionActorDescView::GetBaseClass() const
 {
-	return ActorDesc->GetClass();
+	return ActorDesc->GetBaseClass();
 }
 
-UClass* FWorldPartitionActorDescView::GetActorClass() const
+FName FWorldPartitionActorDescView::GetNativeClass() const
 {
-	return ActorDesc->GetActorClass();
+	return ActorDesc->GetNativeClass();
+}
+
+UClass* FWorldPartitionActorDescView::GetActorNativeClass() const
+{
+	return ActorDesc->GetActorNativeClass();
 }
 
 FVector FWorldPartitionActorDescView::GetOrigin() const
diff --git a/Engine/Source/Runtime/Engine/Private/WorldPartition/WorldPartitionHelpers.cpp b/Engine/Source/Runtime/Engine/Private/WorldPartition/WorldPartitionHelpers.cpp
--- a/Engine/Source/Runtime/Engine/Private/WorldPartition/WorldPartitionHelpers.cpp
+++ b/Engine/Source/Runtime/Engine/Private/WorldPartition/WorldPartitionHelpers.cpp
@@ -9,14 +9,34 @@
 
 DEFINE_LOG_CATEGORY_STATIC(LogWorldPartitionHelpers, Log, All);
 
-
 #if WITH_EDITOR
 
+UClass* FWorldPartitionHelpers::ResolveActorDescClass(const FWorldPartitionActorDesc* ActorDesc)
+{
+	UClass* ActorNativeClass = ActorDesc->GetActorNativeClass();
+	UClass* ActorBaseClass = ActorNativeClass;
+
+	if (FName ActorBaseClassName = ActorDesc->GetBaseClass(); !ActorBaseClassName.IsNone())
+	{
+		ActorBaseClass = LoadClass<AActor>(nullptr, *ActorBaseClassName.ToString(), nullptr, LOAD_None, nullptr);
+
+		if (!ActorBaseClass)
+		{
+			UE_LOG(LogWorldPartitionHelpers, Warning, TEXT("Failed to find actor base class: %s."), *ActorBaseClassName.ToString());
+			ActorBaseClass = ActorNativeClass;
+		}
+	}
+
+	return ActorBaseClass;
+}
+
 void FWorldPartitionHelpers::ForEachIntersectingActorDesc(UWorldPartition* WorldPartition, const FBox& Box, TSubclassOf<AActor> ActorClass, TFunctionRef<bool(const FWorldPartitionActorDesc*)> Predicate)
 {
 	WorldPartition->EditorHash->ForEachIntersectingActor(Box, [&ActorClass, Predicate](const FWorldPartitionActorDesc* ActorDesc)
 	{
-		if (ActorDesc->GetActorClass()->IsChildOf(ActorClass))
+		UClass* ActorDescClass = ResolveActorDescClass(ActorDesc);
+
+		if (ActorDescClass->IsChildOf(ActorClass))
 		{
 			Predicate(ActorDesc);
 		}
@@ -25,11 +45,16 @@ void FWorldPartitionHelpers::ForEachIntersectingActorDesc(UWorldPartition* World
 
 void FWorldPartitionHelpers::ForEachActorDesc(UWorldPartition* WorldPartition, TSubclassOf<AActor> ActorClass, TFunctionRef<bool(const FWorldPartitionActorDesc*)> Predicate)
 {
-	for (UActorDescContainer::TConstIterator<> ActorDescIterator(WorldPartition, ActorClass); ActorDescIterator; ++ActorDescIterator)
+	for (UActorDescContainer::TConstIterator<> ActorDescIterator(WorldPartition); ActorDescIterator; ++ActorDescIterator)
 	{
-		if (!Predicate(*ActorDescIterator))
+		UClass* ActorDescClass = ResolveActorDescClass(*ActorDescIterator);
+
+		if (ActorDescClass->IsChildOf(ActorClass))
 		{
-			return;
+			if (!Predicate(*ActorDescIterator))
+			{
+				return;
+			}
 		}
 	}
 }
@@ -81,9 +106,16 @@ void FWorldPartitionHelpers::ForEachActorWithLoading(UWorldPartition* WorldParti
 {
 	TMap<FGuid, FWorldPartitionReference> ActorReferences;
 
-	ForEachActorDesc(WorldPartition, ActorClass, [&](const FWorldPartitionActorDesc* ActorDesc)
+	ForEachActorDesc(WorldPartition, [&](const FWorldPartitionActorDesc* ActorDesc)
 	{
-		return WorldPartitionHelpers::ForEachActorWithLoadingBody(ActorDesc->GetGuid(), WorldPartition, Predicate, OnReleasingActorReferences, bGCPerActor, ActorReferences);
+		UClass* ActorDescClass = ResolveActorDescClass(ActorDesc);
+
+		if (ActorDescClass->IsChildOf(ActorClass))
+		{
+			return WorldPartitionHelpers::ForEachActorWithLoadingBody(ActorDesc->GetGuid(), WorldPartition, Predicate, OnReleasingActorReferences, bGCPerActor, ActorReferences);
+		}
+
+		return true;
 	});
 
 	OnReleasingActorReferences();
diff --git a/Engine/Source/Runtime/Engine/Private/WorldPartition/WorldPartitionRuntimeSpatialHash.cpp b/Engine/Source/Runtime/Engine/Private/WorldPartition/WorldPartitionRuntimeSpatialHash.cpp
--- a/Engine/Source/Runtime/Engine/Private/WorldPartition/WorldPartitionRuntimeSpatialHash.cpp
+++ b/Engine/Source/Runtime/Engine/Private/WorldPartition/WorldPartitionRuntimeSpatialHash.cpp
@@ -803,7 +803,7 @@ bool UWorldPartitionRuntimeSpatialHash::GenerateStreaming(UWorldPartitionStreami
 	for (auto& ActorDescViewPair : ContainerInstance->ActorDescViewMap)
 	{
 		const FWorldPartitionActorDescView& ActorDescView = ActorDescViewPair.Value;
-		if (ActorDescView.GetActorClass()->IsChildOf<ASpatialHashRuntimeGridInfo>())
+		if (ActorDescView.GetActorNativeClass()->IsChildOf<ASpatialHashRuntimeGridInfo>())
 		{
 			FWorldPartitionReference Ref(WorldPartition, ActorDescView.GetGuid());
 			if (ASpatialHashRuntimeGridInfo* RuntimeGridActor = Cast<ASpatialHashRuntimeGridInfo>(Ref->GetActor()))
diff --git a/Engine/Source/Runtime/Engine/Private/WorldPartition/WorldPartitionStreamingGeneration.cpp b/Engine/Source/Runtime/Engine/Private/WorldPartition/WorldPartitionStreamingGeneration.cpp
--- a/Engine/Source/Runtime/Engine/Private/WorldPartition/WorldPartitionStreamingGeneration.cpp
+++ b/Engine/Source/Runtime/Engine/Private/WorldPartition/WorldPartitionStreamingGeneration.cpp
@@ -648,7 +648,7 @@ void UWorldPartition::GenerateHLOD(ISourceControlHelper* SourceControlHelper, bo
 	// Preparation Phase :: Actor Clusters Creation
 	FActorClusterContext ActorClusterContext = StreamingGenerator.CreateActorClusters([](const FWorldPartitionActorDescView& ActorDescView)
 	{
-		return !ActorDescView.GetActorClass()->IsChildOf<AWorldPartitionHLOD>();
+		return !ActorDescView.GetActorNativeClass()->IsChildOf<AWorldPartitionHLOD>();
 	});
 
 	RuntimeHash->GenerateHLOD(SourceControlHelper, ActorClusterContext, bCreateActorsOnly);
diff --git a/Engine/Source/Runtime/Engine/Public/WorldPartition/ActorDescList.h b/Engine/Source/Runtime/Engine/Public/WorldPartition/ActorDescList.h
--- a/Engine/Source/Runtime/Engine/Public/WorldPartition/ActorDescList.h
+++ b/Engine/Source/Runtime/Engine/Public/WorldPartition/ActorDescList.h
@@ -51,6 +51,7 @@ public:
 			: ActorsIterator(InActorDescList->ActorsByGuid)
 			, ActorClass(InActorClass)
 		{
+			check(ActorClass->IsNative());
 			check(ActorClass->IsChildOf(ActorType::StaticClass()));
 
 			if (ShouldSkip())
@@ -112,7 +113,7 @@ public:
 				return false;
 			}
 
-			return !ActorsIterator->Value->Get()->GetActorClass()->IsChildOf(ActorClass);
+			return !ActorsIterator->Value->Get()->GetActorNativeClass()->IsChildOf(ActorClass);
 		}
 
 		IteratorType ActorsIterator;
diff --git a/Engine/Source/Runtime/Engine/Public/WorldPartition/WorldPartitionActorDesc.h b/Engine/Source/Runtime/Engine/Public/WorldPartition/WorldPartitionActorDesc.h
--- a/Engine/Source/Runtime/Engine/Public/WorldPartition/WorldPartitionActorDesc.h
+++ b/Engine/Source/Runtime/Engine/Public/WorldPartition/WorldPartitionActorDesc.h
@@ -45,8 +45,12 @@ public:
 
 	inline const FGuid& GetGuid() const { return Guid; }
 	
-	inline FName GetClass() const { return Class; }
-	inline UClass* GetActorClass() const { return ActorClass; }
+	UE_DEPRECATED(5.1, "GetClass is deprecated, GetBaseClass or GetNativeClass should be used instead.")
+	inline FName GetClass() const { return GetNativeClass(); }
+
+	inline FName GetBaseClass() const { return BaseClass; }
+	inline FName GetNativeClass() const { return NativeClass; }
+	inline UClass* GetActorNativeClass() const { return ActorNativeClass; }
 	inline FVector GetOrigin() const { return GetBounds().GetCenter(); }
 	inline FName GetRuntimeGrid() const { return RuntimeGrid; }
 	inline bool GetIsSpatiallyLoaded() const { return bIsForcedNonSpatiallyLoaded ? false : bIsSpatiallyLoaded; }
@@ -69,6 +73,7 @@ public:
 
 	FName GetActorName() const;
 	FName GetActorLabelOrName() const;
+	FName GetDisplayClassName() const;
 
 	virtual bool GetContainerInstance(const UActorDescContainer*& OutLevelContainer, FTransform& OutLevelTransform, EContainerClusterMode& OutClusterMode) const { return false; }
 	virtual const FGuid& GetSceneOutlinerParent() const { return GetParentActor(); }
@@ -176,7 +181,8 @@ protected:
 
 	// Persistent
 	FGuid							Guid;
-	FName							Class;
+	FName							BaseClass;
+	FName							NativeClass;
 	FName							ActorPackage;
 	FName							ActorPath;
 	FName							ActorLabel;
@@ -199,7 +205,7 @@ protected:
 	// Transient
 	mutable uint32					SoftRefCount;
 	mutable uint32					HardRefCount;
-	UClass*							ActorClass;
+	UClass*							ActorNativeClass;
 	mutable TWeakObjectPtr<AActor>	ActorPtr;
 	UActorDescContainer*			Container;
 	TArray<FName>					DataLayerInstanceNames;
diff --git a/Engine/Source/Runtime/Engine/Public/WorldPartition/WorldPartitionActorDescView.h b/Engine/Source/Runtime/Engine/Public/WorldPartition/WorldPartitionActorDescView.h
--- a/Engine/Source/Runtime/Engine/Public/WorldPartition/WorldPartitionActorDescView.h
+++ b/Engine/Source/Runtime/Engine/Public/WorldPartition/WorldPartitionActorDescView.h
@@ -23,8 +23,9 @@ public:
 	FWorldPartitionActorDescView(const FWorldPartitionActorDesc* InActorDesc);
 
 	const FGuid& GetGuid() const;
-	FName GetClass() const;
-	UClass* GetActorClass() const;
+	FName GetBaseClass() const;
+	FName GetNativeClass() const;
+	UClass* GetActorNativeClass() const;
 	FVector GetOrigin() const;
 	FName GetRuntimeGrid() const;
 	bool GetIsSpatiallyLoaded() const;
diff --git a/Engine/Source/Runtime/Engine/Public/WorldPartition/WorldPartitionHelpers.h b/Engine/Source/Runtime/Engine/Public/WorldPartition/WorldPartitionHelpers.h
--- a/Engine/Source/Runtime/Engine/Public/WorldPartition/WorldPartitionHelpers.h
+++ b/Engine/Source/Runtime/Engine/Public/WorldPartition/WorldPartitionHelpers.h
@@ -12,6 +12,8 @@ class FWorldPartitionActorDesc;
 
 class ENGINE_API FWorldPartitionHelpers
 {
+	static UClass* ResolveActorDescClass(const FWorldPartitionActorDesc* ActorDesc);
+
 public:
 	template <class ActorClass = AActor>
 	static void ForEachIntersectingActorDesc(UWorldPartition* WorldPartition, const FBox& Box, TFunctionRef<bool(const FWorldPartitionActorDesc*)> Func)
