Fixed actors not showing in their folders when running PIE.

diff --git a/Engine/Source/Runtime/Engine/Classes/Engine/Level.h b/Engine/Source/Runtime/Engine/Classes/Engine/Level.h
--- a/Engine/Source/Runtime/Engine/Classes/Engine/Level.h
+++ b/Engine/Source/Runtime/Engine/Classes/Engine/Level.h
@@ -76,6 +76,7 @@ private:
 	friend class UActorFolder;
 	friend class UWorldPartitionConvertCommandlet;
 	friend class FWorldPartitionLevelHelper;
+	friend class UWorldPartitionLevelStreamingDynamic;
 };
 
 class ENGINE_API FLevelPartitionOperationScope
diff --git a/Engine/Source/Runtime/Engine/Private/WorldPartition/WorldPartitionLevelHelper.cpp b/Engine/Source/Runtime/Engine/Private/WorldPartition/WorldPartitionLevelHelper.cpp
--- a/Engine/Source/Runtime/Engine/Private/WorldPartition/WorldPartitionLevelHelper.cpp
+++ b/Engine/Source/Runtime/Engine/Private/WorldPartition/WorldPartitionLevelHelper.cpp
@@ -383,11 +383,7 @@ bool FWorldPartitionLevelHelper::LoadActors(UWorld* InOwningWorld, ULevel* InDes
 				// Duplicate Folder if needed (this will recreate a transient folder structure in InDestLevel if it doesn't exist (only used in PIE)
 				if (FGuid ActorFolderGuid = Actor->GetFolderGuid(); InDestLevel && ActorFolderGuid.IsValid())
 				{
-					// Make sure Dest level is properly setup
-					if (!InDestLevel->IsUsingActorFolders())
-					{
-						FLevelActorFoldersHelper::SetUseActorFolders(InDestLevel, true);
-					}
+					check(InDestLevel->IsUsingActorFolders());
 					InDestLevel->bFixupActorFoldersAtLoad = false;
 
 					// Make sure Source level actor folder fixup was called
diff --git a/Engine/Source/Runtime/Engine/Private/WorldPartition/WorldPartitionLevelStreamingDynamic.cpp b/Engine/Source/Runtime/Engine/Private/WorldPartition/WorldPartitionLevelStreamingDynamic.cpp
--- a/Engine/Source/Runtime/Engine/Private/WorldPartition/WorldPartitionLevelStreamingDynamic.cpp
+++ b/Engine/Source/Runtime/Engine/Private/WorldPartition/WorldPartitionLevelStreamingDynamic.cpp
@@ -134,6 +134,11 @@ void UWorldPartitionLevelStreamingDynamic::CreateRuntimeLevel()
 	RuntimeLevel = FWorldPartitionLevelHelper::CreateEmptyLevelForRuntimeCell(StreamingCell.Get(), World, GetWorldAsset().ToString());
 	check(RuntimeLevel);
 
+	if (World->PersistentLevel->IsUsingActorFolders())
+	{
+		FLevelActorFoldersHelper::SetUseActorFolders(RuntimeLevel, true);
+	}
+
 	UPackage* RuntimeLevelPackage = RuntimeLevel->GetPackage();
 	check(RuntimeLevelPackage);
 
