Fixed infinite loop when changing the Data Layer parent on a child Data Layer

diff --git a/Engine/Source/Runtime/Engine/Private/WorldPartition/DataLayer/DataLayerInstance.cpp b/Engine/Source/Runtime/Engine/Private/WorldPartition/DataLayer/DataLayerInstance.cpp
--- a/Engine/Source/Runtime/Engine/Private/WorldPartition/DataLayer/DataLayerInstance.cpp
+++ b/Engine/Source/Runtime/Engine/Private/WorldPartition/DataLayer/DataLayerInstance.cpp
@@ -121,6 +121,20 @@ bool UDataLayerInstance::SetParent(UDataLayerInstance* InParent)
 	}
 
 	Modify();
+
+	// If we find ourself in the parent chain of the provided parent
+	UDataLayerInstance* CurrentInstance = InParent;
+	while (CurrentInstance)
+	{
+		if (CurrentInstance == this)
+		{
+			// Detach the parent from its parent
+			InParent->SetParent(nullptr);
+			break;
+		}
+		CurrentInstance = CurrentInstance->GetParent();
+	};
+
 	if (Parent)
 	{
 		Parent->RemoveChild(this);
