Fixed Level Instance at wrong location when editing a level instance after a drag & drop of a level from the content browser on the level property of the level instance.

diff --git a/Engine/Source/Runtime/Engine/Private/LevelInstance/LevelInstanceEditorLevelStreaming.cpp b/Engine/Source/Runtime/Engine/Private/LevelInstance/LevelInstanceEditorLevelStreaming.cpp
--- a/Engine/Source/Runtime/Engine/Private/LevelInstance/LevelInstanceEditorLevelStreaming.cpp
+++ b/Engine/Source/Runtime/Engine/Private/LevelInstance/LevelInstanceEditorLevelStreaming.cpp
@@ -10,6 +10,7 @@
 #include "Editor.h"
 #include "Engine/LevelBounds.h"
 #include "GameFramework/WorldSettings.h"
+#include "PackageTools.h"
 #endif
 
 #if WITH_EDITOR
@@ -61,6 +62,16 @@ ULevelStreamingLevelInstanceEditor* ULevelStreamingLevelInstanceEditor::Load(ILe
 	AActor* LevelInstanceActor = CastChecked<AActor>(LevelInstance);
 	UWorld* CurrentWorld = LevelInstanceActor->GetWorld();
 
+	// Make sure the level is not already loaded
+	// For instance, drag and drop a Level into a Level Instance, then Edit the Level Instance will
+	// result in the level at the wrong location (level transform wasn't applied)
+	UPackage* LevelPackage = FindObject<UPackage>(nullptr, *LevelInstance->GetWorldAssetPackage());
+	UWorld* LevelWorld = LevelPackage ? UWorld::FindWorldInPackage(LevelPackage) : nullptr;
+	if (LevelWorld && LevelWorld->bIsWorldInitialized && !LevelWorld->GetPackage()->HasAnyPackageFlags(PKG_NewlyCreated))
+	{
+		UPackageTools::UnloadPackages({ LevelPackage });
+	}
+
 	TGuardValue<FLevelInstanceID> GuardEditLevelInstanceID(EditLevelInstanceID, LevelInstance->GetLevelInstanceID());
 	ULevelStreamingLevelInstanceEditor* LevelStreaming = nullptr;
 	// If Asset is null we can create a level here
