Added critical section to Property Access Library to prevent race conditions

In async loading thread-enabled builds, copying the property access library during CDO serialization can race with the parent classes postload-fixup (i.e. PatchPropertyOffsets). Adding a critical section prevents this race, and is rarely contended.

diff --git a/Engine/Source/Runtime/Engine/Private/PropertyAccess.cpp b/Engine/Source/Runtime/Engine/Private/PropertyAccess.cpp
--- a/Engine/Source/Runtime/Engine/Private/PropertyAccess.cpp
+++ b/Engine/Source/Runtime/Engine/Private/PropertyAccess.cpp
@@ -6,6 +6,11 @@
 
 #define LOCTEXT_NAMESPACE "PropertyAccess"
 
+namespace PropertyAccess
+{
+	static FCriticalSection CriticalSection;
+}
+
 struct FPropertyAccessSystem
 {
 	struct FResolveIndirectionsOnLoadContext
@@ -216,6 +221,9 @@ struct FPropertyAccessSystem
 
 	static void PatchPropertyOffsets(FPropertyAccessLibrary& InLibrary)
 	{
+		// Need to perform a lock as copying can race with PatchPropertyOffsets in async loading thread-enabled builds
+		FScopeLock Lock(&PropertyAccess::CriticalSection);
+
 		InLibrary.Indirections.Reset();
 
 		const int32 SrcCount = InLibrary.SrcPaths.Num();
@@ -602,26 +610,17 @@ struct FPropertyAccessSystem
 
 const FPropertyAccessLibrary& FPropertyAccessLibrary::operator =(const FPropertyAccessLibrary& Other)
 {
+	// Need to perform a lock as copying can race with PatchPropertyOffsets in async loading thread-enabled builds
+	FScopeLock Lock(&PropertyAccess::CriticalSection);
+
 	PathSegments = Other.PathSegments;
 	SrcPaths = Other.SrcPaths;
 	DestPaths = Other.DestPaths;
-	CopyBatchArray = Other.CopyBatchArray; 
-
-	if (Other.bHasBeenPostLoaded)
-	{
-		SrcAccesses = Other.SrcAccesses;
-		DestAccesses = Other.DestAccesses;
-		Indirections = Other.Indirections;
-		bHasBeenPostLoaded = Other.bHasBeenPostLoaded;
-	}
-	else
-	{
-		SrcAccesses.Empty();
-		DestAccesses.Empty();
-		Indirections.Empty();
-		bHasBeenPostLoaded = false;
-		::FPropertyAccessSystem::PatchPropertyOffsets(*this);
-	}
+	CopyBatchArray = Other.CopyBatchArray;
+	SrcAccesses = Other.SrcAccesses;
+	DestAccesses = Other.DestAccesses;
+	Indirections = Other.Indirections;
+	bHasBeenPostLoaded = Other.bHasBeenPostLoaded;
 	
 	return *this;
 }
diff --git a/Engine/Source/Runtime/Engine/Public/PropertyAccess.h b/Engine/Source/Runtime/Engine/Public/PropertyAccess.h
--- a/Engine/Source/Runtime/Engine/Public/PropertyAccess.h
+++ b/Engine/Source/Runtime/Engine/Public/PropertyAccess.h
@@ -426,15 +426,12 @@ private:
 	TArray<FPropertyAccessCopyBatch> CopyBatchArray;
 	
 	// All source property accesses
-	UPROPERTY(Transient)
 	TArray<FPropertyAccessIndirectionChain> SrcAccesses;
 
 	// All destination accesses (that are copied to our instances).
-	UPROPERTY(Transient)
 	TArray<FPropertyAccessIndirectionChain> DestAccesses;
 
 	// Indirections
-	UPROPERTY(Transient)
 	TArray<FPropertyAccessIndirection> Indirections;
 	
 	// Whether this library has been post-loaded
