Adds mesh streaming reporting to niagara & cascade primitive component to ensure renderable meshes are being adequately streamed in
-reports a scaled instance at the origin of the component
-does not take into account dynamic mesh scale

diff --git a/Engine/Plugins/FX/Niagara/Source/Niagara/Private/NiagaraComponent.cpp b/Engine/Plugins/FX/Niagara/Source/Niagara/Private/NiagaraComponent.cpp
--- a/Engine/Plugins/FX/Niagara/Source/Niagara/Private/NiagaraComponent.cpp
+++ b/Engine/Plugins/FX/Niagara/Source/Niagara/Private/NiagaraComponent.cpp
@@ -1977,6 +1977,17 @@ void UNiagaraComponent::GetUsedMaterials(TArray<UMaterialInterface*>& OutMateria
 	}
 }
 
+void UNiagaraComponent::GetStreamingRenderAssetInfo(FStreamingTextureLevelContext& LevelContext, TArray<FStreamingRenderAssetPrimitiveInfo>& OutStreamingRenderAssets) const
+{
+	// UPrimitiveComponent::GetStreamingRenderAssetInfo takes care of getting the referenced textures
+	Super::GetStreamingRenderAssetInfo(LevelContext, OutStreamingRenderAssets);
+
+	if (SystemInstanceController.IsValid())
+	{
+		SystemInstanceController->GetStreamingMeshInfo(Bounds, LevelContext, OutStreamingRenderAssets);
+	}
+}
+
 bool UNiagaraComponent::ResolveOwnerAllowsScalability(bool bRegister)
 {
 	bOwnerAllowsScalabiltiy = true;
diff --git a/Engine/Plugins/FX/Niagara/Source/Niagara/Private/NiagaraMeshRendererProperties.cpp b/Engine/Plugins/FX/Niagara/Source/Niagara/Private/NiagaraMeshRendererProperties.cpp
--- a/Engine/Plugins/FX/Niagara/Source/Niagara/Private/NiagaraMeshRendererProperties.cpp
+++ b/Engine/Plugins/FX/Niagara/Source/Niagara/Private/NiagaraMeshRendererProperties.cpp
@@ -468,6 +468,25 @@ void UNiagaraMeshRendererProperties::GetUsedMaterials(const FNiagaraEmitterInsta
 	}
 }
 
+void UNiagaraMeshRendererProperties::GetStreamingMeshInfo(const FBoxSphereBounds& OwnerBounds, const FNiagaraEmitterInstance* InEmitter, TArray<FStreamingRenderAssetPrimitiveInfo>& OutStreamingRenderAssets) const
+{
+	for (const FNiagaraMeshRendererMeshProperties& MeshProperties : Meshes)
+	{
+		UStaticMesh* Mesh = MeshProperties.ResolveStaticMesh(InEmitter);
+
+		if (Mesh && Mesh->RenderResourceSupportsStreaming() && Mesh->GetRenderAssetType() == EStreamableRenderAssetType::StaticMesh)
+		{
+			const FBoxSphereBounds MeshBounds = Mesh->GetBounds();
+			const FBoxSphereBounds StreamingBounds = FBoxSphereBounds(
+				OwnerBounds.Origin + MeshBounds.Origin,
+				MeshBounds.BoxExtent * MeshProperties.Scale,
+				MeshBounds.SphereRadius * MeshProperties.Scale.GetMax());
+			const float MeshTexelFactor = MeshBounds.SphereRadius * 2.0f;
+
+			new (OutStreamingRenderAssets) FStreamingRenderAssetPrimitiveInfo(Mesh, StreamingBounds, MeshTexelFactor);
+		}
+	}
+}
 
 bool UNiagaraMeshRendererProperties::PopulateRequiredBindings(FNiagaraParameterStore& InParameterStore)
 {
diff --git a/Engine/Plugins/FX/Niagara/Source/Niagara/Private/NiagaraSystemInstanceController.cpp b/Engine/Plugins/FX/Niagara/Source/Niagara/Private/NiagaraSystemInstanceController.cpp
--- a/Engine/Plugins/FX/Niagara/Source/Niagara/Private/NiagaraSystemInstanceController.cpp
+++ b/Engine/Plugins/FX/Niagara/Source/Niagara/Private/NiagaraSystemInstanceController.cpp
@@ -155,6 +155,25 @@ void FNiagaraSystemInstanceController::GetUsedMaterials(TArray<UMaterialInterfac
 	}
 }
 
+void FNiagaraSystemInstanceController::GetStreamingMeshInfo(const FBoxSphereBounds& OwnerBounds, FStreamingTextureLevelContext& LevelContext, TArray<FStreamingRenderAssetPrimitiveInfo>& OutStreamingRenderAssets) const
+{
+	if (!SystemInstance.IsValid())
+	{
+		return;
+	}
+
+	for (TSharedRef<FNiagaraEmitterInstance, ESPMode::ThreadSafe> EmitterInst : SystemInstance->GetEmitters())
+	{
+		if (UNiagaraEmitter* Emitter = EmitterInst->GetCachedEmitter())
+		{
+			Emitter->ForEachEnabledRenderer([&](UNiagaraRendererProperties* Properties)
+			{
+				Properties->GetStreamingMeshInfo(OwnerBounds, &EmitterInst.Get(), OutStreamingRenderAssets);
+			});
+		}
+	}
+}
+
 UMaterialInterface* FNiagaraSystemInstanceController::GetMaterialOverride(const UNiagaraRendererProperties* InProps, int32 InMaterialSubIndex) const
 {
 	for (const FMaterialOverride& Override : EmitterMaterials)
diff --git a/Engine/Plugins/FX/Niagara/Source/Niagara/Public/NiagaraComponent.h b/Engine/Plugins/FX/Niagara/Source/Niagara/Public/NiagaraComponent.h
--- a/Engine/Plugins/FX/Niagara/Source/Niagara/Public/NiagaraComponent.h
+++ b/Engine/Plugins/FX/Niagara/Source/Niagara/Public/NiagaraComponent.h
@@ -216,7 +216,7 @@ public:
 	virtual FBoxSphereBounds CalcBounds(const FTransform& LocalToWorld) const override;
 	virtual FPrimitiveSceneProxy* CreateSceneProxy() override;
 	virtual void GetUsedMaterials(TArray<UMaterialInterface*>& OutMaterials, bool bGetDebugMaterials = false) const override;
-
+	virtual void GetStreamingRenderAssetInfo(FStreamingTextureLevelContext& LevelContext, TArray<FStreamingRenderAssetPrimitiveInfo>& OutStreamingRenderAssets) const override;
 	virtual void OnAttachmentChanged() override;
 	//~ End UPrimitiveComponent Interface
 
diff --git a/Engine/Plugins/FX/Niagara/Source/Niagara/Public/NiagaraMeshRendererProperties.h b/Engine/Plugins/FX/Niagara/Source/Niagara/Public/NiagaraMeshRendererProperties.h
--- a/Engine/Plugins/FX/Niagara/Source/Niagara/Public/NiagaraMeshRendererProperties.h
+++ b/Engine/Plugins/FX/Niagara/Source/Niagara/Public/NiagaraMeshRendererProperties.h
@@ -177,6 +177,7 @@ public:
 	virtual FNiagaraRenderer* CreateEmitterRenderer(ERHIFeatureLevel::Type FeatureLevel, const FNiagaraEmitterInstance* Emitter, const FNiagaraSystemInstanceController& InController) override;
 	virtual class FNiagaraBoundsCalculator* CreateBoundsCalculator() override;
 	virtual void GetUsedMaterials(const FNiagaraEmitterInstance* InEmitter, TArray<UMaterialInterface*>& OutMaterials) const override;
+	virtual void GetStreamingMeshInfo(const FBoxSphereBounds& OwnerBounds, const FNiagaraEmitterInstance* InEmitter, TArray<FStreamingRenderAssetPrimitiveInfo>& OutStreamingRenderAssets) const override;
 	virtual bool IsSimTargetSupported(ENiagaraSimTarget InSimTarget) const override { return true; };
 	virtual bool PopulateRequiredBindings(FNiagaraParameterStore& InParameterStore) override;
 #if WITH_EDITORONLY_DATA
diff --git a/Engine/Plugins/FX/Niagara/Source/Niagara/Public/NiagaraRendererProperties.h b/Engine/Plugins/FX/Niagara/Source/Niagara/Public/NiagaraRendererProperties.h
--- a/Engine/Plugins/FX/Niagara/Source/Niagara/Public/NiagaraRendererProperties.h
+++ b/Engine/Plugins/FX/Niagara/Source/Niagara/Public/NiagaraRendererProperties.h
@@ -199,6 +199,8 @@ public:
 	virtual class FNiagaraBoundsCalculator* CreateBoundsCalculator() PURE_VIRTUAL(UNiagaraRendererProperties::CreateBoundsCalculator, return nullptr;);
 	virtual void GetUsedMaterials(const FNiagaraEmitterInstance* InEmitter, TArray<UMaterialInterface*>& OutMaterials) const PURE_VIRTUAL(UNiagaraRendererProperties::GetUsedMaterials,);
 
+	virtual void GetStreamingMeshInfo(const FBoxSphereBounds& OwnerBounds, const FNiagaraEmitterInstance* InEmitter, TArray<FStreamingRenderAssetPrimitiveInfo>& OutStreamingRenderAssets) const {}
+
 	virtual bool IsSimTargetSupported(ENiagaraSimTarget InSimTarget) const { return false; };
 
 	const TArray<const FNiagaraVariableAttributeBinding*>& GetAttributeBindings() const { return AttributeBindings; }
diff --git a/Engine/Plugins/FX/Niagara/Source/Niagara/Public/NiagaraSystemInstanceController.h b/Engine/Plugins/FX/Niagara/Source/Niagara/Public/NiagaraSystemInstanceController.h
--- a/Engine/Plugins/FX/Niagara/Source/Niagara/Public/NiagaraSystemInstanceController.h
+++ b/Engine/Plugins/FX/Niagara/Source/Niagara/Public/NiagaraSystemInstanceController.h
@@ -82,6 +82,8 @@ public:
 	UMaterialInterface* GetMaterialOverride(const UNiagaraRendererProperties* InProps, int32 InMaterialSubIndex) const;
 	void SetOnMaterialsUpdated(const FOnMaterialsUpdated& Delegate) { OnMaterialsUpdatedDelegate = Delegate; }
 
+	void GetStreamingMeshInfo(const FBoxSphereBounds& OwnerBounds, FStreamingTextureLevelContext& LevelContext, TArray<FStreamingRenderAssetPrimitiveInfo>& OutStreamingRenderAssets) const;
+
 	bool GetParticleValueVec3_DebugOnly(TArray<FVector>& OutValues, FName EmitterName, FName ValueName) const;
 	bool GetParticleValues_DebugOnly(TArray<float>& OutValues, FName EmitterName, FName ValueName) const;
 
diff --git a/Engine/Source/Runtime/Engine/Classes/Particles/ParticleLODLevel.h b/Engine/Source/Runtime/Engine/Classes/Particles/ParticleLODLevel.h
--- a/Engine/Source/Runtime/Engine/Classes/Particles/ParticleLODLevel.h
+++ b/Engine/Source/Runtime/Engine/Classes/Particles/ParticleLODLevel.h
@@ -142,6 +142,13 @@ class UParticleLODLevel : public UObject
  	 * @param EmitterMaterials - the material slot materials.
 	 */
 	void GetUsedMaterials( TArray<UMaterialInterface*>& OutMaterials, const TArray<struct FNamedEmitterMaterial>& NamedMaterialSlots, const TArray<UMaterialInterface*>& EmitterMaterials) const;
+
+	/**
+	 * Appends information on streaming meshes
+	 * @param Bounds - the bounds of the parent particle system component
+	 * @param OutStreamingRenderAssets - where the streaming mesh information will be added
+	 */
+	void GetStreamingMeshInfo(const FBoxSphereBounds& Bounds, TArray<FStreamingRenderAssetPrimitiveInfo>& OutStreamingRenderAssets) const;
 };
 
 
diff --git a/Engine/Source/Runtime/Engine/Private/Particles/ParticleComponents.cpp b/Engine/Source/Runtime/Engine/Private/Particles/ParticleComponents.cpp
--- a/Engine/Source/Runtime/Engine/Private/Particles/ParticleComponents.cpp
+++ b/Engine/Source/Runtime/Engine/Private/Particles/ParticleComponents.cpp
@@ -7739,6 +7739,29 @@ void UParticleLODLevel::GetUsedMaterials(TArray<UMaterialInterface*>& OutMateria
 	}
 }
 
+void UParticleLODLevel::GetStreamingMeshInfo(const FBoxSphereBounds& Bounds, TArray<FStreamingRenderAssetPrimitiveInfo>& OutStreamingRenderAssets) const
+{
+	if (bEnabled)
+	{
+		const UParticleModuleTypeDataMesh* MeshTypeData = Cast<UParticleModuleTypeDataMesh>(TypeDataModule);
+
+		if (UStaticMesh* Mesh = MeshTypeData ? MeshTypeData->Mesh : nullptr)
+		{
+			if (Mesh->RenderResourceSupportsStreaming() && Mesh->GetRenderAssetType() == EStreamableRenderAssetType::StaticMesh)
+			{
+				const FBoxSphereBounds MeshBounds = Mesh->GetBounds();
+				const FBoxSphereBounds StreamingBounds = FBoxSphereBounds(
+					Bounds.Origin + MeshBounds.Origin,
+					MeshBounds.BoxExtent * MeshTypeData->LODSizeScale,
+					MeshBounds.SphereRadius * MeshTypeData->LODSizeScale);
+				const float MeshTexelFactor = MeshBounds.SphereRadius * 2.0f;
+
+				new (OutStreamingRenderAssets) FStreamingRenderAssetPrimitiveInfo(Mesh, StreamingBounds, MeshTexelFactor);
+			}
+		}
+	}
+}
+
 void UParticleSystemComponent::GetUsedMaterials( TArray<UMaterialInterface*>& OutMaterials, bool bGetDebugMaterials ) const
 {
 	if (Template)
@@ -7814,6 +7837,8 @@ void UParticleSystemComponent::GetStreamingRenderAssetInfo(FStreamingTextureLeve
 				LODLevelMaterials.Reset();
 				LOD->GetUsedMaterials(LODLevelMaterials, Template->NamedMaterialSlots, EmitterMaterials);
 				AddMaterials(MaterialWithScales, LODLevelMaterials, (float)FMath::Max<int32>(LOD->RequiredModule->SubImages_Horizontal, LOD->RequiredModule->SubImages_Vertical));
+
+				LOD->GetStreamingMeshInfo(Bounds, OutStreamingRenderAssets);
 			}
 		}
 
