Support runtime change of the ViewDistances with PerQualityProperty (MinLOD).
Change the way we support Perf mode on PC when QualityMinLOD is enable. Dont use LodBiasModifer, but instead override the global per quality cvar.
Allow Viewdistances setrting to be changed at runtime on PC

diff --git a/Engine/Source/Runtime/Engine/Classes/Engine/SkeletalMesh.h b/Engine/Source/Runtime/Engine/Classes/Engine/SkeletalMesh.h
--- a/Engine/Source/Runtime/Engine/Classes/Engine/SkeletalMesh.h
+++ b/Engine/Source/Runtime/Engine/Classes/Engine/SkeletalMesh.h
@@ -1248,7 +1248,7 @@ public:
 	static void OnLodStrippingQualityLevelChanged(IConsoleVariable* Variable);
 
 	/*Choose either PerPlatform or PerQuality override. Note: Enable PerQuality override in the Project Settings/ General Settings/ UseSkeletalMeshMinLODPerQualityLevels*/
-	int32 GetMinLodIdx() const;
+	int32 GetMinLodIdx(bool bForceLowestLODIdx = false) const;
 	int32 GetDefaultMinLod() const;
 	void SetMinLodIdx(int32 InMinLOD);
 
diff --git a/Engine/Source/Runtime/Engine/Classes/Engine/StaticMesh.h b/Engine/Source/Runtime/Engine/Classes/Engine/StaticMesh.h
--- a/Engine/Source/Runtime/Engine/Classes/Engine/StaticMesh.h
+++ b/Engine/Source/Runtime/Engine/Classes/Engine/StaticMesh.h
@@ -781,7 +781,7 @@ public:
 	}
 
 	/*Choose either PerPlatform or PerQuality override. Note: Enable PerQuality override in the Project Settings/ General Settings/ UseStaticMeshMinLODPerQualityLevels*/
-	ENGINE_API int32 GetMinLODIdx() const;
+	ENGINE_API int32 GetMinLODIdx(bool bForceLowestLODIdx = false) const;
 	ENGINE_API int32 GetDefaultMinLOD() const;
 	ENGINE_API void SetMinLODIdx(int32 InMinLOD);
 
diff --git a/Engine/Source/Runtime/Engine/Private/PerQualityLevelProperties.cpp b/Engine/Source/Runtime/Engine/Private/PerQualityLevelProperties.cpp
--- a/Engine/Source/Runtime/Engine/Private/PerQualityLevelProperties.cpp
+++ b/Engine/Source/Runtime/Engine/Private/PerQualityLevelProperties.cpp
@@ -259,6 +259,22 @@ bool FPerQualityLevelProperty<_StructType, _ValueType, _BasePropertyName>::IsQua
 }
 #endif
 
+template<typename _StructType, typename _ValueType, EName _BasePropertyName>
+_ValueType FPerQualityLevelProperty<_StructType, _ValueType, _BasePropertyName>::GetLowestValue() const
+{
+	const _StructType* This = StaticCast<const _StructType*>(this);
+	_ValueType Value = This->Default;
+
+	for (const TPair<int32, int32>& Pair : This->PerQuality)
+	{
+		if (Pair.Value < Value)
+		{
+			Value = Pair.Value;
+		}
+	}
+	return Value;
+}
+
 /** Serializer to cook out the most appropriate platform override */
 template<typename _StructType, typename _ValueType, EName _BasePropertyName>
 ENGINE_API FArchive& operator<<(FArchive& Ar, FPerQualityLevelProperty<_StructType, _ValueType, _BasePropertyName>& Property)
@@ -315,6 +331,7 @@ template FSupportedQualityLevelArray FPerQualityLevelProperty<FPerQualityLevelIn
 template void FPerQualityLevelProperty<FPerQualityLevelInt, int32, NAME_IntProperty>::StripQualtiyLevelForCooking(const TCHAR* InPlatformName);
 template bool FPerQualityLevelProperty<FPerQualityLevelInt, int32, NAME_IntProperty>::IsQualityLevelValid(int32 QualityLevel) const;
 #endif
+template int32 FPerQualityLevelProperty<FPerQualityLevelInt, int32, NAME_IntProperty>::GetLowestValue() const;
 
 FString FPerQualityLevelInt::ToString() const
 {
diff --git a/Engine/Source/Runtime/Engine/Private/SkeletalMesh.cpp b/Engine/Source/Runtime/Engine/Private/SkeletalMesh.cpp
--- a/Engine/Source/Runtime/Engine/Private/SkeletalMesh.cpp
+++ b/Engine/Source/Runtime/Engine/Private/SkeletalMesh.cpp
@@ -1020,7 +1020,7 @@ void USkeletalMesh::InitResources()
 
 		{
 			const int32 NumLODs = SkelMeshRenderData->LODRenderData.Num();
-			const int32 MinFirstLOD = GetMinLodIdx();
+			const int32 MinFirstLOD = GetMinLodIdx(true);
 
 			CachedSRRState.NumNonStreamingLODs = SkelMeshRenderData->NumInlinedLODs;
 			CachedSRRState.NumNonOptionalLODs = SkelMeshRenderData->NumNonOptionalLODs;
@@ -5740,11 +5740,11 @@ TArray<FString> USkeletalMesh::K2_GetAllMorphTargetNames() const
 	return Names;
 }
 
-int32 USkeletalMesh::GetMinLodIdx() const
+int32 USkeletalMesh::GetMinLodIdx(bool bForceLowestLODIdx) const
 {
 	if (IsMinLodQualityLevelEnable())
 	{
-		return GetQualityLevelMinLod().GetValue(GSkeletalMeshMinLodQualityLevel);
+		return bForceLowestLODIdx ? GetQualityLevelMinLod().GetLowestValue() : GetQualityLevelMinLod().GetValue(GSkeletalMeshMinLodQualityLevel);
 	}
 	else
 	{
@@ -5782,7 +5782,7 @@ bool USkeletalMesh::IsMinLodQualityLevelEnable() const
 }
 
 void USkeletalMesh::OnLodStrippingQualityLevelChanged(IConsoleVariable* Variable) {
-#ifdef WITH_EDITOR
+#if WITH_EDITOR || PLATFORM_DESKTOP
 	if (GEngine && GEngine->UseStaticMeshMinLODPerQualityLevels)
 	{
 		for (TObjectIterator<USkeletalMesh> It; It; ++It)
diff --git a/Engine/Source/Runtime/Engine/Private/SkeletalMeshRenderData.cpp b/Engine/Source/Runtime/Engine/Private/SkeletalMeshRenderData.cpp
--- a/Engine/Source/Runtime/Engine/Private/SkeletalMeshRenderData.cpp
+++ b/Engine/Source/Runtime/Engine/Private/SkeletalMeshRenderData.cpp
@@ -678,16 +678,14 @@ void FSkeletalMeshRenderData::Serialize(FArchive& Ar, USkeletalMesh* Owner)
 				&& GStripSkeletalMeshLodsDuringCooking != 0
 				&& CVarSkeletalMeshKeepMobileMinLODSettingOnDesktop.GetValueOnAnyThread() != 0)
 			{
-				if (Owner->IsMinLodQualityLevelEnable())
-				{
-					MinMobileLODIdx = Owner->GetQualityLevelMinLod().GetValueForQualityLevel(0/*Low*/) - FSkeletalMeshLODRenderData::GetPlatformMinLODIdx(Ar.CookingTarget(), Owner);
-				}
-				else
+				// Serialize 0 value when per quality level properties are used
+				if (!Owner->IsMinLodQualityLevelEnable())
 				{
 					MinMobileLODIdx = Owner->GetMinLod().GetValueForPlatform(TEXT("Mobile")) - Owner->GetMinLod().GetValueForPlatform(TEXT("Desktop"));
-				}
-				MinMobileLODIdx = FMath::Clamp(MinMobileLODIdx, 0, 255); // Will be cast to uint8 when applying LOD bias. Also, make sure it's not < 0,
-																		 // which can happen if the desktop min LOD is higher than the mobile setting
+					// Will be cast to uint8 when applying LOD bias. Also, make sure it's not < 0,
+					// which can happen if the desktop min LOD is higher than the mobile setting
+					MinMobileLODIdx = FMath::Clamp(MinMobileLODIdx, 0, 255); 
+				}													
 			}
 			else
 			{
diff --git a/Engine/Source/Runtime/Engine/Private/StaticMesh.cpp b/Engine/Source/Runtime/Engine/Private/StaticMesh.cpp
--- a/Engine/Source/Runtime/Engine/Private/StaticMesh.cpp
+++ b/Engine/Source/Runtime/Engine/Private/StaticMesh.cpp
@@ -1636,16 +1636,14 @@ void FStaticMeshRenderData::Serialize(FArchive& Ar, UStaticMesh* Owner, bool bCo
 				&& CVarStripMinLodDataDuringCooking.GetValueOnAnyThread() != 0
 				&& CVarStaticMeshKeepMobileMinLODSettingOnDesktop.GetValueOnAnyThread() != 0)
 			{
-				if (Owner->IsMinLodQualityLevelEnable())
-				{
-					MinMobileLODIdx = Owner->GetQualityLevelMinLOD().GetValueForQualityLevel(0/*Low*/) - FStaticMeshLODResources::GetPlatformMinLODIdx(Ar.CookingTarget(), Owner);
-				}
-				else
+				// Serialize 0 value when per quality level properties are used
+				if (!Owner->IsMinLodQualityLevelEnable())
 				{
 					MinMobileLODIdx = Owner->GetMinLOD().GetValueForPlatform(TEXT("Mobile")) - FStaticMeshLODResources::GetPlatformMinLODIdx(Ar.CookingTarget(), Owner);
+					// Will be cast to uint8 when applying LOD bias. Also, make sure it's not < 0,
+					// which can happen if the desktop min LOD is higher than the mobile setting
+					MinMobileLODIdx = FMath::Clamp(MinMobileLODIdx, 0, 255);
 				}
-				MinMobileLODIdx = FMath::Clamp(MinMobileLODIdx, 0, 255); // Will be cast to uint8 when applying LOD bias. Also, make sure it's not < 0,
-																		 // which can happen if the desktop min LOD is higher than the mobile setting
 			}
 			else
 			{
@@ -3220,7 +3218,7 @@ void UStaticMesh::InitResources()
 	{
 		{
 			const int32 NumLODs = GetNumLODs();
-			const int32 MinFirstLOD = GetMinLODIdx();
+			const int32 MinFirstLOD = GetMinLODIdx(true);
 
 			CachedSRRState.NumNonStreamingLODs = GetRenderData()->NumInlinedLODs;
 			CachedSRRState.NumNonOptionalLODs = GetRenderData()->GetNumNonOptionalLODs();
@@ -7799,11 +7797,19 @@ ENGINE_API int32 UStaticMesh::GetDefaultMinLOD() const
 	}
 }
 
-ENGINE_API int32 UStaticMesh::GetMinLODIdx() const
+ENGINE_API int32 UStaticMesh::GetMinLODIdx(bool bForceLowestLODIdx) const
 {
 	if (IsMinLodQualityLevelEnable())
 	{
-		return GetQualityLevelMinLOD().GetValue(GMinLodQualityLevel);
+		int32 CurrentMinLodQualityLevel = GMinLodQualityLevel;
+#if PLATFORM_DESKTOP
+		extern int32 GUseMobileLODBiasOnDesktopES31;
+		if (GUseMobileLODBiasOnDesktopES31 != 0 && GMaxRHIFeatureLevel == ERHIFeatureLevel::ES3_1)
+		{
+			CurrentMinLodQualityLevel = (int32)QualityLevelProperty::EQualityLevels::Low;
+		}
+#endif
+		return bForceLowestLODIdx ? GetQualityLevelMinLOD().GetLowestValue() : GetQualityLevelMinLOD().GetValue(CurrentMinLodQualityLevel);
 	}
 	else
 	{
@@ -7830,7 +7836,7 @@ bool UStaticMesh::IsMinLodQualityLevelEnable() const
 }
 
 void UStaticMesh::OnLodStrippingQualityLevelChanged(IConsoleVariable* Variable){
-#ifdef WITH_EDITOR
+#if WITH_EDITOR || PLATFORM_DESKTOP
 	if (GEngine && GEngine->UseStaticMeshMinLODPerQualityLevels)
 	{
 		for (TObjectIterator<UStaticMesh> It; It; ++It)
diff --git a/Engine/Source/Runtime/Engine/Private/StaticMeshRender.cpp b/Engine/Source/Runtime/Engine/Private/StaticMeshRender.cpp
--- a/Engine/Source/Runtime/Engine/Private/StaticMeshRender.cpp
+++ b/Engine/Source/Runtime/Engine/Private/StaticMeshRender.cpp
@@ -251,10 +251,7 @@ FStaticMeshSceneProxy::FStaticMeshSceneProxy(UStaticMeshComponent* InComponent,
 	extern int32 GUseMobileLODBiasOnDesktopES31;
 	if (GUseMobileLODBiasOnDesktopES31 != 0 && FeatureLevel == ERHIFeatureLevel::ES3_1)
 	{
-		int32 BiasModifier = InComponent->GetStaticMesh()->GetRenderData()->LODBiasModifier;
-		int32 MinLODIdx = FMath::Max(StaticMesh->GetMinLODIdx(), 0);
-		BiasModifier = FMath::Max(BiasModifier - MinLODIdx, 0);
-		EffectiveMinLOD += BiasModifier;
+		EffectiveMinLOD += InComponent->GetStaticMesh()->GetRenderData()->LODBiasModifier;
 	}
 #endif
 
diff --git a/Engine/Source/Runtime/Engine/Private/Streaming/StreamingTexture.cpp b/Engine/Source/Runtime/Engine/Private/Streaming/StreamingTexture.cpp
--- a/Engine/Source/Runtime/Engine/Private/Streaming/StreamingTexture.cpp
+++ b/Engine/Source/Runtime/Engine/Private/Streaming/StreamingTexture.cpp
@@ -211,21 +211,7 @@ void FStreamingRenderAsset::UpdateDynamicData(const int32* NumStreamedMips, int3
 #if PLATFORM_DESKTOP
 			if (GUseMobileLODBiasOnDesktopES31 != 0 && GMaxRHIFeatureLevel == ERHIFeatureLevel::ES3_1)
 			{
-				if (RenderAssetType == EStreamableRenderAssetType::StaticMesh)
-				{
-					const UStaticMesh* StaticMesh = Cast<UStaticMesh>(RenderAsset);
-					if (StaticMesh)
-					{
-						int32 BiasModifier = ResourceState.LODBiasModifier;
-						int32 MinLODIdx = FMath::Max(StaticMesh->GetMinLODIdx(), 0);
-						BiasModifier = FMath::Max(BiasModifier - MinLODIdx, 0);
-						LODBias += BiasModifier;
-					}
-				}
-				else
-				{
-					LODBias += ResourceState.LODBiasModifier;
-				}
+				LODBias += ResourceState.LODBiasModifier;
 			}
 #endif
 		}
diff --git a/Engine/Source/Runtime/Engine/Public/PerQualityLevelProperties.h b/Engine/Source/Runtime/Engine/Public/PerQualityLevelProperties.h
--- a/Engine/Source/Runtime/Engine/Public/PerQualityLevelProperties.h
+++ b/Engine/Source/Runtime/Engine/Public/PerQualityLevelProperties.h
@@ -101,6 +101,8 @@ struct ENGINE_API FPerQualityLevelProperty
 		return GetValueForQualityLevel(QualityLevel);
 	}
 
+	_ValueType GetLowestValue() const;
+
 	/* Load old properties that have been converted to FPerQualityLevel */
 	bool SerializeFromMismatchedTag(const FPropertyTag& Tag, FArchive& Ar)
 	{
