Fix for missing forward declaration with ParticleLODLevel.h

diff --git a/Engine/Source/Runtime/Engine/Classes/Particles/ParticleLODLevel.h b/Engine/Source/Runtime/Engine/Classes/Particles/ParticleLODLevel.h
--- a/Engine/Source/Runtime/Engine/Classes/Particles/ParticleLODLevel.h
+++ b/Engine/Source/Runtime/Engine/Classes/Particles/ParticleLODLevel.h
@@ -15,6 +15,7 @@
 class UInterpCurveEdSetup;
 class UParticleModule;
 class UMaterialInterface;
+struct FStreamingRenderAssetPrimitiveInfo;
 
 UCLASS(collapsecategories, hidecategories=Object, editinlinenew, MinimalAPI)
 class UParticleLODLevel : public UObject
