Removing another place where USoundWave manually deletes data shared with active proxies when it gets GC'd (old code).

diff --git a/Engine/Source/Runtime/Engine/Private/SoundWave.cpp b/Engine/Source/Runtime/Engine/Private/SoundWave.cpp
--- a/Engine/Source/Runtime/Engine/Private/SoundWave.cpp
+++ b/Engine/Source/Runtime/Engine/Private/SoundWave.cpp
@@ -2527,9 +2527,6 @@ void USoundWave::FreeResources(bool bStopSoundsUsingThisResource)
 		RawPCMData = nullptr;
 	}
 
-	// Remove the compressed copy of the data
-	RemoveAudioResource();
-
 	// Stat housekeeping
 	DEC_DWORD_STAT_BY(STAT_AudioMemorySize, TrackedMemoryUsage);
 	DEC_DWORD_STAT_BY(STAT_AudioMemory, TrackedMemoryUsage);
