Making an au.debug version of stat streaming

diff --git a/Engine/Source/Runtime/Engine/Public/Audio/AudioDebug.h b/Engine/Source/Runtime/Engine/Public/Audio/AudioDebug.h
--- a/Engine/Source/Runtime/Engine/Public/Audio/AudioDebug.h
+++ b/Engine/Source/Runtime/Engine/Public/Audio/AudioDebug.h
@@ -146,8 +146,8 @@ namespace Audio
 
 		void SetStats(FDeviceId DeviceId, const TSet<FName>& StatsToSet);
 
-		static bool ToggleStats(UWorld* World, const uint32 StatToToggle);
-		void ToggleStats(FDeviceId DeviceId, const uint32 StatsToToggle);
+		static bool ToggleStats(UWorld* World, const TSet<FName>& StatToToggle);
+		void ToggleStats(FDeviceId DeviceId, const TSet<FName>& StatsToToggle);
 
 		FDelegateHandle WorldRegisteredWithDeviceHandle;
 
