Invalidate shader cached data when strata update links.
Thanks Jason for the fix!

diff --git a/Engine/Source/Runtime/Engine/Private/Materials/Material.cpp b/Engine/Source/Runtime/Engine/Private/Materials/Material.cpp
--- a/Engine/Source/Runtime/Engine/Private/Materials/Material.cpp
+++ b/Engine/Source/Runtime/Engine/Private/Materials/Material.cpp
@@ -3242,6 +3242,32 @@ void UMaterial::ConvertMaterialToStrataMaterial()
 
 	// For rebuild the shading mode since we have change it
 	RebuildShadingModelField();
+
+	if (bInvalidateShader)
+	{
+		// We might have moved connections above so update the CachedExpressionData from the EditorOnly connection data (ground truth).
+		if (!bUseMaterialAttributes)
+		{
+			// Reset the connection bitmask so we can set only the valid ones below.
+			CachedExpressionData->PropertyConnectedBitmask = 0;
+
+			for (int32 PropertyIndex = 0; PropertyIndex < MP_MAX; ++PropertyIndex)
+			{
+				const EMaterialProperty Property = (EMaterialProperty)PropertyIndex;
+				const FExpressionInput* Input = GetExpressionInputForProperty(Property);
+				if (Input && Input->IsConnected())
+				{
+					CachedExpressionData->SetPropertyConnected(Property);
+				}
+			}
+		}
+
+		CachedExpressionData->Validate();
+
+		// Make sure our cached data is consistent
+		ensureMsgf(!(GetEditorOnlyData()->Refraction.IsConnected() && !GetCachedExpressionData().IsPropertyConnected(MP_Refraction)),
+			TEXT("GetCachedExpressionData() says refraction isn't connected, but GetEditorOnlyData() says it is"));
+	}
 #endif
 }
 
