Fix incorrect formats shown in TC descriptions in Editor GUI

diff --git a/Engine/Source/Runtime/Engine/Classes/Engine/TextureDefines.h b/Engine/Source/Runtime/Engine/Classes/Engine/TextureDefines.h
--- a/Engine/Source/Runtime/Engine/Classes/Engine/TextureDefines.h
+++ b/Engine/Source/Runtime/Engine/Classes/Engine/TextureDefines.h
@@ -329,13 +329,13 @@ enum TextureCompressionSettings
 	TC_Default					UMETA(DisplayName = "Default (DXT1/5, BC1/3 on DX11)"),
 	TC_Normalmap				UMETA(DisplayName = "Normalmap (DXT5, BC5 on DX11)"),
 	TC_Masks					UMETA(DisplayName = "Masks (no sRGB)"),
-	TC_Grayscale				UMETA(DisplayName = "Grayscale (R8, RGB8 sRGB)"),
-	TC_Displacementmap			UMETA(DisplayName = "Displacementmap (8/16bit)"),
+	TC_Grayscale				UMETA(DisplayName = "Grayscale (G8/16, RGB8 sRGB)"),
+	TC_Displacementmap			UMETA(DisplayName = "Displacementmap (G8/16)"),
 	TC_VectorDisplacementmap	UMETA(DisplayName = "VectorDisplacementmap (RGBA8)"),
 	TC_HDR						UMETA(DisplayName = "HDR (RGBA16F, no sRGB)"),
 	TC_EditorIcon				UMETA(DisplayName = "UserInterface2D (RGBA)"),
 	TC_Alpha					UMETA(DisplayName = "Alpha (no sRGB, BC4 on DX11)"),
-	TC_DistanceFieldFont		UMETA(DisplayName = "DistanceFieldFont (R8)"),
+	TC_DistanceFieldFont		UMETA(DisplayName = "DistanceFieldFont (G8)"),
 	TC_HDR_Compressed			UMETA(DisplayName = "HDRCompressed (RGB, BC6H, DX11)"),
 	TC_BC7						UMETA(DisplayName = "BC7 (DX11, optional A)"),
 	TC_HalfFloat				UMETA(DisplayName = "Half Float (R16F)"),
diff --git a/Engine/Source/Runtime/Engine/Private/Texture.cpp b/Engine/Source/Runtime/Engine/Private/Texture.cpp
--- a/Engine/Source/Runtime/Engine/Private/Texture.cpp
+++ b/Engine/Source/Runtime/Engine/Private/Texture.cpp
@@ -3001,6 +3001,7 @@ FName GetDefaultTextureFormatName( const ITargetPlatform* TargetPlatform, const
 
 		if (Texture->HasHDRSource(LayerIndex))
 		{
+			// @todo Oodle : this is a very poor selection from SourceFormat
 			// R16F and R32F is available but not used here even if their TC_ would have chosen them!
 			TextureFormatName = NameRGBA16F;
 		}
@@ -3043,23 +3044,15 @@ FName GetDefaultTextureFormatName( const ITargetPlatform* TargetPlatform, const
 	{
 		TextureFormatName = bUseDXT5NormalMap ? NameDXT5n : NameBC5;
 	}
-	else if (FormatSettings.CompressionSettings == TC_Displacementmap)
-	{
-		if (SourceFormat == TSF_G16)
-		{
-			TextureFormatName = NameG16;
-		}
-		else
-		{
-			TextureFormatName = NameG8;
-		}
-	}
 	else if (FormatSettings.CompressionSettings == TC_VectorDisplacementmap)
 	{
 		TextureFormatName = NameBGRA8;
 	}
-	else if (FormatSettings.CompressionSettings == TC_Grayscale)
+	else if (FormatSettings.CompressionSettings == TC_Grayscale || 
+		FormatSettings.CompressionSettings == TC_Displacementmap)
 	{
+		// @todo Oodle : this is a very poor selection from SourceFormat
+		// eg. doesn't use G16 if source is RGBA16 or float
 		if (SourceFormat == TSF_G16)
 		{
 			TextureFormatName = NameG16;
