Removing stray pragma disable optimization.

diff --git a/Engine/Source/Runtime/Engine/Private/Materials/MaterialInstance.cpp b/Engine/Source/Runtime/Engine/Private/Materials/MaterialInstance.cpp
--- a/Engine/Source/Runtime/Engine/Private/Materials/MaterialInstance.cpp
+++ b/Engine/Source/Runtime/Engine/Private/Materials/MaterialInstance.cpp
@@ -2409,7 +2409,6 @@ void UMaterialInstance::ClearAllCachedCookedPlatformData()
 
 #endif
 
-PRAGMA_DISABLE_OPTIMIZATION
 void UMaterialInstance::Serialize(FArchive& Ar)
 {
 	LLM_SCOPE(ELLMTag::MaterialInstance);
@@ -2628,7 +2627,6 @@ void UMaterialInstance::Serialize(FArchive& Ar)
 		ValidateTextureOverrides(GMaxRHIFeatureLevel);
 	}
 }
-PRAGMA_ENABLE_OPTIMIZATION
 
 void UMaterialInstance::PostLoad()
 {
