Use UWorld's RemoveNetworkActor instead of directly calling the NetDriver's during seamless travel tick, in case there is no NetDriver

diff --git a/Engine/Source/Runtime/Engine/Private/World.cpp b/Engine/Source/Runtime/Engine/Private/World.cpp
--- a/Engine/Source/Runtime/Engine/Private/World.cpp
+++ b/Engine/Source/Runtime/Engine/Private/World.cpp
@@ -7296,7 +7296,7 @@ UWorld* FSeamlessTravelHandler::Tick()
 				if (WorldSettings)
 				{
 					WorldSettings->SetReplicates(false);
-					NetDriver->RemoveNetworkActor(WorldSettings);
+					LoadedWorld->RemoveNetworkActor(WorldSettings);
 				}
 			}
 
