Guard post load fixup for EditorOnlyData to avoid crashes when the data can be nulled out.

- This scenario specifically happens when we are nulling out an object in a blueprint so EditorOnlyData on the material can be null.
- We really should only be doing serialization fixup on load.
- Fix another error where functions that were overridden by return type **were not** virtual and marked as `overidden`.

diff --git a/Engine/Source/Runtime/Engine/Classes/Materials/MaterialInstance.h b/Engine/Source/Runtime/Engine/Classes/Materials/MaterialInstance.h
--- a/Engine/Source/Runtime/Engine/Classes/Materials/MaterialInstance.h
+++ b/Engine/Source/Runtime/Engine/Classes/Materials/MaterialInstance.h
@@ -471,8 +471,8 @@ class UMaterialInstance : public UMaterialInterface
 #if WITH_EDITORONLY_DATA
 	ENGINE_API virtual const UClass* GetEditorOnlyDataClass() const override { return UMaterialInstanceEditorOnlyData::StaticClass(); }
 
-	UMaterialInstanceEditorOnlyData* GetEditorOnlyData() { return CastChecked<UMaterialInstanceEditorOnlyData>(Super::GetEditorOnlyData()); }
-	const UMaterialInstanceEditorOnlyData* GetEditorOnlyData() const { return CastChecked<UMaterialInstanceEditorOnlyData>(Super::GetEditorOnlyData()); }
+	virtual UMaterialInstanceEditorOnlyData* GetEditorOnlyData() override { return CastChecked<UMaterialInstanceEditorOnlyData>(Super::GetEditorOnlyData()); }
+	virtual const UMaterialInstanceEditorOnlyData* GetEditorOnlyData() const override { return CastChecked<UMaterialInstanceEditorOnlyData>(Super::GetEditorOnlyData()); }
 #endif // WITH_EDITORONLY_DATA
 
 	/** Physical material to use for this graphics material. Used for sounds, effects etc.*/
diff --git a/Engine/Source/Runtime/Engine/Classes/Materials/MaterialInterface.h b/Engine/Source/Runtime/Engine/Classes/Materials/MaterialInterface.h
--- a/Engine/Source/Runtime/Engine/Classes/Materials/MaterialInterface.h
+++ b/Engine/Source/Runtime/Engine/Classes/Materials/MaterialInterface.h
@@ -228,8 +228,9 @@ protected:
 	ENGINE_API UMaterialInterfaceEditorOnlyData* CreateEditorOnlyData();
 
 public:
-	UMaterialInterfaceEditorOnlyData* GetEditorOnlyData() { return EditorOnlyData; }
-	const UMaterialInterfaceEditorOnlyData* GetEditorOnlyData() const { return EditorOnlyData; }
+	virtual UMaterialInterfaceEditorOnlyData* GetEditorOnlyData() { return EditorOnlyData; }
+	virtual const UMaterialInterfaceEditorOnlyData* GetEditorOnlyData() const { return EditorOnlyData; }
+	bool IsEditorOnlyDataValid() const { return EditorOnlyData != nullptr; }
 #endif // WITH_EDITORONLY_DATA
 
 public:
diff --git a/Engine/Source/Runtime/Engine/Private/Materials/MaterialInstance.cpp b/Engine/Source/Runtime/Engine/Private/Materials/MaterialInstance.cpp
--- a/Engine/Source/Runtime/Engine/Private/Materials/MaterialInstance.cpp
+++ b/Engine/Source/Runtime/Engine/Private/Materials/MaterialInstance.cpp
@@ -2427,7 +2427,6 @@ void UMaterialInstance::Serialize(FArchive& Ar)
 	Super::Serialize(Ar);
 
 #if WITH_EDITOR
-	UMaterialInstanceEditorOnlyData* EditorOnly = GetEditorOnlyData();
 	if (Ar.CustomVer(FRenderingObjectVersion::GUID) < FRenderingObjectVersion::MaterialAttributeLayerParameters)
 	{
 		// Material attribute layers parameter refactor fix-up
@@ -2459,10 +2458,10 @@ void UMaterialInstance::Serialize(FArchive& Ar)
 		StaticParameters_DEPRECATED.UpdateLegacyTerrainLayerWeightData();
 	}
 
-	if (!StaticParameters_DEPRECATED.IsEmpty())
+	if (Ar.IsLoading() && !StaticParameters_DEPRECATED.IsEmpty())
 	{
 		StaticParametersRuntime = MoveTemp(StaticParameters_DEPRECATED.GetRuntime());
-		EditorOnly->StaticParameters = MoveTemp(StaticParameters_DEPRECATED.EditorOnly);
+		GetEditorOnlyData()->StaticParameters = MoveTemp(StaticParameters_DEPRECATED.EditorOnly);
 		StaticParameters_DEPRECATED.Empty();
 	}
 
@@ -2535,7 +2534,7 @@ void UMaterialInstance::Serialize(FArchive& Ar)
 			{
 				StaticParameters_DEPRECATED.SerializeLegacy(Ar);
 				StaticParametersRuntime = MoveTemp(StaticParameters_DEPRECATED.GetRuntime());
-				EditorOnly->StaticParameters = MoveTemp(StaticParameters_DEPRECATED.EditorOnly);
+				GetEditorOnlyData()->StaticParameters = MoveTemp(StaticParameters_DEPRECATED.EditorOnly);
 			}
 
 			static_assert(!STORE_ONLY_ACTIVE_SHADERMAPS, "Only discard unused SMs in cooked build");
@@ -2562,12 +2561,15 @@ void UMaterialInstance::Serialize(FArchive& Ar)
 			FMaterialShaderMapId LegacyId;
 			LegacyId.Serialize(Ar, bLoadedByCookedMaterial);
 
-			EditorOnly->StaticParameters.StaticSwitchParameters = LegacyId.GetStaticSwitchParameters();
-			EditorOnly->StaticParameters.StaticComponentMaskParameters = LegacyId.GetStaticComponentMaskParameters();
-			EditorOnly->StaticParameters.TerrainLayerWeightParameters = LegacyId.GetTerrainLayerWeightParameters();
+			if (IsEditorOnlyDataValid())
+			{
+				GetEditorOnlyData()->StaticParameters.StaticSwitchParameters = LegacyId.GetStaticSwitchParameters();
+				GetEditorOnlyData()->StaticParameters.StaticComponentMaskParameters = LegacyId.GetStaticComponentMaskParameters();
+				GetEditorOnlyData()->StaticParameters.TerrainLayerWeightParameters = LegacyId.GetTerrainLayerWeightParameters();
 
-			TrimToOverriddenOnly(EditorOnly->StaticParameters.StaticSwitchParameters);
-			TrimToOverriddenOnly(EditorOnly->StaticParameters.StaticComponentMaskParameters);
+				TrimToOverriddenOnly(GetEditorOnlyData()->StaticParameters.StaticSwitchParameters);
+				TrimToOverriddenOnly(GetEditorOnlyData()->StaticParameters.StaticComponentMaskParameters);
+			}
 		}
 #endif // WITH_EDITOR
 	}
