fixed static destruction order bug: the completion event of FRenderCommandFence was destroyed after the allocator it uses

diff --git a/Engine/Source/Runtime/Engine/Private/UnrealEngine.cpp b/Engine/Source/Runtime/Engine/Private/UnrealEngine.cpp
--- a/Engine/Source/Runtime/Engine/Private/UnrealEngine.cpp
+++ b/Engine/Source/Runtime/Engine/Private/UnrealEngine.cpp
@@ -12063,6 +12063,33 @@ static FAutoConsoleVariableRef GDoAsyncLoadingWhileWaitingForVSyncCVar(
 	ECVF_Default
 );
 
+FFrameEndSync::FFrameEndSync()
+{
+	// FFrameEndSync instances are often used as static local vars. we need to cleanup them on engine exit to avoid static destruciton order problem
+	CleanupDelegate = FCoreDelegates::OnEnginePreExit.AddRaw(this, &FFrameEndSync::Cleanup);
+}
+
+FFrameEndSync::~FFrameEndSync()
+{
+	// if it's destroyed before the engine exit, remove the delegate so it doesn't use the instance after destruction
+	if (CleanupDelegate.IsValid())
+	{
+		FCoreDelegates::OnEnginePreExit.Remove(CleanupDelegate);
+	}
+}
+
+void FFrameEndSync::Cleanup()
+{
+	// forcing static deallocation order:
+	// fences can hold task completion handles that need to be freed before their allocator is destroyed. wating for fences does the job
+	Fence[0].Wait();
+	Fence[1].Wait();
+
+	FCoreDelegates::OnEnginePreExit.Remove(CleanupDelegate);
+	// notify the destructor
+	CleanupDelegate = {};
+}
+
 void FFrameEndSync::Sync( bool bAllowOneFrameThreadLag )
 {
 	check(IsInGameThread());			
diff --git a/Engine/Source/Runtime/Engine/Public/UnrealEngine.h b/Engine/Source/Runtime/Engine/Public/UnrealEngine.h
--- a/Engine/Source/Runtime/Engine/Public/UnrealEngine.h
+++ b/Engine/Source/Runtime/Engine/Public/UnrealEngine.h
@@ -377,12 +377,21 @@ class FFrameEndSync
 	FRenderCommandFence Fence[2];
 	/** Current index into events array. */
 	int32 EventIndex;
+	/** cleanup delegate for engine pre-exit */
+	FDelegateHandle CleanupDelegate;
+
 public:
+	ENGINE_API FFrameEndSync();
+	ENGINE_API ~FFrameEndSync();
+
 	/**
 	 * Syncs the game thread with the render thread. Depending on passed in bool this will be a total
 	 * sync or a one frame lag.
 	 */
 	ENGINE_API void Sync( bool bAllowOneFrameThreadLag );
+
+private:
+	void Cleanup();
 };
 
 
diff --git a/Engine/Source/Runtime/RenderCore/Private/RenderingThread.cpp b/Engine/Source/Runtime/RenderCore/Private/RenderingThread.cpp
--- a/Engine/Source/Runtime/RenderCore/Private/RenderingThread.cpp
+++ b/Engine/Source/Runtime/RenderCore/Private/RenderingThread.cpp
@@ -1272,6 +1272,7 @@ void FRenderCommandFence::Wait(bool bProcessGameThreadTasks) const
 		}
 #endif
 		GameThreadWaitForTask(CompletionEvent, TriggerThreadIndex, bProcessGameThreadTasks);
+		CompletionEvent = nullptr; // release the internal memory as soon as it's not needed anymore
 	}
 }
 
