Modified VisualizeCalibrationMaterial tooltips re: post-process material requirements.

diff --git a/Engine/Source/Runtime/Engine/Classes/Engine/RendererSettings.h b/Engine/Source/Runtime/Engine/Classes/Engine/RendererSettings.h
--- a/Engine/Source/Runtime/Engine/Classes/Engine/RendererSettings.h
+++ b/Engine/Source/Runtime/Engine/Classes/Engine/RendererSettings.h
@@ -1128,19 +1128,19 @@ class ENGINE_API URendererSettings : public UDeveloperSettings
 	*/
 	UPROPERTY(config, EditAnywhere, Category = PostProcessCalibrationMaterials, meta = (AllowedClasses = "Material",
 		DisplayName = "Visualize Calibration Color Material Path",
-		ToolTip = "When the VisualizeCalibrationColor show flag is enabled, this path will be used as the post-process material to render.",
+		ToolTip = "When the VisualizeCalibrationColor show flag is enabled, this path will be used as the post-process material to render. The post-process material's Blendable Location property must be set to \"After Tonemapping\" for proper calibration display.",
 		ConfigRestartRequired = false))
 	FSoftObjectPath VisualizeCalibrationColorMaterialPath;
 
 	UPROPERTY(config, EditAnywhere, Category = PostProcessCalibrationMaterials, meta = (AllowedClasses = "Material",
 		DisplayName = "Visualize Calibration Custom Material Path",
-		ToolTip = "When the VisualizeCalibrationCustom show flag is enabled, this path will be used as the post-process material to render.",
+		ToolTip = "When the VisualizeCalibrationCustom show flag is enabled, this path will be used as the post-process material to render. The post-process material's Blendable Location property must be set to \"After Tonemapping\" for proper calibration display.",
 		ConfigRestartRequired = false))
 	FSoftObjectPath VisualizeCalibrationCustomMaterialPath;
 
 	UPROPERTY(config, EditAnywhere, Category = PostProcessCalibrationMaterials, meta = (AllowedClasses = "Material",
 		DisplayName = "Visualize Calibration Grayscale Material Path",
-		ToolTip = "When the VisualizeCalibrationGrayscale show flag is enabled, this path will be used as the post-process material to render.",
+		ToolTip = "When the VisualizeCalibrationGrayscale show flag is enabled, this path will be used as the post-process material to render. The post-process material's Blendable Location property must be set to \"After Tonemapping\" for proper calibration display.",
 		ConfigRestartRequired = false))
 		FSoftObjectPath VisualizeCalibrationGrayscaleMaterialPath;
 
