Fix decal rendering when strata is enabled.

diff --git a/Engine/Shaders/Private/DecalCommon.ush b/Engine/Shaders/Private/DecalCommon.ush
--- a/Engine/Shaders/Private/DecalCommon.ush
+++ b/Engine/Shaders/Private/DecalCommon.ush
@@ -16,9 +16,10 @@
 float DecalCommonOutput(
 	inout FPixelShaderOut Out, 
 	in FStrataPixelHeader InHeader,
+	in FStrataBSDF InBSDF,
 	in float CoverageWeight)
 {
-	FStrataDBuffer DBuffer = StrataConvertToDBuffer(InHeader);
+	FStrataDBuffer DBuffer = StrataConvertToDBuffer(InHeader, InBSDF);
 	DBuffer.Coverage *= CoverageWeight;
 
 	float3 Color = DBuffer.Emissive;
diff --git a/Engine/Shaders/Private/DeferredDecal.usf b/Engine/Shaders/Private/DeferredDecal.usf
--- a/Engine/Shaders/Private/DeferredDecal.usf
+++ b/Engine/Shaders/Private/DeferredDecal.usf
@@ -146,7 +146,7 @@ void FPixelShaderInOut_MainPS(inout FPixelShaderIn In, inout FPixelShaderOut Out
 	StrataHeader.IrradianceAO = InitIrradianceAndOcclusion();
 	StrataHeader.IrradianceAO.MaterialAO = GetMaterialAmbientOcclusion(PixelMaterialInputs);
 	StrataHeader.SharedLocalBases = MaterialParameters.SharedLocalBases;
-	const float Coverage = DecalCommonOutput(Out, StrataHeader, DecalFading);
+	const float Coverage = DecalCommonOutput(Out, StrataHeader, StrataHeader.StrataTree.BSDFs[0], DecalFading);
 	const float Opacity = Coverage;
 #else // STRATA_ENABLED
 	float3 Color = GetMaterialEmissive(PixelMaterialInputs);
diff --git a/Engine/Shaders/Private/MeshDecals.usf b/Engine/Shaders/Private/MeshDecals.usf
--- a/Engine/Shaders/Private/MeshDecals.usf
+++ b/Engine/Shaders/Private/MeshDecals.usf
@@ -138,7 +138,7 @@ void FPixelShaderInOut_MainPS(
 	StrataHeader.IrradianceAO = InitIrradianceAndOcclusion();
 	StrataHeader.IrradianceAO.MaterialAO = GetMaterialAmbientOcclusion(PixelMaterialInputs);
 	StrataHeader.SharedLocalBases = MaterialParameters.SharedLocalBases;
-	const float Coverage = DecalCommonOutput(Out, StrataHeader, 1.f /*CoverageWeight*/);
+	const float Coverage = DecalCommonOutput(Out, StrataHeader, StrataHeader.StrataTree.BSDFs[0], 1.f /*CoverageWeight*/);
 	const float Opacity = Coverage;
 #else // STRATA_ENABLED
 	float3 Color = GetMaterialEmissive(PixelMaterialInputs);
diff --git a/Engine/Shaders/Private/Strata/Strata.ush b/Engine/Shaders/Private/Strata/Strata.ush
--- a/Engine/Shaders/Private/Strata/Strata.ush
+++ b/Engine/Shaders/Private/Strata/Strata.ush
@@ -2805,24 +2805,23 @@ FStrataDBuffer StrataUnpackDBuffer(
 }
 
 FStrataDBuffer StrataConvertToDBuffer(
-	in FStrataPixelHeader InHeader)
+	in FStrataPixelHeader InHeader,
+	in FStrataBSDF InBSDF)
 {
 	FStrataDBuffer Out = (FStrataDBuffer)0;
 
 #if STRATA_INLINE_SHADING
 
-	FStrataBSDF BSDF = InHeader.StrataTree.BSDFs[0];
+	const float3 DiffuseColor = SLAB_DIFFUSEALBEDO(InBSDF);
+	const float3 F0 = SLAB_F0(InBSDF);
 
-	const float3 DiffuseColor = SLAB_DIFFUSEALBEDO(BSDF);
-	const float3 F0 = SLAB_F0(BSDF);
-
-	Out.Roughness = SLAB_ROUGHNESS(BSDF);
+	Out.Roughness = SLAB_ROUGHNESS(InBSDF);
 	Out.Metallic = F0RGBToMetallic(F0);
 	Out.Specular = F0RGBToDielectricSpecular(F0);
 	Out.BaseColor = lerp(DiffuseColor, F0, Out.Metallic);
-	Out.Emissive = BSDF_GETEMISSIVE(BSDF);
-	Out.WorldNormal = InHeader.SharedLocalBases.Normals[BSDF_GETSHAREDLOCALBASISID(BSDF)];
-	Out.Coverage = BSDF.Coverage;
+	Out.Emissive = BSDF_GETEMISSIVE(InBSDF);
+	Out.WorldNormal = InHeader.SharedLocalBases.Normals[BSDF_GETSHAREDLOCALBASISID(InBSDF)];
+	Out.Coverage = InBSDF.Coverage;
 
 #else
 	Out.WorldNormal = float3(0,0,1);
diff --git a/Engine/Source/Runtime/Engine/Classes/Materials/MaterialExpressionStrata.h b/Engine/Source/Runtime/Engine/Classes/Materials/MaterialExpressionStrata.h
--- a/Engine/Source/Runtime/Engine/Classes/Materials/MaterialExpressionStrata.h
+++ b/Engine/Source/Runtime/Engine/Classes/Materials/MaterialExpressionStrata.h
@@ -610,8 +610,15 @@ class UMaterialExpressionStrataConvertToDecal : public UMaterialExpressionStrata
 	UPROPERTY()
 	FExpressionInput DecalMaterial;
 
+	/**
+	 * The coverage of the decal (default 1)
+	 */
+	UPROPERTY()
+	FExpressionInput Coverage;
+
 	// When a material is converted to decals, we use parameter blending to ensure that a single slab of matter is provided (current restriction).
-	const bool bUseParameterBlending = true;
+	UPROPERTY()
+	uint32 bUseParameterBlending : 1;
 
 	//~ Begin UMaterialExpression Interface
 #if WITH_EDITOR
diff --git a/Engine/Source/Runtime/Engine/Private/Materials/Material.cpp b/Engine/Source/Runtime/Engine/Private/Materials/Material.cpp
--- a/Engine/Source/Runtime/Engine/Private/Materials/Material.cpp
+++ b/Engine/Source/Runtime/Engine/Private/Materials/Material.cpp
@@ -3195,9 +3195,13 @@ void UMaterial::ConvertMaterialToStrataMaterial()
 			AddStrataShadingModelFromMaterialShadingModel(ConvertNode->ConvertedStrataMaterialInfo, ShadingModels);
 			check(ConvertNode->ConvertedStrataMaterialInfo.CountShadingModels() == 1);
 
-			// Now pas through the convert to decal node
+			// Now pass through the convert to decal node, which flag the material as SSM_Decal, which will set the domain to Decal.
+			// This node doesn't apply parameter blending unlike when used on non-legacy material. This is because legacy conversion 
+			// node does not support parameter blending. Coverage is also left to default (1.0), as the opacity is managed within the 
+			// legacy node
 			UMaterialExpressionStrataConvertToDecal* ConvertToDecalNode= NewObject<UMaterialExpressionStrataConvertToDecal>(this);
 			ConvertToDecalNode->DecalMaterial.Connect(0, ConvertNode);
+			ConvertToDecalNode->bUseParameterBlending = 0;
 
 			EditorOnly->FrontMaterial.Connect(0, ConvertToDecalNode);
 			bInvalidateShader = true;
diff --git a/Engine/Source/Runtime/Engine/Private/Materials/MaterialExpressions.cpp b/Engine/Source/Runtime/Engine/Private/Materials/MaterialExpressions.cpp
--- a/Engine/Source/Runtime/Engine/Private/Materials/MaterialExpressions.cpp
+++ b/Engine/Source/Runtime/Engine/Private/Materials/MaterialExpressions.cpp
@@ -22081,6 +22081,7 @@ FStrataOperator* UMaterialExpressionStrataPostProcess::StrataGenerateMaterialTop
 
 UMaterialExpressionStrataConvertToDecal::UMaterialExpressionStrataConvertToDecal(const FObjectInitializer& ObjectInitializer)
 	: Super(ObjectInitializer)
+	, bUseParameterBlending(1)
 {
 	struct FConstructorStatics
 	{
@@ -22100,26 +22101,41 @@ int32 UMaterialExpressionStrataConvertToDecal::Compile(class FMaterialCompiler*
 	{
 		return Compiler->Errorf(TEXT("Missing DecalMaterial input"));
 	}
-	int32 WeightOfOneCodeChunk = Compiler->Constant(1.0f);
 
-	FStrataOperator& StrataOperator = Compiler->StrataCompilationGetOperator(this);
+	int32 CoverageCodeChunk = Coverage.GetTracedInput().Expression ? Coverage.Compile(Compiler) : Compiler->Constant(1.0f);
+	int32 DecalMaterialCodeChunk = DecalMaterial.Compile(Compiler);
 
-	if (!StrataOperator.bUseParameterBlending)
-	{
-		return Compiler->Errorf(TEXT("Strata Convert To Decal node is not receiveing a parameter blended strata material tree."));
-	}
-	else if (!StrataOperator.bRootOfParameterBlendingSubTree)
+	int32 OutputCodeChunk = INDEX_NONE;
+	FStrataOperator& StrataOperator = Compiler->StrataCompilationGetOperator(this);
+	if (bUseParameterBlending)
 	{
-		return Compiler->Errorf(TEXT("Strata Convert To Decal node must be the root of a parameter blending sun tree."));
-	}
+		if (!StrataOperator.bUseParameterBlending)
+		{
+			return Compiler->Errorf(TEXT("Strata Convert To Decal node is not receiveing a parameter blended strata material tree."));
+		}
+		else if (!StrataOperator.bRootOfParameterBlendingSubTree)
+		{
+			return Compiler->Errorf(TEXT("Strata Convert To Decal node must be the root of a parameter blending sun tree."));
+		}
 
-	int32 OutputCodeChunk = INDEX_NONE;
-	if (StrataOperator.bUseParameterBlending)
+		if (StrataOperator.bUseParameterBlending)
+		{
+			// Propagate the parameter blended normal
+			FStrataOperator* Operator = Compiler->StrataCompilationGetOperatorFromIndex(StrataOperator.LeftIndex);
+			if (StrataOperator.bRootOfParameterBlendingSubTree)
+			{
+				StrataOperator.BSDFRegisteredSharedLocalBasis = Operator->BSDFRegisteredSharedLocalBasis;
+				StrataOperator.CopyFlagsForParameterBlending(*Operator);
+			}
+
+			OutputCodeChunk = Compiler->StrataWeightParameterBlending(
+				DecalMaterialCodeChunk, CoverageCodeChunk,
+				StrataOperator.bRootOfParameterBlendingSubTree ? &StrataOperator : nullptr);
+		}
+	}
+	else
 	{
-		int32 DecalMaterialCodeChunk = DecalMaterial.Compile(Compiler);
-		OutputCodeChunk = Compiler->StrataWeightParameterBlending(
-			DecalMaterialCodeChunk, WeightOfOneCodeChunk,
-			StrataOperator.bRootOfParameterBlendingSubTree ? &StrataOperator : nullptr);
+		OutputCodeChunk = Compiler->StrataWeight(DecalMaterialCodeChunk, CoverageCodeChunk, StrataOperator.Index, StrataOperator.MaxDistanceFromLeaves);
 	}
 	return OutputCodeChunk;
 }
@@ -22136,7 +22152,18 @@ uint32 UMaterialExpressionStrataConvertToDecal::GetOutputType(int32 OutputIndex)
 
 uint32 UMaterialExpressionStrataConvertToDecal::GetInputType(int32 InputIndex)
 {
-	return MCT_Strata;
+	switch (InputIndex)
+	{
+	case 0:
+		return MCT_Strata;
+		break;
+	case 1:
+		return MCT_Float1;
+		break;
+	}
+
+	check(false);
+	return MCT_Float1;
 }
 
 bool UMaterialExpressionStrataConvertToDecal::IsResultStrataMaterial(int32 OutputIndex)
diff --git a/Engine/Source/Runtime/Engine/Private/Materials/MaterialShared.cpp b/Engine/Source/Runtime/Engine/Private/Materials/MaterialShared.cpp
--- a/Engine/Source/Runtime/Engine/Private/Materials/MaterialShared.cpp
+++ b/Engine/Source/Runtime/Engine/Private/Materials/MaterialShared.cpp
@@ -1749,23 +1749,7 @@ bool FMaterialResource::HasMaterialPropertyConnected(EMaterialProperty In) const
 	if (Engine_IsStrataEnabled() && bIsStrataSupportedDomain)
 	{
 		// Strata material traversal is cached as this is an expensive operation
-		#if WITH_EDITOR
-		if (Material->HasStrataFrontMaterialConnected())
-		{
-			if (!CachedStrataMaterialInfo.IsValid())
-			{
-				check(Material->HasStrataFrontMaterialConnected());
-				const UMaterialEditorOnlyData* MaterialEditorOnly = Material->GetEditorOnlyData();
-				if (MaterialEditorOnly->FrontMaterial.Expression->IsResultStrataMaterial(MaterialEditorOnly->FrontMaterial.OutputIndex))
-				{
-					MaterialEditorOnly->FrontMaterial.Expression->GatherStrataMaterialInfo(CachedStrataMaterialInfo, MaterialEditorOnly->FrontMaterial.OutputIndex);
-				}
-			}
-		}
-		return CachedStrataMaterialInfo.HasPropertyConnected(In);
-		#else
 		return FStrataMaterialInfo::HasPropertyConnected(Material->CachedConnectedInputs, In);
-		#endif
 	}
 	else
 	{
