Strata - fixed typos.

diff --git a/Engine/Source/Runtime/Engine/Private/Materials/Material.cpp b/Engine/Source/Runtime/Engine/Private/Materials/Material.cpp
--- a/Engine/Source/Runtime/Engine/Private/Materials/Material.cpp
+++ b/Engine/Source/Runtime/Engine/Private/Materials/Material.cpp
@@ -3196,9 +3196,9 @@ void UMaterial::ConvertMaterialToStrataMaterial()
 			check(ConvertNode->ConvertedStrataMaterialInfo.CountShadingModels() == 1);
 
 			// Now pass through the convert to decal node, which flag the material as SSM_Decal, which will set the domain to Decal.
-			// This node doesn't apply parameter blending unlike when used on non-legacy material. This is because legacy conversion 
+			// This node doesn't apply parameter blending unlike when used on non-legacy material. This is because the legacy conversion 
 			// node does not support parameter blending. Coverage is also left to default (1.0), as the opacity is managed within the 
-			// legacy node
+			// legacy node. The first BSDF of the legacy node is used as the decal material.
 			UMaterialExpressionStrataConvertToDecal* ConvertToDecalNode= NewObject<UMaterialExpressionStrataConvertToDecal>(this);
 			ConvertToDecalNode->DecalMaterial.Connect(0, ConvertNode);
 			ConvertToDecalNode->bUseParameterBlending = 0;
diff --git a/Engine/Source/Runtime/Engine/Private/Materials/MaterialExpressions.cpp b/Engine/Source/Runtime/Engine/Private/Materials/MaterialExpressions.cpp
--- a/Engine/Source/Runtime/Engine/Private/Materials/MaterialExpressions.cpp
+++ b/Engine/Source/Runtime/Engine/Private/Materials/MaterialExpressions.cpp
@@ -22115,7 +22115,7 @@ int32 UMaterialExpressionStrataConvertToDecal::Compile(class FMaterialCompiler*
 		}
 		else if (!StrataOperator.bRootOfParameterBlendingSubTree)
 		{
-			return Compiler->Errorf(TEXT("Strata Convert To Decal node must be the root of a parameter blending sun tree."));
+			return Compiler->Errorf(TEXT("Strata Convert To Decal node must be the root of a parameter blending sub tree."));
 		}
 
 		if (StrataOperator.bUseParameterBlending)
@@ -22398,7 +22398,7 @@ int32 UMaterialExpressionStrataSingleLayerWaterBSDF::Compile(class FMaterialComp
 	}
 	else if (StrataOperator.bRootOfParameterBlendingSubTree)
 	{
-		return Compiler->Errorf(TEXT("Strata SingleLayerWater BSDF node cannot be the root of a parameter blending sun tree."));
+		return Compiler->Errorf(TEXT("Strata SingleLayerWater BSDF node cannot be the root of a parameter blending sub tree."));
 	}
 
 	int32 OutputCodeChunk = Compiler->StrataSingleLayerWaterBSDF(
