adding async save api functions + change gameplay statics version to use them.

diff --git a/Engine/Source/Runtime/Engine/Private/GameplayStatics.cpp b/Engine/Source/Runtime/Engine/Private/GameplayStatics.cpp
--- a/Engine/Source/Runtime/Engine/Private/GameplayStatics.cpp
+++ b/Engine/Source/Runtime/Engine/Private/GameplayStatics.cpp
@@ -50,6 +50,7 @@
 #include "Sound/SoundWave.h"
 #include "Audio/ActorSoundParameterInterface.h"
 #include "Engine/DamageEvents.h"
+#include "HAL/PlatformMisc.h"
 
 #if WITH_ACCESSIBILITY
 #include "Framework/Application/SlateApplication.h"
@@ -2275,24 +2276,24 @@ bool UGameplayStatics::SaveDataToSlot(const TArray<uint8>& InSaveData, const FSt
 
 void UGameplayStatics::AsyncSaveGameToSlot(USaveGame* SaveGameObject, const FString& SlotName, const int32 UserIndex, FAsyncSaveGameToSlotDelegate SavedDelegate)
 {
-	TArray<uint8> ObjectBytes;
-	if (SaveGameToMemory(SaveGameObject, ObjectBytes))
+	ISaveGameSystem* SaveSystem = IPlatformFeaturesModule::Get().GetSaveGameSystem();
+
+	TSharedRef<TArray<uint8>> ObjectBytes(new TArray<uint8>());
+
+	if (SaveSystem && (SlotName.Len() > 0) && 
+		SaveGameToMemory(SaveGameObject, *ObjectBytes) && (ObjectBytes->Num() > 0) )
 	{
-		AsyncTask(ENamedThreads::AnyHiPriThreadNormalTask, [SlotName, UserIndex, SavedDelegate, ObjectBytes]()
-		{
-			bool bSuccess = SaveDataToSlot(ObjectBytes, SlotName, UserIndex);
+		FPlatformUserId PlatformUserId = FPlatformMisc::GetPlatformUserForUserIndex(UserIndex);
 
-			// Now schedule the callback on the game thread, but only if it was bound to anything
-			if (SavedDelegate.IsBound())
+		SaveSystem->SaveGameAsync(false, *SlotName, PlatformUserId, ObjectBytes, 
+			[SavedDelegate, UserIndex](const FString& SlotName, FPlatformUserId PlatformUserId, bool bSuccess)
 			{
-				AsyncTask(ENamedThreads::GameThread, [SlotName, UserIndex, SavedDelegate, bSuccess]()
-				{
-					SavedDelegate.ExecuteIfBound(SlotName, UserIndex, bSuccess);
-				});
+				check(IsInGameThread());
+				SavedDelegate.ExecuteIfBound(SlotName, UserIndex, bSuccess);
 			}
-		});
+		);
 	}
-	else if (SavedDelegate.IsBound())
+	else
 	{
 		SavedDelegate.ExecuteIfBound(SlotName, UserIndex, false);
 	}
@@ -2380,24 +2381,30 @@ bool UGameplayStatics::LoadDataFromSlot(TArray<uint8>& OutSaveData, const FStrin
 
 void UGameplayStatics::AsyncLoadGameFromSlot(const FString& SlotName, const int32 UserIndex, FAsyncLoadGameFromSlotDelegate LoadedDelegate)
 {
-	AsyncTask(ENamedThreads::AnyHiPriThreadNormalTask, [SlotName, UserIndex, LoadedDelegate]()
+	ISaveGameSystem* SaveSystem = IPlatformFeaturesModule::Get().GetSaveGameSystem();
+	if (SaveSystem && (SlotName.Len() > 0))
 	{
-		// Do the actual I/O on the background thread
-		TArray<uint8> ObjectBytes;
-		LoadDataFromSlot(ObjectBytes, SlotName, UserIndex);
+		FPlatformUserId PlatformUserId = FPlatformMisc::GetPlatformUserForUserIndex(UserIndex);
 
-		// Now schedule the serialize and callback on the game thread
-		AsyncTask(ENamedThreads::GameThread, [SlotName, UserIndex, LoadedDelegate, ObjectBytes]()
-		{
-			USaveGame* LoadedGame = nullptr;
-			if (ObjectBytes.Num() > 0)
+		SaveSystem->LoadGameAsync(false, *SlotName, PlatformUserId,
+			[LoadedDelegate, UserIndex](const FString& SlotName, FPlatformUserId PlatformUserId, bool bSuccess, const TArray<uint8>& Data)
 			{
-				LoadedGame = LoadGameFromMemory(ObjectBytes);
-			}
+				check(IsInGameThread());
+
+				USaveGame* LoadedGame = nullptr;
+				if (bSuccess)
+				{
+					LoadedGame = LoadGameFromMemory(Data);
+				}
 
-			LoadedDelegate.ExecuteIfBound(SlotName, UserIndex, LoadedGame);
-		});
-	});
+				LoadedDelegate.ExecuteIfBound(SlotName, UserIndex, LoadedGame);
+			}
+		);
+	}
+	else
+	{
+		LoadedDelegate.ExecuteIfBound(SlotName, UserIndex, nullptr);
+	}
 }
 
 USaveGame* UGameplayStatics::LoadGameFromSlot(const FString& SlotName, const int32 UserIndex)
diff --git a/Engine/Source/Runtime/Engine/Private/SaveGameSystem.cpp b/Engine/Source/Runtime/Engine/Private/SaveGameSystem.cpp
new file mode 100644
--- /dev/null
+++ b/Engine/Source/Runtime/Engine/Private/SaveGameSystem.cpp
@@ -0,0 +1,146 @@
+// Copyright Epic Games, Inc. All Rights Reserved.
+
+#include "SaveGameSystem.h"
+#include "HAL/PlatformMisc.h"
+#include "Containers/Ticker.h"
+
+
+UE::Tasks::FPipe ISaveGameSystem::AsyncTaskPipe{ TEXT("SaveGamePipe") };
+
+
+
+void ISaveGameSystem::DoesSaveGameExistAsync(const TCHAR* Name, FPlatformUserId PlatformUserId, FSaveGameAsyncExistsCallback Callback)
+{
+	FString SlotName(Name);
+
+	// start the save operation on a background thread.
+	AsyncTaskPipe.Launch(UE_SOURCE_LOCATION,
+		[this, SlotName, PlatformUserId, Callback]()
+		{
+			// check if the savegame exists
+			int32 UserIndex = FPlatformMisc::GetUserIndexForPlatformUser(PlatformUserId);
+			const ESaveExistsResult Result = DoesSaveGameExistWithResult(*SlotName, UserIndex);
+
+			// trigger the callback on the game thread.
+			if (Callback)
+			{
+				OnAsyncComplete(
+					[SlotName, PlatformUserId, Result, Callback]()
+					{
+						Callback(SlotName, PlatformUserId, Result);
+					}
+				);
+			}
+		}
+	);
+}
+
+
+void ISaveGameSystem::SaveGameAsync(bool bAttemptToUseUI, const TCHAR* Name, FPlatformUserId PlatformUserId, TSharedRef<const TArray<uint8>> Data, FSaveGameAsyncOpCompleteCallback Callback)
+{
+	FString SlotName(Name);
+
+	// start the save operation on a background thread.
+	AsyncTaskPipe.Launch(UE_SOURCE_LOCATION,
+		[this, bAttemptToUseUI, SlotName, PlatformUserId, Data, Callback]()
+		{
+			// save the savegame
+			int32 UserIndex = FPlatformMisc::GetUserIndexForPlatformUser(PlatformUserId);
+			const bool bResult = SaveGame(bAttemptToUseUI, *SlotName, UserIndex, *Data);
+
+			// trigger the callback on the game thread
+			if (Callback)
+			{
+				OnAsyncComplete(
+					[SlotName, PlatformUserId, bResult, Callback]()
+					{
+						Callback(SlotName, PlatformUserId, bResult);
+					}
+				);
+			}
+		}
+	);
+}
+
+
+void ISaveGameSystem::LoadGameAsync(bool bAttemptToUseUI, const TCHAR* Name, FPlatformUserId PlatformUserId, FSaveGameAsyncLoadCompleteCallback Callback)
+{
+	FString SlotName(Name);
+
+	// start the load operation on a background thread.
+	AsyncTaskPipe.Launch(UE_SOURCE_LOCATION,
+		[this, bAttemptToUseUI, SlotName, PlatformUserId, Callback]()
+		{
+			// load the savegame
+			TSharedRef<TArray<uint8>> Data = MakeShared<TArray<uint8>>();
+			int32 UserIndex = FPlatformMisc::GetUserIndexForPlatformUser(PlatformUserId);
+			const bool bResult = LoadGame(bAttemptToUseUI, *SlotName, UserIndex, Data.Get());
+
+			// trigger the callback on the game thread
+			if (Callback)
+			{
+				OnAsyncComplete(
+					[SlotName, PlatformUserId, bResult, Callback, Data]()
+					{
+						Callback(SlotName, PlatformUserId, bResult, Data.Get());
+					}
+				);
+			}
+		}
+	);
+}
+
+
+void ISaveGameSystem::DeleteGameAsync(bool bAttemptToUseUI, const TCHAR* Name, FPlatformUserId PlatformUserId, FSaveGameAsyncOpCompleteCallback Callback)
+{
+	FString SlotName(Name);
+
+	// start the delete operation on a background thread.
+	AsyncTaskPipe.Launch(UE_SOURCE_LOCATION,
+		[this, bAttemptToUseUI, SlotName, PlatformUserId, Callback]()
+		{
+			// delete the savegame
+			int32 UserIndex = FPlatformMisc::GetUserIndexForPlatformUser(PlatformUserId);
+			const bool bResult = DeleteGame(bAttemptToUseUI, *SlotName, UserIndex);
+
+			// trigger the callback on the game thread
+			if (Callback)
+			{
+				OnAsyncComplete(
+					[SlotName, PlatformUserId, bResult, Callback]()
+					{
+						Callback(SlotName, PlatformUserId, bResult);
+					}
+				);
+			}
+		}
+	);
+}
+
+
+void ISaveGameSystem::InitAsync(bool bAttemptToUseUI, FPlatformUserId PlatformUserId, FSaveGameAsyncInitCompleteCallback Callback)
+{
+	// default implementation does nothing, so just trigger the completion callback on the game thread immediately
+	if (Callback)
+	{
+		OnAsyncComplete(
+			[PlatformUserId, Callback]()
+			{
+				Callback(PlatformUserId, true);
+			}
+		);
+	}
+}
+
+
+void ISaveGameSystem::OnAsyncComplete(TFunction<void()> Callback)
+{
+	// NB. Using Ticker because AsyncTask may run during async package loading which may not be suitable for save data
+	FTSTicker::GetCoreTicker().AddTicker(FTickerDelegate::CreateLambda(
+		[Callback = MoveTemp(Callback)](float) -> bool
+		{
+			Callback();
+			return false;
+		}
+	));
+}
diff --git a/Engine/Source/Runtime/Engine/Public/SaveGameSystem.h b/Engine/Source/Runtime/Engine/Public/SaveGameSystem.h
--- a/Engine/Source/Runtime/Engine/Public/SaveGameSystem.h
+++ b/Engine/Source/Runtime/Engine/Public/SaveGameSystem.h
@@ -6,13 +6,16 @@
 #include "HAL/FileManager.h"
 #include "Misc/FileHelper.h"
 #include "Misc/Paths.h"
+#include "Templates/Function.h"
+#include "Templates/SharedPointer.h"
+#include "Tasks/Pipe.h"
 
 /**
  * Interface for platform feature modules
  */
 
 /** Defines the interface to platform's save game system (or a generic file based one) */
-class ISaveGameSystem
+class ENGINE_API ISaveGameSystem
 {
 public:
 
@@ -43,11 +46,49 @@ public:
 
 	/** Delete an existing save game, blocking until complete */
 	virtual bool DeleteGame(bool bAttemptToUseUI, const TCHAR* Name, const int32 UserIndex) = 0;
+
+
+	// indicates that an async savegame operation has completed. SlotName, PlatformUserId, bool=success. Should always be called from the gamethread.
+	typedef TFunction<void(const FString&, FPlatformUserId, bool)> FSaveGameAsyncOpCompleteCallback;
+
+	// indicates that an async savegame load has completed. SlotName, PlatformUserId, bool=success, SavedData. Should always be called from the gamethread.
+	typedef TFunction<void(const FString&, FPlatformUserId, bool, const TArray<uint8>&)> FSaveGameAsyncLoadCompleteCallback;
+
+	// indicates that an async savegame exists check has completed. SlotName, PlatformUserId, save exists result. Should always be called from the gamethread.
+	typedef TFunction<void(const FString&, FPlatformUserId, ISaveGameSystem::ESaveExistsResult)> FSaveGameAsyncExistsCallback;
+
+	// indicates that an async savegame initialize has completed. PlatformUserId, bool=success. Should always be called from the gamethread.
+	typedef TFunction<void(FPlatformUserId, bool)> FSaveGameAsyncInitCompleteCallback;
+
+
+	/* Asyncronously checks if the named savegame exists. Default implementation runs blocking version on a background thread */
+	virtual void DoesSaveGameExistAsync(const TCHAR* Name, FPlatformUserId PlatformUserId, FSaveGameAsyncExistsCallback Callback);
+
+	/* Saves the named savegame asyncronously. Default implementation runs blocking version on a background thread */
+	virtual void SaveGameAsync(bool bAttemptToUseUI, const TCHAR* Name, FPlatformUserId PlatformUserId, TSharedRef<const TArray<uint8>> Data, FSaveGameAsyncOpCompleteCallback Callback);
+
+	/* Load the named savegame asyncronously. Default implementation runs blocking version on a background thread */
+	virtual void LoadGameAsync(bool bAttemptToUseUI, const TCHAR* Name, FPlatformUserId PlatformUserId, FSaveGameAsyncLoadCompleteCallback Callback);
+
+	/* Delete the named savegame asyncronously. Default implementation runs blocking version on a background thread */
+	virtual void DeleteGameAsync(bool bAttemptToUseUI, const TCHAR* Name, FPlatformUserId PlatformUserId, FSaveGameAsyncOpCompleteCallback Callback);
+
+	/* (Optional) initialise the save system for the given user. Useful if the platform may display UI on first use - this can be called as part of the user login flow, for example. 
+	   If this is not used, any UI may be displayed on the first use of the other save api functions */
+	virtual void InitAsync(bool bAttemptToUseUI, FPlatformUserId PlatformUserId, FSaveGameAsyncInitCompleteCallback Callback);
+
+protected:
+
+	// save task pipe prevents multiple async save operations happening in parallel. note that the order is not guaranteed
+	static UE::Tasks::FPipe AsyncTaskPipe;
+
+	// helper function for calling back to the game thread when an async save operation has completed
+	void OnAsyncComplete(TFunction<void()> Callback);
 };
 
 
 /** A generic save game system that just uses IFileManager to save/load with normal files */
-class FGenericSaveGameSystem : public ISaveGameSystem
+class ENGINE_API FGenericSaveGameSystem : public ISaveGameSystem
 {
 public:
 	virtual bool PlatformHasNativeUI() override
