Cloth - Fix crash when calling FParallelClothTask::GetDesiredThread() outside of the game thread.

diff --git a/Engine/Source/Runtime/Engine/Private/SkeletalMeshComponentPhysics.cpp b/Engine/Source/Runtime/Engine/Private/SkeletalMeshComponentPhysics.cpp
--- a/Engine/Source/Runtime/Engine/Private/SkeletalMeshComponentPhysics.cpp
+++ b/Engine/Source/Runtime/Engine/Private/SkeletalMeshComponentPhysics.cpp
@@ -3762,7 +3762,7 @@ public:
 	}
 	static ENamedThreads::Type GetDesiredThread()
 	{
-		if (CVarEnableClothPhysicsUseTaskThread.GetValueOnGameThread() != 0)
+		if (CVarEnableClothPhysicsUseTaskThread.GetValueOnAnyThread() != 0)
 		{
 			return CPrio_FParallelClothTask.Get();
 		}
