Fix crash and assert because the texture source was modified in presave during the texture async build
   - Can be reproed by running QAGame SimpleLevel -run=WorldPartitionBuilderCommandlet -SCCProvider=None -AllowCommandletRendering -Builder=WorldPartitionHLODsBuilder

#rnx

diff --git a/Engine/Source/Runtime/Engine/Private/Texture.cpp b/Engine/Source/Runtime/Engine/Private/Texture.cpp
--- a/Engine/Source/Runtime/Engine/Private/Texture.cpp
+++ b/Engine/Source/Runtime/Engine/Private/Texture.cpp
@@ -1028,18 +1028,19 @@ void UTexture::PreSave(FObjectPreSaveContext ObjectSaveContext)
 		UpdateResource();
 	}
 
-	if (!GEngine->IsAutosaving() && !ObjectSaveContext.IsProceduralSave())
-	{
-		GWarn->StatusUpdate(0, 0, FText::Format(NSLOCTEXT("UnrealEd", "SavingPackage_CompressingSourceArt", "Compressing source art for texture:  {0}"), FText::FromString(GetName())));
-		Source.Compress();
-	}
-
 	// Ensure that compilation has finished before saving the package
 	// otherwise async compilation might try to read the bulkdata
 	// while it's being serialized to the package.
-	if (IsCompiling())
+	// This also needs to happen before the source is modified below
+	// because it invalidates the texture build due to source hash change
+	// and could cause another build to be triggered during PostCompilation
+	// causing reentrancy problems.
+	FTextureCompilingManager::Get().FinishCompilation({ this });
+
+	if (!GEngine->IsAutosaving() && !ObjectSaveContext.IsProceduralSave())
 	{
-		FTextureCompilingManager::Get().FinishCompilation({ this });
+		GWarn->StatusUpdate(0, 0, FText::Format(NSLOCTEXT("UnrealEd", "SavingPackage_CompressingSourceArt", "Compressing source art for texture:  {0}"), FText::FromString(GetName())));
+		Source.Compress();
 	}
 #endif // #if WITH_EDITOR
 }
