Moved PreshaderBuffer to full precision

diff --git a/Engine/Shaders/Public/ShaderVersion.ush b/Engine/Shaders/Public/ShaderVersion.ush
--- a/Engine/Shaders/Public/ShaderVersion.ush
+++ b/Engine/Shaders/Public/ShaderVersion.ush
@@ -2,4 +2,4 @@
 // Each time the console command is executed it generates a new GUID. As this file is included
 // in Platform.ush (which should be included in any shader) it allows to invalidate the shader DDC.
 // 
-// GUID = 975D1179-809E-420E-88B8-9C77884C5FA7
+// GUID = 7CA85F68-C697-4626-B292-ABF0E15D1B7E
diff --git a/Engine/Source/Runtime/Engine/Private/Materials/MaterialUniformExpressions.cpp b/Engine/Source/Runtime/Engine/Private/Materials/MaterialUniformExpressions.cpp
--- a/Engine/Source/Runtime/Engine/Private/Materials/MaterialUniformExpressions.cpp
+++ b/Engine/Source/Runtime/Engine/Private/Materials/MaterialUniformExpressions.cpp
@@ -361,7 +361,7 @@ FShaderParametersMetadata* FUniformExpressionSet::CreateBufferStruct()
 
 	if (UniformPreshaderBufferSize > 0u)
 	{
-		new(Members) FShaderParametersMetadata::FMember(TEXT("PreshaderBuffer"),TEXT(""),__LINE__,NextMemberOffset,UBMT_FLOAT32,EShaderPrecisionModifier::Half,1,4, UniformPreshaderBufferSize,NULL);
+		new(Members) FShaderParametersMetadata::FMember(TEXT("PreshaderBuffer"),TEXT(""),__LINE__,NextMemberOffset,UBMT_FLOAT32,EShaderPrecisionModifier::Float,1,4, UniformPreshaderBufferSize,NULL);
 		NextMemberOffset += UniformPreshaderBufferSize * sizeof(FVector4f);
 	}
 
