Fix duplicate conditional causing CIS warning

diff --git a/Engine/Source/Runtime/Engine/Private/AudioDerivedData.cpp b/Engine/Source/Runtime/Engine/Private/AudioDerivedData.cpp
--- a/Engine/Source/Runtime/Engine/Private/AudioDerivedData.cpp
+++ b/Engine/Source/Runtime/Engine/Private/AudioDerivedData.cpp
@@ -1688,7 +1688,7 @@ bool FDerivedAudioDataCompressor::Build(TArray<uint8>& OutData)
 	FAudioStatusMessageContext StatusMessage(FText::Format(NSLOCTEXT("Engine", "BuildingCompressedAudioTaskStatus", "Building compressed audio format {AudioFormat} hash {Hash} wave {SoundNodeName}..."), Args));
 
 	// these types of sounds do not need cooked data
-	if(CookInputs->bIsASourceBus || CookInputs->bIsASourceBus)
+	if(CookInputs->bIsASourceBus || CookInputs->bIsSoundWaveProcedural)
 	{
 		return false;
 	}
