Fix for static for static analysis issues.

diff --git a/Engine/Source/Runtime/Engine/Private/Materials/MaterialExpressions.cpp b/Engine/Source/Runtime/Engine/Private/Materials/MaterialExpressions.cpp
--- a/Engine/Source/Runtime/Engine/Private/Materials/MaterialExpressions.cpp
+++ b/Engine/Source/Runtime/Engine/Private/Materials/MaterialExpressions.cpp
@@ -13098,12 +13098,15 @@ void UMaterialFunction::PostLoad()
 		EditorOnly->ExpressionCollection.ExpressionExecEnd = MoveTemp(ExpressionExecEnd_DEPRECATED);
 	}
 
-	for (UMaterialExpression* Expression : EditorOnly->ExpressionCollection.Expressions)
+	if (EditorOnly)
 	{
-		// Expressions whose type was removed can be nullptr
-		if (Expression)
+		for (UMaterialExpression* Expression : EditorOnly->ExpressionCollection.Expressions)
 		{
-			Expression->ConditionalPostLoad();
+			// Expressions whose type was removed can be nullptr
+			if (Expression)
+			{
+				Expression->ConditionalPostLoad();
+			}
 		}
 	}
 #endif // WITH_EDITORONLY_DATA
@@ -13115,7 +13118,7 @@ void UMaterialFunction::PostLoad()
 	}
 	UpdateDependentFunctionCandidates();
 
-	if (GIsEditor)
+	if (GIsEditor && EditorOnly)
 	{
 		// Clean up any removed material expression classes	
 		if (EditorOnly->ExpressionCollection.Expressions.Remove(nullptr) != 0)
@@ -13130,56 +13133,71 @@ void UMaterialFunction::PostLoad()
 	if (GMaterialFunctionsThatNeedExpressionsFlipped.Get(this))
 	{
 		GMaterialFunctionsThatNeedExpressionsFlipped.Clear(this);
-		UMaterial::FlipExpressionPositions(EditorOnly->ExpressionCollection.Expressions, EditorOnly->ExpressionCollection.EditorComments, true);
+		if (EditorOnly)
+		{
+			UMaterial::FlipExpressionPositions(EditorOnly->ExpressionCollection.Expressions, EditorOnly->ExpressionCollection.EditorComments, true);
+		}
 	}
 	else if (GMaterialFunctionsThatNeedCoordinateCheck.Get(this))
 	{
 		GMaterialFunctionsThatNeedCoordinateCheck.Clear(this);
-		if (HasFlippedCoordinates())
+		if (EditorOnly)
 		{
-			UMaterial::FlipExpressionPositions(EditorOnly->ExpressionCollection.Expressions, EditorOnly->ExpressionCollection.EditorComments, false);
+			if (HasFlippedCoordinates())
+			{
+				UMaterial::FlipExpressionPositions(EditorOnly->ExpressionCollection.Expressions, EditorOnly->ExpressionCollection.EditorComments, false);
+			}
+			UMaterial::FixCommentPositions(EditorOnly->ExpressionCollection.EditorComments);
 		}
-		UMaterial::FixCommentPositions(EditorOnly->ExpressionCollection.EditorComments);
 	}
 	else if (GMaterialFunctionsThatNeedCommentFix.Get(this))
 	{
 		GMaterialFunctionsThatNeedCommentFix.Clear(this);
-		UMaterial::FixCommentPositions(EditorOnly->ExpressionCollection.EditorComments);
+		if (EditorOnly)
+		{
+			UMaterial::FixCommentPositions(EditorOnly->ExpressionCollection.EditorComments);
+		}
 	}
 
 	if (GMaterialFunctionsThatNeedFeatureLevelSM6Fix.Get(this))
 	{
 		GMaterialFunctionsThatNeedFeatureLevelSM6Fix.Clear(this);
-		UMaterial::FixFeatureLevelNodesForSM6(EditorOnly->ExpressionCollection.Expressions);
+		if (EditorOnly)
+		{
+			UMaterial::FixFeatureLevelNodesForSM6(EditorOnly->ExpressionCollection.Expressions);
+		}
 	}
 
 	if (GMaterialFunctionsThatNeedSamplerFixup.Get(this))
 	{
 		GMaterialFunctionsThatNeedSamplerFixup.Clear(this);
-		for (UMaterialExpression* Expression : EditorOnly->ExpressionCollection.Expressions)
+		if (EditorOnly)
 		{
-			UMaterialExpressionTextureBase* TextureExpression = Cast<UMaterialExpressionTextureBase>(Expression);
-			if (TextureExpression && TextureExpression->Texture)
+			for (UMaterialExpression* Expression : EditorOnly->ExpressionCollection.Expressions)
 			{
-				switch (TextureExpression->Texture->CompressionSettings)
+				UMaterialExpressionTextureBase* TextureExpression = Cast<UMaterialExpressionTextureBase>(Expression);
+				if (TextureExpression && TextureExpression->Texture)
 				{
-				case TC_Normalmap:
-					TextureExpression->SamplerType = SAMPLERTYPE_Normal;
-					break;
-				case TC_Grayscale:
-					TextureExpression->SamplerType = TextureExpression->Texture->SRGB ? SAMPLERTYPE_Grayscale : SAMPLERTYPE_LinearGrayscale;
-					break;
+					switch (TextureExpression->Texture->CompressionSettings)
+					{
+					case TC_Normalmap:
+						TextureExpression->SamplerType = SAMPLERTYPE_Normal;
+						break;
+					case TC_Grayscale:
+						TextureExpression->SamplerType = TextureExpression->Texture->SRGB ? SAMPLERTYPE_Grayscale : SAMPLERTYPE_LinearGrayscale;
+						break;
 
-				case TC_Masks:
-					TextureExpression->SamplerType = SAMPLERTYPE_Masks;
-					break;
+					case TC_Masks:
+						TextureExpression->SamplerType = SAMPLERTYPE_Masks;
+						break;
 
-				case TC_Alpha:
-					TextureExpression->SamplerType = SAMPLERTYPE_Alpha;
-					break;
-				default:
-					TextureExpression->SamplerType = TextureExpression->Texture->SRGB ? SAMPLERTYPE_Color : SAMPLERTYPE_LinearColor;
-					break;
+					case TC_Alpha:
+						TextureExpression->SamplerType = SAMPLERTYPE_Alpha;
+						break;
+					default:
+						TextureExpression->SamplerType = TextureExpression->Texture->SRGB ? SAMPLERTYPE_Color : SAMPLERTYPE_LinearColor;
+						break;
+					}
 				}
 			}
 		}
