World Partition Actor Descriptors:
- Added actor tags to actor descriptors so they can be queried for unloaded actors.
- Added filtering by actor tags in WorldPartitionResaveActorsBuilder.

#rnx

diff --git a/Engine/Source/Editor/UnrealEd/Private/WorldPartition/WorldPartitionResaveActorsBuilder.cpp b/Engine/Source/Editor/UnrealEd/Private/WorldPartition/WorldPartitionResaveActorsBuilder.cpp
--- a/Engine/Source/Editor/UnrealEd/Private/WorldPartition/WorldPartitionResaveActorsBuilder.cpp
+++ b/Engine/Source/Editor/UnrealEd/Private/WorldPartition/WorldPartitionResaveActorsBuilder.cpp
@@ -69,6 +69,11 @@ bool UWorldPartitionResaveActorsBuilder::PreRun(UWorld* World, FPackageSourceCon
 			UE_LOG(LogWorldPartitionResaveActorsBuilder, Error, TEXT("SwitchActorPackagingSchemeToReduced is not compatible with ActorClassName"));
 			return false;
 		}
+		else if (!ActorTags.IsEmpty())
+		{
+			UE_LOG(LogWorldPartitionResaveActorsBuilder, Error, TEXT("SwitchActorPackagingSchemeToReduced is not compatible with ActorTags"));
+			return false;
+		}
 		else if (bResaveDirtyActorDescsOnly)
 		{
 			UE_LOG(LogWorldPartitionResaveActorsBuilder, Error, TEXT("SwitchActorPackagingSchemeToReduced is not compatible with ResaveDirtyActorDescsOnly"));
@@ -266,9 +271,35 @@ bool UWorldPartitionResaveActorsBuilder::RunInternal(UWorld* World, const FCellI
 	}
 	else
 	{
-		TArray<UPackage*> PackagesToCheckout;
-		FWorldPartitionHelpers::ForEachActorWithLoading(WorldPartition, ActorClass, 
-		[this, &PackagesToDelete, &PackagesToCheckout](const FWorldPartitionActorDesc* ActorDesc)
+		TArray<UPackage*> PackagesToSave;
+
+		FWorldPartitionHelpers::FForEachActorWithLoadiongParams ForEachActorWithLoadingParams;
+
+		ForEachActorWithLoadingParams.ActorClass = ActorClass;
+
+		if (ActorTags.Num())
+		{
+			TSet<FName> ActorTagsSet(ActorTags);
+			ForEachActorWithLoadingParams.FilterActorDesc = [this, &ActorTagsSet](const FWorldPartitionActorDesc* ActorDesc) -> bool
+			{
+				for (const FName& ActorTag : ActorDesc->GetTags())
+				{
+					if (ActorTagsSet.Contains(ActorTag))
+					{
+						return true;
+					}
+				}
+				return false;
+			};
+		}
+
+		ForEachActorWithLoadingParams.OnPreGarbageCollect = [&PackagesToSave, &PackageHelper]()
+		{
+			UWorldPartitionBuilder::SavePackages(PackagesToSave, PackageHelper, true);
+			PackagesToSave.Empty();
+		};
+
+		FWorldPartitionHelpers::ForEachActorWithLoading(WorldPartition, ForEachActorWithLoadingParams, [this, &PackagesToDelete, &PackagesToSave](const FWorldPartitionActorDesc* ActorDesc)
 		{
 			AActor* Actor = ActorDesc->GetActor();
 
@@ -308,17 +339,12 @@ bool UWorldPartitionResaveActorsBuilder::RunInternal(UWorld* World, const FCellI
 
 				if (!bReportOnly)
 				{
-					PackagesToCheckout.Add(Package);
+					PackagesToSave.Add(Package);
 					return true;
 				}
 			}
 
 			return true;
-		},
-		[&PackagesToCheckout, &PackageHelper]()
-		{
-			UWorldPartitionBuilder::SavePackages(PackagesToCheckout, PackageHelper, /*bErrorsAsWarnings*/true);
-			PackagesToCheckout.Empty();
 		});
 	}
 
diff --git a/Engine/Source/Editor/UnrealEd/Public/WorldPartition/WorldPartitionResaveActorsBuilder.h b/Engine/Source/Editor/UnrealEd/Public/WorldPartition/WorldPartitionResaveActorsBuilder.h
--- a/Engine/Source/Editor/UnrealEd/Public/WorldPartition/WorldPartitionResaveActorsBuilder.h
+++ b/Engine/Source/Editor/UnrealEd/Public/WorldPartition/WorldPartitionResaveActorsBuilder.h
@@ -5,7 +5,7 @@
 #include "WorldPartition/WorldPartitionBuilder.h"
 #include "WorldPartitionResaveActorsBuilder.generated.h"
 
-// Example Command Line: ProjectName MapName -run=WorldPartitionBuilderCommandlet -SCCProvider=Perforce -Builder=WorldPartitionResaveActorsBuilder [-ActorClassName=StaticMeshActor] [-SwitchActorPackagingSchemeToReduced]
+// Example Command Line: ProjectName MapName -run=WorldPartitionBuilderCommandlet -SCCProvider=Perforce -Builder=WorldPartitionResaveActorsBuilder [-ActorClassName=StaticMeshActor] [-SwitchActorPackagingSchemeToReduced] [-ActorTags=(Tag1,Tag2,...)]
 
 UCLASS()
 class UWorldPartitionResaveActorsBuilder : public UWorldPartitionBuilder
@@ -35,4 +35,7 @@ private:
 
 	UPROPERTY()
 	bool bEnableActorFolders;
+
+	UPROPERTY()
+	TArray<FName> ActorTags;
 };
\ No newline at end of file
diff --git a/Engine/Source/Runtime/Core/Public/UObject/FortniteNCBranchObjectVersion.h b/Engine/Source/Runtime/Core/Public/UObject/FortniteNCBranchObjectVersion.h
--- a/Engine/Source/Runtime/Core/Public/UObject/FortniteNCBranchObjectVersion.h
+++ b/Engine/Source/Runtime/Core/Public/UObject/FortniteNCBranchObjectVersion.h
@@ -27,6 +27,9 @@ struct CORE_API FFortniteNCBranchObjectVersion
 		// Serialize native and base class for actor descriptors
 		WorldPartitionActorDescNativeBaseClassSerialization,
 
+		// Serialize tags for actor descriptors
+		WorldPartitionActorDescTagsSerialization,
+
 		// -----<new versions can be added above this line>-------------------------------------------------
 		VersionPlusOne,
 		LatestVersion = VersionPlusOne - 1
diff --git a/Engine/Source/Runtime/Engine/Private/WorldPartition/WorldPartitionActorDesc.cpp b/Engine/Source/Runtime/Engine/Private/WorldPartition/WorldPartitionActorDesc.cpp
--- a/Engine/Source/Runtime/Engine/Private/WorldPartition/WorldPartitionActorDesc.cpp
+++ b/Engine/Source/Runtime/Engine/Private/WorldPartition/WorldPartitionActorDesc.cpp
@@ -96,6 +96,8 @@ void FWorldPartitionActorDesc::Init(const AActor* InActor)
 		DataLayerInstanceNames = FDataLayerUtils::ResolvedDataLayerInstanceNames(this);
 	}
 
+	Tags = InActor->Tags;
+
 	ActorPackage = InActor->GetPackage()->GetFName();
 	ActorPath = *InActor->GetPathName();
 	FolderPath = InActor->GetFolderPath();
@@ -178,7 +180,16 @@ bool FWorldPartitionActorDesc::Equals(const FWorldPartitionActorDesc* Other) con
 			TArray<FGuid> SortedReferencesOther(Other->References);
 			SortedReferences.Sort();
 			SortedReferencesOther.Sort();
-			return (SortedReferences == SortedReferencesOther);
+
+			if (SortedReferences == SortedReferencesOther)
+			{
+				TArray<FName> SortedTags(Tags);
+				TArray<FName> SortedTagsOther(Other->Tags);
+				SortedTags.Sort([](const FName& LHS, const FName& RHS) { return LHS.LexicalLess(RHS); });
+				SortedTagsOther.Sort([](const FName& LHS, const FName& RHS) { return LHS.LexicalLess(RHS); });
+
+				return SortedTags == SortedTagsOther;
+			}
 		}
 	}
 
@@ -290,6 +301,11 @@ void FWorldPartitionActorDesc::Serialize(FArchive& Ar)
 
 	Ar << References;
 
+	if (Ar.CustomVer(FFortniteNCBranchObjectVersion::GUID) >= FFortniteNCBranchObjectVersion::WorldPartitionActorDescTagsSerialization)
+	{
+		Ar << Tags;
+	}
+
 	if (Ar.CustomVer(FUE5MainStreamObjectVersion::GUID) < FUE5MainStreamObjectVersion::WorldPartitionActorDescSerializeArchivePersistent)
 	{
 		Ar << ActorPackage << ActorPath;
diff --git a/Engine/Source/Runtime/Engine/Private/WorldPartition/WorldPartitionHelpers.cpp b/Engine/Source/Runtime/Engine/Private/WorldPartition/WorldPartitionHelpers.cpp
--- a/Engine/Source/Runtime/Engine/Private/WorldPartition/WorldPartitionHelpers.cpp
+++ b/Engine/Source/Runtime/Engine/Private/WorldPartition/WorldPartitionHelpers.cpp
@@ -30,20 +30,20 @@ UClass* FWorldPartitionHelpers::ResolveActorDescClass(const FWorldPartitionActor
 	return ActorBaseClass;
 }
 
-void FWorldPartitionHelpers::ForEachIntersectingActorDesc(UWorldPartition* WorldPartition, const FBox& Box, TSubclassOf<AActor> ActorClass, TFunctionRef<bool(const FWorldPartitionActorDesc*)> Predicate)
+void FWorldPartitionHelpers::ForEachIntersectingActorDesc(UWorldPartition* WorldPartition, const FBox& Box, TSubclassOf<AActor> ActorClass, TFunctionRef<bool(const FWorldPartitionActorDesc*)> Func)
 {
-	WorldPartition->EditorHash->ForEachIntersectingActor(Box, [&ActorClass, Predicate](const FWorldPartitionActorDesc* ActorDesc)
+	WorldPartition->EditorHash->ForEachIntersectingActor(Box, [&ActorClass, Func](const FWorldPartitionActorDesc* ActorDesc)
 	{
 		UClass* ActorDescClass = ResolveActorDescClass(ActorDesc);
 
 		if (ActorDescClass->IsChildOf(ActorClass))
 		{
-			Predicate(ActorDesc);
+			Func(ActorDesc);
 		}
 	});
 }
 
-void FWorldPartitionHelpers::ForEachActorDesc(UWorldPartition* WorldPartition, TSubclassOf<AActor> ActorClass, TFunctionRef<bool(const FWorldPartitionActorDesc*)> Predicate)
+void FWorldPartitionHelpers::ForEachActorDesc(UWorldPartition* WorldPartition, TSubclassOf<AActor> ActorClass, TFunctionRef<bool(const FWorldPartitionActorDesc*)> Func)
 {
 	for (UActorDescContainer::TConstIterator<> ActorDescIterator(WorldPartition); ActorDescIterator; ++ActorDescIterator)
 	{
@@ -51,7 +51,7 @@ void FWorldPartitionHelpers::ForEachActorDesc(UWorldPartition* WorldPartition, T
 
 		if (ActorDescClass->IsChildOf(ActorClass))
 		{
-			if (!Predicate(*ActorDescIterator))
+			if (!Func(*ActorDescIterator))
 			{
 				return;
 			}
@@ -81,12 +81,12 @@ namespace WorldPartitionHelpers
 		}
 	}
 
-	bool ForEachActorWithLoadingBody(const FGuid& ActorGuid, UWorldPartition* WorldPartition, TFunctionRef<bool(const FWorldPartitionActorDesc*)> Predicate, TFunctionRef<void()> OnReleasingActorReferences, bool bGCPerActor, TMap<FGuid, FWorldPartitionReference>& ActorReferences)
+	bool ForEachActorWithLoadingBody(const FGuid& ActorGuid, UWorldPartition* WorldPartition, TFunctionRef<bool(const FWorldPartitionActorDesc*)> Func, TFunctionRef<void()> OnReleasingActorReferences, bool bGCPerActor, TMap<FGuid, FWorldPartitionReference>& ActorReferences)
 	{
 		WorldPartitionHelpers::LoadReferences(WorldPartition, ActorGuid, ActorReferences);
 
 		FWorldPartitionReference ActorReference(WorldPartition, ActorGuid);
-		if (!Predicate(ActorReference.Get()))
+		if (!Func(ActorReference.Get()))
 		{
 			return false;
 		}
@@ -102,7 +102,12 @@ namespace WorldPartitionHelpers
 	}
 }
 
-void FWorldPartitionHelpers::ForEachActorWithLoading(UWorldPartition* WorldPartition, TSubclassOf<AActor> ActorClass, TFunctionRef<bool(const FWorldPartitionActorDesc*)> Predicate, TFunctionRef<void()> OnReleasingActorReferences, bool bGCPerActor)
+FWorldPartitionHelpers::FForEachActorWithLoadiongParams::FForEachActorWithLoadiongParams()
+	: bGCPerActor(false)
+	, ActorClass(AActor::StaticClass())
+{}
+
+void FWorldPartitionHelpers::ForEachActorWithLoading(UWorldPartition* WorldPartition, TSubclassOf<AActor> ActorClass, TFunctionRef<bool(const FWorldPartitionActorDesc*)> Func, TFunctionRef<void()> OnReleasingActorReferences, bool bGCPerActor)
 {
 	TMap<FGuid, FWorldPartitionReference> ActorReferences;
 
@@ -112,7 +117,7 @@ void FWorldPartitionHelpers::ForEachActorWithLoading(UWorldPartition* WorldParti
 
 		if (ActorDescClass->IsChildOf(ActorClass))
 		{
-			return WorldPartitionHelpers::ForEachActorWithLoadingBody(ActorDesc->GetGuid(), WorldPartition, Predicate, OnReleasingActorReferences, bGCPerActor, ActorReferences);
+			return WorldPartitionHelpers::ForEachActorWithLoadingBody(ActorDesc->GetGuid(), WorldPartition, Func, OnReleasingActorReferences, bGCPerActor, ActorReferences);
 		}
 
 		return true;
@@ -123,13 +128,13 @@ void FWorldPartitionHelpers::ForEachActorWithLoading(UWorldPartition* WorldParti
 	DoCollectGarbage();
 }
 
-void FWorldPartitionHelpers::ForEachActorWithLoading(UWorldPartition* WorldPartition, const TArray<FGuid>& ActorGuids, TFunctionRef<bool(const FWorldPartitionActorDesc*)> Predicate, TFunctionRef<void()> OnReleasingActorReferences, bool bGCPerActor)
+void FWorldPartitionHelpers::ForEachActorWithLoading(UWorldPartition* WorldPartition, const TArray<FGuid>& ActorGuids, TFunctionRef<bool(const FWorldPartitionActorDesc*)> Func, TFunctionRef<void()> OnReleasingActorReferences, bool bGCPerActor)
 {
 	TMap<FGuid, FWorldPartitionReference> ActorReferences;
 
 	for(const FGuid& ActorGuid : ActorGuids)
 	{
-		if (!WorldPartitionHelpers::ForEachActorWithLoadingBody(ActorGuid, WorldPartition, Predicate, OnReleasingActorReferences, bGCPerActor, ActorReferences))
+		if (!WorldPartitionHelpers::ForEachActorWithLoadingBody(ActorGuid, WorldPartition, Func, OnReleasingActorReferences, bGCPerActor, ActorReferences))
 		{
 			break;
 		}
@@ -140,6 +145,38 @@ void FWorldPartitionHelpers::ForEachActorWithLoading(UWorldPartition* WorldParti
 	DoCollectGarbage();
 }
 
+void FWorldPartitionHelpers::ForEachActorWithLoading(UWorldPartition* WorldPartition, const FForEachActorWithLoadiongParams& Params, TFunctionRef<bool(const FWorldPartitionActorDesc*)> Func)
+{
+	TMap<FGuid, FWorldPartitionReference> ActorReferences;
+
+	auto CallOnPreGarbageCollect = [&Params]()
+	{
+		if (Params.OnPreGarbageCollect)
+		{
+			Params.OnPreGarbageCollect();
+		}
+	};
+
+	ForEachActorDesc(WorldPartition, [&Func, &CallOnPreGarbageCollect, &Params, &ActorReferences, WorldPartition](const FWorldPartitionActorDesc* ActorDesc)
+	{
+		UClass* ActorDescClass = ResolveActorDescClass(ActorDesc);
+
+		if (ActorDescClass->IsChildOf(Params.ActorClass))
+		{
+			if (!Params.FilterActorDesc || Params.FilterActorDesc(ActorDesc))
+			{
+				return WorldPartitionHelpers::ForEachActorWithLoadingBody(ActorDesc->GetGuid(), WorldPartition, Func, [&CallOnPreGarbageCollect, &ActorReferences]() { CallOnPreGarbageCollect(); }, Params.bGCPerActor, ActorReferences);
+			}
+		}
+
+		return true;
+	});
+
+	CallOnPreGarbageCollect();
+	ActorReferences.Empty();
+	DoCollectGarbage();
+}
+
 bool FWorldPartitionHelpers::HasExceededMaxMemory()
 {
 	const FPlatformMemoryStats MemStats = FPlatformMemory::GetStats();
diff --git a/Engine/Source/Runtime/Engine/Public/WorldPartition/WorldPartitionActorDesc.h b/Engine/Source/Runtime/Engine/Public/WorldPartition/WorldPartitionActorDesc.h
--- a/Engine/Source/Runtime/Engine/Public/WorldPartition/WorldPartitionActorDesc.h
+++ b/Engine/Source/Runtime/Engine/Public/WorldPartition/WorldPartitionActorDesc.h
@@ -20,7 +20,6 @@ struct FWorldPartitionActorDescInitData
 };
 
 class UActorDescContainer;
-class UDataLayerSubsystem;
 struct FActorContainerID;
 
 enum class EContainerClusterMode : uint8
@@ -62,6 +61,7 @@ public:
 	inline FName GetHLODLayer() const { return HLODLayer; }
 	inline const TArray<FName>& GetDataLayers() const { return DataLayers; }
 	inline const TArray<FName>& GetDataLayerInstanceNames() const { return DataLayerInstanceNames; }
+	inline const TArray<FName>& GetTags() const { return Tags; }
 	inline void SetDataLayerInstanceNames(const TArray<FName>& InDataLayerInstanceNames) { DataLayerInstanceNames = InDataLayerInstanceNames; }
 	inline FName GetActorPackage() const { return ActorPackage; }
 	inline FName GetActorPath() const { return ActorPath; }
@@ -198,6 +198,7 @@ protected:
 	FName							HLODLayer;
 	TArray<FName>					DataLayers;
 	TArray<FGuid>					References;
+	TArray<FName>					Tags;
 	FName							FolderPath;
 	FGuid							FolderGuid;
 
diff --git a/Engine/Source/Runtime/Engine/Public/WorldPartition/WorldPartitionHelpers.h b/Engine/Source/Runtime/Engine/Public/WorldPartition/WorldPartitionHelpers.h
--- a/Engine/Source/Runtime/Engine/Public/WorldPartition/WorldPartitionHelpers.h
+++ b/Engine/Source/Runtime/Engine/Public/WorldPartition/WorldPartitionHelpers.h
@@ -27,16 +27,36 @@ public:
 		ForEachActorDesc(WorldPartition, ActorClass::StaticClass(), Func);
 	}
 
+	static void ForEachIntersectingActorDesc(UWorldPartition* WorldPartition, const FBox& Box, TSubclassOf<AActor> ActorClass, TFunctionRef<bool(const FWorldPartitionActorDesc*)> Func);
+	static void ForEachActorDesc(UWorldPartition* WorldPartition, TSubclassOf<AActor> ActorClass, TFunctionRef<bool(const FWorldPartitionActorDesc*)> Func);
+
+	/* Struct of optional parameters passed to foreach actordesc functions. */
+	struct ENGINE_API FForEachActorWithLoadiongParams
+	{
+		FForEachActorWithLoadiongParams();
+
+		/* Perform a garbage collection per-actor, useful to test if the caller properly handle GCs. */
+		bool bGCPerActor;
+
+		/* The class used to filter actors loading. */
+		TSubclassOf<AActor> ActorClass;
+
+		/* Custom filter function used to filter actors loading. */
+		TUniqueFunction<bool(const FWorldPartitionActorDesc*)> FilterActorDesc;
+
+		/* Called right before releasing actor references and performing garbage collection. */
+		TUniqueFunction<void()> OnPreGarbageCollect;
+	};
+
 	template<class ActorClass = AActor>
 	static void ForEachActorWithLoading(UWorldPartition* WorldPartition, TFunctionRef<bool(const FWorldPartitionActorDesc*)> Func)
 	{
 		ForEachActorWithLoading(WorldPartition, ActorClass::StaticClass(), Func);
 	}
 
-	static void ForEachIntersectingActorDesc(UWorldPartition* WorldPartition, const FBox& Box, TSubclassOf<AActor> ActorClass, TFunctionRef<bool(const FWorldPartitionActorDesc*)> Func);
-	static void ForEachActorDesc(UWorldPartition* WorldPartition, TSubclassOf<AActor> ActorClass, TFunctionRef<bool(const FWorldPartitionActorDesc*)> Func);
 	static void ForEachActorWithLoading(UWorldPartition* WorldPartition, TSubclassOf<AActor> ActorClass, TFunctionRef<bool(const FWorldPartitionActorDesc*)> Func, TFunctionRef<void()> OnReleasingActorReferences = [](){}, bool bGCPerActor = false);
 	static void ForEachActorWithLoading(UWorldPartition* WorldPartition, const TArray<FGuid>& ActorGuids, TFunctionRef<bool(const FWorldPartitionActorDesc*)> Func, TFunctionRef<void()> OnReleasingActorReferences = [](){}, bool bGCPerActor = false);
+	static void ForEachActorWithLoading(UWorldPartition* WorldPartition, const FForEachActorWithLoadiongParams& Params, TFunctionRef<bool(const FWorldPartitionActorDesc*)> Func);
 	
 	static bool HasExceededMaxMemory();
 	static void DoCollectGarbage();
