Store ISMC custom primitive data in the HLOD proxy descriptor

diff --git a/Engine/Source/Runtime/Engine/Private/HLOD/HLODProxyDesc.cpp b/Engine/Source/Runtime/Engine/Private/HLOD/HLODProxyDesc.cpp
--- a/Engine/Source/Runtime/Engine/Private/HLOD/HLODProxyDesc.cpp
+++ b/Engine/Source/Runtime/Engine/Private/HLOD/HLODProxyDesc.cpp
@@ -25,6 +25,14 @@ FHLODISMComponentDesc::FHLODISMComponentDesc(const UInstancedStaticMeshComponent
 		FTransform InstanceTransform;
 		InISMComponent->GetInstanceTransform(InstanceIndex, InstanceTransform);
 		Instances.Emplace(InstanceTransform);
+
+		const int32 NumCustomDataFloats = InISMComponent->NumCustomDataFloats;
+		FCustomPrimitiveData InstanceCustomData;
+		InstanceCustomData.Data.SetNumUninitialized(NumCustomDataFloats);
+		check(InISMComponent->PerInstanceSMCustomData.GetTypeSize() == InstanceCustomData.Data.GetTypeSize());
+		FMemory::Memcpy(InstanceCustomData.Data.GetData(), &InISMComponent->PerInstanceSMCustomData[InstanceIndex * NumCustomDataFloats], NumCustomDataFloats * InstanceCustomData.Data.GetTypeSize());
+
+		InstancesCustomPrimitiveData.Emplace(MoveTemp(InstanceCustomData));
 	}
 }
 
@@ -63,6 +71,11 @@ bool FHLODISMComponentDesc::operator==(const FHLODISMComponentDesc& Other) const
 		}
 	}
 
+	if (InstancesCustomPrimitiveData != Other.InstancesCustomPrimitiveData)
+	{
+		return false;
+	}
+
 	return true;
 }
 
@@ -310,11 +323,11 @@ ALODActor* UHLODProxyDesc::SpawnLODActor(ULevel* InLevel) const
 				Transform *= ActorTransform;
 			}
 
-			LODActor->AddInstances(ISMStaticMesh, ISMMaterial, Transforms);
+			LODActor->AddInstances(ISMStaticMesh, ISMMaterial, Transforms, ISMComponentDesc.InstancesCustomPrimitiveData);
 		}
 		else
 		{
-			LODActor->AddInstances(ISMStaticMesh, ISMMaterial, ISMComponentDesc.Instances);
+			LODActor->AddInstances(ISMStaticMesh, ISMMaterial, ISMComponentDesc.Instances, ISMComponentDesc.InstancesCustomPrimitiveData);
 		}
 	}
 
diff --git a/Engine/Source/Runtime/Engine/Public/HLOD/HLODProxyDesc.h b/Engine/Source/Runtime/Engine/Public/HLOD/HLODProxyDesc.h
--- a/Engine/Source/Runtime/Engine/Public/HLOD/HLODProxyDesc.h
+++ b/Engine/Source/Runtime/Engine/Public/HLOD/HLODProxyDesc.h
@@ -32,6 +32,9 @@ public:
 
 	UPROPERTY()
 	TArray<FTransform> Instances;
+
+	UPROPERTY()
+	TArray<FCustomPrimitiveData> InstancesCustomPrimitiveData;
 };
 
 /** Describe a LODActor */
diff --git a/Engine/Source/Runtime/Engine/Public/SceneTypes.h b/Engine/Source/Runtime/Engine/Public/SceneTypes.h
--- a/Engine/Source/Runtime/Engine/Public/SceneTypes.h
+++ b/Engine/Source/Runtime/Engine/Public/SceneTypes.h
@@ -26,6 +26,11 @@ struct FCustomPrimitiveData
 {
 	GENERATED_USTRUCT_BODY()
 
+	inline bool operator==(const FCustomPrimitiveData& Other) const
+	{
+		return Data == Other.Data;
+	}
+
 	static constexpr int32 NumCustomPrimitiveDataFloat4s = 9; // Must match NUM_CUSTOM_PRIMITIVE_DATA in SceneData.ush
 	static constexpr int32 NumCustomPrimitiveDataFloats = NumCustomPrimitiveDataFloat4s * 4;
 
