NetConditionGroup feature
* Added new ELifetimeCondition COND_NetGroup, that is only supported on replicated subobjects.
* Any subobject using COND_NetGroup must be part of a netcondition group via the NetConditionGroupManager world subsystem for it to get replicated.
* PlayerControllers can be members of a group via IncludeInNetConditionGroup
* When a netgroup subobject replicates to a connection we will only replicate it if the playercontroller of the connection is a member of atleast one group of that subobject.
* NetGroupOwner and NetGroupReplay are special groups that automatically replicate to the owner's connection or the replay net driver.
* This feature is only available to Actors or Actorcomponents with bReplicateUsingRegisteredSubObjectList = true

diff --git a/Engine/Source/Runtime/CoreUObject/Public/UObject/CoreNetTypes.h b/Engine/Source/Runtime/CoreUObject/Public/UObject/CoreNetTypes.h
--- a/Engine/Source/Runtime/CoreUObject/Public/UObject/CoreNetTypes.h
+++ b/Engine/Source/Runtime/CoreUObject/Public/UObject/CoreNetTypes.h
@@ -24,7 +24,8 @@ enum ELifetimeCondition
 	COND_SimulatedOrPhysicsNoReplay = 12	UMETA(DisplayName = "Simulated Or Physics No Replay"),	// This property will send to simulated Or bRepPhysics actors, but not to replay connections
 	COND_SkipReplay = 13					UMETA(DisplayName = "Skip Replay"),						// This property will not send to the replay connection
 	COND_Never = 15							UMETA(Hidden),											// This property will never be replicated
-	COND_Max = 16							UMETA(Hidden)
+	COND_NetGroup = 16						UMETA(Hidden),											// This subobject will replicate to connections that are part of the same group the subobject is registered to. Not usable on properties.
+	COND_Max = 17							UMETA(Hidden)
 };
 
 
diff --git a/Engine/Source/Runtime/Engine/Classes/Engine/ActorChannel.h b/Engine/Source/Runtime/Engine/Classes/Engine/ActorChannel.h
--- a/Engine/Source/Runtime/Engine/Classes/Engine/ActorChannel.h
+++ b/Engine/Source/Runtime/Engine/Classes/Engine/ActorChannel.h
@@ -400,6 +400,8 @@ private:
 	/** Replicate a list of subobjects */
 	bool WriteSubObjects(UActorComponent* ReplicatedComponent, const UE::Net::FSubObjectRegistry& SubObjectList, FOutBunch& Bunch, FReplicationFlags RepFlags, const TStaticBitArray<COND_Max>& ConditionMap);
 
+	bool CanSubObjectReplicateToClient(ELifetimeCondition NetCondition, FObjectKey SubObjectKey, const TStaticBitArray<COND_Max>& ConditionMap) const;
+
 	bool ValidateReplicatedSubObjects();
 
 	void TestLegacyReplicateSubObjects(UActorComponent* ReplicatedComponent, FOutBunch& Bunch, FReplicationFlags RepFlags);
diff --git a/Engine/Source/Runtime/Engine/Classes/GameFramework/Actor.h b/Engine/Source/Runtime/Engine/Classes/GameFramework/Actor.h
--- a/Engine/Source/Runtime/Engine/Classes/GameFramework/Actor.h
+++ b/Engine/Source/Runtime/Engine/Classes/GameFramework/Actor.h
@@ -3610,7 +3610,7 @@ public:
 	/** 
 	* Allows classes to control if a replicated component can actually be replicated or not in a specific actor class. 
 	* You can also choose a netcondition to filter to whom the component is replicated to.
-	* Called on every replicated component during BeginPlay() and from there on every new replicated component added to the OwnedComponent list
+	* Called on existing replicated component right before BeginPlay() and after that on every new replicated component added to the OwnedComponent list
     *
 	* @param ComponentToReplicate The replicated component added to the actor.
 	* @return Return COND_None if this component should be replicated to everyone, COND_Never if it should not be replicated at all or any other conditions for specific filtering.
diff --git a/Engine/Source/Runtime/Engine/Classes/GameFramework/PlayerController.h b/Engine/Source/Runtime/Engine/Classes/GameFramework/PlayerController.h
--- a/Engine/Source/Runtime/Engine/Classes/GameFramework/PlayerController.h
+++ b/Engine/Source/Runtime/Engine/Classes/GameFramework/PlayerController.h
@@ -1581,6 +1581,32 @@ public:
 	UFUNCTION(BlueprintCallable, Category="Game|Player", meta=(Keywords = "Camera"))
 	virtual void SetViewTargetWithBlend(class AActor* NewViewTarget, float BlendTime = 0, enum EViewTargetBlendFunction BlendFunc = VTBlend_Linear, float BlendExp = 0, bool bLockOutgoing = false);
 
+	/** 
+	* Make this player a member of a netcondition group. 
+	* Any subobject registered in the group may now be replicated to this player's connection.
+	*/
+	void IncludeInNetConditionGroup(FName NetGroup);
+
+	/** Remove this player from a netcondition group. */
+	void RemoveFromNetConditionGroup(FName NetGroup)
+	{
+		NetConditionGroups.RemoveSingleSwap(NetGroup);
+	}
+
+	/** Returns true if the player controller is a member of the netcondition group */
+	bool IsMemberOfNetConditionGroup(FName NetGroup) const
+	{
+		return NetConditionGroups.Find(NetGroup) != INDEX_NONE;
+	}
+
+	/** Returns the list of netcondition groups we are part of. */
+	const TArray<FName>& GetNetConditionGroups() const { return NetConditionGroups; }
+
+private:
+	
+	/** List of netcondition groups we are currently a member of. */
+	TArray<FName> NetConditionGroups;
+
 protected:
 	/** Clickable object currently under the mouse cursor. */
 	TWeakObjectPtr<UPrimitiveComponent> CurrentClickablePrimitive;
diff --git a/Engine/Source/Runtime/Engine/Private/Actor.cpp b/Engine/Source/Runtime/Engine/Private/Actor.cpp
--- a/Engine/Source/Runtime/Engine/Private/Actor.cpp
+++ b/Engine/Source/Runtime/Engine/Private/Actor.cpp
@@ -339,6 +339,8 @@ void AActor::ResetOwnedComponents()
 			if (Component->GetIsReplicated())
 			{
 				ReplicatedComponents.Add(Component);
+
+				AddComponentForReplication(Component);
 			}
 		}
 	}, true, RF_NoFlags, EInternalObjectFlags::Garbage);
@@ -3148,10 +3150,7 @@ void AActor::AddOwnedComponent(UActorComponent* Component)
 		{
 			ReplicatedComponents.AddUnique(Component);
 
-			if (!HasAnyFlags(RF_ClassDefaultObject) && HasActorBegunPlay())
-			{
-				AddComponentForReplication(Component);
-			}
+			AddComponentForReplication(Component);
 		}
 		
 
@@ -3218,7 +3217,7 @@ void AActor::UpdateReplicatedComponent(UActorComponent* Component)
 	}
 	else if (NetCondition != COND_Never)
 	{
-		ReplicatedComponentsInfo.Emplace(FReplicatedComponentInfo{ Component, NetCondition });
+		ReplicatedComponentsInfo.Emplace(FReplicatedComponentInfo(Component, NetCondition));
 	}
 }
 
@@ -3241,12 +3240,12 @@ void AActor::UpdateAllReplicatedComponents()
 			// We reset the array so no need to add unique
 			ReplicatedComponents.Add(Component);
 
-			if (HasActorBegunPlay())
+			if (ActorHasBegunPlay != EActorBeginPlayState::HasNotBegunPlay)
 			{
 				const ELifetimeCondition NetCondition = AllowActorComponentToReplicate(Component);
 				if (NetCondition != COND_Never)
 				{
-					const int32 Index = ReplicatedComponentsInfo.AddUnique(FReplicatedComponentInfo{ Component });
+					const int32 Index = ReplicatedComponentsInfo.AddUnique(FReplicatedComponentInfo(Component));
 					ReplicatedComponentsInfo[Index].NetCondition = NetCondition;
 				}
 			}
@@ -3913,6 +3912,19 @@ void AActor::DispatchBeginPlay(bool bFromLevelStreaming)
 
 		bActorBeginningPlayFromLevelStreaming = bFromLevelStreaming;
 		ActorHasBegunPlay = EActorBeginPlayState::BeginningPlay;
+
+		// Ask the actor class if they want to override any replicated components
+		for (UActorComponent* ReplicatedComponent : ReplicatedComponents)
+		{
+			const ELifetimeCondition NetCondition = AllowActorComponentToReplicate(ReplicatedComponent);
+
+			if (NetCondition != COND_Never)
+			{
+				const int32 Index = ReplicatedComponentsInfo.AddUnique(UE::Net::FReplicatedComponentInfo(ReplicatedComponent));
+				ReplicatedComponentsInfo[Index].NetCondition = NetCondition;
+			}
+		}
+
 		BeginPlay();
 
 		ensure(BeginPlayCallDepth - 1 == CurrentCallDepth);
@@ -3976,19 +3988,6 @@ void AActor::BeginPlay()
 	ReceiveBeginPlay();
 
 	ActorHasBegunPlay = EActorBeginPlayState::HasBegunPlay;
-
-	// Check all replicated components to find if any got overriden by the actor class now that the class finished it's BeginPlay.
-	// We check after BeginPlay because some components can get removed or added during BeginPlay.
-	for (UActorComponent* ReplicatedComponent : ReplicatedComponents)
-	{
-		const ELifetimeCondition NetCondition = AllowActorComponentToReplicate(ReplicatedComponent);
-
-		if (NetCondition != COND_Never)
-		{
-			const int32 Index = ReplicatedComponentsInfo.AddUnique(UE::Net::FReplicatedComponentInfo{ ReplicatedComponent });
-			ReplicatedComponentsInfo[Index].NetCondition = NetCondition;
-		}
-	}
 }
 
 void AActor::UpdateInitialOverlaps(bool bFromLevelStreaming)
diff --git a/Engine/Source/Runtime/Engine/Private/ActorReplication.cpp b/Engine/Source/Runtime/Engine/Private/ActorReplication.cpp
--- a/Engine/Source/Runtime/Engine/Private/ActorReplication.cpp
+++ b/Engine/Source/Runtime/Engine/Private/ActorReplication.cpp
@@ -518,7 +518,8 @@ void AActor::SetReplicatedComponentNetCondition(const UActorComponent* Replicate
 
 void AActor::AddComponentForReplication(UActorComponent* Component)
 {
-	if (ActorHasBegunPlay == EActorBeginPlayState::HasNotBegunPlay)
+	constexpr EObjectFlags CDOFlags = RF_ClassDefaultObject | RF_ArchetypeObject;
+	if (ActorHasBegunPlay == EActorBeginPlayState::HasNotBegunPlay || HasAnyFlags(CDOFlags))
 	{
 		return;
 	}
@@ -530,7 +531,7 @@ void AActor::AddComponentForReplication(UActorComponent* Component)
 		FReplicatedComponentInfo* ComponentInfo = ReplicatedComponentsInfo.FindByKey(Component);
 		if (!ComponentInfo)
 		{
-			ReplicatedComponentsInfo.Emplace(FReplicatedComponentInfo{ Component, NetCondition });
+			ReplicatedComponentsInfo.Emplace(FReplicatedComponentInfo(Component, NetCondition));
 		}
 		else
 		{
@@ -542,7 +543,11 @@ void AActor::AddComponentForReplication(UActorComponent* Component)
 
 void AActor::RemoveReplicatedComponent(UActorComponent* Component)
 {
-	ReplicatedComponentsInfo.RemoveSingleSwap(FReplicatedComponentInfo{ Component });
+	int32 Index = ReplicatedComponentsInfo.IndexOfByKey(Component);
+	if (Index != INDEX_NONE)
+	{
+		ReplicatedComponentsInfo.RemoveAtSwap(Index);
+	}
 }
 
 void AActor::AddReplicatedSubObject(UObject* SubObject, ELifetimeCondition NetCondition)
@@ -573,6 +578,7 @@ void AActor::AddActorComponentReplicatedSubObject(UActorComponent* OwnerComponen
 {
 	check(IsValid(OwnerComponent));
 	check(IsValid(SubObject));
+	ensureMsgf(NetCondition != COND_Custom, TEXT("Custom netconditions do not work with SubObjects. %s - %s will not be replicated."), *GetName(), *SubObject->GetName());
 
 	FReplicatedComponentInfo* ComponentInfo = ReplicatedComponentsInfo.FindByKey(OwnerComponent);
 	if (ComponentInfo)
@@ -593,7 +599,7 @@ void AActor::AddActorComponentReplicatedSubObject(UActorComponent* OwnerComponen
 		// Let's allow him to register subobjects but set a condition that prevents him from getting replicated until we call IsAllowed.
 		constexpr ELifetimeCondition NeverReplicate = COND_Never;
 
-		const int32 Index = ReplicatedComponentsInfo.Emplace(FReplicatedComponentInfo{OwnerComponent, NeverReplicate});
+		const int32 Index = ReplicatedComponentsInfo.Emplace(FReplicatedComponentInfo(OwnerComponent, NeverReplicate));
 		ReplicatedComponentsInfo[Index].SubObjects.AddSubObjectUnique(SubObject, NetCondition);
 
 		UE_LOG(LogNetSubObject, Verbose, TEXT("%s::%s (0x%p) added replicated subobject %s (0x%p) [%s]"),
diff --git a/Engine/Source/Runtime/Engine/Private/DataChannel.cpp b/Engine/Source/Runtime/Engine/Private/DataChannel.cpp
--- a/Engine/Source/Runtime/Engine/Private/DataChannel.cpp
+++ b/Engine/Source/Runtime/Engine/Private/DataChannel.cpp
@@ -6,6 +6,7 @@
 
 #include "Net/DataChannel.h"
 #include "UObject/UObjectIterator.h"
+#include "UObject/ObjectKey.h"
 #include "EngineStats.h"
 #include "EngineGlobals.h"
 #include "Engine/Engine.h"
@@ -22,6 +23,7 @@
 #include "Stats/StatsMisc.h"
 #include "ProfilingDebugging/CsvProfiler.h"
 #include "Net/NetworkGranularMemoryLogging.h"
+#include "Net/Subsystems/NetworkSubsystem.h"
 #include "Net/Core/Trace/NetTrace.h"
 #include "Net/Core/Misc/NetSubObjectRegistry.h"
 #include "Misc/NetworkVersion.h"
@@ -117,21 +119,19 @@ namespace UE::Net
 		QueuedBunchTimeoutSeconds,
 		TEXT("Time in seconds to wait for queued bunches on a channel to flush before logging a warning."));
 
-	int32 GCVarCompareSubObjectsReplicated = 0;
-	FAutoConsoleVariableRef CVarCompareSubObjectsReplicated(
+	static int32 GCVarCompareSubObjectsReplicated = 0;
+	static FAutoConsoleVariableRef CVarCompareSubObjectsReplicated(
 		TEXT("net.SubObjects.CompareWithLegacy"),
 		GCVarCompareSubObjectsReplicated,
 		TEXT("When turned on we will collect the subobjects replicated by the ReplicateSubObjects method and compare them with the ones replicated via the Actor's registered list. If a divergence is detected it will trigger an ensure."));
 
-	TAutoConsoleVariable<int32> CVarLogAllSubObjectComparisonErrors(
+	static TAutoConsoleVariable<int32> CVarLogAllSubObjectComparisonErrors(
 		TEXT("net.SubObjects.LogAllComparisonErrors"),
 		0,
 		TEXT("If enabled log all the errors detected by the CompareWithLegacy cheat. Otherwise only the first ensure triggered gets logged."));
 
 }; // namespace UE::Net
 
-using namespace UE::Net;
-
 template<typename T>
 static const bool IsBunchTooLarge(UNetConnection* Connection, T* Bunch)
 {
@@ -145,25 +145,27 @@ public:
 	FSubObjectGetter() = delete;
 	~FSubObjectGetter() = delete;
 
-	static const TArray<FSubObjectRegistry::FEntry>& GetSubObjects(AActor* InActor)
+	using FSubObjectRegistry = UE::Net::FSubObjectRegistry;
+
+	static const FSubObjectRegistry& GetSubObjects(AActor* InActor)
 	{
-		return InActor->ReplicatedSubObjects.GetRegistryList();
+		return InActor->ReplicatedSubObjects;
 	}
 
 	static const FSubObjectRegistry* GetSubObjectsOfActorCompoment(AActor* InActor, UActorComponent* InActorComp)
 	{
-		FReplicatedComponentInfo* ComponentInfo = InActor->ReplicatedComponentsInfo.FindByKey(InActorComp);
+		UE::Net::FReplicatedComponentInfo* ComponentInfo = InActor->ReplicatedComponentsInfo.FindByKey(InActorComp);
 		return ComponentInfo ? &(ComponentInfo->SubObjects) : nullptr;
 	}
 
-	static const TArray<FReplicatedComponentInfo>& GetReplicatedComponents(AActor* InActor)
+	static const TArray<UE::Net::FReplicatedComponentInfo>& GetReplicatedComponents(AActor* InActor)
 	{
 		return InActor->ReplicatedComponentsInfo;
 	}
 
 	static bool IsSubObjectInRegistry(AActor* InActor, UActorComponent* InActorComp, UObject* InSubObject)
 	{
-		FReplicatedComponentInfo* ComponentInfo = InActor->ReplicatedComponentsInfo.FindByKey(InActorComp);
+		UE::Net::FReplicatedComponentInfo* ComponentInfo = InActor->ReplicatedComponentsInfo.FindByKey(InActorComp);
 		return ComponentInfo ? ComponentInfo->SubObjects.IsSubObjectInRegistry(InSubObject) : false;
 	}
 };
@@ -223,6 +225,12 @@ namespace DataChannelInternal
     * It is emptied at the end of every call to UActorChannel::ReplicateActor()
 	*/
 	TArray<FSubObjectReplicatedInfo> ReplicatedSubObjectsTracker;
+
+	/** 
+	* Temporary pointer to the NetworkSubsystem that gets set and released in ReplicateActor
+	* Prevents having to grab it for every unique subobject condition test.
+	*/
+	static UNetworkSubsystem* CachedNetworkSubsystem = nullptr;
 }
 
 /*-----------------------------------------------------------------------------
@@ -1787,7 +1795,7 @@ void UControlChannel::ReceivedBunch( FInBunch& Bunch )
 			if (FNetControlMessage<NMT_NetPing>::Receive(Bunch, NetPingMessageType, MessageStr) &&
 				NetPingMessageType <= ENetPingControlMessage::Max)
 			{
-				FNetPing::HandleNetPingControlMessage(Connection, NetPingMessageType, MessageStr);
+				UE::Net::FNetPing::HandleNetPingControlMessage(Connection, NetPingMessageType, MessageStr);
 			}
 
 			break;
@@ -2121,7 +2129,7 @@ void UControlChannel::ReceiveDestructionInfo(FInBunch& Bunch)
 					TheActor->PreDestroyFromReplication();
 					TheActor->Destroy(true);
 
-					if (FilterGuidRemapping != 0)
+					if (UE::Net::FilterGuidRemapping != 0)
 					{
 						// Remove this actor's NetGUID from the list of unmapped values, it will be added back if it replicates again
 						if (NetGUID.IsValid() && Connection != nullptr && Connection->Driver != nullptr && Connection->Driver->GuidCache.IsValid())
@@ -2296,7 +2304,7 @@ void UActorChannel::DestroyActorAndComponents()
 		Actor->Destroy(true);
 	}
 
-	if (FilterGuidRemapping != 0)
+	if (UE::Net::FilterGuidRemapping != 0)
 	{
 		// Remove this actor's NetGUID from the list of unmapped values, it will be added back if it replicates again
 		if (ActorNetGUID.IsValid() && Connection != nullptr && Connection->Driver != nullptr && Connection->Driver->GuidCache.IsValid())
@@ -2733,7 +2741,7 @@ bool UActorChannel::ProcessQueuedBunches()
 		}
 		else
 		{
-			if ((FPlatformTime::Seconds() - QueuedBunchStartTime) > QueuedBunchTimeoutSeconds)
+			if ((FPlatformTime::Seconds() - QueuedBunchStartTime) > UE::Net::QueuedBunchTimeoutSeconds)
 			{
 				if (!bSuppressQueuedBunchWarningsDueToHitches && FPlatformProperties::RequiresCookedData())
 				{
@@ -3292,6 +3300,14 @@ int64 UActorChannel::ReplicateActor()
 		return 0;
 	}
 
+	// Cache the netgroup manager so we don't have to access it for every subobject.
+	DataChannelInternal::CachedNetworkSubsystem = ActorWorld->GetSubsystem<UNetworkSubsystem>();
+	check(DataChannelInternal::CachedNetworkSubsystem);
+	ON_SCOPE_EXIT
+	{
+		DataChannelInternal::CachedNetworkSubsystem = nullptr;
+	};
+
 #if UE_NET_TRACE_ENABLED
 	SetTraceCollector(Bunch, UE_NET_TRACE_CREATE_COLLECTOR(ENetTraceVerbosity::Trace));
 #endif
@@ -3540,7 +3556,7 @@ bool UActorChannel::DoSubObjectReplication(FOutBunch& Bunch, FReplicationFlags&
 		// If the actor replicates it's subobjects and those of it's components via the list.
 		bWroteSomethingImportant |= ReplicateRegisteredSubObjects(Bunch, OutRepFlags);
 
-		if (GCVarCompareSubObjectsReplicated)
+		if (UE::Net::GCVarCompareSubObjectsReplicated)
 		{
 			ValidateReplicatedSubObjects();
 
@@ -3559,8 +3575,40 @@ bool UActorChannel::DoSubObjectReplication(FOutBunch& Bunch, FReplicationFlags&
 	return bWroteSomethingImportant;
 }
 
+bool UActorChannel::CanSubObjectReplicateToClient(ELifetimeCondition NetCondition, FObjectKey SubObjectKey, const TStaticBitArray<COND_Max>& ConditionMap) const
+{
+	bool bCanReplicateToClient = ConditionMap[NetCondition];
+
+	if (NetCondition == COND_NetGroup)
+	{
+		TArrayView<const FName> NetGroups = DataChannelInternal::CachedNetworkSubsystem->GetNetConditionGroupManager().GetSubObjectNetConditionGroups(SubObjectKey);
+
+		for (int i=0; i < NetGroups.Num() && !bCanReplicateToClient; ++i)
+		{
+			const FName NetGroup = NetGroups[i];
+
+			if (UE::Net::NetGroupOwner == NetGroup)
+			{
+				bCanReplicateToClient = ConditionMap[COND_OwnerOnly];
+			}
+			else if (UE::Net::NetGroupReplay == NetGroup)
+			{
+				bCanReplicateToClient = ConditionMap[COND_ReplayOnly];
+			}
+			else
+			{
+				bCanReplicateToClient = Connection->PlayerController->IsMemberOfNetConditionGroup(NetGroup);
+			}
+		}
+	}
+
+	return bCanReplicateToClient;
+}
+	
 bool UActorChannel::ReplicateRegisteredSubObjects(FOutBunch& Bunch, FReplicationFlags RepFlags)
 {
+	using namespace UE::Net;
+
 	check(Actor->IsUsingRegisteredSubObjectList());
 
 	if (GCVarCompareSubObjectsReplicated)
@@ -3578,10 +3626,9 @@ bool UActorChannel::ReplicateRegisteredSubObjects(FOutBunch& Bunch, FReplication
 	// Start with the Actor's subobjects
 	{
 		SetCurrentSubObjectOwner(Actor);
-		const TArray<FSubObjectRegistry::FEntry>& ActorSubObjects = FSubObjectGetter::GetSubObjects(Actor);
-		for (const FSubObjectRegistry::FEntry& SubObjectInfo : ActorSubObjects)
+		for (const FSubObjectRegistry::FEntry& SubObjectInfo : FSubObjectGetter::GetSubObjects(Actor).GetRegistryList())
 		{
-			if (ConditionMap[SubObjectInfo.NetCondition])
+			if (CanSubObjectReplicateToClient(SubObjectInfo.NetCondition, SubObjectInfo.Key, ConditionMap))
 			{
 				bWroteSomethingImportant |= WriteSubObjectInBunch(SubObjectInfo.SubObject, Bunch, RepFlags);
 			}
@@ -3593,7 +3640,7 @@ bool UActorChannel::ReplicateRegisteredSubObjects(FOutBunch& Bunch, FReplication
 	{
 		check(IsValid(RepComponentInfo.Component));
 
-		if (ConditionMap[RepComponentInfo.NetCondition])
+		if (CanSubObjectReplicateToClient(RepComponentInfo.NetCondition, RepComponentInfo.Key, ConditionMap))
 		{
 			UActorComponent* ReplicatedComponent = RepComponentInfo.Component;
 
@@ -3696,7 +3743,7 @@ bool UActorChannel::ReplicateSubobject(UActorComponent* ReplicatedComponent, FOu
 	// This traps actor components using the registration list even if their owning Actor isn't using the list itself.
 	if (ReplicatedComponent->IsUsingRegisteredSubObjectList() && !DataChannelInternal::bTestingLegacyMethodForComparison)
 	{
-		if (GCVarCompareSubObjectsReplicated)
+		if (UE::Net::GCVarCompareSubObjectsReplicated)
 		{
 			TestLegacyReplicateSubObjects(ReplicatedComponent, Bunch, RepFlags);
 
@@ -3714,7 +3761,7 @@ bool UActorChannel::ReplicateSubobject(UActorComponent* ReplicatedComponent, FOu
 		// Write all the subobjects of the component
 		bWroteSomethingImportant |= WriteSubObjects(ReplicatedComponent, Bunch, RepFlags, ConditionMap);
 
-		if (GCVarCompareSubObjectsReplicated)
+		if (UE::Net::GCVarCompareSubObjectsReplicated)
 		{
 			ValidateReplicatedSubObjects();
 
@@ -3735,7 +3782,7 @@ bool UActorChannel::ReplicateSubobject(UActorComponent* ReplicatedComponent, FOu
 
 bool UActorChannel::WriteSubObjects(UActorComponent* ReplicatedComponent, FOutBunch& Bunch, FReplicationFlags RepFlags, const TStaticBitArray<COND_Max>& ConditionMap)
 {
-	if (const FSubObjectRegistry* SubObjectList = FSubObjectGetter::GetSubObjectsOfActorCompoment(Actor, ReplicatedComponent))
+	if (const UE::Net::FSubObjectRegistry* SubObjectList = FSubObjectGetter::GetSubObjectsOfActorCompoment(Actor, ReplicatedComponent))
 	{
 		return WriteSubObjects(ReplicatedComponent, *SubObjectList, Bunch, RepFlags, ConditionMap);
 	}
@@ -3743,14 +3790,14 @@ bool UActorChannel::WriteSubObjects(UActorComponent* ReplicatedComponent, FOutBu
 	return false;
 }
 
-bool UActorChannel::WriteSubObjects(UActorComponent* ReplicatedComponent, const FSubObjectRegistry& SubObjectList, FOutBunch& Bunch, FReplicationFlags RepFlags, const TStaticBitArray<COND_Max>& ConditionMap)
+bool UActorChannel::WriteSubObjects(UActorComponent* ReplicatedComponent, const UE::Net::FSubObjectRegistry& SubObjectList, FOutBunch& Bunch, FReplicationFlags RepFlags, const TStaticBitArray<COND_Max>& ConditionMap)
 {
 	bool bWroteSomethingImportant = false;
-	for (const FSubObjectRegistry::FEntry& SubObjectInfo : SubObjectList.GetRegistryList())
+	for (const UE::Net::FSubObjectRegistry::FEntry& SubObjectInfo : SubObjectList.GetRegistryList())
 	{
 		checkf(IsValid(SubObjectInfo.SubObject), TEXT("Found invalid subobject (%s) registered in %s::%s"), *GetNameSafe(SubObjectInfo.SubObject), *Actor->GetName(), *GetNameSafe(ReplicatedComponent));
 
-		if (ConditionMap[SubObjectInfo.NetCondition])
+		if (CanSubObjectReplicateToClient(SubObjectInfo.NetCondition, SubObjectInfo.Key, ConditionMap))
 		{
 			bWroteSomethingImportant |= WriteSubObjectInBunch(SubObjectInfo.SubObject, Bunch, RepFlags);
 		}
@@ -3762,11 +3809,12 @@ bool UActorChannel::WriteSubObjects(UActorComponent* ReplicatedComponent, const
 bool UActorChannel::ValidateReplicatedSubObjects()
 {
 	bool bErrorDetected = false;
-	const bool bLogAllErrors = CVarLogAllSubObjectComparisonErrors.GetValueOnAnyThread() != 0;
-	
-	// Make sure both lists contain the same subobjects and trigger an ensure if an element is missing from the other.
-	for (const DataChannelInternal::FSubObjectReplicatedInfo& Info : DataChannelInternal::LegacySubObjectsCollected)
+	const bool bLogAllErrors = UE::Net::CVarLogAllSubObjectComparisonErrors.GetValueOnAnyThread() != 0;
+
+	for (int32 i=DataChannelInternal::LegacySubObjectsCollected.Num()-1; i >= 0; i--)
 	{
+		const DataChannelInternal::FSubObjectReplicatedInfo& Info = DataChannelInternal::LegacySubObjectsCollected[i];
+
 		constexpr bool bNoShrinking = false;
 		const bool bWasFound = DataChannelInternal::ReplicatedSubObjectsTracker.RemoveSingleSwap(Info, bNoShrinking) != 0;
 
@@ -3777,8 +3825,10 @@ bool UActorChannel::ValidateReplicatedSubObjects()
 	}
 
 	// Leftover subobjects in this list means we replicated more elements than the legacy method would.
-	for (const DataChannelInternal::FSubObjectReplicatedInfo& Info : DataChannelInternal::ReplicatedSubObjectsTracker)
+	for (int32 i=DataChannelInternal::ReplicatedSubObjectsTracker.Num()-1; i >= 0; i--)
 	{
+		const DataChannelInternal::FSubObjectReplicatedInfo& Info = DataChannelInternal::ReplicatedSubObjectsTracker[i];
+
 		ensureMsgf(false, TEXT("%s was replicated only in the registered subobject list. Not in the legacy path."), *Info.Describe(Actor));
 		UE_CLOG(bLogAllErrors, LogNet, Warning, TEXT("%s was replicated only in the registered subobject list. Not in the legacy path."), *Info.Describe(Actor));
 	}
diff --git a/Engine/Source/Runtime/Engine/Private/Net/Subsystems/NetworkSubsystem.cpp b/Engine/Source/Runtime/Engine/Private/Net/Subsystems/NetworkSubsystem.cpp
new file mode 100644
--- /dev/null
+++ b/Engine/Source/Runtime/Engine/Private/Net/Subsystems/NetworkSubsystem.cpp
@@ -0,0 +1,16 @@
+// Copyright Epic Games, Inc. All Rights Reserved.
+
+#include "Net/Subsystems/NetworkSubsystem.h"
+
+#include "Net/NetworkGranularMemoryLogging.h"
+
+
+void UNetworkSubsystem::Serialize(FArchive& Ar)
+{
+	if (Ar.IsCountingMemory())
+	{
+		GRANULAR_NETWORK_MEMORY_TRACKING_INIT(Ar, "UNetworkSubsystem::Serialize");
+
+		GRANULAR_NETWORK_MEMORY_TRACKING_TRACK("GroupsManager", GroupsManager.CountBytes(Ar));
+	}
+}
\ No newline at end of file
diff --git a/Engine/Source/Runtime/Engine/Private/NetDriver.cpp b/Engine/Source/Runtime/Engine/Private/NetDriver.cpp
--- a/Engine/Source/Runtime/Engine/Private/NetDriver.cpp
+++ b/Engine/Source/Runtime/Engine/Private/NetDriver.cpp
@@ -5858,9 +5858,9 @@ void UNetDriver::SetWorld(class UWorld* InWorld)
 	{
 		// Remove old world association
 		UnregisterTickEvents(World);
-		World = NULL;
-		WorldPackage = NULL;
-		Notify = NULL;
+		World = nullptr;
+		WorldPackage = nullptr;
+		Notify = nullptr;
 
 		GetNetworkObjectList().Reset();
 	}
diff --git a/Engine/Source/Runtime/Engine/Private/PlayerController.cpp b/Engine/Source/Runtime/Engine/Private/PlayerController.cpp
--- a/Engine/Source/Runtime/Engine/Private/PlayerController.cpp
+++ b/Engine/Source/Runtime/Engine/Private/PlayerController.cpp
@@ -32,6 +32,7 @@
 #include "Widgets/SViewport.h"
 #include "Engine/Console.h"
 #include "Net/UnrealNetwork.h"
+#include "Net/Core/Misc/NetConditionGroupManager.h"
 #include "Engine/WorldComposition.h"
 #include "Engine/LevelScriptActor.h"
 #include "GameFramework/GameNetworkManager.h"
@@ -5773,4 +5774,11 @@ bool APlayerController::ShouldPerformFullTickWhenPaused() const
 			GEngine->XRSystem.IsValid() && GEngine->XRSystem->GetHMDDevice() && GEngine->XRSystem->GetHMDDevice()->IsHMDConnected());
 }
 
+void APlayerController::IncludeInNetConditionGroup(FName NetGroup)
+{
+	checkf(!UE::Net::IsSpecialNetConditionGroup(NetGroup), TEXT("Cannot add a player to special netcondition group %s manually. This group membership is managed by the network engine automatically."), *NetGroup.ToString());
+	checkf(!NetGroup.IsNone(), TEXT("Invalid netcondition group: NONE"));
+	NetConditionGroups.AddUnique(NetGroup);
+}
+
 #undef LOCTEXT_NAMESPACE
diff --git a/Engine/Source/Runtime/Engine/Private/UnrealNetwork.cpp b/Engine/Source/Runtime/Engine/Private/UnrealNetwork.cpp
--- a/Engine/Source/Runtime/Engine/Private/UnrealNetwork.cpp
+++ b/Engine/Source/Runtime/Engine/Private/UnrealNetwork.cpp
@@ -39,6 +39,7 @@ void RegisterReplicatedLifetimeProperty(
 	TArray<FLifetimeProperty>& OutLifetimeProps,
 	const FDoRepLifetimeParams& Params)
 {
+	checkf(Params.Condition != COND_NetGroup, TEXT("Invalid lifetime condition for %s. COND_NetGroup can only be used with registered subobjects"), PropertyDescriptor.PropertyName);
 	for (int32 i = 0; i < PropertyDescriptor.ArrayDim; i++)
 	{
 		const uint16 RepIndex = PropertyDescriptor.RepIndex + i;
diff --git a/Engine/Source/Runtime/Engine/Public/Net/Subsystems/NetworkSubsystem.h b/Engine/Source/Runtime/Engine/Public/Net/Subsystems/NetworkSubsystem.h
new file mode 100644
--- /dev/null
+++ b/Engine/Source/Runtime/Engine/Public/Net/Subsystems/NetworkSubsystem.h
@@ -0,0 +1,26 @@
+// Copyright Epic Games, Inc. All Rights Reserved.
+
+#pragma once
+
+#include "Subsystems/WorldSubsystem.h"
+#include "Net/Core/Misc/NetConditionGroupManager.h"
+
+#include "NetworkSubsystem.generated.h"
+
+
+UCLASS()
+class ENGINE_API UNetworkSubsystem : public UWorldSubsystem
+{
+	GENERATED_BODY()
+public:
+
+	/** Access the NetConditionGroupManager */
+	UE::Net::FNetConditionGroupManager& GetNetConditionGroupManager() { return GroupsManager;  }
+
+private:
+
+	virtual void Serialize(FArchive& Ar) override;
+
+	/** Manage the of the subobjects and their relationships to different groups */
+	UE::Net::FNetConditionGroupManager GroupsManager;
+};
diff --git a/Engine/Source/Runtime/Engine/Public/Net/UnrealNetwork.h b/Engine/Source/Runtime/Engine/Public/Net/UnrealNetwork.h
--- a/Engine/Source/Runtime/Engine/Public/Net/UnrealNetwork.h
+++ b/Engine/Source/Runtime/Engine/Public/Net/UnrealNetwork.h
@@ -256,6 +256,7 @@ static FProperty* GetReplicatedProperty(const UClass* CallingClass, const UClass
 
 #define DOREPLIFETIME_CONDITION(c,v,cond) \
 { \
+	static_assert(cond != COND_NetGroup, "COND_NetGroup cannot be used on replicated properties. Only when registering subobjects"); \
 	FDoRepLifetimeParams LocalDoRepParams; \
 	LocalDoRepParams.Condition = cond; \
 	DOREPLIFETIME_WITH_PARAMS(c,v,LocalDoRepParams); \
@@ -264,6 +265,7 @@ static FProperty* GetReplicatedProperty(const UClass* CallingClass, const UClass
 /** Allows gamecode to specify RepNotify condition: REPNOTIFY_OnChanged (default) or REPNOTIFY_Always for when repnotify function is called  */
 #define DOREPLIFETIME_CONDITION_NOTIFY(c,v,cond,rncond) \
 { \
+	static_assert(cond != COND_NetGroup, "COND_NetGroup cannot be used on replicated properties. Only when registering subobjects"); \
 	FDoRepLifetimeParams LocalDoRepParams; \
 	LocalDoRepParams.Condition = cond; \
 	LocalDoRepParams.RepNotifyCondition = rncond; \
diff --git a/Engine/Source/Runtime/Net/Core/Private/Net/Core/Misc/NetConditionGroupManager.cpp b/Engine/Source/Runtime/Net/Core/Private/Net/Core/Misc/NetConditionGroupManager.cpp
new file mode 100644
--- /dev/null
+++ b/Engine/Source/Runtime/Net/Core/Private/Net/Core/Misc/NetConditionGroupManager.cpp
@@ -0,0 +1,86 @@
+// Copyright Epic Games, Inc. All Rights Reserved.
+
+#include "Net/Core/Misc/NetConditionGroupManager.h"
+
+extern const FName UE::Net::NetGroupOwner(TEXT("NetGroupOwner"));
+extern const FName UE::Net::NetGroupReplay(TEXT("NetGroupReplay"));
+
+namespace UE::Net
+{
+
+bool IsSpecialNetConditionGroup(FName NetConditionGroup)
+{
+	return NetConditionGroup == NetGroupOwner || NetConditionGroup == NetGroupReplay;
+}
+
+void FNetConditionGroupManager::RegisterSubObjectInGroup(UObject* SubObject, FName NetConditionGroup)
+{
+	check(IsValid(SubObject));
+	checkf(!NetConditionGroup.IsNone(), TEXT("%s was registered with invalid netcondition group NONE"), *GetNameSafe(SubObject));
+
+	FSubObjectGroupInfo& SubObjectInfo = SubObjectGroupMap.FindOrAdd(SubObject);
+	SubObjectInfo.NetGroups.AddUnique(NetConditionGroup);
+}
+
+void FNetConditionGroupManager::RegisterSubObjectInMultipleGroups(UObject* SubObject, TArrayView<FName> NetConditionGroups)
+{
+	check(IsValid(SubObject));
+
+	FSubObjectGroupInfo& SubObjectInfo = SubObjectGroupMap.FindOrAdd(SubObject);
+
+	for (FName GroupName : NetConditionGroups)
+	{
+		checkf(!GroupName.IsNone(), TEXT("%s was registered with invalid netcondition group NONE"), *GetNameSafe(SubObject));
+		SubObjectInfo.NetGroups.AddUnique(GroupName);
+	}
+}
+
+void FNetConditionGroupManager::UnregisterSubObjectFromGroup(UObject* SubObject, FName NetConditionGroup)
+{
+	check(SubObject);
+	const FObjectKey SubObjectKey = SubObject;
+
+	if (FSubObjectGroupInfo* SubObjectInfo = SubObjectGroupMap.Find(SubObjectKey))
+	{
+		SubObjectInfo->NetGroups.RemoveSingleSwap(NetConditionGroup);
+
+		if (SubObjectInfo->NetGroups.Num() == 0)
+		{
+			// The object is not in any group anymore. Delete the entry from the TMap
+			SubObjectGroupMap.Remove(SubObjectKey);
+		}
+	}
+}
+
+void FNetConditionGroupManager::UnregisterSubObjectFromMultipleGroups(UObject* SubObject, TArrayView<FName> NetConditionGroups)
+{
+	check(SubObject);
+	const FObjectKey SubObjectKey = SubObject;
+
+	if (FSubObjectGroupInfo* SubObjectInfo = SubObjectGroupMap.Find(SubObjectKey))
+	{
+		for (FName GroupName : NetConditionGroups)
+		{
+			SubObjectInfo->NetGroups.RemoveSingleSwap(GroupName);
+		}
+
+		if (SubObjectInfo->NetGroups.Num() == 0)
+		{
+			// The object is not in any group anymore. Delete the entry from the TMap
+			SubObjectGroupMap.Remove(SubObjectKey);
+		}
+	}
+}
+
+void FNetConditionGroupManager::UnregisterSubObjectFromAllGroups(UObject* SubObject)
+{
+	check(SubObject);
+	SubObjectGroupMap.Remove(SubObject);
+}
+
+void FNetConditionGroupManager::CountBytes(FArchive& Ar) const
+{
+	SubObjectGroupMap.CountBytes(Ar);
+}
+
+} // end namespace UE::Net
\ No newline at end of file
diff --git a/Engine/Source/Runtime/Net/Core/Private/Net/Core/Misc/NetSubObjectRegistry.cpp b/Engine/Source/Runtime/Net/Core/Private/Net/Core/Misc/NetSubObjectRegistry.cpp
--- a/Engine/Source/Runtime/Net/Core/Private/Net/Core/Misc/NetSubObjectRegistry.cpp
+++ b/Engine/Source/Runtime/Net/Core/Private/Net/Core/Misc/NetSubObjectRegistry.cpp
@@ -2,8 +2,8 @@
 
 #include "Net/Core/Misc/NetSubObjectRegistry.h"
 
-namespace UE {
-namespace Net {
+namespace UE::Net 
+{
 
 // ----------------------------------------------------------------
 //	FSubObjectRegistry
@@ -12,7 +12,7 @@ FSubObjectRegistry::EResult FSubObjectRegistry::AddSubObjectUnique(UObject* InSu
 {
 	check(InSubObject);
 
-	const FEntry SubObjectEntry{ InSubObject, InNetCondition };
+	const FEntry SubObjectEntry(InSubObject, InNetCondition);
 
 	const int32 CurrentIndex = Registry.Find(SubObjectEntry);
 	if (CurrentIndex == INDEX_NONE)
@@ -36,7 +36,13 @@ ELifetimeCondition FSubObjectRegistry::GetNetCondition(UObject* InSubObject) con
 
 bool FSubObjectRegistry::RemoveSubObject(UObject* InSubObject)
 {
-	return Registry.RemoveSingle(FEntry{ InSubObject }) == 1;
+	int32 Index = Registry.IndexOfByKey(InSubObject);
+	if (Index != INDEX_NONE)
+	{
+		Registry.RemoveAt(Index);
+		return true;
+	}
+	return false;
 }
 
 bool FSubObjectRegistry::IsSubObjectInRegistry(UObject* SubObject) const
@@ -44,5 +50,4 @@ bool FSubObjectRegistry::IsSubObjectInRegistry(UObject* SubObject) const
 	return Registry.FindByKey(SubObject) != nullptr;
 }
 
-} //namespace Net
-} //namespace UE
\ No newline at end of file
+} //namespace UE::Net
\ No newline at end of file
diff --git a/Engine/Source/Runtime/Net/Core/Public/Net/Core/Misc/NetConditionGroupManager.h b/Engine/Source/Runtime/Net/Core/Public/Net/Core/Misc/NetConditionGroupManager.h
new file mode 100644
--- /dev/null
+++ b/Engine/Source/Runtime/Net/Core/Public/Net/Core/Misc/NetConditionGroupManager.h
@@ -0,0 +1,86 @@
+// Copyright Epic Games, Inc. All Rights Reserved.
+
+#pragma once
+
+#include "Containers/Map.h"
+#include "Containers/ArrayView.h"
+#include "UObject/ObjectKey.h"
+
+namespace UE::Net
+{
+	/**
+	* Special netcondition groups that allow a subobject to be replicated to unique conditions.
+	* For these groups the network engine will determine automatically if the connection matches the conditions needed.
+	* So the playercontroller should never be added to any of these groups.
+	*/
+
+	/** Replicate a subobject to it's owning connection. */
+	extern NETCORE_API const FName NetGroupOwner;
+
+	/** Replicate a subobject to the replay recording connection */
+	extern NETCORE_API const FName NetGroupReplay;
+
+	/** Returns true if the netcondition is one of the special explicit groups. */
+	NETCORE_API bool IsSpecialNetConditionGroup(FName NetConditionGroup);
+}
+
+
+namespace UE::Net
+{
+
+/**
+* Manage the membership of subobject's replicated with a COND_NetGroup condition.
+* For a COND_NetGroup subobject to be replicated to a connection, they both need to be part of the same netcondition group.
+* A netcondition group is simply an FName. 
+* This manager holds onto the groups the subobject is a part of, while the connection's player controller is
+* where we hold information on the user's groups.
+* 
+* Groups are inclusive and non-binding so subobject's can be part of multiple groups at the same time.
+* As long as a player is a member of one of the subobject's group, that subobject will be replicated to that connection.
+*/
+class NETCORE_API FNetConditionGroupManager
+{
+public:
+
+	/** Add a subobject into a netcondition group */
+	void RegisterSubObjectInGroup(UObject* SubObject, FName NetConditionGroup);
+
+	/** Add a subobject to many netcondition groups */
+	void RegisterSubObjectInMultipleGroups(UObject* SubObject, TArrayView<FName> NetConditionGroups);
+
+	/** Remove a subobject from a netcondition group */
+	void UnregisterSubObjectFromGroup(UObject* SubObject, FName NetConditionGroup);
+
+	/** Remove a subobject from many netcondition groups */
+	void UnregisterSubObjectFromMultipleGroups(UObject* SubObject, TArrayView<FName> NetConditionGroups);
+
+	/** Remove the subobject from all groups he was included into */
+	void UnregisterSubObjectFromAllGroups(UObject* SubObject);
+
+	/** Return the groups the subobject is a part of. Will be empty if he wasn't added to any. */
+	TArrayView<const FName> GetSubObjectNetConditionGroups(FObjectKey SubObjectKey) const
+	{
+		const FSubObjectGroupInfo* GroupInfo = SubObjectGroupMap.Find(SubObjectKey);
+		return GroupInfo ? TArrayView<const FName>(GroupInfo->NetGroups) : TArrayView<const FName>();
+	}
+
+	void CountBytes(FArchive& Ar) const;
+
+private:
+
+#ifndef DEFAULT_NETGROUP_INLINE_SIZE
+	// Sets the inlined array size of each subobject entry.  Configure it to the average groups the subobjects are registered to make them quicker to access.
+	#define DEFAULT_NETGROUP_INLINE_SIZE 4
+#endif
+
+	struct FSubObjectGroupInfo
+	{
+		TArray<FName, TInlineAllocator<DEFAULT_NETGROUP_INLINE_SIZE>> NetGroups;
+	};
+
+	/** Map of the subobjects and the groups they are a part of */
+	TMap<FObjectKey, FSubObjectGroupInfo> SubObjectGroupMap;
+};
+
+
+} // end namespace UE::Net
\ No newline at end of file
diff --git a/Engine/Source/Runtime/Net/Core/Public/Net/Core/Misc/NetSubObjectRegistry.h b/Engine/Source/Runtime/Net/Core/Public/Net/Core/Misc/NetSubObjectRegistry.h
--- a/Engine/Source/Runtime/Net/Core/Public/Net/Core/Misc/NetSubObjectRegistry.h
+++ b/Engine/Source/Runtime/Net/Core/Public/Net/Core/Misc/NetSubObjectRegistry.h
@@ -2,10 +2,13 @@
 
 #pragma once
 
-#include "CoreMinimal.h"
 #include "UObject/CoreNetTypes.h"
+#include "UObject/ObjectKey.h"
 
-namespace UE::Net {
+class UActorComponent;
+
+namespace UE::Net 
+{
 
 /**
 * Stores the SubObjects to replicate and the network condition dictating to which connection they can replicate to.
@@ -38,10 +41,20 @@ public:
 		/** Raw pointer since users are obligated to call RemoveReplicatedSubobject before destroying the subobject otherwise it will cause a crash.  */
 		UObject* SubObject = nullptr;
 
+		/** Store the object key info for fast access later */
+		FObjectKey Key;
+
 		/** The network condition that chooses which connection this subobject can be replicated to. Default is none which means all connections receive it. */
 		ELifetimeCondition NetCondition = COND_None;
 
-		bool operator==(const FEntry& rhs) const { return SubObject == rhs.SubObject; }
+		FEntry(UObject* InSubObject, ELifetimeCondition InNetCondition = COND_None)
+			: SubObject(InSubObject)
+			, Key(InSubObject)
+			, NetCondition(InNetCondition)
+		{
+		}
+
+		bool operator==(const FEntry& rhs) const { return Key == rhs.Key; }
 		bool operator==(UObject* rhs) const { return SubObject == rhs; }
 	};
 
@@ -59,17 +72,27 @@ private:
 /** Keep track of replicated components and their subobject list */
 struct FReplicatedComponentInfo
 {
-	// Component that will be replicated
-	class UActorComponent* Component = nullptr;
+	/** Component that will be replicated */
+	UActorComponent* Component = nullptr;
 
-	// NetCondition of the component
+	/** Store the object key info for fast access later */
+	FObjectKey Key;
+
+	/** NetCondition of the component */
 	ELifetimeCondition NetCondition = COND_None;
 
-	// Collection of subobjects replicated with this component
+	/** Collection of subobjects replicated with this component */
 	FSubObjectRegistry SubObjects;
 
-	bool operator==(const FReplicatedComponentInfo& rhs) const { return Component == rhs.Component; }
-	bool operator==(const class UActorComponent* rhs) const { return Component == rhs; }
+	FReplicatedComponentInfo(UActorComponent* InComponent, ELifetimeCondition InNetCondition = COND_None)
+		: Component(InComponent)
+		, Key((UObject*)InComponent)
+		, NetCondition(InNetCondition)
+	{
+	};
+
+	bool operator==(const FReplicatedComponentInfo& rhs) const { return Key == rhs.Key; }
+	bool operator==(const UActorComponent* rhs) const { return Component == rhs; }
 };
 
 
