Non-unity build fix

diff --git a/Engine/Source/Runtime/Engine/Public/HLOD/HLODProxyDesc.h b/Engine/Source/Runtime/Engine/Public/HLOD/HLODProxyDesc.h
--- a/Engine/Source/Runtime/Engine/Public/HLOD/HLODProxyDesc.h
+++ b/Engine/Source/Runtime/Engine/Public/HLOD/HLODProxyDesc.h
@@ -4,6 +4,7 @@
 
 #include "CoreMinimal.h"
 #include "Engine/MaterialMerging.h"
+#include "SceneTypes.h"
 #include "HLODProxyDesc.generated.h"
 
 class ALODActor;
