Level Instance: set the Level Instance actor label to match the level at creation time.

#rnx

diff --git a/Engine/Source/Runtime/Engine/Private/LevelInstance/LevelInstanceSubsystem.cpp b/Engine/Source/Runtime/Engine/Private/LevelInstance/LevelInstanceSubsystem.cpp
--- a/Engine/Source/Runtime/Engine/Private/LevelInstance/LevelInstanceSubsystem.cpp
+++ b/Engine/Source/Runtime/Engine/Private/LevelInstance/LevelInstanceSubsystem.cpp
@@ -960,6 +960,7 @@ ILevelInstanceInterface* ULevelInstanceSubsystem::CreateLevelInstanceFrom(const
 	ILevelInstanceInterface* NewLevelInstance = CastChecked<ILevelInstanceInterface>(NewLevelInstanceActor);
 	NewLevelInstance->SetWorldAsset(WorldPtr);
 	NewLevelInstanceActor->SetActorLocation(LevelInstanceLocation);
+	NewLevelInstanceActor->SetActorLabel(WorldPtr.GetAssetName());
 	
 	// Actors were moved and kept their World positions so when saving we want their positions to actually be relative to the LevelInstance Actor
 	// so we set the LevelTransform and we mark the level as having moved its actors. 
