Fixed issue with ISM that have no custom primitive data

diff --git a/Engine/Source/Runtime/Engine/Private/HLOD/HLODProxyDesc.cpp b/Engine/Source/Runtime/Engine/Private/HLOD/HLODProxyDesc.cpp
--- a/Engine/Source/Runtime/Engine/Private/HLOD/HLODProxyDesc.cpp
+++ b/Engine/Source/Runtime/Engine/Private/HLOD/HLODProxyDesc.cpp
@@ -20,19 +20,29 @@ FHLODISMComponentDesc::FHLODISMComponentDesc(const UInstancedStaticMeshComponent
 
 	Instances.Reset(InISMComponent->GetInstanceCount());
 
+	const int32 NumCustomDataFloats = InISMComponent->NumCustomDataFloats;
+
 	for (int32 InstanceIndex = 0; InstanceIndex < InISMComponent->GetInstanceCount(); ++InstanceIndex)
 	{
 		FTransform InstanceTransform;
 		InISMComponent->GetInstanceTransform(InstanceIndex, InstanceTransform);
 		Instances.Emplace(InstanceTransform);
 
-		const int32 NumCustomDataFloats = InISMComponent->NumCustomDataFloats;
-		FCustomPrimitiveData InstanceCustomData;
-		InstanceCustomData.Data.SetNumUninitialized(NumCustomDataFloats);
-		check(InISMComponent->PerInstanceSMCustomData.GetTypeSize() == InstanceCustomData.Data.GetTypeSize());
-		FMemory::Memcpy(InstanceCustomData.Data.GetData(), &InISMComponent->PerInstanceSMCustomData[InstanceIndex * NumCustomDataFloats], NumCustomDataFloats * InstanceCustomData.Data.GetTypeSize());
+		if (NumCustomDataFloats > 0)
+		{
+			if (ensure(InISMComponent->PerInstanceSMCustomData.IsValidIndex(InstanceIndex * NumCustomDataFloats)))
+			{
+				FCustomPrimitiveData InstanceCustomData;
+				InstanceCustomData.Data.SetNumUninitialized(NumCustomDataFloats);
+				
+				check(InISMComponent->PerInstanceSMCustomData.GetTypeSize() == InstanceCustomData.Data.GetTypeSize());
+
+				void* Dest = 
+				FMemory::Memcpy(InstanceCustomData.Data.GetData(), &InISMComponent->PerInstanceSMCustomData[InstanceIndex * NumCustomDataFloats], NumCustomDataFloats * InstanceCustomData.Data.GetTypeSize());
 
-		InstancesCustomPrimitiveData.Emplace(MoveTemp(InstanceCustomData));
+				InstancesCustomPrimitiveData.Emplace(MoveTemp(InstanceCustomData));
+			}
+		}
 	}
 }
 
