World Partition: refactor deprecated functions ForEachActorWithLoading in FWorldPartitionHelpers to minimize code duplication and prepare for removal.

#rnx

diff --git a/Engine/Source/Runtime/Engine/Private/WorldPartition/WorldPartitionHelpers.cpp b/Engine/Source/Runtime/Engine/Private/WorldPartition/WorldPartitionHelpers.cpp
--- a/Engine/Source/Runtime/Engine/Private/WorldPartition/WorldPartitionHelpers.cpp
+++ b/Engine/Source/Runtime/Engine/Private/WorldPartition/WorldPartitionHelpers.cpp
@@ -80,26 +80,6 @@ namespace WorldPartitionHelpers
 			InOutActorReferences[ActorGuid] = FWorldPartitionReference(WorldPartition, ActorGuid);
 		}
 	}
-
-	bool ForEachActorWithLoadingBody(const FGuid& ActorGuid, UWorldPartition* WorldPartition, TFunctionRef<bool(const FWorldPartitionActorDesc*)> Func, TFunctionRef<void()> OnReleasingActorReferences, bool bGCPerActor, TMap<FGuid, FWorldPartitionReference>& ActorReferences)
-	{
-		WorldPartitionHelpers::LoadReferences(WorldPartition, ActorGuid, ActorReferences);
-
-		FWorldPartitionReference ActorReference(WorldPartition, ActorGuid);
-		if (!Func(ActorReference.Get()))
-		{
-			return false;
-		}
-
-		if (bGCPerActor || FWorldPartitionHelpers::HasExceededMaxMemory())
-		{
-			OnReleasingActorReferences();
-			ActorReferences.Empty();
-			FWorldPartitionHelpers::DoCollectGarbage();
-		}
-
-		return true;
-	}
 }
 
 FWorldPartitionHelpers::FForEachActorWithLoadingParams::FForEachActorWithLoadingParams()
@@ -109,55 +89,37 @@ FWorldPartitionHelpers::FForEachActorWithLoadingParams::FForEachActorWithLoading
 
 void FWorldPartitionHelpers::ForEachActorWithLoading(UWorldPartition* WorldPartition, TSubclassOf<AActor> ActorClass, TFunctionRef<bool(const FWorldPartitionActorDesc*)> Func, TFunctionRef<void()> OnReleasingActorReferences, bool bGCPerActor)
 {
-	TMap<FGuid, FWorldPartitionReference> ActorReferences;
-
-	ForEachActorDesc(WorldPartition, [&](const FWorldPartitionActorDesc* ActorDesc)
-	{
-		UClass* ActorDescClass = ResolveActorDescClass(ActorDesc);
-
-		if (ActorDescClass->IsChildOf(ActorClass))
-		{
-			return WorldPartitionHelpers::ForEachActorWithLoadingBody(ActorDesc->GetGuid(), WorldPartition, Func, OnReleasingActorReferences, bGCPerActor, ActorReferences);
-		}
-
-		return true;
-	});
-
-	OnReleasingActorReferences();
-	ActorReferences.Empty();
-	DoCollectGarbage();
+	FForEachActorWithLoadingParams Params;
+	Params.ActorClass = ActorClass;
+	Params.OnPreGarbageCollect = [&OnReleasingActorReferences]() { OnReleasingActorReferences(); };
+	ForEachActorWithLoading(WorldPartition, Func, Params);
 }
 
 void FWorldPartitionHelpers::ForEachActorWithLoading(UWorldPartition* WorldPartition, const TArray<FGuid>& ActorGuids, TFunctionRef<bool(const FWorldPartitionActorDesc*)> Func, TFunctionRef<void()> OnReleasingActorReferences, bool bGCPerActor)
 {
-	TMap<FGuid, FWorldPartitionReference> ActorReferences;
-
-	for(const FGuid& ActorGuid : ActorGuids)
-	{
-		if (!WorldPartitionHelpers::ForEachActorWithLoadingBody(ActorGuid, WorldPartition, Func, OnReleasingActorReferences, bGCPerActor, ActorReferences))
-		{
-			break;
-		}
-	}
-
-	OnReleasingActorReferences();
-	ActorReferences.Empty();
-	DoCollectGarbage();
+	TSet<FGuid> ActorGuidsSet(ActorGuids);
+	FForEachActorWithLoadingParams Params;
+	Params.FilterActorDesc = [&ActorGuidsSet](const FWorldPartitionActorDesc* ActorDesc) -> bool { return ActorGuidsSet.Contains(ActorDesc->GetGuid());	};
+	Params.OnPreGarbageCollect = [&OnReleasingActorReferences]() { OnReleasingActorReferences(); };	
+	ForEachActorWithLoading(WorldPartition, Func, Params);
 }
 
 void FWorldPartitionHelpers::ForEachActorWithLoading(UWorldPartition* WorldPartition, TFunctionRef<bool(const FWorldPartitionActorDesc*)> Func, const FForEachActorWithLoadingParams& Params)
 {
 	TMap<FGuid, FWorldPartitionReference> ActorReferences;
 
-	auto CallOnPreGarbageCollect = [&Params]()
+	auto CallGarbageCollect = [&Params, &ActorReferences]()
 	{
 		if (Params.OnPreGarbageCollect)
 		{
 			Params.OnPreGarbageCollect();
 		}
+
+		ActorReferences.Empty();
+		DoCollectGarbage();
 	};
 
-	ForEachActorDesc(WorldPartition, [&Func, &CallOnPreGarbageCollect, &Params, &ActorReferences, WorldPartition](const FWorldPartitionActorDesc* ActorDesc)
+	ForEachActorDesc(WorldPartition, [&Func, &CallGarbageCollect, &Params, &ActorReferences, WorldPartition](const FWorldPartitionActorDesc* ActorDesc)
 	{
 		UClass* ActorDescClass = ResolveActorDescClass(ActorDesc);
 
@@ -165,16 +127,27 @@ void FWorldPartitionHelpers::ForEachActorWithLoading(UWorldPartition* WorldParti
 		{
 			if (!Params.FilterActorDesc || Params.FilterActorDesc(ActorDesc))
 			{
-				return WorldPartitionHelpers::ForEachActorWithLoadingBody(ActorDesc->GetGuid(), WorldPartition, Func, [&CallOnPreGarbageCollect, &ActorReferences]() { CallOnPreGarbageCollect(); }, Params.bGCPerActor, ActorReferences);
+				WorldPartitionHelpers::LoadReferences(WorldPartition, ActorDesc->GetGuid(), ActorReferences);
+
+				FWorldPartitionReference ActorReference(WorldPartition, ActorDesc->GetGuid());
+				if (!Func(ActorReference.Get()))
+				{
+					return false;
+				}
+
+				if (Params.bGCPerActor || FWorldPartitionHelpers::HasExceededMaxMemory())
+				{
+					CallGarbageCollect();
+				}
+
+				return true;
 			}
 		}
 
 		return true;
 	});
 
-	CallOnPreGarbageCollect();
-	ActorReferences.Empty();
-	DoCollectGarbage();
+	CallGarbageCollect();
 }
 
 bool FWorldPartitionHelpers::HasExceededMaxMemory()
