World Partition: avoid loading all actors BP classes when filtering by a native class.

#rnx

diff --git a/Engine/Source/Runtime/Engine/Private/WorldPartition/WorldPartitionHelpers.cpp b/Engine/Source/Runtime/Engine/Private/WorldPartition/WorldPartitionHelpers.cpp
--- a/Engine/Source/Runtime/Engine/Private/WorldPartition/WorldPartitionHelpers.cpp
+++ b/Engine/Source/Runtime/Engine/Private/WorldPartition/WorldPartitionHelpers.cpp
@@ -11,32 +11,35 @@ DEFINE_LOG_CATEGORY_STATIC(LogWorldPartitionHelpers, Log, All);
 
 #if WITH_EDITOR
 
-UClass* FWorldPartitionHelpers::ResolveActorDescClass(const FWorldPartitionActorDesc* ActorDesc)
+bool FWorldPartitionHelpers::IsActorDescClassCompatibleWith(const FWorldPartitionActorDesc* ActorDesc, const UClass* Class)
 {
+	check(Class);
+
 	UClass* ActorNativeClass = ActorDesc->GetActorNativeClass();
 	UClass* ActorBaseClass = ActorNativeClass;
 
-	if (FName ActorBaseClassName = ActorDesc->GetBaseClass(); !ActorBaseClassName.IsNone())
+	if (!Class->IsNative())
 	{
-		ActorBaseClass = LoadClass<AActor>(nullptr, *ActorBaseClassName.ToString(), nullptr, LOAD_None, nullptr);
-
-		if (!ActorBaseClass)
+		if (FName ActorBaseClassName = ActorDesc->GetBaseClass(); !ActorBaseClassName.IsNone())
 		{
-			UE_LOG(LogWorldPartitionHelpers, Warning, TEXT("Failed to find actor base class: %s."), *ActorBaseClassName.ToString());
-			ActorBaseClass = ActorNativeClass;
+			ActorBaseClass = LoadClass<AActor>(nullptr, *ActorBaseClassName.ToString(), nullptr, LOAD_None, nullptr);
+
+			if (!ActorBaseClass)
+			{
+				UE_LOG(LogWorldPartitionHelpers, Warning, TEXT("Failed to find actor base class: %s."), *ActorBaseClassName.ToString());
+				ActorBaseClass = ActorNativeClass;
+			}
 		}
 	}
 
-	return ActorBaseClass;
+	return ActorBaseClass->IsChildOf(Class);
 }
 
 void FWorldPartitionHelpers::ForEachIntersectingActorDesc(UWorldPartition* WorldPartition, const FBox& Box, TSubclassOf<AActor> ActorClass, TFunctionRef<bool(const FWorldPartitionActorDesc*)> Func)
 {
 	WorldPartition->EditorHash->ForEachIntersectingActor(Box, [&ActorClass, Func](const FWorldPartitionActorDesc* ActorDesc)
 	{
-		UClass* ActorDescClass = ResolveActorDescClass(ActorDesc);
-
-		if (ActorDescClass->IsChildOf(ActorClass))
+		if (IsActorDescClassCompatibleWith(ActorDesc, ActorClass))
 		{
 			Func(ActorDesc);
 		}
@@ -47,9 +50,7 @@ void FWorldPartitionHelpers::ForEachActorDesc(UWorldPartition* WorldPartition, T
 {
 	for (UActorDescContainer::TConstIterator<> ActorDescIterator(WorldPartition); ActorDescIterator; ++ActorDescIterator)
 	{
-		UClass* ActorDescClass = ResolveActorDescClass(*ActorDescIterator);
-
-		if (ActorDescClass->IsChildOf(ActorClass))
+		if (IsActorDescClassCompatibleWith(*ActorDescIterator, ActorClass))
 		{
 			if (!Func(*ActorDescIterator))
 			{
@@ -119,33 +120,27 @@ void FWorldPartitionHelpers::ForEachActorWithLoading(UWorldPartition* WorldParti
 		DoCollectGarbage();
 	};
 
-	ForEachActorDesc(WorldPartition, [&Func, &CallGarbageCollect, &Params, &ActorReferences, WorldPartition](const FWorldPartitionActorDesc* ActorDesc)
+	for (UActorDescContainer::TConstIterator<> ActorDescIterator(WorldPartition); ActorDescIterator; ++ActorDescIterator)
 	{
-		UClass* ActorDescClass = ResolveActorDescClass(ActorDesc);
-
-		if (ActorDescClass->IsChildOf(Params.ActorClass))
+		if (IsActorDescClassCompatibleWith(*ActorDescIterator, Params.ActorClass))
 		{
-			if (!Params.FilterActorDesc || Params.FilterActorDesc(ActorDesc))
+			if (!Params.FilterActorDesc || Params.FilterActorDesc(*ActorDescIterator))
 			{
-				WorldPartitionHelpers::LoadReferences(WorldPartition, ActorDesc->GetGuid(), ActorReferences);
+				WorldPartitionHelpers::LoadReferences(WorldPartition, ActorDescIterator->GetGuid(), ActorReferences);
 
-				FWorldPartitionReference ActorReference(WorldPartition, ActorDesc->GetGuid());
+				FWorldPartitionReference ActorReference(WorldPartition, ActorDescIterator->GetGuid());
 				if (!Func(ActorReference.Get()))
 				{
-					return false;
+					break;
 				}
 
 				if (Params.bGCPerActor || FWorldPartitionHelpers::HasExceededMaxMemory())
 				{
 					CallGarbageCollect();
 				}
-
-				return true;
 			}
 		}
-
-		return true;
-	});
+	}
 
 	CallGarbageCollect();
 }
diff --git a/Engine/Source/Runtime/Engine/Public/WorldPartition/WorldPartitionHelpers.h b/Engine/Source/Runtime/Engine/Public/WorldPartition/WorldPartitionHelpers.h
--- a/Engine/Source/Runtime/Engine/Public/WorldPartition/WorldPartitionHelpers.h
+++ b/Engine/Source/Runtime/Engine/Public/WorldPartition/WorldPartitionHelpers.h
@@ -12,7 +12,7 @@ class FWorldPartitionActorDesc;
 
 class ENGINE_API FWorldPartitionHelpers
 {
-	static UClass* ResolveActorDescClass(const FWorldPartitionActorDesc* ActorDesc);
+	static bool IsActorDescClassCompatibleWith(const FWorldPartitionActorDesc* ActorDesc, const UClass* Class);
 
 public:
 	template <class ActorClass = AActor>
