World Partition : Apply level transform on actors already part of level when initializing a partitioned level with a transform.

diff --git a/Engine/Source/Runtime/Engine/Private/WorldPartition/WorldPartition.cpp b/Engine/Source/Runtime/Engine/Private/WorldPartition/WorldPartition.cpp
--- a/Engine/Source/Runtime/Engine/Private/WorldPartition/WorldPartition.cpp
+++ b/Engine/Source/Runtime/Engine/Private/WorldPartition/WorldPartition.cpp
@@ -506,6 +506,24 @@ void UWorldPartition::Initialize(UWorld* InWorld, const FTransform& InTransform)
 			WorldDataLayers = AWorldDataLayers::Create(World);
 			World->SetWorldDataLayers(WorldDataLayers);
 		}
+
+		// Apply level transform on actors already part of the level
+		if (GetInstanceTransformPtr())
+		{
+			check(!OuterWorld->PersistentLevel->bAlreadyMovedActors);
+			for (AActor* Actor : OuterWorld->PersistentLevel->Actors)
+			{
+				if (Actor)
+				{
+					FLevelUtils::FApplyLevelTransformParams TransformParams(Actor->GetLevel(), GetInstanceTransform());
+					TransformParams.Actor = Actor;
+					TransformParams.bDoPostEditMove = true;
+					FLevelUtils::ApplyLevelTransform(TransformParams);
+				}
+			}
+			// Flag Level's bAlreadyMovedActors to true so that ULevelStreaming::PrepareLoadedLevel won't reapply the same transform again.
+			OuterWorld->PersistentLevel->bAlreadyMovedActors = true;
+		}
 	}
 
 	if (bIsEditor && !bIsCooking)
@@ -531,14 +549,6 @@ void UWorldPartition::Initialize(UWorld* InWorld, const FTransform& InTransform)
 			GEditor->OnEditorClose().AddUObject(this, &UWorldPartition::SavePerUserSettings);
 		}
 	}
-
-	if (bIsEditor)
-	{
-		// Because WorldPartition applies instance transforms to each actor (see ApplyActorTransform)
-		// We need to flag Level's bAlreadyMovedActors to true so that ULevelStreaming::PrepareLoadedLevel won't reapply the same transform again.
-		// We asssume that non-external actors won't be transformed (which is fine since the only 2 cases are the worldsettings and the default brush).
-		OuterWorld->PersistentLevel->bAlreadyMovedActors = true;
-	}
 #endif //WITH_EDITOR
 
 	InitState = EWorldPartitionInitState::Initialized;
