World Partition : preload WorldDataLayers actor also in -game for partitioned levels that were not saved with the reference to the actor.

diff --git a/Engine/Source/Runtime/Engine/Private/WorldPartition/WorldPartition.cpp b/Engine/Source/Runtime/Engine/Private/WorldPartition/WorldPartition.cpp
--- a/Engine/Source/Runtime/Engine/Private/WorldPartition/WorldPartition.cpp
+++ b/Engine/Source/Runtime/Engine/Private/WorldPartition/WorldPartition.cpp
@@ -481,7 +481,7 @@ void UWorldPartition::Initialize(UWorld* InWorld, const FTransform& InTransform)
 		}
 	}
 
-	if (bIsEditor)
+	if (bIsEditor || bIsGame)
 	{
 		// Here we need to flush any async loading before starting any synchronous load (mixing synchronous and asynchronous load on the same package is not handled properly).
 		// Also, this will ensure that FindObject will work on an external actor that was loading asynchronously.
@@ -496,7 +496,10 @@ void UWorldPartition::Initialize(UWorld* InWorld, const FTransform& InTransform)
 				break;
 			}
 		}
-		
+	}
+	
+	if (bIsEditor)
+	{
 		// Repair existing maps with no UWorld::WorldDataLayers actor.
 		AWorldDataLayers* WorldDataLayers = World->GetWorldDataLayers();
 		if (!WorldDataLayers)
