Do not read past the mip tail when loading streaming mips from DDC. There is no data for mips past the mip tail so loading from DDC will be attempted because they don't have any bulk data but always fail causing false cache misses to be reported.

diff --git a/Engine/Source/Runtime/Engine/Private/TextureDerivedData.cpp b/Engine/Source/Runtime/Engine/Private/TextureDerivedData.cpp
--- a/Engine/Source/Runtime/Engine/Private/TextureDerivedData.cpp
+++ b/Engine/Source/Runtime/Engine/Private/TextureDerivedData.cpp
@@ -1355,11 +1355,13 @@ static bool LoadDerivedStreamingMips(FTexturePlatformData& PlatformData, int32 F
 	bool bMiss = false;
 
 	TIndirectArray<FTexture2DMipMap>& Mips = PlatformData.Mips;
+	const int32 ReadableMipCount = Mips.Num() - (PlatformData.GetNumMipsInTail() > 0 ? PlatformData.GetNumMipsInTail() - 1 : 0);
+
 	if (PlatformData.DerivedDataKey.IsType<FString>())
 	{
 		TArray<FCacheGetValueRequest, TInlineAllocator<MAX_TEXTURE_MIP_COUNT>> Requests;
 
-		for (int32 MipIndex = FirstMipToLoad; MipIndex < Mips.Num(); ++MipIndex)
+		for (int32 MipIndex = FirstMipToLoad; MipIndex < ReadableMipCount; ++MipIndex)
 		{
 			const FTexture2DMipMap& Mip = Mips[MipIndex];
 			if (Mip.IsPagedToDerivedData() && !Mip.BulkData.IsBulkDataLoaded())
@@ -1399,7 +1401,7 @@ static bool LoadDerivedStreamingMips(FTexturePlatformData& PlatformData, int32 F
 		TArray<FCacheGetChunkRequest, TInlineAllocator<MAX_TEXTURE_MIP_COUNT>> Requests;
 
 		const FCacheKey& Key = *PlatformData.DerivedDataKey.Get<UE::DerivedData::FCacheKeyProxy>().AsCacheKey();
-		for (int32 MipIndex = FirstMipToLoad; MipIndex < Mips.Num(); ++MipIndex)
+		for (int32 MipIndex = FirstMipToLoad; MipIndex < ReadableMipCount; ++MipIndex)
 		{
 			const FTexture2DMipMap& Mip = Mips[MipIndex];
 			if (Mip.IsPagedToDerivedData() && !Mip.BulkData.IsBulkDataLoaded())
