Fix for static analysis warning for ParticleComponents.cpp

diff --git a/Engine/Source/Runtime/Engine/Private/Particles/ParticleComponents.cpp b/Engine/Source/Runtime/Engine/Private/Particles/ParticleComponents.cpp
--- a/Engine/Source/Runtime/Engine/Private/Particles/ParticleComponents.cpp
+++ b/Engine/Source/Runtime/Engine/Private/Particles/ParticleComponents.cpp
@@ -7743,20 +7743,21 @@ void UParticleLODLevel::GetStreamingMeshInfo(const FBoxSphereBounds& Bounds, TAr
 {
 	if (bEnabled)
 	{
-		const UParticleModuleTypeDataMesh* MeshTypeData = Cast<UParticleModuleTypeDataMesh>(TypeDataModule);
-
-		if (UStaticMesh* Mesh = MeshTypeData ? MeshTypeData->Mesh : nullptr)
+		if (const UParticleModuleTypeDataMesh* MeshTypeData = Cast<UParticleModuleTypeDataMesh>(TypeDataModule))
 		{
-			if (Mesh->RenderResourceSupportsStreaming() && Mesh->GetRenderAssetType() == EStreamableRenderAssetType::StaticMesh)
+			if (UStaticMesh* Mesh = MeshTypeData->Mesh)
 			{
-				const FBoxSphereBounds MeshBounds = Mesh->GetBounds();
-				const FBoxSphereBounds StreamingBounds = FBoxSphereBounds(
-					Bounds.Origin + MeshBounds.Origin,
-					MeshBounds.BoxExtent * MeshTypeData->LODSizeScale,
-					MeshBounds.SphereRadius * MeshTypeData->LODSizeScale);
-				const float MeshTexelFactor = MeshBounds.SphereRadius * 2.0f;
-
-				new (OutStreamingRenderAssets) FStreamingRenderAssetPrimitiveInfo(Mesh, StreamingBounds, MeshTexelFactor);
+				if (Mesh->RenderResourceSupportsStreaming() && Mesh->GetRenderAssetType() == EStreamableRenderAssetType::StaticMesh)
+				{
+					const FBoxSphereBounds MeshBounds = Mesh->GetBounds();
+					const FBoxSphereBounds StreamingBounds = FBoxSphereBounds(
+						Bounds.Origin + MeshBounds.Origin,
+						MeshBounds.BoxExtent * MeshTypeData->LODSizeScale,
+						MeshBounds.SphereRadius * MeshTypeData->LODSizeScale);
+					const float MeshTexelFactor = MeshBounds.SphereRadius * 2.0f;
+
+					new (OutStreamingRenderAssets) FStreamingRenderAssetPrimitiveInfo(Mesh, StreamingBounds, MeshTexelFactor);
+				}
 			}
 		}
 	}
