World Partition Actor Descriptors: fix case where the actor pointer was stale and wouldn't be requeried.

#rnx

diff --git a/Engine/Source/Runtime/Engine/Private/WorldPartition/WorldPartitionActorDesc.cpp b/Engine/Source/Runtime/Engine/Private/WorldPartition/WorldPartitionActorDesc.cpp
--- a/Engine/Source/Runtime/Engine/Private/WorldPartition/WorldPartitionActorDesc.cpp
+++ b/Engine/Source/Runtime/Engine/Private/WorldPartition/WorldPartitionActorDesc.cpp
@@ -390,7 +390,7 @@ FName FWorldPartitionActorDesc::GetDisplayClassName() const
 
 bool FWorldPartitionActorDesc::IsLoaded(bool bEvenIfPendingKill) const
 {
-	if (ActorPtr.IsExplicitlyNull())
+	if (ActorPtr.IsExplicitlyNull() || ActorPtr.IsStale())
 	{
 		ActorPtr = FindObject<AActor>(nullptr, *ActorPath.ToString());
 	}
@@ -400,7 +400,7 @@ bool FWorldPartitionActorDesc::IsLoaded(bool bEvenIfPendingKill) const
 
 AActor* FWorldPartitionActorDesc::GetActor(bool bEvenIfPendingKill, bool bEvenIfUnreachable) const
 {
-	if (ActorPtr.IsExplicitlyNull())
+	if (ActorPtr.IsExplicitlyNull() || ActorPtr.IsStale())
 	{
 		ActorPtr = FindObject<AActor>(nullptr, *ActorPath.ToString());
 	}
@@ -410,7 +410,7 @@ AActor* FWorldPartitionActorDesc::GetActor(bool bEvenIfPendingKill, bool bEvenIf
 
 AActor* FWorldPartitionActorDesc::Load() const
 {
-	if (ActorPtr.IsExplicitlyNull())
+	if (ActorPtr.IsExplicitlyNull() || ActorPtr.IsStale())
 	{
 		// First, try to find the existing actor which could have been loaded by another actor (through standard serialization)
 		ActorPtr = FindObject<AActor>(nullptr, *ActorPath.ToString());
