[Replication]
Improve replication of root motion to avoid issues with the character not evaluating environnement properly

#RNX

diff --git a/Engine/Source/Runtime/Engine/Private/Components/CharacterMovementComponent.cpp b/Engine/Source/Runtime/Engine/Private/Components/CharacterMovementComponent.cpp
--- a/Engine/Source/Runtime/Engine/Private/Components/CharacterMovementComponent.cpp
+++ b/Engine/Source/Runtime/Engine/Private/Components/CharacterMovementComponent.cpp
@@ -1746,6 +1746,13 @@ void UCharacterMovementComponent::SimulatedTick(float DeltaSeconds)
 			CharacterOwner->RootMotionRepMoves.Reset();
 		}
 
+		// Update replicated movement mode.
+		if (bNetworkMovementModeChanged)
+		{
+			ApplyNetworkMovementMode(CharacterOwner->GetReplicatedMovementMode());
+			bNetworkMovementModeChanged = false;
+		}
+
 		// Perform movement
 		PerformMovement(DeltaSeconds);
 
@@ -1763,7 +1770,6 @@ void UCharacterMovementComponent::SimulatedTick(float DeltaSeconds)
 		// (Root Motion could leave Velocity out of sync w/ ReplicatedMovement)
 		if( bWasSimulatingRootMotion )
 		{
-			bWasSimulatingRootMotion = false;
 			CharacterOwner->RootMotionRepMoves.Empty();
 			CharacterOwner->OnRep_ReplicatedMovement();
 			CharacterOwner->OnRep_ReplicatedBasedMovement();
@@ -1782,6 +1788,13 @@ void UCharacterMovementComponent::SimulatedTick(float DeltaSeconds)
 				const FScopedPreventAttachedComponentMove PreventMeshMovement(bPreventMeshMovement ? Mesh : nullptr);
 				if (CharacterOwner->IsPlayingRootMotion())
 				{
+					// Update replicated movement mode.
+					if (bNetworkMovementModeChanged)
+					{
+						ApplyNetworkMovementMode(CharacterOwner->GetReplicatedMovementMode());
+						bNetworkMovementModeChanged = false;
+					}
+
 					PerformMovement(DeltaSeconds);
 				}
 				else
@@ -1805,6 +1818,11 @@ void UCharacterMovementComponent::SimulatedTick(float DeltaSeconds)
 				}
 			}
 		}
+
+		if (bWasSimulatingRootMotion)
+		{
+			bWasSimulatingRootMotion = false;
+		}
 	}
 
 	// Smooth mesh location after moving the capsule above.
