World Partition: fix crash when the Level Script references editor-only actors.

#rnx

diff --git a/Engine/Source/Runtime/Engine/Private/WorldPartition/WorldPartitionStreamingGeneration.cpp b/Engine/Source/Runtime/Engine/Private/WorldPartition/WorldPartitionStreamingGeneration.cpp
--- a/Engine/Source/Runtime/Engine/Private/WorldPartition/WorldPartitionStreamingGeneration.cpp
+++ b/Engine/Source/Runtime/Engine/Private/WorldPartition/WorldPartitionStreamingGeneration.cpp
@@ -214,18 +214,19 @@ class FWorldPartitionStreamingGenerator
 
 					for (AActor* Actor : LevelScriptExternalActorReferences)
 					{
-						FWorldPartitionActorDescView& ActorDescView = ContainerDescriptor.ActorDescViewMap.FindChecked(Actor->GetActorGuid());
-
-						if (ActorDescView.GetIsSpatiallyLoaded())
+						if (FWorldPartitionActorDescView* ActorDescView = ContainerDescriptor.ActorDescViewMap.Find(Actor->GetActorGuid()))
 						{
-							ErrorHandler->OnInvalidReferenceLevelScriptStreamed(ActorDescView);
-							ActorDescView.SetForcedNonSpatiallyLoaded();
-						}
+							if (ActorDescView->GetIsSpatiallyLoaded())
+							{
+								ErrorHandler->OnInvalidReferenceLevelScriptStreamed(*ActorDescView);
+								ActorDescView->SetForcedNonSpatiallyLoaded();
+							}
 
-						if (ActorDescView.GetRuntimeDataLayers().Num())
-						{
-							ErrorHandler->OnInvalidReferenceLevelScriptDataLayers(ActorDescView);
-							ActorDescView.SetInvalidDataLayers();
+							if (ActorDescView->GetRuntimeDataLayers().Num())
+							{
+								ErrorHandler->OnInvalidReferenceLevelScriptDataLayers(*ActorDescView);
+								ActorDescView->SetInvalidDataLayers();
+							}
 						}
 					}
 				}
