Deprecated RHICreateTargetableShaderResource utility method, since it's not actually used to create multisampled resources. Removed instances where it appeared multiple textures were being created, but only one was actually created. Cleaned up FTextureRenderTargetResource derived classes to use TextureRHI and TextureRenderTargetRHI.

diff --git a/Engine/Plugins/Media/AndroidCamera/Source/AndroidCamera/Private/Player/AndroidCameraTextureSample.h b/Engine/Plugins/Media/AndroidCamera/Source/AndroidCamera/Private/Player/AndroidCameraTextureSample.h
--- a/Engine/Plugins/Media/AndroidCamera/Source/AndroidCamera/Private/Player/AndroidCameraTextureSample.h
+++ b/Engine/Plugins/Media/AndroidCamera/Source/AndroidCamera/Private/Player/AndroidCameraTextureSample.h
@@ -146,9 +146,10 @@ public:
 			FRHITextureCreateDesc::Create2D(TEXT("DummyTexture2D"))
 			.SetExtent(Dim)
 			.SetFormat(PF_B8G8R8A8)
-			.SetFlags(ETextureCreateFlags::Dynamic | ETextureCreateFlags::SRGB);
+			.SetFlags(ETextureCreateFlags::Dynamic | ETextureCreateFlags::SRGB | ETextureCreateFlags::RenderTargetable | ETextureCreateFlags::ShaderResource)
+			.SetInitialState(ERHIAccess::SRVMask);
 
-		RHICreateTargetableShaderResource(Desc, ETextureCreateFlags::RenderTargetable, Texture);
+		Texture = RHICreateTexture(Desc);
 
 		return Texture;
 	}
diff --git a/Engine/Plugins/Media/AndroidMedia/Source/AndroidMedia/Private/Player/AndroidMediaTextureSample.h b/Engine/Plugins/Media/AndroidMedia/Source/AndroidMedia/Private/Player/AndroidMediaTextureSample.h
--- a/Engine/Plugins/Media/AndroidMedia/Source/AndroidMedia/Private/Player/AndroidMediaTextureSample.h
+++ b/Engine/Plugins/Media/AndroidMedia/Source/AndroidMedia/Private/Player/AndroidMediaTextureSample.h
@@ -147,9 +147,10 @@ public:
 			FRHITextureCreateDesc::Create2D(TEXT("DummyTexture2D"))
 			.SetExtent(Dim)
 			.SetFormat(PF_B8G8R8A8)
-			.SetFlags(ETextureCreateFlags::Dynamic | ETextureCreateFlags::SRGB);
+			.SetFlags(ETextureCreateFlags::Dynamic | ETextureCreateFlags::SRGB | ETextureCreateFlags::RenderTargetable | ETextureCreateFlags::ShaderResource)
+			.SetInitialState(ERHIAccess::SRVMask);
 
-		RHICreateTargetableShaderResource(Desc, ETextureCreateFlags::RenderTargetable, Texture);
+		Texture = RHICreateTexture(Desc);
 
 		return Texture;
 	}
diff --git a/Engine/Plugins/Media/BinkMedia/Source/BinkMediaPlayer/Private/Assets/BinkMediaTextureResource.cpp b/Engine/Plugins/Media/BinkMedia/Source/BinkMediaPlayer/Private/Assets/BinkMediaTextureResource.cpp
--- a/Engine/Plugins/Media/BinkMedia/Source/BinkMediaPlayer/Private/Assets/BinkMediaTextureResource.cpp
+++ b/Engine/Plugins/Media/BinkMedia/Source/BinkMediaPlayer/Private/Assets/BinkMediaTextureResource.cpp
@@ -43,9 +43,10 @@ void FBinkMediaTextureResource::InitDynamicRHI()
 		FRHITextureCreateDesc::Create2D(DebugName)
 		.SetExtent(w, h)
 		.SetFormat(PixelFormat)
-		.SetFlags(TexCreateFlags);
+		.SetFlags(TexCreateFlags | ETextureCreateFlags::RenderTargetable | ETextureCreateFlags::ShaderResource)
+		.SetInitialState(ERHIAccess::SRVMask);
 
-	RHICreateTargetableShaderResource(Desc, ETextureCreateFlags::RenderTargetable, TextureRHI);
+	TextureRHI = RHICreateTexture(Desc);
 
 	// Don't bother updating if its not a valid video
 	if (Owner->GetSurfaceWidth() && Owner->GetSurfaceHeight()) 
diff --git a/Engine/Plugins/Media/ElectraUtil/Source/ElectraSamples/Private/ElectraTextureSampleAndroid.cpp b/Engine/Plugins/Media/ElectraUtil/Source/ElectraSamples/Private/ElectraTextureSampleAndroid.cpp
--- a/Engine/Plugins/Media/ElectraUtil/Source/ElectraSamples/Private/ElectraTextureSampleAndroid.cpp
+++ b/Engine/Plugins/Media/ElectraUtil/Source/ElectraSamples/Private/ElectraTextureSampleAndroid.cpp
@@ -426,9 +426,10 @@ void FElectraTextureSample::InitializeTexture()
 		FRHITextureCreateDesc::Create2D(TEXT("FElectraTextureSample"))
 		.SetExtent(Dim)
 		.SetFormat(PF_B8G8R8A8)
-		.SetFlags(ETextureCreateFlags::Dynamic | ETextureCreateFlags::SRGB);
+		.SetFlags(ETextureCreateFlags::Dynamic | ETextureCreateFlags::SRGB | ETextureCreateFlags::RenderTargetable | ETextureCreateFlags::ShaderResource)
+		.SetInitialState(ERHIAccess::SRVMask);
 
-	RHICreateTargetableShaderResource(Desc, ETextureCreateFlags::RenderTargetable, Texture);
+	Texture = RHICreateTexture(Desc);
 
 	return;
 }
diff --git a/Engine/Plugins/Media/ElectraUtil/Source/ElectraSamples/Public/GStreamer/ElectraTextureSample.h b/Engine/Plugins/Media/ElectraUtil/Source/ElectraSamples/Public/GStreamer/ElectraTextureSample.h
--- a/Engine/Plugins/Media/ElectraUtil/Source/ElectraSamples/Public/GStreamer/ElectraTextureSample.h
+++ b/Engine/Plugins/Media/ElectraUtil/Source/ElectraSamples/Public/GStreamer/ElectraTextureSample.h
@@ -38,9 +38,10 @@ public:
 			FRHITextureCreateDesc::Create2D(TEXT("DummyTexture2D"))
 			.SetExtent(TotalSize)
 			.SetFormat(PF_B8G8R8A8)
-			.SetFlags(ETextureCreateFlags::Dynamic | ETextureCreateFlags::SRGB);
+			.SetFlags(ETextureCreateFlags::Dynamic | ETextureCreateFlags::SRGB | ETextureCreateFlags::RenderTargetable | ETextureCreateFlags::ShaderResource)
+			.SetInitialState(ERHIAccess::SRVMask);
 
-		RHICreateTargetableShaderResource(Desc, ETextureCreateFlags::RenderTargetable, Texture);
+		Texture = RHICreateTexture(Desc);
 	}	
 
 public:
diff --git a/Engine/Plugins/Media/MediaIOFramework/Source/MediaIOCore/Private/Player/MediaIOCorePlayerBase.cpp b/Engine/Plugins/Media/MediaIOFramework/Source/MediaIOCore/Private/Player/MediaIOCorePlayerBase.cpp
--- a/Engine/Plugins/Media/MediaIOFramework/Source/MediaIOCore/Private/Player/MediaIOCorePlayerBase.cpp
+++ b/Engine/Plugins/Media/MediaIOFramework/Source/MediaIOCore/Private/Player/MediaIOCorePlayerBase.cpp
@@ -634,16 +634,20 @@ void FMediaIOCorePlayerBase::CreateAndRegisterTextures(const IMediaOptions* Opti
 					checkf(false, TEXT("Format not supported"));
 				}
 
-				TRefCountPtr<FRHITexture2D> DummyTexture2DRHI;
-
-				ETextureCreateFlags CreateFlags = TexCreate_Shared;
+				ETextureCreateFlags CreateFlags = ETextureCreateFlags::Shared;
 				if (RHIGetInterfaceType() == ERHIInterfaceType::Vulkan)
 				{
-					CreateFlags = TexCreate_External;
+					CreateFlags = ETextureCreateFlags::External;
 				}
 
-				FRHITextureCreateDesc CreateDesc = FRHITextureCreateDesc::Create2D(*TextureName, FIntPoint(TextureWidth, TextureHeight), InputFormat, FClearValueBinding::White, CreateFlags);
-				RHICreateTargetableShaderResource(CreateDesc, TexCreate_RenderTargetable, RHITexture, DummyTexture2DRHI);
+				FRHITextureCreateDesc CreateDesc = FRHITextureCreateDesc::Create2D(*TextureName)
+					.SetExtent(TextureWidth, TextureHeight)
+					.SetFormat(InputFormat)
+					.SetClearValue(FClearValueBinding::White)
+					.SetFlags(CreateFlags | ETextureCreateFlags::RenderTargetable | ETextureCreateFlags::ShaderResource)
+					.SetInitialState(ERHIAccess::SRVMask);
+
+				RHITexture = RHICreateTexture(CreateDesc);
 
 				UE_LOG(LogMediaIOCore, Verbose, TEXT("Registering texture %u"), reinterpret_cast<uintptr_t>(RHITexture->GetNativeResource()));
 
diff --git a/Engine/Plugins/Media/WebMMedia/Source/WebMMedia/Private/Player/WebMMediaTextureSample.cpp b/Engine/Plugins/Media/WebMMedia/Source/WebMMedia/Private/Player/WebMMediaTextureSample.cpp
--- a/Engine/Plugins/Media/WebMMedia/Source/WebMMedia/Private/Player/WebMMediaTextureSample.cpp
+++ b/Engine/Plugins/Media/WebMMedia/Source/WebMMedia/Private/Player/WebMMediaTextureSample.cpp
@@ -24,9 +24,10 @@ void FWebMMediaTextureSample::CreateTexture()
 		FRHITextureCreateDesc::Create2D(TEXT("FWebMMediaTextureSample_DummyTexture2D"))
 		.SetExtent(TotalSize)
 		.SetFormat(PF_B8G8R8A8)
-		.SetFlags(ETextureCreateFlags::Dynamic | ETextureCreateFlags::SRGB);
+		.SetFlags(ETextureCreateFlags::Dynamic | ETextureCreateFlags::SRGB | ETextureCreateFlags::RenderTargetable | ETextureCreateFlags::ShaderResource)
+		.SetInitialState(ERHIAccess::SRVMask);
 
-	RHICreateTargetableShaderResource(Desc, ETextureCreateFlags::RenderTargetable, Texture);
+	Texture = RHICreateTexture(Desc);
 }
 
 const void* FWebMMediaTextureSample::GetBuffer()
diff --git a/Engine/Plugins/Runtime/Steam/SteamVR/Source/SteamVR/Private/SteamVRHMD.cpp b/Engine/Plugins/Runtime/Steam/SteamVR/Source/SteamVR/Private/SteamVRHMD.cpp
--- a/Engine/Plugins/Runtime/Steam/SteamVR/Source/SteamVR/Private/SteamVRHMD.cpp
+++ b/Engine/Plugins/Runtime/Steam/SteamVR/Source/SteamVR/Private/SteamVRHMD.cpp
@@ -1653,7 +1653,8 @@ bool FSteamVRHMD::AllocateRenderTargetTexture(uint32 Index, uint32 SizeX, uint32
 		FRHITextureCreateDesc::Create2D(TEXT("FSteamVRHMD"))
 		.SetExtent(SizeX, SizeY)
 		.SetFormat(PF_B8G8R8A8)
-		.SetNumSamples(NumSamples);
+		.SetFlags(ETextureCreateFlags::RenderTargetable | ETextureCreateFlags::ShaderResource)
+		.SetInitialState(ERHIAccess::SRVMask);
 
 	for (uint32 SwapChainIter = 0; SwapChainIter < SteamVRSwapChainLength; ++SwapChainIter)
 	{
@@ -1667,7 +1668,7 @@ bool FSteamVRHMD::AllocateRenderTargetTexture(uint32 Index, uint32 SizeX, uint32
 
 		FTexture2DRHIRef TargetableTexture;
 
-		RHICreateTargetableShaderResource(Desc, ETextureCreateFlags::RenderTargetable, TargetableTexture);
+		TargetableTexture = RHICreateTexture(Desc);
 
 		SwapChainTextures.Add((FTextureRHIRef&)TargetableTexture);
 
@@ -1720,14 +1721,13 @@ bool FSteamVRHMD::AllocateDepthTexture(uint32 Index, uint32 SizeX, uint32 SizeY,
 		FRHITextureCreateDesc::Create2D(TEXT("SteamVRDepthStencil"))
 		.SetExtent(SizeX, SizeY)
 		.SetFormat(PF_DepthStencil)
-		.SetFlags(Flags)
+		.SetFlags(Flags | TargetableTextureFlags | ETextureCreateFlags::ShaderResource)
+		.SetInitialState(ERHIAccess::SRVMask)
 		.SetClearValue(ClearValue);
 
 	for (uint32 SwapChainIter = 0; SwapChainIter < SteamVRSwapChainLength; ++SwapChainIter)
 	{
-		FTextureRHIRef TargetableTexture;
-
-		RHICreateTargetableShaderResource(Desc, TargetableTextureFlags, TargetableTexture);
+		FTextureRHIRef TargetableTexture = RHICreateTexture(Desc);
 
 		SwapChainTextures.Add(TargetableTexture);
 
diff --git a/Engine/Plugins/Runtime/WarpUtils/Source/PFMExporter/Private/PFMExporterMesh.cpp b/Engine/Plugins/Runtime/WarpUtils/Source/PFMExporter/Private/PFMExporterMesh.cpp
--- a/Engine/Plugins/Runtime/WarpUtils/Source/PFMExporter/Private/PFMExporterMesh.cpp
+++ b/Engine/Plugins/Runtime/WarpUtils/Source/PFMExporter/Private/PFMExporterMesh.cpp
@@ -49,20 +49,27 @@ bool FPFMExporterMesh::SaveToFile(const FString& FileName)
 bool FPFMExporterMesh::BeginExport_RenderThread(FRHICommandListImmediate& RHICmdList)
 {
 	// Allocate RT+SR RHI resources
-	const FRHITextureCreateDesc Desc =
+	FRHITextureCreateDesc Desc =
 		FRHITextureCreateDesc::Create2D(TEXT("FPFMExporterMesh"))
 		.SetExtent(DimWidth, DimHeight)
-		.SetFormat(PFMTextureFormat);
+		.SetFormat(PFMTextureFormat)
+		.SetInitialState(ERHIAccess::CopySrc);
 
-	RHICreateTargetableShaderResource(Desc, ETextureCreateFlags::RenderTargetable, true, RenderTargetTexture, ShaderResourceTexture);
+	ShaderResourceTexture = RHICreateTexture(Desc);
+
+	Desc.SetFlags(ETextureCreateFlags::RenderTargetable | ETextureCreateFlags::ShaderResource)
+		.SetInitialState(ERHIAccess::RTV);
+
+	RenderTargetTexture = RHICreateTexture(Desc);
 
 	return true;
 }
 
 bool FPFMExporterMesh::FinishExport_RenderThread(FRHICommandListImmediate& RHICmdList)
 {
-	//Copy texture to SR
-	RHICmdList.CopyToResolveTarget(RenderTargetTexture, ShaderResourceTexture, FResolveParams());
+	RHICmdList.Transition(FRHITransitionInfo(RenderTargetTexture, ERHIAccess::RTV, ERHIAccess::CopyDest));
+
+	RHICmdList.CopyTexture(RenderTargetTexture, ShaderResourceTexture, {});
 	RHICmdList.ImmediateFlush(EImmediateFlushType::FlushRHIThreadFlushResources);
 
 	// Extract result to memory
diff --git a/Engine/Plugins/Runtime/WindowsMixedReality/Source/HoloLensAR/Private/HoloLensCameraImageTexture.cpp b/Engine/Plugins/Runtime/WindowsMixedReality/Source/HoloLensAR/Private/HoloLensCameraImageTexture.cpp
--- a/Engine/Plugins/Runtime/WindowsMixedReality/Source/HoloLensAR/Private/HoloLensCameraImageTexture.cpp
+++ b/Engine/Plugins/Runtime/WindowsMixedReality/Source/HoloLensAR/Private/HoloLensCameraImageTexture.cpp
@@ -228,10 +228,11 @@ public:
 					FRHITextureCreateDesc::Create2D(TEXT("FHoloLensCameraImageResource_DummyTexture2D"))
 					.SetExtent(Size)
 					.SetFormat(PF_B8G8R8A8)
-					.SetFlags(ETextureCreateFlags::Dynamic);
+					.SetFlags(ETextureCreateFlags::Dynamic | ETextureCreateFlags::RenderTargetable | ETextureCreateFlags::ShaderResource)
+					.SetInitialState(ERHIAccess::SRVMask);
 
 				// Create our render target that we'll convert to
-				RHICreateTargetableShaderResource(CreateDesc, ETextureCreateFlags::RenderTargetable, DecodedTextureRef);
+				DecodedTextureRef = RHICreateTexture(CreateDesc);
 			}
 
 			if (PerformCopy(cameraImageTexture, D3D11DeviceContext))
diff --git a/Engine/Plugins/Runtime/WindowsMixedReality/Source/WindowsMixedRealityHMD/Private/WindowsMixedRealityHMD.cpp b/Engine/Plugins/Runtime/WindowsMixedReality/Source/WindowsMixedRealityHMD/Private/WindowsMixedRealityHMD.cpp
--- a/Engine/Plugins/Runtime/WindowsMixedReality/Source/WindowsMixedRealityHMD/Private/WindowsMixedRealityHMD.cpp
+++ b/Engine/Plugins/Runtime/WindowsMixedReality/Source/WindowsMixedRealityHMD/Private/WindowsMixedRealityHMD.cpp
@@ -1635,7 +1635,8 @@ namespace WindowsMixedReality
 			.SetExtent(sizeX, sizeY)
 			.SetFormat(PF_B8G8R8A8) // must be BGRA
 			.SetNumMips(numMips)
-			.SetFlags(flags);
+			.SetFlags(flags | targetableTextureFlags | ETextureCreateFlags::ShaderResource)
+			.SetInitialState(ERHIAccess::SRVMask);
 
 		if (bMultiView)
 		{
@@ -1643,7 +1644,7 @@ namespace WindowsMixedReality
 				.SetArraySize(2);
 		}
 
-		RHICreateTargetableShaderResource(Desc, targetableTextureFlags, outTargetableTexture, outShaderResourceTexture);
+		outTargetableTexture = outShaderResourceTexture = RHICreateTexture(Desc);
 
 		if (bMultiView)
 		{
@@ -1689,10 +1690,11 @@ namespace WindowsMixedReality
 			.SetFormat(PF_DepthStencil) // must be BGRA
 			// Do not use input mips, this will resolve to 0 which will throw creating the texture.
 			.SetNumMips(1)
-			.SetFlags(InTexFlags)
+			.SetFlags(InTexFlags | TargetableTextureFlags | ETextureCreateFlags::ShaderResource)
 			// This binding is necessary - without it there will be a runtime error.
 			// Current shader assumes far depth since scene depth uses far depth.
-			.SetClearValue(FClearValueBinding::DepthFar);
+			.SetClearValue(FClearValueBinding::DepthFar)
+			.SetInitialState(ERHIAccess::SRVMask);
 
 		if (bMultiView)
 		{
@@ -1705,9 +1707,8 @@ namespace WindowsMixedReality
 			Desc.SetExtent(SizeX, SizeY);
 		}
 
-		RHICreateTargetableShaderResource(Desc, TargetableTextureFlags, OutTargetableTexture, OutShaderResourceTexture);
+		CurrentDepthBuffer = OutTargetableTexture = OutShaderResourceTexture = RHICreateTexture(Desc);
 
-		CurrentDepthBuffer = OutTargetableTexture;
 		bNeedReallocateDepthTexture = false;
 
 		return true;
diff --git a/Engine/Plugins/Runtime/WindowsMoviePlayer/Source/WindowsMoviePlayer/Private/WindowsMovieStreamer.cpp b/Engine/Plugins/Runtime/WindowsMoviePlayer/Source/WindowsMoviePlayer/Private/WindowsMovieStreamer.cpp
--- a/Engine/Plugins/Runtime/WindowsMoviePlayer/Source/WindowsMoviePlayer/Private/WindowsMovieStreamer.cpp
+++ b/Engine/Plugins/Runtime/WindowsMoviePlayer/Source/WindowsMoviePlayer/Private/WindowsMovieStreamer.cpp
@@ -92,9 +92,10 @@ void FMediaFoundationMovieStreamer::ConvertSample()
 				FRHITextureCreateDesc::Create2D(TEXT("FMediaFoundationMovieStreamer"))
 				.SetExtent(SourceFormat.BufferDim)
 				.SetFormat(InputPixelFormat)
-				.SetFlags(InputCreateFlags);
+				.SetFlags(InputCreateFlags | ETextureCreateFlags::RenderTargetable | ETextureCreateFlags::ShaderResource)
+				.SetInitialState(ERHIAccess::SRVMask);
 
-			RHICreateTargetableShaderResource(Desc, ETextureCreateFlags::RenderTargetable, InputTarget);
+			InputTarget = RHICreateTexture(Desc);
 		}
 
 		// copy sample data to input render target
diff --git a/Engine/Plugins/Runtime/nDisplay/Source/DisplayCluster/Private/Render/Viewport/RenderTarget/DisplayClusterRenderTargetResource.cpp b/Engine/Plugins/Runtime/nDisplay/Source/DisplayCluster/Private/Render/Viewport/RenderTarget/DisplayClusterRenderTargetResource.cpp
--- a/Engine/Plugins/Runtime/nDisplay/Source/DisplayCluster/Private/Render/Viewport/RenderTarget/DisplayClusterRenderTargetResource.cpp
+++ b/Engine/Plugins/Runtime/nDisplay/Source/DisplayCluster/Private/Render/Viewport/RenderTarget/DisplayClusterRenderTargetResource.cpp
@@ -43,10 +43,11 @@ void FDisplayClusterViewportResource::ImplInitDynamicRHI_RenderTargetResource2D(
 		.SetExtent(GetSizeX(), GetSizeY())
 		.SetFormat(ViewportResourceSettings.Format)
 		.SetNumMips(NumMips)
-		.SetFlags(CreateFlags)
+		.SetFlags(CreateFlags | ETextureCreateFlags::RenderTargetable | ETextureCreateFlags::ShaderResource)
+		.SetInitialState(ERHIAccess::SRVMask)
 		.SetClearValue(FClearValueBinding::Black);
 
-	RHICreateTargetableShaderResource(Desc, ETextureCreateFlags::RenderTargetable, OutRenderTargetTextureRHI, OutTextureRHI);
+	OutRenderTargetTextureRHI = OutTextureRHI = RHICreateTexture(Desc);
 }
 
 void FDisplayClusterViewportResource::ImplInitDynamicRHI_TextureResource2D(FTexture2DRHIRef& OutTextureRHI)
diff --git a/Engine/Plugins/Runtime/nDisplay/Source/DisplayClusterLightCardEditor/Private/Viewport/DisplayClusterLightCardEditorViewportClient.cpp b/Engine/Plugins/Runtime/nDisplay/Source/DisplayClusterLightCardEditor/Private/Viewport/DisplayClusterLightCardEditorViewportClient.cpp
--- a/Engine/Plugins/Runtime/nDisplay/Source/DisplayClusterLightCardEditor/Private/Viewport/DisplayClusterLightCardEditorViewportClient.cpp
+++ b/Engine/Plugins/Runtime/nDisplay/Source/DisplayClusterLightCardEditor/Private/Viewport/DisplayClusterLightCardEditorViewportClient.cpp
@@ -71,9 +71,11 @@ void FDisplayClusterLightCardEditorViewportClient::FNormalMap::Init(const FScene
 				FRHITextureCreateDesc::Create2D(TEXT("NormalMapTexture"))
 				.SetExtent(SizeX, SizeY)
 				.SetFormat(PF_FloatRGBA)
-				.SetClearValue(FClearValueBinding::Black);
+				.SetClearValue(FClearValueBinding::Black)
+				.SetFlags(ETextureCreateFlags::RenderTargetable | ETextureCreateFlags::ShaderResource)
+				.SetInitialState(ERHIAccess::SRVMask);
 
-			RHICreateTargetableShaderResource(Desc, ETextureCreateFlags::RenderTargetable, RenderTargetTextureRHI);
+			RenderTargetTextureRHI = RHICreateTexture(Desc);
 		});
 }
 
diff --git a/Engine/Source/Runtime/Engine/Private/Particles/ParticleCurveTexture.cpp b/Engine/Source/Runtime/Engine/Private/Particles/ParticleCurveTexture.cpp
--- a/Engine/Source/Runtime/Engine/Private/Particles/ParticleCurveTexture.cpp
+++ b/Engine/Source/Runtime/Engine/Private/Particles/ParticleCurveTexture.cpp
@@ -169,7 +169,6 @@ TGlobalResource<FParticleCurveInjectionVertexDeclaration> GParticleCurveInjectio
 static void InjectCurves(
 	FRHICommandListImmediate& RHICmdList,
 	FRHITexture2D* CurveTextureRHI,
-	FRHITexture2D* CurveTextureTargetRHI,
 	TArray<FCurveSamples>& InPendingCurves)
 {
 	static bool bFirstCall = true;
@@ -178,7 +177,7 @@ static void InjectCurves(
 
 	SCOPED_DRAW_EVENT(RHICmdList, InjectParticleCurves);
 
-	RHICmdList.BeginUpdateMultiFrameResource(CurveTextureTargetRHI);
+	RHICmdList.BeginUpdateMultiFrameResource(CurveTextureRHI);
 
 	ERenderTargetLoadAction LoadAction = ERenderTargetLoadAction::ELoad;
 
@@ -188,8 +187,8 @@ static void InjectCurves(
 		bFirstCall = false;
 	}
 
-	FRHIRenderPassInfo RPInfo(CurveTextureTargetRHI, MakeRenderTargetActions(LoadAction, ERenderTargetStoreAction::EStore));
-	RHICmdList.Transition(FRHITransitionInfo(CurveTextureTargetRHI, ERHIAccess::Unknown, ERHIAccess::RTV));
+	FRHIRenderPassInfo RPInfo(CurveTextureRHI, MakeRenderTargetActions(LoadAction, ERenderTargetStoreAction::EStore));
+	RHICmdList.Transition(FRHITransitionInfo(CurveTextureRHI, ERHIAccess::Unknown, ERHIAccess::RTV));
 	RHICmdList.BeginRenderPass(RPInfo, TEXT("InjectCurves"));
 	{
 		FGraphicsPipelineStateInitializer GraphicsPSOInit;
@@ -282,8 +281,8 @@ static void InjectCurves(
 		}
 	}
 	RHICmdList.EndRenderPass();
-	RHICmdList.CopyToResolveTarget(CurveTextureTargetRHI, CurveTextureRHI, FResolveParams());
-	RHICmdList.EndUpdateMultiFrameResource(CurveTextureTargetRHI);
+	RHICmdList.Transition(FRHITransitionInfo(CurveTextureRHI, ERHIAccess::RTV, ERHIAccess::SRVMask));
+	RHICmdList.EndUpdateMultiFrameResource(CurveTextureRHI);
 }
 
 /*------------------------------------------------------------------------------
@@ -518,10 +517,11 @@ void FParticleCurveTexture::InitRHI()
 	const FRHITextureCreateDesc Desc =
 		FRHITextureCreateDesc::Create2D(TEXT("ParticleCurveTexture"))
 		.SetExtent(GParticleCurveTextureSizeX, GParticleCurveTextureSizeY)
+		.SetFlags(ETextureCreateFlags::RenderTargetable | ETextureCreateFlags::ShaderResource | ETextureCreateFlags::NoFastClear)
 		.SetFormat(PF_B8G8R8A8)
 		.SetClearValue(FClearValueBinding(FLinearColor::Blue));
 
-	RHICreateTargetableShaderResource(Desc, ETextureCreateFlags::RenderTargetable | ETextureCreateFlags::NoFastClear, CurveTextureTargetRHI, CurveTextureRHI);
+	CurveTextureRHI = RHICreateTexture(Desc);
 }
 
 /**
@@ -529,7 +529,6 @@ void FParticleCurveTexture::InitRHI()
 */
 void FParticleCurveTexture::ReleaseRHI()
 {
-	CurveTextureTargetRHI.SafeRelease();
 	CurveTextureRHI.SafeRelease();
 }
 
@@ -611,7 +610,6 @@ void FParticleCurveTexture::SubmitPendingCurves()
 				InjectCurves(
 					RHICmdList,
 					ParticleCurveTexture->CurveTextureRHI,
-					ParticleCurveTexture->CurveTextureTargetRHI,
 					PendingCurves
 				);
 			});
diff --git a/Engine/Source/Runtime/Engine/Private/Particles/ParticleCurveTexture.h b/Engine/Source/Runtime/Engine/Private/Particles/ParticleCurveTexture.h
--- a/Engine/Source/Runtime/Engine/Private/Particles/ParticleCurveTexture.h
+++ b/Engine/Source/Runtime/Engine/Private/Particles/ParticleCurveTexture.h
@@ -166,11 +166,8 @@ public:
 	void SubmitPendingCurves();
 
 private:
-
-	/** Targetable texture for uploading curve samples. */
-	FTexture2DRHIRef CurveTextureTargetRHI;
 	/** Texture for sampling curves on the GPU. */
-	FTexture2DRHIRef CurveTextureRHI;
+	FTextureRHIRef CurveTextureRHI;
 	/** The texel allocator for this texture. */
 	FTexelAllocator TexelAllocator;
 	/** A list of pending curves that need to be uploaded. */
diff --git a/Engine/Source/Runtime/Engine/Private/Particles/ParticleGpuSimulation.cpp b/Engine/Source/Runtime/Engine/Private/Particles/ParticleGpuSimulation.cpp
--- a/Engine/Source/Runtime/Engine/Private/Particles/ParticleGpuSimulation.cpp
+++ b/Engine/Source/Runtime/Engine/Private/Particles/ParticleGpuSimulation.cpp
@@ -226,10 +226,8 @@ class FParticleStateTextures : public FRenderResource
 public:
 
 	/** Contains the positions of all simulating particles. */
-	FTexture2DRHIRef PositionTextureTargetRHI;
 	FTexture2DRHIRef PositionTextureRHI;
 	/** Contains the velocity of all simulating particles. */
-	FTexture2DRHIRef VelocityTextureTargetRHI;
 	FTexture2DRHIRef VelocityTextureRHI;
 
 	bool bTexturesCleared;
@@ -244,7 +242,6 @@ public:
 		const int32 SizeY = GParticleSimulationTextureSizeY;
 
 		// 32-bit per channel RGBA texture for position.
-		check( !IsValidRef( PositionTextureTargetRHI ) );
 		check( !IsValidRef( PositionTextureRHI ) );
 
 #define PARTICLE_STATE_POSITION_TEXTURE_NAME	TEXT("ParticleStatePosition")
@@ -255,13 +252,14 @@ public:
 				FRHITextureCreateDesc::Create2D(PARTICLE_STATE_POSITION_TEXTURE_NAME)
 				.SetExtent(SizeX, SizeY)
 				.SetFormat(PF_A32B32G32R32F)
-				.SetClearValue(FClearValueBinding::Transparent);
+				.SetClearValue(FClearValueBinding::Transparent)
+				.SetFlags(ETextureCreateFlags::RenderTargetable | ETextureCreateFlags::ShaderResource)
+				.SetInitialState(ERHIAccess::SRVMask);
 
-			RHICreateTargetableShaderResource(Desc, ETextureCreateFlags::RenderTargetable, PositionTextureTargetRHI, PositionTextureRHI);
+			PositionTextureRHI = RHICreateTexture(Desc);
 		}
 
 		// 16-bit per channel RGBA texture for velocity.
-		check( !IsValidRef( VelocityTextureTargetRHI ) );
 		check( !IsValidRef( VelocityTextureRHI ) );
 
 		{
@@ -269,18 +267,20 @@ public:
 				FRHITextureCreateDesc::Create2D(PARTICLE_STATE_VELOCITY_TEXTURE_NAME)
 				.SetExtent(SizeX, SizeY)
 				.SetFormat(PF_FloatRGBA)
-				.SetClearValue(FClearValueBinding::Transparent);
+				.SetClearValue(FClearValueBinding::Transparent)
+				.SetFlags(ETextureCreateFlags::RenderTargetable | ETextureCreateFlags::ShaderResource)
+				.SetInitialState(ERHIAccess::SRVMask);
 
-			RHICreateTargetableShaderResource(Desc, ETextureCreateFlags::RenderTargetable, VelocityTextureTargetRHI, VelocityTextureRHI);
+			VelocityTextureRHI = RHICreateTexture(Desc);
 		}
 
 		// using FName's ability to append a number to a string (..._0) without an extra string allocation, except suffixing is done when number > 0 hence the +1 here : 
 		FName PositionTextureName(PARTICLE_STATE_POSITION_TEXTURE_NAME, ParticleStateIndex + 1); 
 		FName VelocityTextureName(PARTICLE_STATE_VELOCITY_TEXTURE_NAME, ParticleStateIndex + 1);
-		PositionTextureTargetRHI->SetName(PositionTextureName);
-		VelocityTextureTargetRHI->SetName(VelocityTextureName);
-		RHIBindDebugLabelName(PositionTextureTargetRHI, *PositionTextureName.ToString());
-		RHIBindDebugLabelName(VelocityTextureTargetRHI, *VelocityTextureName.ToString());
+		PositionTextureRHI->SetName(PositionTextureName);
+		VelocityTextureRHI->SetName(VelocityTextureName);
+		RHIBindDebugLabelName(PositionTextureRHI, *PositionTextureName.ToString());
+		RHIBindDebugLabelName(VelocityTextureRHI, *VelocityTextureName.ToString());
 #undef PARTICLE_STATE_VELOCITY_TEXTURE_NAME
 #undef PARTICLE_STATE_POSITION_TEXTURE_NAME
 
@@ -292,10 +292,7 @@ public:
 	 */
 	virtual void ReleaseRHI() override
 	{
-		// Release textures.
-		PositionTextureTargetRHI.SafeRelease();
 		PositionTextureRHI.SafeRelease();
-		VelocityTextureTargetRHI.SafeRelease();
 		VelocityTextureRHI.SafeRelease();
 	}
 };
@@ -308,8 +305,7 @@ class FParticleAttributesTexture : public FRenderResource
 public:
 
 	/** Contains the attributes of all simulating particles. */
-	FTexture2DRHIRef TextureTargetRHI;
-	FTexture2DRHIRef TextureRHI;
+	FTextureRHIRef TextureRHI;
 
 	/**
 	 * Initialize RHI resources used for particle simulation.
@@ -327,22 +323,20 @@ public:
 			FRHITextureCreateDesc::Create2D(ATTRIBUTES_TEXTURE_NAME)
 			.SetExtent(SizeX, SizeY)
 			.SetFormat(PF_B8G8R8A8)
-			.SetClearValue(FClearValueBinding::Transparent);
+			.SetClearValue(FClearValueBinding::Transparent)
+			.SetFlags(ETextureCreateFlags::RenderTargetable | ETextureCreateFlags::ShaderResource | ETextureCreateFlags::NoFastClear | ExtraFlags)
+			.SetInitialState(ERHIAccess::RTV);
 
-		RHICreateTargetableShaderResource(Desc, ETextureCreateFlags::RenderTargetable | ETextureCreateFlags::NoFastClear | ExtraFlags, TextureTargetRHI, TextureRHI);
+		TextureRHI = RHICreateTexture(Desc);
 
 		static FName AttributesTextureName(ATTRIBUTES_TEXTURE_NAME);
-		TextureTargetRHI->SetName(AttributesTextureName);
-		RHIBindDebugLabelName(TextureTargetRHI, ATTRIBUTES_TEXTURE_NAME);
+		TextureRHI->SetName(AttributesTextureName);
+		RHIBindDebugLabelName(TextureRHI, ATTRIBUTES_TEXTURE_NAME);
 #undef ATTRIBUTES_TEXTURE_NAME
-		
-		{
-			FRHIRenderPassInfo RPInfo(TextureTargetRHI, ERenderTargetActions::Clear_Store);
-			FRHICommandListImmediate& RHICmdList = FRHICommandListExecutor::GetImmediateCommandList();
-			RHICmdList.Transition(FRHITransitionInfo(TextureTargetRHI, ERHIAccess::SRVMask, ERHIAccess::RTV));
-			RHICmdList.BeginRenderPass(RPInfo, TEXT("Clear"));
-			RHICmdList.EndRenderPass();
-		}
+
+		FRHICommandListImmediate& RHICmdList = FRHICommandListExecutor::GetImmediateCommandList();
+		ClearRenderTarget(RHICmdList, TextureRHI);
+		RHICmdList.Transition(FRHITransitionInfo(TextureRHI, ERHIAccess::RTV, ERHIAccess::SRVMask));
 	}
 
 	/**
@@ -350,7 +344,6 @@ public:
 	 */
 	virtual void ReleaseRHI() override
 	{
-		TextureTargetRHI.SafeRelease();
 		TextureRHI.SafeRelease();
 	}
 };
@@ -394,10 +387,11 @@ public:
 			RHICmdList.Transition({ 
 				FRHITransitionInfo(ParticleResources->RenderAttributesTexture.TextureRHI, ERHIAccess::Unknown, ERHIAccess::SRVGraphics),
 				FRHITransitionInfo(ParticleResources->SimulationAttributesTexture.TextureRHI, ERHIAccess::Unknown, ERHIAccess::SRVGraphics),
-				FRHITransitionInfo(ParticleResources->StateTextures[0].PositionTextureTargetRHI, ERHIAccess::Unknown, ERHIAccess::SRVGraphics),
-				FRHITransitionInfo(ParticleResources->StateTextures[0].VelocityTextureTargetRHI, ERHIAccess::Unknown, ERHIAccess::SRVGraphics),
-				FRHITransitionInfo(ParticleResources->StateTextures[1].PositionTextureTargetRHI, ERHIAccess::Unknown, ERHIAccess::SRVGraphics),
-				FRHITransitionInfo(ParticleResources->StateTextures[1].VelocityTextureTargetRHI, ERHIAccess::Unknown, ERHIAccess::SRVGraphics)});
+				FRHITransitionInfo(ParticleResources->StateTextures[0].PositionTextureRHI, ERHIAccess::Unknown, ERHIAccess::SRVGraphics),
+				FRHITransitionInfo(ParticleResources->StateTextures[0].VelocityTextureRHI, ERHIAccess::Unknown, ERHIAccess::SRVGraphics),
+				FRHITransitionInfo(ParticleResources->StateTextures[1].PositionTextureRHI, ERHIAccess::Unknown, ERHIAccess::SRVGraphics),
+				FRHITransitionInfo(ParticleResources->StateTextures[1].VelocityTextureRHI, ERHIAccess::Unknown, ERHIAccess::SRVGraphics)
+			});
 		});
 	}
 
@@ -4665,8 +4659,8 @@ void FFXSystem::PrepareGPUSimulation(FRHICommandListImmediate& RHICmdList)
 	{
 		// Setup render states.
 		RHICmdList.Transition({
-			FRHITransitionInfo(CurrentStateTextures.PositionTextureTargetRHI, ERHIAccess::SRVGraphics, ERHIAccess::RTV),
-			FRHITransitionInfo(CurrentStateTextures.VelocityTextureTargetRHI, ERHIAccess::SRVGraphics, ERHIAccess::RTV)
+			FRHITransitionInfo(CurrentStateTextures.PositionTextureRHI, ERHIAccess::SRVGraphics, ERHIAccess::RTV),
+			FRHITransitionInfo(CurrentStateTextures.VelocityTextureRHI, ERHIAccess::SRVGraphics, ERHIAccess::RTV)
 		});
 	}
 }
@@ -4695,8 +4689,8 @@ void FFXSystem::SimulateGPUParticles(
 	FParticleStateTextures& PrevStateTextures = ParticleSimulationResources->GetPreviousStateTextures();	
 
 	// Setup render states.
-	FRHITexture* CurrentStateRenderTargets[2] = { CurrentStateTextures.PositionTextureTargetRHI, CurrentStateTextures.VelocityTextureTargetRHI };
-	FRHITexture* PreviousStateRenderTargets[2] = { PrevStateTextures.PositionTextureTargetRHI, PrevStateTextures.VelocityTextureTargetRHI };
+	FRHITexture* CurrentStateRenderTargets[2] = { CurrentStateTextures.PositionTextureRHI, CurrentStateTextures.VelocityTextureRHI };
+	FRHITexture* PreviousStateRenderTargets[2] = { PrevStateTextures.PositionTextureRHI, PrevStateTextures.VelocityTextureRHI };
 
 
 #if WITH_MGPU
@@ -5000,20 +4994,21 @@ void FFXSystem::SimulateGPUParticles(
 		// Set render targets.
 		FRHITexture* InjectRenderTargets[4] =
 		{
-			CurrentStateTextures.PositionTextureTargetRHI,
-			CurrentStateTextures.VelocityTextureTargetRHI,
-			ParticleSimulationResources->RenderAttributesTexture.TextureTargetRHI,
-			ParticleSimulationResources->SimulationAttributesTexture.TextureTargetRHI
+			CurrentStateTextures.PositionTextureRHI,
+			CurrentStateTextures.VelocityTextureRHI,
+			ParticleSimulationResources->RenderAttributesTexture.TextureRHI,
+			ParticleSimulationResources->SimulationAttributesTexture.TextureRHI
 		};
-		RHICmdList.BeginUpdateMultiFrameResource(ParticleSimulationResources->RenderAttributesTexture.TextureTargetRHI);
-		RHICmdList.BeginUpdateMultiFrameResource(ParticleSimulationResources->SimulationAttributesTexture.TextureTargetRHI);
+		RHICmdList.BeginUpdateMultiFrameResource(ParticleSimulationResources->RenderAttributesTexture.TextureRHI);
+		RHICmdList.BeginUpdateMultiFrameResource(ParticleSimulationResources->SimulationAttributesTexture.TextureRHI);
 
 		FRHIRenderPassInfo RPInfo(4, InjectRenderTargets, ERenderTargetActions::Load_Store);
 		{
 			// Transition attribute textures to writeble, particle state texture are in writeble state already
 			RHICmdList.Transition({ 
 				FRHITransitionInfo(InjectRenderTargets[2], ERHIAccess::SRVGraphics, ERHIAccess::RTV),
-				FRHITransitionInfo(InjectRenderTargets[3], ERHIAccess::SRVGraphics, ERHIAccess::RTV)});
+				FRHITransitionInfo(InjectRenderTargets[3], ERHIAccess::SRVGraphics, ERHIAccess::RTV)
+			});
 
 			RHICmdList.BeginRenderPass(RPInfo, TEXT("ParticleInjection"));
 
@@ -5030,17 +5025,10 @@ void FFXSystem::SimulateGPUParticles(
 
 			RHICmdList.EndRenderPass();
 
-			// Resolve attributes textures. State textures are resolved later.
-			RHICmdList.CopyToResolveTarget(
-				ParticleSimulationResources->RenderAttributesTexture.TextureTargetRHI,
-				ParticleSimulationResources->RenderAttributesTexture.TextureRHI,
-				FResolveParams()
-			);
-			RHICmdList.CopyToResolveTarget(
-				ParticleSimulationResources->SimulationAttributesTexture.TextureTargetRHI,
-				ParticleSimulationResources->SimulationAttributesTexture.TextureRHI,
-				FResolveParams()
-			);
+			RHICmdList.Transition({
+				FRHITransitionInfo(InjectRenderTargets[2], ERHIAccess::RTV, ERHIAccess::SRVGraphics),
+				FRHITransitionInfo(InjectRenderTargets[3], ERHIAccess::RTV, ERHIAccess::SRVGraphics)
+			});
 		}
 
 		if (GNumAlternateFrameRenderingGroups > 1)
@@ -5055,20 +5043,20 @@ void FFXSystem::SimulateGPUParticles(
 			else
 			{
 #if WITH_MGPU
-				TemporalEffectTextures.Add(ParticleSimulationResources->RenderAttributesTexture.TextureTargetRHI);
-				TemporalEffectTextures.Add(ParticleSimulationResources->SimulationAttributesTexture.TextureTargetRHI);
+				TemporalEffectTextures.Add(ParticleSimulationResources->RenderAttributesTexture.TextureRHI);
+				TemporalEffectTextures.Add(ParticleSimulationResources->SimulationAttributesTexture.TextureRHI);
 #endif
 			}
 		}
 #if WITH_MGPU
 		if (bCrossTransferEnabled)
 		{
-			AddCrossGPUTransferResource(ParticleSimulationResources->RenderAttributesTexture.TextureTargetRHI);
-			AddCrossGPUTransferResource(ParticleSimulationResources->SimulationAttributesTexture.TextureTargetRHI);
+			AddCrossGPUTransferResource(ParticleSimulationResources->RenderAttributesTexture.TextureRHI);
+			AddCrossGPUTransferResource(ParticleSimulationResources->SimulationAttributesTexture.TextureRHI);
 		}
 #endif
-		RHICmdList.EndUpdateMultiFrameResource(ParticleSimulationResources->RenderAttributesTexture.TextureTargetRHI);
-		RHICmdList.EndUpdateMultiFrameResource(ParticleSimulationResources->SimulationAttributesTexture.TextureTargetRHI);
+		RHICmdList.EndUpdateMultiFrameResource(ParticleSimulationResources->RenderAttributesTexture.TextureRHI);
+		RHICmdList.EndUpdateMultiFrameResource(ParticleSimulationResources->SimulationAttributesTexture.TextureRHI);
 	}
 	
 	// finish current state render
@@ -5086,10 +5074,11 @@ void FFXSystem::SimulateGPUParticles(
 		FParticleStateTextures& VisualizeStateTextures = ParticleSimulationResources->GetPreviousStateTextures();
 		
 		RHICmdList.Transition({ 
-			FRHITransitionInfo(VisualizeStateTextures.PositionTextureTargetRHI, ERHIAccess::SRVGraphics, ERHIAccess::RTV),
-			FRHITransitionInfo(VisualizeStateTextures.VelocityTextureTargetRHI, ERHIAccess::SRVGraphics, ERHIAccess::RTV)});
+			FRHITransitionInfo(VisualizeStateTextures.PositionTextureRHI, ERHIAccess::SRVGraphics, ERHIAccess::RTV),
+			FRHITransitionInfo(VisualizeStateTextures.VelocityTextureRHI, ERHIAccess::SRVGraphics, ERHIAccess::RTV)
+		});
 				
-		FRHITexture* VisualizeStateRHIs[2] = { VisualizeStateTextures.PositionTextureTargetRHI, VisualizeStateTextures.VelocityTextureTargetRHI };
+		FRHITexture* VisualizeStateRHIs[2] = { VisualizeStateTextures.PositionTextureRHI, VisualizeStateTextures.VelocityTextureRHI };
 		FRHIRenderPassInfo RPInfo(2, VisualizeStateRHIs, ERenderTargetActions::Load_Store);
 		{
 			RHICmdList.BeginRenderPass(RPInfo, TEXT("ExecuteSimulationCommands"));
@@ -5117,8 +5106,9 @@ void FFXSystem::SimulateGPUParticles(
 		}
 
 		RHICmdList.Transition({ 
-			FRHITransitionInfo(VisualizeStateTextures.PositionTextureTargetRHI, ERHIAccess::RTV, ERHIAccess::SRVGraphics),
-			FRHITransitionInfo(VisualizeStateTextures.VelocityTextureTargetRHI, ERHIAccess::RTV, ERHIAccess::SRVGraphics)});
+			FRHITransitionInfo(VisualizeStateTextures.PositionTextureRHI, ERHIAccess::RTV, ERHIAccess::SRVGraphics),
+			FRHITransitionInfo(VisualizeStateTextures.VelocityTextureRHI, ERHIAccess::RTV, ERHIAccess::SRVGraphics)
+		});
 	}
 
 #if WITH_MGPU
@@ -5177,8 +5167,8 @@ void FFXSystem::UpdateMultiGPUResources(FRHICommandListImmediate& RHICmdList)
 		// Set render targets.
 		FRHITexture* InjectRenderTargets[2] =
 		{
-			ParticleSimulationResources->RenderAttributesTexture.TextureTargetRHI,
-			ParticleSimulationResources->SimulationAttributesTexture.TextureTargetRHI
+			ParticleSimulationResources->RenderAttributesTexture.TextureRHI,
+			ParticleSimulationResources->SimulationAttributesTexture.TextureRHI
 		};
 
 		FRHIRenderPassInfo RPInfo(2, InjectRenderTargets, ERenderTargetActions::Load_Store);
@@ -5186,8 +5176,8 @@ void FFXSystem::UpdateMultiGPUResources(FRHICommandListImmediate& RHICmdList)
 		RPInfo.ColorRenderTargets[1].ResolveTarget = ParticleSimulationResources->SimulationAttributesTexture.TextureRHI;
 		{
 			RHICmdList.Transition({
-				FRHITransitionInfo(ParticleSimulationResources->RenderAttributesTexture.TextureTargetRHI, ERHIAccess::Unknown, ERHIAccess::RTV),
-				FRHITransitionInfo(ParticleSimulationResources->SimulationAttributesTexture.TextureTargetRHI, ERHIAccess::Unknown, ERHIAccess::RTV)
+				FRHITransitionInfo(ParticleSimulationResources->RenderAttributesTexture.TextureRHI, ERHIAccess::Unknown, ERHIAccess::RTV),
+				FRHITransitionInfo(ParticleSimulationResources->SimulationAttributesTexture.TextureRHI, ERHIAccess::Unknown, ERHIAccess::RTV)
 			});
 			RHICmdList.BeginRenderPass(RPInfo, TEXT("UpdateMultiGPUResources"));
 
diff --git a/Engine/Source/Runtime/Engine/Private/Slate/SceneViewport.cpp b/Engine/Source/Runtime/Engine/Private/Slate/SceneViewport.cpp
--- a/Engine/Source/Runtime/Engine/Private/Slate/SceneViewport.cpp
+++ b/Engine/Source/Runtime/Engine/Private/Slate/SceneViewport.cpp
@@ -2045,9 +2045,11 @@ void FSceneViewport::InitDynamicRHI()
 				const FRHITextureCreateDesc Desc =
 					FRHITextureCreateDesc::Create2D(TEXT("BufferedRT"))
 					.SetExtent(TexSizeX, TexSizeY)
-					.SetFormat(SceneTargetFormat);
+					.SetFormat(SceneTargetFormat)
+					.SetFlags(ETextureCreateFlags::RenderTargetable | ETextureCreateFlags::ShaderResource)
+					.SetInitialState(ERHIAccess::SRVMask);
 
-				RHICreateTargetableShaderResource(Desc, ETextureCreateFlags::RenderTargetable, BufferedRTRHI, BufferedSRVRHI);
+				BufferedRTRHI = BufferedSRVRHI = RHICreateTexture(Desc);
 			}
 			BufferedRenderTargetsRHI[i] = BufferedRTRHI;
 			BufferedShaderResourceTexturesRHI[i] = BufferedSRVRHI;
diff --git a/Engine/Source/Runtime/Engine/Private/Slate/SlateTextures.cpp b/Engine/Source/Runtime/Engine/Private/Slate/SlateTextures.cpp
--- a/Engine/Source/Runtime/Engine/Private/Slate/SlateTextures.cpp
+++ b/Engine/Source/Runtime/Engine/Private/Slate/SlateTextures.cpp
@@ -319,9 +319,11 @@ void FSlateTextureRenderTarget2DResource::InitDynamicRHI()
 			FRHITextureCreateDesc::Create2D(TEXT("FSlateTextureRenderTarget2DResource"))
 			.SetExtent(TargetSizeX, TargetSizeY)
 			.SetFormat((EPixelFormat)Format)
-			.SetClearValue(FClearValueBinding(ClearColor));
+			.SetClearValue(FClearValueBinding(ClearColor))
+			.SetFlags(ETextureCreateFlags::RenderTargetable | ETextureCreateFlags::ShaderResource)
+			.SetInitialState(ERHIAccess::SRVMask);
 
-		RHICreateTargetableShaderResource(Desc, ETextureCreateFlags::RenderTargetable, RenderTargetTextureRHI, TextureRHI);
+		RenderTargetTextureRHI = TextureRHI = RHICreateTexture(Desc);
 	}
 
 	// Create the sampler state RHI resource.
@@ -342,7 +344,7 @@ void FSlateTextureRenderTarget2DResource::ReleaseDynamicRHI()
 	// Release the FTexture RHI resources here as well
 	ReleaseRHI();
 
-	RenderTargetTextureRHI.SafeRelease();	
+	RenderTargetTextureRHI.SafeRelease();
 
 	// Remove from global list of deferred clears
 	RemoveFromDeferredUpdateList();
@@ -355,13 +357,10 @@ void FSlateTextureRenderTarget2DResource::UpdateDeferredResource(FRHICommandList
 	// Clear the target surface to green
 	if (bClearRenderTarget)
 	{
-		FRHIRenderPassInfo RPInfo(RenderTargetTextureRHI, ERenderTargetActions::Clear_Store);
-		RHICmdList.BeginRenderPass(RPInfo, TEXT("Slate2DUpdateDeferred_Clear"));
-		RHICmdList.EndRenderPass();
+		RHICmdList.Transition(FRHITransitionInfo(RenderTargetTextureRHI, ERHIAccess::Unknown, ERHIAccess::RTV));
+		ClearRenderTarget(RHICmdList, RenderTargetTextureRHI);
+		RHICmdList.Transition(FRHITransitionInfo(RenderTargetTextureRHI, ERHIAccess::RTV, ERHIAccess::SRVGraphics));
 	}
-
-	// Copy surface to the texture for use
-	RHICmdList.CopyToResolveTarget(RenderTargetTextureRHI, TextureRHI, FResolveParams());
 }
 
 uint32 FSlateTextureRenderTarget2DResource::GetSizeX() const
diff --git a/Engine/Source/Runtime/Engine/Private/TextureRenderTarget2D.cpp b/Engine/Source/Runtime/Engine/Private/TextureRenderTarget2D.cpp
--- a/Engine/Source/Runtime/Engine/Private/TextureRenderTarget2D.cpp
+++ b/Engine/Source/Runtime/Engine/Private/TextureRenderTarget2D.cpp
@@ -575,24 +575,31 @@ void FTextureRenderTarget2DResource::InitDynamicRHI()
 			TexCreateFlags |= ETextureCreateFlags::UAV;
 		}
 
-		const FRHITextureCreateDesc Desc =
+		FRHITextureCreateDesc Desc =
 			FRHITextureCreateDesc::Create2D(*ResourceName)
 			.SetExtent(Owner->SizeX, Owner->SizeY)
 			.SetFormat(Format)
 			.SetNumMips(Owner->GetNumMips())
-			.SetFlags(TexCreateFlags)
+			.SetFlags(TexCreateFlags | ETextureCreateFlags::RenderTargetable | ETextureCreateFlags::ShaderResource)
+			.SetInitialState(ERHIAccess::SRVMask)
 			.SetClearValue(FClearValueBinding(ClearColor));
 
-		RHICreateTargetableShaderResource(Desc, ETextureCreateFlags::RenderTargetable, Owner->bNeedsTwoCopies, RenderTargetTextureRHI, Texture2DRHI);
+		TextureRHI = RenderTargetTextureRHI = RHICreateTexture(Desc);
+
+		if (Owner->bNeedsTwoCopies)
+		{
+			Desc.SetFlags(TexCreateFlags | ETextureCreateFlags::ShaderResource);
+
+			TextureRHI = RHICreateTexture(Desc);
+		}
 
 		if (EnumHasAnyFlags(TexCreateFlags, ETextureCreateFlags::UAV))
 		{
 			UnorderedAccessViewRHI = RHICreateUnorderedAccessView(RenderTargetTextureRHI);
 		}
 
-		SetGPUMask(Desc.GPUMask);
-		TextureRHI = (FTextureRHIRef&)Texture2DRHI;
-		RHIUpdateTextureReference(Owner->TextureReference.TextureReferenceRHI,TextureRHI);
+		SetGPUMask(FRHIGPUMask::All());
+		RHIUpdateTextureReference(Owner->TextureReference.TextureReferenceRHI, TextureRHI);
 
 		AddToDeferredUpdateList(true);
 	}
@@ -619,7 +626,6 @@ void FTextureRenderTarget2DResource::ReleaseDynamicRHI()
 	ReleaseRHI();
 
 	RHIUpdateTextureReference(Owner->TextureReference.TextureReferenceRHI, nullptr);
-	Texture2DRHI.SafeRelease();
 	RenderTargetTextureRHI.SafeRelease();
 	MipGenerationCache.SafeRelease();
 
@@ -635,11 +641,15 @@ void FTextureRenderTarget2DResource::ReleaseDynamicRHI()
  */
 void FTextureRenderTarget2DResource::UpdateDeferredResource( FRHICommandListImmediate& RHICmdList, bool bClearRenderTarget/*=true*/ )
 {
-	FMemMark Mark(FMemStack::Get());
-
 	SCOPED_DRAW_EVENT(RHICmdList, GPUResourceUpdate)
 	RemoveFromDeferredUpdateList();
 
+	// Skip executing an empty graph.
+	if (TextureRHI == RenderTargetTextureRHI && !bClearRenderTarget && !Owner->bAutoGenerateMips)
+	{
+		return;
+	}
+
 	FRDGBuilder GraphBuilder(RHICmdList);
 
 	CacheRenderTarget(RenderTargetTextureRHI, TEXT("MipGeneration"), MipGenerationCache);
@@ -663,7 +673,7 @@ void FTextureRenderTarget2DResource::UpdateDeferredResource( FRHICommandListImme
 			Owner->MipsAddressV == TA_Wrap ? AM_Wrap : (Owner->MipsAddressV == TA_Mirror ? AM_Mirror : AM_Clamp)});
 	}
 
-	AddCopyToResolveTargetPass(GraphBuilder, RenderTargetTextureRDG, TextureRDG, FResolveParams());
+	AddCopyTexturePass(GraphBuilder, RenderTargetTextureRDG, TextureRDG, FRHICopyTextureInfo());
 
 	GraphBuilder.Execute();
 }
diff --git a/Engine/Source/Runtime/Engine/Private/TextureRenderTarget2DArray.cpp b/Engine/Source/Runtime/Engine/Private/TextureRenderTarget2DArray.cpp
--- a/Engine/Source/Runtime/Engine/Private/TextureRenderTarget2DArray.cpp
+++ b/Engine/Source/Runtime/Engine/Private/TextureRenderTarget2DArray.cpp
@@ -246,22 +246,20 @@ void FTextureRenderTarget2DArrayResource::InitDynamicRHI()
 				.SetArraySize(Owner->Slices)
 				.SetFormat(Owner->GetFormat())
 				.SetNumMips(Owner->GetNumMips())
-				.SetFlags(TexCreateFlags)
-				.SetClearValue(FClearValueBinding(Owner->ClearColor));
+				.SetFlags(TexCreateFlags | ETextureCreateFlags::RenderTargetable | ETextureCreateFlags::ShaderResource)
+				.SetClearValue(FClearValueBinding(Owner->ClearColor))
+				.SetInitialState(ERHIAccess::SRVMask);
 
-			RHICreateTargetableShaderResource(Desc, ETextureCreateFlags::RenderTargetable, RenderTarget2DArrayRHI, Texture2DArrayRHI);
+			TextureRHI = RHICreateTexture(Desc);
 		}
 
 		if (EnumHasAnyFlags(TexCreateFlags, ETextureCreateFlags::UAV))
 		{
-			UnorderedAccessViewRHI = RHICreateUnorderedAccessView(RenderTarget2DArrayRHI);
+			UnorderedAccessViewRHI = RHICreateUnorderedAccessView(TextureRHI);
 		}
 
-		TextureRHI = Texture2DArrayRHI;
 		RHIUpdateTextureReference(Owner->TextureReference.TextureReferenceRHI, TextureRHI);
-
-		// Can't set this as it's a texture 2D
-		//RenderTargetTextureRHI = 2DArraySurfaceRHI;
+		RenderTargetTextureRHI = TextureRHI;
 
 		AddToDeferredUpdateList(true);
 	}
@@ -288,8 +286,7 @@ void FTextureRenderTarget2DArrayResource::ReleaseDynamicRHI()
 	ReleaseRHI();
 
 	RHIUpdateTextureReference(Owner->TextureReference.TextureReferenceRHI, nullptr);
-	RenderTarget2DArrayRHI.SafeRelease();
-	Texture2DArrayRHI.SafeRelease();
+	RenderTargetTextureRHI.SafeRelease();
 
 	// remove from global list of deferred clears
 	RemoveFromDeferredUpdateList();
@@ -302,16 +299,14 @@ void FTextureRenderTarget2DArrayResource::ReleaseDynamicRHI()
  */
 void FTextureRenderTarget2DArrayResource::UpdateDeferredResource(FRHICommandListImmediate& RHICmdList, bool bClearRenderTarget/*=true*/)
 {
-	const FIntPoint Dims = GetSizeXY();
-
-	const ERenderTargetLoadAction LoadAction = bClearRenderTarget ? ERenderTargetLoadAction::EClear : ERenderTargetLoadAction::ELoad;
+	if (!bClearRenderTarget)
+	{
+		return;
+	}
 
-	FRHIRenderPassInfo RPInfo(RenderTarget2DArrayRHI, MakeRenderTargetActions(LoadAction, ERenderTargetStoreAction::EStore));
-	RHICmdList.Transition(FRHITransitionInfo(RenderTarget2DArrayRHI, ERHIAccess::Unknown, ERHIAccess::RTV));
-	RHICmdList.BeginRenderPass(RPInfo, TEXT("UpdateTarget2DArray"));
-	RHICmdList.SetViewport(0.0f, 0.0f, 0.0f, (float)Dims.X, (float)Dims.Y, 1.0f);
-	RHICmdList.EndRenderPass();
-	RHICmdList.CopyToResolveTarget(RenderTarget2DArrayRHI, Texture2DArrayRHI, FResolveParams());
+	RHICmdList.Transition(FRHITransitionInfo(TextureRHI, ERHIAccess::Unknown, ERHIAccess::RTV));
+	ClearRenderTarget(RHICmdList, TextureRHI);
+	RHICmdList.Transition(FRHITransitionInfo(TextureRHI, ERHIAccess::RTV, ERHIAccess::SRVMask));
 }
 
 /** 
@@ -366,7 +361,7 @@ bool FTextureRenderTarget2DArrayResource::ReadPixels(TArray<FColor>& OutImageDat
 		[RenderTarget_RT=this, OutData_RT=&OutImageData, Rect_RT=InRect, Slice_RT=InSlice, bSRGB_RT= bSRGB](FRHICommandListImmediate& RHICmdList)
 		{
 			TArray<FFloat16Color> TempData;
-			RHICmdList.ReadSurfaceFloatData(RenderTarget_RT->Texture2DArrayRHI, Rect_RT, TempData, (ECubeFace)0, Slice_RT, 0);
+			RHICmdList.ReadSurfaceFloatData(RenderTarget_RT->TextureRHI, Rect_RT, TempData, (ECubeFace)0, Slice_RT, 0);
 			for (const FFloat16Color& SrcColor : TempData)
 			{
 				OutData_RT->Emplace(FLinearColor(SrcColor).ToFColor(bSRGB_RT));
@@ -391,7 +386,7 @@ bool FTextureRenderTarget2DArrayResource::ReadPixels(TArray<FFloat16Color>& OutI
 	(
 		[RenderTarget_RT=this, OutData_RT=&OutImageData, Rect_RT=InRect, Slice_RT=InSlice](FRHICommandListImmediate& RHICmdList)
 		{
-			RHICmdList.ReadSurfaceFloatData(RenderTarget_RT->Texture2DArrayRHI, Rect_RT, *OutData_RT, (ECubeFace)0, Slice_RT, 0);
+			RHICmdList.ReadSurfaceFloatData(RenderTarget_RT->TextureRHI, Rect_RT, *OutData_RT, (ECubeFace)0, Slice_RT, 0);
 		}
 	);
 	FlushRenderingCommands();
diff --git a/Engine/Source/Runtime/Engine/Private/TextureRenderTargetCube.cpp b/Engine/Source/Runtime/Engine/Private/TextureRenderTargetCube.cpp
--- a/Engine/Source/Runtime/Engine/Private/TextureRenderTargetCube.cpp
+++ b/Engine/Source/Runtime/Engine/Private/TextureRenderTargetCube.cpp
@@ -247,37 +247,21 @@ void FTextureRenderTargetCubeResource::InitDynamicRHI()
 				.SetExtent(Owner->SizeX)
 				.SetFormat(Owner->GetFormat())
 				.SetNumMips(Owner->GetNumMips())
-				.SetFlags(TexCreateFlags)
-				.SetClearValue(FClearValueBinding(Owner->ClearColor));
+				.SetFlags(TexCreateFlags | ETextureCreateFlags::RenderTargetable | ETextureCreateFlags::ShaderResource)
+				.SetClearValue(FClearValueBinding(Owner->ClearColor))
+				.SetInitialState(ERHIAccess::SRVMask);
 
-			RHICreateTargetableShaderResource(Desc, ETextureCreateFlags::RenderTargetable, RenderTargetCubeRHI, TextureCubeRHI);
+			TextureRHI = RHICreateTexture(Desc);
 		}
 
-		if (EnumHasAnyFlags(TexCreateFlags, ETextureCreateFlags::UAV))
-		{
-			UnorderedAccessViewRHI = RHICreateUnorderedAccessView(RenderTargetCubeRHI);
-		}
+		SetGPUMask(FRHIGPUMask::All());
 
-		TextureRHI = TextureCubeRHI;
-		RHIUpdateTextureReference(Owner->TextureReference.TextureReferenceRHI,TextureRHI);
-
-		// Create the RHI target surface used for rendering to
+		if (EnumHasAnyFlags(TexCreateFlags, ETextureCreateFlags::UAV))
 		{
-			FRHIResourceCreateInfo CreateInfo(TEXT("FTextureRenderTargetCubeResource"), FClearValueBinding(Owner->ClearColor));
-			CubeFaceSurfaceRHI = RHICreateTexture2D(
-				Owner->SizeX, 
-				Owner->SizeX, 
-				Owner->GetFormat(),
-				Owner->GetNumMips(), 
-				/* NumSamples =*/ 1,
-				ETextureCreateFlags::RenderTargetable|TexCreateFlags,
-				CreateInfo
-				);
-			SetGPUMask(CreateInfo.GPUMask);
+			UnorderedAccessViewRHI = RHICreateUnorderedAccessView(TextureRHI);
 		}
 
-		// Set render target to 2D surface.
-		RenderTargetTextureRHI = CubeFaceSurfaceRHI;
+		RHIUpdateTextureReference(Owner->TextureReference.TextureReferenceRHI, TextureRHI);
 
 		AddToDeferredUpdateList(true);
 	}
@@ -304,9 +288,7 @@ void FTextureRenderTargetCubeResource::ReleaseDynamicRHI()
 	ReleaseRHI();
 
 	RHIUpdateTextureReference(Owner->TextureReference.TextureReferenceRHI, nullptr);
-	CubeFaceSurfaceRHI.SafeRelease();
-	RenderTargetCubeRHI.SafeRelease();
-	RenderTargetTextureRHI.SafeRelease();	
+	RenderTargetTextureRHI.SafeRelease();
 
 	// remove from global list of deferred clears
 	RemoveFromDeferredUpdateList();
@@ -319,26 +301,19 @@ void FTextureRenderTargetCubeResource::ReleaseDynamicRHI()
  */
 void FTextureRenderTargetCubeResource::UpdateDeferredResource(FRHICommandListImmediate& RHICmdList, bool bClearRenderTarget/*=true*/)
 {
-	const FIntPoint Dims = GetSizeXY();
+	if (!bClearRenderTarget)
+	{
+		return;
+	}
+
+	RHICmdList.Transition(FRHITransitionInfo(TextureRHI, ERHIAccess::Unknown, ERHIAccess::RTV));
+
 	for(int32 FaceIdx = CubeFace_PosX; FaceIdx < CubeFace_MAX; FaceIdx++)
 	{
-		ERenderTargetLoadAction LoadAction = ERenderTargetLoadAction::ELoad;
-		// clear each face of the cube target texture to ClearColor
-		if (bClearRenderTarget)
-		{
-			LoadAction = ERenderTargetLoadAction::EClear;
-		}
-		
-		FRHIRenderPassInfo RPInfo(RenderTargetTextureRHI, MakeRenderTargetActions(LoadAction, ERenderTargetStoreAction::EStore));
-		RHICmdList.Transition(FRHITransitionInfo(RenderTargetTextureRHI, ERHIAccess::Unknown, ERHIAccess::RTV));
-		RHICmdList.BeginRenderPass(RPInfo, TEXT("UpdateTargetCube"));
-		RHICmdList.SetViewport(0.0f, 0.0f, 0.0f, (float)Dims.X, (float)Dims.Y, 1.0f);
-		RHICmdList.EndRenderPass();
-		// copy surface to the texture for use
-		FResolveParams ResolveParams;
-		ResolveParams.CubeFace = (ECubeFace)FaceIdx;
-		RHICmdList.CopyToResolveTarget(RenderTargetTextureRHI, TextureCubeRHI, ResolveParams);
+		ClearRenderTarget(RHICmdList, TextureRHI, 0, FaceIdx);
 	}
+
+	RHICmdList.Transition(FRHITransitionInfo(TextureRHI, ERHIAccess::RTV, ERHIAccess::SRVMask));
 }
 
 /** 
@@ -416,7 +391,7 @@ bool FTextureRenderTargetCubeResource::ReadPixels(TArray< FColor >& OutImageData
 		[Context](FRHICommandListImmediate& RHICmdList)
 		{
 			RHICmdList.ReadSurfaceData(
-				Context.SrcRenderTarget->TextureCubeRHI,
+				Context.SrcRenderTarget->TextureRHI,
 				Context.Rect,
 				*Context.OutData,
 				Context.Flags
@@ -461,7 +436,7 @@ bool FTextureRenderTargetCubeResource::ReadPixels(TArray<FFloat16Color>& OutImag
 		[Context](FRHICommandListImmediate& RHICmdList)
 		{
 			RHICmdList.ReadSurfaceFloatData(
-				Context.SrcRenderTarget->TextureCubeRHI,
+				Context.SrcRenderTarget->TextureRHI,
 				Context.Rect,
 				*Context.OutData,
 				Context.CubeFace,
diff --git a/Engine/Source/Runtime/Engine/Private/TextureRenderTargetVolume.cpp b/Engine/Source/Runtime/Engine/Private/TextureRenderTargetVolume.cpp
--- a/Engine/Source/Runtime/Engine/Private/TextureRenderTargetVolume.cpp
+++ b/Engine/Source/Runtime/Engine/Private/TextureRenderTargetVolume.cpp
@@ -253,18 +253,17 @@ void FTextureRenderTargetVolumeResource::InitDynamicRHI()
 				.SetDepth(Owner->SizeZ)
 				.SetFormat(Owner->GetFormat())
 				.SetNumMips(Owner->GetNumMips())
-				.SetFlags(TexCreateFlags)
+				.SetFlags(TexCreateFlags | ETextureCreateFlags::RenderTargetable | ETextureCreateFlags::ShaderResource)
 				.SetClearValue(FClearValueBinding(Owner->ClearColor));
 
-			RHICreateTargetableShaderResource(Desc, ETextureCreateFlags::RenderTargetable, RenderTargetVolumeRHI, TextureVolumeRHI);
+			TextureRHI = RenderTargetTextureRHI = RHICreateTexture(Desc);
 		}
 
 		if (EnumHasAnyFlags(TexCreateFlags, TexCreate_UAV))
 		{
-			UnorderedAccessViewRHI = RHICreateUnorderedAccessView(RenderTargetVolumeRHI);
+			UnorderedAccessViewRHI = RHICreateUnorderedAccessView(TextureRHI);
 		}
 
-		TextureRHI = TextureVolumeRHI;
 		RHIUpdateTextureReference(Owner->TextureReference.TextureReferenceRHI,TextureRHI);
 
 		// Can't set this as it's a texture 2D
@@ -295,8 +294,7 @@ void FTextureRenderTargetVolumeResource::ReleaseDynamicRHI()
 	ReleaseRHI();
 
 	RHIUpdateTextureReference(Owner->TextureReference.TextureReferenceRHI, nullptr);
-	RenderTargetVolumeRHI.SafeRelease();
-	TextureVolumeRHI.SafeRelease();
+	RenderTargetTextureRHI.SafeRelease();
 
 	// remove from global list of deferred clears
 	RemoveFromDeferredUpdateList();
@@ -309,16 +307,12 @@ void FTextureRenderTargetVolumeResource::ReleaseDynamicRHI()
  */
 void FTextureRenderTargetVolumeResource::UpdateDeferredResource(FRHICommandListImmediate& RHICmdList, bool bClearRenderTarget/*=true*/)
 {
-	const FIntPoint Dims = GetSizeXY();
-
-	const ERenderTargetLoadAction LoadAction = bClearRenderTarget ? ERenderTargetLoadAction::EClear : ERenderTargetLoadAction::ELoad;
-
-	FRHIRenderPassInfo RPInfo(RenderTargetVolumeRHI, MakeRenderTargetActions(LoadAction, ERenderTargetStoreAction::EStore));
-	RHICmdList.Transition(FRHITransitionInfo(RenderTargetVolumeRHI, ERHIAccess::Unknown, ERHIAccess::RTV));
-	RHICmdList.BeginRenderPass(RPInfo, TEXT("UpdateTargetVolume"));
-	RHICmdList.SetViewport(0.0f, 0.0f, 0.0f, (float)Dims.X, (float)Dims.Y, 1.0f);
-	RHICmdList.EndRenderPass();
-	RHICmdList.CopyToResolveTarget(RenderTargetVolumeRHI, TextureVolumeRHI, FResolveParams());
+	if (bClearRenderTarget)
+	{
+		RHICmdList.Transition(FRHITransitionInfo(TextureRHI, ERHIAccess::Unknown, ERHIAccess::RTV));
+		ClearRenderTarget(RHICmdList, TextureRHI);
+		RHICmdList.Transition(FRHITransitionInfo(TextureRHI, ERHIAccess::RTV, ERHIAccess::SRVMask));
+	}
 }
 
 /** 
@@ -373,7 +367,7 @@ bool FTextureRenderTargetVolumeResource::ReadPixels(TArray<FColor>& OutImageData
 		[RenderTarget_RT=this, OutData_RT=&OutImageData, Rect_RT=InRect, DepthSlice_RT=InDepthSlice, bSRGB_RT= bSRGB](FRHICommandListImmediate& RHICmdList)
 		{
 			TArray<FFloat16Color> TempData;
-			RHICmdList.Read3DSurfaceFloatData(RenderTarget_RT->TextureVolumeRHI, Rect_RT, FIntPoint(DepthSlice_RT, DepthSlice_RT + 1), TempData);
+			RHICmdList.Read3DSurfaceFloatData(RenderTarget_RT->TextureRHI, Rect_RT, FIntPoint(DepthSlice_RT, DepthSlice_RT + 1), TempData);
 			for (const FFloat16Color& SrcColor : TempData)
 			{
 				OutData_RT->Emplace(FLinearColor(SrcColor).ToFColor(bSRGB_RT));
@@ -398,7 +392,7 @@ bool FTextureRenderTargetVolumeResource::ReadPixels(TArray<FFloat16Color>& OutIm
 	(
 		[RenderTarget_RT=this, OutData_RT=&OutImageData, Rect_RT=InRect, DepthSlice_RT=InDepthSlice](FRHICommandListImmediate& RHICmdList)
 		{
-			RHICmdList.Read3DSurfaceFloatData(RenderTarget_RT->TextureVolumeRHI, Rect_RT, FIntPoint(DepthSlice_RT, DepthSlice_RT+1), *OutData_RT);
+			RHICmdList.Read3DSurfaceFloatData(RenderTarget_RT->TextureRHI, Rect_RT, FIntPoint(DepthSlice_RT, DepthSlice_RT+1), *OutData_RT);
 		}
 	);
 	FlushRenderingCommands();
diff --git a/Engine/Source/Runtime/Engine/Private/UnrealClient.cpp b/Engine/Source/Runtime/Engine/Private/UnrealClient.cpp
--- a/Engine/Source/Runtime/Engine/Private/UnrealClient.cpp
+++ b/Engine/Source/Runtime/Engine/Private/UnrealClient.cpp
@@ -1727,7 +1727,6 @@ const TArray<FColor>& FViewport::GetRawHitProxyData(FIntRect InRect)
 				// Copy (resolve) the rendered thumbnail from the render target to its texture
 				RHICmdList.Transition(FRHITransitionInfo(RenderTargetTexture, ERHIAccess::Unknown, ERHIAccess::CopySrc));
 				RHICmdList.CopyToResolveTarget(RenderTargetTexture, HitProxyMapPtr->GetHitProxyCPUTexture(), ResolveParams);
-				RHICmdList.CopyToResolveTarget(RenderTargetTexture, HitProxyMapPtr->GetHitProxyTexture(), FResolveParams());
 			});
 
 		ENQUEUE_RENDER_COMMAND(EndDrawingCommand)(
@@ -2102,9 +2101,11 @@ void FViewport::FHitProxyMap::Init(uint32 NewSizeX,uint32 NewSizeY)
 			FRHITextureCreateDesc::Create2D(TEXT("HitProxyTexture"))
 			.SetExtent(SizeX, SizeY)
 			.SetFormat(PF_B8G8R8A8)
-			.SetClearValue(FClearValueBinding::White);
+			.SetFlags(ETextureCreateFlags::RenderTargetable | ETextureCreateFlags::ShaderResource)
+			.SetClearValue(FClearValueBinding::White)
+			.SetInitialState(ERHIAccess::SRVMask);
 
-		RHICreateTargetableShaderResource(Desc, ETextureCreateFlags::RenderTargetable, RenderTargetTextureRHI, HitProxyTexture);
+		RenderTargetTextureRHI = RHICreateTexture(Desc);
 	}
 	{
 		const FRHITextureCreateDesc Desc =
@@ -2119,7 +2120,6 @@ void FViewport::FHitProxyMap::Init(uint32 NewSizeX,uint32 NewSizeY)
 
 void FViewport::FHitProxyMap::Release()
 {
-	HitProxyTexture.SafeRelease();
 	HitProxyCPUTexture.SafeRelease();
 	RenderTargetTextureRHI.SafeRelease();
 }
@@ -2388,7 +2388,9 @@ void FDummyViewport::InitDynamicRHI()
 	const FRHITextureCreateDesc Desc =
 		FRHITextureCreateDesc::Create2D(TEXT("FDummyViewport"))
 		.SetExtent(SizeX, SizeY)
-		.SetFormat(PF_A2B10G10R10);
+		.SetFormat(PF_A2B10G10R10)
+		.SetFlags(ETextureCreateFlags::RenderTargetable | ETextureCreateFlags::ShaderResource)
+		.SetInitialState(ERHIAccess::SRVMask);
 
-	RHICreateTargetableShaderResource(Desc, ETextureCreateFlags::RenderTargetable, RenderTargetTextureRHI);
-}
+	RenderTargetTextureRHI = RHICreateTexture(Desc);
+}
\ No newline at end of file
diff --git a/Engine/Source/Runtime/Engine/Public/PlanarReflectionSceneProxy.h b/Engine/Source/Runtime/Engine/Public/PlanarReflectionSceneProxy.h
--- a/Engine/Source/Runtime/Engine/Public/PlanarReflectionSceneProxy.h
+++ b/Engine/Source/Runtime/Engine/Public/PlanarReflectionSceneProxy.h
@@ -26,11 +26,6 @@ public:
 		Size(InSize)
 	{}
 
-	virtual const FTexture2DRHIRef& GetRenderTargetTexture() const 
-	{
-		return (const FTexture2DRHIRef&)TextureRHI;
-	}
-
 	virtual void InitDynamicRHI()
 	{
 		// Create the sampler state RHI resource.
@@ -47,9 +42,11 @@ public:
 			FRHITextureCreateDesc::Create2D(TEXT("FPlanarReflectionRenderTarget"))
 			.SetExtent(GetSizeXY())
 			.SetFormat(PF_FloatRGBA)
-			.SetClearValue(FClearValueBinding::Black);
+			.SetClearValue(FClearValueBinding::Black)
+			.SetFlags(ETextureCreateFlags::RenderTargetable | ETextureCreateFlags::ShaderResource)
+			.SetInitialState(ERHIAccess::SRVMask);
 
-		RHICreateTargetableShaderResource(Desc, ETextureCreateFlags::RenderTargetable, RenderTargetTextureRHI, TextureRHI);
+		RenderTargetTextureRHI = TextureRHI = RHICreateTexture(Desc);
 	}
 
 	virtual FIntPoint GetSizeXY() const { return Size; }
diff --git a/Engine/Source/Runtime/Engine/Public/Slate/SlateTextures.h b/Engine/Source/Runtime/Engine/Public/Slate/SlateTextures.h
--- a/Engine/Source/Runtime/Engine/Public/Slate/SlateTextures.h
+++ b/Engine/Source/Runtime/Engine/Public/Slate/SlateTextures.h
@@ -153,9 +153,6 @@ public:
 	/** Resizes the render target */
 	virtual void SetSize(int32 InSizeX,int32 InSizeY);
 
-	/** Gets the RHI resource for this render target */
-	FTexture2DRHIRef GetTextureRHI() { return TextureRHI; }
-
 public:
 	// FTextureRenderTargetResource implementation
 	virtual void ClampSize(int32 SizeX,int32 SizeY) override;
diff --git a/Engine/Source/Runtime/Engine/Public/TextureRenderTarget2DArrayResource.h b/Engine/Source/Runtime/Engine/Public/TextureRenderTarget2DArrayResource.h
--- a/Engine/Source/Runtime/Engine/Public/TextureRenderTarget2DArrayResource.h
+++ b/Engine/Source/Runtime/Engine/Public/TextureRenderTarget2DArrayResource.h
@@ -63,11 +63,6 @@ public:
 	 */
 	virtual FIntPoint GetSizeXY() const override;
 
-	/**
-	 * @return TextureRHI for rendering
-	 */
-	FTexture2DArrayRHIRef GetTextureRHI() { return Texture2DArrayRHI; }
-
 	/**
 	 * @return UnorderedAccessView for rendering
 	 */
@@ -111,9 +106,11 @@ private:
 	const class UTextureRenderTarget2DArray* Owner;
 
 	/** Represents the current render target (from one of the slices)*/
-	FTexture2DArrayRHIRef RenderTarget2DArrayRHI;
+	UE_DEPRECATED(5.1, "RenderTarget2DArrayRHI is deprecated. Use TextureRHI instead.")
+	FTextureRHIRef RenderTarget2DArrayRHI;
 	/** Texture resource used for rendering with and resolving to */
-	FTexture2DArrayRHIRef Texture2DArrayRHI;
+	UE_DEPRECATED(5.1, "Texture2DArrayRHI is deprecated. Use TextureRHI instead.")
+	FTextureRHIRef Texture2DArrayRHI;
 	/** Optional Unordered Access View for the resource, automatically created if bCanCreateUAV is true */
 	FUnorderedAccessViewRHIRef UnorderedAccessViewRHI;
 };
diff --git a/Engine/Source/Runtime/Engine/Public/TextureRenderTargetVolumeResource.h b/Engine/Source/Runtime/Engine/Public/TextureRenderTargetVolumeResource.h
--- a/Engine/Source/Runtime/Engine/Public/TextureRenderTargetVolumeResource.h
+++ b/Engine/Source/Runtime/Engine/Public/TextureRenderTargetVolumeResource.h
@@ -63,11 +63,6 @@ public:
 	 */
 	virtual FIntPoint GetSizeXY() const override;
 
-	/**
-	 * @return TextureRHI for rendering
-	 */
-	FTexture3DRHIRef GetTextureRHI() { return TextureVolumeRHI; }
-
 	/**
 	 * @return UnorderedAccessView for rendering
 	 */
@@ -110,10 +105,11 @@ private:
 	/** The UTextureRenderTargetVolume which this resource represents. */
 	const class UTextureRenderTargetVolume* Owner;
 
-	/** Represents the current render target (from one of the slices)*/
-	FTexture3DRHIRef RenderTargetVolumeRHI;
-	/** Texture resource used for rendering with and resolving to */
-	FTexture3DRHIRef TextureVolumeRHI;
+	UE_DEPRECATED(5.1, "RenderTargetVolumeRHI has been deprecated. Use RenderTargetTextureRHI instead.")
+	FTextureRHIRef RenderTargetVolumeRHI;
+	UE_DEPRECATED(5.1, "TextureVolumeRHI has been deprecated. Use TextureRHI instead.")
+	FTextureRHIRef TextureVolumeRHI;
+
 	/** Optional Unordered Access View for the resource, automatically created if bCanCreateUAV is true */
 	FUnorderedAccessViewRHIRef UnorderedAccessViewRHI;
 };
diff --git a/Engine/Source/Runtime/Engine/Public/TextureResource.h b/Engine/Source/Runtime/Engine/Public/TextureResource.h
--- a/Engine/Source/Runtime/Engine/Public/TextureResource.h
+++ b/Engine/Source/Runtime/Engine/Public/TextureResource.h
@@ -127,17 +127,17 @@ public:
 		return TextureRHI.IsValid() && !!(TextureRHI->GetFlags() & TexCreate_Virtual);
 	}
 
-	FORCEINLINE FRHITexture2D* GetTexture2DRHI() const
+	FORCEINLINE FRHITexture* GetTexture2DRHI() const
 	{
 		return TextureRHI.IsValid() ? TextureRHI->GetTexture2D() : nullptr;
 	}
 
-	FORCEINLINE FRHITexture3D* GetTexture3DRHI() const
+	FORCEINLINE FRHITexture* GetTexture3DRHI() const
 	{
 		return TextureRHI.IsValid() ? TextureRHI->GetTexture3D() : nullptr;
 	}
 
-	FORCEINLINE FRHITexture2DArray* GetTexture2DArrayRHI() const
+	FORCEINLINE FRHITexture* GetTexture2DArrayRHI() const
 	{
 		return TextureRHI.IsValid() ? TextureRHI->GetTexture2DArray() : nullptr;
 	}
@@ -458,11 +458,6 @@ public:
 	 */
 	virtual float GetDisplayGamma() const override;
 
-	/** 
-	 * @return TextureRHI for rendering 
-	 */
-	FTexture2DRHIRef GetTextureRHI() { return Texture2DRHI; }
-
 	/**
 	 * @return UnorderedAccessView for rendering
 	 */
@@ -482,6 +477,7 @@ private:
 	/** The UTextureRenderTarget2D which this resource represents. */
 	const class UTextureRenderTarget2D* Owner;
 	/** Texture resource used for rendering with and resolving to */
+	UE_DEPRECATED(5.1, "Texture2DRHI is deprecated. Use TextureRHI instead.")
 	FTexture2DRHIRef Texture2DRHI;
 	/** Optional Unordered Access View for the resource, automatically created if bCanCreateUAV is true */
 	FUnorderedAccessViewRHIRef UnorderedAccessViewRHI;
@@ -549,11 +545,6 @@ public:
 	 */
 	virtual FIntPoint GetSizeXY() const override;
 
-	/** 
-	 * @return TextureRHI for rendering 
-	 */
-	FTextureCubeRHIRef GetTextureRHI() { return TextureCubeRHI; }
-
 	/**
 	 * @return UnorderedAccessView for rendering
 	 */
@@ -596,13 +587,15 @@ protected:
 private:
 	/** The UTextureRenderTargetCube which this resource represents. */
 	const class UTextureRenderTargetCube* Owner;
-	/** Texture resource used for rendering with and resolving to */
-	FTextureCubeRHIRef TextureCubeRHI;
-	/** Target surfaces for each cube face */
-	FTexture2DRHIRef CubeFaceSurfaceRHI;
 
-	/** Represents the current render target (from one of the cube faces)*/
-	FTextureCubeRHIRef RenderTargetCubeRHI;
+	UE_DEPRECATED(5.1, "TextureCubeRHI is deprecated. Use TextureRHI instead.")
+	FTextureRHIRef TextureCubeRHI;
+
+	UE_DEPRECATED(5.1, "CubeFaceSurfaceRHI is deprecated. Use TextureRHI instead.")
+	FTextureRHIRef CubeFaceSurfaceRHI;
+
+	UE_DEPRECATED(5.1, "RenderTargetCubeRHI is deprecated. Use TextureRHI instead.")
+	FTextureRHIRef RenderTargetCubeRHI;
 
 	/** Optional Unordered Access View for the resource, automatically created if bCanCreateUAV is true */
 	FUnorderedAccessViewRHIRef UnorderedAccessViewRHI;
diff --git a/Engine/Source/Runtime/Engine/Public/UnrealClient.h b/Engine/Source/Runtime/Engine/Public/UnrealClient.h
--- a/Engine/Source/Runtime/Engine/Public/UnrealClient.h
+++ b/Engine/Source/Runtime/Engine/Public/UnrealClient.h
@@ -691,7 +691,7 @@ protected:
 		virtual void AddReferencedObjects( FReferenceCollector& Collector ) override;
 		virtual FString GetReferencerName() const override;
 
-		const FTexture2DRHIRef& GetHitProxyTexture(void) const		{ return HitProxyTexture; }
+		const FTexture2DRHIRef& GetHitProxyTexture(void) const		{ return RenderTargetTextureRHI; }
 		const FTexture2DRHIRef& GetHitProxyCPUTexture(void) const		{ return HitProxyCPUTexture; }
 
 	private:
@@ -705,7 +705,6 @@ protected:
 		/** References to the hit proxies cached by the hit proxy map. */
 		TArray<TRefCountPtr<HHitProxy> > HitProxies;
 
-		FTexture2DRHIRef HitProxyTexture;
 		FTexture2DRHIRef HitProxyCPUTexture;
 	};
 
diff --git a/Engine/Source/Runtime/MediaAssets/Private/Misc/MediaTextureResource.cpp b/Engine/Source/Runtime/MediaAssets/Private/Misc/MediaTextureResource.cpp
--- a/Engine/Source/Runtime/MediaAssets/Private/Misc/MediaTextureResource.cpp
+++ b/Engine/Source/Runtime/MediaAssets/Private/Misc/MediaTextureResource.cpp
@@ -1206,10 +1206,11 @@ void FMediaTextureResource::CreateOutputRenderTarget(const FIntPoint & InDim, EP
 			.SetExtent(InDim)
 			.SetFormat(InPixelFormat)
 			.SetNumMips(InNumMips)
-			.SetFlags(OutputCreateFlags)
+			.SetFlags(OutputCreateFlags | ETextureCreateFlags::RenderTargetable | ETextureCreateFlags::ShaderResource)
+			.SetInitialState(ERHIAccess::SRVMask)
 			.SetClearValue(FClearValueBinding(InClearColor));
 
-		RHICreateTargetableShaderResource(Desc, ETextureCreateFlags::RenderTargetable, OutputTarget);
+		OutputTarget = RHICreateTexture(Desc);
 
 		OutputTarget->SetName(TEXT("MediaTextureResourceOutput"));
 
diff --git a/Engine/Source/Runtime/RHI/Public/RHIUtilities.h b/Engine/Source/Runtime/RHI/Public/RHIUtilities.h
--- a/Engine/Source/Runtime/RHI/Public/RHIUtilities.h
+++ b/Engine/Source/Runtime/RHI/Public/RHIUtilities.h
@@ -660,46 +660,6 @@ inline void TransitionRenderPassTargets(FRHICommandList& RHICmdList, const FRHIR
 	RHICmdList.Transition(MakeArrayView(Transitions, TransitionIndex));
 }
 
-RHI_API void RHICreateTargetableShaderResource(
-	const FRHITextureCreateDesc& BaseDesc,
-	ETextureCreateFlags TargetableTextureFlags,
-	FTextureRHIRef& OutTargetableTexture);
-
-/**
- * Creates 1 or 2 textures with the same dimensions/format.
- * If the RHI supports textures that can be used as both shader resources and render targets,
- * and bForceSeparateTargetAndShaderResource=false, then a single texture is created.
- * Otherwise two textures are created, one of them usable as a shader resource and resolve target, and one of them usable as a render target.
- * Two texture references are always returned, but they may reference the same texture.
- * If two different textures are returned, the render-target texture must be manually copied to the shader-resource texture.
- */
-RHI_API void RHICreateTargetableShaderResource(
-	const FRHITextureCreateDesc& Desc,
-	ETextureCreateFlags TargetableTextureFlags,
-	bool bForceSeparateTargetAndShaderResource,
-	bool bForceSharedTargetAndShaderResource,
-	FTextureRHIRef& OutTargetableTexture,
-	FTextureRHIRef& OutShaderResourceTexture);
-
-inline void RHICreateTargetableShaderResource(
-	const FRHITextureCreateDesc& Desc,
-	ETextureCreateFlags TargetableTextureFlags,
-	bool bForceSeparateTargetAndShaderResource,
-	FTextureRHIRef& OutTargetableTexture,
-	FTextureRHIRef& OutShaderResourceTexture)
-{
-	RHICreateTargetableShaderResource(Desc, TargetableTextureFlags, bForceSeparateTargetAndShaderResource, false, OutTargetableTexture, OutShaderResourceTexture);
-}
-
-inline void RHICreateTargetableShaderResource(
-	const FRHITextureCreateDesc& Desc,
-	ETextureCreateFlags TargetableTextureFlags,
-	FTextureRHIRef& OutTargetableTexture,
-	FTextureRHIRef& OutShaderResourceTexture)
-{
-	RHICreateTargetableShaderResource(Desc, TargetableTextureFlags, false, false, OutTargetableTexture, OutShaderResourceTexture);
-}
-
 /**
  * Creates 1 or 2 textures with the same dimensions/format.
  * If the RHI supports textures that can be used as both shader resources and render targets,
@@ -708,7 +668,7 @@ inline void RHICreateTargetableShaderResource(
  * Two texture references are always returned, but they may reference the same texture.
  * If two different textures are returned, the render-target texture must be manually copied to the shader-resource texture.
  */
-UE_DEPRECATED(5.1, "RHICreateTargetableShaderResource2D is deprecated, RHICreateTargetableShaderResource with a FRHITextureCreateDesc should be used instead")
+UE_DEPRECATED(5.1, "RHICreateTargetableShaderResource2D is deprecated. RHICreateTexture should be used instead.")
 RHI_API void RHICreateTargetableShaderResource2D(
 	uint32 SizeX,
 	uint32 SizeY,
@@ -723,7 +683,7 @@ RHI_API void RHICreateTargetableShaderResource2D(
 	FTextureRHIRef& OutShaderResourceTexture,
 	uint32 NumSamples = 1);
 
-UE_DEPRECATED(5.1, "RHICreateTargetableShaderResource2D is deprecated, RHICreateTargetableShaderResource with a FRHITextureCreateDesc should be used instead")
+UE_DEPRECATED(5.1, "RHICreateTargetableShaderResource2D is deprecated. RHICreateTexture should be used instead.")
 RHI_API void RHICreateTargetableShaderResource2D(
 	uint32 SizeX,
 	uint32 SizeY,
@@ -737,7 +697,7 @@ RHI_API void RHICreateTargetableShaderResource2D(
 	FTextureRHIRef& OutShaderResourceTexture,
 	uint32 NumSamples = 1);
 
-UE_DEPRECATED(5.1, "RHICreateTargetableShaderResource2DArray is deprecated, RHICreateTargetableShaderResource with a FRHITextureCreateDesc should be used instead")
+UE_DEPRECATED(5.1, "RHICreateTargetableShaderResource2DArray is deprecated. RHICreateTexture should be used instead.")
 RHI_API void RHICreateTargetableShaderResource2DArray(
 	uint32 SizeX,
 	uint32 SizeY,
@@ -753,7 +713,7 @@ RHI_API void RHICreateTargetableShaderResource2DArray(
 	FTextureRHIRef& OutShaderResourceTexture,
 	uint32 NumSamples = 1);
 
-UE_DEPRECATED(5.1, "RHICreateTargetableShaderResource2DArray is deprecated, RHICreateTargetableShaderResource with a FRHITextureCreateDesc should be used instead")
+UE_DEPRECATED(5.1, "RHICreateTargetableShaderResource2DArray is deprecated. RHICreateTexture should be used instead.")
 RHI_API void RHICreateTargetableShaderResource2DArray(
 	uint32 SizeX,
 	uint32 SizeY,
@@ -775,7 +735,7 @@ RHI_API void RHICreateTargetableShaderResource2DArray(
  * Two texture references are always returned, but they may reference the same texture.
  * If two different textures are returned, the render-target texture must be manually copied to the shader-resource texture.
  */
-UE_DEPRECATED(5.1, "RHICreateTargetableShaderResourceCube is deprecated, RHICreateTargetableShaderResource with a FRHITextureCreateDesc should be used instead")
+UE_DEPRECATED(5.1, "RHICreateTargetableShaderResourceCube is deprecated. RHICreateTexture should be used instead.")
 RHI_API void RHICreateTargetableShaderResourceCube(
 	uint32 LinearSize,
 	uint8 Format,
@@ -795,7 +755,7 @@ RHI_API void RHICreateTargetableShaderResourceCube(
  * Two texture references are always returned, but they may reference the same texture.
  * If two different textures are returned, the render-target texture must be manually copied to the shader-resource texture.
  */
-UE_DEPRECATED(5.1, "RHICreateTargetableShaderResourceCubeArray is deprecated, RHICreateTargetableShaderResource with a FRHITextureCreateDesc should be used instead")
+UE_DEPRECATED(5.1, "RHICreateTargetableShaderResourceCubeArray is deprecated. RHICreateTexture should be used instead.")
 RHI_API void RHICreateTargetableShaderResourceCubeArray(
 	uint32 LinearSize,
 	uint32 ArraySize,
@@ -816,7 +776,7 @@ RHI_API void RHICreateTargetableShaderResourceCubeArray(
  * Two texture references are always returned, but they may reference the same texture.
  * If two different textures are returned, the render-target texture must be manually copied to the shader-resource texture.
  */
-UE_DEPRECATED(5.1, "RHICreateTargetableShaderResource3D is deprecated, RHICreateTargetableShaderResource with a FRHITextureCreateDesc should be used instead")
+UE_DEPRECATED(5.1, "RHICreateTargetableShaderResource3D is deprecated. RHICreateTexture should be used instead.")
 RHI_API void RHICreateTargetableShaderResource3D(
 	uint32 SizeX,
 	uint32 SizeY,
@@ -830,6 +790,18 @@ RHI_API void RHICreateTargetableShaderResource3D(
 	FTextureRHIRef& OutTargetableTexture,
 	FTextureRHIRef& OutShaderResourceTexture);
 
+/** Performs a clear render pass on an RHI texture. The texture is expected to be in the RTV state. */
+inline void ClearRenderTarget(FRHICommandList& RHICmdList, FRHITexture* Texture, uint32 MipIndex = 0, uint32 ArraySlice = 0)
+{
+	check(Texture);
+	const FIntPoint Extent = Texture->GetSizeXY();
+	FRHIRenderPassInfo Info(Texture, ERenderTargetActions::Clear_Store);
+	Info.ColorRenderTargets[0].MipIndex = MipIndex;
+	Info.ColorRenderTargets[0].ArraySlice = ArraySlice;
+	RHICmdList.BeginRenderPass(Info, TEXT("ClearRenderTarget"));
+	RHICmdList.EndRenderPass();
+}
+
 /**
  * Computes the vertex count for a given number of primitives of the specified type.
  * @param NumPrimitives The number of primitives.
diff --git a/Engine/Source/Runtime/RenderCore/Public/RenderResource.h b/Engine/Source/Runtime/RenderCore/Public/RenderResource.h
--- a/Engine/Source/Runtime/RenderCore/Public/RenderResource.h
+++ b/Engine/Source/Runtime/RenderCore/Public/RenderResource.h
@@ -452,7 +452,9 @@ public:
 
 	// Destructor
 	virtual ~FTexture() {}
-	
+
+	FRHITexture* GetTextureRHI() { return TextureRHI; }
+
 	/** Returns the width of the texture in pixels. */
 	virtual uint32 GetSizeX() const
 	{
diff --git a/Engine/Source/Runtime/Renderer/Private/RayTracing/RaytracingSkylight.cpp b/Engine/Source/Runtime/Renderer/Private/RayTracing/RaytracingSkylight.cpp
--- a/Engine/Source/Runtime/Renderer/Private/RayTracing/RaytracingSkylight.cpp
+++ b/Engine/Source/Runtime/Renderer/Private/RayTracing/RaytracingSkylight.cpp
@@ -773,10 +773,9 @@ void FDeferredShadingSceneRenderer::VisualizeSkyLightMipTree(
 	RHICmdList.EndRenderPass();
 	GVisualizeTexture.SetCheckPoint(RHICmdList, SkyLightMipTreeRT);
 
-	RHICmdList.CopyToResolveTarget(SceneColor->GetRHI(), SceneColor->GetRHI(), FResolveParams());
-
 	// Transition to compute
 	RHICmdList.Transition({
+		FRHITransitionInfo(SceneColor->GetRHI(), ERHIAccess::RTV, ERHIAccess::SRVMask),
 		FRHITransitionInfo(SkyLightMipTreePosX.UAV, ERHIAccess::SRVGraphics, ERHIAccess::UAVCompute),
 		FRHITransitionInfo(SkyLightMipTreeNegX.UAV, ERHIAccess::SRVGraphics, ERHIAccess::UAVCompute),
 		FRHITransitionInfo(SkyLightMipTreePosY.UAV, ERHIAccess::SRVGraphics, ERHIAccess::UAVCompute),
diff --git a/Engine/Source/Runtime/SlateRHIRenderer/Private/Slate3DRenderer.cpp b/Engine/Source/Runtime/SlateRHIRenderer/Private/Slate3DRenderer.cpp
--- a/Engine/Source/Runtime/SlateRHIRenderer/Private/Slate3DRenderer.cpp
+++ b/Engine/Source/Runtime/SlateRHIRenderer/Private/Slate3DRenderer.cpp
@@ -155,7 +155,7 @@ void FSlate3DRenderer::DrawWindowToTarget_RenderThread(FRHICommandListImmediate&
 	RenderTargetPolicy->BeginDrawingWindows();
 
 	// Set render target and clear.
-	FTexture2DRHIRef RTTextureRHI = Context.RenderTarget->GetRenderTargetTexture();
+	FRHITexture* RTTextureRHI = Context.RenderTarget->GetRenderTargetTexture();
 	InRHICmdList.Transition(FRHITransitionInfo(RTTextureRHI, ERHIAccess::Unknown, ERHIAccess::RTV));
 	
 	FRHIRenderPassInfo RPInfo(RTTextureRHI, ERenderTargetActions::Load_Store);
@@ -199,9 +199,10 @@ void FSlate3DRenderer::DrawWindowToTarget_RenderThread(FRHICommandListImmediate&
 						FRHITextureCreateDesc::Create2D(TEXT("SlateWindowDepthStencil"))
 						.SetExtent(ColorTarget->GetSizeXY())
 						.SetFormat(PF_DepthStencil)
-						.SetClearValue(FClearValueBinding::DepthZero);
+						.SetClearValue(FClearValueBinding::DepthZero)
+						.SetFlags(ETextureCreateFlags::DepthStencilTargetable | ETextureCreateFlags::ShaderResource);
 
-					RHICreateTargetableShaderResource(Desc, ETextureCreateFlags::DepthStencilTargetable, DepthStencil);
+					DepthStencil = RHICreateTexture(Desc);
 
 					check(IsValidRef(DepthStencil));
 				}
@@ -244,7 +245,7 @@ void FSlate3DRenderer::DrawWindowToTarget_RenderThread(FRHICommandListImmediate&
 	}
 
 	FSlateEndDrawingWindowsCommand::EndDrawingWindows(InRHICmdList, Context.WindowDrawBuffer, *RenderTargetPolicy);
-	InRHICmdList.CopyToResolveTarget(Context.RenderTarget->GetRenderTargetTexture(), RTTextureRHI, FResolveParams());
+	InRHICmdList.Transition(FRHITransitionInfo(RTTextureRHI, ERHIAccess::RTV, ERHIAccess::SRVMask));
 
 	// Enqueue a command to keep "this" alive.
 	InRHICmdList.EnqueueLambda([Self = SharedThis(this)](FRHICommandListImmediate&){});
diff --git a/Engine/Source/Runtime/SlateRHIRenderer/Private/SlatePostProcessResource.cpp b/Engine/Source/Runtime/SlateRHIRenderer/Private/SlatePostProcessResource.cpp
--- a/Engine/Source/Runtime/SlateRHIRenderer/Private/SlatePostProcessResource.cpp
+++ b/Engine/Source/Runtime/SlateRHIRenderer/Private/SlatePostProcessResource.cpp
@@ -45,10 +45,11 @@ void FSlatePostProcessResource::ResizeTargets(const FIntPoint& NewSize)
 			const FRHITextureCreateDesc Desc =
 				FRHITextureCreateDesc::Create2D(TEXT("FSlatePostProcessResource"))
 				.SetExtent(RenderTargetSize)
-				.SetFormat(PixelFormat);
+				.SetFormat(PixelFormat)
+				.SetFlags(ETextureCreateFlags::RenderTargetable | ETextureCreateFlags::ShaderResource)
+				.SetInitialState(ERHIAccess::SRVMask);
 
-			FTexture2DRHIRef RenderTargetTextureRHI;
-			RHICreateTargetableShaderResource(Desc, ETextureCreateFlags::RenderTargetable, RenderTargetTextureRHI);
+			FTexture2DRHIRef RenderTargetTextureRHI = RHICreateTexture(Desc);
 
 			RenderTargets.Add(RenderTargetTextureRHI);
 		}
diff --git a/Engine/Source/Runtime/SlateRHIRenderer/Private/SlateRHIRenderer.cpp b/Engine/Source/Runtime/SlateRHIRenderer/Private/SlateRHIRenderer.cpp
--- a/Engine/Source/Runtime/SlateRHIRenderer/Private/SlateRHIRenderer.cpp
+++ b/Engine/Source/Runtime/SlateRHIRenderer/Private/SlateRHIRenderer.cpp
@@ -159,10 +159,11 @@ void FViewportInfo::RecreateDepthBuffer_RenderThread()
 			FRHITextureCreateDesc::Create2D(TEXT("SlateViewportDepthStencil"))
 			.SetExtent(Width, Height)
 			.SetFormat(PF_DepthStencil)
+			.SetFlags(TargetableTextureFlags | ETextureCreateFlags::ShaderResource)
+			.SetInitialState(ERHIAccess::SRVMask)
 			.SetClearValue(FClearValueBinding::DepthZero);
 
-		RHICreateTargetableShaderResource(Desc, TargetableTextureFlags, DepthStencil);
-		
+		DepthStencil = RHICreateTexture(Desc);
 		check(IsValidRef(DepthStencil));
 	}
 }
@@ -898,16 +899,17 @@ void FSlateRHIRenderer::DrawWindow_RenderThread(FRHICommandListImmediate& RHICmd
 
 				// LUT
 				{
-					ViewportInfo.ColorSpaceLUTRT.SafeRelease();
-					ViewportInfo.ColorSpaceLUTSRV.SafeRelease();
+					ViewportInfo.ColorSpaceLUT.SafeRelease();
 
 					const FRHITextureCreateDesc Desc =
 						FRHITextureCreateDesc::Create3D(TEXT("ColorSpaceLUT"))
 						.SetExtent(CompositionLUTSize)
 						.SetDepth(CompositionLUTSize)
-						.SetFormat(PF_A2B10G10R10);
+						.SetFormat(PF_A2B10G10R10)
+						.SetFlags(ETextureCreateFlags::RenderTargetable | ETextureCreateFlags::ShaderResource)
+						.SetInitialState(ERHIAccess::SRVMask);
 
-					RHICreateTargetableShaderResource(Desc, ETextureCreateFlags::RenderTargetable, ViewportInfo.ColorSpaceLUTRT, ViewportInfo.ColorSpaceLUTSRV);
+					ViewportInfo.ColorSpaceLUT = RHICreateTexture(Desc);
 				}
 
 				bLUTStale = true;
@@ -951,9 +953,6 @@ void FSlateRHIRenderer::DrawWindow_RenderThread(FRHICommandListImmediate& RHICmd
 
 				GRenderTargetPool.FindFreeElement(RHICmdList, Desc, HDRRenderRT, TEXT("HDRTargetRT"));
 
-				// FIXME: is this necessary?
-				RHICmdList.Transition(FRHITransitionInfo(FinalBuffer, ERHIAccess::Unknown, ERHIAccess::SRVMask));
-
 				BackBuffer = HDRRenderRT->GetRHI();
 			}
 #endif
@@ -992,8 +991,8 @@ void FSlateRHIRenderer::DrawWindow_RenderThread(FRHICommandListImmediate& RHICmd
 				if (bLUTStale)
 				{
 					// #todo-renderpasses will this touch every pixel? use NoAction?
-					FRHIRenderPassInfo RPInfo(ViewportInfo.ColorSpaceLUTRT, ERenderTargetActions::Load_Store);
-					RHICmdList.Transition(FRHITransitionInfo(ViewportInfo.ColorSpaceLUTRT, ERHIAccess::Unknown, ERHIAccess::RTV));
+					FRHIRenderPassInfo RPInfo(ViewportInfo.ColorSpaceLUT, ERenderTargetActions::Load_Store);
+					RHICmdList.Transition(FRHITransitionInfo(ViewportInfo.ColorSpaceLUT, ERHIAccess::Unknown, ERHIAccess::RTV));
 					RHICmdList.BeginRenderPass(RPInfo, TEXT("GenerateLUT"));
 					{
 						FGraphicsPipelineStateInitializer GraphicsPSOInit;
@@ -1026,7 +1025,7 @@ void FSlateRHIRenderer::DrawWindow_RenderThread(FRHICommandListImmediate& RHICmd
 						RasterizeToVolumeTexture(RHICmdList, VolumeBounds);
 					}
 					RHICmdList.EndRenderPass();
-					RHICmdList.CopyToResolveTarget(ViewportInfo.ColorSpaceLUTRT, ViewportInfo.ColorSpaceLUTSRV, FResolveParams());
+					RHICmdList.Transition(FRHITransitionInfo(ViewportInfo.ColorSpaceLUT, ERHIAccess::RTV, ERHIAccess::SRVMask));
 				}
 
 				// Composition pass
@@ -1067,7 +1066,7 @@ void FSlateRHIRenderer::DrawWindow_RenderThread(FRHICommandListImmediate& RHICmd
 
 							SetGraphicsPipelineState(RHICmdList, GraphicsPSOInit, 0);
 
-							PixelShader->SetParameters(RHICmdList, ViewportInfo.UITargetRT->GetRHI(), UITargetRTMaskTexture, ViewportInfo.HDRSourceRT->GetRHI(), ViewportInfo.ColorSpaceLUTSRV);
+							PixelShader->SetParameters(RHICmdList, ViewportInfo.UITargetRT->GetRHI(), UITargetRTMaskTexture, ViewportInfo.HDRSourceRT->GetRHI(), ViewportInfo.ColorSpaceLUT);
 						}
 						else
 						{
@@ -1081,7 +1080,7 @@ void FSlateRHIRenderer::DrawWindow_RenderThread(FRHICommandListImmediate& RHICmd
 
 							SetGraphicsPipelineState(RHICmdList, GraphicsPSOInit, 0);
 
-							PixelShader->SetParameters(RHICmdList, ViewportInfo.UITargetRT->GetRHI(), UITargetRTMaskTexture, ViewportInfo.HDRSourceRT->GetRHI(), ViewportInfo.ColorSpaceLUTSRV);
+							PixelShader->SetParameters(RHICmdList, ViewportInfo.UITargetRT->GetRHI(), UITargetRTMaskTexture, ViewportInfo.HDRSourceRT->GetRHI(), ViewportInfo.ColorSpaceLUT);
 						}
 
 						RendererModule.DrawRectangle(
diff --git a/Engine/Source/Runtime/SlateRHIRenderer/Private/SlateRHIRenderer.h b/Engine/Source/Runtime/SlateRHIRenderer/Private/SlateRHIRenderer.h
--- a/Engine/Source/Runtime/SlateRHIRenderer/Private/SlateRHIRenderer.h
+++ b/Engine/Source/Runtime/SlateRHIRenderer/Private/SlateRHIRenderer.h
@@ -67,8 +67,7 @@ struct FViewportInfo : public FRenderResource
 	TRefCountPtr<IPooledRenderTarget> HDRSourceRT;
 
 	/** Color-space LUT for HDR UI composition. */
-	FTexture3DRHIRef ColorSpaceLUTRT;
-	FTexture3DRHIRef ColorSpaceLUTSRV;
+	FTextureRHIRef ColorSpaceLUT;
 	int32 ColorSpaceLUTOutputDevice;
 	int32 ColorSpaceLUTOutputGamut;
 		
@@ -130,8 +129,7 @@ struct FViewportInfo : public FRenderResource
 	~FViewportInfo()
 	{
 		DepthStencil.SafeRelease();
-		ColorSpaceLUTRT.SafeRelease();
-		ColorSpaceLUTSRV.SafeRelease();
+		ColorSpaceLUT.SafeRelease();
 	}
 
 	void ConditionallyUpdateDepthBuffer(bool bInRequiresStencilTest, uint32 Width, uint32 Height);
diff --git a/Engine/Source/Runtime/WebBrowserTexture/Private/WebBrowserTextureResource.cpp b/Engine/Source/Runtime/WebBrowserTexture/Private/WebBrowserTextureResource.cpp
--- a/Engine/Source/Runtime/WebBrowserTexture/Private/WebBrowserTextureResource.cpp
+++ b/Engine/Source/Runtime/WebBrowserTexture/Private/WebBrowserTextureResource.cpp
@@ -146,7 +146,7 @@ void FWebBrowserTextureResource::ClearTexture(const FLinearColor& ClearColor)
 {
 	FPlatformMisc::LowLevelOutputDebugStringf(TEXT("FWebBrowserTextureResource:ClearTexture"));
 	// create output render target if we don't have one yet
-	const ETextureCreateFlags OutputCreateFlags = TexCreate_Dynamic | TexCreate_SRGB;
+	const ETextureCreateFlags OutputCreateFlags = ETextureCreateFlags::Dynamic | ETextureCreateFlags::SRGB;
 
 	if ((ClearColor != CurrentClearColor) || !OutputTarget.IsValid() || !EnumHasAllFlags(OutputTarget->GetFlags(), OutputCreateFlags))
 	{
@@ -154,10 +154,11 @@ void FWebBrowserTextureResource::ClearTexture(const FLinearColor& ClearColor)
 			FRHITextureCreateDesc::Create2D(TEXT("FWebBrowserTextureResource"))
 			.SetExtent(2, 2)
 			.SetFormat(PF_B8G8R8A8)
-			.SetFlags(OutputCreateFlags)
+			.SetFlags(OutputCreateFlags | ETextureCreateFlags::RenderTargetable | ETextureCreateFlags::ShaderResource)
+			.SetInitialState(ERHIAccess::SRVMask)
 			.SetClearValue(FClearValueBinding(ClearColor));
 
-		RHICreateTargetableShaderResource(Desc, ETextureCreateFlags::RenderTargetable, OutputTarget);
+		OutputTarget = RHICreateTexture(Desc);
 
 		CurrentClearColor = ClearColor;
 		UpdateResourceSize();
@@ -171,10 +172,8 @@ void FWebBrowserTextureResource::ClearTexture(const FLinearColor& ClearColor)
 	// draw the clear color
 	FRHICommandListImmediate& CommandList = FRHICommandListExecutor::GetImmediateCommandList();
 	{
-		FRHIRenderPassInfo RPInfo(RenderTargetTextureRHI, ERenderTargetActions::Clear_Store);
-		CommandList.BeginRenderPass(RPInfo, TEXT("ClearTexture"));
-		CommandList.EndRenderPass();
-
+		CommandList.Transition(FRHITransitionInfo(RenderTargetTextureRHI, ERHIAccess::Unknown, ERHIAccess::RTV));
+		ClearRenderTarget(CommandList, RenderTargetTextureRHI);
 		CommandList.Transition(FRHITransitionInfo(RenderTargetTextureRHI, ERHIAccess::RTV, ERHIAccess::SRVMask));
 	}
 
@@ -216,10 +215,11 @@ void FWebBrowserTextureResource::CopySample(const TSharedPtr<FWebBrowserTextureS
 				FRHITextureCreateDesc::Create2D(TEXT("FWebBrowserTextureResource"))
 				.SetExtent(SampleDim)
 				.SetFormat(PF_B8G8R8A8)
-				.SetFlags(OutputCreateFlags)
+				.SetFlags(OutputCreateFlags | ETextureCreateFlags::RenderTargetable | ETextureCreateFlags::ShaderResource)
+				.SetInitialState(ERHIAccess::SRVMask)
 				.SetClearValue(FClearValueBinding(ClearColor));
 
-			RHICreateTargetableShaderResource(Desc, ETextureCreateFlags::RenderTargetable, OutputTarget);
+			OutputTarget = RHICreateTexture(Desc);
 
 			CurrentClearColor = ClearColor;
 			UpdateResourceSize();
diff --git a/Engine/Source/Runtime/WebBrowserTexture/Public/WebBrowserTextureSample.h b/Engine/Source/Runtime/WebBrowserTexture/Public/WebBrowserTextureSample.h
--- a/Engine/Source/Runtime/WebBrowserTexture/Public/WebBrowserTextureSample.h
+++ b/Engine/Source/Runtime/WebBrowserTexture/Public/WebBrowserTextureSample.h
@@ -131,9 +131,10 @@ public:
 			FRHITextureCreateDesc::Create2D(TEXT("FWebBrowserTextureSample"))
 			.SetExtent(Dim)
 			.SetFormat(PF_B8G8R8A8)
-			.SetFlags(ETextureCreateFlags::Dynamic | ETextureCreateFlags::SRGB);
+			.SetFlags(ETextureCreateFlags::Dynamic | ETextureCreateFlags::SRGB | ETextureCreateFlags::RenderTargetable | ETextureCreateFlags::ShaderResource)
+			.SetInitialState(ERHIAccess::SRVMask);
 
-		RHICreateTargetableShaderResource(Desc, ETextureCreateFlags::RenderTargetable, Texture);
+		Texture = RHICreateTexture(Desc);
 
 		return Texture;
 	}
