Log more info when hitting conduit as entry state ensure on non-editor builds (on behalf of Jose.Villarroel)
#rnx

diff --git a/Engine/Source/Runtime/Engine/Private/Animation/AnimNode_StateMachine.cpp b/Engine/Source/Runtime/Engine/Private/Animation/AnimNode_StateMachine.cpp
--- a/Engine/Source/Runtime/Engine/Private/Animation/AnimNode_StateMachine.cpp
+++ b/Engine/Source/Runtime/Engine/Private/Animation/AnimNode_StateMachine.cpp
@@ -778,12 +778,13 @@ void FAnimNode_StateMachine::Evaluate_AnyThread(FPoseContext& Output)
 		{ 
 			UAnimBlueprint* AnimBlueprint = Output.AnimInstanceProxy->GetAnimBlueprint();
 
-			FText Message = FText::Format(LOCTEXT("CurrentStateIsConduitWarning", "Current state is a conduit and will reset to ref pose. This can happen if a conduit is an entry state and there isn't at least one valid transition to take. Conduit({0}), State Machine({1}), AnimBlueprint({2})."),
+			FText Message = FText::Format(LOCTEXT("CurrentStateIsConduitWarning", "Current state is a conduit and will reset to ref pose. This can happen if a conduit is an entry state and there isn't at least one valid transition to take. State Machine({0}), Conduit({1}), AnimBlueprint({2})."),
 				FText::FromName(Machine->MachineName), FText::FromName(GetCurrentStateName()), FText::FromString(GetPathNameSafe(AnimBlueprint)));
 			Output.LogMessage(EMessageSeverity::Error, Message);
 		}
 #else
-		ensure(!bCurrentStateIsConduit);
+		ensureMsgf(!bCurrentStateIsConduit, TEXT("Current state is a conduit and will reset to ref pose. State Machine(%s), Conduit(%s), AnimInstance(%s)."),
+			*Machine->MachineName.ToString(), *GetCurrentStateName().ToString(), *Output.AnimInstanceProxy->GetAnimInstanceName());
 #endif
 
 		if (bCurrentStateIsConduit || Machine->States.Num() == 0 || !Machine->States.IsValidIndex(CurrentState))
