Implemented optional support for assigning a TransactionId for LevelStatusUpdate and ServerUpdateLevelVisibility requests. This is used in multiplayer to allow server to rely on levelvisibility status reported by clients to know what data to replicate to clients.

Behavior is controlled by two CVars:
LevelStreaming.ShouldClientUseMakingInvisibleTransactionRequest
- if enabled client will send a request to server and wait for an ack before making a streaming level invisible

LevelStreaming.ShouldServerUseMakingVisibleTransactionRequest
- If enabled server will associate a transaction id to each call to ClientUpdateLevelStreamingStatus and wait for the response from the client of that particular request before updating LevelVisibility for replicated data. This is to solve some issues where we might have multiple requests in flight.

By default both CVars are disabled.

diff --git a/Engine/Source/Runtime/Engine/Classes/Engine/LevelStreaming.h b/Engine/Source/Runtime/Engine/Classes/Engine/LevelStreaming.h
--- a/Engine/Source/Runtime/Engine/Classes/Engine/LevelStreaming.h
+++ b/Engine/Source/Runtime/Engine/Classes/Engine/LevelStreaming.h
@@ -20,6 +20,15 @@ class ALevelScriptActor;
 class ALevelStreamingVolume;
 class ULevel;
 class ULevelStreaming;
+struct FNetLevelVisibilityTransactionId;
+
+struct FNetLevelVisibilityState
+{
+	uint32 ServerRequestIndex = 0;
+	uint32 PendingClientRequestIndex = 0;
+	uint32 AckedClientRequestIndex = 0;
+	bool bPendingInvisibleRequest = false;
+};
 
 ENGINE_API DECLARE_LOG_CATEGORY_EXTERN(LogLevelStreaming, Log, All);
 
@@ -299,6 +308,20 @@ private:
 	/** Internal function for checking if the desired level is the currently loaded level */
 	bool IsDesiredLevelLoaded() const;
 
+public:
+
+	/** Begin a client instigated NetVisibility request */
+	void BeginClientNetVisibilityRequest(bool bInShouldBeVisible);
+
+	/** Check if we are waiting for a making invisible streaming transaction */
+	bool IsWaitingForNetVisibilityTransactionAck() const;
+
+	/** Set the current state of the current visibility/streaming transaction */
+	void UpdateNetVisibilityTransactionState(bool bInShouldBeVisible, FNetLevelVisibilityTransactionId TransactionId);
+
+	/** Ack a client instigated visibility/streaming transaction */
+	void AckNetVisibilityTransaction(FNetLevelVisibilityTransactionId AckedClientTransactionId);
+
 public:
 
 	/** Returns the value of bShouldBeVisible. Use ShouldBeVisible to query whether a streaming level should be visible based on its own criteria. */
@@ -499,6 +522,12 @@ public:
 	UPROPERTY(BlueprintAssignable)
 	FLevelStreamingVisibilityStatus		OnLevelHidden;
 
+	/** If true client will wait for acknowledgment from server before making streaming levels invisible */
+	static bool ShouldClientUseMakingInvisibleTransactionRequest();
+
+	/** If true server will wait for client acknowledgment before making treating streaming levels as visible for the client */
+	static bool ShouldServerUseMakingVisibleTransactionRequest();
+
 	/** 
 	 * Traverses all streaming level objects in the persistent world and in all inner worlds and calls appropriate delegate for streaming objects that refer specified level 
 	 *
@@ -594,6 +623,9 @@ private:
 	/** The cached package name of the currently loaded level. */
 	mutable FName CachedLoadedLevelPackageName;
 
+	/** State for server handshake of NetLevelVisibilty requests */
+	FNetLevelVisibilityState NetVisibilityState;
+
 	friend struct FStreamingLevelPrivateAccessor;
 };
 
diff --git a/Engine/Source/Runtime/Engine/Classes/Engine/NetConnection.h b/Engine/Source/Runtime/Engine/Classes/Engine/NetConnection.h
--- a/Engine/Source/Runtime/Engine/Classes/Engine/NetConnection.h
+++ b/Engine/Source/Runtime/Engine/Classes/Engine/NetConnection.h
@@ -45,6 +45,7 @@ class FObjectReplicator;
 class StatelessConnectHandlerComponent;
 class UActorChannel;
 class UChildConnection;
+class ULevelStreaming;
 struct FEncryptionKeyResponse;
 
 namespace UE::Net
@@ -777,7 +778,7 @@ public:
 	 */
 	TSet<FName> ClientVisibleLevelNames;
 
-	/** Called by PlayerController to tell connection about client level visiblity change */
+	/** Called by PlayerController to tell connection about client level visibility change */
 	ENGINE_API void UpdateLevelVisibility(const struct FUpdateLevelVisibilityLevelInfo& LevelVisibility);
 	
 #if DO_ENABLE_NET_TEST
@@ -1439,6 +1440,9 @@ public:
 
 	ENGINE_API virtual void NotifyActorChannelCleanedUp(UActorChannel* Channel, EChannelCloseReason CloseReason);
 
+	/** Update FNetLevelVisibilityTransactionId for server instigated level streaming */
+	FNetLevelVisibilityTransactionId UpdateLevelStreamStatusChangedTransactionId(const ULevelStreaming* LevelObject, const FName PackageName, bool bShouldBeVisible);
+
 protected:
 
 	bool GetPendingCloseDueToSocketSendFailure() const
@@ -1486,6 +1490,9 @@ private:
 	/** This is used to capture visibility updates while the server is in transition and deffer the update until the server has completed the transition */
 	TMap<FName, FUpdateLevelVisibilityLevelInfo> PendingUpdateLevelVisibility;
 
+	/** This is used to track any pending streaming status requests the server has sent */
+	TMap<FName, FNetLevelVisibilityTransactionId> ClientPendingStreamingStatusRequest;
+
 private:
 
 	/** Called by PlayerController to tell connection about client level visibility change */
diff --git a/Engine/Source/Runtime/Engine/Classes/Engine/World.h b/Engine/Source/Runtime/Engine/Classes/Engine/World.h
--- a/Engine/Source/Runtime/Engine/Classes/Engine/World.h
+++ b/Engine/Source/Runtime/Engine/Classes/Engine/World.h
@@ -16,6 +16,7 @@
 #include "CollisionQueryParams.h"
 #include "WorldCollision.h"
 #include "GameFramework/Pawn.h"
+#include "GameFramework/UpdateLevelVisibilityLevelInfo.h"
 #include "EngineDefines.h"
 #include "Engine/Blueprint.h"
 #include "Engine/PendingNetGame.h"
@@ -2940,15 +2941,18 @@ public:
 	 *
 	 * @param Level				Level object we should add
 	 * @param LevelTransform	Transformation to apply to each actor in the level
+	 * @param bConsiderTimeLimie optional bool indicating if we should consider timelimit or not, default is true
+	 * @param TransactionId optional parameter that carries the current transaction id associated with calls updating LevelVisibility used when communicating level visibility with server
 	 */
-	void AddToWorld( ULevel* Level, const FTransform& LevelTransform = FTransform::Identity, bool bConsiderTimeLimit = true );
+	void AddToWorld(ULevel* Level, const FTransform& LevelTransform = FTransform::Identity, bool bConsiderTimeLimit = true, FNetLevelVisibilityTransactionId TransactionId = FNetLevelVisibilityTransactionId());
 
 	/** 
 	 * Dissociates the passed in level from the world. The removal is blocking.
 	 *
 	 * @param Level			Level object we should remove
+	 * @param TransactionId optional parameter that carries the current transaction id associated with calls updating LevelVisibility used when communicating level visibility with server
 	 */
-	void RemoveFromWorld( ULevel* Level, bool bAllowIncrementalRemoval = false );
+	void RemoveFromWorld(ULevel* Level, bool bAllowIncrementalRemoval = false, FNetLevelVisibilityTransactionId TransactionId = FNetLevelVisibilityTransactionId());
 
 	/**
 	 * Updates sub-levels (load/unload/show/hide) using streaming levels current state
diff --git a/Engine/Source/Runtime/Engine/Classes/GameFramework/PlayerController.h b/Engine/Source/Runtime/Engine/Classes/GameFramework/PlayerController.h
--- a/Engine/Source/Runtime/Engine/Classes/GameFramework/PlayerController.h
+++ b/Engine/Source/Runtime/Engine/Classes/GameFramework/PlayerController.h
@@ -813,6 +813,14 @@ public:
 	UFUNCTION(Reliable, Client)
 	virtual void ClientEnableNetworkVoice(bool bEnable);
 
+	/** 
+	 * Acknowledge received LevelVisibilityTransactionId
+	 * @param PackageName - Identifying the level that we are acknowledging levelvisibility for
+	 * @param TransactionId - TransactionId being acknowledged
+	 */
+	UFUNCTION(Reliable, Client)
+	void ClientAckUpdateLevelVisibility(FName PackageName, FNetLevelVisibilityTransactionId TransactionId);
+
 	/** Enable voice chat transmission */
 	void StartTalking();
 
@@ -1341,9 +1349,11 @@ public:
 	 * @param bNewShouldBeVisible - Whether the level should be visible if it is loaded	
 	 * @param bNewShouldBlockOnLoad - Whether we want to force a blocking load
 	 * @param LODIndex				- Current LOD index for a streaming level
+	 * @param TransactionId			- Optional parameter used when communicating LevelVisibility changes between server and client
 	 */
 	UFUNCTION(Reliable, Client)
-	void ClientUpdateLevelStreamingStatus(FName PackageName, bool bNewShouldBeLoaded, bool bNewShouldBeVisible, bool bNewShouldBlockOnLoad, int32 LODIndex);
+	void ClientUpdateLevelStreamingStatus(FName PackageName, bool bNewShouldBeLoaded, bool bNewShouldBeVisible, bool bNewShouldBlockOnLoad, int32 LODIndex, FNetLevelVisibilityTransactionId TransactionId);
+	void ClientUpdateLevelStreamingStatus(FName PackageName, bool bNewShouldBeLoaded, bool bNewShouldBeVisible, bool bNewShouldBlockOnLoad, int32 LODIndex) { ClientUpdateLevelStreamingStatus(PackageName, bNewShouldBeLoaded, bNewShouldBeVisible, bNewShouldBlockOnLoad, LODIndex, FNetLevelVisibilityTransactionId()); }
 
 	/**
 	 * Replicated Update streaming status.  This version allows for the streaming state of many levels to be sent in a single RPC.
@@ -1439,17 +1449,6 @@ public:
 	UFUNCTION(reliable, server, WithValidation, SealedEvent)
 	void ServerUpdateLevelVisibility(const FUpdateLevelVisibilityLevelInfo& LevelVisibility);
 
-	UE_DEPRECATED(4.24, "Use ServerUpdateLevelVisibility that accepts a LevelVisibility struct.")
-	void ServerUpdateLevelVisibility(FName PackageName, bool bIsVisible)
-	{
-		FUpdateLevelVisibilityLevelInfo LevelVisibility;
-		LevelVisibility.PackageName = PackageName;
-		LevelVisibility.FileName = PackageName;
-		LevelVisibility.bIsVisible = bIsVisible;
-
-		ServerUpdateLevelVisibility(LevelVisibility);
-	}
-
 	/** 
 	 * Called when the client adds/removes a streamed level.  This version of the function allows you to pass the state of 
 	 * multiple levels at once, to reduce the number of RPC events that will be sent.
diff --git a/Engine/Source/Runtime/Engine/Classes/GameFramework/UpdateLevelVisibilityLevelInfo.h b/Engine/Source/Runtime/Engine/Classes/GameFramework/UpdateLevelVisibilityLevelInfo.h
--- a/Engine/Source/Runtime/Engine/Classes/GameFramework/UpdateLevelVisibilityLevelInfo.h
+++ b/Engine/Source/Runtime/Engine/Classes/GameFramework/UpdateLevelVisibilityLevelInfo.h
@@ -6,6 +6,48 @@
 #include "UObject/ObjectMacros.h"
 #include "UpdateLevelVisibilityLevelInfo.generated.h"
 
+/** This structure is used to to identify NetLevelVisibility transactions between server and client */
+USTRUCT()
+struct ENGINE_API FNetLevelVisibilityTransactionId
+{
+	GENERATED_BODY()
+
+	enum : uint32
+	{
+		InvalidTransactionIndex = 0U,
+		InvalidTransactionId = 0U,
+		IsClientTransactionMask = 0x80000000U,
+		ValueMask = ~IsClientTransactionMask,
+	};
+
+	FNetLevelVisibilityTransactionId();
+
+	bool IsClientTransaction() const { return Data & IsClientTransactionMask; }
+	bool IsValid() const { return Data != InvalidTransactionId; }
+	void SetIsClientInstigator(bool bValue) { Data = (Data & ~IsClientTransactionMask) | (bValue ? IsClientTransactionMask : 0U); }
+	uint32 GetTransactionIndex() const { return Data & ValueMask; }
+	void SetTransactionIndex(uint32 TransactionIndex) { Data = (Data & IsClientTransactionMask) | (TransactionIndex & ValueMask); }
+
+	uint32 IncrementTransactionIndex();
+	bool operator==(const FNetLevelVisibilityTransactionId& Other) const { return Data == Other.Data; }
+	bool NetSerialize(FArchive& Ar, UPackageMap* PackageMap, bool& bOutSuccess);
+
+private:
+	FNetLevelVisibilityTransactionId(uint32 LevelVisibilityTransactionIndex, bool bIsClientInstigator);
+
+	UPROPERTY()
+	uint32 Data;
+};
+
+template<>
+struct TStructOpsTypeTraits<FNetLevelVisibilityTransactionId> : public TStructOpsTypeTraitsBase2<FNetLevelVisibilityTransactionId>
+{
+	enum
+	{
+		WithNetSerializer = true
+	};
+};
+
 /** This structure is used to pass arguments to ServerUpdateLevelVisibilty() and ServerUpdateMultipleLevelsVisibility() server RPC functions */
 USTRUCT()
 struct ENGINE_API FUpdateLevelVisibilityLevelInfo
@@ -34,6 +76,10 @@ struct ENGINE_API FUpdateLevelVisibilityLevelInfo
 	UPROPERTY()
 	FName FileName;
 
+	/** Identifies this visibility request when communicating with server */
+	UPROPERTY()
+	FNetLevelVisibilityTransactionId VisibilityRequestId;
+
 	/** The new visibility state for this level. */
 	UPROPERTY()
 	uint32 bIsVisible : 1;
@@ -52,3 +98,25 @@ struct TStructOpsTypeTraits<FUpdateLevelVisibilityLevelInfo> : public TStructOps
 		WithNetSerializer = true
 	};
 };
+
+inline FNetLevelVisibilityTransactionId::FNetLevelVisibilityTransactionId()
+: Data(InvalidTransactionId)
+{
+}
+
+inline FNetLevelVisibilityTransactionId::FNetLevelVisibilityTransactionId(uint32 LevelVisibilityTransactionIndex, bool bIsClientInstigator)
+: Data((bIsClientInstigator ? IsClientTransactionMask : 0U) | (LevelVisibilityTransactionIndex & ValueMask))
+{
+}
+
+inline uint32 FNetLevelVisibilityTransactionId::IncrementTransactionIndex()
+{
+	uint32 NewIndex = (GetTransactionIndex() + 1U) & ValueMask;
+	// Deal with wraparound since zero is reserved as the invalid index
+	if (NewIndex == InvalidTransactionIndex)
+	{
+		++NewIndex;
+	}
+	SetTransactionIndex(NewIndex);
+	return NewIndex;
+}
diff --git a/Engine/Source/Runtime/Engine/Private/GameFramework/UpdateLevelVisibilityLevelInfo.cpp b/Engine/Source/Runtime/Engine/Private/GameFramework/UpdateLevelVisibilityLevelInfo.cpp
--- a/Engine/Source/Runtime/Engine/Private/GameFramework/UpdateLevelVisibilityLevelInfo.cpp
+++ b/Engine/Source/Runtime/Engine/Private/GameFramework/UpdateLevelVisibilityLevelInfo.cpp
@@ -48,9 +48,37 @@ bool FUpdateLevelVisibilityLevelInfo::NetSerialize(FArchive& Ar, UPackageMap* Pa
 		FileName = PackageName;
 	}
 
+	VisibilityRequestId.NetSerialize(Ar, PackageMap, bOutSuccess);
+
 	bIsVisible = bLocalIsVisible;
 
 	bOutSuccess = !Ar.IsError();
 	return true;
 }
 
+bool FNetLevelVisibilityTransactionId::NetSerialize(FArchive& Ar, UPackageMap* PackageMap, bool& bOutSuccess)
+{
+	if (Ar.IsLoading())
+	{
+		bool bIsClientInstigator = false;
+		uint32 Value = 0U;
+
+		Ar.SerializeBits(&bIsClientInstigator, 1);
+		Ar.SerializeIntPacked(Value);
+		
+		*this = FNetLevelVisibilityTransactionId(Value, bIsClientInstigator);
+	}
+	else
+	{
+		bool bIsClientInstigator = IsClientTransaction();
+		uint32 Value = GetTransactionIndex();
+
+		Ar.SerializeBits(&bIsClientInstigator, 1);
+		Ar.SerializeIntPacked(Value);		
+	}
+
+	bOutSuccess = !Ar.IsError();
+	return true;
+}
+
+
diff --git a/Engine/Source/Runtime/Engine/Private/LevelStreaming.cpp b/Engine/Source/Runtime/Engine/Private/LevelStreaming.cpp
--- a/Engine/Source/Runtime/Engine/Private/LevelStreaming.cpp
+++ b/Engine/Source/Runtime/Engine/Private/LevelStreaming.cpp
@@ -29,6 +29,7 @@
 #include "PhysicsEngine/BodySetup.h"
 #include "SceneInterface.h"
 #include "Engine/NetDriver.h"
+#include "Engine/NetConnection.h"
 #include "Engine/PackageMapClient.h"
 #include "Serialization/LoadTimeTrace.h"
 #include "AssetRegistry/AssetRegistryModule.h"
@@ -37,6 +38,39 @@ DEFINE_LOG_CATEGORY(LogLevelStreaming);
 
 #define LOCTEXT_NAMESPACE "World"
 
+// CVars
+namespace LevelStreamingCVars
+{
+	// There are cases where we might have multiple visibility requests (and data) in flight leading to the server 
+	// starting to replicate data based on an older visibility/streamingstatus update which can lead to broken channels
+	// to mitigate this problem we assign a TransactionId to each request/update to make sure that we are acting on the correct data
+	static bool bShouldClientUseMakingInvisibleTransactionRequest = false;
+	FAutoConsoleVariableRef CVarShouldClientUseMakingInvisibleTransactionRequest(
+		TEXT("LevelStreaming.ShouldClientUseMakingInvisibleTransactionRequest"),
+		bShouldClientUseMakingInvisibleTransactionRequest,
+		TEXT("Whether client should wait for server to acknowledge visibility update before making streaming levels invisible.\n")
+		TEXT("0: Disable, 1: Enable"),
+		ECVF_Default);
+
+	static bool bShouldServerUseMakingVisibleTransactionRequest = false;
+	FAutoConsoleVariableRef CVarShouldServerUseMakingVisibleTransactionRequest(
+		TEXT("LevelStreaming.ShouldServerUseMakingVisibleTransactionRequest"),
+		bShouldServerUseMakingVisibleTransactionRequest,
+		TEXT("Whether server should wait for client to acknowledge visibility update before treating streaming levels as visible by the client.\n")
+		TEXT("0: Disable, 1: Enable"),
+		ECVF_Default);
+}
+
+bool ULevelStreaming::ShouldClientUseMakingInvisibleTransactionRequest()
+{
+	return LevelStreamingCVars::bShouldClientUseMakingInvisibleTransactionRequest;
+}
+
+bool ULevelStreaming::ShouldServerUseMakingVisibleTransactionRequest()
+{
+	return LevelStreamingCVars::bShouldServerUseMakingVisibleTransactionRequest;
+}
+
 int32 ULevelStreamingDynamic::UniqueLevelInstanceId = 0;
 
 /**
@@ -527,6 +561,96 @@ bool ULevelStreaming::DetermineTargetState()
 	return bContinueToConsider;
 }
 
+static bool ShouldWaitForServerAckBeforeMakingInvisible(const UWorld* World, const ULevel* LoadedLevel, const bool bShouldBeVisible, FNetLevelVisibilityState& NetVisibilityState)
+{
+	if (ULevelStreaming::ShouldClientUseMakingInvisibleTransactionRequest())
+	{
+		if (!LoadedLevel->bClientOnlyVisible && World->IsNetMode(NM_Client) && (World->NetDriver && World->NetDriver->ServerConnection->GetConnectionState() == USOCK_Open))
+		{
+			if (NetVisibilityState.PendingClientRequestIndex != NetVisibilityState.AckedClientRequestIndex)
+			{
+				// Wait for server to acknowledge the visibility change so that we are sure that we have no incoming data in flight when removing the actors from the world
+				return true;
+			}
+			else if (!bShouldBeVisible && NetVisibilityState.bPendingInvisibleRequest && NetVisibilityState.PendingClientRequestIndex == NetVisibilityState.AckedClientRequestIndex)
+			{
+				// We have a pending request, IncrementTransactionIndex and send ServerUpdateLevelVisibility request to server
+				FNetLevelVisibilityTransactionId TransactionId;
+				TransactionId.SetIsClientInstigator(true);
+				TransactionId.SetTransactionIndex(NetVisibilityState.PendingClientRequestIndex);
+				NetVisibilityState.PendingClientRequestIndex = TransactionId.IncrementTransactionIndex();
+				NetVisibilityState.bPendingInvisibleRequest = false;
+
+				for (FConstPlayerControllerIterator Iterator = World->GetPlayerControllerIterator(); Iterator; ++Iterator)
+				{
+					if (APlayerController* PlayerController = Iterator->Get())
+					{
+						FUpdateLevelVisibilityLevelInfo LevelVisibility(LoadedLevel, false);
+						LevelVisibility.PackageName = PlayerController->NetworkRemapPath(LevelVisibility.PackageName, false);
+						LevelVisibility.VisibilityRequestId = TransactionId;
+							
+						PlayerController->ServerUpdateLevelVisibility(LevelVisibility);
+					}
+				}
+				return true;
+			}
+		}
+	}
+
+	// Invalidate request
+	NetVisibilityState.bPendingInvisibleRequest = false;
+	NetVisibilityState.AckedClientRequestIndex = NetVisibilityState.PendingClientRequestIndex;	
+
+	return false;
+}
+
+bool ULevelStreaming::IsWaitingForNetVisibilityTransactionAck() const
+{
+	if (ShouldClientUseMakingInvisibleTransactionRequest())
+	{
+		if (LoadedLevel && !LoadedLevel->bClientOnlyVisible)
+		{
+			const UWorld* World = GetWorld();
+			if (World->IsNetMode(NM_Client) && (World->NetDriver && World->NetDriver->ServerConnection->GetConnectionState() == USOCK_Open) && ((NetVisibilityState.PendingClientRequestIndex != NetVisibilityState.AckedClientRequestIndex) || NetVisibilityState.bPendingInvisibleRequest))
+			{
+				return true;
+			}	
+		}
+	}
+
+	return false;
+}
+
+void ULevelStreaming::UpdateNetVisibilityTransactionState(bool bInShouldBeVisible, FNetLevelVisibilityTransactionId TransactionId)
+{
+	if (GetCurrentState() != ULevelStreaming::ECurrentState::MakingInvisible)
+	{
+		// If this is a client request to unload, we will conditionally send a notification to the server before we make the level invisible
+		const bool bIsClientTransaction = TransactionId.IsClientTransaction();
+
+		NetVisibilityState.bPendingInvisibleRequest = bIsClientTransaction && (bInShouldBeVisible == false);
+		NetVisibilityState.ServerRequestIndex = bIsClientTransaction ? FNetLevelVisibilityTransactionId::InvalidTransactionIndex : TransactionId.GetTransactionIndex();
+	}
+}
+
+void ULevelStreaming::BeginClientNetVisibilityRequest(bool bInShouldBeVisible)
+{
+	FNetLevelVisibilityTransactionId TransactionId;
+	TransactionId.SetIsClientInstigator(true);
+
+	UpdateNetVisibilityTransactionState(bInShouldBeVisible, TransactionId);
+}
+
+void ULevelStreaming::AckNetVisibilityTransaction(FNetLevelVisibilityTransactionId AckedClientTransactionId)
+{
+	if (ensure(NetVisibilityState.PendingClientRequestIndex != NetVisibilityState.AckedClientRequestIndex))
+	{
+		// Invalidate the pending request as there should only ever be a single one in flight
+		NetVisibilityState.AckedClientRequestIndex = AckedClientTransactionId.GetTransactionIndex();
+		NetVisibilityState.bPendingInvisibleRequest = false;
+	}
+}
+
 void ULevelStreaming::UpdateStreamingState(bool& bOutUpdateAgain, bool& bOutRedetermineTarget)
 {
 	FScopeCycleCounterUObject ContextScope(this);
@@ -576,7 +700,13 @@ void ULevelStreaming::UpdateStreamingState(bool& bOutUpdateAgain, bool& bOutRede
 			}
 			else
 			{
-				World->AddToWorld(LoadedLevel, LevelTransform, !bShouldBlockOnLoad);
+				// Only respond with ServerTransactionId if the is the target visibility state is supposed to be visible
+				FNetLevelVisibilityTransactionId TransactionId;
+				if (bShouldBeVisible)
+				{
+					TransactionId.SetTransactionIndex(NetVisibilityState.ServerRequestIndex);
+				}
+				World->AddToWorld(LoadedLevel, LevelTransform, !bShouldBlockOnLoad, TransactionId);
 
 				if (LoadedLevel->bIsVisible)
 				{
@@ -611,8 +741,21 @@ void ULevelStreaming::UpdateStreamingState(bool& bOutUpdateAgain, bool& bOutRede
 
 			const bool bWasVisible = LoadedLevel->bIsVisible;
 
+			// We do not want to have any changes in flights when ending play, so before making invisible we wait for server acknowledgment
+			if (ShouldWaitForServerAckBeforeMakingInvisible(World, LoadedLevel, bShouldBeVisible, NetVisibilityState))
+			{
+				break;
+			}
+
+			FNetLevelVisibilityTransactionId TransactionId;
+			// Only respond with ServerTransactionId if the is the target visibility state is supposed to be not visible
+			if (bShouldBeVisible == false)
+			{
+				TransactionId.SetTransactionIndex(NetVisibilityState.ServerRequestIndex);
+			}
+
 			// Hide loaded level, incrementally if necessary
-			World->RemoveFromWorld(LoadedLevel, !bShouldBlockOnUnload && World->IsGameWorld());
+			World->RemoveFromWorld(LoadedLevel, !bShouldBlockOnUnload && World->IsGameWorld(), TransactionId);
 
 			// Hide loaded level immediately if bRequireFullVisibilityToRender is set
 			const bool LevelBecameInvisible = bWasVisible && !LoadedLevel->bIsVisible;
diff --git a/Engine/Source/Runtime/Engine/Private/LevelUtils.cpp b/Engine/Source/Runtime/Engine/Private/LevelUtils.cpp
--- a/Engine/Source/Runtime/Engine/Private/LevelUtils.cpp
+++ b/Engine/Source/Runtime/Engine/Private/LevelUtils.cpp
@@ -50,32 +50,19 @@ bool FLevelUtils::bMovingLevel = false;
 bool FLevelUtils::bApplyingLevelTransform = false;
 #endif
 
-/**
- * Returns the streaming level corresponding to the specified ULevel, or NULL if none exists.
- *
- * @param		Level		The level to query.
- * @return					The level's streaming level, or NULL if none exists.
- */
 ULevelStreaming* FLevelUtils::FindStreamingLevel(const ULevel* Level)
 {
 	return ULevelStreaming::FindStreamingLevel(Level);
 }
 
-/**
- * Returns the streaming level by package name, or NULL if none exists.
- *
- * @param		PackageName		Name of the package containing the ULevel to query
- * @return						The level's streaming level, or NULL if none exists.
- */
-ULevelStreaming* FLevelUtils::FindStreamingLevel(UWorld* InWorld, const TCHAR* InPackageName)
+ULevelStreaming* FLevelUtils::FindStreamingLevel(UWorld* InWorld, const FName PackageName)
 {
-	const FName PackageName( InPackageName );
 	ULevelStreaming* MatchingLevel = NULL;
-	if( InWorld)
+	if (InWorld && !PackageName.IsNone())
 	{
 		for (ULevelStreaming* CurStreamingLevel : InWorld->GetStreamingLevels())
 		{
-			if( CurStreamingLevel && CurStreamingLevel->GetWorldAssetPackageFName() == PackageName )
+			if (CurStreamingLevel && CurStreamingLevel->GetWorldAssetPackageFName() == PackageName)
 			{
 				MatchingLevel = CurStreamingLevel;
 				break;
@@ -85,6 +72,11 @@ ULevelStreaming* FLevelUtils::FindStreamingLevel(UWorld* InWorld, const TCHAR* I
 	return MatchingLevel;
 }
 
+ULevelStreaming* FLevelUtils::FindStreamingLevel(UWorld* InWorld, const TCHAR* InPackageName)
+{
+	return FindStreamingLevel(InWorld, FName(InPackageName));
+}
+
 /////////////////////////////////////////////////////////////////////////////////////////
 //
 //	Level locking/unlocking.
diff --git a/Engine/Source/Runtime/Engine/Private/NetConnection.cpp b/Engine/Source/Runtime/Engine/Private/NetConnection.cpp
--- a/Engine/Source/Runtime/Engine/Private/NetConnection.cpp
+++ b/Engine/Source/Runtime/Engine/Private/NetConnection.cpp
@@ -41,6 +41,7 @@
 #include "UObject/UnrealNames.h"
 #include "HAL/LowLevelMemStats.h"
 #include "Net/NetPing.h"
+#include "LevelUtils.h"
 
 DECLARE_LLM_MEMORY_STAT(TEXT("NetConnection"), STAT_NetConnectionLLM, STATGROUP_LLMFULL);
 LLM_DEFINE_TAG(NetConnection, NAME_None, TEXT("Networking"), GET_STATFNAME(STAT_NetConnectionLLM), GET_STATFNAME(STAT_NetworkingSummaryLLM));
@@ -124,25 +125,15 @@ namespace UE_NetConnectionPrivate
 		// If we have a file it is fine too
 		// If its in our own streaming level list, its good
 
-		auto IsInLevelList = [&World](FName InPackageName)
-		{
-			for (ULevelStreaming* StreamingLevel : World->GetStreamingLevels())
-			{
-				if (StreamingLevel && (StreamingLevel->GetWorldAssetPackageFName() == InPackageName))
-				{
-					return true;
-				}
-			}
-			return false;
-		};
-
 		FValidateLevelVisibilityResult Result;
 		FString PackageNameStr = LevelVisibility.PackageName.ToString();
 		Result.Package = FindPackage(nullptr, *PackageNameStr);
 		Result.Linker = FLinkerLoad::FindExistingLinkerForPackage(Result.Package);
+
+		const ULevelStreaming* StreamingLevel = FLevelUtils::FindStreamingLevel(World, LevelVisibility.PackageName);
 		Result.bLevelExists = Result.Linker ||
 			(!FPackageName::IsMemoryPackage(PackageNameStr) && FPackageName::DoesPackageExist(LevelVisibility.FileName.ToString())) ||
-			IsInLevelList(LevelVisibility.PackageName);
+			StreamingLevel != nullptr;
 
 		return Result;
 	}
@@ -844,7 +835,7 @@ void UNetConnection::Serialize( FArchive& Ar )
 		GRANULAR_NETWORK_MEMORY_TRACKING_TRACK("DormantReplicatorMap", DormantReplicatorMap.CountBytes(Ar));
 
 		GRANULAR_NETWORK_MEMORY_TRACKING_TRACK("ClientVisibleLevelNames", ClientVisibleLevelNames.CountBytes(Ar));
-		GRANULAR_NETWORK_MEMORY_TRACKING_TRACK("ClientVisibileActorOuters", ClientVisibleActorOuters.CountBytes(Ar));
+		GRANULAR_NETWORK_MEMORY_TRACKING_TRACK("ClientVisibleActorOuters", ClientVisibleActorOuters.CountBytes(Ar));
 
 		GRANULAR_NETWORK_MEMORY_TRACKING_TRACK("ActorsStarvedByClassTimeMap",
 			ActorsStarvedByClassTimeMap.CountBytes(Ar);
@@ -1299,7 +1290,7 @@ void UNetConnection::AddReferencedObjects(UObject* InThis, FReferenceCollector&
 		}
 	}
 
-	// ClientVisibileActorOuters acceleration map
+	// ClientVisibleActorOuters acceleration map
 	for (auto& MapIt : This->ClientVisibleActorOuters)
 	{
 		Collector.AddReferencedObject(MapIt.Key, This);
@@ -1343,6 +1334,23 @@ void UNetConnection::AssertValid()
 
 }
 
+FNetLevelVisibilityTransactionId UNetConnection::UpdateLevelStreamStatusChangedTransactionId(const ULevelStreaming* LevelObject, const FName PackageName, bool bShouldBeVisible)
+{
+	// Increment transactionId
+	FNetLevelVisibilityTransactionId& TransactionId = ClientPendingStreamingStatusRequest.FindOrAdd(PackageName, FNetLevelVisibilityTransactionId());
+	TransactionId.IncrementTransactionIndex();
+
+	// Disable replication of actors on the invisible level
+	// We want to do this now to avoid having data in flight that the client will not be able to receive.
+	if (bShouldBeVisible == false)
+	{
+		ClientVisibleLevelNames.Remove(PackageName);
+		UpdateAllCachedLevelVisibility();
+	}
+
+	return TransactionId;
+}
+
 bool UNetConnection::ClientHasInitializedLevelFor(const AActor* TestActor) const
 {
 	checkSlow(Driver);
@@ -1401,7 +1409,7 @@ void UNetConnection::UpdateLevelVisibility(const FUpdateLevelVisibilityLevelInfo
 	if (Driver && Driver->GetWorld())
 	{
 		// If we are doing seamless travel we need to defer visibility updates until after the server has completed loading the level
-		// otherwise we might end up in a situation where visibilty is not correctly updated
+		// otherwise we might end up in a situation where visibility is not correctly updated
 		if (Driver->GetWorld()->IsInSeamlessTravel())
 		{
 			PendingUpdateLevelVisibility.FindOrAdd(LevelVisibility.PackageName) = LevelVisibility;
@@ -1418,13 +1426,28 @@ void UNetConnection::UpdateLevelVisibilityInternal(const FUpdateLevelVisibilityL
 
 	GNumClientUpdateLevelVisibility++;
 
+	bool bVerifiedTransaction = true;
+	if (ULevelStreaming::ShouldServerUseMakingVisibleTransactionRequest())
+	{
+		// If we have a valid server instigated VisibilityRequest we verify it against our pending transactions before treating the level as visible
+		// This is to avoid issues with multiple possibly conflicting requests in flight.
+		const FNetLevelVisibilityTransactionId VisibilityRequestId = LevelVisibility.VisibilityRequestId;
+		if (VisibilityRequestId.IsValid() && !VisibilityRequestId.IsClientTransaction())
+		{
+			if (FNetLevelVisibilityTransactionId* PendingTransactionIndex = ClientPendingStreamingStatusRequest.Find(LevelVisibility.PackageName))
+			{
+				bVerifiedTransaction = VisibilityRequestId == *PendingTransactionIndex;
+			}
+		}
+	}
+
 	// add or remove the level package name from the list, as requested
 	if (LevelVisibility.bIsVisible)
 	{
 		// verify that we were passed a valid level name
 		FValidateLevelVisibilityResult VisibilityResult = ValidateLevelVisibility(GetWorld(), LevelVisibility);
 
-		if (VisibilityResult.bLevelExists)
+		if (bVerifiedTransaction && VisibilityResult.bLevelExists)
 		{
 			ClientVisibleLevelNames.Add(LevelVisibility.PackageName);
 			UE_LOG(LogPlayerController, Verbose, TEXT("ServerUpdateLevelVisibility() Added '%s'"), *LevelVisibility.PackageName.ToString());
@@ -1469,9 +1492,8 @@ void UNetConnection::UpdateLevelVisibilityInternal(const FUpdateLevelVisibilityL
 			{
 				ReplicationConnectionDriver->NotifyClientVisibleLevelNamesAdd(LevelVisibility.PackageName, LevelWorld);
 			}
-
 		}
-		else
+		else if (!VisibilityResult.bLevelExists)
 		{
 			FString PackageNameStr = LevelVisibility.PackageName.ToString();
 			FString FileNameStr = LevelVisibility.FileName.ToString();
@@ -1499,7 +1521,7 @@ void UNetConnection::UpdateLevelVisibilityInternal(const FUpdateLevelVisibilityL
 		{
 			ReplicationConnectionDriver->NotifyClientVisibleLevelNamesRemove(LevelVisibility.PackageName);
 		}
-			
+
 		// Close any channels now that have actors that were apart of the level the client just unloaded
 		for (auto It = ActorChannels.CreateIterator(); It; ++It)
 		{
@@ -3422,6 +3444,12 @@ void UNetConnection::ReceivedPacket( FBitReader& Reader, bool bIsReinjectedPacke
 
 	// Trace end marker of this incoming data packet.
 	UE_NET_TRACE_PACKET_RECV(NetTraceId, GetConnectionId(), InPacketId, Reader.GetNumBits());
+
+	if (bSkipAck && !IsInternalAck())
+	{
+		// Trace packet as lost on receiving end to indicate bSkipAck in traced data
+		UE_NET_TRACE_PACKET_DROPPED(NetTraceId, GetConnectionId(), InPacketId, ENetTracePacketType::Incoming);
+	}
 }
 
 void UNetConnection::RestoreRemappedChannel(const int32 ChIndex)
@@ -3707,7 +3735,6 @@ int32 UNetConnection::SendRawBunch(FOutBunch& Bunch, bool InAllowMerge, const FN
 	check(!Bunch.IsError());
 	Driver->OutBunches++;
 	Driver->OutTotalBunches++;
-	TimeSensitive = 1;
 
 	// Build header.
 	SendBunchHeader.Reset();
@@ -3790,6 +3817,10 @@ int32 UNetConnection::SendRawBunch(FOutBunch& Bunch, bool InAllowMerge, const FN
 	// flush packet now so that we can report collected stats in the correct scope
 	PrepareWriteBitsToSendBuffer(BunchHeaderBits, BunchBits);
 
+	// We want to mark the packet in which we write the data as TimeSensitive
+	// Note: we want to mark the packet as TimeSensitive here, as PrepareWriteBitsToSendBuffer migth flush the packet
+	TimeSensitive = 1;
+
 	// Report bunch
 	UE_NET_TRACE_END_BUNCH(OutTraceCollector, Bunch, Bunch.ChName, 0, BunchHeaderBits, BunchBits, BunchCollector);
 
@@ -4557,6 +4588,7 @@ void UNetConnection::ResetGameWorldState()
 {
 	//Clear out references and do whatever else so that nothing holds onto references that it doesn't need to.
 	ResetDestructionInfos();
+	ClientPendingStreamingStatusRequest.Empty();
 	ClientVisibleLevelNames.Empty();
 	KeepProcessingActorChannelBunchesMap.Empty();
 	DormantReplicatorMap.Empty();
diff --git a/Engine/Source/Runtime/Engine/Private/PlayerController.cpp b/Engine/Source/Runtime/Engine/Private/PlayerController.cpp
--- a/Engine/Source/Runtime/Engine/Private/PlayerController.cpp
+++ b/Engine/Source/Runtime/Engine/Private/PlayerController.cpp
@@ -67,6 +67,7 @@
 #include "GameFramework/GameSession.h"
 #include "GameMapsSettings.h"
 #include "Particles/EmitterCameraLensEffectBase.h"
+#include "LevelUtils.h"
 #include "WorldPartition/WorldPartitionStreamingSource.h"
 
 DEFINE_LOG_CATEGORY(LogPlayerController);
@@ -254,7 +255,7 @@ FName APlayerController::NetworkRemapPath(FName InPackageName, bool bReading)
 
 /// @cond DOXYGEN_WARNINGS
 
-void APlayerController::ClientUpdateLevelStreamingStatus_Implementation(FName PackageName, bool bNewShouldBeLoaded, bool bNewShouldBeVisible, bool bNewShouldBlockOnLoad, int32 LODIndex )
+void APlayerController::ClientUpdateLevelStreamingStatus_Implementation(FName PackageName, bool bNewShouldBeLoaded, bool bNewShouldBeVisible, bool bNewShouldBlockOnLoad, int32 LODIndex, FNetLevelVisibilityTransactionId TransactionId)
 {
 	PackageName = NetworkRemapPath(PackageName, true);
 	
@@ -297,6 +298,7 @@ void APlayerController::ClientUpdateLevelStreamingStatus_Implementation(FName Pa
 						LevelStreamingObject->SetShouldBeVisible(bNewShouldBeVisible);
 						LevelStreamingObject->bShouldBlockOnLoad = bNewShouldBlockOnLoad;
 						LevelStreamingObject->SetLevelLODIndex(LODIndex);
+						LevelStreamingObject->UpdateNetVisibilityTransactionState(bNewShouldBeVisible, TransactionId);
 					}
 					else
 					{
@@ -320,7 +322,21 @@ void APlayerController::ClientUpdateMultipleLevelsStreamingStatus_Implementation
 {
 	for( const FUpdateLevelStreamingLevelStatus& LevelStatus : LevelStatuses )
 	{
-		ClientUpdateLevelStreamingStatus_Implementation( LevelStatus.PackageName, LevelStatus.bNewShouldBeLoaded, LevelStatus.bNewShouldBeVisible, LevelStatus.bNewShouldBlockOnLoad, LevelStatus.LODIndex );
+		ClientUpdateLevelStreamingStatus_Implementation(LevelStatus.PackageName, LevelStatus.bNewShouldBeLoaded, LevelStatus.bNewShouldBeVisible, LevelStatus.bNewShouldBlockOnLoad, LevelStatus.LODIndex, FNetLevelVisibilityTransactionId());
+	}
+}
+
+void APlayerController::ClientAckUpdateLevelVisibility_Implementation(FName PackageName, FNetLevelVisibilityTransactionId TransactionId)
+{
+	if (ensureAlwaysMsgf(TransactionId.IsClientTransaction(), TEXT("APlayerController::ClientAckUpdateLevelVisibility Expected TransactionId to be ClientTransaction")))
+	{
+		// find streaming levels and update request id
+		PackageName = NetworkRemapPath(PackageName, true);
+
+		if (ULevelStreaming* LevelStreamingObject = FLevelUtils::FindStreamingLevel(GetWorld(), PackageName))
+		{
+			LevelStreamingObject->AckNetVisibilityTransaction(TransactionId);
+		}
 	}
 }
 
@@ -346,8 +362,16 @@ void APlayerController::ServerUpdateLevelVisibility_Implementation(const FUpdate
 	{
 		FUpdateLevelVisibilityLevelInfo LevelVisibilityCopy = LevelVisibility;
 		LevelVisibilityCopy.PackageName = NetworkRemapPath(LevelVisibilityCopy.PackageName, true);
+		// FileName and packageName might differ so we have to remap the filename as well.
+		LevelVisibilityCopy.FileName = NetworkRemapPath(LevelVisibilityCopy.FileName, true);
 
 		Connection->UpdateLevelVisibility(LevelVisibilityCopy);
+
+		// If this is a client instigated request respond with the request id so that the client knows that we have received the visibility update
+		if (LevelVisibilityCopy.VisibilityRequestId.IsClientTransaction())
+		{
+			ClientAckUpdateLevelVisibility(LevelVisibility.PackageName, LevelVisibilityCopy.VisibilityRequestId);
+		}
 	}
 }
 
@@ -3371,7 +3395,20 @@ void APlayerController::ClientForceGarbageCollection_Implementation()
 
 void APlayerController::LevelStreamingStatusChanged(ULevelStreaming* LevelObject, bool bNewShouldBeLoaded, bool bNewShouldBeVisible, bool bNewShouldBlockOnLoad, int32 LODIndex )
 {
-	ClientUpdateLevelStreamingStatus(NetworkRemapPath(LevelObject->GetWorldAssetPackageFName(), false), bNewShouldBeLoaded, bNewShouldBeVisible, bNewShouldBlockOnLoad, LODIndex);
+	FNetLevelVisibilityTransactionId TransactionId;
+	if (GetNetMode() == NM_Client)
+	{
+		TransactionId.SetIsClientInstigator(true);
+	}
+	else if (NetConnection)
+	{
+		// For server instigated visibility status changes we assign a transaction id that is used to ensure that we do not enable replication until visibility is confirmed by the client
+		const FName PackageName = NetworkRemapPath(LevelObject->GetWorldAssetPackageFName(), true);
+
+		TransactionId = NetConnection->UpdateLevelStreamStatusChangedTransactionId(LevelObject, PackageName, bNewShouldBeVisible);
+	}
+
+	ClientUpdateLevelStreamingStatus(NetworkRemapPath(LevelObject->GetWorldAssetPackageFName(), false), bNewShouldBeLoaded, bNewShouldBeVisible, bNewShouldBlockOnLoad, LODIndex, TransactionId);
 }
 
 /// @cond DOXYGEN_WARNINGS
diff --git a/Engine/Source/Runtime/Engine/Private/World.cpp b/Engine/Source/Runtime/Engine/Private/World.cpp
--- a/Engine/Source/Runtime/Engine/Private/World.cpp
+++ b/Engine/Source/Runtime/Engine/Private/World.cpp
@@ -2704,7 +2704,7 @@ private:
 // Cumulated time passed in UWorld::AddToWorld since last call to UWorld::UpdateLevelStreaming.
 static double GAddToWorldTimeCumul = 0.0;
 
-void UWorld::AddToWorld( ULevel* Level, const FTransform& LevelTransform, bool bConsiderTimeLimit )
+void UWorld::AddToWorld( ULevel* Level, const FTransform& LevelTransform, bool bConsiderTimeLimit, FNetLevelVisibilityTransactionId TransactionId)
 {
 	TRACE_CPUPROFILER_EVENT_SCOPE(AddToWorld);
 	SCOPE_CYCLE_COUNTER(STAT_AddToWorldTime);
@@ -2977,6 +2977,7 @@ void UWorld::AddToWorld( ULevel* Level, const FTransform& LevelTransform, bool b
 				if (APlayerController* LocalPlayerController = It->GetPlayerController(this))
 				{
 					LevelVisibility.PackageName = LocalPlayerController->NetworkRemapPath(UnmappedPackageName, false);
+					LevelVisibility.VisibilityRequestId = TransactionId;
 					LocalPlayerController->ServerUpdateLevelVisibility(LevelVisibility);
 				}
 			}
@@ -3054,7 +3055,7 @@ void UWorld::BeginTearingDown()
 // Cumulated time doing IncrementalUnregisterComponents in UWorld::RemoveFromWorld since last call to UWorld::UpdateLevelStreaming.
 static double GRemoveFromWorldUnregisterComponentTimeCumul = 0.0;
 
-void UWorld::RemoveFromWorld( ULevel* Level, bool bAllowIncrementalRemoval )
+void UWorld::RemoveFromWorld( ULevel* Level, bool bAllowIncrementalRemoval, FNetLevelVisibilityTransactionId TransactionId)
 {
 	TRACE_CPUPROFILER_EVENT_SCOPE(RemoveFromWorld);
 	SCOPE_CYCLE_COUNTER(STAT_RemoveFromWorldTime);
@@ -3175,6 +3176,7 @@ void UWorld::RemoveFromWorld( ULevel* Level, bool bAllowIncrementalRemoval )
 					if (APlayerController* LocalPlayerController = It->GetPlayerController(this))
 					{
 						LevelVisibility.PackageName = LocalPlayerController->NetworkRemapPath(UnmappedPackageName, false);
+						LevelVisibility.VisibilityRequestId = TransactionId;
 						LocalPlayerController->ServerUpdateLevelVisibility(LevelVisibility);
 					}
 				}
diff --git a/Engine/Source/Runtime/Engine/Public/LevelUtils.h b/Engine/Source/Runtime/Engine/Public/LevelUtils.h
--- a/Engine/Source/Runtime/Engine/Public/LevelUtils.h
+++ b/Engine/Source/Runtime/Engine/Public/LevelUtils.h
@@ -25,6 +25,15 @@ public:
 	 */
 	static ULevelStreaming* FindStreamingLevel(const ULevel* Level);
 
+	/**
+	 * Returns the streaming level by package FName, or NULL if none exists.
+	 *
+	 * @param		InWorld			World to look in for the streaming level
+	 * @param		PackageFName	FName of the package containing the ULevel to query
+	 * @return						The level's streaming level, or NULL if none exists.
+	 */
+	static ULevelStreaming* FindStreamingLevel(UWorld* InWorld, const FName PackageName);
+
 	/**
 	 * Returns the streaming level by package name, or NULL if none exists.
 	 *
