New save game functions to enumerate save games & deterime if the save interface supports multiple users.

 - async enumeration function will be added in a subsequent change.

diff --git a/Engine/Source/Runtime/Engine/Public/SaveGameSystem.h b/Engine/Source/Runtime/Engine/Public/SaveGameSystem.h
--- a/Engine/Source/Runtime/Engine/Public/SaveGameSystem.h
+++ b/Engine/Source/Runtime/Engine/Public/SaveGameSystem.h
@@ -32,12 +32,21 @@ public:
 	/** Returns true if the platform has a native UI (like many consoles) */
 	virtual bool PlatformHasNativeUI() = 0;
 
+	/** Some save systems don't support more than one user, in those cases the game needs to avoid saving settings for anyone but the primary user. */
+	virtual bool DoesSaveSystemSupportMultipleUsers() = 0;
+
 	/** Return true if the named savegame exists (probably not useful with NativeUI */
 	virtual bool DoesSaveGameExist(const TCHAR* Name, const int32 UserIndex) = 0;
 
 	/** Similar to DoesSaveGameExist, except returns a result code with more information. */
 	virtual ESaveExistsResult DoesSaveGameExistWithResult(const TCHAR* Name, const int32 UserIndex) = 0;
 
+	/** Gets a list of all known saves.  This isn't possible on all platforms,  */
+	virtual bool GetSaveGameNames(TArray<FString>& FoundSaves, const int32 UserIndex)
+	{
+		return false;
+	}
+
 	/** Saves the game, blocking until complete. Platform may use FGameDelegates to get more information from the game */
 	virtual bool SaveGame(bool bAttemptToUseUI, const TCHAR* Name, const int32 UserIndex, const TArray<uint8>& Data) = 0;
 
@@ -96,6 +105,11 @@ public:
 		return false;
 	}
 
+	virtual bool DoesSaveSystemSupportMultipleUsers() override
+	{
+		return false;
+	}
+
 	virtual ESaveExistsResult DoesSaveGameExistWithResult(const TCHAR* Name, const int32 UserIndex) override
 	{
 		if (IFileManager::Get().FileSize(*GetSaveGamePath(Name)) >= 0)
@@ -110,6 +124,20 @@ public:
 		return ESaveExistsResult::OK == DoesSaveGameExistWithResult(Name, UserIndex);
 	}
 
+	virtual bool GetSaveGameNames(TArray<FString>& FoundSaves, const int32 UserIndex) override
+	{
+		TArray<FString> FoundFiles;
+		const FString SaveGameDirectory = FPaths::ProjectSavedDir() / TEXT("SaveGames/");
+		IFileManager::Get().FindFiles(FoundFiles, *SaveGameDirectory, TEXT("*.sav"));
+
+		for (FString& File : FoundFiles)
+		{
+			FoundSaves.Add(FPaths::GetBaseFilename(File));
+		}
+
+		return true;
+	}
+
 	virtual bool SaveGame(bool bAttemptToUseUI, const TCHAR* Name, const int32 UserIndex, const TArray<uint8>& Data) override
 	{
 		return FFileHelper::SaveArrayToFile(Data, *GetSaveGamePath(Name));
diff --git a/Engine/Source/Runtime/IOS/IOSPlatformFeatures/Private/IOSSaveGameSystem.cpp b/Engine/Source/Runtime/IOS/IOSPlatformFeatures/Private/IOSSaveGameSystem.cpp
--- a/Engine/Source/Runtime/IOS/IOSPlatformFeatures/Private/IOSSaveGameSystem.cpp
+++ b/Engine/Source/Runtime/IOS/IOSPlatformFeatures/Private/IOSSaveGameSystem.cpp
@@ -27,6 +27,20 @@ void FIOSSaveGameSystem::Shutdown()
 {
 }
 
+bool FIOSSaveGameSystem::GetSaveGameNames(TArray<FString>& FoundSaves, const int32 UserIndex)
+{
+	TArray<FString> FoundFiles;
+	const FString SaveGameDirectory = FPaths::ProjectSavedDir() / TEXT("SaveGames/");
+	IFileManager::Get().FindFiles(FoundFiles, *SaveGameDirectory, TEXT("*.sav"));
+
+	for (FString& File : FoundFiles)
+	{
+		FoundSaves.Add(FPaths::GetBaseFilename(File));
+	}
+
+	return true;
+}
+
 ISaveGameSystem::ESaveExistsResult FIOSSaveGameSystem::DoesSaveGameExistWithResult(const TCHAR* Name, const int32 UserIndex)
 {
 	if (IFileManager::Get().FileSize(*GetSaveGamePath(Name)) >= 0)
diff --git a/Engine/Source/Runtime/IOS/IOSPlatformFeatures/Public/IOSSaveGameSystem.h b/Engine/Source/Runtime/IOS/IOSPlatformFeatures/Public/IOSSaveGameSystem.h
--- a/Engine/Source/Runtime/IOS/IOSPlatformFeatures/Public/IOSSaveGameSystem.h
+++ b/Engine/Source/Runtime/IOS/IOSPlatformFeatures/Public/IOSSaveGameSystem.h
@@ -16,6 +16,13 @@ public:
 		return false;
 	}
 
+	virtual bool DoesSaveSystemSupportMultipleUsers() override
+	{
+		return false;
+	}
+
+	virtual bool GetSaveGameNames(TArray<FString>& FoundSaves, const int32 UserIndex) override;
+
 	virtual bool DoesSaveGameExist(const TCHAR* Name, const int32 UserIndex) override
 	{
 		return ESaveExistsResult::OK == DoesSaveGameExistWithResult(Name, UserIndex);
