Fixed problem following 20043286. Volumetric domain material we no longer converted correctly extinction from grey to colored.
Conversion needs to work on deprecated data for the post load to be able to pick up and convert to the new EditorOnly data.

 juan.canada

diff --git a/Engine/Source/Runtime/Engine/Private/Materials/Material.cpp b/Engine/Source/Runtime/Engine/Private/Materials/Material.cpp
--- a/Engine/Source/Runtime/Engine/Private/Materials/Material.cpp
+++ b/Engine/Source/Runtime/Engine/Private/Materials/Material.cpp
@@ -2599,13 +2599,13 @@ void UMaterial::Serialize(FArchive& Ar)
 #if WITH_EDITORONLY_DATA
 	if (MaterialDomain == MD_Volume && Ar.IsLoading() && Ar.CustomVer(FRenderingObjectVersion::GUID) < FRenderingObjectVersion::VolumeExtinctionBecomesRGB)
 	{
-		UMaterialEditorOnlyData* EditorOnly = GetEditorOnlyData();
-		if (EditorOnly && EditorOnly->Opacity.IsConnected()) // Base material input cannot have default values so we only deal with connected expression
+		// Note: we work on _DEPRECATED data because that will be used later to create the EditorOnly data in PostLoad().
+		if (Opacity_DEPRECATED.IsConnected()) // Base material input cannot have default values so we only deal with connected expression
 		{
 			// Change expression output from the Opacity to SubSurfaceColor that is now representing RGB extinction. Leave opacity connected as it is unused now anyway
-			EditorOnly->SubsurfaceColor.Connect(EditorOnly->Opacity.OutputIndex, EditorOnly->Opacity.Expression);
+			SubsurfaceColor_DEPRECATED.Connect(Opacity_DEPRECATED.OutputIndex, Opacity_DEPRECATED.Expression);
 			// Now disconnect Opacity
-			EditorOnly->Opacity.Expression = nullptr;
+			Opacity_DEPRECATED.Expression = nullptr;
 
 			// Now force the material to recompile and we use a hash of the original StateId.
 			// This is to avoid having different StateId each time we load the material and to not forever recompile it,i.e. use a cached version.
