Remove unnecessary overhead from FEdGraphPinReference::GetTypeHash().

diff --git a/Engine/Source/Runtime/Engine/Classes/EdGraph/EdGraphPin.h b/Engine/Source/Runtime/Engine/Classes/EdGraph/EdGraphPin.h
--- a/Engine/Source/Runtime/Engine/Classes/EdGraph/EdGraphPin.h
+++ b/Engine/Source/Runtime/Engine/Classes/EdGraph/EdGraphPin.h
@@ -271,10 +271,7 @@ struct ENGINE_API FEdGraphPinReference
 
 	friend uint32 GetTypeHash(const FEdGraphPinReference& EdGraphPinReference)
 	{
-		UEdGraphNode* ResolvedOwningNode = EdGraphPinReference.OwningNode.Get();
-		ensureMsgf(ResolvedOwningNode || !EdGraphPinReference.PinId.IsValid(), TEXT("Trying to reference an unowned pin: %s"), *EdGraphPinReference.PinId.ToString());
-		uint32 NodeHash = ResolvedOwningNode ? FCrc::StrCrc32(*ResolvedOwningNode->GetName()) : 0;
-		return FCrc::StrCrc32(*EdGraphPinReference.PinId.ToString(), 0 );
+		return GetTypeHash(EdGraphPinReference.PinId);
 	}
 
 	bool operator==(const FEdGraphPinReference& Other) const
