Inline edit support to material expression inputs
- Inline edit support for UPROPERTY of type float, int32, uint32, uint8, enum, bool.
- Editable UPROPERTYs with "OverridingInputProperty" metadata are created as inline widgets next to their corresponding inputs.
- Rest of the editable UPROPERTYs may specify "ShowAsInputPin" metadata to become inline edit pins, with 2 choices: "Primary" - show in primary view, "Advanced" - show in in advanced view.
- Update a bunch of material expressions to reflect the changes, rest of the expressions still need to be worked through.

diff --git a/Engine/Source/Editor/EditorStyle/Private/StarshipStyle.cpp b/Engine/Source/Editor/EditorStyle/Private/StarshipStyle.cpp
--- a/Engine/Source/Editor/EditorStyle/Private/StarshipStyle.cpp
+++ b/Engine/Source/Editor/EditorStyle/Private/StarshipStyle.cpp
@@ -3737,6 +3737,8 @@ void FStarshipEditorStyle::FStyle::SetupGraphEditorStyles()
 		Set( "Graph.Pin.ObjectSet", new IMAGE_BRUSH( "Old/Graph/Pin_ObjectSet", Icon12x12 ) );
 		Set( "Graph.Pin.ObjectEmpty", new IMAGE_BRUSH( "Old/Graph/Pin_ObjectEmpty", Icon12x12 ) );
 
+		Set("Graph.Pin.Dummy", new IMAGE_BRUSH("/Graph/Pin_dummy", FVector2D(15, 11)));
+
 		Set( "Graph.ConnectorFeedback.Border", new BOX_BRUSH( "Old/Menu_Background", FMargin(8.0f/64.0f) ) );
 		Set( "Graph.ConnectorFeedback.OK", new IMAGE_BRUSH( "Old/Graph/Feedback_OK", Icon16x16 ) );
 		Set( "Graph.ConnectorFeedback.OKWarn", new IMAGE_BRUSH( "Old/Graph/Feedback_OKWarn", Icon16x16 ) );
diff --git a/Engine/Source/Editor/GraphEditor/Private/MaterialNodes/SGraphNodeMaterialBase.cpp b/Engine/Source/Editor/GraphEditor/Private/MaterialNodes/SGraphNodeMaterialBase.cpp
--- a/Engine/Source/Editor/GraphEditor/Private/MaterialNodes/SGraphNodeMaterialBase.cpp
+++ b/Engine/Source/Editor/GraphEditor/Private/MaterialNodes/SGraphNodeMaterialBase.cpp
@@ -21,7 +21,8 @@
 #include "Materials/MaterialExpressionMakeMaterialAttributes.h"
 #include "Materials/MaterialExpressionBreakMaterialAttributes.h"
 
-
+static const FName NAME_Pin_NotConnectable("Graph.Pin.Dummy");
+static const FSlateBrush* CacheImg_Pin_NotConnectable = nullptr;
 
 /**
 * Simple representation of the backbuffer that the preview canvas renders to
@@ -291,6 +292,16 @@ void SGraphNodeMaterialBase::CreatePinWidgets()
 			TSharedPtr<SGraphPin> NewPin = CreatePinWidget(CurPin);
 			check(NewPin.IsValid());
 
+			// Assign an custom icon to not connectible pins
+			if (CurPin->bNotConnectable)
+			{
+				if (!CacheImg_Pin_NotConnectable)
+				{
+					CacheImg_Pin_NotConnectable = FAppStyle::Get().GetBrush(NAME_Pin_NotConnectable);
+				}
+				NewPin->SetCustomPinIcon(CacheImg_Pin_NotConnectable, CacheImg_Pin_NotConnectable);
+			}
+
 			this->AddPin(NewPin.ToSharedRef());
 		}
 	}
@@ -310,6 +321,14 @@ void SGraphNodeMaterialBase::AddPin( const TSharedRef<SGraphPin>& PinToAdd )
 {
 	PinToAdd->SetOwner( SharedThis(this) );
 
+	// Set visibility on advanced view pins
+	const UEdGraphPin* PinObj = PinToAdd->GetPinObj();
+	const bool bAdvancedParameter = (PinObj != nullptr) && PinObj->bAdvancedView;
+	if (bAdvancedParameter)
+	{
+		PinToAdd->SetVisibility(TAttribute<EVisibility>(PinToAdd, &SGraphPin::IsPinVisibleAsAdvanced));
+	}
+
 	if (PinToAdd->GetDirection() == EEdGraphPinDirection::EGPD_Input)
 	{
 		FMargin Padding = Settings->GetInputPinPadding();
@@ -348,7 +367,25 @@ void SGraphNodeMaterialBase::CreateBelowPinControls(TSharedPtr<SVerticalBox> Mai
 {
 	if (GraphNode && MainBox.IsValid())
 	{
-		int32 LeftPinCount = InputPins.Num();
+		// Count the number of visible input pins on the left
+		int32 LeftPinCount = 0;
+		if (GraphNode->AdvancedPinDisplay == ENodeAdvancedPins::Hidden)
+		{
+			// Advanced view pins are hidden so exclude them from the pin count
+			for (int32 i = 0; i < InputPins.Num(); ++i)
+			{
+				const UEdGraphPin* PinObj = InputPins[i]->GetPinObj();
+				if (!PinObj->bAdvancedView)
+				{
+					LeftPinCount++;
+				}
+			}
+		}
+		else
+		{
+			LeftPinCount = InputPins.Num();
+		}
+
 		int32 RightPinCount = OutputPins.Num();
 
 		const float NegativeHPad = FMath::Max<float>(-Settings->PaddingTowardsNodeEdge, 0.0f);
@@ -446,6 +483,14 @@ TSharedRef<SWidget> SGraphNodeMaterialBase::CreateNodeContentArea()
 		];
 }
 
+void SGraphNodeMaterialBase::OnAdvancedViewChanged(const ECheckBoxState NewCheckedState)
+{
+	SGraphNode::OnAdvancedViewChanged(NewCheckedState);
+
+	// Update the graph node so that the preview is recreated to update its position
+	UpdateGraphNode();
+}
+
 TSharedRef<SWidget> SGraphNodeMaterialBase::CreatePreviewWidget()
 {
 	PreviewViewport.Reset();
diff --git a/Engine/Source/Editor/GraphEditor/Private/MaterialNodes/SGraphNodeMaterialBase.h b/Engine/Source/Editor/GraphEditor/Private/MaterialNodes/SGraphNodeMaterialBase.h
--- a/Engine/Source/Editor/GraphEditor/Private/MaterialNodes/SGraphNodeMaterialBase.h
+++ b/Engine/Source/Editor/GraphEditor/Private/MaterialNodes/SGraphNodeMaterialBase.h
@@ -106,6 +106,7 @@ protected:
 	virtual void CreateBelowPinControls(TSharedPtr<SVerticalBox> MainBox) override;
 	virtual void SetDefaultTitleAreaWidget(TSharedRef<SOverlay> DefaultTitleAreaWidget) override;
 	virtual TSharedRef<SWidget> CreateNodeContentArea() override;
+	virtual void OnAdvancedViewChanged(const ECheckBoxState NewCheckedState) override;
 	// End of SGraphNode interface
 
 	/** Creates a preview viewport if necessary */
diff --git a/Engine/Source/Editor/GraphEditor/Private/MaterialPins/SGraphPinMaterialInput.h b/Engine/Source/Editor/GraphEditor/Private/MaterialPins/SGraphPinMaterialInput.h
deleted file mode 100644
diff --git a/Engine/Source/Editor/GraphEditor/Private/NodeFactory.cpp b/Engine/Source/Editor/GraphEditor/Private/NodeFactory.cpp
--- a/Engine/Source/Editor/GraphEditor/Private/NodeFactory.cpp
+++ b/Engine/Source/Editor/GraphEditor/Private/NodeFactory.cpp
@@ -253,22 +253,6 @@ TSharedPtr<SGraphPin> FNodeFactory::CreatePinWidget(UEdGraphPin* InPin)
 			return K2PinWidget;
 		}
 	}
-
-	if (const UMaterialGraphSchema* MaterialGraphSchema = Cast<const UMaterialGraphSchema>(InPin->GetSchema()))
-	{
-		if (InPin->PinType.PinCategory == MaterialGraphSchema->PC_Exec)
-		{
-			return SNew(SGraphPinExec, InPin);
-		}
-		else if (InPin->PinType.PinCategory == MaterialGraphSchema->PC_MaterialInput)
-		{
-			return SNew(SGraphPinMaterialInput, InPin);
-		}
-		else
-		{
-			return SNew(SGraphPin, InPin);
-		}
-	}
 	
 	// If we didn't pick a custom pin widget, use an uncustomized basic pin
 	return SNew(SGraphPin, InPin);
diff --git a/Engine/Source/Editor/GraphEditor/Public/MaterialPins/SGraphPinMaterialInput.h b/Engine/Source/Editor/GraphEditor/Public/MaterialPins/SGraphPinMaterialInput.h
new file mode 100644
--- /dev/null
+++ b/Engine/Source/Editor/GraphEditor/Public/MaterialPins/SGraphPinMaterialInput.h
@@ -0,0 +1,23 @@
+// Copyright Epic Games, Inc. All Rights Reserved.
+
+#pragma once
+
+#include "CoreMinimal.h"
+#include "Widgets/DeclarativeSyntaxSupport.h"
+#include "Styling/SlateColor.h"
+#include "SGraphPin.h"
+
+class GRAPHEDITOR_API SGraphPinMaterialInput : public SGraphPin
+{
+public:
+	SLATE_BEGIN_ARGS(SGraphPinMaterialInput) {}
+	SLATE_END_ARGS()
+
+	void Construct(const FArguments& InArgs, UEdGraphPin* InGraphPinObj);
+
+protected:
+	//~ Begin SGraphPin Interface
+	virtual FSlateColor GetPinColor() const override;
+	//~ End SGraphPin Interface
+
+};
diff --git a/Engine/Source/Editor/GraphEditor/Public/SGraphNode.h b/Engine/Source/Editor/GraphEditor/Public/SGraphNode.h
--- a/Engine/Source/Editor/GraphEditor/Public/SGraphNode.h
+++ b/Engine/Source/Editor/GraphEditor/Public/SGraphNode.h
@@ -336,7 +336,7 @@ protected:
 	EVisibility AdvancedViewArrowVisibility() const;
 
 	/** Show/hide advanced view */
-	void OnAdvancedViewChanged( const ECheckBoxState NewCheckedState );
+	virtual void OnAdvancedViewChanged( const ECheckBoxState NewCheckedState );
 
 	/** hidden == unchecked, shown == checked */
 	ECheckBoxState IsAdvancedViewChecked() const;
diff --git a/Engine/Source/Editor/MaterialEditor/Private/MaterialEditor.cpp b/Engine/Source/Editor/MaterialEditor/Private/MaterialEditor.cpp
--- a/Engine/Source/Editor/MaterialEditor/Private/MaterialEditor.cpp
+++ b/Engine/Source/Editor/MaterialEditor/Private/MaterialEditor.cpp
@@ -6990,6 +6990,11 @@ void FMaterialEditor::AddSelectedExpressionToFavorites()
 	}
 }
 
+void FMaterialEditor::UpdateDetailView()
+{
+	GetDetailView()->InvalidateCachedState();
+}
+
 void FMaterialEditor::OnSelectedNodesChanged(const TSet<class UObject*>& NewSelection)
 {
 	TArray<UObject*> SelectedObjects;
@@ -7067,6 +7072,8 @@ void FMaterialEditor::OnNodeDoubleClicked(class UEdGraphNode* Node)
 			ChannelEditStruct.Red = &Constant3Expression->Constant.R;
 			ChannelEditStruct.Green = &Constant3Expression->Constant.G;
 			ChannelEditStruct.Blue = &Constant3Expression->Constant.B;
+			static FName ConstantName = FName(TEXT("Constant"));
+			ColorPickerProperty = Constant3Expression->GetClass()->FindPropertyByName(ConstantName);
 		}
 		else if( Constant4Expression )
 		{
@@ -7074,6 +7081,8 @@ void FMaterialEditor::OnNodeDoubleClicked(class UEdGraphNode* Node)
 			ChannelEditStruct.Green = &Constant4Expression->Constant.G;
 			ChannelEditStruct.Blue = &Constant4Expression->Constant.B;
 			ChannelEditStruct.Alpha = &Constant4Expression->Constant.A;
+			static FName ConstantName = FName(TEXT("Constant"));
+			ColorPickerProperty = Constant4Expression->GetClass()->FindPropertyByName(ConstantName);
 		}
 		else if (InputExpression)
 		{
diff --git a/Engine/Source/Editor/MaterialEditor/Private/MaterialEditor.h b/Engine/Source/Editor/MaterialEditor/Private/MaterialEditor.h
--- a/Engine/Source/Editor/MaterialEditor/Private/MaterialEditor.h
+++ b/Engine/Source/Editor/MaterialEditor/Private/MaterialEditor.h
@@ -331,6 +331,8 @@ public:
 	// Widget Accessors
 	TSharedRef<class IDetailsView> GetDetailView() const {return MaterialDetailsView.ToSharedRef();}
 	
+	virtual void UpdateDetailView() override;
+
 	// FTickableGameObject interface
 	virtual void Tick(float DeltaTime) override;
 
diff --git a/Engine/Source/Editor/MaterialEditor/Private/MaterialEditorGraphPanelPinFactory.cpp b/Engine/Source/Editor/MaterialEditor/Private/MaterialEditorGraphPanelPinFactory.cpp
new file mode 100644
--- /dev/null
+++ b/Engine/Source/Editor/MaterialEditor/Private/MaterialEditorGraphPanelPinFactory.cpp
@@ -0,0 +1,77 @@
+// Copyright Epic Games, Inc. All Rights Reserved.
+
+#include "MaterialEditorGraphPanelPinFactory.h"
+#include "MaterialGraph/MaterialGraphSchema.h"
+#include "MaterialPins/SGraphPinMaterialInput.h"
+#include "KismetPins/SGraphPinExec.h"
+#include "KismetPins/SGraphPinNum.h"
+#include "KismetPins/SGraphPinVector.h"
+#include "KismetPins/SGraphPinVector2D.h"
+#include "KismetPins/SGraphPinVector4.h"
+#include "KismetPins/SGraphPinColor.h"
+#include "KismetPins/SGraphPinInteger.h"
+#include "Editor/GraphEditor/Private/KismetPins/SGraphPinEnum.h"
+#include "KismetPins/SGraphPinBool.h"
+#include "NodeFactory.h"
+
+TSharedPtr<class SGraphPin> FMaterialEditorGraphPanelPinFactory::CreatePin(class UEdGraphPin* InPin) const
+{
+	if (const UMaterialGraphSchema* MaterialGraphSchema = Cast<const UMaterialGraphSchema>(InPin->GetSchema()))
+	{
+		if (InPin->PinType.PinCategory == MaterialGraphSchema->PC_Exec)
+		{
+			return SNew(SGraphPinExec, InPin);
+		}
+		else if (InPin->PinType.PinCategory == MaterialGraphSchema->PC_MaterialInput)
+		{
+			return SNew(SGraphPinMaterialInput, InPin);
+		}
+		else
+		{
+			if (InPin->PinType.PinSubCategory == MaterialGraphSchema->PSC_Red || InPin->PinType.PinSubCategory == MaterialGraphSchema->PSC_Float)
+			{
+				return SNew(SGraphPinNum<double>, InPin);
+			}
+			else if (InPin->PinType.PinSubCategory == MaterialGraphSchema->PSC_RG)
+			{
+				return SNew(SGraphPinVector2D, InPin);
+			}
+			else if (InPin->PinType.PinSubCategory == MaterialGraphSchema->PSC_RGB)
+			{
+				return SNew(SGraphPinVector, InPin);
+			}
+			else if (InPin->PinType.PinSubCategory == MaterialGraphSchema->PSC_RGBA)
+			{
+				return SNew(SGraphPinColor, InPin);
+			}
+			else if (InPin->PinType.PinSubCategory == MaterialGraphSchema->PSC_Vector4)
+			{
+				return SNew(SGraphPinVector4, InPin);
+			}
+			else if (InPin->PinType.PinSubCategory == MaterialGraphSchema->PSC_Int)
+			{
+				return SNew(SGraphPinInteger, InPin);
+			}
+			else if (InPin->PinType.PinSubCategory == MaterialGraphSchema->PSC_Byte)
+			{
+				// Check for valid enum object reference
+				if ((InPin->PinType.PinSubCategoryObject != NULL) && (InPin->PinType.PinSubCategoryObject->IsA(UEnum::StaticClass())))
+				{
+					return SNew(SGraphPinEnum, InPin);
+				}
+				else
+				{
+					return SNew(SGraphPinInteger, InPin);
+				}
+			}
+			else if (InPin->PinType.PinSubCategory == MaterialGraphSchema->PSC_Bool)
+			{
+				return SNew(SGraphPinBool, InPin);
+			}
+		}
+
+		return SNew(SGraphPin, InPin);
+	}
+
+	return nullptr;
+}
diff --git a/Engine/Source/Editor/MaterialEditor/Private/MaterialEditorModule.cpp b/Engine/Source/Editor/MaterialEditor/Private/MaterialEditorModule.cpp
--- a/Engine/Source/Editor/MaterialEditor/Private/MaterialEditorModule.cpp
+++ b/Engine/Source/Editor/MaterialEditor/Private/MaterialEditorModule.cpp
@@ -18,9 +18,10 @@
 #include "Widgets/Images/SImage.h"
 #include "Misc/EngineBuildSettings.h"
 #include "MaterialEditorSettings.h"
-
+#include "EdGraphUtilities.h"
 #include "ISettingsModule.h"
 #include "Interfaces/IMainFrameModule.h"
+#include "MaterialEditorGraphPanelPinFactory.h"
 
 const FName MaterialEditorAppIdentifier = FName(TEXT("MaterialEditorApp"));
 const FName MaterialInstanceEditorAppIdentifier = FName(TEXT("MaterialInstanceEditorApp"));
@@ -62,6 +63,9 @@ public:
 				NSLOCTEXT("MaterialEditorModule", "SettingsDesc", "Settings related to the material editor."),
 				GetMutableDefault<UMaterialEditorSettings>());
 		}
+
+		GraphPanelPinFactory = MakeShared<FMaterialEditorGraphPanelPinFactory>();
+		FEdGraphUtilities::RegisterVisualPinFactory(GraphPanelPinFactory);
 	}
 
 	/**
@@ -69,6 +73,8 @@ public:
 	 */
 	virtual void ShutdownModule() override
 	{
+		FEdGraphUtilities::UnregisterVisualPinFactory(GraphPanelPinFactory);
+
 		MenuExtensibilityManager.Reset();
 		ToolBarExtensibilityManager.Reset();
 	}
@@ -125,6 +131,7 @@ private:
 	TSharedPtr<FExtensibilityManager> MenuExtensibilityManager;
 	TSharedPtr<FExtensibilityManager> ToolBarExtensibilityManager;
 	FDelegateHandle ContentBrowserAssetExtenderDelegateHandle;
+	TSharedPtr<FMaterialEditorGraphPanelPinFactory> GraphPanelPinFactory;
 };
 
 IMPLEMENT_MODULE( FMaterialEditorModule, MaterialEditor );
diff --git a/Engine/Source/Editor/MaterialEditor/Private/MaterialEditorUtilities.cpp b/Engine/Source/Editor/MaterialEditor/Private/MaterialEditorUtilities.cpp
--- a/Engine/Source/Editor/MaterialEditor/Private/MaterialEditorUtilities.cpp
+++ b/Engine/Source/Editor/MaterialEditor/Private/MaterialEditorUtilities.cpp
@@ -130,6 +130,15 @@ void FMaterialEditorUtilities::UpdateMaterialAfterGraphChange(const class UEdGra
 	}
 }
 
+void FMaterialEditorUtilities::UpdateDetailView(const class UEdGraph* Graph)
+{
+	TSharedPtr<class IMaterialEditor> MaterialEditor = GetIMaterialEditorForObject(Graph);
+	if (MaterialEditor.IsValid())
+	{
+		MaterialEditor->UpdateDetailView();
+	}
+}
+
 bool FMaterialEditorUtilities::CanPasteNodes(const class UEdGraph* Graph)
 {
 	bool bCanPaste = false;
diff --git a/Engine/Source/Editor/MaterialEditor/Public/IMaterialEditor.h b/Engine/Source/Editor/MaterialEditor/Public/IMaterialEditor.h
--- a/Engine/Source/Editor/MaterialEditor/Public/IMaterialEditor.h
+++ b/Engine/Source/Editor/MaterialEditor/Public/IMaterialEditor.h
@@ -90,6 +90,8 @@ public:
 	 */
 	virtual void UpdateMaterialAfterGraphChange() {};
 
+	virtual void UpdateDetailView() {}
+
 	/**
 	 * Jumps to a specific node or graph
 	 *
diff --git a/Engine/Source/Editor/MaterialEditor/Public/MaterialEditorGraphPanelPinFactory.h b/Engine/Source/Editor/MaterialEditor/Public/MaterialEditorGraphPanelPinFactory.h
new file mode 100644
--- /dev/null
+++ b/Engine/Source/Editor/MaterialEditor/Public/MaterialEditorGraphPanelPinFactory.h
@@ -0,0 +1,14 @@
+// Copyright Epic Games, Inc. All Rights Reserved.
+
+#pragma once
+
+#include "CoreMinimal.h"
+#include "EdGraphUtilities.h"
+
+class FMaterialEditorGraphPanelPinFactory :
+	public FGraphPanelPinFactory
+{
+public:
+	// FGraphPanelPinFactory interface
+	virtual TSharedPtr<class SGraphPin> CreatePin(class UEdGraphPin* InPin) const override;
+};
diff --git a/Engine/Source/Editor/MaterialEditor/Public/MaterialEditorUtilities.h b/Engine/Source/Editor/MaterialEditor/Public/MaterialEditorUtilities.h
--- a/Engine/Source/Editor/MaterialEditor/Public/MaterialEditorUtilities.h
+++ b/Engine/Source/Editor/MaterialEditor/Public/MaterialEditorUtilities.h
@@ -118,6 +118,8 @@ public:
 	 */
 	static void UpdateMaterialAfterGraphChange(const class UEdGraph* Graph);
 
+	static void UpdateDetailView(const class UEdGraph* Graph);
+
 	/** Can we paste to this graph? */
 	static bool CanPasteNodes(const class UEdGraph* Graph);
 
diff --git a/Engine/Source/Editor/UnrealEd/Classes/MaterialGraph/MaterialGraphNode.h b/Engine/Source/Editor/UnrealEd/Classes/MaterialGraph/MaterialGraphNode.h
--- a/Engine/Source/Editor/UnrealEd/Classes/MaterialGraph/MaterialGraphNode.h
+++ b/Engine/Source/Editor/UnrealEd/Classes/MaterialGraph/MaterialGraphNode.h
@@ -72,6 +72,7 @@ public:
 	virtual FString GetDocumentationExcerptName() const override;
 	virtual bool CanUserDeleteNode() const override;
 	virtual uint32 GetPinMaterialType(const UEdGraphPin* Pin) const override;
+	virtual void PinDefaultValueChanged(UEdGraphPin* Pin) override;
 	//~ End UEdGraphNode Interface.
 
 	//~ Begin UMaterialGraphNode_Base Interface
diff --git a/Engine/Source/Editor/UnrealEd/Classes/MaterialGraph/MaterialGraphNode_Base.h b/Engine/Source/Editor/UnrealEd/Classes/MaterialGraph/MaterialGraphNode_Base.h
--- a/Engine/Source/Editor/UnrealEd/Classes/MaterialGraph/MaterialGraphNode_Base.h
+++ b/Engine/Source/Editor/UnrealEd/Classes/MaterialGraph/MaterialGraphNode_Base.h
@@ -67,4 +67,7 @@ protected:
 	virtual uint32 GetPinMaterialType(const UEdGraphPin* Pin) const;
 
 	void EmptyPins();
+
+	/** Return the first input pin matching the name */
+	class UEdGraphPin* GetInputPin(const FName& PinName) const;
 };
diff --git a/Engine/Source/Editor/UnrealEd/Classes/MaterialGraph/MaterialGraphSchema.h b/Engine/Source/Editor/UnrealEd/Classes/MaterialGraph/MaterialGraphSchema.h
--- a/Engine/Source/Editor/UnrealEd/Classes/MaterialGraph/MaterialGraphSchema.h
+++ b/Engine/Source/Editor/UnrealEd/Classes/MaterialGraph/MaterialGraphSchema.h
@@ -192,6 +192,13 @@ class UMaterialGraphSchema : public UEdGraphSchema
 	UNREALED_API static const FName PSC_Blue;
 	UNREALED_API static const FName PSC_Alpha;
 	UNREALED_API static const FName PSC_RGBA;
+	UNREALED_API static const FName PSC_RGB;
+	UNREALED_API static const FName PSC_RG;
+	UNREALED_API static const FName PSC_Int;
+	UNREALED_API static const FName PSC_Byte;
+	UNREALED_API static const FName PSC_Bool;
+	UNREALED_API static const FName PSC_Float;
+	UNREALED_API static const FName PSC_Vector4;
 
 	UNREALED_API static const FName PN_Execute; // Category=PC_Exec, singleton, input
 
@@ -242,6 +249,12 @@ class UMaterialGraphSchema : public UEdGraphSchema
 	/** Check whether the types of pins are compatible */
 	bool ArePinsCompatible_Internal(const UEdGraphPin* InputPin, const UEdGraphPin* OutputPin, FText& ResponseMessage) const;
 
+	/** update material when the default value of a material node's pin has changed */
+	UNREALED_API void UpdateMaterialOnDefaultValueChanged(const UEdGraph* Graph) const;
+
+	/** Update the detail view */
+	UNREALED_API void UpdateDetailView(const UEdGraph* Graph) const;
+
 	/** Gets the type of this pin (must be part of a UMaterialGraphNode_Base) */
 	UNREALED_API static uint32 GetMaterialValueType(const UEdGraphPin* MaterialPin);
 
@@ -250,7 +263,6 @@ class UMaterialGraphSchema : public UEdGraphSchema
 	virtual void GetContextMenuActions(class UToolMenu* Menu, class UGraphNodeContextMenuContext* Context) const override;
 	virtual const FPinConnectionResponse CanCreateConnection(const UEdGraphPin* A, const UEdGraphPin* B) const override;
 	virtual bool TryCreateConnection(UEdGraphPin* A, UEdGraphPin* B) const override;
-	virtual bool ShouldHidePinDefaultValue(UEdGraphPin* Pin) const override { return true; }
 	virtual FLinearColor GetPinTypeColor(const FEdGraphPinType& PinType) const override;
 	virtual void BreakNodeLinks(UEdGraphNode& TargetNode) const override;
 	virtual void BreakPinLinks(UEdGraphPin& TargetPin, bool bSendsNodeNotifcation) const override;
diff --git a/Engine/Source/Editor/UnrealEd/Private/MaterialGraph.cpp b/Engine/Source/Editor/UnrealEd/Private/MaterialGraph.cpp
--- a/Engine/Source/Editor/UnrealEd/Private/MaterialGraph.cpp
+++ b/Engine/Source/Editor/UnrealEd/Private/MaterialGraph.cpp
@@ -501,7 +501,8 @@ void UMaterialGraph::LinkGraphNodesFromMaterial()
 		{
 			if (Pin->Direction == EGPD_Input && Pin->PinType.PinCategory != UMaterialGraphSchema::PC_Exec)
 			{
-				if (ExpressionInputs[Pin->SourceIndex]->Expression)
+				// Implicitly generated property inputs are not returned by GetInputs(), so check index is within valid range.
+				if (ExpressionInputs.IsValidIndex(Pin->SourceIndex) && ExpressionInputs[Pin->SourceIndex]->Expression)
 				{
 					// Unclear why this is null sometimes outside of composite reroute, but this is safer than crashing
 					if (UMaterialGraphNode* GraphNode = Cast<UMaterialGraphNode>(ExpressionInputs[Pin->SourceIndex]->Expression->GraphNode))
@@ -631,7 +632,8 @@ void UMaterialGraph::LinkMaterialExpressionsFromGraph()
 						if (Pin->Direction == EGPD_Input && Pin->PinType.PinCategory != UMaterialGraphSchema::PC_Exec)
 						{
 							// Wire up non-execution input pins
-							FExpressionInput* ExpressionInput = ExpressionInputs[Pin->SourceIndex];
+							// Implicitly generated property inputs are not returned by GetInputs(), so check index is within valid range.
+							FExpressionInput* ExpressionInput = ExpressionInputs.IsValidIndex(Pin->SourceIndex) ? ExpressionInputs[Pin->SourceIndex] : nullptr;
 							if (Pin->LinkedTo.Num() > 0)
 							{
 								UEdGraphPin* ConnectedPin = Pin->LinkedTo[0];
diff --git a/Engine/Source/Editor/UnrealEd/Private/MaterialGraphNode.cpp b/Engine/Source/Editor/UnrealEd/Private/MaterialGraphNode.cpp
--- a/Engine/Source/Editor/UnrealEd/Private/MaterialGraphNode.cpp
+++ b/Engine/Source/Editor/UnrealEd/Private/MaterialGraphNode.cpp
@@ -651,6 +651,11 @@ uint32 UMaterialGraphNode::GetPinMaterialType(const UEdGraphPin* Pin) const
 	}
 }
 
+void UMaterialGraphNode::PinDefaultValueChanged(UEdGraphPin* Pin)
+{
+	MaterialExpression->PinDefaultValueChanged(Pin->SourceIndex, Pin->DefaultValue);
+}
+
 void UMaterialGraphNode::CreateInputPins()
 {
 	if (MaterialExpression->HasExecInput())
@@ -662,14 +667,14 @@ void UMaterialGraphNode::CreateInputPins()
 		NewPin->PinFriendlyName = SpaceText;
 	}
 
-	const TArray<FExpressionInput*> ExpressionInputs = MaterialExpression->GetInputs();
-	for (int32 Index = 0; Index < ExpressionInputs.Num(); ++Index)
+	TArray<FExpressionInput*> ExpressionInputs = MaterialExpression->GetInputs();
+	const int32 NumExpressionInputs = ExpressionInputs.Num();
+	for (int32 Index = 0; Index < NumExpressionInputs; ++Index)
 	{
-		FExpressionInput* Input = ExpressionInputs[Index];
-		FName PinCategory;
-
 		FName InputName = MaterialExpression->GetInputName(Index);
 		InputName = GetShortenPinName(InputName);
+
+		FName PinCategory;
 		if (MaterialExpression->IsInputConnectionRequired(Index))
 		{
 			PinCategory = UMaterialGraphSchema::PC_Required;
@@ -679,8 +684,12 @@ void UMaterialGraphNode::CreateInputPins()
 			PinCategory = UMaterialGraphSchema::PC_Optional;
 		}
 
-		UEdGraphPin* NewPin = CreatePin(EGPD_Input, PinCategory, InputName);
+		FName PinSubCategory = MaterialExpression->GetInputPinSubCategory(Index);
+		UObject* PinSubCategoryObject = MaterialExpression->GetInputPinSubCategoryObject(Index);
+
+		UEdGraphPin* NewPin = CreatePin(EGPD_Input, PinCategory, PinSubCategory, PinSubCategoryObject, InputName);
 		NewPin->SourceIndex = Index;
+		NewPin->DefaultValue = MaterialExpression->GetInputPinDefaultValue(Index);
 		if (NewPin->PinName.IsNone())
 		{
 			// Makes sure pin has a name for lookup purposes but user will never see it
@@ -688,6 +697,49 @@ void UMaterialGraphNode::CreateInputPins()
 			NewPin->PinFriendlyName = SpaceText;
 		}
 	}
+
+	// Next create pins for property inputs
+	int32 AdvancedPins = 0;
+	static FName ShowAsInputPinMetaData(TEXT("ShowAsInputPin"));
+	const TArray<FProperty*> PropertyInputs = MaterialExpression->GetPropertyInputs();
+	for (int32 i = 0; i < PropertyInputs.Num(); ++i)
+	{
+		const FProperty* Property = PropertyInputs[i];
+		FString InputName = Property->GetDisplayNameText().ToString();
+		FName PinCategory = UMaterialGraphSchema::PC_Optional;
+
+		int32 PinIndex = NumExpressionInputs + i;
+		FName PinSubCategory = MaterialExpression->GetInputPinSubCategory(PinIndex);
+		UObject* PinSubCategoryObject = MaterialExpression->GetInputPinSubCategoryObject(PinIndex);
+
+		UEdGraphPin* NewPin = CreatePin(EGPD_Input, PinCategory, PinSubCategory, PinSubCategoryObject, FName(InputName));
+		NewPin->SourceIndex = PinIndex;
+		NewPin->DefaultValue = MaterialExpression->GetInputPinDefaultValue(PinIndex);
+		// Property pins can't connect externally
+		NewPin->bNotConnectable = true;
+		NewPin->bDefaultValueIsReadOnly = !MaterialExpression->CanEditChange(Property);
+		const FString ShowAsInputPin = Property->GetMetaData(ShowAsInputPinMetaData);
+		NewPin->bAdvancedView = (ShowAsInputPin == TEXT("Advanced"));
+		if (NewPin->PinName.IsNone())
+		{
+			// Makes sure pin has a name for lookup purposes but user will never see it
+			NewPin->PinName = CreateUniquePinName(TEXT("PropertyInput"));
+			NewPin->PinFriendlyName = SpaceText;
+		}
+
+		if (NewPin->bAdvancedView)
+		{
+			AdvancedPins++;
+		}
+	}
+	if (AdvancedPins > 0)
+	{
+		// Turn on the advanced view arrow button
+		if (AdvancedPinDisplay == ENodeAdvancedPins::NoPins)
+		{
+			AdvancedPinDisplay = ENodeAdvancedPins::Hidden;
+		}
+	}
 }
 
 void UMaterialGraphNode::CreateOutputPins()
diff --git a/Engine/Source/Editor/UnrealEd/Private/MaterialGraphNode_Base.cpp b/Engine/Source/Editor/UnrealEd/Private/MaterialGraphNode_Base.cpp
--- a/Engine/Source/Editor/UnrealEd/Private/MaterialGraphNode_Base.cpp
+++ b/Engine/Source/Editor/UnrealEd/Private/MaterialGraphNode_Base.cpp
@@ -37,6 +37,21 @@ UEdGraphPin* UMaterialGraphNode_Base::GetInputPin(int32 InputIndex) const
 	return nullptr;
 }
 
+UEdGraphPin* UMaterialGraphNode_Base::GetInputPin(const FName& PinName) const
+{
+	// Return the first input pin matching the name
+	for (UEdGraphPin* Pin : Pins)
+	{
+		if (Pin->PinType.PinCategory != UMaterialGraphSchema::PC_Exec &&
+			Pin->Direction == EGPD_Input &&
+			Pin->PinName == PinName)
+		{
+			return Pin;
+		}
+	}
+	return nullptr;
+}
+
 UEdGraphPin* UMaterialGraphNode_Base::GetOutputPin(int32 OutputIndex) const
 {
 	for (UEdGraphPin* Pin : Pins)
@@ -273,7 +288,7 @@ void UMaterialGraphNode_Base::ReconstructNode()
 		}
 		else
 		{
-			if (OldPin->Direction == EGPD_Input) NewPin = GetInputPin(OldPin->SourceIndex);
+			if (OldPin->Direction == EGPD_Input) NewPin = GetInputPin(OldPin->PinName);
 			else NewPin = GetOutputPin(OldPin->SourceIndex);
 		}
 		if (NewPin)
diff --git a/Engine/Source/Editor/UnrealEd/Private/MaterialGraphSchema.cpp b/Engine/Source/Editor/UnrealEd/Private/MaterialGraphSchema.cpp
--- a/Engine/Source/Editor/UnrealEd/Private/MaterialGraphSchema.cpp
+++ b/Engine/Source/Editor/UnrealEd/Private/MaterialGraphSchema.cpp
@@ -238,6 +238,13 @@ const FName UMaterialGraphSchema::PSC_Green(TEXT("green"));
 const FName UMaterialGraphSchema::PSC_Blue(TEXT("blue"));
 const FName UMaterialGraphSchema::PSC_Alpha(TEXT("alpha"));
 const FName UMaterialGraphSchema::PSC_RGBA(TEXT("rgba"));
+const FName UMaterialGraphSchema::PSC_RGB(TEXT("rgb"));
+const FName UMaterialGraphSchema::PSC_RG(TEXT("rg"));
+const FName UMaterialGraphSchema::PSC_Int(TEXT("int"));
+const FName UMaterialGraphSchema::PSC_Byte(TEXT("byte"));
+const FName UMaterialGraphSchema::PSC_Bool(TEXT("bool"));
+const FName UMaterialGraphSchema::PSC_Float(TEXT("float"));
+const FName UMaterialGraphSchema::PSC_Vector4(TEXT("vector4"));
 
 const FName UMaterialGraphSchema::PN_Execute("execute");
 
@@ -377,6 +384,14 @@ bool UMaterialGraphSchema::ArePinsCompatible_Internal(const UEdGraphPin* InputPi
 	uint32 InputType = GetMaterialValueType(InputPin);
 	uint32 OutputType = GetMaterialValueType(OutputPin);
 
+	if (InputPin->bNotConnectable)
+	{
+		FFormatNamedArguments Args;
+		Args.Add(TEXT("PinName"), FText::FromName(InputPin->PinName));
+		ResponseMessage = FText::Format(LOCTEXT("PinNotConnectable", "Pin '{PinName}' is not connectable"), Args);	
+		return false;
+	}
+
 	bool bPinsCompatible = CanConnectMaterialValueTypes(InputType, OutputType);
 	if (!bPinsCompatible)
 	{
@@ -873,6 +888,16 @@ void UMaterialGraphSchema::DroppedAssetsOnGraph(const TArray<struct FAssetData>&
 	}
 }
 
+void UMaterialGraphSchema::UpdateMaterialOnDefaultValueChanged(const UEdGraph* Graph) const
+{
+	FMaterialEditorUtilities::UpdateMaterialAfterGraphChange(Graph);
+}
+
+void UMaterialGraphSchema::UpdateDetailView(const UEdGraph* Graph) const
+{
+	FMaterialEditorUtilities::UpdateDetailView(Graph);
+}
+
 int32 UMaterialGraphSchema::GetNodeSelectionCount(const UEdGraph* Graph) const
 {
 	return FMaterialEditorUtilities::GetNumberOfSelectedNodes(Graph);
diff --git a/Engine/Source/Runtime/Engine/Classes/Materials/MaterialExpression.h b/Engine/Source/Runtime/Engine/Classes/Materials/MaterialExpression.h
--- a/Engine/Source/Runtime/Engine/Classes/Materials/MaterialExpression.h
+++ b/Engine/Source/Runtime/Engine/Classes/Materials/MaterialExpression.h
@@ -5,6 +5,7 @@
 #include "CoreMinimal.h"
 #include "UObject/ObjectMacros.h"
 #include "UObject/Object.h"
+#include "UObject/UnrealType.h"
 #include "Misc/Guid.h"
 #include "MaterialShared.h"
 #include "MaterialExpressionIO.h"
@@ -339,6 +340,7 @@ class ENGINE_API UMaterialExpression : public UObject
 	 *	@param	Outputs		The TArray of outputs to fill in.
 	 */
 	virtual TArray<FExpressionOutput>& GetOutputs();
+	/** Get the expression inputs supported by this expression (Note: property inputs NOT included). */
 	virtual const TArray<FExpressionInput*> GetInputs();
 	virtual FExpressionInput* GetInput(int32 InputIndex);
 	virtual FName GetInputName(int32 InputIndex) const;
@@ -348,6 +350,14 @@ class ENGINE_API UMaterialExpression : public UObject
 		return true;
 	};
 
+	/** Find the property that is associated with the input pin. */
+	virtual FProperty* GetInputPinProperty(int32 PinIndex);
+	virtual FName GetInputPinSubCategory(int32 PinIndex);
+	virtual UObject* GetInputPinSubCategoryObject(int32 PinIndex);
+	virtual void PinDefaultValueChanged(int32 PinIndex, const FString& DefaultValue);
+	virtual FString GetInputPinDefaultValue(int32 PinIndex);
+	virtual TArray<FProperty*> GetPropertyInputs() const;
+
 	virtual uint32 GetInputType(int32 InputIndex);
 	virtual uint32 GetOutputType(int32 OutputIndex);
 
diff --git a/Engine/Source/Runtime/Engine/Classes/Materials/MaterialExpressionBumpOffset.h b/Engine/Source/Runtime/Engine/Classes/Materials/MaterialExpressionBumpOffset.h
--- a/Engine/Source/Runtime/Engine/Classes/Materials/MaterialExpressionBumpOffset.h
+++ b/Engine/Source/Runtime/Engine/Classes/Materials/MaterialExpressionBumpOffset.h
@@ -31,7 +31,7 @@ class UMaterialExpressionBumpOffset : public UMaterialExpression
 	float ReferencePlane;    // Height at which no offset is applied.
 
 	/** only used if Coordinate is not hooked up */
-	UPROPERTY(EditAnywhere, Category = MaterialExpressionBumpOffset)
+	UPROPERTY(EditAnywhere, Category = MaterialExpressionBumpOffset, meta = (OverridingInputProperty = "Coordinate"))
 	uint32 ConstCoordinate;
 
 	//~ Begin UMaterialExpression Interface
diff --git a/Engine/Source/Runtime/Engine/Classes/Materials/MaterialExpressionClamp.h b/Engine/Source/Runtime/Engine/Classes/Materials/MaterialExpressionClamp.h
--- a/Engine/Source/Runtime/Engine/Classes/Materials/MaterialExpressionClamp.h
+++ b/Engine/Source/Runtime/Engine/Classes/Materials/MaterialExpressionClamp.h
@@ -31,7 +31,7 @@ class UMaterialExpressionClamp : public UMaterialExpression
 	UPROPERTY(meta = (RequiredInput = "false", ToolTip = "Defaults to 'MaxDefault' if not specified"))
 	FExpressionInput Max;
 
-	UPROPERTY(EditAnywhere, Category=MaterialExpressionClamp)
+	UPROPERTY(EditAnywhere, Category=MaterialExpressionClamp, Meta = (ShowAsInputPin = "Advanced"))
 	TEnumAsByte<enum EClampMode> ClampMode;
 
 	UPROPERTY(EditAnywhere, Category=MaterialExpressionClamp, meta=(OverridingInputProperty = "Min"))
diff --git a/Engine/Source/Runtime/Engine/Classes/Materials/MaterialExpressionComponentMask.h b/Engine/Source/Runtime/Engine/Classes/Materials/MaterialExpressionComponentMask.h
--- a/Engine/Source/Runtime/Engine/Classes/Materials/MaterialExpressionComponentMask.h
+++ b/Engine/Source/Runtime/Engine/Classes/Materials/MaterialExpressionComponentMask.h
@@ -17,16 +17,16 @@ class UMaterialExpressionComponentMask : public UMaterialExpression
 	UPROPERTY()
 	FExpressionInput Input;
 
-	UPROPERTY(EditAnywhere, Category=MaterialExpressionComponentMask)
+	UPROPERTY(EditAnywhere, Category=MaterialExpressionComponentMask, Meta = (ShowAsInputPin = "Advanced"))
 	uint32 R:1;
 
-	UPROPERTY(EditAnywhere, Category=MaterialExpressionComponentMask)
+	UPROPERTY(EditAnywhere, Category=MaterialExpressionComponentMask, Meta = (ShowAsInputPin = "Advanced"))
 	uint32 G:1;
 
-	UPROPERTY(EditAnywhere, Category=MaterialExpressionComponentMask)
+	UPROPERTY(EditAnywhere, Category=MaterialExpressionComponentMask, Meta = (ShowAsInputPin = "Advanced"))
 	uint32 B:1;
 
-	UPROPERTY(EditAnywhere, Category=MaterialExpressionComponentMask)
+	UPROPERTY(EditAnywhere, Category=MaterialExpressionComponentMask, Meta = (ShowAsInputPin = "Advanced"))
 	uint32 A:1;
 
 
diff --git a/Engine/Source/Runtime/Engine/Classes/Materials/MaterialExpressionConstant.h b/Engine/Source/Runtime/Engine/Classes/Materials/MaterialExpressionConstant.h
--- a/Engine/Source/Runtime/Engine/Classes/Materials/MaterialExpressionConstant.h
+++ b/Engine/Source/Runtime/Engine/Classes/Materials/MaterialExpressionConstant.h
@@ -13,7 +13,7 @@ class UMaterialExpressionConstant : public UMaterialExpression
 {
 	GENERATED_UCLASS_BODY()
 
-	UPROPERTY(EditAnywhere, BlueprintReadWrite, Category=MaterialExpressionConstant, DisplayName="Value")
+	UPROPERTY(EditAnywhere, BlueprintReadWrite, Category=MaterialExpressionConstant, DisplayName="Value", Meta=(ShowAsInputPin = "Primary"))
 	float R;
 
 
diff --git a/Engine/Source/Runtime/Engine/Classes/Materials/MaterialExpressionConstant2Vector.h b/Engine/Source/Runtime/Engine/Classes/Materials/MaterialExpressionConstant2Vector.h
--- a/Engine/Source/Runtime/Engine/Classes/Materials/MaterialExpressionConstant2Vector.h
+++ b/Engine/Source/Runtime/Engine/Classes/Materials/MaterialExpressionConstant2Vector.h
@@ -13,10 +13,10 @@ class UMaterialExpressionConstant2Vector : public UMaterialExpression
 {
 	GENERATED_UCLASS_BODY()
 
-	UPROPERTY(EditAnywhere, BlueprintReadWrite, Category = MaterialExpressionConstant2Vector)
+	UPROPERTY(EditAnywhere, BlueprintReadWrite, Category = MaterialExpressionConstant2Vector, DisplayName = "X", Meta = (ShowAsInputPin = "Primary"))
 	float R;
 
-	UPROPERTY(EditAnywhere, BlueprintReadWrite, Category=MaterialExpressionConstant2Vector)
+	UPROPERTY(EditAnywhere, BlueprintReadWrite, Category=MaterialExpressionConstant2Vector, DisplayName = "Y", Meta = (ShowAsInputPin = "Primary"))
 	float G;
 
 
diff --git a/Engine/Source/Runtime/Engine/Classes/Materials/MaterialExpressionConstant3Vector.h b/Engine/Source/Runtime/Engine/Classes/Materials/MaterialExpressionConstant3Vector.h
--- a/Engine/Source/Runtime/Engine/Classes/Materials/MaterialExpressionConstant3Vector.h
+++ b/Engine/Source/Runtime/Engine/Classes/Materials/MaterialExpressionConstant3Vector.h
@@ -13,7 +13,7 @@ class UMaterialExpressionConstant3Vector
 {
 	GENERATED_UCLASS_BODY()
 
- 	UPROPERTY(EditAnywhere, BlueprintReadWrite, Category=MaterialExpressionConstant3Vector, meta=(HideAlphaChannel))
+ 	UPROPERTY(EditAnywhere, BlueprintReadWrite, Category=MaterialExpressionConstant3Vector, meta=(HideAlphaChannel, ShowAsInputPin = "Primary"))
 	FLinearColor Constant;
 
 public:
diff --git a/Engine/Source/Runtime/Engine/Classes/Materials/MaterialExpressionConstant4Vector.h b/Engine/Source/Runtime/Engine/Classes/Materials/MaterialExpressionConstant4Vector.h
--- a/Engine/Source/Runtime/Engine/Classes/Materials/MaterialExpressionConstant4Vector.h
+++ b/Engine/Source/Runtime/Engine/Classes/Materials/MaterialExpressionConstant4Vector.h
@@ -13,7 +13,7 @@ class UMaterialExpressionConstant4Vector : public UMaterialExpression
 {
 	GENERATED_UCLASS_BODY()
 
-	UPROPERTY(EditAnywhere, BlueprintReadWrite, Category=MaterialExpressionConstant4Vector)
+	UPROPERTY(EditAnywhere, BlueprintReadWrite, Category=MaterialExpressionConstant4Vector, Meta = (ShowAsInputPin = "Primary"))
 	FLinearColor Constant;
 
 	//~ Begin UMaterialExpression Interface
diff --git a/Engine/Source/Runtime/Engine/Classes/Materials/MaterialExpressionConstantBiasScale.h b/Engine/Source/Runtime/Engine/Classes/Materials/MaterialExpressionConstantBiasScale.h
--- a/Engine/Source/Runtime/Engine/Classes/Materials/MaterialExpressionConstantBiasScale.h
+++ b/Engine/Source/Runtime/Engine/Classes/Materials/MaterialExpressionConstantBiasScale.h
@@ -17,10 +17,10 @@ class UMaterialExpressionConstantBiasScale : public UMaterialExpression
 	UPROPERTY()
 	FExpressionInput Input;
 
-	UPROPERTY(EditAnywhere, Category=MaterialExpressionConstantBiasScale)
+	UPROPERTY(EditAnywhere, Category=MaterialExpressionConstantBiasScale, Meta = (ShowAsInputPin = "Advanced"))
 	float Bias;
 
-	UPROPERTY(EditAnywhere, Category=MaterialExpressionConstantBiasScale)
+	UPROPERTY(EditAnywhere, Category=MaterialExpressionConstantBiasScale, Meta = (ShowAsInputPin = "Advanced"))
 	float Scale;
 
 
diff --git a/Engine/Source/Runtime/Engine/Classes/Materials/MaterialExpressionCosine.h b/Engine/Source/Runtime/Engine/Classes/Materials/MaterialExpressionCosine.h
--- a/Engine/Source/Runtime/Engine/Classes/Materials/MaterialExpressionCosine.h
+++ b/Engine/Source/Runtime/Engine/Classes/Materials/MaterialExpressionCosine.h
@@ -17,7 +17,7 @@ class UMaterialExpressionCosine : public UMaterialExpression
 	UPROPERTY()
 	FExpressionInput Input;
 
-	UPROPERTY(EditAnywhere, Category=MaterialExpressionCosine)
+	UPROPERTY(EditAnywhere, Category=MaterialExpressionCosine, Meta = (ShowAsInputPin = "Advanced"))
 	float Period;
 
 
diff --git a/Engine/Source/Runtime/Engine/Classes/Materials/MaterialExpressionDesaturation.h b/Engine/Source/Runtime/Engine/Classes/Materials/MaterialExpressionDesaturation.h
--- a/Engine/Source/Runtime/Engine/Classes/Materials/MaterialExpressionDesaturation.h
+++ b/Engine/Source/Runtime/Engine/Classes/Materials/MaterialExpressionDesaturation.h
@@ -21,7 +21,7 @@ class UMaterialExpressionDesaturation : public UMaterialExpression
 	UPROPERTY()
 	FExpressionInput Fraction;
 
-	UPROPERTY(EditAnywhere, Category=MaterialExpressionDesaturation)
+	UPROPERTY(EditAnywhere, Category=MaterialExpressionDesaturation, Meta = (ShowAsInputPin = "Advanced"))
 	FLinearColor LuminanceFactors;    // Color component factors for converting a color to greyscale.
 
 
diff --git a/Engine/Source/Runtime/Engine/Classes/Materials/MaterialExpressionDoubleVectorParameter.h b/Engine/Source/Runtime/Engine/Classes/Materials/MaterialExpressionDoubleVectorParameter.h
--- a/Engine/Source/Runtime/Engine/Classes/Materials/MaterialExpressionDoubleVectorParameter.h
+++ b/Engine/Source/Runtime/Engine/Classes/Materials/MaterialExpressionDoubleVectorParameter.h
@@ -14,7 +14,7 @@ class UMaterialExpressionDoubleVectorParameter : public UMaterialExpressionParam
 {
 	GENERATED_UCLASS_BODY()
 
-	UPROPERTY(EditAnywhere, Category = MaterialExpressionVectorParameter)
+	UPROPERTY(EditAnywhere, Category = MaterialExpressionVectorParameter, Meta = (ShowAsInputPin = "Primary"))
 	FVector4d DefaultValue;
 
 	//~ Begin UMaterialExpression Interface
diff --git a/Engine/Source/Runtime/Engine/Classes/Materials/MaterialExpressionDynamicParameter.h b/Engine/Source/Runtime/Engine/Classes/Materials/MaterialExpressionDynamicParameter.h
--- a/Engine/Source/Runtime/Engine/Classes/Materials/MaterialExpressionDynamicParameter.h
+++ b/Engine/Source/Runtime/Engine/Classes/Materials/MaterialExpressionDynamicParameter.h
@@ -29,11 +29,11 @@ class UMaterialExpressionDynamicParameter : public UMaterialExpression
 	UPROPERTY(EditAnywhere, editfixedsize, Category=MaterialExpressionDynamicParameter)
 	TArray<FString> ParamNames;
 
-	UPROPERTY(EditAnywhere, Category = MaterialExpressionDynamicParameter)
+	UPROPERTY(EditAnywhere, Category = MaterialExpressionDynamicParameter, Meta = (ShowAsInputPin = "Primary"))
 	FLinearColor DefaultValue;
 
 	/** The index of the dynamic parameter for use in tools that allow > 1 */
-	UPROPERTY(EditAnywhere, Category = MaterialExpressionDynamicParameter, meta = (UIMin = 0, ClampMin = 0, UIMax = 3, ClampMax = 3))
+	UPROPERTY(EditAnywhere, Category = MaterialExpressionDynamicParameter, meta = (UIMin = 0, ClampMin = 0, UIMax = 3, ClampMax = 3, ShowAsInputPin = "Primary"))
 	uint32 ParameterIndex;
 
 	//~ Begin UObject Interface
diff --git a/Engine/Source/Runtime/Engine/Classes/Materials/MaterialExpressionSceneColor.h b/Engine/Source/Runtime/Engine/Classes/Materials/MaterialExpressionSceneColor.h
--- a/Engine/Source/Runtime/Engine/Classes/Materials/MaterialExpressionSceneColor.h
+++ b/Engine/Source/Runtime/Engine/Classes/Materials/MaterialExpressionSceneColor.h
@@ -43,7 +43,7 @@ class UMaterialExpressionSceneColor : public UMaterialExpression
 	FExpressionInput OffsetFraction_DEPRECATED;
 
 	/** only used if Input is not hooked up */
-	UPROPERTY(EditAnywhere, Category = MaterialExpressionSceneColor)
+	UPROPERTY(EditAnywhere, Category = MaterialExpressionSceneColor, meta = (OverridingInputProperty = "Input"))
 	FVector2D ConstInput;
 
 	//~ Begin UObject Interface.
diff --git a/Engine/Source/Runtime/Engine/Classes/Materials/MaterialExpressionSine.h b/Engine/Source/Runtime/Engine/Classes/Materials/MaterialExpressionSine.h
--- a/Engine/Source/Runtime/Engine/Classes/Materials/MaterialExpressionSine.h
+++ b/Engine/Source/Runtime/Engine/Classes/Materials/MaterialExpressionSine.h
@@ -17,7 +17,7 @@ class UMaterialExpressionSine : public UMaterialExpression
 	UPROPERTY()
 	FExpressionInput Input;
 
-	UPROPERTY(EditAnywhere, Category=MaterialExpressionSine)
+	UPROPERTY(EditAnywhere, Category=MaterialExpressionSine, Meta = (ShowAsInputPin = "Advanced"))
 	float Period;
 
 
diff --git a/Engine/Source/Runtime/Engine/Classes/Materials/MaterialExpressionTangent.h b/Engine/Source/Runtime/Engine/Classes/Materials/MaterialExpressionTangent.h
--- a/Engine/Source/Runtime/Engine/Classes/Materials/MaterialExpressionTangent.h
+++ b/Engine/Source/Runtime/Engine/Classes/Materials/MaterialExpressionTangent.h
@@ -16,7 +16,7 @@ class UMaterialExpressionTangent : public UMaterialExpression
 	UPROPERTY()
 	FExpressionInput Input;
 
-	UPROPERTY(EditAnywhere, Category=MaterialExpressionTangent)
+	UPROPERTY(EditAnywhere, Category=MaterialExpressionTangent, Meta = (ShowAsInputPin = "Advanced"))
 	float Period;
 
 	//~ Begin UMaterialExpression Interface
diff --git a/Engine/Source/Runtime/Engine/Classes/Materials/MaterialExpressionTextureBase.h b/Engine/Source/Runtime/Engine/Classes/Materials/MaterialExpressionTextureBase.h
--- a/Engine/Source/Runtime/Engine/Classes/Materials/MaterialExpressionTextureBase.h
+++ b/Engine/Source/Runtime/Engine/Classes/Materials/MaterialExpressionTextureBase.h
@@ -23,7 +23,7 @@ class ENGINE_API UMaterialExpressionTextureBase : public UMaterialExpression
 	UPROPERTY(EditAnywhere, BlueprintReadWrite, Category=MaterialExpressionTextureBase)
 	TObjectPtr<class UTexture> Texture;
 
-	UPROPERTY(EditAnywhere, BlueprintReadWrite, Category=MaterialExpressionTextureBase)
+	UPROPERTY(EditAnywhere, BlueprintReadWrite, Category=MaterialExpressionTextureBase, meta = (ShowAsInputPin = "Advanced"))
 	TEnumAsByte<enum EMaterialSamplerType> SamplerType;
 	
 	/** Is default selected texture when using mesh paint mode texture painting */
diff --git a/Engine/Source/Runtime/Engine/Classes/Materials/MaterialExpressionTextureSample.h b/Engine/Source/Runtime/Engine/Classes/Materials/MaterialExpressionTextureSample.h
--- a/Engine/Source/Runtime/Engine/Classes/Materials/MaterialExpressionTextureSample.h
+++ b/Engine/Source/Runtime/Engine/Classes/Materials/MaterialExpressionTextureSample.h
@@ -42,7 +42,7 @@ class ENGINE_API UMaterialExpressionTextureSample : public UMaterialExpressionTe
 	FExpressionInput AutomaticViewMipBiasValue;
 
 	/** Defines how the MipValue property is applied to the texture lookup */
-	UPROPERTY(EditAnywhere, Category=MaterialExpressionTextureSample, meta=(DisplayName = "MipValueMode"))
+	UPROPERTY(EditAnywhere, Category=MaterialExpressionTextureSample, meta=(DisplayName = "MipValueMode", ShowAsInputPin = "Advanced"))
 	TEnumAsByte<enum ETextureMipValueMode> MipValueMode;
 
 	/** 
@@ -51,7 +51,7 @@ class ENGINE_API UMaterialExpressionTextureSample : public UMaterialExpressionTe
 	 * Otherwise use one of the global samplers, which will not consume a sampler slot.
 	 * This allows materials to use more than 16 unique textures on SM5 platforms.
 	 */
-	UPROPERTY(EditAnywhere, Category=MaterialExpressionTextureSample)
+	UPROPERTY(EditAnywhere, Category=MaterialExpressionTextureSample, Meta = (ShowAsInputPin = "Advanced"))
 	TEnumAsByte<enum ESamplerSourceMode> SamplerSource;
 
 	/** Whether the texture should be sampled with per view mip biasing for sharper output with Temporal AA. */
@@ -69,7 +69,7 @@ protected:
 public:
 
 	/** only used if Coordinates is not hooked up */
-	UPROPERTY(EditAnywhere, Category = MaterialExpressionTextureSample)
+	UPROPERTY(EditAnywhere, Category = MaterialExpressionTextureSample, meta = (OverridingInputProperty = "Coordinates"))
 	uint8 ConstCoordinate;
 
 	/** only used if MipValue is not hooked up */
diff --git a/Engine/Source/Runtime/Engine/Private/Materials/MaterialExpressions.cpp b/Engine/Source/Runtime/Engine/Private/Materials/MaterialExpressions.cpp
--- a/Engine/Source/Runtime/Engine/Private/Materials/MaterialExpressions.cpp
+++ b/Engine/Source/Runtime/Engine/Private/Materials/MaterialExpressions.cpp
@@ -934,6 +934,406 @@ void UMaterialExpression::PostDuplicate(bool bDuplicateForPIE)
 
 #if WITH_EDITOR
 
+FProperty* UMaterialExpression::GetInputPinProperty(int32 PinIndex)
+{
+	// Explicit input pins are before property input pins
+	TArray<FProperty*> PropertyInputs = GetPropertyInputs();
+	const int32 NumInputs = GetInputs().Num();
+	if (PinIndex < NumInputs)
+	{
+		FExpressionInput* Input = GetInput(PinIndex);
+
+		// Find the UPROPERTY that has OverridingInputProperty meta data pointing to the expression input.
+		static FName OverridingInputPropertyMetaData(TEXT("OverridingInputProperty"));
+		for (TFieldIterator<FProperty> InputIt(GetClass(), EFieldIteratorFlags::IncludeSuper, EFieldIteratorFlags::ExcludeDeprecated); InputIt; ++InputIt)
+		{
+			FProperty* Property = *InputIt;
+			if (Property->HasMetaData(OverridingInputPropertyMetaData))
+			{
+				const FString& OverridingPropertyName = Property->GetMetaData(OverridingInputPropertyMetaData);
+				FStructProperty* StructProp = FindFProperty<FStructProperty>(GetClass(), *OverridingPropertyName);
+				if (ensure(StructProp != nullptr))
+				{
+					if (Input == StructProp->ContainerPtrToValuePtr<FExpressionInput>(this))
+					{
+						return Property;
+					}
+				}
+			}
+		}
+	}
+	else if (PinIndex < NumInputs + PropertyInputs.Num())
+	{
+		FName PropertyName = PropertyInputs[PinIndex - NumInputs]->GetFName();
+		for (TFieldIterator<FProperty> InputIt(GetClass(), EFieldIteratorFlags::IncludeSuper, EFieldIteratorFlags::ExcludeDeprecated); InputIt; ++InputIt)
+		{
+			FProperty* Property = *InputIt;
+			if (PropertyName == Property->GetFName())
+			{
+				return Property;
+			}
+		}
+	}
+	return nullptr;
+}
+
+FName UMaterialExpression::GetInputPinSubCategory(int32 PinIndex)
+{
+	FName PinSubCategory;
+
+	// Find the UPROPERTY associated with the pin
+	FProperty* Property = GetInputPinProperty(PinIndex);
+	if (Property)
+	{
+		// This is the UPROPERTY matching with the target input
+		FFieldClass* PropertyClass = Property->GetClass();
+		if (PropertyClass == FFloatProperty::StaticClass())
+		{
+			PinSubCategory = UMaterialGraphSchema::PSC_Red;
+		}
+		else if (PropertyClass == FIntProperty::StaticClass())
+		{
+			PinSubCategory = UMaterialGraphSchema::PSC_Int;
+		}
+		else if (PropertyClass == FUInt32Property::StaticClass())
+		{
+			PinSubCategory = UMaterialGraphSchema::PSC_Int;
+		}
+		else if (PropertyClass == FByteProperty::StaticClass())
+		{
+			PinSubCategory = UMaterialGraphSchema::PSC_Byte;
+		}
+		else if (PropertyClass == FBoolProperty::StaticClass())
+		{
+			PinSubCategory = UMaterialGraphSchema::PSC_Bool;
+		}
+		else if (PropertyClass == FStructProperty::StaticClass())
+		{
+			UScriptStruct* Struct = CastField<FStructProperty>(Property)->Struct;
+			if (Struct == TBaseStructure<FLinearColor>::Get())
+			{
+				PinSubCategory = Property->HasMetaData(TEXT("HideAlphaChannel")) ? UMaterialGraphSchema::PSC_RGB : UMaterialGraphSchema::PSC_RGBA;
+			}
+			else if (Struct == TBaseStructure<FVector4>::Get() || Struct == TVariantStructure<FVector4d>::Get())
+			{
+				PinSubCategory = UMaterialGraphSchema::PSC_Vector4;
+			}
+			else if (Struct == TBaseStructure<FVector>::Get() || Struct == TVariantStructure<FVector3f>::Get())
+			{
+				PinSubCategory = UMaterialGraphSchema::PSC_RGB;
+			}
+			else if (Struct == TBaseStructure<FVector2D>::Get())
+			{
+				PinSubCategory = UMaterialGraphSchema::PSC_RG;
+			}	
+		}
+	}
+
+	return PinSubCategory;
+}
+
+UObject* UMaterialExpression::GetInputPinSubCategoryObject(int32 PinIndex)
+{
+	UObject* PinSubCategoryObject = nullptr;
+	FProperty* Property = GetInputPinProperty(PinIndex);
+	if (Property)
+	{
+		if (FByteProperty* ByteProperty = CastField<FByteProperty>(Property))
+		{
+			PinSubCategoryObject = ByteProperty->GetIntPropertyEnum();
+		}
+	}
+	return PinSubCategoryObject;
+}
+
+void UMaterialExpression::PinDefaultValueChanged(int32 PinIndex, const FString& DefaultValue)
+{
+	// Update the default value of the expression input when pin value changes
+	// Find the UPROPERTY that has OverridingInputProperty meta data pointing to the input, override its value.
+	FProperty* Property = GetInputPinProperty(PinIndex);
+	if (Property)
+	{
+		const FString& ClampMin = Property->GetMetaData(TEXT("ClampMin"));
+		const FString& ClampMax = Property->GetMetaData(TEXT("ClampMax"));
+
+		FFieldClass* PropertyClass = Property->GetClass();
+		if (PropertyClass == FFloatProperty::StaticClass())
+		{
+			float Value = FCString::Atof(*DefaultValue);
+			Value = ClampMin.Len() ? (FMath::Max<float>(FCString::Atof(*ClampMin), Value)) : Value;
+			Value = ClampMax.Len() ? (FMath::Min<float>(FCString::Atof(*ClampMax), Value)) : Value;
+			FFloatProperty* FloatProperty = CastField<FFloatProperty>(Property);
+			FloatProperty->SetPropertyValue_InContainer(this, Value);
+		}
+		else if (PropertyClass == FIntProperty::StaticClass())
+		{
+			int32 Value = FCString::Atoi(*DefaultValue);
+			Value = ClampMin.Len() ? (FMath::Max<int32>(FCString::Atoi(*ClampMin), Value)) : Value;
+			Value = ClampMax.Len() ? (FMath::Min<int32>(FCString::Atoi(*ClampMax), Value)) : Value;
+			FIntProperty* IntProperty = CastField<FIntProperty>(Property);
+			IntProperty->SetPropertyValue_InContainer(this, Value);
+		}
+		else if (PropertyClass == FUInt32Property::StaticClass())
+		{	
+			int32 IntValue = FCString::Atoi(*DefaultValue);
+			IntValue = ClampMin.Len() ? (FMath::Max<int32>(FCString::Atoi(*ClampMin), IntValue)) : IntValue;
+			IntValue = ClampMax.Len() ? (FMath::Min<int32>(FCString::Atoi(*ClampMax), IntValue)) : IntValue;
+			// Make sure the value is not negative
+			uint32 Value = (uint32)FMath::Max(IntValue, 0);
+			FUInt32Property* UInt32Property = CastField<FUInt32Property>(Property);
+			UInt32Property->SetPropertyValue_InContainer(this, Value);
+		}
+		else if (PropertyClass == FByteProperty::StaticClass())
+		{
+			FByteProperty* ByteProperty = CastField<FByteProperty>(Property);
+			uint8 Value;
+			UEnum* Enum = ByteProperty->GetIntPropertyEnum();
+			if (Enum)
+			{
+				Value = (uint8)Enum->GetValueByName(FName(DefaultValue));
+			}
+			else
+			{
+				int32 IntValue = FCString::Atoi(*DefaultValue);
+				IntValue = ClampMin.Len() ? (FMath::Max<int32>(FCString::Atoi(*ClampMin), IntValue)) : IntValue;
+				IntValue = ClampMax.Len() ? (FMath::Min<int32>(FCString::Atoi(*ClampMax), IntValue)) : IntValue;
+				// Make sure the value doesn't exceed byte limit
+				Value = (uint8)FMath::Max(FMath::Min(IntValue, 255), 0);
+			}	
+			ByteProperty->SetPropertyValue_InContainer(this, Value);
+		}
+		else if (PropertyClass == FBoolProperty::StaticClass())
+		{
+			bool Value = FCString::ToBool(*DefaultValue);
+			FBoolProperty* BoolProperty = CastField<FBoolProperty>(Property);
+			BoolProperty->SetPropertyValue_InContainer(this, Value);
+		}
+		else if (PropertyClass == FStructProperty::StaticClass())
+		{
+			UScriptStruct* Struct = ((FStructProperty*)Property)->Struct;
+			if (Struct == TBaseStructure<FLinearColor>::Get())
+			{
+				FLinearColor* ColorProperty = Property->ContainerPtrToValuePtr<FLinearColor>(this);
+				if (Property->HasMetaData(TEXT("HideAlphaChannel")))
+				{
+					// This is a 3 element vector
+					TArray<FString> Elements;
+					DefaultValue.ParseIntoArray(Elements, TEXT(","), true);
+					check(Elements.Num() == 3);
+					ColorProperty->R = FCString::Atof(*Elements[0]);
+					ColorProperty->G = FCString::Atof(*Elements[1]);
+					ColorProperty->B = FCString::Atof(*Elements[2]);
+				}
+				else
+				{
+					// This is a 4 element vector
+					ColorProperty->InitFromString(DefaultValue);
+				}
+			}
+			else if (Struct == TBaseStructure<FVector4>::Get())
+			{
+				TArray<FString> Elements;
+				DefaultValue.ParseIntoArray(Elements, TEXT(","), true);
+				check(Elements.Num() == 4);
+				FVector4* Value = Property->ContainerPtrToValuePtr<FVector4>(this);
+				Value->X = FCString::Atod(*Elements[0]);
+				Value->Y = FCString::Atod(*Elements[1]);
+				Value->Z = FCString::Atod(*Elements[2]);
+				Value->W = FCString::Atod(*Elements[3]);
+			}
+			else if (Struct == TVariantStructure<FVector4d>::Get())
+			{
+				TArray<FString> Elements;
+				DefaultValue.ParseIntoArray(Elements, TEXT(","), true);
+				check(Elements.Num() == 4);
+				FVector4d* Value = Property->ContainerPtrToValuePtr<FVector4d>(this);
+				Value->X = FCString::Atod(*Elements[0]);
+				Value->Y = FCString::Atod(*Elements[1]);
+				Value->Z = FCString::Atod(*Elements[2]);
+				Value->W = FCString::Atod(*Elements[3]);
+			}
+			else if (Struct == TBaseStructure<FVector>::Get())
+			{
+				TArray<FString> Elements;
+				DefaultValue.ParseIntoArray(Elements, TEXT(","), true);
+				check(Elements.Num() == 3);
+				FVector* Value = Property->ContainerPtrToValuePtr<FVector>(this);
+				Value->X = FCString::Atod(*Elements[0]);
+				Value->Y = FCString::Atod(*Elements[1]);
+				Value->Z = FCString::Atod(*Elements[2]);
+			}
+			else if (Struct == TVariantStructure<FVector3f>::Get())
+			{
+				TArray<FString> Elements;
+				DefaultValue.ParseIntoArray(Elements, TEXT(","), true);
+				check(Elements.Num() == 3);
+				FVector3f* Value = Property->ContainerPtrToValuePtr<FVector3f>(this);
+				Value->X = FCString::Atof(*Elements[0]);
+				Value->Y = FCString::Atof(*Elements[1]);
+				Value->Z = FCString::Atof(*Elements[2]);
+			}
+			else if (Struct == TBaseStructure<FVector2D>::Get())
+			{
+				FVector2D* Value = Property->ContainerPtrToValuePtr<FVector2D>(this);
+				Value->InitFromString(DefaultValue);
+			}
+		}
+
+		FPropertyChangedEvent Event(Property);
+		PostEditChangeProperty(Event);
+
+		// Update the expression preview and the material to reflect the change
+		Modify();
+		GetDefault<UMaterialGraphSchema>()->ForceVisualizationCacheClear();
+		const UMaterialGraphSchema* Schema = CastChecked<const UMaterialGraphSchema>(GraphNode->GetSchema());
+		Schema->UpdateMaterialOnDefaultValueChanged(GraphNode->GetGraph());
+		// There might be other properties affected by this property change (e.g. propertyA determines if propertyB is read-only) so refresh the detail view
+		Schema->UpdateDetailView(GraphNode->GetGraph());
+	}
+}
+
+FString UMaterialExpression::GetInputPinDefaultValue(int32 PinIndex)
+{
+	// Find the UPROPERTY for the input pin, retrieve its value.
+	FProperty* Property = GetInputPinProperty(PinIndex);
+	if (Property)
+	{
+		// This is the UPROPERTY matching with the target input
+		FFieldClass* PropertyClass = Property->GetClass();
+		if (PropertyClass == FFloatProperty::StaticClass())
+		{
+			FFloatProperty* FloatProperty = CastField<FFloatProperty>(Property);
+			float Value = FloatProperty->GetPropertyValue_InContainer(this);
+			return FString::SanitizeFloat(Value);
+		}
+		else if (PropertyClass == FIntProperty::StaticClass())
+		{
+			FIntProperty* IntProperty = CastField<FIntProperty>(Property);
+			int32 Value = IntProperty->GetPropertyValue_InContainer(this);
+			return FString::FromInt(Value);
+		}
+		else if (PropertyClass == FUInt32Property::StaticClass())
+		{
+			FUInt32Property* UInt32Property = CastField<FUInt32Property>(Property);
+			uint32 Value = UInt32Property->GetPropertyValue_InContainer(this);
+			return FString::FromInt((int32)Value);
+		}
+		else if (PropertyClass == FByteProperty::StaticClass())
+		{
+			FByteProperty* ByteProperty = CastField<FByteProperty>(Property);
+			uint8 Value = ByteProperty->GetPropertyValue_InContainer(this);
+			return ByteProperty->Enum ? ByteProperty->Enum->GetDisplayNameTextByValue(Value).ToString() : FString::FromInt(Value);
+		}
+		else if (PropertyClass == FBoolProperty::StaticClass())
+		{
+			FBoolProperty* BoolProperty = CastField<FBoolProperty>(Property);
+			bool Value = BoolProperty->GetPropertyValue_InContainer(this);
+			return Value ? TEXT("true") : TEXT("false");
+		}
+		else if (PropertyClass == FStructProperty::StaticClass())
+		{
+			UScriptStruct* Struct = ((FStructProperty*)Property)->Struct;
+			if (Struct == TBaseStructure<FLinearColor>::Get())
+			{
+				FLinearColor Value = *Property->ContainerPtrToValuePtr<FLinearColor>(this);
+				if (Property->HasMetaData(TEXT("HideAlphaChannel")))
+				{
+					// This is a 3 element vector
+					return FString::SanitizeFloat(Value.R) + FString(TEXT(",")) + FString::SanitizeFloat(Value.G) + FString(TEXT(",")) + FString::SanitizeFloat(Value.B);
+				}
+				else
+				{
+					// This is a 4 element vector
+					return Value.ToString();
+				}
+			}
+			else if (Struct == TBaseStructure<FVector4>::Get())
+			{
+				FVector4 Value = *Property->ContainerPtrToValuePtr<FVector4>(this);		
+				return FString::SanitizeFloat(Value.X) + FString(TEXT(",")) + FString::SanitizeFloat(Value.Y) + FString(TEXT(",")) + FString::SanitizeFloat(Value.Z) + FString(TEXT(",")) + FString::SanitizeFloat(Value.W);
+			}
+			else if (Struct == TVariantStructure<FVector4d>::Get())
+			{
+				FVector4d Value = *Property->ContainerPtrToValuePtr<FVector4d>(this);		
+				return FString::SanitizeFloat(Value.X) + FString(TEXT(",")) + FString::SanitizeFloat(Value.Y) + FString(TEXT(",")) + FString::SanitizeFloat(Value.Z) + FString(TEXT(",")) + FString::SanitizeFloat(Value.W);
+			}
+			else if (Struct == TBaseStructure<FVector>::Get())
+			{
+				FVector Value = *Property->ContainerPtrToValuePtr<FVector>(this);
+				return FString::SanitizeFloat(Value.X) + FString(TEXT(",")) + FString::SanitizeFloat(Value.Y) + FString(TEXT(",")) + FString::SanitizeFloat(Value.Z);
+			}
+			else if (Struct == TVariantStructure<FVector3f>::Get())
+			{
+				FVector3f Value = *Property->ContainerPtrToValuePtr<FVector3f>(this);
+				return FString::SanitizeFloat(Value.X) + FString(TEXT(",")) + FString::SanitizeFloat(Value.Y) + FString(TEXT(",")) + FString::SanitizeFloat(Value.Z);
+			}
+			else if (Struct == TBaseStructure<FVector2D>::Get())
+			{
+				FVector2D* Value = Property->ContainerPtrToValuePtr<FVector2D>(this);
+				return Value->ToString();
+			}
+		}
+	}
+
+	return TEXT("");
+}
+
+TArray<FProperty*> UMaterialExpression::GetPropertyInputs() const
+{
+	TArray<FProperty*> PropertyInputs;
+
+	static FName OverridingInputPropertyMetaData(TEXT("OverridingInputProperty"));
+	static FName ShowAsInputPinMetaData(TEXT("ShowAsInputPin"));
+	for (TFieldIterator<FProperty> InputIt(GetClass(), EFieldIteratorFlags::IncludeSuper, EFieldIteratorFlags::ExcludeDeprecated); InputIt; ++InputIt)
+	{
+		bool bCreateInput = false;
+		FProperty* Property = *InputIt;
+		// Don't create an expression input if the property is already associated with one explicitly declared
+		bool bOverridingInputProperty = Property->HasMetaData(OverridingInputPropertyMetaData);
+		// It needs to have the 'EditAnywhere' specifier
+		const bool bEditAnywhere = Property->HasAnyPropertyFlags(CPF_Edit);
+		// It needs to be marked with a valid pin category meta data
+		const FString ShowAsInputPin = Property->GetMetaData(ShowAsInputPinMetaData);
+		const bool bShowAsInputPin = ShowAsInputPin == TEXT("Primary") || ShowAsInputPin == TEXT("Advanced");
+
+		if (!bOverridingInputProperty && bEditAnywhere && bShowAsInputPin)
+		{
+			// Check if the property type fits within the allowed widget types
+			FFieldClass* PropertyClass = Property->GetClass();
+			if (PropertyClass == FFloatProperty::StaticClass()
+				|| PropertyClass == FIntProperty::StaticClass()
+				|| PropertyClass == FUInt32Property::StaticClass()
+				|| PropertyClass == FByteProperty::StaticClass()
+				|| PropertyClass == FBoolProperty::StaticClass()
+				)
+			{
+				bCreateInput = true;
+			}
+			else if (PropertyClass == FStructProperty::StaticClass())
+			{
+				FStructProperty* StructProperty = CastField<FStructProperty>(Property);
+				UScriptStruct* Struct = StructProperty->Struct;
+				if (Struct == TBaseStructure<FLinearColor>::Get()
+					|| Struct == TBaseStructure<FVector4>::Get()
+					|| Struct == TVariantStructure<FVector4d>::Get()
+					|| Struct == TBaseStructure<FVector>::Get()
+					|| Struct == TVariantStructure<FVector3f>::Get()
+					|| Struct == TBaseStructure<FVector2D>::Get()
+					)
+				{
+					bCreateInput = true;
+				}
+			}
+		}
+		if (bCreateInput)
+		{
+			PropertyInputs.Add(Property);
+		}
+	}
+
+	return PropertyInputs;
+}
+
 void UMaterialExpression::PostEditChangeProperty(FPropertyChangedEvent& PropertyChangedEvent)
 {
 	Super::PostEditChangeProperty(PropertyChangedEvent);
@@ -956,6 +1356,64 @@ void UMaterialExpression::PostEditChangeProperty(FPropertyChangedEvent& Property
 		}
 	}
 
+	// PropertyChangedEvent.MemberProperty is the owner of PropertyChangedEvent.Property so check for MemberProperty
+	FProperty* MemberPropertyThatChanged = PropertyChangedEvent.MemberProperty;
+	if (MemberPropertyThatChanged != nullptr)
+	{
+		int32 PinIndex = -1;
+		const TArray<FExpressionInput*> AllInputs = GetInputs();
+
+		// Find the expression input this UPROPERTY points to with OverridingInputProperty meta data
+		static FName OverridingInputPropertyMetaData(TEXT("OverridingInputProperty"));
+		if (MemberPropertyThatChanged->HasMetaData(OverridingInputPropertyMetaData))
+		{
+			const FString& OverridingPropertyName = MemberPropertyThatChanged->GetMetaData(OverridingInputPropertyMetaData);
+			FStructProperty* StructProp = FindFProperty<FStructProperty>(GetClass(), *OverridingPropertyName);
+			if (StructProp)
+			{
+				const FExpressionInput* TargetInput = StructProp->ContainerPtrToValuePtr<FExpressionInput>(this);
+				for (int32 i = 0; i < AllInputs.Num(); ++i)
+				{
+					if (TargetInput == AllInputs[i])
+					{
+						PinIndex = i;
+						break;
+					}
+				}
+			}
+		}
+		else
+		{
+			// Not found in explicit expression inputs, so search in property inputs.
+			TArray<FProperty*> PropertyInputs = GetPropertyInputs();
+			for (int32 i = 0; i < PropertyInputs.Num(); ++i)
+			{
+				if (MemberPropertyThatChanged->GetFName() == PropertyInputs[i]->GetFName())
+				{
+					PinIndex = AllInputs.Num() + i;
+				}
+			}
+		}
+
+		if (PinIndex > -1)
+		{
+			// Update the pin value of the expression input
+			UEdGraphPin* Pin = GraphNode->GetPinAt(PinIndex);
+			Pin->DefaultValue = GetInputPinDefaultValue(PinIndex);
+
+			// If the property is linked as inline toggle to another property, both pins need updating to reflect the change.
+			bool bInlineEditConditionToggle = MemberPropertyThatChanged->HasMetaData(TEXT("InlineEditConditionToggle"));
+			bool bEditCondition = MemberPropertyThatChanged->HasMetaData(TEXT("EditCondition"));
+			if (bInlineEditConditionToggle || bEditCondition)
+			{
+				if (GraphNode)
+				{
+					CastChecked<UMaterialGraphNode>(GraphNode)->RecreateAndLinkNode();
+				}
+			}
+		}
+	}
+
 	FProperty* PropertyThatChanged = PropertyChangedEvent.Property;
 	if( PropertyThatChanged != nullptr )
 	{
@@ -1035,6 +1493,19 @@ bool UMaterialExpression::CanEditChange(const FProperty* InProperty) const
 				}
 			}
 		}
+
+		if (bIsEditable)
+		{
+			// If the property has EditCondition metadata, then whether it's editable depends on the other EditCondition property
+			const FString EditConditionPropertyName = InProperty->GetMetaData(TEXT("EditCondition"));
+			if (!EditConditionPropertyName.IsEmpty())
+			{
+				FBoolProperty* EditConditionProperty = FindFProperty<FBoolProperty>(GetClass(), *EditConditionPropertyName);
+				{
+					bIsEditable = *EditConditionProperty->ContainerPtrToValuePtr<bool>(this);
+				}
+			}
+		}
 	}
 
 	return bIsEditable;
diff --git a/Engine/Source/Runtime/Landscape/Classes/Materials/MaterialExpressionLandscapeLayerWeight.h b/Engine/Source/Runtime/Landscape/Classes/Materials/MaterialExpressionLandscapeLayerWeight.h
--- a/Engine/Source/Runtime/Landscape/Classes/Materials/MaterialExpressionLandscapeLayerWeight.h
+++ b/Engine/Source/Runtime/Landscape/Classes/Materials/MaterialExpressionLandscapeLayerWeight.h
@@ -31,7 +31,7 @@ class LANDSCAPE_API UMaterialExpressionLandscapeLayerWeight : public UMaterialEx
 	float PreviewWeight;
 
 	/** only used if Base is not hooked up */
-	UPROPERTY(EditAnywhere, Category = MaterialExpressionLandscapeLayerWeight)
+	UPROPERTY(EditAnywhere, Category = MaterialExpressionLandscapeLayerWeight, meta = (OverridingInputProperty = "Base"))
 	FVector ConstBase;
 
 public:
