StaticMesh: Assert that the streaming data for Nanite is loaded when cooking

#rnx

diff --git a/Engine/Source/Runtime/Engine/Private/StaticMesh.cpp b/Engine/Source/Runtime/Engine/Private/StaticMesh.cpp
--- a/Engine/Source/Runtime/Engine/Private/StaticMesh.cpp
+++ b/Engine/Source/Runtime/Engine/Private/StaticMesh.cpp
@@ -4653,6 +4653,7 @@ bool UStaticMesh::IsCachedCookedPlatformDataLoaded(const ITargetPlatform* Target
 			return false;
 		}
 		NaniteResources.EndRebuildBulkDataFromCache();
+		check(NaniteResources.StreamablePages.GetBulkDataSize() > 0);
 	}
 
 	return true;
