Ensure on missing linker in SCS instead of crash

diff --git a/Engine/Source/Runtime/Engine/Private/SCS_Node.cpp b/Engine/Source/Runtime/Engine/Private/SCS_Node.cpp
--- a/Engine/Source/Runtime/Engine/Private/SCS_Node.cpp
+++ b/Engine/Source/Runtime/Engine/Private/SCS_Node.cpp
@@ -332,7 +332,10 @@ void USCS_Node::PreloadChain()
 
 	if (ComponentTemplate && ComponentTemplate->HasAnyFlags(RF_NeedLoad))
 	{
-		ComponentTemplate->GetLinker()->Preload(ComponentTemplate);
+		if (ensure(ComponentTemplate->GetLinker()))
+		{
+			ComponentTemplate->GetLinker()->Preload(ComponentTemplate);
+		}
 	}
 
 	for( TArray<USCS_Node*>::TIterator ChildIt(ChildNodes); ChildIt; ++ChildIt )
