add TC for SingleFloat and HDR_F32
to expose GPU texture formats for R32F and RGBA32F
PF_A32B32G32R32F
PF_R32_FLOAT
OpenGL ES platforms do not support F32 with filtering so must convert to F16 unless the mip filter is set to nearest

diff --git a/Engine/Source/Developer/Android/AndroidTargetPlatform/Private/AndroidTargetPlatform.cpp b/Engine/Source/Developer/Android/AndroidTargetPlatform/Private/AndroidTargetPlatform.cpp
--- a/Engine/Source/Developer/Android/AndroidTargetPlatform/Private/AndroidTargetPlatform.cpp
+++ b/Engine/Source/Developer/Android/AndroidTargetPlatform/Private/AndroidTargetPlatform.cpp
@@ -74,14 +74,15 @@ namespace AndroidTexFormat
 	const static FName NameAutoASTC(TEXT("ASTC_RGBAuto"));
 	
 	// Uncompressed Texture Formats
-	const static FName NameBGRA8(TEXT("BGRA8"));
-	const static FName NameG8(TEXT("G8"));
-	const static FName NameVU8(TEXT("VU8"));
-	const static FName NameRGBA16F(TEXT("RGBA16F"));
-	const static FName NameR16F(TEXT("R16F"));
-	const static FName NameR5G6B5(TEXT("R5G6B5"));
-	const static FName NameA1RGB555(TEXT("A1RGB555"));
+	//const static FName NameBGRA8(TEXT("BGRA8"));
+	//const static FName NameG8(TEXT("G8"));
+	//const static FName NameVU8(TEXT("VU8"));
+	//const static FName NameRGBA16F(TEXT("RGBA16F"));
+	//const static FName NameR16F(TEXT("R16F"));
+	//const static FName NameR5G6B5(TEXT("R5G6B5"));
+
 	//A1RGB555 is mapped to RGB555A1, because OpenGL GL_RGB5_A1 only supports alpha on the lowest bit
+	const static FName NameA1RGB555(TEXT("A1RGB555"));
 	const static FName NameRGB555A1(TEXT("RGB555A1"));
 
 	const static FName ASTCRemap[][2] =
@@ -487,9 +488,11 @@ void FAndroidTargetPlatform::GetTextureFormats( const UTexture* Texture, TArray<
 	const bool bSupportCompressedVolumeTexture = SupportsTextureFormatCategory(EAndroidTextureFormatCategory::ASTC);
 	// TODO: compressed HDR formats
 	const bool bSupportDX11TextureFormats = false;
+	// OpenGL ES has F32 textures but doesn't allow linear filtering unless OES_texture_float_linear
+	const bool bSupportFilteredFloat32Textures = false;
 
 	TArray<FName>& LayerFormats = OutFormats.AddDefaulted_GetRef();
-	GetDefaultTextureFormatNamePerLayer(LayerFormats, this, Texture, bSupportDX11TextureFormats, bSupportCompressedVolumeTexture, 1);
+	GetDefaultTextureFormatNamePerLayer(LayerFormats, this, Texture, bSupportDX11TextureFormats, bSupportCompressedVolumeTexture, 1, bSupportFilteredFloat32Textures);
 
 	for (FName& TextureFormatName : LayerFormats)
 	{
diff --git a/Engine/Source/Developer/HoloLens/HoloLensTargetPlatform/Private/HoloLensTargetPlatform.cpp b/Engine/Source/Developer/HoloLens/HoloLensTargetPlatform/Private/HoloLensTargetPlatform.cpp
--- a/Engine/Source/Developer/HoloLens/HoloLensTargetPlatform/Private/HoloLensTargetPlatform.cpp
+++ b/Engine/Source/Developer/HoloLens/HoloLensTargetPlatform/Private/HoloLensTargetPlatform.cpp
@@ -190,7 +190,8 @@ void FHoloLensTargetPlatform::GetReflectionCaptureFormats(TArray<FName>& OutForm
 
 void FHoloLensTargetPlatform::GetTextureFormats(const UTexture* InTexture, TArray< TArray<FName> >& OutFormats) const
 {
-	GetDefaultTextureFormatNamePerLayer(OutFormats.AddDefaulted_GetRef(), this, InTexture, false);
+	bool bSupportFilteredFloat32Textures = false;
+	GetDefaultTextureFormatNamePerLayer(OutFormats.AddDefaulted_GetRef(), this, InTexture, false, false,4, bSupportFilteredFloat32Textures);
 }
 
 void FHoloLensTargetPlatform::GetAllTextureFormats(TArray<FName>& OutFormats) const
diff --git a/Engine/Source/Developer/TextureFormatUncompressed/Private/TextureFormatUncompressed.cpp b/Engine/Source/Developer/TextureFormatUncompressed/Private/TextureFormatUncompressed.cpp
--- a/Engine/Source/Developer/TextureFormatUncompressed/Private/TextureFormatUncompressed.cpp
+++ b/Engine/Source/Developer/TextureFormatUncompressed/Private/TextureFormatUncompressed.cpp
@@ -40,6 +40,7 @@ class FUncompressedTextureBuildFunction final : public FTextureBuildFunction
 	op(G16) \
 	op(VU8) \
 	op(RGBA16F) \
+	op(RGBA32F) \
 	op(XGXR8) \
 	op(RGBA8) \
 	op(POTERROR) \
@@ -91,6 +92,10 @@ class FTextureFormatUncompressed : public ITextureFormat
 		{
 			return TEXT("RGBA16F");
 		}
+		else if (InBuildSettings.TextureFormatName == GTextureFormatNameRGBA32F)
+		{
+			return TEXT("RGBA32F");
+		}
 		else if (InBuildSettings.TextureFormatName == GTextureFormatNameR16F)
 		{
 			return TEXT("R16F");
@@ -144,6 +149,10 @@ class FTextureFormatUncompressed : public ITextureFormat
 		{
 			return PF_FloatRGBA;
 		}
+		else if (BuildSettings.TextureFormatName == GTextureFormatNameRGBA32F)
+		{
+			return PF_A32B32G32R32F;
+		}
 		else if (BuildSettings.TextureFormatName == GTextureFormatNameR16F)
 		{
 			return PF_R16F;
@@ -309,6 +318,18 @@ class FTextureFormatUncompressed : public ITextureFormat
 
 			return true;
 		}
+		else if (BuildSettings.TextureFormatName == GTextureFormatNameRGBA32F)
+		{
+			FImage Image;
+			InImage.CopyTo(Image, ERawImageFormat::RGBA32F, EGammaSpace::Linear);
+
+			OutCompressedImage.SizeX = Image.SizeX;
+			OutCompressedImage.SizeY = Image.SizeY;
+			OutCompressedImage.SizeZ = (BuildSettings.bVolume || BuildSettings.bTextureArray) ? Image.NumSlices : 1;
+			OutCompressedImage.RawData = MoveTemp(Image.RawData);
+
+			return true;
+		}
 		else if (BuildSettings.TextureFormatName == GTextureFormatNameR16F)
 		{
 			FImage Image;
diff --git a/Engine/Source/Runtime/Engine/Classes/Engine/Texture.h b/Engine/Source/Runtime/Engine/Classes/Engine/Texture.h
--- a/Engine/Source/Runtime/Engine/Classes/Engine/Texture.h
+++ b/Engine/Source/Runtime/Engine/Classes/Engine/Texture.h
@@ -132,7 +132,7 @@ struct FTextureSource
 	// should match ERawImageFormat::GetBytesPerPixel
 	ENGINE_API static int32 GetBytesPerPixel(ETextureSourceFormat Format);
 	// should match ERawImageFormat::IsHDR
-	FORCEINLINE static bool IsHDR(ETextureSourceFormat Format) { return (Format == TSF_BGRE8 || Format == TSF_RGBA16F || Format == TSF_R16F || Format == TSF_R32F || Format == TSF_RGBA32F); }
+	FORCEINLINE static bool IsHDR(ETextureSourceFormat Format) { return UE::TextureDefines::IsHDR(Format); }
 	
 	enum class ELockState : uint8
 	{
@@ -1704,14 +1704,7 @@ public:
 	/** @return true if the texture has an uncompressed texture setting */
 	bool IsUncompressed() const
 	{
-		return (CompressionSettings == TC_Grayscale ||
-				CompressionSettings == TC_Displacementmap ||
-				CompressionSettings == TC_VectorDisplacementmap ||
-				CompressionSettings == TC_HDR ||
-				CompressionSettings == TC_EditorIcon ||
-				CompressionSettings == TC_DistanceFieldFont ||
-				CompressionSettings == TC_HalfFloat
-			);
+		return UE::TextureDefines::IsUncompressed(CompressionSettings);
 	}
 
 	/**
diff --git a/Engine/Source/Runtime/Engine/Classes/Engine/TextureDefines.h b/Engine/Source/Runtime/Engine/Classes/Engine/TextureDefines.h
--- a/Engine/Source/Runtime/Engine/Classes/Engine/TextureDefines.h
+++ b/Engine/Source/Runtime/Engine/Classes/Engine/TextureDefines.h
@@ -336,11 +336,13 @@ enum TextureCompressionSettings
 	TC_EditorIcon				UMETA(DisplayName = "UserInterface2D (RGBA)"),
 	TC_Alpha					UMETA(DisplayName = "Alpha (no sRGB, BC4 on DX11)"),
 	TC_DistanceFieldFont		UMETA(DisplayName = "DistanceFieldFont (G8)"),
-	TC_HDR_Compressed			UMETA(DisplayName = "HDRCompressed (RGB, BC6H, DX11)"),
+	TC_HDR_Compressed			UMETA(DisplayName = "HDR Compressed (RGB, BC6H, DX11)"),
 	TC_BC7						UMETA(DisplayName = "BC7 (DX11, optional A)"),
 	TC_HalfFloat				UMETA(DisplayName = "Half Float (R16F)"),
 	TC_LQ				        UMETA(Hidden, DisplayName = "Low Quality (BGR565/BGR555A1)", ToolTip = "BGR565/BGR555A1, fallback to DXT1/DXT5 on Mac platform"),
-	TC_EncodedReflectionCapture		UMETA(Hidden),
+	TC_EncodedReflectionCapture	UMETA(Hidden),
+	TC_SingleFloat				UMETA(DisplayName = "Single Float (R32F)"),
+	TC_HDR_F32					UMETA(DisplayName = "HDR High Precision (RGBA32F)"),
 	TC_MAX,
 };
 
@@ -400,3 +402,46 @@ enum class ETextureChromaticAdaptationMethod : uint8
 	TCAM_CAT02		= 2 UMETA(DisplayName = "CAT02", ToolTip = "Chromatic adaptation is applied using the CAT02 method."),
 };
 
+
+namespace UE
+{
+namespace TextureDefines
+{
+
+static FORCEINLINE bool IsHDR(ETextureSourceFormat Format)
+{
+	return (Format == TSF_BGRE8 || Format == TSF_RGBA16F || Format == TSF_RGBA32F || Format == TSF_R16F || Format == TSF_R32F);
+}
+
+static FORCEINLINE bool IsHDR(TextureCompressionSettings CompressionSettings)
+{
+	switch(CompressionSettings)
+	{
+	case TC_HDR:
+	case TC_HDR_F32:
+	case TC_HDR_Compressed:
+	case TC_HalfFloat:
+	case TC_SingleFloat:
+		return true;
+	default:
+		return false;
+	}
+}
+
+static FORCEINLINE bool IsUncompressed(TextureCompressionSettings CompressionSettings)
+{
+	return (CompressionSettings == TC_Grayscale ||
+			CompressionSettings == TC_Displacementmap ||
+			CompressionSettings == TC_VectorDisplacementmap ||
+			CompressionSettings == TC_HDR ||
+			CompressionSettings == TC_HDR_F32 ||
+			CompressionSettings == TC_EditorIcon ||
+			CompressionSettings == TC_DistanceFieldFont ||
+			CompressionSettings == TC_HalfFloat ||
+			CompressionSettings == TC_SingleFloat
+		);
+}
+
+} // TextureDefines
+} // UE
+
diff --git a/Engine/Source/Runtime/Engine/Private/Rendering/StreamableTextureResource.cpp b/Engine/Source/Runtime/Engine/Private/Rendering/StreamableTextureResource.cpp
--- a/Engine/Source/Runtime/Engine/Private/Rendering/StreamableTextureResource.cpp
+++ b/Engine/Source/Runtime/Engine/Private/Rendering/StreamableTextureResource.cpp
@@ -94,7 +94,7 @@ FStreamableTextureResource::FStreamableTextureResource(UTexture* InOwner, const
 	, PixelFormat(InPlatformData->PixelFormat)
 {
 	// HDR images are stored in linear but still require gamma correction to display correctly.
-	bIgnoreGammaConversions = !InOwner->SRGB && InOwner->CompressionSettings != TC_HDR && InOwner->CompressionSettings != TC_HDR_Compressed && InOwner->CompressionSettings != TC_HalfFloat;
+	bIgnoreGammaConversions = !InOwner->SRGB && ! UE::TextureDefines::IsHDR(InOwner->CompressionSettings);
 	bSRGB = InOwner->SRGB;
 	bGreyScaleFormat = (PixelFormat == PF_G8) || (PixelFormat == PF_BC4);
 
diff --git a/Engine/Source/Runtime/Engine/Private/Texture.cpp b/Engine/Source/Runtime/Engine/Private/Texture.cpp
--- a/Engine/Source/Runtime/Engine/Private/Texture.cpp
+++ b/Engine/Source/Runtime/Engine/Private/Texture.cpp
@@ -39,6 +39,7 @@
 #include "ColorSpace.h"
 #include "ImageCoreUtils.h"
 #include "ImageUtils.h"
+#include "Algo/Unique.h"
 
 #if WITH_EDITOR
 #include "DerivedDataBuildVersion.h"
@@ -381,7 +382,6 @@ void UTexture::ValidateSettingsAfterImportOrEdit(bool * pRequiresNotifyMaterials
 		UE_LOG(LogTexture, Display, TEXT("Power of 2 padding cannot be used with LeaveExistingMips, disabled. (%s)"), *GetName());
 
 		PowerOfTwoMode = ETexturePowerOfTwoSetting::None;
-		// @@!! this is broken in the Texture Editor Details GUI ; the PowerOfTwoMode button does not update
 	}
 
 	// IsPowerOfTwo only checks XY :
@@ -474,7 +474,8 @@ void UTexture::ValidateSettingsAfterImportOrEdit(bool * pRequiresNotifyMaterials
 #endif // #if WITH_EDITORONLY_DATA
 
 	// check TC_ CompressionSettings that should have SRGB off
-	const bool bPreventSRGB = (CompressionSettings == TC_Alpha || CompressionSettings == TC_Normalmap || CompressionSettings == TC_Masks || CompressionSettings == TC_HDR || CompressionSettings == TC_HDR_Compressed || CompressionSettings == TC_HalfFloat);
+	const bool bPreventSRGB = CompressionSettings == TC_Alpha || CompressionSettings == TC_Normalmap || CompressionSettings == TC_Masks || 
+		UE::TextureDefines::IsHDR(CompressionSettings);
 	if (bPreventSRGB && SRGB == true)
 	{
 		SRGB = false;
@@ -2855,13 +2856,13 @@ static FName ConditionalGetPrefixedFormat(FName TextureFormatName, const ITarget
 	return TextureFormatName;
 }
 
-FName GetDefaultTextureFormatName( const ITargetPlatform* TargetPlatform, const UTexture* Texture, int32 LayerIndex, bool bSupportDX11TextureFormats, bool bSupportCompressedVolumeTexture, int32 BlockSize )
+FName GetDefaultTextureFormatName( const ITargetPlatform* TargetPlatform, const UTexture* Texture, int32 LayerIndex, bool bSupportDX11TextureFormats, bool bSupportCompressedVolumeTexture, int32 BlockSize, bool bSupportFilteredFloat32Textures )
 {
 	FName TextureFormatName = NAME_None;
 	bool bOodleTextureSdkVersionIsNone = true;
 
 	/**
-	 * IF you add a format to this function don't forget to update GetAllDefaultTextureFormatNames 
+	 * IF you add a format to this function don't forget to update GetAllDefaultTextureFormats 
 	 */
 
 #if WITH_EDITOR
@@ -2879,6 +2880,7 @@ FName GetDefaultTextureFormatName( const ITargetPlatform* TargetPlatform, const
 	static FName NameG16(TEXT("G16"));
 	static FName NameVU8(TEXT("VU8"));
 	static FName NameRGBA16F(TEXT("RGBA16F"));
+	static FName NameRGBA32F(TEXT("RGBA32F"));
 	static FName NameR16F(TEXT("R16F"));
 	static FName NameR32F(TEXT("R32F"));
 	static FName NameBC6H(TEXT("BC6H"));
@@ -2900,13 +2902,7 @@ FName GetDefaultTextureFormatName( const ITargetPlatform* TargetPlatform, const
 	// Determine the pixel format of the (un/)compressed texture
 	
 	// Identify TC groups that will map to uncompressed :		
-	bool bIsTCThatMapsToUncompressed =
-		FormatSettings.CompressionSettings == TC_HDR || 
-		FormatSettings.CompressionSettings == TC_Displacementmap ||
-		FormatSettings.CompressionSettings == TC_VectorDisplacementmap ||
-		FormatSettings.CompressionSettings == TC_Grayscale ||
-		FormatSettings.CompressionSettings == TC_DistanceFieldFont ||
-		FormatSettings.CompressionSettings == TC_HalfFloat;
+	bool bIsTCThatMapsToUncompressed = UE::TextureDefines::IsUncompressed(FormatSettings.CompressionSettings);
 
 	// see if compression needs to be forced off even if requested :
 	bool bNoCompression = FormatSettings.CompressionNone				// Code wants the texture uncompressed.
@@ -2915,7 +2911,7 @@ FName GetDefaultTextureFormatName( const ITargetPlatform* TargetPlatform, const
 		|| (Texture->LODGroup == TEXTUREGROUP_ColorLookupTable)	// Textures in certain LOD groups should remain uncompressed.
 		|| (Texture->LODGroup == TEXTUREGROUP_Bokeh)
 		|| (Texture->LODGroup == TEXTUREGROUP_IESLightProfile)
-		|| (Texture->GetMaterialType() == MCT_VolumeTexture && !bSupportCompressedVolumeTexture)
+		|| ((Texture->GetMaterialType() == MCT_VolumeTexture) && !bSupportCompressedVolumeTexture)
 		|| FormatSettings.CompressionSettings == TC_EncodedReflectionCapture;
 
 	// note that bNoCompression is not the same as bIsTCThatMapsToUncompressed
@@ -2970,6 +2966,7 @@ FName GetDefaultTextureFormatName( const ITargetPlatform* TargetPlatform, const
 			}
 
 			// note if bIsTCThatMapsToUncompressed , this might change the mapping :
+			// @todo Oodle : should not do this if bIsTCThatMapsToUncompressed already
 			bNoCompression = true;
 		}
 
@@ -3000,6 +2997,15 @@ FName GetDefaultTextureFormatName( const ITargetPlatform* TargetPlatform, const
 		// eg. if you were in bIsTCThatMapsToUncompressed , and for some reason you get bNoCompression set
 		//	 then your output format can change from one uncompressed format to another
 
+		// @todo Oodle : if this was redone, I would make two changes
+		// 1. first of all detect if you already have a TC that maps to uncompressed with a specific format,
+		//		and if so, stick to that
+		//		eg. currently if you choose TC_HDR and your source image is G8 you will get RGBA16F output
+		//		but if the source image is not a multiple of 4, bNoCompression turns on, and you get G8 output
+		//		this can be fixed by putting the bNoCompression check after the standard TC maps
+		// 2. second if you do the bNoCompression, fix TSF formats not mapping to the closest possible
+		//		uncompressed renderable format
+
 		if (Texture->HasHDRSource(LayerIndex))
 		{
 			// @todo Oodle : this is a very poor selection from SourceFormat
@@ -3017,11 +3023,14 @@ FName GetDefaultTextureFormatName( const ITargetPlatform* TargetPlatform, const
 		}
 		else if (FormatSettings.CompressionSettings == TC_Normalmap && bUseDXT5NormalMap)
 		{
+			// move R to A like we do for DXT5 normal maps :
 			TextureFormatName = NameXGXR8;
 		}
 		else
 		{
 			// note CompressionNoAlpha no longer kills alpha on AutoDXT
+			//  CompressionNoAlpha really just means "pick BC1 instead of BC3"
+			//	but if your image is not a multiple of 4 you still get alpha RGBA
 			TextureFormatName = NameBGRA8;
 		}
 	}
@@ -3041,6 +3050,10 @@ FName GetDefaultTextureFormatName( const ITargetPlatform* TargetPlatform, const
 	{
 		TextureFormatName = NameRGBA16F;
 	}
+	else if (FormatSettings.CompressionSettings == TC_HDR_F32)
+	{
+		TextureFormatName = NameRGBA32F;
+	}
 	else if (FormatSettings.CompressionSettings == TC_Normalmap)
 	{
 		TextureFormatName = bUseDXT5NormalMap ? NameDXT5n : NameBC5;
@@ -3083,6 +3096,10 @@ FName GetDefaultTextureFormatName( const ITargetPlatform* TargetPlatform, const
 	{
 		TextureFormatName = NameR16F;
 	}
+	else if (FormatSettings.CompressionSettings == TC_SingleFloat)
+	{
+		TextureFormatName = NameR32F;
+	}
 	else
 	{
 		check( FormatSettings.CompressionSettings == TC_Default ||
@@ -3101,6 +3118,8 @@ FName GetDefaultTextureFormatName( const ITargetPlatform* TargetPlatform, const
 	// Some PC GPUs don't support sRGB read from G8 textures (e.g. AMD DX10 cards on ShaderModel3.0)
 	// This solution requires 4x more memory but a lot of PC HW emulate the format anyway
 	// note: GrayscaleSRGB is off on all targetplatforms currently
+	// @todo Oodle : I think this could use G16 instead and be half the size
+	//	 (that's doing the gamma->linear in the G8->G16 conversion)
 	if ((TextureFormatName == NameG8) && FormatSettings.SRGB && !TargetPlatform->SupportsFeature(ETargetPlatformFeatures::GrayscaleSRGB))
 	{
 		TextureFormatName = NameBGRA8;
@@ -3123,25 +3142,55 @@ FName GetDefaultTextureFormatName( const ITargetPlatform* TargetPlatform, const
 		}
 	}
 	
+	// remap 32F to 16F if not supported :
+	if ( !bSupportFilteredFloat32Textures &&
+		( TextureFormatName == NameR32F || TextureFormatName == NameRGBA32F ) )
+	{
+		// Texture::Filter can be manually set to TF_Nearest , if it's Default it comes from LOD Group
+		//   eg. Nearest for TEXTUREGROUP_ColorLookupTable and TEXTUREGROUP_Pixels2D
+		const UTextureLODSettings& LODSettings = TargetPlatform->GetTextureLODSettings();
+		ETextureSamplerFilter Filter = LODSettings.GetSamplerFilter(Texture);
+
+		if ( Filter != ETextureSamplerFilter::Point )
+		{
+			// non-Point filters require remap
+
+			UE_LOG(LogTexture, Display, TEXT("32 bit float texture changed to 16F because Filter is not Nearest and !bSupportFilteredFloat32Textures : %s"), *Texture->GetPathName());
+
+			if (TextureFormatName == NameR32F)
+			{
+				TextureFormatName = NameR16F;
+			}
+			else 
+			{
+				check(TextureFormatName == NameRGBA32F);
+
+				TextureFormatName = NameRGBA16F;
+			}			
+		}
+	}
+
 	bOodleTextureSdkVersionIsNone = Texture->OodleTextureSdkVersion.IsNone();
 #endif //WITH_EDITOR
 
 	return ConditionalGetPrefixedFormat(TextureFormatName, TargetPlatform, bOodleTextureSdkVersionIsNone);
 }
 
-void GetDefaultTextureFormatNamePerLayer(TArray<FName>& OutFormatNames, const class ITargetPlatform* TargetPlatform, const class UTexture* Texture, bool bSupportDX11TextureFormats, bool bSupportCompressedVolumeTexture, int32 BlockSize)
+void GetDefaultTextureFormatNamePerLayer(TArray<FName>& OutFormatNames, const class ITargetPlatform* TargetPlatform, const class UTexture* Texture, bool bSupportDX11TextureFormats, bool bSupportCompressedVolumeTexture, int32 BlockSize, bool bSupportFilteredFloat32Textures )
 {
 #if WITH_EDITOR
 	OutFormatNames.Reserve(Texture->Source.GetNumLayers());
 	for (int32 LayerIndex = 0; LayerIndex < Texture->Source.GetNumLayers(); ++LayerIndex)
 	{
-		OutFormatNames.Add(GetDefaultTextureFormatName(TargetPlatform, Texture, LayerIndex, bSupportDX11TextureFormats, bSupportCompressedVolumeTexture, BlockSize));
+		OutFormatNames.Add(GetDefaultTextureFormatName(TargetPlatform, Texture, LayerIndex, bSupportDX11TextureFormats, bSupportCompressedVolumeTexture, BlockSize, bSupportFilteredFloat32Textures));
 	}
 #endif // WITH_EDITOR
 }
 
 void GetAllDefaultTextureFormats(const class ITargetPlatform* TargetPlatform, TArray<FName>& OutFormats, bool bSupportDX11TextureFormats)
 {
+	// this is only used by CookOnTheFlyServer, it could be removed entirely
+
 #if WITH_EDITOR
 	static FName NameDXT1(TEXT("DXT1"));
 	static FName NameDXT3(TEXT("DXT3"));
@@ -3156,6 +3205,7 @@ void GetAllDefaultTextureFormats(const class ITargetPlatform* TargetPlatform, TA
 	static FName NameG16(TEXT("G16"));
 	static FName NameVU8(TEXT("VU8"));
 	static FName NameRGBA16F(TEXT("RGBA16F"));
+	static FName NameRGBA32F(TEXT("RGBA32F"));
 	static FName NameR16F(TEXT("R16F"));
 	static FName NameR32F(TEXT("R32F"));
 	static FName NameBC6H(TEXT("BC6H"));
@@ -3174,6 +3224,7 @@ void GetAllDefaultTextureFormats(const class ITargetPlatform* TargetPlatform, TA
 	OutFormats.Add(NameG16);
 	OutFormats.Add(NameVU8);
 	OutFormats.Add(NameRGBA16F);
+	OutFormats.Add(NameRGBA32F);
 	OutFormats.Add(NameR16F);
 	OutFormats.Add(NameR32F);
 	if (bSupportDX11TextureFormats)
@@ -3181,14 +3232,21 @@ void GetAllDefaultTextureFormats(const class ITargetPlatform* TargetPlatform, TA
 		OutFormats.Add(NameBC6H);
 		OutFormats.Add(NameBC7);
 	}
+	// is there any drawback to just adding the 32F textures here even if we don't want them? -> no.
+	//	what is this list even used for?
+	//  AFAICT it's only used by CookOnTheFlyServer for GetVersionFormatNumbersForIniVersionStrings
 
 	// go over the original base formats only, and possibly add on to the end of the array if there is a prefix needed
 	int NumBaseFormats = OutFormats.Num();
 	for (int Index = 0; Index < NumBaseFormats; Index++)
 	{
-		OutFormats.AddUnique(ConditionalGetPrefixedFormat(OutFormats[Index], TargetPlatform, true));
-		OutFormats.AddUnique(ConditionalGetPrefixedFormat(OutFormats[Index], TargetPlatform, false));
+		OutFormats.Add(ConditionalGetPrefixedFormat(OutFormats[Index], TargetPlatform, true));
+		OutFormats.Add(ConditionalGetPrefixedFormat(OutFormats[Index], TargetPlatform, false));
 	}
+	
+	// make unique:		
+	OutFormats.Sort( FNameFastLess() );
+	OutFormats.SetNum( Algo::Unique( OutFormats ) );
 #endif
 }
 
diff --git a/Engine/Source/Runtime/Engine/Private/Texture2D.cpp b/Engine/Source/Runtime/Engine/Private/Texture2D.cpp
--- a/Engine/Source/Runtime/Engine/Private/Texture2D.cpp
+++ b/Engine/Source/Runtime/Engine/Private/Texture2D.cpp
@@ -1463,7 +1463,7 @@ void FVirtualTexture2DResource::InitializeEditorResources(IVirtualTexture* InVir
 		RHIBindDebugLabelName(TextureRHI, *Name);
 		RHIUpdateTextureReference(TextureOwner->TextureReference.TextureReferenceRHI, TextureRHI);
 
-		bIgnoreGammaConversions = !TextureOwner->SRGB && TextureOwner->CompressionSettings != TC_HDR && TextureOwner->CompressionSettings != TC_HalfFloat;
+		bIgnoreGammaConversions = !TextureOwner->SRGB && ! UE::TextureDefines::IsHDR(TextureOwner->CompressionSettings);
 
 		// re factored to ensure this is set earlier...make sure it's correct
 		ensure(bSRGB == TextureOwner->SRGB);
diff --git a/Engine/Source/Runtime/Engine/Public/TextureResource.h b/Engine/Source/Runtime/Engine/Public/TextureResource.h
--- a/Engine/Source/Runtime/Engine/Public/TextureResource.h
+++ b/Engine/Source/Runtime/Engine/Public/TextureResource.h
@@ -605,10 +605,10 @@ private:
 };
 
 /** Gets the name of a format for the given LayerIndex */
-ENGINE_API FName GetDefaultTextureFormatName( const class ITargetPlatform* TargetPlatform, const class UTexture* Texture, int32 LayerIndex, bool bSupportDX11TextureFormats, bool bSupportCompressedVolumeTexture = false, int32 BlockSize = 4);
+ENGINE_API FName GetDefaultTextureFormatName( const class ITargetPlatform* TargetPlatform, const class UTexture* Texture, int32 LayerIndex, bool bSupportDX11TextureFormats, bool bSupportCompressedVolumeTexture = false, int32 BlockSize = 4, bool bSupportFilteredFloat32Textures = true);
 
 /** Gets an array of format names for each layer in the texture */
-ENGINE_API void GetDefaultTextureFormatNamePerLayer(TArray<FName>& OutFormatNames, const class ITargetPlatform* TargetPlatform, const class UTexture* Texture, bool bSupportDX11TextureFormats, bool bSupportCompressedVolumeTexture = false, int32 BlockSize = 4);
+ENGINE_API void GetDefaultTextureFormatNamePerLayer(TArray<FName>& OutFormatNames, const class ITargetPlatform* TargetPlatform, const class UTexture* Texture, bool bSupportDX11TextureFormats, bool bSupportCompressedVolumeTexture = false, int32 BlockSize = 4, bool bSupportFilteredFloat32Textures = true);
 
 // returns all the texture formats which can be returned by GetDefaultTextureFormatName
 ENGINE_API void GetAllDefaultTextureFormats( const class ITargetPlatform* TargetPlatform, TArray<FName>& OutFormats, bool bSupportDX11TextureFormats);
