Fixed more static analysis warning.

diff --git a/Engine/Source/Runtime/Engine/Private/Materials/Material.cpp b/Engine/Source/Runtime/Engine/Private/Materials/Material.cpp
--- a/Engine/Source/Runtime/Engine/Private/Materials/Material.cpp
+++ b/Engine/Source/Runtime/Engine/Private/Materials/Material.cpp
@@ -3316,6 +3316,7 @@ void UMaterial::PostLoad()
 	const int32 UE5MainVer = GetLinkerCustomVersion(FUE5MainStreamObjectVersion::GUID);
 
 	UMaterialEditorOnlyData* EditorOnly = GetEditorOnlyData();
+	check(EditorOnly != nullptr);
 
 	MoveExpressionInput(BaseColor_DEPRECATED, EditorOnly->BaseColor);
 	MoveExpressionInput(Metallic_DEPRECATED, EditorOnly->Metallic);
