Create NetworkSubsystem in every world

diff --git a/Engine/Source/Runtime/Engine/Public/Net/Subsystems/NetworkSubsystem.h b/Engine/Source/Runtime/Engine/Public/Net/Subsystems/NetworkSubsystem.h
--- a/Engine/Source/Runtime/Engine/Public/Net/Subsystems/NetworkSubsystem.h
+++ b/Engine/Source/Runtime/Engine/Public/Net/Subsystems/NetworkSubsystem.h
@@ -17,6 +17,11 @@ public:
 	/** Access the NetConditionGroupManager */
 	UE::Net::FNetConditionGroupManager& GetNetConditionGroupManager() { return GroupsManager;  }
 
+protected:
+	
+	/** Allow the subsystem on every world type. It's essentially free until it's used but accessed by any world that does netdriver replication. */
+	virtual bool DoesSupportWorldType(const EWorldType::Type WorldType) const override { return true; }
+
 private:
 
 	virtual void Serialize(FArchive& Ar) override;
