Strata - Fix for material blending mode override to AlphaTest to work when starta is enabled.
For that, we always need to propagate OpacityMask to the root node.
Virtual texture does not seem to complain anymore about OpacityMask being fed.

diff --git a/Engine/Source/Runtime/Engine/Private/Materials/Material.cpp b/Engine/Source/Runtime/Engine/Private/Materials/Material.cpp
--- a/Engine/Source/Runtime/Engine/Private/Materials/Material.cpp
+++ b/Engine/Source/Runtime/Engine/Private/Materials/Material.cpp
@@ -3001,13 +3001,10 @@ void UMaterial::ConvertMaterialToStrataMaterial()
 		// * Remove support for material attribute
 		// * explicitely connect the Strata node to the root node
 		// * Forward inputs to the root node (Do not reconnect the Opacity as we handle the opacity by internally within the conversion node)
-		// * Forward masked opacity only if blend mode is set to masked, as certain material (e.g., FlattenVT requires to have no OpacityMask plugged)
+		// * Always forward masked opacity because this is required when the blend mode is overriden to masked on a material instance.
 		bUseMaterialAttributes = false;
 		EditorOnly->FrontMaterial.Connect(0, ConvertNode);
-		if (BlendMode == BLEND_Masked)
-		{
-			EditorOnly->OpacityMask.Connect(7, BreakMatAtt);
-		}
+		EditorOnly->OpacityMask.Connect(7, BreakMatAtt);
 		EditorOnly->WorldPositionOffset.Connect(10, BreakMatAtt);
 		EditorOnly->AmbientOcclusion.Connect(14, BreakMatAtt);
 		EditorOnly->PixelDepthOffset.Connect(24, BreakMatAtt);
