SubsystemCollection:
- Add optional disabling of FGCObject collection
- New FWorldSubsystemCollection which disables FGCObject collection
- UWorld::AddReferencedObjects collects its FWorldSubsystemCollection references as the referencer

Fixes memory reference when deleting a Level in the content browser now initialized worlds (even Inactive ones) always initialize their Subsystems

diff --git a/Engine/Source/Runtime/Engine/Classes/Engine/World.h b/Engine/Source/Runtime/Engine/Classes/Engine/World.h
--- a/Engine/Source/Runtime/Engine/Classes/Engine/World.h
+++ b/Engine/Source/Runtime/Engine/Classes/Engine/World.h
@@ -1014,6 +1014,15 @@ private:
 	static void BroadcastWorldPartitionUninitialized(UWorld* InWorld, UWorldPartition* InWorldPartition);
 };
 
+/**
+ * World specific Subsystem collection which delegates object references to its owning UWorld
+ */
+class FWorldSubsystemCollection : public FSubsystemCollection<UWorldSubsystem>
+{
+public:
+	virtual bool IsGCObjectReferencer() const override { return false; }
+};
+
 /** 
  * The World is the top level object representing a map or a sandbox in which Actors and Components will exist and be rendered.  
  *
@@ -4126,7 +4135,7 @@ public:
 	FWorldPSCPool PSCPool;
 
 	//PSC Pooling END
-	FSubsystemCollection<UWorldSubsystem> SubsystemCollection;
+	FWorldSubsystemCollection SubsystemCollection;
 };
 
 /** Global UWorld pointer. Use of this pointer should be avoided whenever possible. */
diff --git a/Engine/Source/Runtime/Engine/Private/Subsystems/SubsystemCollection.cpp b/Engine/Source/Runtime/Engine/Private/Subsystems/SubsystemCollection.cpp
--- a/Engine/Source/Runtime/Engine/Private/Subsystems/SubsystemCollection.cpp
+++ b/Engine/Source/Runtime/Engine/Private/Subsystems/SubsystemCollection.cpp
@@ -217,7 +217,15 @@ USubsystem* FSubsystemCollectionBase::InitializeDependency(TSubclassOf<USubsyste
 
 void FSubsystemCollectionBase::AddReferencedObjects(FReferenceCollector& Collector)
 {
-	Collector.AddReferencedObjects(SubsystemMap);
+	if (IsGCObjectReferencer())
+	{
+		AddReferencedObjects(nullptr, Collector);
+	}
+}
+
+void FSubsystemCollectionBase::AddReferencedObjects(UObject* Referencer, FReferenceCollector& Collector)
+{
+	Collector.AddReferencedObjects(SubsystemMap, Referencer);
 }
 
 FString FSubsystemCollectionBase::GetReferencerName() const
diff --git a/Engine/Source/Runtime/Engine/Private/World.cpp b/Engine/Source/Runtime/Engine/Private/World.cpp
--- a/Engine/Source/Runtime/Engine/Private/World.cpp
+++ b/Engine/Source/Runtime/Engine/Private/World.cpp
@@ -571,6 +571,8 @@ void UWorld::AddReferencedObjects(UObject* InThis, FReferenceCollector& Collecto
 
 	This->StreamingLevelsToConsider.AddReferencedObjects(InThis, Collector);
 
+	This->SubsystemCollection.AddReferencedObjects(InThis, Collector);
+
 	Super::AddReferencedObjects( InThis, Collector );
 }
 
diff --git a/Engine/Source/Runtime/Engine/Public/Subsystems/SubsystemCollection.h b/Engine/Source/Runtime/Engine/Public/Subsystems/SubsystemCollection.h
--- a/Engine/Source/Runtime/Engine/Public/Subsystems/SubsystemCollection.h
+++ b/Engine/Source/Runtime/Engine/Public/Subsystems/SubsystemCollection.h
@@ -53,6 +53,12 @@ public:
 	 */
 	static void DeactivateExternalSubsystem(UClass* SubsystemClass);
 
+	/** Add referenced objects with referencer */
+	void AddReferencedObjects(UObject* Referencer, FReferenceCollector& Collector);
+
+	/** If True, this collection will add its references through its FGCObject::AddReferencedObjects interface implementation */
+	virtual bool IsGCObjectReferencer() const { return true; }
+
 protected:
 	/** protected constructor - for use by the template only(FSubsystemCollection<TBaseType>) */
 	FSubsystemCollectionBase(UClass* InBaseType);
