Exposed StaticMesh DDC build key outside of StaticMesh.cpp (private namespace) so that internal UE5 plugins can leverage it.

diff --git a/Engine/Source/Runtime/Engine/Classes/Engine/StaticMesh.h b/Engine/Source/Runtime/Engine/Classes/Engine/StaticMesh.h
--- a/Engine/Source/Runtime/Engine/Classes/Engine/StaticMesh.h
+++ b/Engine/Source/Runtime/Engine/Classes/Engine/StaticMesh.h
@@ -2063,4 +2063,11 @@ public:
 
 	UStaticMesh::FBuildParameters BuildParameters;
 	bool bHasRenderDataChanged = false;
-};
\ No newline at end of file
+};
+
+namespace UE::Private::StaticMesh
+{
+#if WITH_EDITOR
+	ENGINE_API FString BuildStaticMeshDerivedDataKey(const ITargetPlatform* TargetPlatform, UStaticMesh* Mesh, const FStaticMeshLODGroup& LODGroup);
+#endif
+}
diff --git a/Engine/Source/Runtime/Engine/Private/StaticMesh.cpp b/Engine/Source/Runtime/Engine/Private/StaticMesh.cpp
--- a/Engine/Source/Runtime/Engine/Private/StaticMesh.cpp
+++ b/Engine/Source/Runtime/Engine/Private/StaticMesh.cpp
@@ -2671,6 +2671,18 @@ static FString BuildStaticMeshDerivedDataKey(const FString& KeySuffix)
 		*KeySuffix);
 }
 
+namespace UE::Private::StaticMesh
+{
+
+#if WITH_EDITOR
+FString BuildStaticMeshDerivedDataKey(const ITargetPlatform* TargetPlatform, UStaticMesh* Mesh, const FStaticMeshLODGroup& LODGroup)
+{
+	return BuildStaticMeshDerivedDataKey(BuildStaticMeshDerivedDataKeySuffix(TargetPlatform, Mesh, LODGroup));
+}
+#endif
+
+} // UE::Private::StaticMesh
+
 static FString BuildStaticMeshLODDerivedDataKey(const FString& KeySuffix, int32 LODIdx)
 {
 	return FDerivedDataCacheInterface::BuildCacheKey(
