Fix potential duplicated actor guids when duplicating an actor.

#rnx

diff --git a/Engine/Source/Runtime/Engine/Private/Actor.cpp b/Engine/Source/Runtime/Engine/Private/Actor.cpp
--- a/Engine/Source/Runtime/Engine/Private/Actor.cpp
+++ b/Engine/Source/Runtime/Engine/Private/Actor.cpp
@@ -882,6 +882,10 @@ void AActor::Serialize(FArchive& Ar)
 		{
 			ActorGuid = FGuid::NewGuid();
 		}
+		else if ((Ar.GetPortFlags() & (PPF_Duplicate | PPF_DuplicateForPIE)) == PPF_Duplicate)
+		{
+			ActorGuid = FGuid::NewGuid();
+		}
 
 		if (!CanChangeIsSpatiallyLoadedFlag() && (Ar.CustomVer(FUE5ReleaseStreamObjectVersion::GUID) < FUE5ReleaseStreamObjectVersion::ActorGridPlacementDeprecateDefaultValueFixup))
 		{
