ActorChannel:
* Unify subobject list error checks
* Fix missing forward declaration

diff --git a/Engine/Source/Runtime/Engine/Classes/Engine/ActorChannel.h b/Engine/Source/Runtime/Engine/Classes/Engine/ActorChannel.h
--- a/Engine/Source/Runtime/Engine/Classes/Engine/ActorChannel.h
+++ b/Engine/Source/Runtime/Engine/Classes/Engine/ActorChannel.h
@@ -21,6 +21,7 @@ class FInBunch;
 class FNetFieldExportGroup;
 class FOutBunch;
 class UNetConnection;
+struct FObjectKey;
 
 namespace UE::Net { struct FSubObjectRegistry;  }
 
diff --git a/Engine/Source/Runtime/Engine/Private/DataChannel.cpp b/Engine/Source/Runtime/Engine/Private/DataChannel.cpp
--- a/Engine/Source/Runtime/Engine/Private/DataChannel.cpp
+++ b/Engine/Source/Runtime/Engine/Private/DataChannel.cpp
@@ -3630,6 +3630,8 @@ bool UActorChannel::ReplicateRegisteredSubObjects(FOutBunch& Bunch, FReplication
 		{
 			if (CanSubObjectReplicateToClient(SubObjectInfo.NetCondition, SubObjectInfo.Key, ConditionMap))
 			{
+				checkf(IsValid(SubObjectInfo.SubObject), TEXT("Found invalid subobject (%s) registered in %s"), *GetNameSafe(SubObjectInfo.SubObject), *Actor->GetName());
+
 				bWroteSomethingImportant |= WriteSubObjectInBunch(SubObjectInfo.SubObject, Bunch, RepFlags);
 			}
 		}
@@ -3638,12 +3640,12 @@ bool UActorChannel::ReplicateRegisteredSubObjects(FOutBunch& Bunch, FReplication
 	// Now the replicated actor components
 	for( const FReplicatedComponentInfo& RepComponentInfo : FSubObjectGetter::GetReplicatedComponents(Actor) )
 	{
-		check(IsValid(RepComponentInfo.Component));
-
 		if (CanSubObjectReplicateToClient(RepComponentInfo.NetCondition, RepComponentInfo.Key, ConditionMap))
 		{
 			UActorComponent* ReplicatedComponent = RepComponentInfo.Component;
 
+			checkf(IsValid(ReplicatedComponent), TEXT("Found invalid replicated component (%s) registered in %s"), *GetNameSafe(ReplicatedComponent), *Actor->GetName());
+
 			SetCurrentSubObjectOwner(ReplicatedComponent);
 			if (ReplicatedComponent->IsUsingRegisteredSubObjectList())
 			{
