Long save times in Take Recorder with timecode
#fix added missing variable name to AnimRecorder controller scoped bracket (meant that it wasn't acting as RAII)
#misc switched AttributeCurve::RemoveRedundant keys to just generate new keys array rather than removing individual keys

diff --git a/Engine/Source/Editor/SequenceRecorder/Private/AnimationRecorder.cpp b/Engine/Source/Editor/SequenceRecorder/Private/AnimationRecorder.cpp
--- a/Engine/Source/Editor/SequenceRecorder/Private/AnimationRecorder.cpp
+++ b/Engine/Source/Editor/SequenceRecorder/Private/AnimationRecorder.cpp
@@ -641,7 +641,7 @@ void FAnimationRecorder::ProcessRecordedTimes(UAnimSequence* AnimSequence, USkel
 	TArray<float> SlateTimes(&Times[0], 1);
 
 	IAnimationDataController& Controller = AnimSequence->GetController();
-	IAnimationDataController::FScopedBracket(Controller, LOCTEXT("AddTimeCodeAttributesBracket", "Adding Time Code attributes"));
+	IAnimationDataController::FScopedBracket ScopedBracket(Controller, LOCTEXT("AddTimeCodeAttributesBracket", "Adding Time Code attributes"));
 
 	// If the user defined bone doesn't exist, fallback to writing timecodes to the root
 	bool bHasUserDefinedBone = false;
diff --git a/Engine/Source/Runtime/Engine/Classes/Animation/AttributeCurve.cpp b/Engine/Source/Runtime/Engine/Classes/Animation/AttributeCurve.cpp
--- a/Engine/Source/Runtime/Engine/Classes/Animation/AttributeCurve.cpp
+++ b/Engine/Source/Runtime/Engine/Classes/Animation/AttributeCurve.cpp
@@ -296,7 +296,7 @@ FKeyHandle FAttributeCurve::FindKeyBeforeOrAt(float KeyTime) const
 
 void FAttributeCurve::RemoveRedundantKeys()
 {
-	TArray<FKeyHandle> KeyHandlesToRemove;
+	TSet<int32> KeyIndicesToRemove;
 	for (int32 KeyIndex = 0; KeyIndex < Keys.Num(); ++KeyIndex)
 	{
 		if (KeyIndex + 2 < Keys.Num())
@@ -308,14 +308,26 @@ void FAttributeCurve::RemoveRedundantKeys()
 			if (ScriptStruct->CompareScriptStruct(CurrentKey.Value.GetPtr<void>(), NextKeyOne.Value.GetPtr<void>(), 0)
 				&& ScriptStruct->CompareScriptStruct(NextKeyOne.Value.GetPtr<void>(), NextKeyTwo.Value.GetPtr<void>(), 0))
 			{
-				KeyHandlesToRemove.Add(GetKeyHandle(KeyIndex + 1));
+				KeyIndicesToRemove.Add(KeyIndex + 1);
 			}
 		}
 	}
 
-	for (const FKeyHandle& Handle : KeyHandlesToRemove)
+	if (KeyIndicesToRemove.Num())
 	{
-		DeleteKey(Handle);
+	    TArray<FAttributeKey> NewKeys;
+		NewKeys.Reserve(Keys.Num() - KeyIndicesToRemove.Num());
+	    for (int32 KeyIndex = 0; KeyIndex < Keys.Num(); ++KeyIndex)
+	    {
+		    if (!KeyIndicesToRemove.Contains(KeyIndex))
+		    {
+			    NewKeys.Add(Keys[KeyIndex]);
+		    }
+	    }
+    
+	    Swap(Keys, NewKeys);
+	    KeyHandlesToIndices.Empty(Keys.Num());
+	    KeyHandlesToIndices.SetKeyHandles(Keys.Num());
 	}
 
 	// If only two keys left and they are identical as well, remove the 2nd one.
