HLSLTranslator: Fixed ComponentMask node allowing out of bounds component access when the input is of LWCScalar type

diff --git a/Engine/Source/Runtime/Engine/Private/Materials/HLSLMaterialTranslator.cpp b/Engine/Source/Runtime/Engine/Private/Materials/HLSLMaterialTranslator.cpp
--- a/Engine/Source/Runtime/Engine/Private/Materials/HLSLMaterialTranslator.cpp
+++ b/Engine/Source/Runtime/Engine/Private/Materials/HLSLMaterialTranslator.cpp
@@ -8044,10 +8044,10 @@ int32 FHLSLMaterialTranslator::ComponentMask(int32 Vector,bool R,bool G,bool B,b
 
 	FString MaskString = FString::Printf(TEXT("%s%s%s%s"),
 		R ? TEXT("r") : TEXT(""),
-		// If VectorType is set to MCT_Float which means it could be any of the float types, assume it is a float1
-		G ? (VectorType == MCT_Float ? TEXT("r") : TEXT("g")) : TEXT(""),
-		B ? (VectorType == MCT_Float ? TEXT("r") : TEXT("b")) : TEXT(""),
-		A ? (VectorType == MCT_Float ? TEXT("r") : TEXT("a")) : TEXT("")
+		// If VectorType is set to MCT_Float/MCT_LWCScalar which means it could be any of the float types, assume it is a float1
+		G ? ((VectorType == MCT_Float || VectorType == MCT_LWCScalar) ? TEXT("r") : TEXT("g")) : TEXT(""),
+		B ? ((VectorType == MCT_Float || VectorType == MCT_LWCScalar) ? TEXT("r") : TEXT("b")) : TEXT(""),
+		A ? ((VectorType == MCT_Float || VectorType == MCT_LWCScalar) ? TEXT("r") : TEXT("a")) : TEXT("")
 		);
 
 	FString SourceString = GetParameterCode(Vector);
