nop just comments about how Texture2D Blueprint_GetSizeX can be wrong

diff --git a/Engine/Source/Runtime/Engine/Private/Texture2D.cpp b/Engine/Source/Runtime/Engine/Private/Texture2D.cpp
--- a/Engine/Source/Runtime/Engine/Private/Texture2D.cpp
+++ b/Engine/Source/Runtime/Engine/Private/Texture2D.cpp
@@ -1086,6 +1086,9 @@ int32 UTexture2D::Blueprint_GetSizeX() const
 	// In that state, the texture size is 0. Here we compute the resolution once cooked.
 	if (!GetSizeX())
 	{
+		//beware: this is wrong in a variety of ways
+		//	MaxTextureSize, PadForPow2, Downscale, etc. are not applied
+
 		const UTextureLODSettings* LODSettings = UDeviceProfileManager::Get().GetActiveProfile()->GetTextureLODSettings();
 		const int32 CookedLODBias = LODSettings->CalculateLODBias(Source.SizeX, Source.SizeY, MaxTextureSize, LODGroup, LODBias, 0, MipGenSettings, IsCurrentlyVirtualTextured());
 		return FMath::Max<int32>(Source.SizeX >> CookedLODBias, 1);
