AttributeRuntime: Added more built-in attributes, Vector3 and Quat, added unit test as well

diff --git a/Engine/Source/Runtime/Engine/Classes/Animation/BuiltInAttributeTypes.h b/Engine/Source/Runtime/Engine/Classes/Animation/BuiltInAttributeTypes.h
--- a/Engine/Source/Runtime/Engine/Classes/Animation/BuiltInAttributeTypes.h
+++ b/Engine/Source/Runtime/Engine/Classes/Animation/BuiltInAttributeTypes.h
@@ -87,7 +87,7 @@ struct FStringAnimationAttribute
 	FString Value;
 };
 
-/** Attribute type supporting the legacy TVariant<FString> atttributes */
+/** Attribute type supporting the legacy TVariant<FTransform> atttributes */
 USTRUCT(BlueprintType)
 struct FTransformAnimationAttribute
 {
@@ -146,6 +146,125 @@ struct FTransformAnimationAttribute
 	}
 };
 
+
+USTRUCT(BlueprintType)
+struct FVectorAnimationAttribute
+{
+	GENERATED_BODY()
+
+	UPROPERTY(EditAnywhere, BlueprintReadWrite, Category=TransformAnimationAttribute)
+	FVector Value;
+
+	void Accumulate(const FVectorAnimationAttribute& Attribute, float Weight, EAdditiveAnimationType AdditiveType)
+	{
+		// if (FAnimWeight::IsRelevant(Weight))
+		{
+			Value += Attribute.Value * Weight;
+		}
+	}
+
+	void MakeAdditive(const FVectorAnimationAttribute& BaseAttribute)
+	{
+		Value = Value - BaseAttribute.Value;
+	}
+
+	void Normalize()
+	{
+		Value.Normalize();
+	}
+	
+	FVectorAnimationAttribute Multiply(const float Weight) const
+	{
+		FVectorAnimationAttribute Out;
+		Out.Value = Value * Weight;
+
+		return Out;
+	}
+
+	void Interpolate(const FVectorAnimationAttribute& Attribute, float Alpha)
+	{
+		Value = FMath::Lerp<FVector>(Value, Attribute.Value, Alpha);
+	}
+};
+
+USTRUCT(BlueprintType)
+struct FQuaternionAnimationAttribute
+{
+	GENERATED_BODY()
+
+	FQuaternionAnimationAttribute() : Value(FQuat::Identity) {}
+	FQuaternionAnimationAttribute(const FQuat& InQuat) : Value(InQuat) {}
+
+	UPROPERTY(EditAnywhere, BlueprintReadWrite, Category=TransformAnimationAttribute)
+	FQuat Value;
+
+	void Accumulate(const FQuaternionAnimationAttribute& Attribute, float Weight, EAdditiveAnimationType AdditiveType)
+	{
+		// if (FAnimWeight::IsRelevant(Weight))
+		{
+			if (AdditiveType == AAT_None)
+			{
+				
+				const FQuat WeightedRotation = Attribute.Value * Weight;
+				
+				// From VectorAccumulateQuaternionShortestPath
+				// Blend rotation
+				//     To ensure the 'shortest route', we make sure the dot product between the both rotations is positive.
+				//     const float Bias = (|A.B| >= 0 ? 1 : -1)
+				//     return A + B * Bias;
+				const FQuat::FReal DotResult = (Value | WeightedRotation);
+				const FQuat::FReal Bias = FMath::FloatSelect(DotResult, FQuat::FReal(1.0f), FQuat::FReal(-1.0f));
+				
+				Value += WeightedRotation * Bias;
+			}
+			else
+			{
+				// Quaternion equivalent of FTransform::BlendFromIdentityAndAccumulate
+				const FQuat WeightedRotation = FQuat::FastLerp(FQuat::Identity, Attribute.Value, Weight).GetNormalized();
+				Value = WeightedRotation * Value;
+			}
+			
+		}
+	}
+
+	void MakeAdditive(const FQuaternionAnimationAttribute& BaseAttribute)
+	{
+		Value = Value * BaseAttribute.Value.Inverse();
+		Value.Normalize();
+	}
+
+	void Normalize()
+	{
+		Value.Normalize();
+	}
+	
+	FQuaternionAnimationAttribute Multiply(const float Weight) const
+	{
+		FQuaternionAnimationAttribute Out;
+		Out.Value = Value * Weight;
+
+		return Out;
+	}
+
+	void Interpolate(const FQuaternionAnimationAttribute& Attribute, float Alpha)
+	{
+		Value = FQuat::FastLerp(Value, Attribute.Value, Alpha).GetNormalized();
+		Value.Normalize();
+	}
+};
+
+USTRUCT()
+struct FNonBlendableQuaternionAnimationAttribute : public FQuaternionAnimationAttribute
+{
+	GENERATED_BODY()
+};
+
+USTRUCT()
+struct FNonBlendableVectorAnimationAttribute : public FVectorAnimationAttribute
+{
+	GENERATED_BODY()
+};
+
 USTRUCT()
 struct FNonBlendableTransformAnimationAttribute : public FTransformAnimationAttribute
 {
@@ -198,6 +317,36 @@ namespace UE
 			};
 		};
 
+		/** Quaternion attribute requires normalization */
+		template<>
+		struct TAttributeTypeTraits<FQuaternionAnimationAttribute> : public TAttributeTypeTraitsBase<FQuaternionAnimationAttribute>
+		{
+			enum
+			{
+				RequiresNormalization = true,
+			};
+		};
+
+		/** Non blendable types*/
+		template<>
+		struct TAttributeTypeTraits<FNonBlendableQuaternionAnimationAttribute> : public TAttributeTypeTraitsBase<FNonBlendableQuaternionAnimationAttribute>
+		{
+			enum
+			{
+				IsBlendable = false,
+			};
+		};
+
+		/** Non blendable types*/
+		template<>
+		struct TAttributeTypeTraits<FNonBlendableVectorAnimationAttribute> : public TAttributeTypeTraitsBase<FNonBlendableVectorAnimationAttribute>
+		{
+			enum
+			{
+				IsBlendable = false,
+			};
+		};
+		
 		/** Non blendable types*/
 		template<>
 		struct TAttributeTypeTraits<FNonBlendableTransformAnimationAttribute> : public TAttributeTypeTraitsBase<FNonBlendableTransformAnimationAttribute>
diff --git a/Engine/Source/Runtime/Engine/Private/Animation/AttributeTypes.cpp b/Engine/Source/Runtime/Engine/Private/Animation/AttributeTypes.cpp
--- a/Engine/Source/Runtime/Engine/Private/Animation/AttributeTypes.cpp
+++ b/Engine/Source/Runtime/Engine/Private/Animation/AttributeTypes.cpp
@@ -28,6 +28,8 @@ namespace UE
 			RegisterType<FIntegerAnimationAttribute>();
 			RegisterType<FStringAnimationAttribute>();
 			RegisterType<FTransformAnimationAttribute>();
+			RegisterType<FVectorAnimationAttribute>();
+			RegisterType<FQuaternionAnimationAttribute>();
 		}
 		
 		static FDelayedAutoRegisterHelper DelayedAttributeTypesInitializationHelper(EDelayedRegisterRunPhase::ObjectSystemReady, []()
