Fix for a regression which no longer preserves the non-native SCS root transform on a serialized Actor BP instanced root component after adding a DSO root to a native parent class.

This was previously a side effect of validating instanced property values during PostLoadSubobjects(), but a recent change (17730206) will now bypass this step if the current value has already been serialized.

diff --git a/Engine/Source/Runtime/Engine/Private/Actor.cpp b/Engine/Source/Runtime/Engine/Private/Actor.cpp
--- a/Engine/Source/Runtime/Engine/Private/Actor.cpp
+++ b/Engine/Source/Runtime/Engine/Private/Actor.cpp
@@ -988,6 +988,23 @@ void AActor::PostLoadSubobjects(FObjectInstancingGraph* OuterInstanceGraph)
 
 	ResetOwnedComponents();
 
+	// Redirect the root component away from a Blueprint-added instance if the native ctor now constructs a default subobject as the root.
+	if (RootComponent && !RootComponent->HasAnyFlags(RF_DefaultSubObject))
+	{
+		// Get the native class default object from which non-native classes (i.e. Blueprint classes) should inherit their root component.
+		const AActor* ActorCDO = GetClass()->GetDefaultObject<AActor>();
+		checkSlow(ActorCDO);
+
+		// Ensure the root component references the appropriate instance; non-native classes (e.g. Blueprints) should always inherit from the CDO.
+		if (const USceneComponent* ActorCDO_RootComponent = ActorCDO->GetRootComponent())
+		{
+			if (USceneComponent* ActorInstance_NativeRootComponent = Cast<USceneComponent>(GetDefaultSubobjectByName(ActorCDO_RootComponent->GetFName())))
+			{
+				SetRootComponent(ActorInstance_NativeRootComponent);
+			}
+		}
+	}
+
 	if (RootComponent && bHadRoot && OldRoot != RootComponent && OldRoot->IsIn(this))
 	{
 		UE_LOG(LogActor, Log, TEXT("Root component has changed, relocating new root component to old position %s->%s"), *OldRoot->GetFullName(), *GetRootComponent()->GetFullName());
