Attribute Runtime: Fixed Uninitialized member variable

diff --git a/Engine/Source/Runtime/Engine/Classes/Animation/BuiltInAttributeTypes.h b/Engine/Source/Runtime/Engine/Classes/Animation/BuiltInAttributeTypes.h
--- a/Engine/Source/Runtime/Engine/Classes/Animation/BuiltInAttributeTypes.h
+++ b/Engine/Source/Runtime/Engine/Classes/Animation/BuiltInAttributeTypes.h
@@ -152,8 +152,8 @@ struct FVectorAnimationAttribute
 {
 	GENERATED_BODY()
 
-	UPROPERTY(EditAnywhere, BlueprintReadWrite, Category=TransformAnimationAttribute)
-	FVector Value;
+	UPROPERTY(EditAnywhere, BlueprintReadWrite, Category=VectorAnimationAttribute)
+	FVector Value = FVector::ZeroVector;
 
 	void Accumulate(const FVectorAnimationAttribute& Attribute, float Weight, EAdditiveAnimationType AdditiveType)
 	{
@@ -192,11 +192,8 @@ struct FQuaternionAnimationAttribute
 {
 	GENERATED_BODY()
 
-	FQuaternionAnimationAttribute() : Value(FQuat::Identity) {}
-	FQuaternionAnimationAttribute(const FQuat& InQuat) : Value(InQuat) {}
-
-	UPROPERTY(EditAnywhere, BlueprintReadWrite, Category=TransformAnimationAttribute)
-	FQuat Value;
+	UPROPERTY(EditAnywhere, BlueprintReadWrite, Category=QuaternionAnimationAttribute)
+	FQuat Value = FQuat::Identity;
 
 	void Accumulate(const FQuaternionAnimationAttribute& Attribute, float Weight, EAdditiveAnimationType AdditiveType)
 	{
