[Navmesh] Fix Foliage instance static meshes do not affecting the navmesh

diff --git a/Engine/Source/Runtime/Engine/Private/HierarchicalInstancedStaticMesh.cpp b/Engine/Source/Runtime/Engine/Private/HierarchicalInstancedStaticMesh.cpp
--- a/Engine/Source/Runtime/Engine/Private/HierarchicalInstancedStaticMesh.cpp
+++ b/Engine/Source/Runtime/Engine/Private/HierarchicalInstancedStaticMesh.cpp
@@ -35,6 +35,7 @@
 #include "PrimitiveSceneInfo.h"
 #include "NaniteSceneProxy.h"
 #include "HierarchicalStaticMeshSceneProxy.h"
+#include "VisualLogger/VisualLogger.h"
 
 #if WITH_EDITOR
 static float GDebugBuildTreeAsyncDelayInSeconds = 0.f;
@@ -3277,8 +3278,14 @@ static void GatherInstanceTransformsInArea(const UHierarchicalInstancedStaticMes
 	if (ClusterTree.Num())
 	{
 		const FClusterNode& ChildNode = ClusterTree[Child];
-		const FBox WorldNodeBox = FBox(ChildNode.BoundMin, ChildNode.BoundMax).TransformBy(Component.GetComponentTransform());
-	
+
+		const FTransform ToWorldTransform = FTransform(Component.TranslatedInstanceSpaceOrigin) * Component.GetComponentTransform();
+		const FBox WorldNodeBox = FBox(ChildNode.BoundMin, ChildNode.BoundMax).TransformBy(ToWorldTransform);
+
+		UE_VLOG_BOX(&Component, LogStaticMesh, Verbose, FBox(ChildNode.BoundMin, ChildNode.BoundMax), FColor::Red, TEXT("LocalNodeBox"));
+		UE_VLOG_BOX(&Component, LogStaticMesh, Verbose, WorldNodeBox, FColor::Green, TEXT("WorldNodeBox"));
+		UE_VLOG_BOX(&Component, LogStaticMesh, Verbose, AreaBox, FColor::Blue, TEXT("AreaBox"));
+
 		if (AreaBox.Intersect(WorldNodeBox))
 		{
 			if (ChildNode.FirstChild < 0 || AreaBox.IsInside(WorldNodeBox))
@@ -3456,7 +3463,9 @@ static void GatherInstancesOverlappingArea(const UHierarchicalInstancedStaticMes
 {
 	const TArray<FClusterNode>& ClusterTree = *Component.ClusterTreePtr;
 	const FClusterNode& ChildNode = ClusterTree[Child];
-	const FBox WorldNodeBox = FBox(ChildNode.BoundMin, ChildNode.BoundMax).TransformBy(Component.GetComponentTransform());
+
+	const FTransform ToWorldTransform = FTransform(Component.TranslatedInstanceSpaceOrigin) * Component.GetComponentTransform();
+	const FBox WorldNodeBox = FBox(ChildNode.BoundMin, ChildNode.BoundMax).TransformBy(ToWorldTransform);
 
 	if (AreaBox.Intersect(WorldNodeBox))
 	{
