Strata - material conversion - Fixed some material having weird behavior due to bac caching with the new EditiorOnlyData. We now follow the proper way to UpdateCachedExpressionData.

diff --git a/Engine/Shaders/Public/ShaderVersion.ush b/Engine/Shaders/Public/ShaderVersion.ush
--- a/Engine/Shaders/Public/ShaderVersion.ush
+++ b/Engine/Shaders/Public/ShaderVersion.ush
@@ -2,4 +2,4 @@
 // Each time the console command is executed it generates a new GUID. As this file is included
 // in Platform.ush (which should be included in any shader) it allows to invalidate the shader DDC.
 // 
-// GUID = 7CA85F68-C697-4626-B292-ABF0E15D1B7E
+// GUID = 8EA7B5C8-3F33-4EF6-922F-F4E242BE7AD3
diff --git a/Engine/Source/Runtime/Engine/Private/Materials/Material.cpp b/Engine/Source/Runtime/Engine/Private/Materials/Material.cpp
--- a/Engine/Source/Runtime/Engine/Private/Materials/Material.cpp
+++ b/Engine/Source/Runtime/Engine/Private/Materials/Material.cpp
@@ -3247,23 +3247,7 @@ void UMaterial::ConvertMaterialToStrataMaterial()
 	if (bInvalidateShader)
 	{
 		// We might have moved connections above so update the CachedExpressionData from the EditorOnly connection data (ground truth).
-		if (!bUseMaterialAttributes)
-		{
-			// Reset the connection bitmask so we can set only the valid ones below.
-			CachedExpressionData->PropertyConnectedBitmask = 0;
-
-			for (int32 PropertyIndex = 0; PropertyIndex < MP_MAX; ++PropertyIndex)
-			{
-				const EMaterialProperty Property = (EMaterialProperty)PropertyIndex;
-				const FExpressionInput* Input = GetExpressionInputForProperty(Property);
-				if (Input && Input->IsConnected())
-				{
-					CachedExpressionData->SetPropertyConnected(Property);
-				}
-			}
-		}
-
-		CachedExpressionData->Validate();
+		UpdateCachedExpressionData();
 
 		// Make sure our cached data is consistent
 		ensureMsgf(!(GetEditorOnlyData()->Refraction.IsConnected() && !GetCachedExpressionData().IsPropertyConnected(MP_Refraction)),
