More  accurate comments on level visibility since the old comments were misleading.

diff --git a/Engine/Source/Runtime/Engine/Classes/Engine/Level.h b/Engine/Source/Runtime/Engine/Classes/Engine/Level.h
--- a/Engine/Source/Runtime/Engine/Classes/Engine/Level.h
+++ b/Engine/Source/Runtime/Engine/Classes/Engine/Level.h
@@ -622,7 +622,10 @@ public:
 	UPROPERTY(Transient, DuplicateTransient, NonTransactional)
 	uint8										bStaticComponentsRegisteredInStreamingManager: 1;
 
-	/** Whether the level is currently visible/ associated with the world */
+	/** 
+	 * Whether the level is currently visible/ associated with the world. 
+	 * If false, may not yet be fully removed from the world.
+	 */
 	UPROPERTY(transient)
 	uint8										bIsVisible:1;
 
diff --git a/Engine/Source/Runtime/Engine/Classes/Engine/LevelStreaming.h b/Engine/Source/Runtime/Engine/Classes/Engine/LevelStreaming.h
--- a/Engine/Source/Runtime/Engine/Classes/Engine/LevelStreaming.h
+++ b/Engine/Source/Runtime/Engine/Classes/Engine/LevelStreaming.h
@@ -514,11 +514,11 @@ public:
 	UPROPERTY(BlueprintAssignable)
 	FLevelStreamingLoadedStatus			OnLevelUnloaded;
 	
-	/** Called when level is added to the world  */
+	/** Called when level is added to the world and is visible  */
 	UPROPERTY(BlueprintAssignable)
 	FLevelStreamingVisibilityStatus		OnLevelShown;
 	
-	/** Called when level is removed from the world  */
+	/** Called when level is no longer visible, may not be removed from world yet  */
 	UPROPERTY(BlueprintAssignable)
 	FLevelStreamingVisibilityStatus		OnLevelHidden;
 
