PR #9158: Expose UCurveBase to Blueprint (Contributed by [REDACTED_PERSON])

diff --git a/Engine/Source/Runtime/Engine/Classes/Curves/CurveBase.h b/Engine/Source/Runtime/Engine/Classes/Curves/CurveBase.h
--- a/Engine/Source/Runtime/Engine/Classes/Curves/CurveBase.h
+++ b/Engine/Source/Runtime/Engine/Classes/Curves/CurveBase.h
@@ -13,7 +13,7 @@
 /**
  * Defines a curve of interpolated points to evaluate over a given range
  */
-UCLASS(abstract)
+UCLASS(abstract, BlueprintType)
 class ENGINE_API UCurveBase
 	: public UObject
 	, public FCurveOwnerInterface
