Adding async implementation of new GetSaveNames function.

#rnx

diff --git a/Engine/Source/Runtime/Engine/Private/SaveGameSystem.cpp b/Engine/Source/Runtime/Engine/Private/SaveGameSystem.cpp
--- a/Engine/Source/Runtime/Engine/Private/SaveGameSystem.cpp
+++ b/Engine/Source/Runtime/Engine/Private/SaveGameSystem.cpp
@@ -117,6 +117,32 @@ void ISaveGameSystem::DeleteGameAsync(bool bAttemptToUseUI, const TCHAR* Name, F
 	);
 }
 
+void ISaveGameSystem::GetSaveGameNamesAsync(FPlatformUserId PlatformUserId, FSaveGameAsyncGetNamesCallback Callback)
+{
+	// start the delete operation on a background thread.
+	AsyncTaskPipe.Launch(UE_SOURCE_LOCATION,
+		[this, PlatformUserId, Callback]()
+		{
+			// get the list of savegames
+			TArray<FString> FoundSaves;
+			int32 UserIndex = FPlatformMisc::GetUserIndexForPlatformUser(PlatformUserId);
+			const bool bResult = GetSaveGameNames(FoundSaves, UserIndex);
+
+			// trigger the callback on the game thread
+			if (Callback)
+			{
+				OnAsyncComplete(
+					[PlatformUserId, bResult, FoundSaves = MoveTemp(FoundSaves), Callback]()
+					{
+						Callback(PlatformUserId, bResult, FoundSaves);
+					}
+				);
+			}
+		}
+	);
+
+}
+
 
 void ISaveGameSystem::InitAsync(bool bAttemptToUseUI, FPlatformUserId PlatformUserId, FSaveGameAsyncInitCompleteCallback Callback)
 {
diff --git a/Engine/Source/Runtime/Engine/Public/SaveGameSystem.h b/Engine/Source/Runtime/Engine/Public/SaveGameSystem.h
--- a/Engine/Source/Runtime/Engine/Public/SaveGameSystem.h
+++ b/Engine/Source/Runtime/Engine/Public/SaveGameSystem.h
@@ -41,7 +41,7 @@ public:
 	/** Similar to DoesSaveGameExist, except returns a result code with more information. */
 	virtual ESaveExistsResult DoesSaveGameExistWithResult(const TCHAR* Name, const int32 UserIndex) = 0;
 
-	/** Gets a list of all known saves.  This isn't possible on all platforms,  */
+	/** Gets a list of all known saves.  This isn't possible on all platforms. */
 	virtual bool GetSaveGameNames(TArray<FString>& FoundSaves, const int32 UserIndex)
 	{
 		return false;
@@ -69,6 +69,9 @@ public:
 	// indicates that an async savegame initialize has completed. PlatformUserId, bool=success. Should always be called from the gamethread.
 	typedef TFunction<void(FPlatformUserId, bool)> FSaveGameAsyncInitCompleteCallback;
 
+	// indicates that an async savegame list names has completed. PlatformUserId, bool=success, FoundSaves. Should always be called from the game thread
+	typedef TFunction<void(FPlatformUserId, bool, const TArray<FString>&)> FSaveGameAsyncGetNamesCallback;
+
 
 	/* Asyncronously checks if the named savegame exists. Default implementation runs blocking version on a background thread */
 	virtual void DoesSaveGameExistAsync(const TCHAR* Name, FPlatformUserId PlatformUserId, FSaveGameAsyncExistsCallback Callback);
@@ -82,6 +85,9 @@ public:
 	/* Delete the named savegame asyncronously. Default implementation runs blocking version on a background thread */
 	virtual void DeleteGameAsync(bool bAttemptToUseUI, const TCHAR* Name, FPlatformUserId PlatformUserId, FSaveGameAsyncOpCompleteCallback Callback);
 
+	/* Gets a list of all known saves if the platform supports it. Default implementatino runs blocking version on a background thread */
+	virtual void GetSaveGameNamesAsync(FPlatformUserId PlatformUserId, FSaveGameAsyncGetNamesCallback Callback);
+
 	/* (Optional) initialise the save system for the given user. Useful if the platform may display UI on first use - this can be called as part of the user login flow, for example. 
 	   If this is not used, any UI may be displayed on the first use of the other save api functions */
 	virtual void InitAsync(bool bAttemptToUseUI, FPlatformUserId PlatformUserId, FSaveGameAsyncInitCompleteCallback Callback);
