Moved RichCurve evaluation code from RichCurve.cpp to separate file and consolidated all instances of copy-pasta behaviour

diff --git a/Engine/Plugins/Enterprise/DatasmithFBXImporter/Source/DatasmithDeltaGenTranslator/Private/DatasmithDeltaGenAnimationInterpolator.cpp b/Engine/Plugins/Enterprise/DatasmithFBXImporter/Source/DatasmithDeltaGenTranslator/Private/DatasmithDeltaGenAnimationInterpolator.cpp
--- a/Engine/Plugins/Enterprise/DatasmithFBXImporter/Source/DatasmithDeltaGenTranslator/Private/DatasmithDeltaGenAnimationInterpolator.cpp
+++ b/Engine/Plugins/Enterprise/DatasmithFBXImporter/Source/DatasmithDeltaGenTranslator/Private/DatasmithDeltaGenAnimationInterpolator.cpp
@@ -3,6 +3,7 @@
 #include "DatasmithDeltaGenAnimationInterpolator.h"
 
 #include "Algo/BinarySearch.h"
+#include "Curves/CurveEvaluation.h"
 
 namespace DeltaGen
 {
@@ -30,104 +31,6 @@ namespace DeltaGen
 		OutIndex2 = Index2;
 	}
 
-	/**
-	 * Copied from FRichCurve
-	 * Assuming that P0, P1, P2 and P3 are sequential control points of an N=4 Bezier curve, returns
-	 * the interpolated value for interpolation constant Alpha in [0, 1]
-	 */
-	FVector BezierInterp(const FVector& P0, const FVector& P1, const FVector& P2, const FVector& P3, float Alpha)
-	{
-		const FVector P01 = FMath::Lerp(P0, P1, Alpha);
-		const FVector P12 = FMath::Lerp(P1, P2, Alpha);
-		const FVector P23 = FMath::Lerp(P2, P3, Alpha);
-		const FVector P012 = FMath::Lerp(P01, P12, Alpha);
-		const FVector P123 = FMath::Lerp(P12, P23, Alpha);
-		const FVector P0123 = FMath::Lerp(P012, P123, Alpha);
-
-		return P0123;
-	}
-
-	/**
-	 * Copied from MovieSceneFloatChannel
-	 * Solve Cubic Equation using Cardano's forumla
-	 * Adopted from Graphic Gems 1
-	 * [REDACTED_PROFILE_URL]/GraphicsGems/blob/master/gems/Roots3And4.c
-	 *
-	 * @param Coeff Coefficient parameters of form  Coeff[0] + Coeff[1]*x + Coeff[2]*x^2 + Coeff[3]*x^3 = 0
-	 * @param Solution Up to 3 real solutions. We don't include imaginary solutions, as would need a complex number object
-	 * @return Returns the number of real solutions returned in the Solution array.
-	 */
-	int32 SolveCubic(double Coeff[4], double Solution[3])
-	{
-		auto Cbrt = [](double x) -> double
-		{
-			return ((x) > 0.0 ? pow((x), 1.0 / 3.0) : ((x) < 0.0 ? -pow((double)-(x), 1.0 / 3.0) : 0.0));
-		};
-		int32 NumSolutions = 0;
-
-		// Normal form: x^3 + Ax^2 + Bx + C = 0
-		double Denominator = (Coeff[3] != 0.0) ? Coeff[3] : SMALL_NUMBER * SMALL_NUMBER;
-		double A = Coeff[2] / Denominator;
-		double B = Coeff[1] / Denominator;
-		double C = Coeff[0] / Denominator;
-
-		// Substitute x = y - A/3 to eliminate quadric term: x^3 +px + q = 0
-		double SqOfA = A * A;
-		double P = 1.0 / 3 * (-1.0 / 3 * SqOfA + B);
-		double Q = 1.0 / 2 * (2.0 / 27 * A * SqOfA - 1.0 / 3 * A * B + C);
-
-		// Use Cardano's formula
-		double CubeOfP = P * P * P;
-		double D = Q * Q + CubeOfP;
-		if (FMath::IsNearlyZero(D))
-		{
-			// One triple solution
-			if (FMath::IsNearlyZero(Q))
-			{
-				Solution[0] = 0;
-				NumSolutions = 1;
-			}
-			// One single and one double solution
-			else
-			{
-				double U = Cbrt(-Q);
-				Solution[0] = 2 * U;
-				Solution[1] = -U;
-				NumSolutions = 2;
-			}
-		}
-		// Three real solutions
-		else if (D < 0)
-		{
-			double Phi = 1.0 / 3 * acos(-Q / sqrt(-CubeOfP));
-			double T = 2 * sqrt(-P);
-
-			Solution[0] = T * cos(Phi);
-			Solution[1] = -T * cos(Phi + UE_DOUBLE_PI / 3);
-			Solution[2] = -T * cos(Phi - UE_DOUBLE_PI / 3);
-			NumSolutions = 3;
-		}
-		// One real solution
-		else
-		{
-			double SqrtD = sqrt(D);
-			double U = Cbrt(SqrtD - Q);
-			double V = -Cbrt(SqrtD + Q);
-
-			Solution[0] = U + V;
-			NumSolutions = 1;
-		}
-
-		// Resubstitute
-		double Sub = 1.0 / 3 * A;
-		for (int32 Index = 0; Index < NumSolutions; ++Index)
-		{
-			Solution[Index] -= Sub;
-		}
-
-		return NumSolutions;
-	}
-
 	FInterpolator::FInterpolator(const TArray<float>& InTimes, const TArray<FVector>& InValues)
 		: Times(InTimes)
 		, Values(InValues)
@@ -271,7 +174,7 @@ namespace DeltaGen
 		const FVector& ArriveControlPoint = Index2 > 0 ?               Values[3 * Index2 - 1] : Values[3 * Index2];
 
 		// Index1 * 3 and Index2 * 3 because Values stores arrive/leave control points too
-		return BezierInterp(Values[Index1 * 3],
+		return UE::Curves::BezierInterp(Values[Index1 * 3],
 							LeaveControlPoint,
 							ArriveControlPoint,
 							Values[Index2 * 3],
@@ -340,7 +243,7 @@ namespace DeltaGen
 		// Solve cubic polynomial with Gerolano Cardano's formula
 		double Coefs[4] = { CoefConst, CoefT1, CoefT2, CoefT3 };
 		double Solutions[3] = {-1., -1., -1.};
-		int32 NumSolutions = SolveCubic(Coefs, Solutions);
+		int32 NumSolutions = UE::Curves::SolveCubic(Coefs, Solutions);
 
 		// The target solution, if it exists, is the only real one within [0, 1]
 		float TargetAlpha = FLT_MAX;
diff --git a/Engine/Plugins/VirtualProduction/CameraCalibrationCore/Source/CameraCalibrationCore/Private/LensFile.cpp b/Engine/Plugins/VirtualProduction/CameraCalibrationCore/Source/CameraCalibrationCore/Private/LensFile.cpp
--- a/Engine/Plugins/VirtualProduction/CameraCalibrationCore/Source/CameraCalibrationCore/Private/LensFile.cpp
+++ b/Engine/Plugins/VirtualProduction/CameraCalibrationCore/Source/CameraCalibrationCore/Private/LensFile.cpp
@@ -14,6 +14,7 @@
 #include "LensDistortionModelHandlerBase.h"
 #include "LensFileRendering.h"
 #include "LensInterpolationUtils.h"
+#include "Curves/CurveEvaluation.h"
 #include "Models/SphericalLensModel.h"
 #include "Tables/BaseLensTable.h"
 #include "Tables/LensTableUtils.h"
@@ -45,19 +46,6 @@ namespace LensFileUtils
 		return NewRenderTarget2D;
 	}
 
-	/** From FRichCurve - Util to find float value on bezier defined by 4 control points */
-	float BezierInterp(float P0, float P1, float P2, float P3, float Alpha)
-	{
-		const float P01 = FMath::Lerp(P0, P1, Alpha);
-		const float P12 = FMath::Lerp(P1, P2, Alpha);
-		const float P23 = FMath::Lerp(P2, P3, Alpha);
-		const float P012 = FMath::Lerp(P01, P12, Alpha);
-		const float P123 = FMath::Lerp(P12, P23, Alpha);
-		const float P0123 = FMath::Lerp(P012, P123, Alpha);
-
-		return P0123;
-	}
-
 	float EvalAtTwoPoints(const float EvalTime, const float Time0, const float Time1, const float Value0, const float Value1, const float Tangent0, const float Tangent1)
 	{
 		if (FMath::IsNearlyEqual(Time0, Time1))
@@ -79,7 +67,7 @@ namespace LensFileUtils
 		const float P3 = Value1;
 		const float P1 = P0 + (CurveTan0 * CurveDiff * OneThird);
 		const float P2 = P3 - (CurveTan1 * CurveDiff * OneThird);
-		return BezierInterp(P0, P1, P2, P3, CurveAlpha);
+		return UE::Curves::BezierInterp(P0, P1, P2, P3, CurveAlpha);
 	}
 
 	void FindWeightsAndInterp(float InEvalTime, TConstArrayView<float> InTimes, TConstArrayView<float> InTangents, TOptional<float> LerpFactor, TConstArrayView<float> Inputs, float& Output)
diff --git a/Engine/Source/Editor/CurveEditor/Private/Filters/CurveEditorReduceFilter.cpp b/Engine/Source/Editor/CurveEditor/Private/Filters/CurveEditorReduceFilter.cpp
--- a/Engine/Source/Editor/CurveEditor/Private/Filters/CurveEditorReduceFilter.cpp
+++ b/Engine/Source/Editor/CurveEditor/Private/Filters/CurveEditorReduceFilter.cpp
@@ -9,23 +9,7 @@
 #include "CurveModel.h"
 #include "Templates/SharedPointer.h"
 #include "CurveEditorSnapMetrics.h"
-
-/**
-The following key reduction is the same as that found in FRichCurve.
-It would be nice if there was just one implementation of the reduction (and) baking algorithms.
-*/
-/** Util to find float value on bezier defined by 4 control points */
-static float BezierInterp(float P0, float P1, float P2, float P3, float Alpha)
-{
-	const float P01 = FMath::Lerp(P0, P1, Alpha);
-	const float P12 = FMath::Lerp(P1, P2, Alpha);
-	const float P23 = FMath::Lerp(P2, P3, Alpha);
-	const float P012 = FMath::Lerp(P01, P12, Alpha);
-	const float P123 = FMath::Lerp(P12, P23, Alpha);
-	const float P0123 = FMath::Lerp(P012, P123, Alpha);
-
-	return P0123;
-}
+#include "Curves/CurveEvaluation.h"
 
 static float EvalForTwoKeys(const FKeyPosition& Key1Pos, const FKeyAttributes& Key1Attrib,
 	const FKeyPosition& Key2Pos, const FKeyAttributes& Key2Attrib,
@@ -49,7 +33,7 @@ static float EvalForTwoKeys(const FKeyPosition& Key1Pos, const FKeyAttributes& K
 			const float P1 = Key1Attrib.HasLeaveTangent() ? P0 + (Key1Attrib.GetLeaveTangent() * Diff*OneThird) : P0;
 			const float P2 = Key2Attrib.HasArriveTangent() ? P3 - (Key2Attrib.GetArriveTangent() * Diff*OneThird) : P3;
 
-			return BezierInterp(P0, P1, P2, P3, Alpha);
+			return UE::Curves::BezierInterp(P0, P1, P2, P3, Alpha);
 		}
 	}
 	else
diff --git a/Engine/Source/Runtime/Engine/Private/Curves/CurveEvaluation.cpp b/Engine/Source/Runtime/Engine/Private/Curves/CurveEvaluation.cpp
new file mode 100644
--- /dev/null
+++ b/Engine/Source/Runtime/Engine/Private/Curves/CurveEvaluation.cpp
@@ -0,0 +1,245 @@
+// Copyright Epic Games, Inc. All Rights Reserved.
+
+#include "Curves/CurveEvaluation.h"
+#include "Curves/RichCurve.h"
+
+namespace UE
+{
+	namespace Curves
+	{		
+		int32 SolveCubic(double Coeff[4], double Solution[3])
+		{
+			auto Cbrt = [](double x) -> double
+			{
+				return ((x) > 0.0 ? pow((x), 1.0 / 3.0) : ((x) < 0.0 ? -pow((double)-(x), 1.0 / 3.0) : 0.0));
+			};
+			int32 NumSolutions = 0;
+
+			// Normal form: x^3 + Ax^2 + Bx + C = 0
+			double Denominator = (Coeff[3] != 0.0) ? Coeff[3] : UE_DOUBLE_SMALL_NUMBER * UE_DOUBLE_SMALL_NUMBER;
+			double A = Coeff[2] / Denominator;
+			double B = Coeff[1] / Denominator;
+			double C = Coeff[0] / Denominator;
+
+			// Substitute x = y - A/3 to eliminate quadric term: x^3 +px + q = 0
+			double SqOfA = A * A;
+			double P = 1.0 / 3 * (-1.0 / 3 * SqOfA + B);
+			double Q = 1.0 / 2 * (2.0 / 27 * A * SqOfA - 1.0 / 3 * A * B + C);
+
+			// Use Cardano's formula
+			double CubeOfP = P * P * P;
+			double D = Q * Q + CubeOfP;
+			if (FMath::IsNearlyZero(D))
+			{
+				// One triple solution
+				if (FMath::IsNearlyZero(Q))
+				{
+					Solution[0] = 0;
+					NumSolutions = 1;
+				}
+				// One single and one double solution
+				else
+				{
+					double U = Cbrt(-Q);
+					Solution[0] = 2 * U;
+					Solution[1] = -U;
+					NumSolutions = 2;
+				}
+			}
+			// Three real solutions
+			else if (D < 0)
+			{
+				double Phi = 1.0 / 3 * acos(-Q / sqrt(-CubeOfP));
+				double T = 2 * sqrt(-P);
+
+				Solution[0] = T * cos(Phi);
+				Solution[1] = -T * cos(Phi + UE_DOUBLE_PI / 3);
+				Solution[2] = -T * cos(Phi - UE_DOUBLE_PI / 3);
+				NumSolutions = 3;
+			}
+			// One real solution
+			else
+			{
+				double SqrtD = sqrt(D);
+				double U = Cbrt(SqrtD - Q);
+				double V = -Cbrt(SqrtD + Q);
+
+				Solution[0] = U + V;
+				NumSolutions = 1;
+			}
+
+			// Resubstitute
+			double Sub = 1.0 / 3 * A;
+			for (int32 Index = 0; Index < NumSolutions; ++Index)
+			{
+				Solution[Index] -= Sub;
+			}
+
+			return NumSolutions;
+		}
+
+		void BezierToPower(double A1, double B1, double C1, double D1, double* A2, double* B2, double* C2, double* D2)
+		{
+			double A = B1 - A1;
+			double B = C1 - B1;
+			double C = D1 - C1;
+			double D = B - A;
+			*A2 = C - B - D;
+			*B2 = 3.0 * D;
+			*C2 = 3.0 * A;
+			*D2 = A1;
+		}
+
+		float WeightedEvalForTwoKeys(
+			float Key1Value, float Key1Time, float Key1LeaveTangent, float Key1LeaveTangentWeight, ERichCurveTangentWeightMode Key1TangentWeightMode,
+			float Key2Value, float Key2Time, float Key2ArriveTangent, float Key2ArriveTangentWeight,  ERichCurveTangentWeightMode Key2TangentWeightMode,
+			float InTime)
+		{
+			const float Diff = Key2Time - Key1Time;
+			const float Alpha = (InTime - Key1Time) / Diff;
+			const float OneThird = 1.0f / 3.0f;
+			const double Time1 = Key1Time;
+			const double Time2 = Key2Time;
+			const float X = Time2 - Time1;
+			float CosAngle, SinAngle;
+			float Angle = FMath::Atan(Key1LeaveTangent);
+			FMath::SinCos(&SinAngle, &CosAngle, Angle);
+			float LeaveWeight;
+			if (Key1TangentWeightMode == RCTWM_WeightedNone || Key1TangentWeightMode == RCTWM_WeightedArrive)
+			{
+				const float LeaveTangentNormalized = Key1LeaveTangent;
+				const float Y = LeaveTangentNormalized * X;
+				LeaveWeight = FMath::Sqrt(X * X + Y * Y) * OneThird;
+			}
+			else
+			{
+				LeaveWeight = Key1LeaveTangentWeight;
+			}
+			const float Key1TanX = CosAngle * LeaveWeight + Time1;
+			const float Key1TanY = SinAngle * LeaveWeight + Key1Value;
+
+			Angle = FMath::Atan(Key2ArriveTangent);
+			FMath::SinCos(&SinAngle, &CosAngle, Angle);
+			float ArriveWeight;
+			if (Key2TangentWeightMode == RCTWM_WeightedNone || Key2TangentWeightMode == RCTWM_WeightedLeave)
+			{
+				const float ArriveTangentNormalized = Key2ArriveTangent;
+				const float Y = ArriveTangentNormalized * X;
+				ArriveWeight = FMath::Sqrt(X * X + Y * Y) * OneThird;
+			}
+			else
+			{
+				ArriveWeight = Key2ArriveTangentWeight;
+			}
+			const float Key2TanX = -CosAngle * ArriveWeight + Time2;
+			const float Key2TanY = -SinAngle * ArriveWeight + Key2Value;
+
+			//Normalize the Time Range
+			const float RangeX = Time2 - Time1;
+
+			const float Dx1 = Key1TanX - Time1;
+			const float Dx2 = Key2TanX - Time1;
+
+			// Normalize values
+			const float NormalizedX1 = Dx1 / RangeX;
+			const float NormalizedX2 = Dx2 / RangeX;
+
+			double Coeff[4];
+			double Results[3];
+
+			//Convert Bezier to Power basis, also float to double for precision for root finding.
+			BezierToPower(
+				0.0, NormalizedX1, NormalizedX2, 1.0,
+				&(Coeff[3]), &(Coeff[2]), &(Coeff[1]), &(Coeff[0])
+			);
+
+			Coeff[0] = Coeff[0] - Alpha;
+
+			const int32 NumResults = SolveCubic(Coeff, Results);
+			double NewInterp = Alpha;
+			if (NumResults == 1)
+			{
+				NewInterp = Results[0];
+			}
+			else
+			{
+				NewInterp = TNumericLimits<double>::Lowest(); //just need to be out of range
+				for (double Result : Results)
+				{
+					if ((Result >= 0.0) && (Result <= 1.0))
+					{
+						if (NewInterp < 0.0 || Result > NewInterp)
+						{
+							NewInterp = Result;
+						}
+					}
+				}
+
+				if (NewInterp == TNumericLimits<double>::Lowest())
+				{
+					NewInterp = 0.0;
+				}
+
+			}
+			//now use NewInterp and adjusted tangents plugged into the Y (Value) part of the graph.
+			const float P0 = Key1Value;
+			const float P1 = Key1TanY;
+			const float P3 = Key2Value;
+			const float P2 = Key2TanY;
+
+			return BezierInterp(P0, P1, P2, P3, NewInterp);
+		}
+
+		bool IsItNotWeighted(const FRichCurveKey& Key1, const FRichCurveKey& Key2)
+		{
+			return ((Key1.TangentWeightMode == RCTWM_WeightedNone || Key1.TangentWeightMode == RCTWM_WeightedArrive)
+				&& (Key2.TangentWeightMode == RCTWM_WeightedNone || Key2.TangentWeightMode == RCTWM_WeightedLeave));
+		}
+
+		bool IsWeighted(const FRichCurveKey& Key1, const FRichCurveKey& Key2)
+		{
+			return (Key1.TangentMode != RCTM_None && (Key1.TangentWeightMode == RCTWM_WeightedBoth || Key1.TangentWeightMode == RCTWM_WeightedLeave)) &&
+				(Key2.TangentMode != RCTM_None && (Key2.TangentWeightMode == RCTWM_WeightedBoth || Key2.TangentWeightMode == RCTWM_WeightedArrive));
+		}
+
+		float EvalForTwoKeys(const FRichCurveKey& Key1, const FRichCurveKey& Key2, const float InTime)
+		{
+			const float Diff = Key2.Time - Key1.Time;
+
+			if (Diff > 0.f && Key1.InterpMode != RCIM_Constant)
+			{
+				const float Alpha = (InTime - Key1.Time) / Diff;
+				const float P0 = Key1.Value;
+				const float P3 = Key2.Value;
+
+				if (Key1.InterpMode == RCIM_Linear)
+				{
+					return FMath::Lerp(P0, P3, Alpha);
+				}
+				else 
+				{					
+					if (IsItNotWeighted(Key1,Key2))
+					{
+					    const float OneThird = 1.0f / 3.0f;
+					    const float P1 = P0 + (Key1.LeaveTangent * Diff * OneThird);
+					    const float P2 = P3 - (Key2.ArriveTangent * Diff * OneThird);
+
+					    return BezierInterp(P0, P1, P2, P3, Alpha);
+					}
+					//if (IsWeighted(Key1, Key2))//it's weighted
+					else	
+					{
+					return WeightedEvalForTwoKeys(
+						Key1.Value, Key1.Time, Key1.LeaveTangent, Key1.LeaveTangentWeight, Key1.TangentWeightMode,
+						Key2.Value, Key2.Time, Key2.ArriveTangent, Key2.ArriveTangentWeight, Key2.TangentWeightMode,
+						InTime);
+					}
+				}
+			}
+			else
+			{
+				return Key1.Value;
+			}
+		}
+	}
+}
diff --git a/Engine/Source/Runtime/Engine/Private/Curves/RichCurve.cpp b/Engine/Source/Runtime/Engine/Private/Curves/RichCurve.cpp
--- a/Engine/Source/Runtime/Engine/Private/Curves/RichCurve.cpp
+++ b/Engine/Source/Runtime/Engine/Private/Curves/RichCurve.cpp
@@ -3,6 +3,7 @@
 #include "Curves/RichCurve.h"
 #include "Templates/Function.h"
 #include "Containers/ArrayView.h"
+#include "Curves/CurveEvaluation.h"
 #include "UObject/UE5MainStreamObjectVersion.h"
 
 // Broken - do not turn on! 
@@ -942,266 +943,6 @@ void FRichCurve::RemoveRedundantKeys(float Tolerance, float FirstKeyTime, float
 }
 
 
-/* Solve Cubic Euqation using Cardano's forumla
-* Adopted from Graphic Gems 1
-* [REDACTED_PROFILE_URL]/GraphicsGems/blob/master/gems/Roots3And4.c
-*  Solve cubic of form
-*
-* @param Coeff Coefficient parameters of form  Coeff[0] + Coeff[1]*x + Coeff[2]*x^2 + Coeff[3]*x^3 + Coeff[4]*x^4 = 0
-* @param Solution Up to 3 real solutions. We don't include imaginary solutions, would need a complex number objecct
-* @return Returns the number of real solutions returned in the Solution array.
-*/
-static int SolveCubic(double Coeff[4], double Solution[3])
-{
-	auto cbrt = [](double x) -> double
-	{
-		return ((x) > 0.0 ? pow((x), 1.0 / 3.0) : ((x) < 0.0 ? -pow((double)-(x), 1.0 / 3.0) : 0.0));
-	};
-	int     NumSolutions = 0;
-
-	/* normal form: x^3 + Ax^2 + Bx + C = 0 */
-
-	double A = Coeff[2] / Coeff[3];
-	double B = Coeff[1] / Coeff[3];
-	double C = Coeff[0] / Coeff[3];
-
-	/*  substitute x = y - A/3 to eliminate quadric term:
-	x^3 +px + q = 0 */
-
-	double SqOfA = A * A;
-	double P = 1.0 / 3 * (-1.0 / 3 * SqOfA + B);
-	double Q = 1.0 / 2 * (2.0 / 27 * A * SqOfA - 1.0 / 3 * A * B + C);
-
-	/* use Cardano's formula */
-
-	double CubeOfP = P * P * P;
-	double D = Q * Q + CubeOfP;
-
-	if (FMath::IsNearlyZero(D))
-	{
-		if (FMath::IsNearlyZero(Q)) /* one triple solution */
-		{
-			Solution[0] = 0;
-			NumSolutions = 1;
-		}
-		else /* one single and one double solution */
-		{
-			double u = cbrt(-Q);
-			Solution[0] = 2 * u;
-			Solution[1] = -u;
-			NumSolutions = 2;
-		}
-	}
-	else if (D < 0) /* Casus irreducibilis: three real solutions */
-	{
-		double phi = 1.0 / 3 * acos(-Q / sqrt(-CubeOfP));
-		double t = 2 * sqrt(-P);
-
-		Solution[0] = t * cos(phi);
-		Solution[1] = -t * cos(phi + UE_DOUBLE_PI / 3);
-		Solution[2] = -t * cos(phi - UE_DOUBLE_PI / 3);
-		NumSolutions = 3;
-	}
-	else /* one real solution */
-	{
-		double sqrt_D = sqrt(D);
-		double u = cbrt(sqrt_D - Q);
-		double v = -cbrt(sqrt_D + Q);
-
-		Solution[0] = u + v;
-		NumSolutions = 1;
-	}
-
-	/* resubstitute */
-
-	double Sub = 1.0 / 3 * A;
-
-	for (int i = 0; i < NumSolutions; ++i)
-	{
-		Solution[i] -= Sub;
-
-	}
-	return NumSolutions;
-}
-
-/** Util to find float value on bezier defined by 4 control points */
-FORCEINLINE_DEBUGGABLE static float BezierInterp(float P0, float P1, float P2, float P3, float Alpha)
-{
-	const float P01 = FMath::Lerp(P0, P1, Alpha);
-	const float P12 = FMath::Lerp(P1, P2, Alpha);
-	const float P23 = FMath::Lerp(P2, P3, Alpha);
-	const float P012 = FMath::Lerp(P01, P12, Alpha);
-	const float P123 = FMath::Lerp(P12, P23, Alpha);
-	const float P0123 = FMath::Lerp(P012, P123, Alpha);
-
-	return P0123;
-}
-
-/*
-*   Convert the control values for a polynomial defined in the Bezier
-*		basis to a polynomial defined in the power basis (t^3 t^2 t 1).
-*/
-static void BezierToPower(double A1, double B1, double C1, double D1,
-	double* A2, double* B2, double* C2, double* D2)
-{
-	double A = B1 - A1;
-	double B = C1 - B1;
-	double C = D1 - C1;
-	double D = B - A;
-	*A2 = C - B - D;
-	*B2 = 3.0 * D;
-	*C2 = 3.0 * A;
-	*D2 = A1;
-}
-
-static float WeightedEvalForTwoKeys(
-	float Key1Value, float Key1Time, float Key1LeaveTangent, float Key1LeaveTangentWeight,ERichCurveTangentWeightMode Key1TangentWeightMode,
-	float Key2Value, float Key2Time, float Key2ArriveTangent, float Key2ArriveTangentWeight,  ERichCurveTangentWeightMode Key2TangentWeightMode,
-	float InTime)
-{
-	const float Diff = Key2Time - Key1Time;
-	const float Alpha = (InTime - Key1Time) / Diff;
-	const float P0 = Key1Value;
-	const float P3 = Key2Value;
-	const float OneThird = 1.0f / 3.0f;
-	const double Time1 = Key1Time;
-	const double Time2 = Key2Time;
-	const float X = Time2 - Time1;
-	float CosAngle, SinAngle;
-	float Angle = FMath::Atan(Key1LeaveTangent);
-	FMath::SinCos(&SinAngle, &CosAngle, Angle);
-	float LeaveWeight;
-	if (Key1TangentWeightMode == RCTWM_WeightedNone || Key1TangentWeightMode == RCTWM_WeightedArrive)
-	{
-		const float LeaveTangentNormalized = Key1LeaveTangent;
-		const float Y = LeaveTangentNormalized * X;
-		LeaveWeight = FMath::Sqrt(X * X + Y * Y) * OneThird;
-	}
-	else
-	{
-		LeaveWeight = Key1LeaveTangentWeight;
-	}
-	const float Key1TanX = CosAngle * LeaveWeight + Time1;
-	const float Key1TanY = SinAngle * LeaveWeight + Key1Value;
-
-	Angle = FMath::Atan(Key2ArriveTangent);
-	FMath::SinCos(&SinAngle, &CosAngle, Angle);
-	float ArriveWeight;
-	if (Key2TangentWeightMode == RCTWM_WeightedNone || Key2TangentWeightMode == RCTWM_WeightedLeave)
-	{
-		const float ArriveTangentNormalized = Key2ArriveTangent;
-		const float Y = ArriveTangentNormalized * X;
-		ArriveWeight = FMath::Sqrt(X * X + Y * Y) * OneThird;
-	}
-	else
-	{
-		ArriveWeight = Key2ArriveTangentWeight;
-	}
-	const float Key2TanX = -CosAngle * ArriveWeight + Time2;
-	const float Key2TanY = -SinAngle * ArriveWeight + Key2Value;
-
-	//Normalize the Time Range
-	const float RangeX = Time2 - Time1;
-
-	const float Dx1 = Key1TanX - Time1;
-	const float Dx2 = Key2TanX - Time1;
-
-	// Normalize values
-	const float NormalizedX1 = Dx1 / RangeX;
-	const float NormalizedX2 = Dx2 / RangeX;
-
-	double Coeff[4];
-	double Results[3];
-
-	//Convert Bezier to Power basis, also float to double for precision for root finding.
-	BezierToPower(
-		0.0, NormalizedX1, NormalizedX2, 1.0,
-		&(Coeff[3]), &(Coeff[2]), &(Coeff[1]), &(Coeff[0])
-	);
-
-	Coeff[0] = Coeff[0] - Alpha;
-
-	int NumResults = SolveCubic(Coeff, Results);
-	float NewInterp = Alpha;
-	if (NumResults == 1)
-	{
-		NewInterp = Results[0];
-	}
-	else
-	{
-		NewInterp = TNumericLimits<float>::Lowest(); //just need to be out of range
-		for (double Result : Results)
-		{
-			if ((Result >= 0.0f) && (Result <= 1.0f))
-			{
-				if (NewInterp < 0.0f || Result > NewInterp)
-				{
-					NewInterp = Result;
-				}
-			}
-		}
-
-		if (NewInterp == TNumericLimits<float>::Lowest())
-		{
-			NewInterp = 0.f;
-		}
-
-	}
-	//now use NewInterp and adjusted tangents plugged into the Y (Value) part of the graph.
-	//const float P0 = Key1.Value;
-	const float P1 = Key1TanY;
-	//const float P3 = Key2.Value;
-	const float P2 = Key2TanY;
-
-	float OutValue = BezierInterp(P0, P1, P2, P3, NewInterp);
-	return OutValue;
-}
-
-static bool IsItNotWeighted(const FRichCurveKey& Key1, const FRichCurveKey& Key2)
-{
-	return ((Key1.TangentWeightMode == RCTWM_WeightedNone || Key1.TangentWeightMode == RCTWM_WeightedArrive)
-		&& (Key2.TangentWeightMode == RCTWM_WeightedNone || Key2.TangentWeightMode == RCTWM_WeightedLeave));
-}
-
-float EvalForTwoKeys(const FRichCurveKey& Key1, const FRichCurveKey& Key2, const float InTime)
-{
-	const float Diff = Key2.Time - Key1.Time;
-
-	if (Diff > 0.f && Key1.InterpMode != RCIM_Constant)
-	{
-		const float Alpha = (InTime - Key1.Time) / Diff;
-		const float P0 = Key1.Value;
-		const float P3 = Key2.Value;
-
-		if (Key1.InterpMode == RCIM_Linear)
-		{
-			return FMath::Lerp(P0, P3, Alpha);
-		}
-		else 
-		{
-			if (IsItNotWeighted(Key1,Key2))
-			{
-				const float OneThird = 1.0f / 3.0f;
-				const float P1 = P0 + (Key1.LeaveTangent * Diff * OneThird);
-				const float P2 = P3 - (Key2.ArriveTangent * Diff * OneThird);
-
-				return BezierInterp(P0, P1, P2, P3, Alpha);
-			}
-			else //it's weighted
-			{
-				return  WeightedEvalForTwoKeys(
-					Key1.Value, Key1.Time, Key1.LeaveTangent, Key1.LeaveTangentWeight, Key1.TangentWeightMode,
-					Key2.Value, Key2.Time, Key2.ArriveTangent, Key2.ArriveTangentWeight, Key2.TangentWeightMode,
-					InTime);
-			}
-		}
-	}
-	else
-	{
-		return Key1.Value;
-	}
-}
-
 void FRichCurve::RemoveRedundantKeysInternal(float Tolerance, int32 InStartKeepKey, int32 InEndKeepKey, FFrameRate SampleRate)
 {
 	if (Keys.Num() < 3) // Will always keep first and last key
@@ -1272,11 +1013,11 @@ void FRichCurve::RemoveRedundantKeysInternal(float Tolerance, int32 InStartKeepK
 				while (SampleFrameTime <= KeyTwoFrameTime)
 				{
 					const float SampleSeconds = SampleRate.AsSeconds(SampleFrameTime);
-					const float ValueWithoutKey = EvalForTwoKeys(Keys[MostRecentKeepKeyIndex], Keys[TestIndex + 1], SampleSeconds);
+					const float ValueWithoutKey = UE::Curves::EvalForTwoKeys(Keys[MostRecentKeepKeyIndex], Keys[TestIndex + 1], SampleSeconds);
 
 					const FRichCurveKey& KeyOne = SampleSeconds <= ToRemoveKeyTime ? Keys[MostRecentKeepKeyIndex] : Keys[TestIndex];
 					const FRichCurveKey& KeyTwo = SampleSeconds <= ToRemoveKeyTime ? Keys[TestIndex] : Keys[TestIndex + 1];
-					const float ValueWithKey = EvalForTwoKeys(KeyOne, KeyTwo, SampleSeconds);
+					const float ValueWithKey = UE::Curves::EvalForTwoKeys(KeyOne, KeyTwo, SampleSeconds);
 				
 					if (FMath::Abs(ValueWithoutKey - ValueWithKey) > Tolerance)
 					{
@@ -1299,7 +1040,7 @@ void FRichCurve::RemoveRedundantKeysInternal(float Tolerance, int32 InStartKeepK
 				if (IsNonUserTangentKey(Keys[MostRecentKeepKeyIndex]) && IsNonUserTangentKey(Keys[TestIndex + 1]) && IsNonUserTangentKey(Keys[TestIndex]))
 				{
 					const float KeyValue = Keys[TestIndex].Value;
-					const float ValueWithoutKey = EvalForTwoKeys(Keys[MostRecentKeepKeyIndex], Keys[TestIndex + 1], Keys[TestIndex].Time);
+					const float ValueWithoutKey = UE::Curves::EvalForTwoKeys(Keys[MostRecentKeepKeyIndex], Keys[TestIndex + 1], Keys[TestIndex].Time);
 
 					bKeepKey = FMath::Abs(ValueWithoutKey - KeyValue) > Tolerance;
 				}
@@ -1450,7 +1191,7 @@ float FRichCurve::Eval(float InTime, float InDefaultValue) const
 			}
 		}
 
-		InterpVal = EvalForTwoKeys(Keys[first - 1], Keys[first], InTime);
+		InterpVal = UE::Curves::EvalForTwoKeys(Keys[first - 1], Keys[first], InTime);
 	}
 	else
 	{
@@ -1582,17 +1323,17 @@ static ERichCurveKeyTimeCompressionFormat FindRichCurveKeyFormat(const FRichCurv
 			{
 				return FMath::Lerp(P0, P3, Alpha);
 			}
-			else if(IsItNotWeighted(Key1,Key2))
+			else if(UE::Curves::IsItNotWeighted(Key1,Key2))
 			{
 				const float OneThird = 1.0f / 3.0f;
 				const float P1 = P0 + (Key1.LeaveTangent * Diff * OneThird);
 				const float P2 = P3 - (Key2.ArriveTangent * Diff * OneThird);
 
-				return BezierInterp(P0, P1, P2, P3, Alpha);
+				return UE::Curves::BezierInterp(P0, P1, P2, P3, Alpha);
 			}
 			else //it's weighted
 			{
-				return  WeightedEvalForTwoKeys(
+				return  UE::Curves::WeightedEvalForTwoKeys(
 					Key1.Value, KeyTime1, Key1.LeaveTangent, Key1.LeaveTangentWeight, Key1.TangentWeightMode,
 					Key2.Value, KeyTime2, Key2.ArriveTangent, Key2.ArriveTangentWeight, Key2.TangentWeightMode,
 					InTime);
@@ -1657,7 +1398,7 @@ static ERichCurveKeyTimeCompressionFormat FindRichCurveKeyFormat(const FRichCurv
 		}
 
 		const float Result_16 = EvalForTwoKeys(Curve.Keys[CurrentLossyKey], LossyTime1_16, Curve.Keys[CurrentLossyKey + 1], LossyTime2_16, CurrentTime);
-		const float Result_Ref = ::EvalForTwoKeys(Curve.Keys[CurrentRefKey], Curve.Keys[CurrentRefKey + 1], CurrentTime);
+		const float Result_Ref = UE::Curves::EvalForTwoKeys(Curve.Keys[CurrentRefKey], Curve.Keys[CurrentRefKey + 1], CurrentTime);
 
 		const float Error_16 = FMath::Abs(Result_Ref - Result_16);
 
@@ -2317,13 +2058,13 @@ FORCEINLINE_DEBUGGABLE static float InterpEval(float InTime, const KeyTimeAdapte
 			const float P1 = P0 + (KeyLeaveTangent0 * ScaledDiff);
 			const float P2 = P3 - (KeyArriveTangent1 * ScaledDiff);
 
-			InterpolatedValue = BezierInterp(P0, P1, P2, P3, Alpha);
+			InterpolatedValue = UE::Curves::BezierInterp(P0, P1, P2, P3, Alpha);
 		}
 		else //it's weighted
 		{
-			return  WeightedEvalForTwoKeys(
-				KeyValue0, KeyTime0, KeyDataAdapter.GetKeyLeaveTangent(KeyValueHandle0), KeyDataAdapter.GetKeyLeaveTangentWeight(KeyValueHandle0), KeyDataAdapter.GetKeyTangentWeightMode(KeyValueHandle0),
-				KeyValue1, KeyTime1, KeyDataAdapter.GetKeyArriveTangent(KeyValueHandle1), KeyDataAdapter.GetKeyArriveTangentWeight(KeyValueHandle1), KeyDataAdapter.GetKeyTangentWeightMode(KeyValueHandle1),
+			return  UE::Curves::WeightedEvalForTwoKeys(
+				KeyValue0, KeyTime0, KeyDataAdapter.GetKeyLeaveTangent(KeyValueHandle0), KeyDataAdapter.GetKeyLeaveTangentWeight(KeyValueHandle0), KeyDataAdapter.GetKeyTangentWeightMode(First - 1),
+				KeyValue1, KeyTime1, KeyDataAdapter.GetKeyArriveTangent(KeyValueHandle1), KeyDataAdapter.GetKeyArriveTangentWeight(KeyValueHandle1), KeyDataAdapter.GetKeyTangentWeightMode(First),
 				InTime);
 		}
 	}
diff --git a/Engine/Source/Runtime/Engine/Public/Curves/CurveEvaluation.h b/Engine/Source/Runtime/Engine/Public/Curves/CurveEvaluation.h
new file mode 100644
--- /dev/null
+++ b/Engine/Source/Runtime/Engine/Public/Curves/CurveEvaluation.h
@@ -0,0 +1,55 @@
+// Copyright Epic Games, Inc. All Rights Reserved.
+
+#pragma once
+
+#include "Curves/RichCurve.h"
+
+namespace UE
+{
+	namespace Curves
+	{
+		/* Solve Cubic Euqation using Cardano's forumla
+		* Adopted from Graphic Gems 1
+		* [REDACTED_PROFILE_URL]/GraphicsGems/blob/master/gems/Roots3And4.c
+		*  Solve cubic of form
+		*
+		* @param Coeff Coefficient parameters of form  Coeff[0] + Coeff[1]*x + Coeff[2]*x^2 + Coeff[3]*x^3 + Coeff[4]*x^4 = 0
+		* @param Solution Up to 3 real solutions. We don't include imaginary solutions, would need a complex number objecct
+		* @return Returns the number of real solutions returned in the Solution array.
+		*/
+		ENGINE_API int32 SolveCubic(double Coeff[4], double Solution[3]);
+
+		/**  Assuming that P0, P1, P2 and P3 are sequential control points of an N=4 Bezier curve, returns
+		 * the interpolated value for interpolation constant Alpha in [0, 1]
+		 */
+		template<typename CurveValueType>
+		CurveValueType BezierInterp(CurveValueType P0, CurveValueType P1, CurveValueType P2, CurveValueType P3, float Alpha)
+		{
+			const CurveValueType P01 = FMath::Lerp(P0, P1, Alpha);
+			const CurveValueType P12 = FMath::Lerp(P1, P2, Alpha);
+			const CurveValueType P23 = FMath::Lerp(P2, P3, Alpha);
+			const CurveValueType P012 = FMath::Lerp(P01, P12, Alpha);
+			const CurveValueType P123 = FMath::Lerp(P12, P23, Alpha);
+			const CurveValueType P0123 = FMath::Lerp(P012, P123, Alpha);
+
+			return P0123;
+		}
+
+		/*
+		*   Convert the control values for a polynomial defined in the Bezier
+		*		basis to a polynomial defined in the power basis (t^3 t^2 t 1).
+		*/
+		ENGINE_API void BezierToPower(double A1, double B1, double C1, double D1, double* A2, double* B2, double* C2, double* D2);
+		
+		float WeightedEvalForTwoKeys(float Key1Value, float Key1Time, float Key1LeaveTangent, float Key1LeaveTangentWeight, ERichCurveTangentWeightMode Key1TangentWeightMode, float Key2Value, float Key2Time, float Key2ArriveTangent, float Key2ArriveTangentWeight,  ERichCurveTangentWeightMode Key2TangentWeightMode, float InTime);
+
+		/** Returns whether the interpolation between Key1 and Key2 does NOT require using weighted tangent data */
+		bool IsItNotWeighted(const FRichCurveKey& Key1, const FRichCurveKey& Key2);
+
+		/** Returns whether  the interpolation between Key1 and Key2 requires using weighted tangent data */
+		bool IsWeighted(const FRichCurveKey& Key1, const FRichCurveKey& Key2);
+
+		/** Evaluates and interpolates specific point, determined by InTime, between Key1 and Key2 */
+		ENGINE_API float EvalForTwoKeys(const FRichCurveKey& Key1, const FRichCurveKey& Key2, const float InTime);
+	}
+}
\ No newline at end of file
diff --git a/Engine/Source/Runtime/MovieScene/Private/Channels/MovieSceneCurveChannelImpl.cpp b/Engine/Source/Runtime/MovieScene/Private/Channels/MovieSceneCurveChannelImpl.cpp
--- a/Engine/Source/Runtime/MovieScene/Private/Channels/MovieSceneCurveChannelImpl.cpp
+++ b/Engine/Source/Runtime/MovieScene/Private/Channels/MovieSceneCurveChannelImpl.cpp
@@ -5,6 +5,7 @@
 #include "Channels/MovieSceneFloatChannel.h"
 #include "HAL/ConsoleManager.h"
 #include "MovieSceneFrameMigration.h"
+#include "Curves/CurveEvaluation.h"
 #include "UObject/FortniteMainBranchObjectVersion.h"
 #include "UObject/SequencerObjectVersion.h"
 
@@ -27,20 +28,6 @@ namespace UE
 namespace MovieScene
 {
 
-/** Util to find value on bezier defined by 4 control points */
-template<typename CurveValueType>
-CurveValueType BezierInterp(CurveValueType P0, CurveValueType P1, CurveValueType P2, CurveValueType P3, float Alpha)
-{
-	const CurveValueType P01   = FMath::Lerp(P0,   P1,   Alpha);
-	const CurveValueType P12   = FMath::Lerp(P1,   P2,   Alpha);
-	const CurveValueType P23   = FMath::Lerp(P2,   P3,   Alpha);
-	const CurveValueType P012  = FMath::Lerp(P01,  P12,  Alpha);
-	const CurveValueType P123  = FMath::Lerp(P12,  P23,  Alpha);
-	const CurveValueType P0123 = FMath::Lerp(P012, P123, Alpha);
-
-	return P0123;
-}
-
 template<typename ChannelType>
 static typename ChannelType::CurveValueType
 EvalForTwoKeys(
@@ -76,7 +63,7 @@ EvalForTwoKeys(
 			const CurveValueType P1 = P0 + (LeaveTangent * Diff*OneThird);
 			const CurveValueType P2 = P3 - (ArriveTangent * Diff*OneThird);
 
-			return BezierInterp(P0, P1, P2, P3, Alpha);
+			return UE::Curves::BezierInterp(P0, P1, P2, P3, Alpha);
 		}
 	}
 	else
@@ -142,104 +129,6 @@ FCycleParams CycleTime(FFrameNumber MinFrame, FFrameNumber MaxFrame, FFrameTime
 	return Params;
 }
 
-/* Solve Cubic Equation using Cardano's forumla
-* Adopted from Graphic Gems 1
-* [REDACTED_PROFILE_URL]/GraphicsGems/blob/master/gems/Roots3And4.c
-*  Solve cubic of form
-*
-* @param Coeff Coefficient parameters of form  Coeff[0] + Coeff[1]*x + Coeff[2]*x^2 + Coeff[3]*x^3 + Coeff[4]*x^4 = 0
-* @param Solution Up to 3 real solutions. We don't include imaginary solutions, would need a complex number objecct
-* @return Returns the number of real solutions returned in the Solution array.
-*/
-static int SolveCubic(double Coeff[4], double Solution[3])
-{
-	auto cbrt = [](double x) -> double
-	{
-		return ((x) > 0.0 ? pow((x), 1.0 / 3.0) : ((x) < 0.0 ? -pow((double)-(x), 1.0 / 3.0) : 0.0));
-	};
-	int     NumSolutions = 0;
-
-	/* normal form: x^3 + Ax^2 + Bx + C = 0 */
-
-	double A = Coeff[2] / Coeff[3];
-	double B = Coeff[1] / Coeff[3];
-	double C = Coeff[0] / Coeff[3];
-
-	/*  substitute x = y - A/3 to eliminate quadric term:
-	x^3 +px + q = 0 */
-
-	double SqOfA = A * A;
-	double P = 1.0 / 3 * (-1.0 / 3 * SqOfA + B);
-	double Q = 1.0 / 2 * (2.0 / 27 * A * SqOfA - 1.0 / 3 * A * B + C);
-
-	/* use Cardano's formula */
-
-	double CubeOfP = P * P * P;
-	double D = Q * Q + CubeOfP;
-
-	if (FMath::IsNearlyZero(D))
-	{
-		if (FMath::IsNearlyZero(Q)) /* one triple solution */
-		{
-			Solution[0] = 0;
-			NumSolutions = 1;
-		}
-		else /* one single and one double solution */
-		{
-			double u = cbrt(-Q);
-			Solution[0] = 2 * u;
-			Solution[1] = -u;
-			NumSolutions = 2;
-		}
-	}
-	else if (D < 0) /* Casus irreducibilis: three real solutions */
-	{
-		double phi = 1.0 / 3 * acos(-Q / sqrt(-CubeOfP));
-		double t = 2 * sqrt(-P);
-
-		Solution[0] = t * cos(phi);
-		Solution[1] = -t * cos(phi + UE_DOUBLE_PI / 3);
-		Solution[2] = -t * cos(phi - UE_DOUBLE_PI / 3);
-		NumSolutions = 3;
-	}
-	else /* one real solution */
-	{
-		double sqrt_D = sqrt(D);
-		double u = cbrt(sqrt_D - Q);
-		double v = -cbrt(sqrt_D + Q);
-
-		Solution[0] = u + v;
-		NumSolutions = 1;
-	}
-
-	/* resubstitute */
-
-	double Sub = 1.0 / 3 * A;
-
-	for (int i = 0; i < NumSolutions; ++i)
-		Solution[i] -= Sub;
-
-	return NumSolutions;
-}
-
-/*
-*   Convert the control values for a polynomial defined in the Bezier
-*		basis to a polynomial defined in the power basis (t^3 t^2 t 1).
-*/
-static void BezierToPower(	double A1, double B1, double C1, double D1,
-	double *A2, double *B2, double *C2, double *D2)
-{
-	double A = B1 - A1;
-	double B = C1 - B1;
-	double C = D1 - C1;
-	double D = B - A;
-	*A2 = C- B - D;
-	*B2 = 3.0 * D;
-	*C2 = 3.0 * A;
-	*D2 = A1;
-}
-
-
 } // namespace MovieScene
 } // namespace UE
 
@@ -508,7 +397,7 @@ bool TMovieSceneCurveChannelImpl<ChannelType>::Evaluate(const ChannelType* InCha
 				const float P3 = Key2.Value;
 				const float P2 = P3 - (Key2.Tangent.ArriveTangent * Diff * OneThird);
 
-				OutValue = Params.ValueOffset + BezierInterp(P0, P1, P2, P3, Interp);
+				OutValue = Params.ValueOffset + UE::Curves::BezierInterp(P0, P1, P2, P3, Interp);
 				break;
 			}
 			else //its weighted
@@ -566,14 +455,14 @@ bool TMovieSceneCurveChannelImpl<ChannelType>::Evaluate(const ChannelType* InCha
 				double Results[3];
 
 				//Convert Bezier to Power basis, also float to double for precision for root finding.
-				BezierToPower(
+				UE::Curves::BezierToPower(
 					0.0, NormalizedX1, NormalizedX2, 1.0,
 					&(Coeff[3]), &(Coeff[2]), &(Coeff[1]), &(Coeff[0])
 				);
 
 				Coeff[0] = Coeff[0] - Interp;
 				
-				int NumResults = SolveCubic(Coeff, Results);
+				const int32 NumResults = UE::Curves::SolveCubic(Coeff, Results);
 				float NewInterp = Interp;
 				if (NumResults == 1)
 				{
@@ -605,7 +494,7 @@ bool TMovieSceneCurveChannelImpl<ChannelType>::Evaluate(const ChannelType* InCha
 				const float P3 = Key2.Value;
 				const float P2 = Key2TanY;
 
-				OutValue = Params.ValueOffset + BezierInterp(P0, P1, P2, P3,  NewInterp);
+				OutValue = Params.ValueOffset + UE::Curves::BezierInterp(P0, P1, P2, P3,  NewInterp);
 			}
 			break;
 		}
