SPP limit in the UX for ray tracing passes bumped to a higher value

diff --git a/Engine/Source/Runtime/Engine/Classes/Engine/Scene.h b/Engine/Source/Runtime/Engine/Classes/Engine/Scene.h
--- a/Engine/Source/Runtime/Engine/Classes/Engine/Scene.h
+++ b/Engine/Source/Runtime/Engine/Classes/Engine/Scene.h
@@ -1583,7 +1583,7 @@ struct FPostProcessSettings
 	int32 RayTracingGIMaxBounces;
 
 	/** Sets the samples per pixel for ray tracing global illumination. */
-	UPROPERTY(interp, EditAnywhere, BlueprintReadWrite, Category = "Global Illumination|Ray Tracing Global Illumination", meta = (ClampMin = "1", ClampMax = "64", editcondition = "bOverride_RayTracingGISamplesPerPixel", DisplayName = "Samples Per Pixel"))
+	UPROPERTY(interp, EditAnywhere, BlueprintReadWrite, Category = "Global Illumination|Ray Tracing Global Illumination", meta = (ClampMin = "1", ClampMax = "65536", editcondition = "bOverride_RayTracingGISamplesPerPixel", DisplayName = "Samples Per Pixel"))
 	int32 RayTracingGISamplesPerPixel;
 	
 	/** Chooses the Reflection method. Not compatible with Forward Shading. */
@@ -1626,7 +1626,7 @@ struct FPostProcessSettings
 	int32 RayTracingReflectionsMaxBounces;
 
 	/** Sets the samples per pixel for ray traced reflections. */
-	UPROPERTY(interp, EditAnywhere, BlueprintReadWrite, Category = "Reflections|Ray Tracing Reflections", meta = (ClampMin = "1", ClampMax = "64", editcondition = "bOverride_RayTracingReflectionsSamplesPerPixel", DisplayName = "Samples Per Pixel"))
+	UPROPERTY(interp, EditAnywhere, BlueprintReadWrite, Category = "Reflections|Ray Tracing Reflections", meta = (ClampMin = "1", ClampMax = "65536", editcondition = "bOverride_RayTracingReflectionsSamplesPerPixel", DisplayName = "Samples Per Pixel"))
 	int32 RayTracingReflectionsSamplesPerPixel;
 
 	/** Sets the reflected shadows type. */
@@ -1933,7 +1933,7 @@ struct FPostProcessSettings
 	uint32 RayTracingAO : 1;
 
 	/** Sets the samples per pixel for ray tracing ambient occlusion. */
-	UPROPERTY(interp, EditAnywhere, BlueprintReadWrite, Category = "Rendering Features|Ray Tracing Ambient Occlusion", meta = (ClampMin = "1", ClampMax = "64", editcondition = "bOverride_RayTracingAOSamplesPerPixel", DisplayName = "Samples Per Pixel"))
+	UPROPERTY(interp, EditAnywhere, BlueprintReadWrite, Category = "Rendering Features|Ray Tracing Ambient Occlusion", meta = (ClampMin = "1", ClampMax = "65536", editcondition = "bOverride_RayTracingAOSamplesPerPixel", DisplayName = "Samples Per Pixel"))
 	int32 RayTracingAOSamplesPerPixel;
 
 	/** Scalar factor on the ray-tracing ambient occlusion score. */
@@ -2077,7 +2077,7 @@ struct FPostProcessSettings
 	int32 RayTracingTranslucencyRefractionRays;
 
 	/** Sets the samples per pixel for ray traced translucency. */
-	UPROPERTY(interp, EditAnywhere, BlueprintReadWrite, Category = "Rendering Features|Ray Tracing Translucency", meta = (ClampMin = "1", ClampMax = "64", editcondition = "bOverride_RayTracingTranslucencySamplesPerPixel", DisplayName = "Samples Per Pixel"))
+	UPROPERTY(interp, EditAnywhere, BlueprintReadWrite, Category = "Rendering Features|Ray Tracing Translucency", meta = (ClampMin = "1", ClampMax = "65536", editcondition = "bOverride_RayTracingTranslucencySamplesPerPixel", DisplayName = "Samples Per Pixel"))
 	int32 RayTracingTranslucencySamplesPerPixel;
 
 	/** Sets the translucency shadows type. */
