Add PopulateRichCurve to CompressedRichCurve, allowing for decompressing key-data into a FRichCurve

diff --git a/Engine/Source/Runtime/Engine/Classes/Curves/RichCurve.h b/Engine/Source/Runtime/Engine/Classes/Curves/RichCurve.h
--- a/Engine/Source/Runtime/Engine/Classes/Curves/RichCurve.h
+++ b/Engine/Source/Runtime/Engine/Classes/Curves/RichCurve.h
@@ -396,6 +396,9 @@ struct ENGINE_API FCompressedRichCurve
 	/** Evaluate this rich curve at the specified time */
 	float Eval(float InTime, float InDefaultValue = 0.0f) const;
 
+	/** Populate RichCurve with decompressed key-data */
+	void PopulateCurve(FRichCurve& OutCurve) const;
+
 	/** Evaluate this rich curve at the specified time */
 	static float StaticEval(ERichCurveCompressionFormat CompressionFormat, ERichCurveKeyTimeCompressionFormat KeyTimeCompressionFormat, ERichCurveExtrapolation PreInfinityExtrap, ERichCurveExtrapolation PostInfinityExtrap, TConstantValueNumKeys ConstantValueNumKeys, const uint8* CompressedKeys, float InTime, float InDefaultValue = 0.0f);
 
diff --git a/Engine/Source/Runtime/Engine/Private/Curves/RichCurve.cpp b/Engine/Source/Runtime/Engine/Private/Curves/RichCurve.cpp
--- a/Engine/Source/Runtime/Engine/Private/Curves/RichCurve.cpp
+++ b/Engine/Source/Runtime/Engine/Private/Curves/RichCurve.cpp
@@ -2180,6 +2180,194 @@ static TFunction<float(ERichCurveExtrapolation PreInfinityExtrap, ERichCurveExtr
 		},
 };
 
+ERichCurveInterpMode GetKeyInterpMode(ERichCurveCompressionFormat Format)
+{
+	switch (Format)
+	{
+	case ERichCurveCompressionFormat::RCCF_Linear:
+		return ERichCurveInterpMode::RCIM_Linear;
+		
+	case ERichCurveCompressionFormat::RCCF_Cubic:
+		return ERichCurveInterpMode::RCIM_Cubic;
+		
+	case ERichCurveCompressionFormat::RCCF_Weighted:
+		return ERichCurveInterpMode::RCIM_Cubic;
+
+	case ERichCurveCompressionFormat::RCCF_Constant:
+		return ERichCurveInterpMode::RCIM_Constant;
+
+	case ERichCurveCompressionFormat::RCCF_Empty:
+		return ERichCurveInterpMode::RCIM_None;
+
+	case ERichCurveCompressionFormat::RCCF_Mixed:
+		check(false);		
+	}
+	
+	return ERichCurveInterpMode::RCIM_None;
+}
+
+static TFunction<FRichCurveKey(FCompressedRichCurve::TConstantValueNumKeys ConstantValueNumKeys, const uint8* CompressedKeys, int32 Index)> ToRichCurveMap[6][2]
+{
+	// RCCF_Empty
+	{
+		// RCKTCF_uint16
+		[](FCompressedRichCurve::TConstantValueNumKeys ConstantValueNumKeys, const uint8* CompressedKeys, int32 Index)
+		{
+			return FRichCurveKey(0.f, ConstantValueNumKeys.ConstantValue);
+		},
+		// RCKTCF_float32
+		[](FCompressedRichCurve::TConstantValueNumKeys ConstantValueNumKeys, const uint8* CompressedKeys, int32 Index)
+		{
+			return FRichCurveKey(0.f, ConstantValueNumKeys.ConstantValue);
+		},
+	},
+	// RCCF_Constant
+	{
+		// RCKTCF_uint16
+		[](FCompressedRichCurve::TConstantValueNumKeys ConstantValueNumKeys, const uint8* CompressedKeys, int32 Index)
+		{
+			return FRichCurveKey(0.f, ConstantValueNumKeys.ConstantValue);
+		},
+		// RCKTCF_float32
+		[](FCompressedRichCurve::TConstantValueNumKeys ConstantValueNumKeys, const uint8* CompressedKeys, int32 Index)
+		{
+			return FRichCurveKey(0.f, ConstantValueNumKeys.ConstantValue);
+		},
+	},
+	// RCCF_Linear
+	{
+		// RCKTCF_uint16
+		[](FCompressedRichCurve::TConstantValueNumKeys ConstantValueNumKeys, const uint8* CompressedKeys, int32 Index)
+		{
+			const int32 KeyTimesOffset = 0;
+			const Quantized16BitKeyTimeAdapter KeyTimeAdapter(CompressedKeys, KeyTimesOffset, ConstantValueNumKeys.NumKeys);
+			const UniformKeyDataAdapter<RCCF_Linear> KeyDataAdapter(CompressedKeys, KeyTimeAdapter);
+
+			const KeyDataHandle Handle = KeyDataAdapter.GetKeyDataHandle(Index);
+			return FRichCurveKey(KeyTimeAdapter.GetTime(Index), KeyDataAdapter.GetKeyValue(Handle));
+		},
+		// RCKTCF_float32
+		[](FCompressedRichCurve::TConstantValueNumKeys ConstantValueNumKeys, const uint8* CompressedKeys, int32 Index)
+		{
+			const int32 KeyTimesOffset = 0;
+			const Float32BitKeyTimeAdapter KeyTimeAdapter(CompressedKeys, KeyTimesOffset, ConstantValueNumKeys.NumKeys);
+			const UniformKeyDataAdapter<RCCF_Linear> KeyDataAdapter(CompressedKeys, KeyTimeAdapter);
+
+			const KeyDataHandle Handle = KeyDataAdapter.GetKeyDataHandle(Index);
+			return FRichCurveKey(KeyTimeAdapter.GetTime(Index), KeyDataAdapter.GetKeyValue(Handle));
+		},
+	},
+	// RCCF_Cubic
+	{
+		// RCKTCF_uint16
+		[](FCompressedRichCurve::TConstantValueNumKeys ConstantValueNumKeys, const uint8* CompressedKeys, int32 Index)
+		{
+			const int32 KeyTimesOffset = 0;
+			const Quantized16BitKeyTimeAdapter KeyTimeAdapter(CompressedKeys, KeyTimesOffset, ConstantValueNumKeys.NumKeys);
+			const UniformKeyDataAdapter<RCCF_Cubic> KeyDataAdapter(CompressedKeys, KeyTimeAdapter);
+
+			const KeyDataHandle Handle = KeyDataAdapter.GetKeyDataHandle(Index);
+			return FRichCurveKey(KeyTimeAdapter.GetTime(Index), KeyDataAdapter.GetKeyValue(Handle), KeyDataAdapter.GetKeyArriveTangent(Handle), KeyDataAdapter.GetKeyLeaveTangent(Handle), RCIM_Cubic);
+		},
+		// RCKTCF_float32
+		[](FCompressedRichCurve::TConstantValueNumKeys ConstantValueNumKeys, const uint8* CompressedKeys, int32 Index)
+		{
+			const int32 KeyTimesOffset = 0;
+			const Float32BitKeyTimeAdapter KeyTimeAdapter(CompressedKeys, KeyTimesOffset, ConstantValueNumKeys.NumKeys);
+			const UniformKeyDataAdapter<RCCF_Cubic> KeyDataAdapter(CompressedKeys, KeyTimeAdapter);
+
+			const KeyDataHandle Handle = KeyDataAdapter.GetKeyDataHandle(Index);
+			return FRichCurveKey(KeyTimeAdapter.GetTime(Index), KeyDataAdapter.GetKeyValue(Handle), KeyDataAdapter.GetKeyArriveTangent(Handle), KeyDataAdapter.GetKeyLeaveTangent(Handle), RCIM_Cubic);
+		},
+	},
+	// RCCF_Mixed
+	{
+		// RCKTCF_uint16
+		[](FCompressedRichCurve::TConstantValueNumKeys ConstantValueNumKeys, const uint8* CompressedKeys, int32 Index)
+		{
+			const int32 InterpModesOffset = 0;
+			const int32 KeyTimesOffset = InterpModesOffset + Align(ConstantValueNumKeys.NumKeys * sizeof(uint8), sizeof(uint16));
+			const Quantized16BitKeyTimeAdapter KeyTimeAdapter(CompressedKeys, KeyTimesOffset, ConstantValueNumKeys.NumKeys);
+			const MixedKeyDataAdapter KeyDataAdapter(CompressedKeys, InterpModesOffset, KeyTimeAdapter);
+
+			const KeyDataHandle Handle = KeyDataAdapter.GetKeyDataHandle(Index);			
+			const FRichCurveKey Key(KeyTimeAdapter.GetTime(Index), KeyDataAdapter.GetKeyValue(Handle), KeyDataAdapter.GetKeyArriveTangent(Handle), KeyDataAdapter.GetKeyLeaveTangent(Handle),
+				GetKeyInterpMode(KeyDataAdapter.GetKeyInterpMode(Index)));
+			
+			return Key;
+		},
+		// RCKTCF_float32
+		[](FCompressedRichCurve::TConstantValueNumKeys ConstantValueNumKeys, const uint8* CompressedKeys, int32 Index)
+		{
+			const int32 InterpModesOffset = 0;
+			const int32 KeyTimesOffset = InterpModesOffset + Align(ConstantValueNumKeys.NumKeys * sizeof(uint8), sizeof(float));
+			const Float32BitKeyTimeAdapter KeyTimeAdapter(CompressedKeys, KeyTimesOffset, ConstantValueNumKeys.NumKeys);
+			const MixedKeyDataAdapter KeyDataAdapter(CompressedKeys, InterpModesOffset, KeyTimeAdapter);
+
+			const KeyDataHandle Handle = KeyDataAdapter.GetKeyDataHandle(Index);			
+			FRichCurveKey Key(KeyTimeAdapter.GetTime(Index), KeyDataAdapter.GetKeyValue(Handle), KeyDataAdapter.GetKeyArriveTangent(Handle), KeyDataAdapter.GetKeyLeaveTangent(Handle),
+				GetKeyInterpMode(KeyDataAdapter.GetKeyInterpMode(Index)));
+			
+			return Key;
+		},
+	},
+	// RCCF_Weighted
+	{
+		// RCKTCF_uint16
+		[](FCompressedRichCurve::TConstantValueNumKeys ConstantValueNumKeys, const uint8* CompressedKeys, int32 Index)
+		{
+			const int32 InterpModesOffset = 0;
+			const int32 KeyTimesOffset = InterpModesOffset + Align(2 * ConstantValueNumKeys.NumKeys * sizeof(uint8), sizeof(uint16));
+			const Quantized16BitKeyTimeAdapter KeyTimeAdapter(CompressedKeys, KeyTimesOffset, ConstantValueNumKeys.NumKeys);
+			const WeightedKeyDataAdapter KeyDataAdapter(CompressedKeys, InterpModesOffset, KeyTimeAdapter);
+
+			const KeyDataHandle Handle = KeyDataAdapter.GetKeyDataHandle(Index);
+			FRichCurveKey Key(KeyTimeAdapter.GetTime(Index), KeyDataAdapter.GetKeyValue(Handle), KeyDataAdapter.GetKeyArriveTangent(Handle), KeyDataAdapter.GetKeyLeaveTangent(Handle),
+				GetKeyInterpMode(KeyDataAdapter.GetKeyInterpMode(Index)));
+
+			Key.TangentMode = ERichCurveTangentMode::RCTM_User;
+			Key.TangentWeightMode = KeyDataAdapter.GetKeyTangentWeightMode(Index);
+			Key.ArriveTangentWeight = KeyDataAdapter.GetKeyArriveTangentWeight(Handle);
+			Key.LeaveTangentWeight = KeyDataAdapter.GetKeyLeaveTangentWeight(Handle);
+			
+			return Key;
+		},
+		// RCKTCF_float32
+		[](FCompressedRichCurve::TConstantValueNumKeys ConstantValueNumKeys, const uint8* CompressedKeys, int32 Index)
+		{
+			const int32 InterpModesOffset = 0;
+			const int32 KeyTimesOffset = InterpModesOffset + Align(2 * ConstantValueNumKeys.NumKeys * sizeof(uint8), sizeof(float));
+			const Float32BitKeyTimeAdapter KeyTimeAdapter(CompressedKeys, KeyTimesOffset, ConstantValueNumKeys.NumKeys);
+			const WeightedKeyDataAdapter KeyDataAdapter(CompressedKeys, InterpModesOffset, KeyTimeAdapter);
+
+			const KeyDataHandle Handle = KeyDataAdapter.GetKeyDataHandle(Index);
+			
+			FRichCurveKey Key(KeyTimeAdapter.GetTime(Index), KeyDataAdapter.GetKeyValue(Handle), KeyDataAdapter.GetKeyArriveTangent(Handle), KeyDataAdapter.GetKeyLeaveTangent(Handle),
+				GetKeyInterpMode(KeyDataAdapter.GetKeyInterpMode(Index)));
+
+			Key.TangentMode = ERichCurveTangentMode::RCTM_User;
+			Key.TangentWeightMode = KeyDataAdapter.GetKeyTangentWeightMode(Index);
+			Key.ArriveTangentWeight = KeyDataAdapter.GetKeyArriveTangentWeight(Handle);
+			Key.LeaveTangentWeight = KeyDataAdapter.GetKeyLeaveTangentWeight(Handle);
+
+			return Key;
+		},
+	},
+};
+
+void FCompressedRichCurve::PopulateCurve(FRichCurve& OutCurve) const
+{
+	if (CompressionFormat != RCCF_Constant)
+	{
+		TArray<FRichCurveKey> Keys;
+		for (int32 KeyIndex = 0; KeyIndex < ConstantValueNumKeys.NumKeys; ++KeyIndex)
+		{
+			Keys.Add(ToRichCurveMap[CompressionFormat][KeyTimeCompressionFormat](ConstantValueNumKeys, CompressedKeys.GetData(), KeyIndex));
+		}
+		OutCurve.SetKeys(Keys);
+	}
+}
+
 float FCompressedRichCurve::Eval(float InTime, float InDefaultValue) const
 {
 	// Dynamic dispatch into a template optimized code path
