Ensure WeightedKeyDataAdapter is correctly indexing into InterpModes array

diff --git a/Engine/Source/Runtime/Engine/Private/Curves/RichCurve.cpp b/Engine/Source/Runtime/Engine/Private/Curves/RichCurve.cpp
--- a/Engine/Source/Runtime/Engine/Private/Curves/RichCurve.cpp
+++ b/Engine/Source/Runtime/Engine/Private/Curves/RichCurve.cpp
@@ -1865,12 +1865,12 @@ struct WeightedKeyDataAdapter
 
 	constexpr ERichCurveCompressionFormat GetKeyInterpMode(int32 KeyIndex) const
 	{
-		return (ERichCurveCompressionFormat)InterpModes[KeyIndex];
+		return (ERichCurveCompressionFormat)InterpModes[(KeyIndex * 2)];
 	}
 
 	constexpr ERichCurveTangentWeightMode GetKeyTangentWeightMode(int32 KeyIndex) const
 	{
-		return (ERichCurveTangentWeightMode)InterpModes[KeyIndex+1];
+		return (ERichCurveTangentWeightMode)InterpModes[(KeyIndex * 2)+1];
 	}
 
 
