Crash: Sequencer Bake to Control Rig
#fix when evaluating raw animation poses

diff --git a/Engine/Source/Runtime/Engine/Private/Animation/AnimSequence.cpp b/Engine/Source/Runtime/Engine/Private/Animation/AnimSequence.cpp
--- a/Engine/Source/Runtime/Engine/Private/Animation/AnimSequence.cpp
+++ b/Engine/Source/Runtime/Engine/Private/Animation/AnimSequence.cpp
@@ -1496,16 +1496,18 @@ void UAnimSequence::GetBonePose(FAnimationPoseData& OutAnimationPoseData, const
 		if (bDisableRetargeting)
 		{
 			TArray<FTransform> const& AuthoredOnRefSkeleton = GetRetargetTransforms();
-			TArray<FBoneIndexType> const& RequireBonesIndexArray = RequiredBones.GetBoneIndicesArray();
-
-			int32 const NumRequiredBones = RequireBonesIndexArray.Num();
+			const int32 NumRawBones = RequiredBones.GetSkeletonAsset()->GetReferenceSkeleton().GetRawBoneNum();
 			for (FCompactPoseBoneIndex PoseBoneIndex : OutPose.ForEachBoneIndex())
 			{
-				int32 const& SkeletonBoneIndex = RequiredBones.GetSkeletonIndex(PoseBoneIndex);
-
+				const FSkeletonPoseBoneIndex SkeletonBoneIndex = RequiredBones.GetSkeletonPoseIndexFromCompactPoseIndex(PoseBoneIndex);
 				// Pose bone index should always exist in Skeleton
 				checkSlow(SkeletonBoneIndex != INDEX_NONE);
-				OutPose[PoseBoneIndex] = AuthoredOnRefSkeleton[SkeletonBoneIndex];
+
+				// Virtual bones are part of the retarget transform pose, so if the pose has not been updated (recently) there might be a mismatch
+				if (SkeletonBoneIndex < NumRawBones || AuthoredOnRefSkeleton.IsValidIndex(SkeletonBoneIndex.GetInt()))
+				{
+					OutPose[PoseBoneIndex] = AuthoredOnRefSkeleton[SkeletonBoneIndex.GetInt()];					
+				}
 			}
 		}
 		else
