Minor typo fixes.

diff --git a/Engine/Plugins/Experimental/ChaosCaching/Source/ChaosCaching/Private/Chaos/Adapters/CacheAdapter.cpp b/Engine/Plugins/Experimental/ChaosCaching/Source/ChaosCaching/Private/Chaos/Adapters/CacheAdapter.cpp
--- a/Engine/Plugins/Experimental/ChaosCaching/Source/ChaosCaching/Private/Chaos/Adapters/CacheAdapter.cpp
+++ b/Engine/Plugins/Experimental/ChaosCaching/Source/ChaosCaching/Private/Chaos/Adapters/CacheAdapter.cpp
@@ -7,11 +7,11 @@
 // Feature registration name for modular features
 const FName Chaos::FComponentCacheAdapter::FeatureName("ChaosCacheAdapter");
 
-// Engine adapters will always use priorities between EngineAdapterPriotityBegin and UserAdapterPriotityBegin
-const uint8 Chaos::FComponentCacheAdapter::EngineAdapterPriotityBegin(0);
+// Engine adapters will always use priorities between EngineAdapterPriorityBegin and UserAdapterPriorityBegin
+const uint8 Chaos::FComponentCacheAdapter::EngineAdapterPriorityBegin(0);
 
 // Any priority of this or above is a user implemented adapter
-const uint8 Chaos::FComponentCacheAdapter::UserAdapterPriotityBegin(1 << 3);
+const uint8 Chaos::FComponentCacheAdapter::UserAdapterPriorityBegin(1 << 3);
 
 DEFINE_LOG_CATEGORY(LogCacheAdapter);
 
diff --git a/Engine/Plugins/Experimental/ChaosCaching/Source/ChaosCaching/Private/Chaos/Adapters/GeometryCollectionComponentCacheAdapter.cpp b/Engine/Plugins/Experimental/ChaosCaching/Source/ChaosCaching/Private/Chaos/Adapters/GeometryCollectionComponentCacheAdapter.cpp
--- a/Engine/Plugins/Experimental/ChaosCaching/Source/ChaosCaching/Private/Chaos/Adapters/GeometryCollectionComponentCacheAdapter.cpp
+++ b/Engine/Plugins/Experimental/ChaosCaching/Source/ChaosCaching/Private/Chaos/Adapters/GeometryCollectionComponentCacheAdapter.cpp
@@ -40,7 +40,7 @@ namespace Chaos
 
 	uint8 FGeometryCollectionCacheAdapter::GetPriority() const
 	{
-		return EngineAdapterPriotityBegin;
+		return EngineAdapterPriorityBegin;
 	}
 	
 	void FGeometryCollectionCacheAdapter::Record_PostSolve(UPrimitiveComponent* InComp, const FTransform& InRootTransform, FPendingFrameWrite& OutFrame, Chaos::FReal InTime) const
diff --git a/Engine/Plugins/Experimental/ChaosCaching/Source/ChaosCaching/Private/Chaos/Adapters/StaticMeshComponentCacheAdapter.cpp b/Engine/Plugins/Experimental/ChaosCaching/Source/ChaosCaching/Private/Chaos/Adapters/StaticMeshComponentCacheAdapter.cpp
--- a/Engine/Plugins/Experimental/ChaosCaching/Source/ChaosCaching/Private/Chaos/Adapters/StaticMeshComponentCacheAdapter.cpp
+++ b/Engine/Plugins/Experimental/ChaosCaching/Source/ChaosCaching/Private/Chaos/Adapters/StaticMeshComponentCacheAdapter.cpp
@@ -33,7 +33,7 @@ namespace Chaos
 
 	uint8 FStaticMeshCacheAdapter::GetPriority() const
 	{
-		return EngineAdapterPriotityBegin;
+		return EngineAdapterPriorityBegin;
 	}
 
 	void RecordToCacheInternal(FSingleParticlePhysicsProxy* InProxy, const FTransform& InRootTransform, FPendingFrameWrite& OutFrame, Chaos::FReal InTime)
diff --git a/Engine/Plugins/Experimental/ChaosCaching/Source/ChaosCaching/Public/Chaos/Adapters/CacheAdapter.h b/Engine/Plugins/Experimental/ChaosCaching/Source/ChaosCaching/Public/Chaos/Adapters/CacheAdapter.h
--- a/Engine/Plugins/Experimental/ChaosCaching/Source/ChaosCaching/Public/Chaos/Adapters/CacheAdapter.h
+++ b/Engine/Plugins/Experimental/ChaosCaching/Source/ChaosCaching/Public/Chaos/Adapters/CacheAdapter.h
@@ -49,8 +49,8 @@ namespace Chaos
 
 		/** Registration name for modular features module */
 		static const FName FeatureName;
-		static const uint8 EngineAdapterPriotityBegin;
-		static const uint8 UserAdapterPriotityBegin;
+		static const uint8 EngineAdapterPriorityBegin;
+		static const uint8 UserAdapterPriorityBegin;
 
 		FComponentCacheAdapter()          = default;
 		virtual ~FComponentCacheAdapter() = default;
@@ -98,7 +98,7 @@ namespace Chaos
 		 * Gets the priority for an adapter.
 		 * When two or more adapters give the same support level for a given component class this priority will be
 		 * used to decide which adapter will be used. All base engine level adapters per type will use a priority
-		 * between EngineAdapterPriotityBegin and UserAdapterPriorityBegin defined in this interface.
+		 * between EngineAdapterPriorityBegin and UserAdapterPriorityBegin defined in this interface.
 		 *
 		 * Users implementing adapters intended to override all engine functionality should return priorities
 		 * above UserAdapterPriorityBegin to ensure they will always be selected above engine adapters.
diff --git a/Engine/Source/Runtime/Engine/Private/WorldPartition/WorldPartitionDebugHelper.cpp b/Engine/Source/Runtime/Engine/Private/WorldPartition/WorldPartitionDebugHelper.cpp
--- a/Engine/Source/Runtime/Engine/Private/WorldPartition/WorldPartitionDebugHelper.cpp
+++ b/Engine/Source/Runtime/Engine/Private/WorldPartition/WorldPartitionDebugHelper.cpp
@@ -180,10 +180,10 @@ void FWorldPartitionDebugHelper::DrawLegendItem(UCanvas* Canvas, const FString&
 	Pos.Y = TextPos.Y;
 }
 
-bool FWorldPartitionDebugHelper::bShowRuntimeSpatialHashCellStreamingPriotity = false;
-FAutoConsoleVariableRef FWorldPartitionDebugHelper::ShowRuntimeSpatialHashCellStreamingPriotityCommand(
-	TEXT("wp.Runtime.ShowRuntimeSpatialHashCellStreamingPriotity"),
-	FWorldPartitionDebugHelper::bShowRuntimeSpatialHashCellStreamingPriotity,
+bool FWorldPartitionDebugHelper::bShowRuntimeSpatialHashCellStreamingPriority = false;
+FAutoConsoleVariableRef FWorldPartitionDebugHelper::ShowRuntimeSpatialHashCellStreamingPriorityCommand(
+	TEXT("wp.Runtime.ShowRuntimeSpatialHashCellStreamingPriority"),
+	FWorldPartitionDebugHelper::bShowRuntimeSpatialHashCellStreamingPriority,
 	TEXT("Enable to show a heatmap of the runtime spatial hash grid cells based on their priority."));
 
 FLinearColor FWorldPartitionDebugHelper::GetHeatMapColor(float ValueNormalized)
diff --git a/Engine/Source/Runtime/Engine/Private/WorldPartition/WorldPartitionRuntimeSpatialHash.cpp b/Engine/Source/Runtime/Engine/Private/WorldPartition/WorldPartitionRuntimeSpatialHash.cpp
--- a/Engine/Source/Runtime/Engine/Private/WorldPartition/WorldPartitionRuntimeSpatialHash.cpp
+++ b/Engine/Source/Runtime/Engine/Private/WorldPartition/WorldPartitionRuntimeSpatialHash.cpp
@@ -400,7 +400,7 @@ void FSpatialHashStreamingGrid::GetFilteredCellsForDebugDraw(const FSpatialHashS
 
 EWorldPartitionRuntimeCellVisualizeMode FSpatialHashStreamingGrid::GetStreamingCellVisualizeMode() const
 {
-	const EWorldPartitionRuntimeCellVisualizeMode VisualizeMode = FWorldPartitionDebugHelper::IsRuntimeSpatialHashCellStreamingPriotityShown() ? EWorldPartitionRuntimeCellVisualizeMode::StreamingPriority : EWorldPartitionRuntimeCellVisualizeMode::StreamingStatus;
+	const EWorldPartitionRuntimeCellVisualizeMode VisualizeMode = FWorldPartitionDebugHelper::IsRuntimeSpatialHashCellStreamingPriorityShown() ? EWorldPartitionRuntimeCellVisualizeMode::StreamingPriority : EWorldPartitionRuntimeCellVisualizeMode::StreamingStatus;
 	return VisualizeMode;
 }
 
diff --git a/Engine/Source/Runtime/Engine/Private/WorldPartition/WorldPartitionStreamingPolicy.cpp b/Engine/Source/Runtime/Engine/Private/WorldPartition/WorldPartitionStreamingPolicy.cpp
--- a/Engine/Source/Runtime/Engine/Private/WorldPartition/WorldPartitionStreamingPolicy.cpp
+++ b/Engine/Source/Runtime/Engine/Private/WorldPartition/WorldPartitionStreamingPolicy.cpp
@@ -355,7 +355,7 @@ void UWorldPartitionStreamingPolicy::UpdateStreamingState()
 #if !UE_BUILD_SHIPPING
 void UWorldPartitionStreamingPolicy::UpdateDebugCellsStreamingPriority(const TSet<const UWorldPartitionRuntimeCell*>& ActivateStreamingCells, const TSet<const UWorldPartitionRuntimeCell*>& LoadStreamingCells)
 {
-	if (FWorldPartitionDebugHelper::IsRuntimeSpatialHashCellStreamingPriotityShown())
+	if (FWorldPartitionDebugHelper::IsRuntimeSpatialHashCellStreamingPriorityShown())
 	{
 		TSet<const UWorldPartitionRuntimeCell*> CellsToSort;
 		CellsToSort.Append(ActivateStreamingCells);
diff --git a/Engine/Source/Runtime/Engine/Public/WorldPartition/WorldPartitionDebugHelper.h b/Engine/Source/Runtime/Engine/Public/WorldPartition/WorldPartitionDebugHelper.h
--- a/Engine/Source/Runtime/Engine/Public/WorldPartition/WorldPartitionDebugHelper.h
+++ b/Engine/Source/Runtime/Engine/Public/WorldPartition/WorldPartitionDebugHelper.h
@@ -23,7 +23,7 @@ public:
 	static void DrawText(UCanvas* Canvas, const FString& Text, const UFont* Font, const FColor& Color, FVector2D& Pos, float* MaxTextWidth = nullptr);
 	static void DrawLegendItem(UCanvas* Canvas, const FString& Text, const UFont* Font, const FColor& Color, const FColor& TextColor, FVector2D& Pos, float* MaxItemWidth = nullptr);
 	static FLinearColor GetHeatMapColor(float Value);
-	FORCEINLINE static bool IsRuntimeSpatialHashCellStreamingPriotityShown() { return bShowRuntimeSpatialHashCellStreamingPriotity; }
+	FORCEINLINE static bool IsRuntimeSpatialHashCellStreamingPriorityShown() { return bShowRuntimeSpatialHashCellStreamingPriority; }
 
 private:
 	static FAutoConsoleCommand DebugFilterByStreamingStatusCommand;
@@ -38,6 +38,6 @@ private:
 	static FAutoConsoleCommand DebugFilterByCellNameCommand;
 	static TArray<FString> DebugCellNameFilter;
 
-	static FAutoConsoleVariableRef ShowRuntimeSpatialHashCellStreamingPriotityCommand;
-	static bool bShowRuntimeSpatialHashCellStreamingPriotity;
+	static FAutoConsoleVariableRef ShowRuntimeSpatialHashCellStreamingPriorityCommand;
+	static bool bShowRuntimeSpatialHashCellStreamingPriority;
 };
\ No newline at end of file
