Inline edit support to material expression inputs (part 2)
- Support for multiple UPROPERTYs associated with the same input pin forming vector parameters, used by DecalMipmapLevel and Panner nodes
- Inline edit support for UPROPERTY of type double
- Remove Min/Max from Clamp node caption
- Update the rest of material expressions to reflect the changes.

diff --git a/Engine/Source/Runtime/Engine/Classes/Materials/MaterialExpression.h b/Engine/Source/Runtime/Engine/Classes/Materials/MaterialExpression.h
--- a/Engine/Source/Runtime/Engine/Classes/Materials/MaterialExpression.h
+++ b/Engine/Source/Runtime/Engine/Classes/Materials/MaterialExpression.h
@@ -351,7 +351,7 @@ class ENGINE_API UMaterialExpression : public UObject
 	};
 
 	/** Find the property that is associated with the input pin. */
-	virtual FProperty* GetInputPinProperty(int32 PinIndex);
+	virtual TArray<FProperty*> GetInputPinProperty(int32 PinIndex);
 	virtual FName GetInputPinSubCategory(int32 PinIndex);
 	virtual UObject* GetInputPinSubCategoryObject(int32 PinIndex);
 	virtual void PinDefaultValueChanged(int32 PinIndex, const FString& DefaultValue);
diff --git a/Engine/Source/Runtime/Engine/Classes/Materials/MaterialExpressionGenericConstant.h b/Engine/Source/Runtime/Engine/Classes/Materials/MaterialExpressionGenericConstant.h
--- a/Engine/Source/Runtime/Engine/Classes/Materials/MaterialExpressionGenericConstant.h
+++ b/Engine/Source/Runtime/Engine/Classes/Materials/MaterialExpressionGenericConstant.h
@@ -29,7 +29,7 @@ class UMaterialExpressionConstantDouble : public UMaterialExpressionGenericConst
 {
 	GENERATED_UCLASS_BODY()
 
-	UPROPERTY(EditAnywhere, BlueprintReadWrite, Category = MaterialExpressionConstant)
+	UPROPERTY(EditAnywhere, BlueprintReadWrite, Category = MaterialExpressionConstant, meta = (ShowAsInputPin = "Primary"))
 	double Value;
 
 	virtual UE::Shader::FValue GetConstantValue() const override { return Value; }
diff --git a/Engine/Source/Runtime/Engine/Classes/Materials/MaterialExpressionIf.h b/Engine/Source/Runtime/Engine/Classes/Materials/MaterialExpressionIf.h
--- a/Engine/Source/Runtime/Engine/Classes/Materials/MaterialExpressionIf.h
+++ b/Engine/Source/Runtime/Engine/Classes/Materials/MaterialExpressionIf.h
@@ -29,11 +29,11 @@ class UMaterialExpressionIf : public UMaterialExpression
 	UPROPERTY()
 	FExpressionInput ALessThanB;
 
-	UPROPERTY(EditAnywhere, Category=MaterialExpressionIf)
+	UPROPERTY(EditAnywhere, Category=MaterialExpressionIf, meta = (ShowAsInputPin = "Advanced"))
 	float EqualsThreshold;
 
 	/** only used if B is not hooked up */
-	UPROPERTY(EditAnywhere, Category = MaterialExpressionIf)
+	UPROPERTY(EditAnywhere, Category = MaterialExpressionIf, meta = (OverridingInputProperty = "B"))
 	float ConstB;
 
 	UPROPERTY()
diff --git a/Engine/Source/Runtime/Engine/Classes/Materials/MaterialExpressionPanner.h b/Engine/Source/Runtime/Engine/Classes/Materials/MaterialExpressionPanner.h
--- a/Engine/Source/Runtime/Engine/Classes/Materials/MaterialExpressionPanner.h
+++ b/Engine/Source/Runtime/Engine/Classes/Materials/MaterialExpressionPanner.h
@@ -23,14 +23,14 @@ class UMaterialExpressionPanner : public UMaterialExpression
 	UPROPERTY(meta = (RequiredInput = "false", ToolTip = "Vector2 speed scale, if specified"))
 	FExpressionInput Speed;
 
-	UPROPERTY(EditAnywhere, Category=MaterialExpressionPanner)
+	UPROPERTY(EditAnywhere, Category=MaterialExpressionPanner, meta = (OverridingInputProperty = "Speed"))
 	float SpeedX;
 
-	UPROPERTY(EditAnywhere, Category=MaterialExpressionPanner)
+	UPROPERTY(EditAnywhere, Category=MaterialExpressionPanner, meta = (OverridingInputProperty = "Speed"))
 	float SpeedY;
 
 	/** only used if Coordinate is not hooked up */
-	UPROPERTY(EditAnywhere, Category = MaterialExpressionPanner)
+	UPROPERTY(EditAnywhere, Category = MaterialExpressionPanner, meta = (OverridingInputProperty = "Coordinate"))
 	uint32 ConstCoordinate;
 
 	// Output only the fractional part of the pan calculation for greater precision.
diff --git a/Engine/Source/Runtime/Engine/Classes/Materials/MaterialExpressionPerInstanceCustomData.h b/Engine/Source/Runtime/Engine/Classes/Materials/MaterialExpressionPerInstanceCustomData.h
--- a/Engine/Source/Runtime/Engine/Classes/Materials/MaterialExpressionPerInstanceCustomData.h
+++ b/Engine/Source/Runtime/Engine/Classes/Materials/MaterialExpressionPerInstanceCustomData.h
@@ -18,7 +18,7 @@ class UMaterialExpressionPerInstanceCustomData : public UMaterialExpression
 	UPROPERTY(EditAnywhere, Category = MaterialExpressionPerInstanceCustomData, meta = (OverridingInputProperty = "DefaultValue"))
 	float ConstDefaultValue;
 
-	UPROPERTY(EditAnywhere, Category = MaterialExpressionPerInstanceCustomData)
+	UPROPERTY(EditAnywhere, Category = MaterialExpressionPerInstanceCustomData, meta = (ShowAsInputPin = "Advanced"))
 	uint32 DataIndex;
 
 	//~ Begin UMaterialExpression Interface
@@ -40,7 +40,7 @@ class UMaterialExpressionPerInstanceCustomData3Vector : public UMaterialExpressi
 	UPROPERTY(EditAnywhere, Category = MaterialExpressionPerInstanceCustomData, meta = (OverridingInputProperty = "DefaultValue"))
 	FLinearColor ConstDefaultValue;
 
-	UPROPERTY(EditAnywhere, Category = MaterialExpressionPerInstanceCustomData)
+	UPROPERTY(EditAnywhere, Category = MaterialExpressionPerInstanceCustomData, meta = (ShowAsInputPin = "Advanced"))
 	uint32 DataIndex;
 
 	//~ Begin UMaterialExpression Interface
diff --git a/Engine/Source/Runtime/Engine/Classes/Materials/MaterialExpressionReflectionVectorWS.h b/Engine/Source/Runtime/Engine/Classes/Materials/MaterialExpressionReflectionVectorWS.h
--- a/Engine/Source/Runtime/Engine/Classes/Materials/MaterialExpressionReflectionVectorWS.h
+++ b/Engine/Source/Runtime/Engine/Classes/Materials/MaterialExpressionReflectionVectorWS.h
@@ -19,7 +19,7 @@ class UMaterialExpressionReflectionVectorWS : public UMaterialExpression
 	FExpressionInput CustomWorldNormal;
 
 	/** (true): The specified world normal will be normalized. (false): WorldNormal will just be used as is, faster but possible artifacts if normal length isn't 1 */
-	UPROPERTY(EditAnywhere, Category=UMaterialExpressionReflectionVectorWS, meta=(DisplayName = "Normalize custom world normal"))
+	UPROPERTY(EditAnywhere, Category=UMaterialExpressionReflectionVectorWS, meta=(DisplayName = "Normalize custom world normal", ShowAsInputPin = "Advanced"))
 	uint32 bNormalizeCustomWorldNormal : 1;
 
 	//~ Begin UMaterialExpression Interface
diff --git a/Engine/Source/Runtime/Engine/Classes/Materials/MaterialExpressionRotateAboutAxis.h b/Engine/Source/Runtime/Engine/Classes/Materials/MaterialExpressionRotateAboutAxis.h
--- a/Engine/Source/Runtime/Engine/Classes/Materials/MaterialExpressionRotateAboutAxis.h
+++ b/Engine/Source/Runtime/Engine/Classes/Materials/MaterialExpressionRotateAboutAxis.h
@@ -26,7 +26,7 @@ class UMaterialExpressionRotateAboutAxis : public UMaterialExpression
 	UPROPERTY()
 	FExpressionInput Position;
 
-	UPROPERTY(EditAnywhere, Category=MaterialExpressionRotateAboutAxis)
+	UPROPERTY(EditAnywhere, Category=MaterialExpressionRotateAboutAxis, meta = (ShowAsInputPin = "Advanced"))
 	float Period;
 
 
diff --git a/Engine/Source/Runtime/Engine/Classes/Materials/MaterialExpressionRotator.h b/Engine/Source/Runtime/Engine/Classes/Materials/MaterialExpressionRotator.h
--- a/Engine/Source/Runtime/Engine/Classes/Materials/MaterialExpressionRotator.h
+++ b/Engine/Source/Runtime/Engine/Classes/Materials/MaterialExpressionRotator.h
@@ -20,17 +20,17 @@ class UMaterialExpressionRotator : public UMaterialExpression
 	UPROPERTY(meta = (RequiredInput = "false", ToolTip = "Defaults to Game Time if not specified"))
 	FExpressionInput Time;
 
-	UPROPERTY(EditAnywhere, Category=MaterialExpressionRotator)
+	UPROPERTY(EditAnywhere, Category=MaterialExpressionRotator, meta = (ShowAsInputPin = "Advanced"))
 	float CenterX;
 
-	UPROPERTY(EditAnywhere, Category=MaterialExpressionRotator)
+	UPROPERTY(EditAnywhere, Category=MaterialExpressionRotator, meta = (ShowAsInputPin = "Advanced"))
 	float CenterY;
 
-	UPROPERTY(EditAnywhere, Category=MaterialExpressionRotator)
+	UPROPERTY(EditAnywhere, Category=MaterialExpressionRotator, meta = (ShowAsInputPin = "Advanced"))
 	float Speed;
 
 	/** only used if Coordinate is not hooked up */
-	UPROPERTY(EditAnywhere, Category = MaterialExpressionRotator)
+	UPROPERTY(EditAnywhere, Category = MaterialExpressionRotator, meta = (OverridingInputProperty = "Coordinate"))
 	uint32 ConstCoordinate;
 
 	//~ Begin UMaterialExpression Interface
diff --git a/Engine/Source/Runtime/Engine/Classes/Materials/MaterialExpressionSamplePhysicsField.h b/Engine/Source/Runtime/Engine/Classes/Materials/MaterialExpressionSamplePhysicsField.h
--- a/Engine/Source/Runtime/Engine/Classes/Materials/MaterialExpressionSamplePhysicsField.h
+++ b/Engine/Source/Runtime/Engine/Classes/Materials/MaterialExpressionSamplePhysicsField.h
@@ -22,7 +22,7 @@ class UMaterialExpressionSamplePhysicsVectorField : public UMaterialExpression
 	FExpressionInput WorldPosition;
 
 	/** Target Type to be accessed */
-	UPROPERTY(EditAnywhere, Category = "Field Setup", meta = (DisplayName = "Target Type"))
+	UPROPERTY(EditAnywhere, Category = "Field Setup", meta = (DisplayName = "Target Type", ShowAsInputPin = "Advanced"))
 	TEnumAsByte<EFieldVectorType> FieldTarget;
 
 	//~ Begin UMaterialExpression Interface
@@ -42,7 +42,7 @@ class UMaterialExpressionSamplePhysicsScalarField : public UMaterialExpression
 	FExpressionInput WorldPosition;
 
 	/** Target Type to be accessed */
-	UPROPERTY(EditAnywhere, Category = "Field Setup", meta = (DisplayName = "Target Type"))
+	UPROPERTY(EditAnywhere, Category = "Field Setup", meta = (DisplayName = "Target Type", ShowAsInputPin = "Advanced"))
 	TEnumAsByte<EFieldScalarType> FieldTarget;
 
 	//~ Begin UMaterialExpression Interface
@@ -62,7 +62,7 @@ class UMaterialExpressionSamplePhysicsIntegerField : public UMaterialExpression
 	FExpressionInput WorldPosition;
 
 	/** Target Type to be accessed */
-	UPROPERTY(EditAnywhere, Category = "Field Setup", meta = (DisplayName = "Target Type"))
+	UPROPERTY(EditAnywhere, Category = "Field Setup", meta = (DisplayName = "Target Type", ShowAsInputPin = "Advanced"))
 	TEnumAsByte<EFieldIntegerType> FieldTarget;
 
 	//~ Begin UMaterialExpression Interface
diff --git a/Engine/Source/Runtime/Engine/Classes/Materials/MaterialExpressionSceneColor.h b/Engine/Source/Runtime/Engine/Classes/Materials/MaterialExpressionSceneColor.h
--- a/Engine/Source/Runtime/Engine/Classes/Materials/MaterialExpressionSceneColor.h
+++ b/Engine/Source/Runtime/Engine/Classes/Materials/MaterialExpressionSceneColor.h
@@ -28,7 +28,7 @@ class UMaterialExpressionSceneColor : public UMaterialExpression
 	* Coordinates - UV coordinates to apply to the scene color lookup.
 	* OffsetFraction - 	An offset to apply to the scene color lookup in a 2d fraction of the screen.
 	*/ 
-	UPROPERTY(EditAnywhere, Category=MaterialExpressionSceneColor)
+	UPROPERTY(EditAnywhere, Category=MaterialExpressionSceneColor, meta = (ShowAsInputPin = "Advanced"))
 	TEnumAsByte<enum EMaterialSceneAttributeInputMode::Type> InputMode;
 
 	/**
diff --git a/Engine/Source/Runtime/Engine/Classes/Materials/MaterialExpressionSceneDepth.h b/Engine/Source/Runtime/Engine/Classes/Materials/MaterialExpressionSceneDepth.h
--- a/Engine/Source/Runtime/Engine/Classes/Materials/MaterialExpressionSceneDepth.h
+++ b/Engine/Source/Runtime/Engine/Classes/Materials/MaterialExpressionSceneDepth.h
@@ -19,7 +19,7 @@ class UMaterialExpressionSceneDepth : public UMaterialExpression
 	* Coordinates - UV coordinates to apply to the scene depth lookup.
 	* OffsetFraction - An offset to apply to the scene depth lookup in a 2d fraction of the screen.
 	*/ 
-	UPROPERTY(EditAnywhere, Category=MaterialExpressionSceneDepth)
+	UPROPERTY(EditAnywhere, Category=MaterialExpressionSceneDepth, meta = (ShowAsInputPin = "Advanced"))
 	TEnumAsByte<enum EMaterialSceneAttributeInputMode::Type> InputMode;
 
 	/**
@@ -34,7 +34,7 @@ class UMaterialExpressionSceneDepth : public UMaterialExpression
 	FExpressionInput Coordinates_DEPRECATED;
 
 	/** only used if Input is not hooked up */
-	UPROPERTY(EditAnywhere, Category = MaterialExpressionSceneDepth)
+	UPROPERTY(EditAnywhere, Category = MaterialExpressionSceneDepth, meta = (OverridingInputProperty = "Input"))
 	FVector2D ConstInput;
 
 	//~ Begin UObject Interface.
diff --git a/Engine/Source/Runtime/Engine/Classes/Materials/MaterialExpressionSceneDepthWithoutWater.h b/Engine/Source/Runtime/Engine/Classes/Materials/MaterialExpressionSceneDepthWithoutWater.h
--- a/Engine/Source/Runtime/Engine/Classes/Materials/MaterialExpressionSceneDepthWithoutWater.h
+++ b/Engine/Source/Runtime/Engine/Classes/Materials/MaterialExpressionSceneDepthWithoutWater.h
@@ -22,7 +22,7 @@ public:
 	* Coordinates - UV coordinates to apply to the scene depth lookup.
 	* OffsetFraction - An offset to apply to the scene depth lookup in a 2d fraction of the screen.
 	*/
-	UPROPERTY(EditAnywhere, Category = MaterialExpressionSceneDepthWithoutWater)
+	UPROPERTY(EditAnywhere, Category = MaterialExpressionSceneDepthWithoutWater, meta = (ShowAsInputPin = "Advanced"))
 	TEnumAsByte<enum EMaterialSceneAttributeInputMode::Type> InputMode;
 
 	/**
@@ -34,11 +34,11 @@ public:
 	FExpressionInput Input;
 
 	/** only used if Input is not hooked up */
-	UPROPERTY(EditAnywhere, Category = MaterialExpressionSceneDepthWithoutWater)
+	UPROPERTY(EditAnywhere, Category = MaterialExpressionSceneDepthWithoutWater, meta = (OverridingInputProperty = "Input"))
 	FVector2D ConstInput;
 
 	/** Depth to fall back to in case the needed texture isn't available on a particular platform or configuration */
-	UPROPERTY(EditAnywhere, Category = MaterialExpressionSceneDepthWithoutWater)
+	UPROPERTY(EditAnywhere, Category = MaterialExpressionSceneDepthWithoutWater, meta = (ShowAsInputPin = "Advanced"))
 	float FallbackDepth = 1000000.0f;
 
 	//~ Begin UMaterialExpression Interface
diff --git a/Engine/Source/Runtime/Engine/Classes/Materials/MaterialExpressionSceneTexture.h b/Engine/Source/Runtime/Engine/Classes/Materials/MaterialExpressionSceneTexture.h
--- a/Engine/Source/Runtime/Engine/Classes/Materials/MaterialExpressionSceneTexture.h
+++ b/Engine/Source/Runtime/Engine/Classes/Materials/MaterialExpressionSceneTexture.h
@@ -20,11 +20,11 @@ class UMaterialExpressionSceneTexture : public UMaterialExpression
 	FExpressionInput Coordinates;
 
 	/** Which scene texture (screen aligned texture) we want to make a lookup into */
-	UPROPERTY(EditAnywhere, Category=UMaterialExpressionSceneTexture, meta=(DisplayName = "Scene Texture Id"))
+	UPROPERTY(EditAnywhere, Category=UMaterialExpressionSceneTexture, meta=(DisplayName = "Scene Texture Id", ShowAsInputPin = "Advanced"))
 	TEnumAsByte<ESceneTextureId> SceneTextureId;
 	
 	/** Whether to use point sampled texture lookup (default) or using [bi-linear] filtered (can be slower, avoid faceted lock with distortions), some SceneTextures cannot be filtered */
-	UPROPERTY(EditAnywhere, Category=UMaterialExpressionSceneTexture, meta=(DisplayName = "Filtered"))
+	UPROPERTY(EditAnywhere, Category=UMaterialExpressionSceneTexture, meta=(DisplayName = "Filtered", ShowAsInputPin = "Advanced"))
 	bool bFiltered;
 
 	//~ Begin UMaterialExpression Interface
diff --git a/Engine/Source/Runtime/Engine/Classes/Materials/MaterialExpressionShadingModel.h b/Engine/Source/Runtime/Engine/Classes/Materials/MaterialExpressionShadingModel.h
--- a/Engine/Source/Runtime/Engine/Classes/Materials/MaterialExpressionShadingModel.h
+++ b/Engine/Source/Runtime/Engine/Classes/Materials/MaterialExpressionShadingModel.h
@@ -31,7 +31,7 @@ class UMaterialExpressionShadingModel : public UMaterialExpression
 	virtual bool GenerateHLSLExpression(FMaterialHLSLGenerator& Generator, UE::HLSLTree::FScope& Scope, int32 OutputIndex, UE::HLSLTree::FExpression const*& OutExpression) const override;
 #endif
 public:
-	UPROPERTY(EditAnywhere, Category=ShadingModel,  meta=(ValidEnumValues="MSM_DefaultLit, MSM_Subsurface, MSM_PreintegratedSkin, MSM_ClearCoat, MSM_SubsurfaceProfile, MSM_TwoSidedFoliage, MSM_Hair, MSM_Cloth, MSM_Eye"))
+	UPROPERTY(EditAnywhere, Category=ShadingModel,  meta=(ValidEnumValues="MSM_DefaultLit, MSM_Subsurface, MSM_PreintegratedSkin, MSM_ClearCoat, MSM_SubsurfaceProfile, MSM_TwoSidedFoliage, MSM_Hair, MSM_Cloth, MSM_Eye", ShowAsInputPin = "Primary"))
 	TEnumAsByte<enum EMaterialShadingModel> ShadingModel = MSM_DefaultLit;
 	//~ End UMaterialExpression Interface
 };
diff --git a/Engine/Source/Runtime/Engine/Classes/Materials/MaterialExpressionSkyAtmosphereLightDirection.h b/Engine/Source/Runtime/Engine/Classes/Materials/MaterialExpressionSkyAtmosphereLightDirection.h
--- a/Engine/Source/Runtime/Engine/Classes/Materials/MaterialExpressionSkyAtmosphereLightDirection.h
+++ b/Engine/Source/Runtime/Engine/Classes/Materials/MaterialExpressionSkyAtmosphereLightDirection.h
@@ -14,7 +14,7 @@ class UMaterialExpressionSkyAtmosphereLightDirection : public UMaterialExpressio
 	GENERATED_UCLASS_BODY()
 
 	/** Index of the atmosphere light to sample. */
-	UPROPERTY(EditAnywhere, BlueprintReadWrite, Category = MaterialExpressionTextureCoordinate, meta = (UIMin = "0", UIMax = "1", ClampMin = "0", ClampMax = "1"))
+	UPROPERTY(EditAnywhere, BlueprintReadWrite, Category = MaterialExpressionTextureCoordinate, meta = (UIMin = "0", UIMax = "1", ClampMin = "0", ClampMax = "1", ShowAsInputPin = "Primary"))
 	int32 LightIndex;
 
 	//~ Begin UMaterialExpression Interface
diff --git a/Engine/Source/Runtime/Engine/Classes/Materials/MaterialExpressionSkyAtmosphereLightIlluminance.h b/Engine/Source/Runtime/Engine/Classes/Materials/MaterialExpressionSkyAtmosphereLightIlluminance.h
--- a/Engine/Source/Runtime/Engine/Classes/Materials/MaterialExpressionSkyAtmosphereLightIlluminance.h
+++ b/Engine/Source/Runtime/Engine/Classes/Materials/MaterialExpressionSkyAtmosphereLightIlluminance.h
@@ -14,7 +14,7 @@ class UMaterialExpressionSkyAtmosphereLightIlluminance : public UMaterialExpress
 	GENERATED_UCLASS_BODY()
 
 	/** Index of the atmosphere light to sample. */
-	UPROPERTY(EditAnywhere, BlueprintReadWrite, Category = MaterialExpressionTextureCoordinate, meta = (UIMin = "0", UIMax = "1", ClampMin = "0", ClampMax = "1"))
+	UPROPERTY(EditAnywhere, BlueprintReadWrite, Category = MaterialExpressionTextureCoordinate, meta = (UIMin = "0", UIMax = "1", ClampMin = "0", ClampMax = "1", ShowAsInputPin = "Primary"))
 	int32 LightIndex;
 
 	/** World position of the sample. If not specified, the pixel world position will be used. */
@@ -37,7 +37,7 @@ class UMaterialExpressionSkyAtmosphereLightDiskLuminance : public UMaterialExpre
 	GENERATED_UCLASS_BODY()
 
 	/** Index of the atmosphere light to sample. */
-	UPROPERTY(EditAnywhere, BlueprintReadWrite, Category = MaterialExpressionTextureCoordinate, meta = (UIMin = "0", UIMax = "1", ClampMin = "0", ClampMax = "1"))
+	UPROPERTY(EditAnywhere, BlueprintReadWrite, Category = MaterialExpressionTextureCoordinate, meta = (UIMin = "0", UIMax = "1", ClampMin = "0", ClampMax = "1", ShowAsInputPin = "Primary"))
 	int32 LightIndex;
 
 	//~ Begin UMaterialExpression Interface
diff --git a/Engine/Source/Runtime/Engine/Classes/Materials/MaterialExpressionStaticBool.h b/Engine/Source/Runtime/Engine/Classes/Materials/MaterialExpressionStaticBool.h
--- a/Engine/Source/Runtime/Engine/Classes/Materials/MaterialExpressionStaticBool.h
+++ b/Engine/Source/Runtime/Engine/Classes/Materials/MaterialExpressionStaticBool.h
@@ -13,7 +13,7 @@ class UMaterialExpressionStaticBool : public UMaterialExpression
 {
 	GENERATED_UCLASS_BODY()
 
-	UPROPERTY(EditAnywhere, Category=MaterialExpressionStaticBool)
+	UPROPERTY(EditAnywhere, Category=MaterialExpressionStaticBool, meta = (ShowAsInputPin = "Primary"))
 	uint32 Value:1;
 
 
diff --git a/Engine/Source/Runtime/Engine/Classes/Materials/MaterialExpressionStaticBoolParameter.h b/Engine/Source/Runtime/Engine/Classes/Materials/MaterialExpressionStaticBoolParameter.h
--- a/Engine/Source/Runtime/Engine/Classes/Materials/MaterialExpressionStaticBoolParameter.h
+++ b/Engine/Source/Runtime/Engine/Classes/Materials/MaterialExpressionStaticBoolParameter.h
@@ -13,7 +13,7 @@ class UMaterialExpressionStaticBoolParameter : public UMaterialExpressionParamet
 {
 	GENERATED_UCLASS_BODY()
 
-	UPROPERTY(EditAnywhere, Category=MaterialExpressionStaticBoolParameter)
+	UPROPERTY(EditAnywhere, Category=MaterialExpressionStaticBoolParameter, meta = (ShowAsInputPin = "Primary"))
 	uint32 DefaultValue:1;
 
 public:
diff --git a/Engine/Source/Runtime/Engine/Classes/Materials/MaterialExpressionStaticComponentMaskParameter.h b/Engine/Source/Runtime/Engine/Classes/Materials/MaterialExpressionStaticComponentMaskParameter.h
--- a/Engine/Source/Runtime/Engine/Classes/Materials/MaterialExpressionStaticComponentMaskParameter.h
+++ b/Engine/Source/Runtime/Engine/Classes/Materials/MaterialExpressionStaticComponentMaskParameter.h
@@ -17,16 +17,16 @@ class UMaterialExpressionStaticComponentMaskParameter : public UMaterialExpressi
 	UPROPERTY()
 	FExpressionInput Input;
 
-	UPROPERTY(EditAnywhere, Category=MaterialExpressionStaticComponentMaskParameter)
+	UPROPERTY(EditAnywhere, Category=MaterialExpressionStaticComponentMaskParameter, meta = (ShowAsInputPin = "Advanced"))
 	uint32 DefaultR:1;
 
-	UPROPERTY(EditAnywhere, Category=MaterialExpressionStaticComponentMaskParameter)
+	UPROPERTY(EditAnywhere, Category=MaterialExpressionStaticComponentMaskParameter, meta = (ShowAsInputPin = "Advanced"))
 	uint32 DefaultG:1;
 
-	UPROPERTY(EditAnywhere, Category=MaterialExpressionStaticComponentMaskParameter)
+	UPROPERTY(EditAnywhere, Category=MaterialExpressionStaticComponentMaskParameter, meta = (ShowAsInputPin = "Advanced"))
 	uint32 DefaultB:1;
 
-	UPROPERTY(EditAnywhere, Category=MaterialExpressionStaticComponentMaskParameter)
+	UPROPERTY(EditAnywhere, Category=MaterialExpressionStaticComponentMaskParameter, meta = (ShowAsInputPin = "Advanced"))
 	uint32 DefaultA:1;
 
 
diff --git a/Engine/Source/Runtime/Engine/Classes/Materials/MaterialExpressionTextureCoordinate.h b/Engine/Source/Runtime/Engine/Classes/Materials/MaterialExpressionTextureCoordinate.h
--- a/Engine/Source/Runtime/Engine/Classes/Materials/MaterialExpressionTextureCoordinate.h
+++ b/Engine/Source/Runtime/Engine/Classes/Materials/MaterialExpressionTextureCoordinate.h
@@ -14,23 +14,23 @@ class UMaterialExpressionTextureCoordinate : public UMaterialExpression
 	GENERATED_UCLASS_BODY()
 
 	/** Texture coordinate index */
-	UPROPERTY(EditAnywhere, BlueprintReadWrite, Category = MaterialExpressionTextureCoordinate)
+	UPROPERTY(EditAnywhere, BlueprintReadWrite, Category = MaterialExpressionTextureCoordinate, meta = (ShowAsInputPin = "Primary"))
 	int32 CoordinateIndex;
 
 	/** Controls how much the texture tiles horizontally, by scaling the U component of the vertex UVs by the specified amount. */
-	UPROPERTY(EditAnywhere, BlueprintReadWrite, Category=MaterialExpressionTextureCoordinate)
+	UPROPERTY(EditAnywhere, BlueprintReadWrite, Category=MaterialExpressionTextureCoordinate, meta = (ShowAsInputPin = "Advanced"))
 	float UTiling;
 
 	/** Controls how much the texture tiles vertically, by scaling the V component of the vertex UVs by the specified amount. */
-	UPROPERTY(EditAnywhere, BlueprintReadWrite, Category=MaterialExpressionTextureCoordinate)
+	UPROPERTY(EditAnywhere, BlueprintReadWrite, Category=MaterialExpressionTextureCoordinate, meta = (ShowAsInputPin = "Advanced"))
 	float VTiling;
 
 	/** Would like to unmirror U or V 
 	 *  - if the texture is mirrored and if you would like to undo mirroring for this texture sample, use this to unmirror */
-	UPROPERTY(EditAnywhere, Category=MaterialExpressionTextureCoordinate)
+	UPROPERTY(EditAnywhere, Category=MaterialExpressionTextureCoordinate, meta = (ShowAsInputPin = "Advanced"))
 	uint32 UnMirrorU:1;
 
-	UPROPERTY(EditAnywhere, Category=MaterialExpressionTextureCoordinate)
+	UPROPERTY(EditAnywhere, Category=MaterialExpressionTextureCoordinate, meta = (ShowAsInputPin = "Advanced"))
 	uint32 UnMirrorV:1;
 
 
diff --git a/Engine/Source/Runtime/Engine/Classes/Materials/MaterialExpressionTextureProperty.h b/Engine/Source/Runtime/Engine/Classes/Materials/MaterialExpressionTextureProperty.h
--- a/Engine/Source/Runtime/Engine/Classes/Materials/MaterialExpressionTextureProperty.h
+++ b/Engine/Source/Runtime/Engine/Classes/Materials/MaterialExpressionTextureProperty.h
@@ -33,7 +33,7 @@ class ENGINE_API UMaterialExpressionTextureProperty : public UMaterialExpression
 	FExpressionInput TextureObject;
 	
 	/** Texture property to be accessed */
-	UPROPERTY(EditAnywhere, Category=UMaterialExpressionTextureProperty, meta=(DisplayName = "Texture Property"))
+	UPROPERTY(EditAnywhere, Category=UMaterialExpressionTextureProperty, meta=(DisplayName = "Texture Property", ShowAsInputPin = "Advanced"))
 	TEnumAsByte<EMaterialExposedTextureProperty> Property;
 	
 	//~ Begin UMaterialExpression Interface
diff --git a/Engine/Source/Runtime/Engine/Classes/Materials/MaterialExpressionTransform.h b/Engine/Source/Runtime/Engine/Classes/Materials/MaterialExpressionTransform.h
--- a/Engine/Source/Runtime/Engine/Classes/Materials/MaterialExpressionTransform.h
+++ b/Engine/Source/Runtime/Engine/Classes/Materials/MaterialExpressionTransform.h
@@ -73,11 +73,11 @@ class UMaterialExpressionTransform : public UMaterialExpression
 	FExpressionInput Input;
 
 	/** Source coordinate space of the FVector */
-	UPROPERTY(EditAnywhere, Category=MaterialExpressionTransform, meta=(DisplayName = "Source"))
+	UPROPERTY(EditAnywhere, Category=MaterialExpressionTransform, meta=(DisplayName = "Source", ShowAsInputPin = "Advanced"))
 	TEnumAsByte<enum EMaterialVectorCoordTransformSource> TransformSourceType;
 
 	/** Destination coordinate space of the FVector */
-	UPROPERTY(EditAnywhere, Category=MaterialExpressionTransform, meta=(DisplayName = "Destination"))
+	UPROPERTY(EditAnywhere, Category=MaterialExpressionTransform, meta=(DisplayName = "Destination", ShowAsInputPin = "Advanced"))
 	TEnumAsByte<enum EMaterialVectorCoordTransform> TransformType;
 
 	//~ Begin UMaterialExpression Interface
diff --git a/Engine/Source/Runtime/Engine/Classes/Materials/MaterialExpressionTransformPosition.h b/Engine/Source/Runtime/Engine/Classes/Materials/MaterialExpressionTransformPosition.h
--- a/Engine/Source/Runtime/Engine/Classes/Materials/MaterialExpressionTransformPosition.h
+++ b/Engine/Source/Runtime/Engine/Classes/Materials/MaterialExpressionTransformPosition.h
@@ -46,11 +46,11 @@ class UMaterialExpressionTransformPosition : public UMaterialExpression
 	FExpressionInput Input;
 
 	/** source format of the position that will be transformed */
-	UPROPERTY(EditAnywhere, Category=MaterialExpressionTransformPosition, meta=(DisplayName = "Source"))
+	UPROPERTY(EditAnywhere, Category=MaterialExpressionTransformPosition, meta=(DisplayName = "Source", ShowAsInputPin = "Advanced"))
 	TEnumAsByte<enum EMaterialPositionTransformSource> TransformSourceType;
 
 	/** type of transform to apply to the input expression */
-	UPROPERTY(EditAnywhere, Category=MaterialExpressionTransformPosition, meta=(DisplayName = "Destination"))
+	UPROPERTY(EditAnywhere, Category=MaterialExpressionTransformPosition, meta=(DisplayName = "Destination", ShowAsInputPin = "Advanced"))
 	TEnumAsByte<enum EMaterialPositionTransformSource> TransformType;
 
 
diff --git a/Engine/Source/Runtime/Engine/Classes/Materials/MaterialExpressionVectorNoise.h b/Engine/Source/Runtime/Engine/Classes/Materials/MaterialExpressionVectorNoise.h
--- a/Engine/Source/Runtime/Engine/Classes/Materials/MaterialExpressionVectorNoise.h
+++ b/Engine/Source/Runtime/Engine/Classes/Materials/MaterialExpressionVectorNoise.h
@@ -58,21 +58,21 @@ class UMaterialExpressionVectorNoise : public UMaterialExpression
 	FExpressionInput Position;
 
 	/** Noise function, affects performance and look */
-	UPROPERTY(EditAnywhere, Category = MaterialExpressionVectorNoise, meta = (DisplayName = "Function"))
+	UPROPERTY(EditAnywhere, Category = MaterialExpressionVectorNoise, meta = (DisplayName = "Function", ShowAsInputPin = "Advanced"))
 	TEnumAsByte<enum EVectorNoiseFunction> NoiseFunction;
 
 	/** For noise functions where applicable, lower numbers are faster and lower quality, higher numbers are slower and higher quality */
-	UPROPERTY(EditAnywhere, Category=MaterialExpressionVectorNoise, meta=(UIMin = "1", UIMax = "4"))
+	UPROPERTY(EditAnywhere, Category=MaterialExpressionVectorNoise, meta=(UIMin = "1", UIMax = "4", ShowAsInputPin = "Advanced"))
 	int32 Quality;
 
 	/** Whether tile the noise pattern, useful for baking to seam-free repeating textures */
-	UPROPERTY(EditAnywhere, Category = MaterialExpressionVectorNoise)
+	UPROPERTY(EditAnywhere, Category = MaterialExpressionVectorNoise, meta = (ShowAsInputPin = "Advanced"))
 	uint32 bTiling:1;
 
 	/** How many units in each tile (if Tiling is on) 
 	  * For Perlin noise functions, Tile Size must be a multiple of three
 	  */
-	UPROPERTY(EditAnywhere, Category = MaterialExpressionVectorNoise, meta=(UIMin = "4"))
+	UPROPERTY(EditAnywhere, Category = MaterialExpressionVectorNoise, meta=(UIMin = "4", ShowAsInputPin = "Advanced"))
 	uint32 TileSize;
 
 	//~ Begin UMaterialExpression Interface
diff --git a/Engine/Source/Runtime/Engine/Classes/Materials/MaterialExpressionViewProperty.h b/Engine/Source/Runtime/Engine/Classes/Materials/MaterialExpressionViewProperty.h
--- a/Engine/Source/Runtime/Engine/Classes/Materials/MaterialExpressionViewProperty.h
+++ b/Engine/Source/Runtime/Engine/Classes/Materials/MaterialExpressionViewProperty.h
@@ -49,7 +49,7 @@ class UMaterialExpressionViewProperty : public UMaterialExpression
 	GENERATED_UCLASS_BODY()
 	
 	/** View input property to be accessed */
-	UPROPERTY(EditAnywhere, Category=UMaterialExpressionViewProperty, meta=(DisplayName = "View Property"))
+	UPROPERTY(EditAnywhere, Category=UMaterialExpressionViewProperty, meta=(DisplayName = "View Property", ShowAsInputPin = "Advanced"))
 	TEnumAsByte<EMaterialExposedViewProperty> Property;
 	
 	//~ Begin UMaterialExpression Interface
diff --git a/Engine/Source/Runtime/Engine/Classes/Materials/MaterialExpressionWorldPosition.h b/Engine/Source/Runtime/Engine/Classes/Materials/MaterialExpressionWorldPosition.h
--- a/Engine/Source/Runtime/Engine/Classes/Materials/MaterialExpressionWorldPosition.h
+++ b/Engine/Source/Runtime/Engine/Classes/Materials/MaterialExpressionWorldPosition.h
@@ -32,7 +32,7 @@ class UMaterialExpressionWorldPosition : public UMaterialExpression
 {
 	GENERATED_UCLASS_BODY()
 	
-	UPROPERTY(EditAnywhere, Category=UMaterialExpressionWorldPosition, meta=(DisplayName = "Shader Offsets"))
+	UPROPERTY(EditAnywhere, Category=UMaterialExpressionWorldPosition, meta=(DisplayName = "Shader Offsets", ShowAsInputPin = "Advanced"))
 	TEnumAsByte<EWorldPositionIncludedOffsets> WorldPositionShaderOffset;
 	
 	//~ Begin UMaterialExpression Interface
diff --git a/Engine/Source/Runtime/Engine/Private/Materials/MaterialExpressions.cpp b/Engine/Source/Runtime/Engine/Private/Materials/MaterialExpressions.cpp
--- a/Engine/Source/Runtime/Engine/Private/Materials/MaterialExpressions.cpp
+++ b/Engine/Source/Runtime/Engine/Private/Materials/MaterialExpressions.cpp
@@ -934,8 +934,10 @@ void UMaterialExpression::PostDuplicate(bool bDuplicateForPIE)
 
 #if WITH_EDITOR
 
-FProperty* UMaterialExpression::GetInputPinProperty(int32 PinIndex)
+TArray<FProperty*> UMaterialExpression::GetInputPinProperty(int32 PinIndex)
 {
+	// Find all UPROPERTYs associated with this input pin
+	TArray<FProperty*> Properties;
 	// Explicit input pins are before property input pins
 	TArray<FProperty*> PropertyInputs = GetPropertyInputs();
 	const int32 NumInputs = GetInputs().Num();
@@ -943,7 +945,8 @@ FProperty* UMaterialExpression::GetInputPinProperty(int32 PinIndex)
 	{
 		FExpressionInput* Input = GetInput(PinIndex);
 
-		// Find the UPROPERTY that has OverridingInputProperty meta data pointing to the expression input.
+		// Find the UPROPERTYs that have OverridingInputProperty meta data pointing to the expression input.
+		// There can be multiple scalar entries together forming a vector parameter, e.g. DecalMipmapLevel node has FExpressionInput TextureSize <-> float ConstWidth/ConstHeight.
 		static FName OverridingInputPropertyMetaData(TEXT("OverridingInputProperty"));
 		for (TFieldIterator<FProperty> InputIt(GetClass(), EFieldIteratorFlags::IncludeSuper, EFieldIteratorFlags::ExcludeDeprecated); InputIt; ++InputIt)
 		{
@@ -956,7 +959,7 @@ FProperty* UMaterialExpression::GetInputPinProperty(int32 PinIndex)
 				{
 					if (Input == StructProp->ContainerPtrToValuePtr<FExpressionInput>(this))
 					{
-						return Property;
+						Properties.Add(Property);
 					}
 				}
 			}
@@ -970,11 +973,11 @@ FProperty* UMaterialExpression::GetInputPinProperty(int32 PinIndex)
 			FProperty* Property = *InputIt;
 			if (PropertyName == Property->GetFName())
 			{
-				return Property;
+				Properties.Add(Property);
 			}
 		}
 	}
-	return nullptr;
+	return Properties;
 }
 
 FName UMaterialExpression::GetInputPinSubCategory(int32 PinIndex)
@@ -982,15 +985,20 @@ FName UMaterialExpression::GetInputPinSubCategory(int32 PinIndex)
 	FName PinSubCategory;
 
 	// Find the UPROPERTY associated with the pin
-	FProperty* Property = GetInputPinProperty(PinIndex);
-	if (Property)
+	TArray<FProperty*> Properties = GetInputPinProperty(PinIndex);
+	if (Properties.Num() == 1)
 	{
 		// This is the UPROPERTY matching with the target input
+		FProperty* Property = Properties[0];
 		FFieldClass* PropertyClass = Property->GetClass();
 		if (PropertyClass == FFloatProperty::StaticClass())
 		{
 			PinSubCategory = UMaterialGraphSchema::PSC_Red;
 		}
+		else if (PropertyClass == FDoubleProperty::StaticClass())
+		{
+			PinSubCategory = UMaterialGraphSchema::PSC_Red;
+		}
 		else if (PropertyClass == FIntProperty::StaticClass())
 		{
 			PinSubCategory = UMaterialGraphSchema::PSC_Int;
@@ -1028,6 +1036,19 @@ FName UMaterialExpression::GetInputPinSubCategory(int32 PinIndex)
 			}	
 		}
 	}
+	// There can be multiple scalar entries together forming a vector2/3/4.
+	else if (Properties.Num() == 2)
+	{
+		PinSubCategory = UMaterialGraphSchema::PSC_RG;
+	}
+	else if (Properties.Num() == 3)
+	{
+		PinSubCategory = UMaterialGraphSchema::PSC_RGB;
+	}
+	else if (Properties.Num() == 4)
+	{
+		PinSubCategory = UMaterialGraphSchema::PSC_Vector4;
+	}
 
 	return PinSubCategory;
 }
@@ -1035,10 +1056,10 @@ FName UMaterialExpression::GetInputPinSubCategory(int32 PinIndex)
 UObject* UMaterialExpression::GetInputPinSubCategoryObject(int32 PinIndex)
 {
 	UObject* PinSubCategoryObject = nullptr;
-	FProperty* Property = GetInputPinProperty(PinIndex);
-	if (Property)
+	TArray<FProperty*> Properties = GetInputPinProperty(PinIndex);
+	if (Properties.Num() > 0)
 	{
-		if (FByteProperty* ByteProperty = CastField<FByteProperty>(Property))
+		if (FByteProperty* ByteProperty = CastField<FByteProperty>(Properties[0]))
 		{
 			PinSubCategoryObject = ByteProperty->GetIntPropertyEnum();
 		}
@@ -1049,138 +1070,171 @@ UObject* UMaterialExpression::GetInputPinSubCategoryObject(int32 PinIndex)
 void UMaterialExpression::PinDefaultValueChanged(int32 PinIndex, const FString& DefaultValue)
 {
 	// Update the default value of the expression input when pin value changes
-	// Find the UPROPERTY that has OverridingInputProperty meta data pointing to the input, override its value.
-	FProperty* Property = GetInputPinProperty(PinIndex);
-	if (Property)
+	// Find the UPROPERTYs that have OverridingInputProperty meta data pointing to the input, override their values.
+	TArray<FProperty*> Properties = GetInputPinProperty(PinIndex);
+	if (Properties.Num() > 0)
 	{
-		const FString& ClampMin = Property->GetMetaData(TEXT("ClampMin"));
-		const FString& ClampMax = Property->GetMetaData(TEXT("ClampMax"));
-
-		FFieldClass* PropertyClass = Property->GetClass();
-		if (PropertyClass == FFloatProperty::StaticClass())
+		TArray<FString> PropertyValues;
+		if (Properties.Num() == 1)
 		{
-			float Value = FCString::Atof(*DefaultValue);
-			Value = ClampMin.Len() ? (FMath::Max<float>(FCString::Atof(*ClampMin), Value)) : Value;
-			Value = ClampMax.Len() ? (FMath::Min<float>(FCString::Atof(*ClampMax), Value)) : Value;
-			FFloatProperty* FloatProperty = CastField<FFloatProperty>(Property);
-			FloatProperty->SetPropertyValue_InContainer(this, Value);
+			PropertyValues.Add(DefaultValue);
 		}
-		else if (PropertyClass == FIntProperty::StaticClass())
+		else if (Properties.Num() == 2)
 		{
-			int32 Value = FCString::Atoi(*DefaultValue);
-			Value = ClampMin.Len() ? (FMath::Max<int32>(FCString::Atoi(*ClampMin), Value)) : Value;
-			Value = ClampMax.Len() ? (FMath::Min<int32>(FCString::Atoi(*ClampMax), Value)) : Value;
-			FIntProperty* IntProperty = CastField<FIntProperty>(Property);
-			IntProperty->SetPropertyValue_InContainer(this, Value);
+			// Vector2 is formatted as (X=0.0, Y=0.0)
+			FVector2D Value;
+			Value.InitFromString(DefaultValue);
+			PropertyValues.Add(FString::SanitizeFloat(Value.X));
+			PropertyValues.Add(FString::SanitizeFloat(Value.Y));
 		}
-		else if (PropertyClass == FUInt32Property::StaticClass())
-		{	
-			int32 IntValue = FCString::Atoi(*DefaultValue);
-			IntValue = ClampMin.Len() ? (FMath::Max<int32>(FCString::Atoi(*ClampMin), IntValue)) : IntValue;
-			IntValue = ClampMax.Len() ? (FMath::Min<int32>(FCString::Atoi(*ClampMax), IntValue)) : IntValue;
-			// Make sure the value is not negative
-			uint32 Value = (uint32)FMath::Max(IntValue, 0);
-			FUInt32Property* UInt32Property = CastField<FUInt32Property>(Property);
-			UInt32Property->SetPropertyValue_InContainer(this, Value);
+		else
+		{
+			// Vector3/4 are formatted as numbers separated by commas
+			DefaultValue.ParseIntoArray(PropertyValues, TEXT(","), true);
+			check(PropertyValues.Num() == Properties.Num());
 		}
-		else if (PropertyClass == FByteProperty::StaticClass())
+
+		for (int32 i = 0; i < Properties.Num(); ++i)
 		{
-			FByteProperty* ByteProperty = CastField<FByteProperty>(Property);
-			uint8 Value;
-			UEnum* Enum = ByteProperty->GetIntPropertyEnum();
-			if (Enum)
+			FProperty* Property = Properties[i];
+			const FString& ClampMin = Property->GetMetaData(TEXT("ClampMin"));
+			const FString& ClampMax = Property->GetMetaData(TEXT("ClampMax"));
+
+			FString PropertyValue = PropertyValues[i];
+			FFieldClass* PropertyClass = Property->GetClass();
+			if (PropertyClass == FFloatProperty::StaticClass())
 			{
-				Value = (uint8)Enum->GetValueByName(FName(DefaultValue));
+				float Value = FCString::Atof(*PropertyValue);
+				Value = ClampMin.Len() ? (FMath::Max<float>(FCString::Atof(*ClampMin), Value)) : Value;
+				Value = ClampMax.Len() ? (FMath::Min<float>(FCString::Atof(*ClampMax), Value)) : Value;
+				FFloatProperty* FloatProperty = CastField<FFloatProperty>(Property);
+				FloatProperty->SetPropertyValue_InContainer(this, Value);
 			}
-			else
+			else if (PropertyClass == FDoubleProperty::StaticClass())
 			{
-				int32 IntValue = FCString::Atoi(*DefaultValue);
+				double Value = FCString::Atod(*PropertyValue);
+				Value = ClampMin.Len() ? (FMath::Max<double>(FCString::Atod(*ClampMin), Value)) : Value;
+				Value = ClampMax.Len() ? (FMath::Min<double>(FCString::Atod(*ClampMax), Value)) : Value;
+				FDoubleProperty* DoubleProperty = CastField<FDoubleProperty>(Property);
+				DoubleProperty->SetPropertyValue_InContainer(this, Value);
+			}
+			else if (PropertyClass == FIntProperty::StaticClass())
+			{
+				int32 Value = FCString::Atoi(*PropertyValue);
+				Value = ClampMin.Len() ? (FMath::Max<int32>(FCString::Atoi(*ClampMin), Value)) : Value;
+				Value = ClampMax.Len() ? (FMath::Min<int32>(FCString::Atoi(*ClampMax), Value)) : Value;
+				FIntProperty* IntProperty = CastField<FIntProperty>(Property);
+				IntProperty->SetPropertyValue_InContainer(this, Value);
+			}
+			else if (PropertyClass == FUInt32Property::StaticClass())
+			{
+				int32 IntValue = FCString::Atoi(*PropertyValue);
 				IntValue = ClampMin.Len() ? (FMath::Max<int32>(FCString::Atoi(*ClampMin), IntValue)) : IntValue;
 				IntValue = ClampMax.Len() ? (FMath::Min<int32>(FCString::Atoi(*ClampMax), IntValue)) : IntValue;
-				// Make sure the value doesn't exceed byte limit
-				Value = (uint8)FMath::Max(FMath::Min(IntValue, 255), 0);
-			}	
-			ByteProperty->SetPropertyValue_InContainer(this, Value);
-		}
-		else if (PropertyClass == FBoolProperty::StaticClass())
-		{
-			bool Value = FCString::ToBool(*DefaultValue);
-			FBoolProperty* BoolProperty = CastField<FBoolProperty>(Property);
-			BoolProperty->SetPropertyValue_InContainer(this, Value);
-		}
-		else if (PropertyClass == FStructProperty::StaticClass())
-		{
-			UScriptStruct* Struct = ((FStructProperty*)Property)->Struct;
-			if (Struct == TBaseStructure<FLinearColor>::Get())
+				// Make sure the value is not negative
+				uint32 Value = (uint32)FMath::Max(IntValue, 0);
+				FUInt32Property* UInt32Property = CastField<FUInt32Property>(Property);
+				UInt32Property->SetPropertyValue_InContainer(this, Value);
+			}
+			else if (PropertyClass == FByteProperty::StaticClass())
 			{
-				FLinearColor* ColorProperty = Property->ContainerPtrToValuePtr<FLinearColor>(this);
-				if (Property->HasMetaData(TEXT("HideAlphaChannel")))
+				FByteProperty* ByteProperty = CastField<FByteProperty>(Property);
+				uint8 Value;
+				UEnum* Enum = ByteProperty->GetIntPropertyEnum();
+				if (Enum)
 				{
-					// This is a 3 element vector
-					TArray<FString> Elements;
-					DefaultValue.ParseIntoArray(Elements, TEXT(","), true);
-					check(Elements.Num() == 3);
-					ColorProperty->R = FCString::Atof(*Elements[0]);
-					ColorProperty->G = FCString::Atof(*Elements[1]);
-					ColorProperty->B = FCString::Atof(*Elements[2]);
+					Value = (uint8)Enum->GetValueByName(FName(PropertyValue));
 				}
 				else
 				{
-					// This is a 4 element vector
-					ColorProperty->InitFromString(DefaultValue);
+					int32 IntValue = FCString::Atoi(*PropertyValue);
+					IntValue = ClampMin.Len() ? (FMath::Max<int32>(FCString::Atoi(*ClampMin), IntValue)) : IntValue;
+					IntValue = ClampMax.Len() ? (FMath::Min<int32>(FCString::Atoi(*ClampMax), IntValue)) : IntValue;
+					// Make sure the value doesn't exceed byte limit
+					Value = (uint8)FMath::Max(FMath::Min(IntValue, 255), 0);
 				}
+				ByteProperty->SetPropertyValue_InContainer(this, Value);
 			}
-			else if (Struct == TBaseStructure<FVector4>::Get())
-			{
-				TArray<FString> Elements;
-				DefaultValue.ParseIntoArray(Elements, TEXT(","), true);
-				check(Elements.Num() == 4);
-				FVector4* Value = Property->ContainerPtrToValuePtr<FVector4>(this);
-				Value->X = FCString::Atod(*Elements[0]);
-				Value->Y = FCString::Atod(*Elements[1]);
-				Value->Z = FCString::Atod(*Elements[2]);
-				Value->W = FCString::Atod(*Elements[3]);
-			}
-			else if (Struct == TVariantStructure<FVector4d>::Get())
-			{
-				TArray<FString> Elements;
-				DefaultValue.ParseIntoArray(Elements, TEXT(","), true);
-				check(Elements.Num() == 4);
-				FVector4d* Value = Property->ContainerPtrToValuePtr<FVector4d>(this);
-				Value->X = FCString::Atod(*Elements[0]);
-				Value->Y = FCString::Atod(*Elements[1]);
-				Value->Z = FCString::Atod(*Elements[2]);
-				Value->W = FCString::Atod(*Elements[3]);
-			}
-			else if (Struct == TBaseStructure<FVector>::Get())
+			else if (PropertyClass == FBoolProperty::StaticClass())
 			{
-				TArray<FString> Elements;
-				DefaultValue.ParseIntoArray(Elements, TEXT(","), true);
-				check(Elements.Num() == 3);
-				FVector* Value = Property->ContainerPtrToValuePtr<FVector>(this);
-				Value->X = FCString::Atod(*Elements[0]);
-				Value->Y = FCString::Atod(*Elements[1]);
-				Value->Z = FCString::Atod(*Elements[2]);
+				bool Value = FCString::ToBool(*PropertyValue);
+				FBoolProperty* BoolProperty = CastField<FBoolProperty>(Property);
+				BoolProperty->SetPropertyValue_InContainer(this, Value);
 			}
-			else if (Struct == TVariantStructure<FVector3f>::Get())
-			{
-				TArray<FString> Elements;
-				DefaultValue.ParseIntoArray(Elements, TEXT(","), true);
-				check(Elements.Num() == 3);
-				FVector3f* Value = Property->ContainerPtrToValuePtr<FVector3f>(this);
-				Value->X = FCString::Atof(*Elements[0]);
-				Value->Y = FCString::Atof(*Elements[1]);
-				Value->Z = FCString::Atof(*Elements[2]);
-			}
-			else if (Struct == TBaseStructure<FVector2D>::Get())
+			else if (PropertyClass == FStructProperty::StaticClass())
 			{
-				FVector2D* Value = Property->ContainerPtrToValuePtr<FVector2D>(this);
-				Value->InitFromString(DefaultValue);
+				UScriptStruct* Struct = ((FStructProperty*)Property)->Struct;
+				if (Struct == TBaseStructure<FLinearColor>::Get())
+				{
+					FLinearColor* ColorProperty = Property->ContainerPtrToValuePtr<FLinearColor>(this);
+					if (Property->HasMetaData(TEXT("HideAlphaChannel")))
+					{
+						// This is a 3 element vector
+						TArray<FString> Elements;
+						PropertyValue.ParseIntoArray(Elements, TEXT(","), true);
+						check(Elements.Num() == 3);
+						ColorProperty->R = FCString::Atof(*Elements[0]);
+						ColorProperty->G = FCString::Atof(*Elements[1]);
+						ColorProperty->B = FCString::Atof(*Elements[2]);
+					}
+					else
+					{
+						// This is a 4 element vector
+						ColorProperty->InitFromString(PropertyValue);
+					}
+				}
+				else if (Struct == TBaseStructure<FVector4>::Get())
+				{
+					TArray<FString> Elements;
+					PropertyValue.ParseIntoArray(Elements, TEXT(","), true);
+					check(Elements.Num() == 4);
+					FVector4* Value = Property->ContainerPtrToValuePtr<FVector4>(this);
+					Value->X = FCString::Atod(*Elements[0]);
+					Value->Y = FCString::Atod(*Elements[1]);
+					Value->Z = FCString::Atod(*Elements[2]);
+					Value->W = FCString::Atod(*Elements[3]);
+				}
+				else if (Struct == TVariantStructure<FVector4d>::Get())
+				{
+					TArray<FString> Elements;
+					PropertyValue.ParseIntoArray(Elements, TEXT(","), true);
+					check(Elements.Num() == 4);
+					FVector4d* Value = Property->ContainerPtrToValuePtr<FVector4d>(this);
+					Value->X = FCString::Atod(*Elements[0]);
+					Value->Y = FCString::Atod(*Elements[1]);
+					Value->Z = FCString::Atod(*Elements[2]);
+					Value->W = FCString::Atod(*Elements[3]);
+				}
+				else if (Struct == TBaseStructure<FVector>::Get())
+				{
+					TArray<FString> Elements;
+					PropertyValue.ParseIntoArray(Elements, TEXT(","), true);
+					check(Elements.Num() == 3);
+					FVector* Value = Property->ContainerPtrToValuePtr<FVector>(this);
+					Value->X = FCString::Atod(*Elements[0]);
+					Value->Y = FCString::Atod(*Elements[1]);
+					Value->Z = FCString::Atod(*Elements[2]);
+				}
+				else if (Struct == TVariantStructure<FVector3f>::Get())
+				{
+					TArray<FString> Elements;
+					PropertyValue.ParseIntoArray(Elements, TEXT(","), true);
+					check(Elements.Num() == 3);
+					FVector3f* Value = Property->ContainerPtrToValuePtr<FVector3f>(this);
+					Value->X = FCString::Atof(*Elements[0]);
+					Value->Y = FCString::Atof(*Elements[1]);
+					Value->Z = FCString::Atof(*Elements[2]);
+				}
+				else if (Struct == TBaseStructure<FVector2D>::Get())
+				{
+					FVector2D* Value = Property->ContainerPtrToValuePtr<FVector2D>(this);
+					Value->InitFromString(PropertyValue);
+				}
 			}
-		}
 
-		FPropertyChangedEvent Event(Property);
-		PostEditChangeProperty(Event);
+			FPropertyChangedEvent Event(Property);
+			PostEditChangeProperty(Event);
+		}
 
 		// Update the expression preview and the material to reflect the change
 		Modify();
@@ -1194,41 +1248,51 @@ void UMaterialExpression::PinDefaultValueChanged(int32 PinIndex, const FString&
 
 FString UMaterialExpression::GetInputPinDefaultValue(int32 PinIndex)
 {
-	// Find the UPROPERTY for the input pin, retrieve its value.
-	FProperty* Property = GetInputPinProperty(PinIndex);
-	if (Property)
+	TArray<FString> PropertyValues;
+
+	// Find the UPROPERTYs for the input pin, retrieve their values.
+	TArray<FProperty*> Properties = GetInputPinProperty(PinIndex);
+	for (int32 i = 0; i < Properties.Num(); ++i)
 	{
 		// This is the UPROPERTY matching with the target input
+		FProperty* Property = Properties[i];
+		FString PropertyValue;
 		FFieldClass* PropertyClass = Property->GetClass();
 		if (PropertyClass == FFloatProperty::StaticClass())
 		{
 			FFloatProperty* FloatProperty = CastField<FFloatProperty>(Property);
 			float Value = FloatProperty->GetPropertyValue_InContainer(this);
-			return FString::SanitizeFloat(Value);
+			PropertyValue = FString::SanitizeFloat(Value);
+		}
+		else if (PropertyClass == FDoubleProperty::StaticClass())
+		{
+			FDoubleProperty* DoubleProperty = CastField<FDoubleProperty>(Property);
+			double Value = DoubleProperty->GetPropertyValue_InContainer(this);
+			PropertyValue = LexToString(Value);
 		}
 		else if (PropertyClass == FIntProperty::StaticClass())
 		{
 			FIntProperty* IntProperty = CastField<FIntProperty>(Property);
 			int32 Value = IntProperty->GetPropertyValue_InContainer(this);
-			return FString::FromInt(Value);
+			PropertyValue = FString::FromInt(Value);
 		}
 		else if (PropertyClass == FUInt32Property::StaticClass())
 		{
 			FUInt32Property* UInt32Property = CastField<FUInt32Property>(Property);
 			uint32 Value = UInt32Property->GetPropertyValue_InContainer(this);
-			return FString::FromInt((int32)Value);
+			PropertyValue = FString::FromInt((int32)Value);
 		}
 		else if (PropertyClass == FByteProperty::StaticClass())
 		{
 			FByteProperty* ByteProperty = CastField<FByteProperty>(Property);
 			uint8 Value = ByteProperty->GetPropertyValue_InContainer(this);
-			return ByteProperty->Enum ? ByteProperty->Enum->GetDisplayNameTextByValue(Value).ToString() : FString::FromInt(Value);
+			PropertyValue = (ByteProperty->Enum ? ByteProperty->Enum->GetDisplayNameTextByValue(Value).ToString() : FString::FromInt(Value));
 		}
 		else if (PropertyClass == FBoolProperty::StaticClass())
 		{
 			FBoolProperty* BoolProperty = CastField<FBoolProperty>(Property);
 			bool Value = BoolProperty->GetPropertyValue_InContainer(this);
-			return Value ? TEXT("true") : TEXT("false");
+			PropertyValue = (Value ? TEXT("true") : TEXT("false"));
 		}
 		else if (PropertyClass == FStructProperty::StaticClass())
 		{
@@ -1239,40 +1303,65 @@ FString UMaterialExpression::GetInputPinDefaultValue(int32 PinIndex)
 				if (Property->HasMetaData(TEXT("HideAlphaChannel")))
 				{
 					// This is a 3 element vector
-					return FString::SanitizeFloat(Value.R) + FString(TEXT(",")) + FString::SanitizeFloat(Value.G) + FString(TEXT(",")) + FString::SanitizeFloat(Value.B);
+					PropertyValue = FString::SanitizeFloat(Value.R) + FString(TEXT(",")) + FString::SanitizeFloat(Value.G) + FString(TEXT(",")) + FString::SanitizeFloat(Value.B);
 				}
 				else
 				{
 					// This is a 4 element vector
-					return Value.ToString();
+					PropertyValue = Value.ToString();
 				}
 			}
 			else if (Struct == TBaseStructure<FVector4>::Get())
 			{
 				FVector4 Value = *Property->ContainerPtrToValuePtr<FVector4>(this);		
-				return FString::SanitizeFloat(Value.X) + FString(TEXT(",")) + FString::SanitizeFloat(Value.Y) + FString(TEXT(",")) + FString::SanitizeFloat(Value.Z) + FString(TEXT(",")) + FString::SanitizeFloat(Value.W);
+				PropertyValue = FString::SanitizeFloat(Value.X) + FString(TEXT(",")) + FString::SanitizeFloat(Value.Y) + FString(TEXT(",")) + FString::SanitizeFloat(Value.Z) + FString(TEXT(",")) + FString::SanitizeFloat(Value.W);
 			}
 			else if (Struct == TVariantStructure<FVector4d>::Get())
 			{
 				FVector4d Value = *Property->ContainerPtrToValuePtr<FVector4d>(this);		
-				return FString::SanitizeFloat(Value.X) + FString(TEXT(",")) + FString::SanitizeFloat(Value.Y) + FString(TEXT(",")) + FString::SanitizeFloat(Value.Z) + FString(TEXT(",")) + FString::SanitizeFloat(Value.W);
+				PropertyValue = FString::SanitizeFloat(Value.X) + FString(TEXT(",")) + FString::SanitizeFloat(Value.Y) + FString(TEXT(",")) + FString::SanitizeFloat(Value.Z) + FString(TEXT(",")) + FString::SanitizeFloat(Value.W);
 			}
 			else if (Struct == TBaseStructure<FVector>::Get())
 			{
 				FVector Value = *Property->ContainerPtrToValuePtr<FVector>(this);
-				return FString::SanitizeFloat(Value.X) + FString(TEXT(",")) + FString::SanitizeFloat(Value.Y) + FString(TEXT(",")) + FString::SanitizeFloat(Value.Z);
+				PropertyValue = FString::SanitizeFloat(Value.X) + FString(TEXT(",")) + FString::SanitizeFloat(Value.Y) + FString(TEXT(",")) + FString::SanitizeFloat(Value.Z);
 			}
 			else if (Struct == TVariantStructure<FVector3f>::Get())
 			{
 				FVector3f Value = *Property->ContainerPtrToValuePtr<FVector3f>(this);
-				return FString::SanitizeFloat(Value.X) + FString(TEXT(",")) + FString::SanitizeFloat(Value.Y) + FString(TEXT(",")) + FString::SanitizeFloat(Value.Z);
+				PropertyValue = FString::SanitizeFloat(Value.X) + FString(TEXT(",")) + FString::SanitizeFloat(Value.Y) + FString(TEXT(",")) + FString::SanitizeFloat(Value.Z);
 			}
 			else if (Struct == TBaseStructure<FVector2D>::Get())
 			{
 				FVector2D* Value = Property->ContainerPtrToValuePtr<FVector2D>(this);
-				return Value->ToString();
+				PropertyValue = Value->ToString();
 			}
 		}
+
+		PropertyValues.Add(PropertyValue);
+	}
+
+	check(PropertyValues.Num() == Properties.Num());
+	if (Properties.Num() == 1)
+	{
+		return PropertyValues[0];
+	}
+	else if (Properties.Num() == 2)
+	{
+		// Vector2 is formatted as (X=0.0, Y=0.0)
+		float X = FCString::Atof(*PropertyValues[0]);
+		float Y = FCString::Atof(*PropertyValues[1]);
+		FVector2D Value(X, Y);
+		return Value.ToString();
+	}
+	// Vector3/4 are formatted as numbers separated by commas
+	else if (Properties.Num() == 3)
+	{
+		return PropertyValues[0] + TEXT(",") + PropertyValues[1] + TEXT(",") + PropertyValues[2];
+	}
+	else if (Properties.Num() == 4)
+	{
+		return PropertyValues[0] + TEXT(",") + PropertyValues[1] + TEXT(",") + PropertyValues[2] + TEXT(",") + PropertyValues[3];
 	}
 
 	return TEXT("");
@@ -1301,6 +1390,7 @@ TArray<FProperty*> UMaterialExpression::GetPropertyInputs() const
 			// Check if the property type fits within the allowed widget types
 			FFieldClass* PropertyClass = Property->GetClass();
 			if (PropertyClass == FFloatProperty::StaticClass()
+				|| PropertyClass == FDoubleProperty::StaticClass()
 				|| PropertyClass == FIntProperty::StaticClass()
 				|| PropertyClass == FUInt32Property::StaticClass()
 				|| PropertyClass == FByteProperty::StaticClass()
@@ -5067,15 +5157,6 @@ int32 UMaterialExpressionClamp::Compile(class FMaterialCompiler* Compiler, int32
 void UMaterialExpressionClamp::GetCaption(TArray<FString>& OutCaptions) const
 {
 	FString	NewCaption = TEXT( "Clamp" );
-
-	if (ClampMode == CMODE_ClampMin || ClampMode == CMODE_Clamp)
-	{
-		NewCaption += Min.GetTracedInput().Expression ? TEXT(" (Min)") : FString::Printf(TEXT(" (Min=%.4g)"), MinDefault);
-	}
-	if (ClampMode == CMODE_ClampMax || ClampMode == CMODE_Clamp)
-	{
-		NewCaption += Max.GetTracedInput().Expression ? TEXT(" (Max)") : FString::Printf(TEXT(" (Max=%.4g)"), MaxDefault);
-	}
 	OutCaptions.Add(NewCaption);
 }
 #endif // WITH_EDITOR
