Compiles all global shaders with HLSL2021 on some platforms

diff --git a/Engine/Config/Mac/DataDrivenPlatformInfo.ini b/Engine/Config/Mac/DataDrivenPlatformInfo.ini
--- a/Engine/Config/Mac/DataDrivenPlatformInfo.ini
+++ b/Engine/Config/Mac/DataDrivenPlatformInfo.ini
@@ -48,6 +48,7 @@ bSupportsCapsuleShadows=true
 bSupports4ComponentUAVReadWrite=true
 bSupportsVolumeTextureAtomics=false
 bSupportsDualSourceBlending=true
+EnablesHLSL2021ByDefault=1
 
 [ShaderPlatform METAL_MRT_MAC]
 Language=Metal
diff --git a/Engine/Config/VulkanPC/DataDrivenPlatformInfo.ini b/Engine/Config/VulkanPC/DataDrivenPlatformInfo.ini
--- a/Engine/Config/VulkanPC/DataDrivenPlatformInfo.ini
+++ b/Engine/Config/VulkanPC/DataDrivenPlatformInfo.ini
@@ -132,6 +132,7 @@ bSupportsAnisotropicMaterials=true
 bSupportsROV=false
 bSupportsOIT=false
 
+EnablesHLSL2021ByDefault=1
 
 [ShaderPlatform VULKAN_PCES3_1]
 Language=Vulkan
diff --git a/Engine/Config/Windows/DataDrivenPlatformInfo.ini b/Engine/Config/Windows/DataDrivenPlatformInfo.ini
--- a/Engine/Config/Windows/DataDrivenPlatformInfo.ini
+++ b/Engine/Config/Windows/DataDrivenPlatformInfo.ini
@@ -140,6 +140,7 @@ bSupportsBindless=true
 bSupportsROV=true
 bSupportsOIT=true
 bSupportsRealTypes=RuntimeDependent
+EnablesHLSL2021ByDefault=1
 
 [PreviewPlatform PCD3D_SM6]
 PlatformName=PCD3D_SM6
diff --git a/Engine/Shaders/Private/TemporalAA.usf b/Engine/Shaders/Private/TemporalAA.usf
--- a/Engine/Shaders/Private/TemporalAA.usf
+++ b/Engine/Shaders/Private/TemporalAA.usf
@@ -540,6 +540,16 @@ float UpscaleFactor; // = 1 / ScreenPercentage
 
 //------------------------------------------------------- FUNCTIONS
 
+#if COMPILER_SUPPORTS_HLSL2021
+
+template<typename T>
+void CorrectExposure(inout T X)
+{
+	X *= HistoryPreExposureCorrection;
+}
+
+#endif
+
 float3 RGBToYCoCg( float3 RGB )
 {
 	float Y  = dot( RGB, float3(  1, 2,  1 ) );
@@ -1858,7 +1868,11 @@ FTAAHistoryPayload SampleHistory(in float2 HistoryScreenPosition)
 		HistoryPayload.CocRadius = RawHistory0.a;
 	#endif
 
-	HistoryPayload.Color.rgb *= HistoryPreExposureCorrection;
+	#if COMPILER_SUPPORTS_HLSL2021
+		CorrectExposure(HistoryPayload.Color.rgb);
+	#else
+		HistoryPayload.Color.rgb *= HistoryPreExposureCorrection;
+	#endif
 
 	HistoryPayload.Color = TransformSceneColor(HistoryPayload.Color);
 
diff --git a/Engine/Shaders/Public/ShaderVersion.ush b/Engine/Shaders/Public/ShaderVersion.ush
--- a/Engine/Shaders/Public/ShaderVersion.ush
+++ b/Engine/Shaders/Public/ShaderVersion.ush
@@ -2,4 +2,4 @@
 // Each time the console command is executed it generates a new GUID. As this file is included
 // in Platform.ush (which should be included in any shader) it allows to invalidate the shader DDC.
 // 
-// GUID = 8EA7B5C8-3F33-4EF6-922F-F4E242BE7AD3
+// GUID = 15749AA0-E6E1-44D9-8D85-687EE3C73A20
diff --git a/Engine/Source/Runtime/Engine/Private/ShaderCompiler/ShaderCompiler.cpp b/Engine/Source/Runtime/Engine/Private/ShaderCompiler/ShaderCompiler.cpp
--- a/Engine/Source/Runtime/Engine/Private/ShaderCompiler/ShaderCompiler.cpp
+++ b/Engine/Source/Runtime/Engine/Private/ShaderCompiler/ShaderCompiler.cpp
@@ -5462,6 +5462,17 @@ void GlobalBeginCompileShader(
 		Input.Environment.SetDefine(TEXT("RAYMISSSHADER"), Target.Frequency == SF_RayMiss);
 	}
 
+	// Enables HLSL 2021
+	uint32 EnablesHLSL2021ByDefault = FDataDrivenShaderPlatformInfo::GetEnablesHLSL2021ByDefault(EShaderPlatform(Target.Platform));
+	if (EnablesHLSL2021ByDefault == uint32(1) && DebugGroupName == TEXT("Global"))
+	{
+		Input.Environment.CompilerFlags.Add(CFLAG_HLSL2021);
+	}
+	else if (EnablesHLSL2021ByDefault == uint32(2))
+	{
+		Input.Environment.CompilerFlags.Add(CFLAG_HLSL2021);
+	}
+
 	// #defines get stripped out by the preprocessor without this. We can override with this
 	Input.Environment.SetDefine(TEXT("COMPILER_DEFINE"), TEXT("#define"));
 
diff --git a/Engine/Source/Runtime/RHI/Private/RHI.cpp b/Engine/Source/Runtime/RHI/Private/RHI.cpp
--- a/Engine/Source/Runtime/RHI/Private/RHI.cpp
+++ b/Engine/Source/Runtime/RHI/Private/RHI.cpp
@@ -2504,6 +2504,7 @@ void FGenericDataDrivenShaderPlatformInfo::ParseDataDrivenShaderInfo(const FConf
 	GET_SECTION_BOOL_HELPER(bSupportsROV);
 	GET_SECTION_BOOL_HELPER(bSupportsOIT);
 	GET_SECTION_SUPPORT_HELPER(bSupportsRealTypes);
+	GET_SECTION_INT_HELPER(EnablesHLSL2021ByDefault);
 #undef GET_SECTION_BOOL_HELPER
 #undef GET_SECTION_INT_HELPER
 #undef GET_SECTION_SUPPORT_HELPER
diff --git a/Engine/Source/Runtime/RHI/Public/RHIDefinitions.h b/Engine/Source/Runtime/RHI/Public/RHIDefinitions.h
--- a/Engine/Source/Runtime/RHI/Public/RHIDefinitions.h
+++ b/Engine/Source/Runtime/RHI/Public/RHIDefinitions.h
@@ -413,6 +413,7 @@ class RHI_API FGenericDataDrivenShaderPlatformInfo
 	uint32 bSupportsROV : 1;
 	uint32 bSupportsOIT : 1;
 	uint32 bSupportsRealTypes : int32(ERHIFeatureSupport::NumBits);
+	uint32 EnablesHLSL2021ByDefault : 2; // 0: disabled, 1: global shaders only, 2: all shaders
 		
 #if WITH_EDITOR
 	FText FriendlyName;
@@ -946,6 +947,11 @@ public:
 		return ERHIFeatureSupport(Infos[Platform].bSupportsRealTypes);
 	}
 
+	static FORCEINLINE_DEBUGGABLE const uint32 GetEnablesHLSL2021ByDefault(const FStaticShaderPlatform Platform)
+	{
+		return Infos[Platform].EnablesHLSL2021ByDefault;
+	}
+
 #if WITH_EDITOR
 	static FORCEINLINE_DEBUGGABLE FText GetFriendlyName(const FStaticShaderPlatform Platform)
 	{
