VulkanRHI:  Set initial state of GParticleCurveTexture to SRVMask since it might get bound before it's filled for the first time.

diff --git a/Engine/Source/Runtime/Engine/Private/Particles/ParticleCurveTexture.cpp b/Engine/Source/Runtime/Engine/Private/Particles/ParticleCurveTexture.cpp
--- a/Engine/Source/Runtime/Engine/Private/Particles/ParticleCurveTexture.cpp
+++ b/Engine/Source/Runtime/Engine/Private/Particles/ParticleCurveTexture.cpp
@@ -519,7 +519,8 @@ void FParticleCurveTexture::InitRHI()
 		.SetExtent(GParticleCurveTextureSizeX, GParticleCurveTextureSizeY)
 		.SetFlags(ETextureCreateFlags::RenderTargetable | ETextureCreateFlags::ShaderResource | ETextureCreateFlags::NoFastClear)
 		.SetFormat(PF_B8G8R8A8)
-		.SetClearValue(FClearValueBinding(FLinearColor::Blue));
+		.SetClearValue(FClearValueBinding(FLinearColor::Blue))
+		.SetInitialState(ERHIAccess::SRVMask);
 
 	CurveTextureRHI = RHICreateTexture(Desc);
 }
