Engine nullptr fixups

diff --git a/Engine/Source/Runtime/Engine/Private/Actor.cpp b/Engine/Source/Runtime/Engine/Private/Actor.cpp
--- a/Engine/Source/Runtime/Engine/Private/Actor.cpp
+++ b/Engine/Source/Runtime/Engine/Private/Actor.cpp
@@ -5429,7 +5429,7 @@ void AActor::SetLifeSpan( float InLifespan )
 	// Store the new value
 	InitialLifeSpan = InLifespan;
 	// Initialize a timer for the actors lifespan if there is one. Otherwise clear any existing timer
-	if ((GetLocalRole() == ROLE_Authority || GetTearOff()) && IsValidChecked(this))
+	if ((GetLocalRole() == ROLE_Authority || GetTearOff()) && IsValidChecked(this) && GetWorld())
 	{
 		if( InLifespan > 0.0f)
 		{
diff --git a/Engine/Source/Runtime/Engine/Private/World.cpp b/Engine/Source/Runtime/Engine/Private/World.cpp
--- a/Engine/Source/Runtime/Engine/Private/World.cpp
+++ b/Engine/Source/Runtime/Engine/Private/World.cpp
@@ -2138,24 +2138,26 @@ UWorld* UWorld::CreateWorld(const EWorldType::Type InWorldType, bool bInformEngi
 
 void UWorld::RemoveActor(AActor* Actor, bool bShouldModifyLevel) const
 {
-	ULevel* CheckLevel = Actor->GetLevel();
-	const int32 ActorListIndex = CheckLevel->Actors.Find( Actor );
-	// Search the entire list.
-	if( ActorListIndex != INDEX_NONE )
+	if (ULevel* CheckLevel = Actor->GetLevel())
 	{
-		if ( bShouldModifyLevel && GUndo )
+		const int32 ActorListIndex = CheckLevel->Actors.Find(Actor);
+		// Search the entire list.
+		if (ActorListIndex != INDEX_NONE)
 		{
-			ModifyLevel( CheckLevel );
-		}
-		
-		if (!IsGameWorld())
-		{
-			CheckLevel->Actors[ActorListIndex]->Modify();
-		}
-		
-		CheckLevel->Actors[ActorListIndex] = nullptr;
+			if (bShouldModifyLevel && GUndo)
+			{
+				ModifyLevel(CheckLevel);
+			}
+
+			if (!IsGameWorld())
+			{
+				CheckLevel->Actors[ActorListIndex]->Modify();
+			}
+
+			CheckLevel->Actors[ActorListIndex] = nullptr;
 
-		CheckLevel->ActorsForGC.RemoveSwap(Actor);
+			CheckLevel->ActorsForGC.RemoveSwap(Actor);
+		}
 	}
 
 	// Remove actor from network list
