Fix crash when material is auto created by python script, in which case GraphNode is null.

diff --git a/Engine/Source/Runtime/Engine/Private/Materials/MaterialExpressions.cpp b/Engine/Source/Runtime/Engine/Private/Materials/MaterialExpressions.cpp
--- a/Engine/Source/Runtime/Engine/Private/Materials/MaterialExpressions.cpp
+++ b/Engine/Source/Runtime/Engine/Private/Materials/MaterialExpressions.cpp
@@ -1448,7 +1448,7 @@ void UMaterialExpression::PostEditChangeProperty(FPropertyChangedEvent& Property
 
 	// PropertyChangedEvent.MemberProperty is the owner of PropertyChangedEvent.Property so check for MemberProperty
 	FProperty* MemberPropertyThatChanged = PropertyChangedEvent.MemberProperty;
-	if (MemberPropertyThatChanged != nullptr)
+	if (MemberPropertyThatChanged != nullptr && GraphNode)
 	{
 		int32 PinIndex = -1;
 		const TArray<FExpressionInput*> AllInputs = GetInputs();
@@ -1489,6 +1489,7 @@ void UMaterialExpression::PostEditChangeProperty(FPropertyChangedEvent& Property
 		{
 			// Update the pin value of the expression input
 			UEdGraphPin* Pin = GraphNode->GetPinAt(PinIndex);
+			check(Pin);
 			Pin->DefaultValue = GetInputPinDefaultValue(PinIndex);
 
 			// If the property is linked as inline toggle to another property, both pins need updating to reflect the change.
@@ -1496,10 +1497,7 @@ void UMaterialExpression::PostEditChangeProperty(FPropertyChangedEvent& Property
 			bool bEditCondition = MemberPropertyThatChanged->HasMetaData(TEXT("EditCondition"));
 			if (bInlineEditConditionToggle || bEditCondition)
 			{
-				if (GraphNode)
-				{
-					CastChecked<UMaterialGraphNode>(GraphNode)->RecreateAndLinkNode();
-				}
+				CastChecked<UMaterialGraphNode>(GraphNode)->RecreateAndLinkNode();
 			}
 		}
 	}
