Scene Capture:  Added Blueprint function node "DeprojectSceneCaptureToWorld" to support translating UV coordinates in scene capture space to a world space ray, similar to "DeprojectScreenToWorld" which operates on the player camera.  Allows hit tracing operations into scene capture surfaces without requiring C++ code.

diff --git a/Engine/Source/Runtime/Engine/Classes/Kismet/GameplayStatics.h b/Engine/Source/Runtime/Engine/Classes/Kismet/GameplayStatics.h
--- a/Engine/Source/Runtime/Engine/Classes/Kismet/GameplayStatics.h
+++ b/Engine/Source/Runtime/Engine/Classes/Kismet/GameplayStatics.h
@@ -19,6 +19,7 @@
 #include "Particles/WorldPSCPool.h"
 #include "GameplayStatics.generated.h"
 
+class ASceneCapture2D;
 class UAudioComponent;
 class UInitialActiveSoundParams;
 class UBlueprint;
@@ -1335,6 +1336,16 @@ public:
 	UFUNCTION(BlueprintPure, Category = "Camera", meta = (Keywords = "unproject"))
 	static bool DeprojectScreenToWorld(APlayerController const* Player, const FVector2D& ScreenPosition, FVector& WorldPosition, FVector& WorldDirection);
 
+	/** 
+	 * Transforms the given 2D UV coordinate into a 3D world-space point and direction.
+	 * @param SceneCapture2D	Deproject using this scene capture's view.
+	 * @param ScreenPosition	UV in scene capture render target to deproject.
+	 * @param WorldPosition		(out) Corresponding 3D position on camera near plane, in world space.
+	 * @param WorldDirection	(out) World space direction vector away from the camera at the given 2d point.
+	 */
+	UFUNCTION(BlueprintPure, Category = "Camera", meta = (Keywords = "unproject"))
+	static bool DeprojectSceneCaptureToWorld(ASceneCapture2D const* SceneCapture2D, const FVector2D& TargetUV, FVector& WorldPosition, FVector& WorldDirection);
+
 	/** 
 	 * Transforms the given 3D world-space point into a its 2D screen space coordinate. 
 	 * @param Player			Project using this player's view.
diff --git a/Engine/Source/Runtime/Engine/Private/GameplayStatics.cpp b/Engine/Source/Runtime/Engine/Private/GameplayStatics.cpp
--- a/Engine/Source/Runtime/Engine/Private/GameplayStatics.cpp
+++ b/Engine/Source/Runtime/Engine/Private/GameplayStatics.cpp
@@ -45,6 +45,7 @@
 #include "ContentStreaming.h"
 #include "Async/Async.h"
 #include "Engine/SceneCapture2D.h"
+#include "Engine/TextureRenderTarget2D.h"
 #include "Components/SceneCaptureComponent2D.h"
 #include "Sound/SoundCue.h"
 #include "Sound/SoundWave.h"
@@ -2992,6 +2993,56 @@ bool UGameplayStatics::DeprojectScreenToWorld(APlayerController const* Player, c
 	return false;
 }
 
+bool UGameplayStatics::DeprojectSceneCaptureToWorld(ASceneCapture2D const* SceneCapture2D, const FVector2D& TargetUV, FVector& WorldPosition, FVector& WorldDirection)
+{
+	if (USceneCaptureComponent2D* SceneCaptureComponent2D = SceneCapture2D->GetCaptureComponent2D())
+	{
+		if (SceneCaptureComponent2D->TextureTarget)
+		{
+			FMinimalViewInfo ViewInfo;
+			SceneCaptureComponent2D->GetCameraView(0.0f, ViewInfo);
+
+			FMatrix ProjectionMatrix;
+			if (SceneCaptureComponent2D->bUseCustomProjectionMatrix)
+			{
+				ProjectionMatrix = AdjustProjectionMatrixForRHI(SceneCaptureComponent2D->CustomProjectionMatrix);
+			}
+			else//
+			{
+				ProjectionMatrix = AdjustProjectionMatrixForRHI(ViewInfo.CalculateProjectionMatrix());
+			}
+			FMatrix InvProjectionMatrix = ProjectionMatrix.Inverse();
+
+			// A view matrix is the inverse of the viewer's matrix, so an inverse view matrix is just the viewer's matrix.
+			// To save precision, we directly compute the viewer's matrix, plus it also avoids the cost of the inverse.
+			// The matrix to convert from world coordinate space to view coordinate space also needs to be included (this
+			// is the transpose of the similar matrix used in CalculateViewProjectionMatricesFromMinimalView).
+			FMatrix InvViewMatrix = FMatrix(
+				FPlane(0, 1, 0, 0),
+				FPlane(0, 0, 1, 0),
+				FPlane(1, 0, 0, 0),
+				FPlane(0, 0, 0, 1)) * FRotationTranslationMatrix(ViewInfo.Rotation, ViewInfo.Location);
+
+			FIntPoint TargetSize = FIntPoint(SceneCaptureComponent2D->TextureTarget->SizeX, SceneCaptureComponent2D->TextureTarget->SizeY);
+
+			FSceneView::DeprojectScreenToWorld(
+				TargetUV * FVector2D(TargetSize),
+				FIntRect(FIntPoint(0, 0), TargetSize),
+				InvViewMatrix,
+				InvProjectionMatrix,
+				WorldPosition,
+				WorldDirection);
+
+			return true;
+		}
+	}
+
+	// something went wrong, zero things and return false
+	WorldPosition = FVector::ZeroVector;
+	WorldDirection = FVector::ZeroVector;
+	return false;
+}
+
 bool UGameplayStatics::ProjectWorldToScreen(APlayerController const* Player, const FVector& WorldPosition, FVector2D& ScreenPosition, bool bPlayerViewportRelative)
 {
 	ULocalPlayer* const LP = Player ? Player->GetLocalPlayer() : nullptr;
