WorldPartition: Server now delays unloads of client visible cells

diff --git a/Engine/Source/Runtime/Engine/Private/WorldPartition/WorldPartitionStreamingPolicy.cpp b/Engine/Source/Runtime/Engine/Private/WorldPartition/WorldPartitionStreamingPolicy.cpp
--- a/Engine/Source/Runtime/Engine/Private/WorldPartition/WorldPartitionStreamingPolicy.cpp
+++ b/Engine/Source/Runtime/Engine/Private/WorldPartition/WorldPartitionStreamingPolicy.cpp
@@ -20,6 +20,8 @@
 #include "Engine/World.h"
 #include "Engine/Engine.h"
 #include "Engine/Canvas.h"
+#include "Engine/NetDriver.h"
+#include "Engine/NetConnection.h"
 #include "DrawDebugHelpers.h"
 #include "DisplayDebugHelpers.h"
 #include "RenderUtils.h"
@@ -226,14 +228,20 @@ void UWorldPartitionStreamingPolicy::UpdateStreamingState()
 	
 	// Update streaming sources
 	UpdateStreamingSources();
-		
+	
+	TSet<FName> ClientVisibleLevelNames;
+	bool bUpdateEpoch = false;
+
 	TSet<const UWorldPartitionRuntimeCell*>& ActivateStreamingCells = FrameActivateCells.GetCells();
 	TSet<const UWorldPartitionRuntimeCell*>& LoadStreamingCells = FrameLoadCells.GetCells();
+		
 	check(ActivateStreamingCells.IsEmpty());
 	check(LoadStreamingCells.IsEmpty());
 
 	const ENetMode NetMode = World->GetNetMode();
-	if (NetMode == NM_Standalone || NetMode == NM_Client || AWorldPartitionReplay::IsEnabled(World))
+	
+	bool bClient = NetMode == NM_Standalone || NetMode == NM_Client || AWorldPartitionReplay::IsEnabled(World);
+	if (bClient)
 	{
 		// When uninitializing, UpdateStreamingState is called, but we don't want any cells to be loaded
 		if (WorldPartition->IsInitialized())
@@ -254,8 +262,18 @@ void UWorldPartitionStreamingPolicy::UpdateStreamingState()
 			return; 
 		}
 
+		bUpdateEpoch = true;
 		const UDataLayerSubsystem* DataLayerSubsystem = WorldPartition->GetWorld()->GetSubsystem<UDataLayerSubsystem>();
-		DataLayersStatesServerEpoch = AWorldDataLayers::GetDataLayersStateEpoch();
+		
+		// Gather Client visible level names
+		if (const UNetDriver* NetDriver = World->GetNetDriver())
+		{
+			for (UNetConnection* Connection : NetDriver->ClientConnections)
+			{
+				ClientVisibleLevelNames.Add(Connection->GetClientWorldPackageName());
+				ClientVisibleLevelNames.Append(Connection->ClientVisibleLevelNames);
+			}
+		}
 
 		// Non Data Layer Cells + Active Data Layers
 		WorldPartition->RuntimeHash->GetAllStreamingCells(ActivateStreamingCells, /*bAllDataLayers=*/ false, /*bDataLayersOnly=*/ false, DataLayerSubsystem->GetEffectiveActiveDataLayerNames());
@@ -301,9 +319,28 @@ void UWorldPartitionStreamingPolicy::UpdateStreamingState()
 	{
 		WORLDPARTITION_LOG_UPDATESTREAMINGSTATE(Log);
 	}
-	
+		
+	// Remove cells state changes based on client visibility. If any cell gets removed, delay epoch change.
+	auto ProcessClientVisibility = [bClient, &ClientVisibleLevelNames, &bUpdateEpoch](TSet<const UWorldPartitionRuntimeCell*>& Cells)
+	{
+		if (!bClient)
+		{
+			for (TSet<const UWorldPartitionRuntimeCell*>::TIterator It = Cells.CreateIterator(); It; ++It)
+			{
+				if (ClientVisibleLevelNames.Contains((*It)->GetLevel()->GetPackage()->GetFName()))
+				{
+					It.RemoveCurrent();
+
+					UE_CLOG(bUpdateEpoch, LogWorldPartition, Verbose, TEXT("Server epoch update delayed by client visibility"));
+					bUpdateEpoch = false;
+				}
+			}
+		}
+	};
+
 	if (ToUnloadCells.Num() > 0)
 	{
+		ProcessClientVisibility(ToUnloadCells);
 		SetTargetStateForCells(EWorldPartitionRuntimeCellState::Unloaded, ToUnloadCells);
 	}
 
@@ -315,6 +352,7 @@ void UWorldPartitionStreamingPolicy::UpdateStreamingState()
 
 	if (ToLoadCells.Num() > 0)
 	{
+		ProcessClientVisibility(ToLoadCells);
 		SetTargetStateForCells(EWorldPartitionRuntimeCellState::Loaded, ToLoadCells);
 	}
 
@@ -350,6 +388,13 @@ void UWorldPartitionStreamingPolicy::UpdateStreamingState()
 	// Reset frame StreamingSourceCells (optimization to avoid reallocation at every call to UpdateStreamingState)
 	FrameActivateCells.Reset();
 	FrameLoadCells.Reset();
+
+	// Update Epoch if we aren't waiting for clients anymore
+	if (bUpdateEpoch)
+	{
+		DataLayersStatesServerEpoch = AWorldDataLayers::GetDataLayersStateEpoch();
+		UE_LOG(LogWorldPartition, Verbose, TEXT("Server epoch updated"));
+	}
 }
 
 #if !UE_BUILD_SHIPPING
