UE-151832 LyraCameraComponent doesn't work with HMD's
-Refactored UCameraComponent a little so the XR handling is pulled into functions.
-Added a check for XR and call to the XR camera handling to the Lyra camera GetViews.

diff --git a/Engine/Source/Runtime/Engine/Classes/Camera/CameraComponent.h b/Engine/Source/Runtime/Engine/Classes/Camera/CameraComponent.h
--- a/Engine/Source/Runtime/Engine/Classes/Camera/CameraComponent.h
+++ b/Engine/Source/Runtime/Engine/Classes/Camera/CameraComponent.h
@@ -136,6 +136,11 @@ public:
 #endif
 	// End of UObject interface
 
+protected:
+	bool IsXRHeadTrackedCamera() const;
+	virtual void HandleXRCamera();
+
+public:
 	/**
 	 * Returns camera's Point of View.
 	 * Called by Camera class. Subclass and postprocess to add any effects.
diff --git a/Engine/Source/Runtime/Engine/Private/Camera/CameraComponent.cpp b/Engine/Source/Runtime/Engine/Private/Camera/CameraComponent.cpp
--- a/Engine/Source/Runtime/Engine/Private/Camera/CameraComponent.cpp
+++ b/Engine/Source/Runtime/Engine/Private/Camera/CameraComponent.cpp
@@ -307,9 +307,9 @@ void UCameraComponent::OnCameraMeshHiddenChanged()
 #endif
 }
 
-void UCameraComponent::GetCameraView(float DeltaTime, FMinimalViewInfo& DesiredView)
+bool UCameraComponent::IsXRHeadTrackedCamera() const
 {
-	if (GEngine && GEngine->XRSystem.IsValid() && GetWorld() && GetWorld()->WorldType != EWorldType::Editor )
+	if (GEngine && GEngine->XRSystem.IsValid() && GetWorld() && GetWorld()->WorldType != EWorldType::Editor)
 	{
 		IXRTrackingSystem* XRSystem = GEngine->XRSystem.Get();
 		auto XRCamera = XRSystem->GetXRCamera();
@@ -318,29 +318,46 @@ void UCameraComponent::GetCameraView(float DeltaTime, FMinimalViewInfo& DesiredV
 		{
 			if (XRSystem->IsHeadTrackingAllowedForWorld(*GetWorld()))
 			{
-				const FTransform ParentWorld = CalcNewComponentToWorld(FTransform());
-
-				XRCamera->SetupLateUpdate(ParentWorld, this, bLockToHmd == 0);
-
-				if (bLockToHmd)
-				{
-					FQuat Orientation;
-					FVector Position;
-					if (XRCamera->UpdatePlayerCamera(Orientation, Position))
-					{
-						SetRelativeTransform(FTransform(Orientation, Position));
-					}
-					else
-					{
-						ResetRelativeTransform();
-					}
-				}
-
-				XRCamera->OverrideFOV(this->FieldOfView);
+				return true;
 			}
 		}
 	}
 
+	return false;
+}
+
+void UCameraComponent::HandleXRCamera()
+{
+	IXRTrackingSystem* XRSystem = GEngine->XRSystem.Get();
+	auto XRCamera = XRSystem->GetXRCamera();
+	const FTransform ParentWorld = CalcNewComponentToWorld(FTransform());
+
+	XRCamera->SetupLateUpdate(ParentWorld, this, bLockToHmd == 0);
+
+	if (bLockToHmd)
+	{
+		FQuat Orientation;
+		FVector Position;
+		if (XRCamera->UpdatePlayerCamera(Orientation, Position))
+		{
+			SetRelativeTransform(FTransform(Orientation, Position));
+		}
+		else
+		{
+			ResetRelativeTransform();
+		}
+	}
+
+	XRCamera->OverrideFOV(this->FieldOfView);
+}
+
+void UCameraComponent::GetCameraView(float DeltaTime, FMinimalViewInfo& DesiredView)
+{
+	if (IsXRHeadTrackedCamera())
+	{
+		HandleXRCamera();
+	}
+
 	if (bUsePawnControlRotation)
 	{
 		const APawn* OwningPawn = Cast<APawn>(GetOwner());
