Added draw distance scale for skeletal mesh overlay (r.ViewDistanceScale.SkeletalMeshOverlay)

diff --git a/Engine/Source/Runtime/Core/Private/HAL/ConsoleManager.cpp b/Engine/Source/Runtime/Core/Private/HAL/ConsoleManager.cpp
--- a/Engine/Source/Runtime/Core/Private/HAL/ConsoleManager.cpp
+++ b/Engine/Source/Runtime/Core/Private/HAL/ConsoleManager.cpp
@@ -3062,6 +3062,13 @@ static TAutoConsoleVariable<int32> CVarViewDistanceScale_FieldOfViewAffectsHLOD(
 	TEXT("If enabled, applies the field of view scaling to HLOD draw distances as well as non-HLODs."),
 	ECVF_Scalability | ECVF_RenderThreadSafe);
 
+static TAutoConsoleVariable<float> CVarViewDistanceScale_SkeletalMeshOverlay(
+	TEXT("r.ViewDistanceScale.SkeletalMeshOverlay"),
+	1.f,
+	TEXT("Controls the distance scale for skeletal mesh overlay, Default = 1.0. \n"
+		 "Higher values will increase skeletal mesh overlay draw distance. This value is applied together with r.ViewDistanceScale"),
+	ECVF_Scalability | ECVF_RenderThreadSafe);
+
 static TAutoConsoleVariable<int32> CVarLightFunctionQuality(
 	TEXT("r.LightFunctionQuality"),
 	2,
diff --git a/Engine/Source/Runtime/Engine/Private/SkeletalMesh.cpp b/Engine/Source/Runtime/Engine/Private/SkeletalMesh.cpp
--- a/Engine/Source/Runtime/Engine/Private/SkeletalMesh.cpp
+++ b/Engine/Source/Runtime/Engine/Private/SkeletalMesh.cpp
@@ -6740,7 +6740,8 @@ void FSkeletalMeshSceneProxy::GetDynamicElementsSection(const TArray<const FScen
 				bool bAddOverlay = true;
 				if (bHasOverlayCullDistance)
 				{
-					float MaxDrawDistanceScale = GetCachedScalabilityCVars().ViewDistanceScale;
+					// this is already combined with ViewDistanceScale
+					float MaxDrawDistanceScale = GetCachedScalabilityCVars().SkeletalMeshOverlayDistanceScale;
 					MaxDrawDistanceScale *= GetCachedScalabilityCVars().CalculateFieldOfViewDistanceScale(View->DesiredFOV);
 					float DistanceSquared = (GetBounds().Origin - View->ViewMatrices.GetViewOrigin()).SizeSquared();
 					if (DistanceSquared > FMath::Square(OverlayMaterialMaxDrawDistance * MaxDrawDistanceScale))
diff --git a/Engine/Source/Runtime/Engine/Private/UnrealEngine.cpp b/Engine/Source/Runtime/Engine/Private/UnrealEngine.cpp
--- a/Engine/Source/Runtime/Engine/Private/UnrealEngine.cpp
+++ b/Engine/Source/Runtime/Engine/Private/UnrealEngine.cpp
@@ -647,6 +647,7 @@ FCachedSystemScalabilityCVars::FCachedSystemScalabilityCVars()
 	, ViewDistanceScale(-1)
 	, ViewDistanceScaleSquared(-1)
 	, StaticMeshLODDistanceScale(1.0f)
+	, SkeletalMeshOverlayDistanceScale(1.0f)
 	, MaxAnisotropy(-1)
 	
 {
@@ -661,7 +662,8 @@ bool FCachedSystemScalabilityCVars::operator==(const FCachedSystemScalabilityCVa
 		MaxCSMShadowResolution == Other.MaxCSMShadowResolution &&
 		ViewDistanceScale == Other.ViewDistanceScale &&
 		ViewDistanceScaleSquared == Other.ViewDistanceScaleSquared &&
-		StaticMeshLODDistanceScale == Other.StaticMeshLODDistanceScale;
+		StaticMeshLODDistanceScale == Other.StaticMeshLODDistanceScale &&
+		SkeletalMeshOverlayDistanceScale == Other.SkeletalMeshOverlayDistanceScale;
 }
 
 void ScalabilityCVarsSinkCallback()
@@ -728,6 +730,12 @@ void ScalabilityCVarsSinkCallback()
 		LocalScalabilityCVars.StaticMeshLODDistanceScale = StaticMeshLODDistanceScale->GetFloat();
 	}
 
+	{
+		static const auto SkeletalMeshOverlayDistanceScale = ConsoleMan.FindTConsoleVariableDataFloat(TEXT("r.ViewDistanceScale.SkeletalMeshOverlay"));
+		LocalScalabilityCVars.SkeletalMeshOverlayDistanceScale = FMath::Max(SkeletalMeshOverlayDistanceScale->GetValueOnGameThread(), 0.0f);
+		LocalScalabilityCVars.SkeletalMeshOverlayDistanceScale *= LocalScalabilityCVars.ViewDistanceScale;
+	}
+
 	LocalScalabilityCVars.bInitialized = true;
 
 	if ((LocalScalabilityCVars == GCachedScalabilityCVars) && GCachedScalabilityCVars.bInitialized)
diff --git a/Engine/Source/Runtime/Engine/Public/UnrealEngine.h b/Engine/Source/Runtime/Engine/Public/UnrealEngine.h
--- a/Engine/Source/Runtime/Engine/Public/UnrealEngine.h
+++ b/Engine/Source/Runtime/Engine/Public/UnrealEngine.h
@@ -425,6 +425,7 @@ struct FCachedSystemScalabilityCVars
 	float ViewDistanceScaleSquared;
 	int32 FieldOfViewAffectsHLOD;
 	float StaticMeshLODDistanceScale;
+	float SkeletalMeshOverlayDistanceScale;
 
 	float CalculateFieldOfViewDistanceScale(const float FieldOfView) const
 	{
