Make FMaterialCachedExpressionData robust to handling missing editor only data when running the cooked editor w/o .o asset files.

- This code existed previously existed, and needed to be resurrected.

diff --git a/Engine/Source/Runtime/Engine/Private/Materials/MaterialCachedData.cpp b/Engine/Source/Runtime/Engine/Private/Materials/MaterialCachedData.cpp
--- a/Engine/Source/Runtime/Engine/Private/Materials/MaterialCachedData.cpp
+++ b/Engine/Source/Runtime/Engine/Private/Materials/MaterialCachedData.cpp
@@ -537,17 +537,22 @@ int32 FMaterialCachedExpressionData::FindParameterIndex(EMaterialParameterType T
 void FMaterialCachedExpressionData::GetParameterValueByIndex(EMaterialParameterType Type, int32 ParameterIndex, FMaterialParameterMetadata& OutResult) const
 {
 #if WITH_EDITORONLY_DATA
+	bool bIsEditorOnlyDataStripped = true;
 	if (EditorOnlyData)
 	{
 		const FMaterialCachedParameterEditorEntry& EditorEntry = EditorOnlyData->EditorEntries[(int32)Type];
-		const FMaterialCachedParameterEditorInfo& EditorInfo = EditorEntry.EditorInfo[ParameterIndex];
-		OutResult.ExpressionGuid = EditorInfo.ExpressionGuid;
-		OutResult.Description = EditorInfo.Description;
-		OutResult.Group = EditorInfo.Group;
-		OutResult.SortPriority = EditorInfo.SortPriority;
-		if (EditorInfo.AssetIndex != INDEX_NONE)
+		bIsEditorOnlyDataStripped = EditorEntry.EditorInfo.Num() == 0;
+		if (!bIsEditorOnlyDataStripped)
 		{
-			OutResult.AssetPath = EditorOnlyData->AssetPaths[EditorInfo.AssetIndex];
+			const FMaterialCachedParameterEditorInfo& EditorInfo = EditorEntry.EditorInfo[ParameterIndex];
+			OutResult.ExpressionGuid = EditorInfo.ExpressionGuid;
+			OutResult.Description = EditorInfo.Description;
+			OutResult.Group = EditorInfo.Group;
+			OutResult.SortPriority = EditorInfo.SortPriority;
+			if (EditorInfo.AssetIndex != INDEX_NONE)
+			{
+				OutResult.AssetPath = EditorOnlyData->AssetPaths[EditorInfo.AssetIndex];
+			}
 		}
 	}
 #endif // WITH_EDITORONLY_DATA
@@ -558,7 +563,7 @@ void FMaterialCachedExpressionData::GetParameterValueByIndex(EMaterialParameterT
 		OutResult.Value = ScalarValues[ParameterIndex];
 		OutResult.PrimitiveDataIndex = ScalarPrimitiveDataIndexValues[ParameterIndex];
 #if WITH_EDITORONLY_DATA
-		if (EditorOnlyData)
+		if (EditorOnlyData && !bIsEditorOnlyDataStripped)
 		{
 			OutResult.ScalarMin = EditorOnlyData->ScalarMinMaxValues[ParameterIndex].X;
 			OutResult.ScalarMax = EditorOnlyData->ScalarMinMaxValues[ParameterIndex].Y;
@@ -579,7 +584,7 @@ void FMaterialCachedExpressionData::GetParameterValueByIndex(EMaterialParameterT
 		OutResult.Value = VectorValues[ParameterIndex];
 		OutResult.PrimitiveDataIndex = VectorPrimitiveDataIndexValues[ParameterIndex];
 #if  WITH_EDITORONLY_DATA
-		if (EditorOnlyData)
+		if (EditorOnlyData && !bIsEditorOnlyDataStripped)
 		{
 			OutResult.ChannelNames = EditorOnlyData->VectorChannelNameValues[ParameterIndex];
 			OutResult.bUsedAsChannelMask = EditorOnlyData->VectorUsedAsChannelMaskValues[ParameterIndex];
@@ -592,7 +597,7 @@ void FMaterialCachedExpressionData::GetParameterValueByIndex(EMaterialParameterT
 	case EMaterialParameterType::Texture:
 		OutResult.Value = TextureValues[ParameterIndex].LoadSynchronous();
 #if WITH_EDITORONLY_DATA
-		if (EditorOnlyData)
+		if (EditorOnlyData && !bIsEditorOnlyDataStripped)
 		{
 			OutResult.ChannelNames = EditorOnlyData->TextureChannelNameValues[ParameterIndex];
 		}
@@ -606,13 +611,13 @@ void FMaterialCachedExpressionData::GetParameterValueByIndex(EMaterialParameterT
 		break;
 #if WITH_EDITORONLY_DATA
 	case EMaterialParameterType::StaticSwitch:
-		if (EditorOnlyData)
+		if (EditorOnlyData && !bIsEditorOnlyDataStripped)
 		{
 			OutResult.Value = EditorOnlyData->StaticSwitchValues[ParameterIndex];
 		}
 		break;
 	case EMaterialParameterType::StaticComponentMask:
-		if (EditorOnlyData)
+		if (EditorOnlyData && !bIsEditorOnlyDataStripped)
 		{
 			OutResult.Value = EditorOnlyData->StaticComponentMaskValues[ParameterIndex];
 		}
@@ -641,7 +646,11 @@ const FGuid& FMaterialCachedExpressionData::GetExpressionGuid(EMaterialParameter
 #if WITH_EDITORONLY_DATA
 	if (EditorOnlyData)
 	{
-		return EditorOnlyData->EditorEntries[(int32)Type].EditorInfo[Index].ExpressionGuid;
+		// cooked materials can strip out expression guids
+		if (EditorOnlyData->EditorEntries[(int32)Type].EditorInfo.Num() != 0)
+		{
+			return EditorOnlyData->EditorEntries[(int32)Type].EditorInfo[Index].ExpressionGuid;
+		}
 	}
 #endif // WITH_EDITORONLY_DATA
 	static const FGuid EmptyGuid;
