[Replay]
Make sure that when playing replay, we rely more on the pawn replication than on the movementComponent replication which might be delayed when using some root motion source

#RNX

diff --git a/Engine/Source/Runtime/Engine/Private/Character.cpp b/Engine/Source/Runtime/Engine/Private/Character.cpp
--- a/Engine/Source/Runtime/Engine/Private/Character.cpp
+++ b/Engine/Source/Runtime/Engine/Private/Character.cpp
@@ -1169,6 +1169,12 @@ void ACharacter::OnRep_ReplicatedMovement()
 		return;
 	}
 
+	if (GetWorld()->IsPlayingReplay())
+	{
+		Super::OnRep_ReplicatedMovement();
+		return;
+	}
+
 	// Skip standard position correction if we are playing root motion, OnRep_RootMotion will handle it.
 	if (!IsPlayingNetworkedRootMotionMontage()) // animation root motion
 	{
diff --git a/Engine/Source/Runtime/Engine/Private/Components/CharacterMovementComponent.cpp b/Engine/Source/Runtime/Engine/Private/Components/CharacterMovementComponent.cpp
--- a/Engine/Source/Runtime/Engine/Private/Components/CharacterMovementComponent.cpp
+++ b/Engine/Source/Runtime/Engine/Private/Components/CharacterMovementComponent.cpp
@@ -1626,7 +1626,7 @@ void UCharacterMovementComponent::SimulatedTick(float DeltaSeconds)
 	SCOPE_CYCLE_COUNTER(STAT_CharacterMovementSimulated);
 	checkSlow(CharacterOwner != nullptr);
 
-	if (NetworkSmoothingMode == ENetworkSmoothingMode::Replay)
+	if (GetWorld()->IsPlayingReplay())
 	{
 		const FVector OldLocation = UpdatedComponent ? UpdatedComponent->GetComponentLocation() : FVector::ZeroVector;
 		const FVector OldVelocity = Velocity;
