Allow the editor-only material data to be null since it's optional data

#rnx

diff --git a/Engine/Source/Runtime/Engine/Classes/Materials/Material.h b/Engine/Source/Runtime/Engine/Classes/Materials/Material.h
--- a/Engine/Source/Runtime/Engine/Classes/Materials/Material.h
+++ b/Engine/Source/Runtime/Engine/Classes/Materials/Material.h
@@ -401,8 +401,8 @@ class UMaterial : public UMaterialInterface
 #if WITH_EDITORONLY_DATA
 	ENGINE_API virtual const UClass* GetEditorOnlyDataClass() const override { return UMaterialEditorOnlyData::StaticClass(); }
 
-	UMaterialEditorOnlyData* GetEditorOnlyData() { return CastChecked<UMaterialEditorOnlyData>(Super::GetEditorOnlyData()); }
-	const UMaterialEditorOnlyData* GetEditorOnlyData() const { return CastChecked<UMaterialEditorOnlyData>(Super::GetEditorOnlyData()); }
+	virtual UMaterialEditorOnlyData* GetEditorOnlyData() override { return CastChecked<UMaterialEditorOnlyData>(Super::GetEditorOnlyData(), ECastCheckedType::NullAllowed); }
+	virtual const UMaterialEditorOnlyData* GetEditorOnlyData() const override { return CastChecked<UMaterialEditorOnlyData>(Super::GetEditorOnlyData(), ECastCheckedType::NullAllowed); }
 #endif // WITH_EDITORONLY_DATA
 
 	// Physics.
diff --git a/Engine/Source/Runtime/Engine/Classes/Materials/MaterialInstance.h b/Engine/Source/Runtime/Engine/Classes/Materials/MaterialInstance.h
--- a/Engine/Source/Runtime/Engine/Classes/Materials/MaterialInstance.h
+++ b/Engine/Source/Runtime/Engine/Classes/Materials/MaterialInstance.h
@@ -471,8 +471,8 @@ class UMaterialInstance : public UMaterialInterface
 #if WITH_EDITORONLY_DATA
 	ENGINE_API virtual const UClass* GetEditorOnlyDataClass() const override { return UMaterialInstanceEditorOnlyData::StaticClass(); }
 
-	virtual UMaterialInstanceEditorOnlyData* GetEditorOnlyData() override { return CastChecked<UMaterialInstanceEditorOnlyData>(Super::GetEditorOnlyData()); }
-	virtual const UMaterialInstanceEditorOnlyData* GetEditorOnlyData() const override { return CastChecked<UMaterialInstanceEditorOnlyData>(Super::GetEditorOnlyData()); }
+	virtual UMaterialInstanceEditorOnlyData* GetEditorOnlyData() override { return CastChecked<UMaterialInstanceEditorOnlyData>(Super::GetEditorOnlyData(), ECastCheckedType::NullAllowed); }
+	virtual const UMaterialInstanceEditorOnlyData* GetEditorOnlyData() const override { return CastChecked<UMaterialInstanceEditorOnlyData>(Super::GetEditorOnlyData(), ECastCheckedType::NullAllowed); }
 #endif // WITH_EDITORONLY_DATA
 
 	/** Physical material to use for this graphics material. Used for sounds, effects etc.*/
