Changing "named events enabled" text to render regardless of which path enabled it, instead of only when using the "stat namedevents" path.

diff --git a/Engine/Source/Runtime/Engine/Classes/Engine/Engine.h b/Engine/Source/Runtime/Engine/Classes/Engine/Engine.h
--- a/Engine/Source/Runtime/Engine/Classes/Engine/Engine.h
+++ b/Engine/Source/Runtime/Engine/Classes/Engine/Engine.h
@@ -3491,6 +3491,16 @@ public:
 	 */
 	void RenderEngineStats(UWorld* World, FViewport* Viewport, FCanvas* Canvas, int32 LHSX, int32& InOutLHSY, int32 RHSX, int32& InOutRHSY, const FVector* ViewLocation, const FRotator* ViewRotation);
 
+	/**
+	 * Function to render text indicating whether named events are enabled.
+	 *
+	 * @param Canvas The canvas to use when drawing.
+	 * @param X The X position to start drawing from.
+	 * @param Y The Y position to start drawing from.
+	 * @return The ending Y position to continue rendering stats at.
+	 */
+	int32 RenderNamedEventsEnabled(FCanvas* Canvas, int32 X, int32 Y);
+
 	/**
 	 * Function definition for those stats which have their own toggle functions (or toggle other stats).
 	 *
@@ -3623,7 +3633,6 @@ private:
 	int32 RenderStatFPS(UWorld* World, FViewport* Viewport, FCanvas* Canvas, int32 X, int32 Y, const FVector* ViewLocation = nullptr, const FRotator* ViewRotation = nullptr);
 	int32 RenderStatHitches(UWorld* World, FViewport* Viewport, FCanvas* Canvas, int32 X, int32 Y, const FVector* ViewLocation = nullptr, const FRotator* ViewRotation = nullptr);
 	int32 RenderStatSummary(UWorld* World, FViewport* Viewport, FCanvas* Canvas, int32 X, int32 Y, const FVector* ViewLocation = nullptr, const FRotator* ViewRotation = nullptr);
-	int32 RenderStatNamedEvents(UWorld* World, FViewport* Viewport, FCanvas* Canvas, int32 X, int32 Y, const FVector* ViewLocation = nullptr, const FRotator* ViewRotation = nullptr);
 	int32 RenderStatColorList(UWorld* World, FViewport* Viewport, FCanvas* Canvas, int32 X, int32 Y, const FVector* ViewLocation = nullptr, const FRotator* ViewRotation = nullptr);
 	int32 RenderStatLevels(UWorld* World, FViewport* Viewport, FCanvas* Canvas, int32 X, int32 Y, const FVector* ViewLocation = nullptr, const FRotator* ViewRotation = nullptr);
 	int32 RenderStatLevelMap(UWorld* World, FViewport* Viewport, FCanvas* Canvas, int32 X, int32 Y, const FVector* ViewLocation = nullptr, const FRotator* ViewRotation = nullptr);
diff --git a/Engine/Source/Runtime/Engine/Private/UnrealEngine.cpp b/Engine/Source/Runtime/Engine/Private/UnrealEngine.cpp
--- a/Engine/Source/Runtime/Engine/Private/UnrealEngine.cpp
+++ b/Engine/Source/Runtime/Engine/Private/UnrealEngine.cpp
@@ -2071,7 +2071,7 @@ void UEngine::Init(IEngineLoop* InEngineLoop)
 #if !UE_BUILD_SHIPPING
 	EngineStats.Add(FEngineStatFuncs(TEXT("STAT_Version"), TEXT("STATCAT_Engine"), FText::FromString(TEXT("Display build version string which includes the Changelist number.")), FEngineStatRender::CreateUObject(this, &UEngine::RenderStatVersion), FEngineStatToggle(), bIsRHS));
 #endif // !UE_BUILD_SHIPPING
-	EngineStats.Add(FEngineStatFuncs(TEXT("STAT_NamedEvents"), TEXT("STATCAT_Engine"), FText::GetEmpty(), FEngineStatRender::CreateUObject(this, &UEngine::RenderStatNamedEvents), FEngineStatToggle::CreateUObject(this, &UEngine::ToggleStatNamedEvents), bIsRHS));
+	EngineStats.Add(FEngineStatFuncs(TEXT("STAT_NamedEvents"), TEXT("STATCAT_Engine"), FText::GetEmpty(), FEngineStatRender(), FEngineStatToggle::CreateUObject(this, &UEngine::ToggleStatNamedEvents), bIsRHS));
 	EngineStats.Add(FEngineStatFuncs(TEXT("STAT_VerboseNamedEvents"), TEXT("STATCAT_Engine"), FText::GetEmpty(), FEngineStatRender(), FEngineStatToggle::CreateUObject(this, &UEngine::ToggleStatVerboseNamedEvents), bIsRHS));
 	EngineStats.Add(FEngineStatFuncs(TEXT("STAT_FPS"), TEXT("STATCAT_Engine"), FText::FromString(TEXT("Display average FPS.")), FEngineStatRender::CreateUObject(this, &UEngine::RenderStatFPS), FEngineStatToggle::CreateUObject(this, &UEngine::ToggleStatFPS), bIsRHS));
 	EngineStats.Add(FEngineStatFuncs(TEXT("STAT_Summary"), TEXT("STATCAT_Engine"), FText::FromString(TEXT("Display UsedPhysical memory stat.")), FEngineStatRender::CreateUObject(this, &UEngine::RenderStatSummary), FEngineStatToggle(), bIsRHS));
@@ -11857,6 +11857,9 @@ PRAGMA_ENABLE_DEPRECATION_WARNINGS
 		// give the viewport first shot at drawing stats
 		Y = Viewport->DrawStatsHUD(Canvas, StatsXOffset, Y);
 
+		// Named events are enabled through multiple paths, so draw them here
+		Y = GEngine->RenderNamedEventsEnabled(Canvas, X, Y);
+
 		// Render all the simple stats
 		GEngine->RenderEngineStats(World, Viewport, Canvas, StatsXOffset, MessageY, X, Y, &ViewLocation, &ViewRotation);
 
@@ -16790,16 +16793,19 @@ bool UEngine::ToggleStatVerboseNamedEvents(UWorld* World, FCommonViewportClient*
 	return false;
 }
 
-int32 UEngine::RenderStatNamedEvents(UWorld* World, FViewport* Viewport, FCanvas* Canvas, int32 X, int32 Y, const FVector* ViewLocation, const FRotator* ViewRotation)
+int32 UEngine::RenderNamedEventsEnabled(FCanvas* Canvas, int32 X, int32 Y)
 {
-	const FText Text = GShouldEmitVerboseNamedEvents ?
-		LOCTEXT("VERBOSENAMEDEVENTSENABLED", "VERBOSE NAMED EVENTS ENABLED") :
-		LOCTEXT("NAMEDEVENTSENABLED", "NAMED EVENTS ENABLED");
-	X -= GShouldEmitVerboseNamedEvents ? 100 : 42;
-	FCanvasTextItem TextItem(FVector2D(X, Y), Text, GetSmallFont(), FLinearColor::Blue);
-	TextItem.EnableShadow(FLinearColor::Black);
-	Canvas->DrawItem(TextItem);
-	Y += TextItem.DrawnSize.Y;
+	if (GCycleStatsShouldEmitNamedEvents)
+	{
+		const FText Text = GShouldEmitVerboseNamedEvents ?
+			LOCTEXT("VERBOSENAMEDEVENTSENABLED", "VERBOSE NAMED EVENTS ENABLED") :
+			LOCTEXT("NAMEDEVENTSENABLED", "NAMED EVENTS ENABLED");
+		X -= GShouldEmitVerboseNamedEvents ? 100 : 42;
+		FCanvasTextItem TextItem(FVector2D(X, Y), Text, GetSmallFont(), FLinearColor::Blue);
+		TextItem.EnableShadow(FLinearColor::Black);
+		Canvas->DrawItem(TextItem);
+		Y += TextItem.DrawnSize.Y;
+	}
 	return Y;
 }
 
