Unifying how we modify GCycleStatsShouldEmitNamedEvents so that we guarantee it won't drop below zero (should never happen, but just in case so that we can change checks to be != 0).

diff --git a/Engine/Source/Runtime/Core/Private/FramePro/FrameProProfiler.cpp b/Engine/Source/Runtime/Core/Private/FramePro/FrameProProfiler.cpp
--- a/Engine/Source/Runtime/Core/Private/FramePro/FrameProProfiler.cpp
+++ b/Engine/Source/Runtime/Core/Private/FramePro/FrameProProfiler.cpp
@@ -426,7 +426,7 @@ FString FFrameProProfiler::StartFrameProRecording(const FString& FilenameRoot, i
 	GFrameProEnabled = true;
 
 	// Enable named events as well
-	GCycleStatsShouldEmitNamedEvents += 1;
+	++GCycleStatsShouldEmitNamedEvents:
 
 	// Set recording flag
 	GFrameProIsRecording = true;
@@ -450,7 +450,7 @@ void FFrameProProfiler::StopFrameProRecording()
 	FramePro::StopRecording();
 
 	// Disable named events
-	GCycleStatsShouldEmitNamedEvents -= 1;
+	GCycleStatsShouldEmitNamedEvents = FMath::Max(0, GCycleStatsShouldEmitNamedEvents - 1);
 
 	// Clear recording flag
 	GFrameProIsRecording = false;
diff --git a/Engine/Source/Runtime/Core/Private/ProfilingDebugging/CsvProfiler.cpp b/Engine/Source/Runtime/Core/Private/ProfilingDebugging/CsvProfiler.cpp
--- a/Engine/Source/Runtime/Core/Private/ProfilingDebugging/CsvProfiler.cpp
+++ b/Engine/Source/Runtime/Core/Private/ProfilingDebugging/CsvProfiler.cpp
@@ -2820,7 +2820,7 @@ void FCsvProfiler::BeginFrame()
 					}
 					SetMetadataInternal(TEXT("TargetFramerate"), *FString::FromInt(TargetFPS));
 					SetMetadataInternal(TEXT("StartTimestamp"), *FString::Printf(TEXT("%lld"), FDateTime::UtcNow().ToUnixTimestamp()));
-					SetMetadataInternal(TEXT("NamedEvents"), (GCycleStatsShouldEmitNamedEvents > 0) ? TEXT("1") : TEXT("0"));
+					SetMetadataInternal(TEXT("NamedEvents"), GCycleStatsShouldEmitNamedEvents ? TEXT("1") : TEXT("0"));
 
 					GCsvStatCounts = !!CVarCsvStatCounts.GetValueOnGameThread();
 
diff --git a/Engine/Source/Runtime/Core/Public/Stats/Stats2.h b/Engine/Source/Runtime/Core/Public/Stats/Stats2.h
--- a/Engine/Source/Runtime/Core/Public/Stats/Stats2.h
+++ b/Engine/Source/Runtime/Core/Public/Stats/Stats2.h
@@ -1577,7 +1577,7 @@ public:
 		}
 
 		// Emit named event for active cycle stat.
-		if (GCycleStatsShouldEmitNamedEvents > 0
+		if (GCycleStatsShouldEmitNamedEvents
 			&& (GShouldEmitVerboseNamedEvents || !EnumHasAnyFlags(InStatFlags, EStatFlags::Verbose)))
 		{
 #if PLATFORM_USES_ANSI_STRING_FOR_EXTERNAL_PROFILING
@@ -2374,7 +2374,7 @@ namespace Stats
 #if STATS && CPUPROFILERTRACE_ENABLED
 		return UE_TRACE_CHANNELEXPR_IS_ENABLED(CpuChannel) || FThreadStats::IsCollectingData();
 #else
-		return GCycleStatsShouldEmitNamedEvents > 0;
+		return GCycleStatsShouldEmitNamedEvents != 0;
 #endif
 	}
 }
diff --git a/Engine/Source/Runtime/Engine/Private/PhysicsEngine/PhysXSupport.cpp b/Engine/Source/Runtime/Engine/Private/PhysicsEngine/PhysXSupport.cpp
--- a/Engine/Source/Runtime/Engine/Private/PhysicsEngine/PhysXSupport.cpp
+++ b/Engine/Source/Runtime/Engine/Private/PhysicsEngine/PhysXSupport.cpp
@@ -624,7 +624,7 @@ PxCollection* MakePhysXCollection(const TArray<UPhysicalMaterial*>& PhysicalMate
 
 void* FPhysXProfilerCallback::zoneStart(const char* eventName, bool detached, uint64_t contextId)
 {
-	if(GCycleStatsShouldEmitNamedEvents > 0)
+	if(GCycleStatsShouldEmitNamedEvents)
 	{
 		FPlatformMisc::BeginNamedEvent(FColor::Red, *FString::Printf(TEXT("PHYSX: %s"), StringCast<TCHAR>(eventName).Get()));
 	}
@@ -634,7 +634,7 @@ void* FPhysXProfilerCallback::zoneStart(const char* eventName, bool detached, ui
 
 void FPhysXProfilerCallback::zoneEnd(void* profilerData, const char* eventName, bool detached, uint64_t contextId)
 {
-	if(GCycleStatsShouldEmitNamedEvents > 0)
+	if(GCycleStatsShouldEmitNamedEvents)
 	{
 		FPlatformMisc::EndNamedEvent();
 	}
diff --git a/Engine/Source/Runtime/Launch/Private/LaunchEngineLoop.cpp b/Engine/Source/Runtime/Launch/Private/LaunchEngineLoop.cpp
--- a/Engine/Source/Runtime/Launch/Private/LaunchEngineLoop.cpp
+++ b/Engine/Source/Runtime/Launch/Private/LaunchEngineLoop.cpp
@@ -1520,13 +1520,13 @@ int32 FEngineLoop::PreInitPreStartupScreen(const TCHAR* CmdLine)
 	// Command line option for enabling named events
 	if (FParse::Param(CmdLine, TEXT("statnamedevents")))
 	{
-		GCycleStatsShouldEmitNamedEvents = 1;
+		++GCycleStatsShouldEmitNamedEvents;
 	}
 	
 	if (FParse::Param(CmdLine, TEXT("verbosenamedevents")))
 	{
-		GCycleStatsShouldEmitNamedEvents = 1;
-		GShouldEmitVerboseNamedEvents = 1;
+		++GCycleStatsShouldEmitNamedEvents;
+		GShouldEmitVerboseNamedEvents = true;
 	}
 
 	// Set the flag for whether we've build DebugGame instead of Development. The engine does not know this (whereas the launch module does) because it is always built in development.
