Expose default value as inputs for ScalarParameter and VectorParameter material expressions

diff --git a/Engine/Source/Runtime/Engine/Classes/Materials/MaterialExpressionScalarParameter.h b/Engine/Source/Runtime/Engine/Classes/Materials/MaterialExpressionScalarParameter.h
--- a/Engine/Source/Runtime/Engine/Classes/Materials/MaterialExpressionScalarParameter.h
+++ b/Engine/Source/Runtime/Engine/Classes/Materials/MaterialExpressionScalarParameter.h
@@ -15,7 +15,7 @@ class UMaterialExpressionScalarParameter : public UMaterialExpressionParameter
 {
 	GENERATED_UCLASS_BODY()
 
-	UPROPERTY(EditAnywhere, Category=MaterialExpressionScalarParameter)
+	UPROPERTY(EditAnywhere, Category=MaterialExpressionScalarParameter, Meta = (ShowAsInputPin = "Primary"))
 	float DefaultValue;
 
 	UPROPERTY(EditAnywhere, Category=CustomPrimitiveData)
diff --git a/Engine/Source/Runtime/Engine/Classes/Materials/MaterialExpressionVectorParameter.h b/Engine/Source/Runtime/Engine/Classes/Materials/MaterialExpressionVectorParameter.h
--- a/Engine/Source/Runtime/Engine/Classes/Materials/MaterialExpressionVectorParameter.h
+++ b/Engine/Source/Runtime/Engine/Classes/Materials/MaterialExpressionVectorParameter.h
@@ -14,7 +14,7 @@ class UMaterialExpressionVectorParameter : public UMaterialExpressionParameter
 {
 	GENERATED_UCLASS_BODY()
 
-	UPROPERTY(EditAnywhere, Category=MaterialExpressionVectorParameter)
+	UPROPERTY(EditAnywhere, Category=MaterialExpressionVectorParameter, meta = (ShowAsInputPin = "Primary"))
 	FLinearColor DefaultValue;
 
 	UPROPERTY(EditAnywhere, Category=CustomPrimitiveData)
