Added new "ShaderPrint" ShowFlag to fix Nanitestats and other shader printing showing up in thumbnails.

#rnx

diff --git a/Engine/Source/Runtime/Engine/Public/ShowFlags.h b/Engine/Source/Runtime/Engine/Public/ShowFlags.h
--- a/Engine/Source/Runtime/Engine/Public/ShowFlags.h
+++ b/Engine/Source/Runtime/Engine/Public/ShowFlags.h
@@ -225,6 +225,8 @@ struct FEngineShowFlags
 
 		// TODO: Remove when Physical page pool size scales automatically with demand
 		SetVirtualShadowMapCaching(false);
+
+		SetShaderPrint(false);
 	}
 
 	void EnableAdvancedFeatures()
diff --git a/Engine/Source/Runtime/Engine/Public/ShowFlagsValues.inl b/Engine/Source/Runtime/Engine/Public/ShowFlagsValues.inl
--- a/Engine/Source/Runtime/Engine/Public/ShowFlagsValues.inl
+++ b/Engine/Source/Runtime/Engine/Public/ShowFlagsValues.inl
@@ -428,5 +428,8 @@ SHOWFLAG_ALWAYS_ACCESSIBLE(LumenReuseShadowMaps, SFG_Lumen, NSLOCTEXT("UnrealEd"
 /** Visualize Skin Cache */
 SHOWFLAG_FIXED_IN_SHIPPING(0, VisualizeGPUSkinCache, SFG_Hidden, NSLOCTEXT("UnrealEd", "VisualizeGPUSkinCache", "Visualize GPU Skin Cache"))
 
+/** Enables ShaderPrint functionality that is used for debug printing from shaders (e.g. 'Nanitestats' command). Disable this for certain editor rendering (e.g. thumbnail rendering) */
+SHOWFLAG_ALWAYS_ACCESSIBLE(ShaderPrint, SFG_Developer, NSLOCTEXT("UnrealEd", "ShaderPrint", "Shader Print"))
+
 #undef SHOWFLAG_ALWAYS_ACCESSIBLE
 #undef SHOWFLAG_FIXED_IN_SHIPPING
diff --git a/Engine/Source/Runtime/Renderer/Private/ShaderPrint.cpp b/Engine/Source/Runtime/Renderer/Private/ShaderPrint.cpp
--- a/Engine/Source/Runtime/Renderer/Private/ShaderPrint.cpp
+++ b/Engine/Source/Runtime/Renderer/Private/ShaderPrint.cpp
@@ -259,7 +259,7 @@ namespace ShaderPrint
 
 	bool IsEnabled(const FSceneView& View)
 	{
-		return IsEnabled() && IsSupported(View.GetShaderPlatform());
+		return IsEnabled() && View.Family->EngineShowFlags.ShaderPrint && IsSupported(View.GetShaderPlatform());
 	}
 
 	// Returns true if the default view exists and has shader debug rendering enabled (this needs to be checked before using a permutation that requires the shader draw parameters)
