Chaos triangle mesh cooking:  Collision UV feature was broken due to reordering of triangles when generating the BVH in CL 20123324.  Modified the BVH generation to properly maintain the face remapping table.  Ran some local unit tests with debug code to treat certain triangles as invalid, to test that remapping logic works when invalid triangles are removed.  Invalid triangles would have caused similar bugs in older versions of the code, as the remap table wasn't updated in that case.

#rnx

diff --git a/Engine/Source/Runtime/Engine/Private/PhysicsEngine/Experimental/ChaosCooking.cpp b/Engine/Source/Runtime/Engine/Private/PhysicsEngine/Experimental/ChaosCooking.cpp
--- a/Engine/Source/Runtime/Engine/Private/PhysicsEngine/Experimental/ChaosCooking.cpp
+++ b/Engine/Source/Runtime/Engine/Private/PhysicsEngine/Experimental/ChaosCooking.cpp
@@ -16,6 +16,14 @@ namespace Chaos
 {
 	namespace Cooking
 	{
+		static void CopyUpdatedFaceRemapFromTriangleMesh(const Chaos::FTriangleMeshImplicitObject& TriangleMesh, TArray<int32>& OutFaceRemap)
+		{
+			for (int32 TriangleIndex = 0; TriangleIndex < OutFaceRemap.Num(); TriangleIndex++)
+			{
+				OutFaceRemap[TriangleIndex] = TriangleMesh.GetExternalFaceIndexFromInternal(TriangleIndex);
+			}
+		}
+
 		TUniquePtr<Chaos::FTriangleMeshImplicitObject> BuildSingleTrimesh(const FTriMeshCollisionData& Desc, TArray<int32>& OutFaceRemap, TArray<int32>& OutVertexRemap)
 		{
 			if(Desc.Vertices.Num() == 0)
@@ -65,6 +73,10 @@ namespace Chaos
 					MaterialIndices.Reserve(NumTriangles);
 				}
 
+				// Need to rebuild face remap array, in case there are any invalid triangles
+				TArray<int32> OldFaceRemap = MoveTemp(OutFaceRemap);
+				OutFaceRemap.Reserve(OldFaceRemap.Num());
+
 				Triangles.Reserve(NumTriangles);
 				for(int32 TriangleIndex = 0; TriangleIndex < NumTriangles; ++TriangleIndex)
 				{
@@ -80,19 +92,23 @@ namespace Chaos
 					if(bIsValidTriangle)
 					{
 						Triangles.Add(Chaos::TVector<int32, 3>(FinalIndices[BaseIndex], FinalIndices[BaseIndex + 1], FinalIndices[BaseIndex + 2]));
+						if (!OldFaceRemap.IsEmpty())
+						{
+							OutFaceRemap.Add(OldFaceRemap[TriangleIndex]);
+						}
 
 						if(bHasMaterials)
 						{
 							if(EnableMeshClean)
 							{
-								if(!ensure(OutFaceRemap.IsValidIndex(TriangleIndex)))
+								if(!ensure(OldFaceRemap.IsValidIndex(TriangleIndex)))
 								{
 									MaterialIndices.Empty();
 									bHasMaterials = false;
 								}
 								else
 								{
-									const int32 OriginalIndex = OutFaceRemap[TriangleIndex];
+									const int32 OriginalIndex = OldFaceRemap[TriangleIndex];
 
 									if(ensure(Desc.MaterialIndices.IsValidIndex(OriginalIndex)))
 									{
@@ -124,7 +140,12 @@ namespace Chaos
 
 				TUniquePtr<TArray<int32>> OutFaceRemapPtr = MakeUnique<TArray<int32>>(OutFaceRemap);
 				TUniquePtr<TArray<int32>> OutVertexRemapPtr = Chaos::TriMeshPerPolySupport ? MakeUnique<TArray<int32>>(OutVertexRemap) : nullptr;
-				return MakeUnique<Chaos::FTriangleMeshImplicitObject>(MoveTemp(TriMeshParticles), MoveTemp(Triangles), MoveTemp(MaterialIndices), MoveTemp(OutFaceRemapPtr), MoveTemp(OutVertexRemapPtr));
+				TUniquePtr<Chaos::FTriangleMeshImplicitObject> TriangleMesh = MakeUnique<Chaos::FTriangleMeshImplicitObject>(MoveTemp(TriMeshParticles), MoveTemp(Triangles), MoveTemp(MaterialIndices), MoveTemp(OutFaceRemapPtr), MoveTemp(OutVertexRemapPtr));
+
+				// Propagate remapped indices from the FTriangleMeshImplicitObject back to the remap array
+				CopyUpdatedFaceRemapFromTriangleMesh(*TriangleMesh.Get(), OutFaceRemap);
+
+				return TriangleMesh;
 			};
 
 			if(FinalVerts.Num() < TNumericLimits<uint16>::Max())
@@ -237,6 +258,10 @@ namespace Chaos
 					MaterialIndices.Reserve(NumTriangles);
 				}
 
+				// Need to rebuild face remap array, in case there are any invalid triangles
+				TArray<int32> OldFaceRemap = MoveTemp(OutFaceRemap);
+				OutFaceRemap.Reserve(OldFaceRemap.Num());
+
 				Triangles.Reserve(NumTriangles);
 				for(int32 TriangleIndex = 0; TriangleIndex < NumTriangles; ++TriangleIndex)
 				{
@@ -252,19 +277,23 @@ namespace Chaos
 					if(bIsValidTriangle)
 					{
 						Triangles.Add(Chaos::TVector<int32, 3>(FinalIndices[BaseIndex], FinalIndices[BaseIndex + 1], FinalIndices[BaseIndex + 2]));
+						if (OldFaceRemap.Num())
+						{
+							OutFaceRemap.Add(OldFaceRemap[TriangleIndex]);
+						}
 
 						if(bHasMaterials)
 						{
 							if(EnableMeshClean)
 							{
-								if(!ensure(OutFaceRemap.IsValidIndex(TriangleIndex)))
+								if(!ensure(OldFaceRemap.IsValidIndex(TriangleIndex)))
 								{
 									MaterialIndices.Empty();
 									bHasMaterials = false;
 								}
 								else
 								{
-									const int32 OriginalIndex = OutFaceRemap[TriangleIndex];
+									const int32 OriginalIndex = OldFaceRemap[TriangleIndex];
 
 									if(ensure(InParams.TriangleMeshDesc.MaterialIndices.IsValidIndex(OriginalIndex)))
 									{
@@ -296,7 +325,11 @@ namespace Chaos
 
 				TUniquePtr<TArray<int32>> OutFaceRemapPtr = MakeUnique<TArray<int32>>(OutFaceRemap);
 				TUniquePtr<TArray<int32>> OutVertexRemapPtr = Chaos::TriMeshPerPolySupport ? MakeUnique<TArray<int32>>(OutVertexRemap) : nullptr;
-				OutTriangleMeshes.Emplace(new Chaos::FTriangleMeshImplicitObject(MoveTemp(TriMeshParticles), MoveTemp(Triangles), MoveTemp(MaterialIndices), MoveTemp(OutFaceRemapPtr), MoveTemp(OutVertexRemapPtr)));
+				Chaos::FTriangleMeshImplicitObject* TriangleMesh = new Chaos::FTriangleMeshImplicitObject(MoveTemp(TriMeshParticles), MoveTemp(Triangles), MoveTemp(MaterialIndices), MoveTemp(OutFaceRemapPtr), MoveTemp(OutVertexRemapPtr));
+				OutTriangleMeshes.Emplace(TriangleMesh);
+
+				// Propagate remapped indices from the FTriangleMeshImplicitObject back to the remap array
+				CopyUpdatedFaceRemapFromTriangleMesh(*TriangleMesh, OutFaceRemap);
 			};
 
 			if(FinalVerts.Num() < TNumericLimits<uint16>::Max())
diff --git a/Engine/Source/Runtime/Experimental/Chaos/Private/Chaos/TriangleMeshImplicitObject.cpp b/Engine/Source/Runtime/Experimental/Chaos/Private/Chaos/TriangleMeshImplicitObject.cpp
--- a/Engine/Source/Runtime/Experimental/Chaos/Private/Chaos/TriangleMeshImplicitObject.cpp
+++ b/Engine/Source/Runtime/Experimental/Chaos/Private/Chaos/TriangleMeshImplicitObject.cpp
@@ -1572,7 +1572,13 @@ void FTriangleMeshImplicitObject::RebuildFastBVHFromTree(const BVHType& TreeBVH)
 	int32 FaceNum = 0;
 	TArray<TVec3<FTrimeshIndexBuffer::LargeIdxType>> LargeIndices;
 	TArray<TVec3<FTrimeshIndexBuffer::SmallIdxType>> SmallIndices;
-		
+	TArray<int32> OldFaceIndexMap;
+
+	if (ExternalFaceIndexMap.IsValid())
+	{
+		OldFaceIndexMap = MoveTemp(*ExternalFaceIndexMap.Get());
+		ExternalFaceIndexMap->Reserve(OldFaceIndexMap.Num());
+	}
 
 	// since we do skip leaf nodes, we need to handle the case where we have only one node that will be a leaf by default
 	if (Nodes.Num() == 1)
@@ -1595,6 +1601,10 @@ void FTriangleMeshImplicitObject::RebuildFastBVHFromTree(const BVHType& TreeBVH)
 			// Reorder triangle indices, triangles will be in the same order as the bounding volume in the BVH structure.
 			// And all triangles in a leaf will be contiguous in memory.
 			AddTriangles(LargeIndices, SmallIndices, MElements, LeafPayload.Payload);
+			if (ExternalFaceIndexMap.IsValid())
+			{
+				ExternalFaceIndexMap->Add(OldFaceIndexMap[LeafPayload.Payload]);
+			}
 		}
 		if (MElements.RequiresLargeIndices())
 		{
@@ -1672,6 +1682,10 @@ void FTriangleMeshImplicitObject::RebuildFastBVHFromTree(const BVHType& TreeBVH)
 							// Reorder triangle indices, triangles will be in the same order as the bounding volume in the BVH structure.
 							// And all triangles in a leaf will be contiguous in memory.
 							AddTriangles(LargeIndices, SmallIndices, MElements, LeafPayload.Payload);
+							if (ExternalFaceIndexMap.IsValid())
+							{
+								ExternalFaceIndexMap->Add(OldFaceIndexMap[LeafPayload.Payload]);
+							}
 							FaceNum++;
 						}
 					}
