Avoid counting bodysetup that are choosen to be built as DDC miss

#rnx

diff --git a/Engine/Source/Runtime/Engine/Private/PhysicsEngine/BodySetup.cpp b/Engine/Source/Runtime/Engine/Private/PhysicsEngine/BodySetup.cpp
--- a/Engine/Source/Runtime/Engine/Private/PhysicsEngine/BodySetup.cpp
+++ b/Engine/Source/Runtime/Engine/Private/PhysicsEngine/BodySetup.cpp
@@ -1616,14 +1616,15 @@ void GetDDCBuiltData(FByteBulkData* OutResult, DDCBuilderType& InBuilder, UBodyS
 				<< BodySetupDDCFetch.FetchReason(*FetchReason);
 #endif
 			bDDCHit = GetDerivedDataCacheRef().GetSynchronous(&InBuilder, OutData, &bDataWasBuilt);
+
+			// Only compute hit/miss if DDC was not skipped
+			COOK_STAT(Timer.AddHitOrMiss(!bDDCHit || bDataWasBuilt ? FCookStats::CallStats::EHitOrMiss::Miss : FCookStats::CallStats::EHitOrMiss::Hit, OutData.Num()));
 		}
 		else
 		{
 			bDataWasBuilt = true;
 			InBuilder.Build(OutData);
 		}
-
-		COOK_STAT(Timer.AddHitOrMiss(!bDDCHit || bDataWasBuilt ? FCookStats::CallStats::EHitOrMiss::Miss : FCookStats::CallStats::EHitOrMiss::Hit, OutData.Num()));
 	}
 
 	if(OutData.Num())
