HLODActor - Removed check() in SetHLODComponents() that triggered when an empty array of components was provided

diff --git a/Engine/Source/Runtime/Engine/Private/WorldPartition/HLOD/HLODActor.cpp b/Engine/Source/Runtime/Engine/Private/WorldPartition/HLOD/HLODActor.cpp
--- a/Engine/Source/Runtime/Engine/Private/WorldPartition/HLOD/HLODActor.cpp
+++ b/Engine/Source/Runtime/Engine/Private/WorldPartition/HLOD/HLODActor.cpp
@@ -109,7 +109,6 @@ TUniquePtr<FWorldPartitionActorDesc> AWorldPartitionHLOD::CreateClassActorDesc()
 void AWorldPartitionHLOD::SetHLODComponents(const TArray<UActorComponent*>& InHLODComponents)
 {
 	TRACE_CPUPROFILER_EVENT_SCOPE(AWorldPartitionHLOD::SetHLODComponents);
-	check(!InHLODComponents.IsEmpty());
 
 	TArray<UActorComponent*> ComponentsToRemove;
 	GetComponents(ComponentsToRemove);
