Fix a crash while cooking android with lightmap VT enabled.

diff --git a/Engine/Source/Runtime/Engine/Private/VT/VirtualTextureDataBuilder.cpp b/Engine/Source/Runtime/Engine/Private/VT/VirtualTextureDataBuilder.cpp
--- a/Engine/Source/Runtime/Engine/Private/VT/VirtualTextureDataBuilder.cpp
+++ b/Engine/Source/Runtime/Engine/Private/VT/VirtualTextureDataBuilder.cpp
@@ -579,7 +579,7 @@ void FVirtualTextureDataBuilder::BuildTiles(const TArray<FVTSourceTileEntry>& Ti
 		TBSettings.bSRGB = BuildSettingsForLayer.bSRGB;
 		TBSettings.bUseLegacyGamma = BuildSettingsForLayer.bUseLegacyGamma;
 		TBSettings.MipGenSettings = TMGS_NoMipmaps;
-		TBSettings.bForceAlphaChannel = BuildSettingsForLayer.bForceAlphaChannel;
+		TBSettings.bForceAlphaChannel = BuildSettingsForLayer.bForceAlphaChannel || LayerData.bHasAlpha;
 		TBSettings.bForceNoAlphaChannel = BuildSettingsForLayer.bForceNoAlphaChannel;
 		TBSettings.bHDRSource = BuildSettingsForLayer.bHDRSource;
 		TBSettings.bVirtualStreamable = true;
@@ -1168,23 +1168,6 @@ void FVirtualTextureDataBuilder::BuildSourcePixels(const FTextureSourceData& Sou
 		FName TextureFormatPrefix;
 		FName TextureFormatName = RemovePrefixFromName(BuildSettingsForLayer.TextureFormatName, TextureFormatPrefix);
 
-		// We handle AutoDXT specially here since otherwise the texture format compressor would choose a DXT format for every tile
-		// individually. Causing tiles in the same VT to use different formats which we don't allow.
-		static FName NameDXT1(TEXT("DXT1"));
-		static FName NameDXT5(TEXT("DXT5"));
-		static FName NameAutoDXT(TEXT("AutoDXT"));
-		if (TextureFormatName == NameAutoDXT)
-		{
-			if (LayerData.bHasAlpha)
-			{
-				TextureFormatName = NameDXT5;
-			}
-			else
-			{
-				TextureFormatName = NameDXT1;
-			}
-		}
-
 		if ( TextureFormatPrefix.IsNone())
 		{
 			LayerData.TextureFormatName = TextureFormatName;
