Improve async input component so that it is used with UObjects and PlayerController can create one by default.

diff --git a/Engine/Source/Runtime/Engine/Classes/GameFramework/PlayerController.h b/Engine/Source/Runtime/Engine/Classes/GameFramework/PlayerController.h
--- a/Engine/Source/Runtime/Engine/Classes/GameFramework/PlayerController.h
+++ b/Engine/Source/Runtime/Engine/Classes/GameFramework/PlayerController.h
@@ -25,6 +25,7 @@
 #include "GameFramework/UpdateLevelVisibilityLevelInfo.h"
 #include "GenericPlatform/IInputInterface.h"
 #include "GameFramework/PlayerInput.h"
+#include "Physics/AsyncPhysicsData.h"
 #include "PlayerController.generated.h"
 
 class ACameraActor;
@@ -43,6 +44,7 @@ class UPlayer;
 class UPrimitiveComponent;
 struct FActiveHapticFeedbackEffect;
 struct FCollisionQueryParams;
+class UAsyncPhysicsInputComponent;
 
 /** Default delegate that provides an implementation for those that don't have special needs other than a toggle */
 DECLARE_DELEGATE_RetVal(bool, FCanUnpause);
@@ -340,8 +342,26 @@ public:
 	TArray<FForceFeedbackEffectHistoryEntry> ForceFeedbackEffectHistoryEntries;
 #endif
 
+	/** The type of async physics data object to use*/
+	UPROPERTY(EditDefaultsOnly, Category=PlayerController)
+	TSubclassOf<UAsyncPhysicsData> AsyncPhysicsDataClass;
+
+	/** Get the async physics data to write to. This data will make its way to the async physics tick on client and server. Should not be used during async tick */
+	UFUNCTION(BlueprintPure, Category = PlayerController)
+	UAsyncPhysicsData* GetAsyncPhysicsDataToWrite() const;
+
+	/** Get the async physics data to execute logic off of. This data should not be modified and will NOT make its way back. Must be used during async tick */
+	UFUNCTION(BlueprintPure, Category = PlayerController)
+	const UAsyncPhysicsData* GetAsyncPhysicsDataToConsume() const;
+
 private:
 
+	UPROPERTY(ReplicatedUsing=OnRep_AsyncPhysicsDataComponent)
+	TObjectPtr<UAsyncPhysicsInputComponent> AsyncPhysicsDataComponent;
+
+	UFUNCTION()
+	void OnRep_AsyncPhysicsDataComponent();
+
 	struct FDynamicForceFeedbackAction
 	{
 		/** Time over which interpolation should happen */
diff --git a/Engine/Source/Runtime/Engine/Private/PhysicsEngine/AsyncPhysicsInputComponent.cpp b/Engine/Source/Runtime/Engine/Private/PhysicsEngine/AsyncPhysicsInputComponent.cpp
--- a/Engine/Source/Runtime/Engine/Private/PhysicsEngine/AsyncPhysicsInputComponent.cpp
+++ b/Engine/Source/Runtime/Engine/Private/PhysicsEngine/AsyncPhysicsInputComponent.cpp
@@ -292,19 +292,54 @@ struct FAsyncPhysicsInputRewindCallback : public Chaos::IRewindCallback
 	}
 };
 
-
 UAsyncPhysicsInputComponent::UAsyncPhysicsInputComponent()
 {
-	bWantsInitializeComponent = true;
 	bAutoActivate = true;
 	SetIsReplicatedByDefault(true);
+}
+
+void UAsyncPhysicsInputComponent::SetDataClass(TSubclassOf<UAsyncPhysicsData> InDataClass)
+{
+	ensureMsgf(DataClass == nullptr, TEXT("You can only set the data class once"));
+	DataClass = InDataClass;
+
+	DataToWrite = DuplicateObject((UAsyncPhysicsData*)DataClass->GetDefaultObject(), nullptr);
+	
+	//now that we have a class we're ready to async tick
 	SetAsyncPhysicsTickEnabled(true);
+	UWorld* World = GetWorld();
+	if (World->WorldType == EWorldType::PIE || World->WorldType == EWorldType::Game)
+	{
+		FPhysScene* PhysScene = World->GetPhysicsScene();
+		if (ensureAlways(PhysScene))
+		{
+			Chaos::FPhysicsSolver* Solver = PhysScene->GetSolver();
+			if (Solver->GetRewindCallback() == nullptr)
+			{
+				// -------------------------------------------------------------------------------------
+				// Hack: enable Rewind capture here. This needs to be moved somewhere else
+				// -------------------------------------------------------------------------------------
+
+				int32 NumFrames = 64;
+				const IConsoleVariable* NumFramesCVar = IConsoleManager::Get().FindConsoleVariable(TEXT("p.RewindCaptureNumFrames"));
+				if (ensure(NumFramesCVar))
+				{
+					// 1 frame is required to enable rewind capture. 
+					NumFrames = FMath::Max<int32>(1, NumFramesCVar->GetInt());
+				}
+				Solver->EnableRewindCapture(NumFrames, false, MakeUnique<FAsyncPhysicsInputRewindCallback>(World));
+			}
+
+			FAsyncPhysicsInputRewindCallback* Callback = static_cast<FAsyncPhysicsInputRewindCallback*>(Solver->GetRewindCallback());
+			Callback->DispatchPhysicsTick.AddUObject(this, &UAsyncPhysicsInputComponent::OnDispatchPhysicsTick);
+		}
+	}
 }
 
 void UAsyncPhysicsInputComponent::AsyncPhysicsTickComponent(float DeltaTime, float SimTime)
 {
 	Super::AsyncPhysicsTickComponent(DeltaTime, SimTime);
-	ensureMsgf(Pool != nullptr, TEXT("You must call RegisterInputPool in InitializeComponent"));
+	ensureMsgf(DataClass != nullptr, TEXT("You must call SetDataClass after creating the component"));
 
 	UWorld* World = GetWorld();
 	FPhysScene* PhysScene = World->GetPhysicsScene();
@@ -320,44 +355,40 @@ void UAsyncPhysicsInputComponent::AsyncPhysicsTickComponent(float DeltaTime, flo
 		Callback->UpdateReplicationMap_Internal(ServerFrame);
 	}
 
-	Pool->SetCurrentInputToAsyncExecute(nullptr);	//set current input to none in case we don't find it
+	DataToConsume = nullptr; //set current input to none in case we don't find it
 
-	for (int32 Idx = BufferedInputs.Num() - 1; Idx >= 0; --Idx)
+	for (int32 Idx = BufferedData.Num() - 1; Idx >= 0; --Idx)
 	{
-		FAsyncPhysicsInput* Input = BufferedInputs[Idx];
-		if (Input->ServerFrame == ServerFrame)
+		UAsyncPhysicsData* Data = BufferedData[Idx];
+		if (Data->ServerFrame == ServerFrame)
 		{
-			Pool->SetCurrentInputToAsyncExecute(Input);
+			DataToConsume = Data;
 		}
 
-		bool bFreeInput = Input->ServerFrame < ServerFrame;	//for most cases once we're passed this frame we can free
+		bool bFreeData = Data->ServerFrame < ServerFrame;	//for most cases once we're passed this frame we can free
 		if (APlayerController* PC = GetPlayerController())
 		{
 			if (PC->IsLocalController() && !InputCmdCVars::ForceInputDrop)
 			{
-				FAsyncPhysicsInputWrapper Wrapper;
-				Wrapper.Input = Input;
-				Wrapper.OwnerComponent = this;
-				ServerRPCBufferInput(Wrapper);
+				ServerRPCBufferInput(Data);
 
-				//If we are the local player then we need to keep this input around to send redundant RPCs to deal with potential packet loss
-				bFreeInput = --Input->Replicated == 0;
+				//If we are the local player then we need to keep this data around to send redundant RPCs to deal with potential packet loss
+				bFreeData = --Data->ReplicationRedundancy <= 0;
 			}
 		}
 
-		if (bFreeInput)
+		if (bFreeData)
 		{
-			Pool->FreeInputToPool(Input);
-			BufferedInputs.RemoveAtSwap(Idx);
+			BufferedData.RemoveAtSwap(Idx);
 		}
 	}
 }
 
-void UAsyncPhysicsInputComponent::ServerRPCBufferInput_Implementation(FAsyncPhysicsInputWrapper Wrapper)
+void UAsyncPhysicsInputComponent::ServerRPCBufferInput_Implementation(UAsyncPhysicsData* NewData)
 {
-	for (FAsyncPhysicsInput* Input : BufferedInputs)
+	for (UAsyncPhysicsData* Data : BufferedData)
 	{
-		if (Input->ServerFrame == Wrapper.Input->ServerFrame)
+		if (Data->ServerFrame == NewData->ServerFrame)
 		{
 			//already buffered this is just a redundant send to deal with potential packet loss
 			//question: what if client's offset changes so we get a new instruction for what we think is an existing frame? The case where offset changes is input fault (hopefully rare) so maybe we just accept it
@@ -365,27 +396,12 @@ void UAsyncPhysicsInputComponent::ServerRPCBufferInput_Implementation(FAsyncPhys
 		}
 	}
 
-	BufferedInputs.Add(Wrapper.Input);
-}
-
-
-bool FAsyncPhysicsInputWrapper::NetSerialize(FArchive& Ar, class UPackageMap* Map, bool& bOutSuccess)
-{
-	Ar << OwnerComponent;
-	FAsyncPhysicsInputPool* Pool = OwnerComponent->GetInputPool();
-	Pool->NetSerializeHelper(Input, Ar, Map, bOutSuccess);
-	return bOutSuccess;
-}
-
-void UAsyncPhysicsInputComponent::RegisterInputPool(FAsyncPhysicsInputPool* InPool)
-{
-	ensureMsgf(Pool == nullptr, TEXT("You can only register the input pool once during initialization"));
-	Pool = InPool;
+	BufferedData.Add(NewData);
 }
 
 void UAsyncPhysicsInputComponent::OnDispatchPhysicsTick(int32 PhysicsStep, int32 NumSteps, int32 ServerFrame)
 {
-	ensureMsgf(Pool != nullptr, TEXT("You must call RegisterInputPool in InitializeComponent"));
+	ensureMsgf(DataClass != nullptr, TEXT("You must call SetDataClass after creating the component"));
 
 	// It would be better if we only registered while we had a local controller,
 	// but this is simpler to implement
@@ -394,107 +410,53 @@ void UAsyncPhysicsInputComponent::OnDispatchPhysicsTick(int32 PhysicsStep, int32
 	{
 		if (PC->IsLocalController())
 		{
-			FAsyncPhysicsInput* Input = Pool->FlushLatestInputToPopulate();
+			UAsyncPhysicsData* DataToSend = DataToWrite;
+			DataToWrite = DuplicateObject((UAsyncPhysicsData*)DataClass->GetDefaultObject(), nullptr);
 			for (int32 Step = 0; Step < NumSteps; ++Step)
 			{
-				if(Step > 0)
+				if(Step > 0 && DataToSend)
 				{
 					//Make sure each sub-step gets its own unique data so we don't have to worry about ref count
 					//TODO: should probably given user opportunity to modify this per sub-step. For example a jump instruction should only happen on first sub-step
-					Input = Pool->CloneInput(Input);
+					DataToSend = DuplicateObject(DataToSend, nullptr);
 				}
-				Input->ServerFrame = ServerFrame + 1 + Step;
-				BufferedInputs.Add(Input);
+				DataToSend->ServerFrame = ServerFrame + 1 + Step;
+				BufferedData.Add(DataToSend);
 			}
 		}
 	}
 }
 
-void UAsyncPhysicsInputComponent::InitializeComponent()
+APlayerController* UAsyncPhysicsInputComponent::GetPlayerController()
 {
-	Super::InitializeComponent();
-
-	// Test is component is valid for Network Physics. Needs a valid physics ActorHandle
-	auto ValidComponent = [](UActorComponent* Component)
+	if (APlayerController* PC = Cast<APlayerController>(GetOwner()))
 	{
-		bool bValid = false;
-		if (UPrimitiveComponent* PrimComp = Cast<UPrimitiveComponent>(Component))
-		{
-			bValid = FPhysicsInterface::IsValid(PrimComp->BodyInstance.ActorHandle);
-		}
-		return bValid;
-	};
-
-	auto SelectComponent = [&ValidComponent](const TArray<UActorComponent*>& Components)
-	{
-		UPrimitiveComponent* Pc = nullptr;
-		for (UActorComponent* Ac : Components)
-		{
-			if (ValidComponent(Ac))
-			{
-				Pc = (UPrimitiveComponent*)Ac;
-				break;
-			}
-
-		}
-		return Pc;
-	};
+		return PC;
+	}
 
-	if (AActor* MyActor = GetOwner())
+	if (APawn* Pawn = Cast<APawn>(GetOwner()))
 	{
-		TInlineComponentArray<UPrimitiveComponent*> FoundComponents;
-		MyActor->GetComponents<UPrimitiveComponent>(FoundComponents);
-		for (UPrimitiveComponent* FoundComponent : FoundComponents)
-		{
-			if (FPhysicsInterface::IsValid(FoundComponent->BodyInstance.ActorHandle))
-			{
-				UpdateComponent = FoundComponent;
-				break;
-			}
-		}
+		Pawn->GetController<APlayerController>();
 	}
 
-	//TODO: set up tick dependency so that async input component runs before the update component
-	//ensure(UpdateComponent);
-
-
-
-	UWorld* World = GetWorld();
-	if (World->WorldType == EWorldType::PIE || World->WorldType == EWorldType::Game)
-	{
-		FPhysScene* PhysScene = World->GetPhysicsScene();
-		if (ensureAlways(PhysScene))
-		{
-			Chaos::FPhysicsSolver* Solver = PhysScene->GetSolver();
-			if (Solver->GetRewindCallback() == nullptr)
-			{
-				// -------------------------------------------------------------------------------------
-				// Hack: enable Rewind capture here. This needs to be moved somewhere else
-				// -------------------------------------------------------------------------------------
+	return nullptr;
+}
 
-				int32 NumFrames = 64;
-				const IConsoleVariable* NumFramesCVar = IConsoleManager::Get().FindConsoleVariable(TEXT("p.RewindCaptureNumFrames"));
-				if (ensure(NumFramesCVar))
-				{
-					// 1 frame is required to enable rewind capture. 
-					NumFrames = FMath::Max<int32>(1, NumFramesCVar->GetInt());
-				}
-				Solver->EnableRewindCapture(NumFrames, false, MakeUnique<FAsyncPhysicsInputRewindCallback>(World));
-			}
 
-			FAsyncPhysicsInputRewindCallback* Callback = static_cast<FAsyncPhysicsInputRewindCallback*>(Solver->GetRewindCallback());
-			Callback->DispatchPhysicsTick.AddUObject(this, &UAsyncPhysicsInputComponent::OnDispatchPhysicsTick);
-		}
-	}
+UAsyncPhysicsData* UAsyncPhysicsInputComponent::GetDataToWrite() const
+{
+	//TODO: ensure not inside async physics tick
+	return DataToWrite;
 }
 
-APlayerController* UAsyncPhysicsInputComponent::GetPlayerController()
+const UAsyncPhysicsData* UAsyncPhysicsInputComponent::GetDataToConsume() const
 {
-	APawn* Pawn = Cast<APawn>(GetOwner());
-	if (!Pawn)
+	//TODO: ensure inside async physics tick
+	ensureMsgf(DataClass != nullptr, TEXT("You must call SetDataClass after creating the component"));
+	if (DataToConsume)
 	{
-		return nullptr;
+		return DataToConsume;
 	}
-
-	return Pawn->GetController<APlayerController>();
+	
+	return (UAsyncPhysicsData*)DataClass->GetDefaultObject();
 }
\ No newline at end of file
diff --git a/Engine/Source/Runtime/Engine/Private/PlayerController.cpp b/Engine/Source/Runtime/Engine/Private/PlayerController.cpp
--- a/Engine/Source/Runtime/Engine/Private/PlayerController.cpp
+++ b/Engine/Source/Runtime/Engine/Private/PlayerController.cpp
@@ -69,6 +69,7 @@
 #include "Particles/EmitterCameraLensEffectBase.h"
 #include "LevelUtils.h"
 #include "WorldPartition/WorldPartitionStreamingSource.h"
+#include "Physics/AsyncPhysicsInputComponent.h"
 
 DEFINE_LOG_CATEGORY(LogPlayerController);
 
@@ -4670,11 +4671,25 @@ void APlayerController::GetLifetimeReplicatedProps(TArray< FLifetimeProperty > &
 
 	// Replicate SpawnLocation for remote spectators
 	DOREPLIFETIME_CONDITION(APlayerController, SpawnLocation, COND_OwnerOnly);
+
+	DOREPLIFETIME_CONDITION(APlayerController, AsyncPhysicsDataComponent, COND_OwnerOnly);
+}
+
+void APlayerController::OnRep_AsyncPhysicsDataComponent()
+{
+	AsyncPhysicsDataComponent->SetDataClass(AsyncPhysicsDataClass);
 }
 
 void APlayerController::SetPlayer( UPlayer* InPlayer )
 {
 	FMoviePlayerProxyBlock MoviePlayerBlock;
+	if(AsyncPhysicsDataClass && GetLocalRole() == ROLE_Authority)
+	{
+		static const FName AsyncPhysicsDataComponentName(TEXT("PC_AsyncPhysicsDataComponent"));
+		AsyncPhysicsDataComponent = NewObject<UAsyncPhysicsInputComponent>(this, AsyncPhysicsDataComponentName);
+		AsyncPhysicsDataComponent->SetDataClass(AsyncPhysicsDataClass);
+		AsyncPhysicsDataComponent->RegisterComponent();
+	}
 
 	check(InPlayer!=NULL);
 
@@ -5818,4 +5833,15 @@ void APlayerController::IncludeInNetConditionGroup(FName NetGroup)
 	NetConditionGroups.AddUnique(NetGroup);
 }
 
+UAsyncPhysicsData* APlayerController::GetAsyncPhysicsDataToWrite() const
+{
+	return AsyncPhysicsDataComponent ? AsyncPhysicsDataComponent->GetDataToWrite() : nullptr;
+}
+
+const UAsyncPhysicsData* APlayerController::GetAsyncPhysicsDataToConsume() const
+{
+	return AsyncPhysicsDataComponent ? AsyncPhysicsDataComponent->GetDataToConsume() : nullptr;
+}
+
+
 #undef LOCTEXT_NAMESPACE
diff --git a/Engine/Source/Runtime/Engine/Public/Physics/AsyncPhysicsData.h b/Engine/Source/Runtime/Engine/Public/Physics/AsyncPhysicsData.h
new file mode 100644
--- /dev/null
+++ b/Engine/Source/Runtime/Engine/Public/Physics/AsyncPhysicsData.h
@@ -0,0 +1,31 @@
+// Copyright Epic Games, Inc. All Rights Reserved.
+
+#pragma once
+
+#include "CoreMinimal.h"
+#include "AsyncPhysicsData.generated.h"
+
+/**
+	The base class for async physics data. Inherit from this to create custom data for async physics tick.
+	When no data is available (say due to massive latency or packet loss) we fall back on the default constructed data.
+	This means you should set the default values to something equivalent to no input (for example bPlayerWantsToJump should probably default to false)
+*/
+UCLASS(BlueprintType, Blueprintable)
+class ENGINE_API UAsyncPhysicsData : public UObject
+{
+	GENERATED_BODY()
+public:
+	virtual ~UAsyncPhysicsData() = default;
+	int32 GetServerFrame() const { return ServerFrame; }
+
+private:
+	UPROPERTY()
+	int32 ServerFrame = INDEX_NONE;
+protected:
+
+	//Determines how many times we redundantly send data to server. The higher this number the less packet loss, but more bandwidth is used
+	UPROPERTY(EditDefaultsOnly, Category = AsyncPhysicsData)
+	int32 ReplicationRedundancy = 4;
+
+	friend class UAsyncPhysicsInputComponent;
+};
\ No newline at end of file
diff --git a/Engine/Source/Runtime/Engine/Public/Physics/AsyncPhysicsInputComponent.h b/Engine/Source/Runtime/Engine/Public/Physics/AsyncPhysicsInputComponent.h
--- a/Engine/Source/Runtime/Engine/Public/Physics/AsyncPhysicsInputComponent.h
+++ b/Engine/Source/Runtime/Engine/Public/Physics/AsyncPhysicsInputComponent.h
@@ -2,240 +2,48 @@
 
 #pragma once
 #include "Components/ActorComponent.h"
-#include "Engine/NetSerialization.h"
-#include "Engine/NetConnection.h"
-#include "Engine/PackageMapClient.h"
-#include "Net/RepLayout.h"
+#include "Physics/AsyncPhysicsData.h"
 
 #include "AsyncPhysicsInputComponent.generated.h"
 
-class UAsyncPhysicsInputComponent;
-
-/** 
-	The base struct for async physics input. Inherit from this to create custom input data for async physics tick.
-	When no input is available (say due to massive latency or packet loss) we fall back on the default constructed input.
-	This means you should set the default values to something equivalent to no input (for example bPlayerWantsToJump should probably default to false)
-*/
-USTRUCT()
-struct FAsyncPhysicsInput
-{
-	GENERATED_BODY()
-public:
-	virtual ~FAsyncPhysicsInput() = default;
-	int32 GetServerFrame() const { return ServerFrame; }
-private:
-	UPROPERTY()
-	int32 ServerFrame = INDEX_NONE;
-
-	int32 Replicated = 4;	//how many times we want to replicate to server. TODO: use cvar
-
-	friend UAsyncPhysicsInputComponent;
-};
-
-/** Helper struct to replicate polymorphic input data */
-USTRUCT()
-struct FAsyncPhysicsInputWrapper
-{
-	GENERATED_BODY()
-
-	FAsyncPhysicsInput* Input = nullptr;
-
-	bool NetSerialize(FArchive& Ar, class UPackageMap* Map, bool& bOutSuccess);
-
-	UPROPERTY()
-	TObjectPtr<UAsyncPhysicsInputComponent> OwnerComponent = nullptr;
-};
-
-template<>
-struct TStructOpsTypeTraits<FAsyncPhysicsInputWrapper> : public TStructOpsTypeTraitsBase2<FAsyncPhysicsInputWrapper>
-{
-	enum
-	{
-		WithNetSerializer = true
-	};
-};
-
-class FAsyncPhysicsInputPool;
-
-/** Base component used to easily send data from variable tick to async physics.
-	Has networking support so server and client execute AsyncPhysicsTickComponent on same data for same step.
-	Inherit from this class to create your own component that you can act on inputs with.
-
-	Extended class must hold a TAsyncPhysicsInputPool<T> where T is your custom input extended from FAsyncPhysicsInput.
-	Must call RegisterInputPool in InitializeComponent
-*/
 UCLASS(BlueprintType)
 class ENGINE_API UAsyncPhysicsInputComponent : public UActorComponent
 {
 	GENERATED_BODY()
 public:
 	UAsyncPhysicsInputComponent();
-	void InitializeComponent() override;
-
-	/** Must be called in InitializeComponent by extended class */
-	void RegisterInputPool(FAsyncPhysicsInputPool* InPool);
+	
+	/** Sets the data class used. Can only call this once after the component has been created */
+	void SetDataClass(TSubclassOf<UAsyncPhysicsData> InDataClass);
 
-	/** Executes physics logic based on the current input (See TNetworkedInputPool::GetCurrentInput)*/
 	virtual void AsyncPhysicsTickComponent(float DeltaTime, float SimTime) override;
 
 	UFUNCTION(Server, unreliable)
-	void ServerRPCBufferInput(FAsyncPhysicsInputWrapper PlayerInput);
+	void ServerRPCBufferInput(UAsyncPhysicsData* AsyncPhysicsData);
 
 	void OnDispatchPhysicsTick(int32 PhysicsStep, int32 NumSteps, int32 ServerFrame);
 
-	FAsyncPhysicsInputPool* GetInputPool() { return Pool; }
-
-protected:
-	UPrimitiveComponent* UpdateComponent = nullptr;
-
 	APlayerController* GetPlayerController();
 
-private:
-	TArray<FAsyncPhysicsInput*> BufferedInputs;
-	FAsyncPhysicsInputPool* Pool = nullptr;
-};
+	/** Get the async physics data to write to. This data will make its way to the async physics tick on client and server. Should not be called from async tick */
+	UFUNCTION(BlueprintPure, Category = PlayerController)
+	UAsyncPhysicsData* GetDataToWrite() const;
 
+	/** Get the async physics data to execute logic off of. This data should not be modified and will NOT make its way back. Should be called from async tick */
+	UFUNCTION(BlueprintPure, Category = PlayerController)
+	const UAsyncPhysicsData* GetDataToConsume() const;
 
-class FAsyncPhysicsInputPool
-{
-public:
-	virtual ~FAsyncPhysicsInputPool() = default;
-
-private:
-	friend class UAsyncPhysicsInputComponent;
-	friend FAsyncPhysicsInputWrapper;
-
-	/** Flushes the current input to populate so that new inputs can be populated. Returns the previous input to be saved off */
-	virtual FAsyncPhysicsInput* FlushLatestInputToPopulate() = 0;
-
-	/** Makes a deep copy of the input so that we can treat it as identical input over multiple physics steps from one slow gamethread step */
-	virtual FAsyncPhysicsInput* CloneInput(const FAsyncPhysicsInput* Input) = 0;
-
-	/** Sets the current input for execution during async physics */
-	virtual void SetCurrentInputToAsyncExecute(FAsyncPhysicsInput* Input) = 0;
-
-	/** Frees the input back into the pool*/
-	virtual void FreeInputToPool(FAsyncPhysicsInput* Input) = 0;
-
-	/** Serialize the underlying data. If needed a new entry in the pool will be created (during deserialization) */
-	virtual void NetSerializeHelper(FAsyncPhysicsInput*& Data, FArchive& Ar, class UPackageMap* Map, bool& bOutSuccess) = 0;
-};
-
-/** Helper class to pool inputs and do various operations on like serialization.
-*   T must have FAsyncPhysicsInput as its base */
-template <typename T>
-class TAsyncPhysicsInputPool final : public FAsyncPhysicsInputPool
-{
-public:
-	TAsyncPhysicsInputPool()
-	{
-		PendingInputToPopulate = NewInput();
-	}
-
-	~TAsyncPhysicsInputPool()
-	{
-		delete PendingInputToPopulate;
-		delete CurrentInputToExecute;
-		while (Pool.Num())
-		{
-			T* Input = Pool.Pop();
-			delete Input;
-		}
-	}
-
-	/** Gets the pending input object for populating. This is what gets sent to the server and the async physics tick to execute logic off of.
-		Should not be used in async tick as it may not be the right physics step*/
-	T& GetPendingInputToPopulate() { return *PendingInputToPopulate; }
-	const T& GetPendingInputToPopulate() const { return *PendingInputToPopulate; }
-
-	/** Gets the current input object to execute logic off of. This should be used during async physics tick. If no input is available we use default constructed values*/
-	const T& GetCurrentInput()
-	{
-		static T NoInput;
-		return CurrentInputToExecute ? *CurrentInputToExecute : NoInput;
-	}
 
 private:
+	UPROPERTY()
+	TSubclassOf<UAsyncPhysicsData> DataClass;
 
-	virtual FAsyncPhysicsInput* FlushLatestInputToPopulate() override
-	{
-		//Assumes CurrentInputToPopulate has been saved into the send to server buffer and will be properly freed back into pool
-		FAsyncPhysicsInput* OldInput = PendingInputToPopulate;
-		PendingInputToPopulate = NewInput();
-		return OldInput;
-	}
-
-	virtual FAsyncPhysicsInput* CloneInput(const FAsyncPhysicsInput* Input) override
-	{
-		T* Input2 = NewInput();
-		*Input2 = *(static_cast<const T*>(Input));
-		return Input2;
-	}
-
-	virtual void SetCurrentInputToAsyncExecute(FAsyncPhysicsInput* Input) override
-	{
-		CurrentInputToExecute = static_cast<T*>(Input);
-	}
-
-	virtual void NetSerializeHelper(FAsyncPhysicsInput*& Data, FArchive& Ar, class UPackageMap* Map, bool& bOutSuccess) override
-	{
-		if (Ar.IsLoading())
-		{
-			Data = NewInput();
-		}
-
-		UScriptStruct* ScriptStruct = T::StaticStruct();
-		if (ScriptStruct->StructFlags & STRUCT_NetSerializeNative)
-		{
-			ScriptStruct->GetCppStructOps()->NetSerialize(Ar, Map, bOutSuccess, Data);
-		}
-		else
-		{
-			UNetConnection* Connection = CastChecked<UPackageMapClient>(Map)->GetConnection();
-			UNetDriver* NetDriver = Connection ? Connection->GetDriver() : nullptr;
-			TSharedPtr<FRepLayout> RepLayout = NetDriver ? NetDriver->GetStructRepLayout(ScriptStruct) : nullptr;
-
-			if (RepLayout.IsValid())
-			{
-				if (FBitArchive* BitAr = static_cast<FBitArchive*>(&Ar))
-				{
-					bool bHasUnmapped = false;
-					RepLayout->SerializePropertiesForStruct(ScriptStruct, *BitAr, Map, Data, bHasUnmapped);
-
-					bOutSuccess = true;
-				}
-			}
-		}
-	}
-
-	T* NewInput()
-	{
-		T* Input;
-		if (Pool.Num())
-		{
-			Input = Pool.Pop();
-			new (Input) T();
-		}
-		else
-		{
-			Input = new T();
-		}
-
-		return Input;
-	}
-
-	void FreeInput(T* Input)
-	{
-		Input->~T();
-		Pool.Add(Input);
-	}
+	UPROPERTY(Transient)
+	TArray<TObjectPtr<UAsyncPhysicsData>> BufferedData;
 
-	virtual void FreeInputToPool(FAsyncPhysicsInput* Input) override
-	{
-		FreeInput(static_cast<T*>(Input));
-	}
+	UPROPERTY(Transient)
+	TObjectPtr<UAsyncPhysicsData> DataToConsume;
 
-	T* CurrentInputToExecute = nullptr;
-	T* PendingInputToPopulate = nullptr;
-	TArray<T*> Pool;
+	UPROPERTY(Transient)
+	TObjectPtr<UAsyncPhysicsData> DataToWrite;
 };
\ No newline at end of file
