Material Baking - Fixed FMaterialShaderMapId operator== not considering Usage when comparing a cooked id
* This prevented material proxy from being picked up as needed to be compiled (as they rely on usages to specify the properties to bake, ex EMaterialShaderMapUsage::MaterialExportBaseColor for diffuse)

diff --git a/Engine/Source/Runtime/Engine/Private/Materials/MaterialShader.cpp b/Engine/Source/Runtime/Engine/Private/Materials/MaterialShader.cpp
--- a/Engine/Source/Runtime/Engine/Private/Materials/MaterialShader.cpp
+++ b/Engine/Source/Runtime/Engine/Private/Materials/MaterialShader.cpp
@@ -785,19 +785,27 @@ bool FMaterialShaderMapId::operator==(const FMaterialShaderMapId& ReferenceSet)
 
 	if (!IsCookedId())
 	{
-		if (Usage != ReferenceSet.Usage
-			|| UsageCustomOutput != ReferenceSet.UsageCustomOutput
-			|| BaseMaterialId != ReferenceSet.BaseMaterialId)
+		if (BaseMaterialId != ReferenceSet.BaseMaterialId)
 		{
 			return false;
 		}
 	}
 	else
 #endif
-	if (CookedShaderMapIdHash != ReferenceSet.CookedShaderMapIdHash)
+	{
+		if (CookedShaderMapIdHash != ReferenceSet.CookedShaderMapIdHash)
+		{
+			return false;
+		}
+	}
+
+#if WITH_EDITOR
+	if (Usage != ReferenceSet.Usage
+		|| UsageCustomOutput != ReferenceSet.UsageCustomOutput)
 	{
 		return false;
 	}
+#endif
 
 	if (QualityLevel != ReferenceSet.QualityLevel
 		|| FeatureLevel != ReferenceSet.FeatureLevel)
