Add ENGINE_API to USkeleton::RebuildLinkup.
Change by [REDACTED_PERSON]

diff --git a/Engine/Source/Runtime/Engine/Classes/Animation/Skeleton.h b/Engine/Source/Runtime/Engine/Classes/Animation/Skeleton.h
--- a/Engine/Source/Runtime/Engine/Classes/Animation/Skeleton.h
+++ b/Engine/Source/Runtime/Engine/Classes/Animation/Skeleton.h
@@ -1013,7 +1013,7 @@ public:
 	 *
 	 * @param	InSkelMesh	: SkeletalMesh to build look up for
 	 */
-	void RebuildLinkup(const USkeletalMesh* InSkelMesh);
+	ENGINE_API void RebuildLinkup(const USkeletalMesh* InSkelMesh);
 
 	/**
 	 * Remove Link up cache for the SkelMesh
