Fix UTimelineComponent::OnRep_Timeline to only set the playback position only if the Position changed and not other variables in the struct

diff --git a/Engine/Source/Runtime/Engine/Classes/Components/TimelineComponent.h b/Engine/Source/Runtime/Engine/Classes/Components/TimelineComponent.h
--- a/Engine/Source/Runtime/Engine/Classes/Components/TimelineComponent.h
+++ b/Engine/Source/Runtime/Engine/Classes/Components/TimelineComponent.h
@@ -498,7 +498,7 @@ public:
 	ENGINE_API void SetLinearColorCurve(UCurveLinearColor* NewLinearColorCurve, FName LinearColorTrackName);
 
 	UFUNCTION()
-	void OnRep_Timeline();
+	void OnRep_Timeline(FTimeline& OldTimeline);
 
 	//~ Begin ActorComponent Interface.
 	virtual void TickComponent(float DeltaTime, enum ELevelTick TickType, FActorComponentTickFunction *ThisTickFunction) override;
diff --git a/Engine/Source/Runtime/Engine/Private/Timeline.cpp b/Engine/Source/Runtime/Engine/Private/Timeline.cpp
--- a/Engine/Source/Runtime/Engine/Private/Timeline.cpp
+++ b/Engine/Source/Runtime/Engine/Private/Timeline.cpp
@@ -948,12 +948,11 @@ void UTimelineComponent::SetDirectionPropertyName(FName DirectionPropertyName)
 	TheTimeline.SetDirectionPropertyName(DirectionPropertyName);
 }
 
-void UTimelineComponent::OnRep_Timeline()
+void UTimelineComponent::OnRep_Timeline(FTimeline& OldTimeline)
 {
-	if (!TheTimeline.IsPlaying())
+	if (!TheTimeline.IsPlaying() && OldTimeline.GetPlaybackPosition() != TheTimeline.GetPlaybackPosition())
 	{
 		// make sure a final update call occurs on the client for the final position
-		// FIXME: this is incomplete, we need to compare vs the last simulated position for firing events and such
 		TheTimeline.SetPlaybackPosition(TheTimeline.GetPlaybackPosition(), false, true);
 	}
 }
