Add more descriptive ensure msg when hitting a circular dependency in simple construction script

diff --git a/Engine/Source/Runtime/Engine/Private/SCS_Node.cpp b/Engine/Source/Runtime/Engine/Private/SCS_Node.cpp
--- a/Engine/Source/Runtime/Engine/Private/SCS_Node.cpp
+++ b/Engine/Source/Runtime/Engine/Private/SCS_Node.cpp
@@ -332,7 +332,7 @@ void USCS_Node::PreloadChain()
 
 	if (ComponentTemplate && ComponentTemplate->HasAnyFlags(RF_NeedLoad))
 	{
-		if (ensure(ComponentTemplate->GetLinker()))
+		if (ensureMsgf(ComponentTemplate->GetLinker(), TEXT("Failed to find linker for %s, likely a circular dependency"), *ComponentTemplate->GetPathName()))
 		{
 			ComponentTemplate->GetLinker()->Preload(ComponentTemplate);
 		}
