World Partition: fix crash when changing maps in the editor.

#rnx

diff --git a/Engine/Source/Runtime/Engine/Private/WorldPartition/WorldPartitionActorLoaderInterface.cpp b/Engine/Source/Runtime/Engine/Private/WorldPartition/WorldPartitionActorLoaderInterface.cpp
--- a/Engine/Source/Runtime/Engine/Private/WorldPartition/WorldPartitionActorLoaderInterface.cpp
+++ b/Engine/Source/Runtime/Engine/Private/WorldPartition/WorldPartitionActorLoaderInterface.cpp
@@ -30,7 +30,7 @@ IWorldPartitionActorLoaderInterface::ILoaderAdapter::ILoaderAdapter(UWorld* InWo
 
 IWorldPartitionActorLoaderInterface::ILoaderAdapter::~ILoaderAdapter()
 {
-	Unload();
+	UnregisterDelegates();
 }
 
 bool IWorldPartitionActorLoaderInterface::ILoaderAdapter::Load()
@@ -39,7 +39,7 @@ bool IWorldPartitionActorLoaderInterface::ILoaderAdapter::Load()
 	{
 		bLoaded = true;
 		RefreshLoadedState();
-		FDataLayersEditorBroadcast::Get().OnActorDataLayersEditorLoadingStateChanged().AddRaw(this, &IWorldPartitionActorLoaderInterface::ILoaderAdapter::OnActorDataLayersEditorLoadingStateChanged);
+		RegisterDelegates();
 	}
 
 	return true;
@@ -65,7 +65,8 @@ bool IWorldPartitionActorLoaderInterface::ILoaderAdapter::Unload()
 			FScopedSlowTask SlowTask(1, LOCTEXT("UpdatingLoading", "Updating loading..."));
 			SlowTask.MakeDialog();
 
-			FDataLayersEditorBroadcast::Get().OnActorDataLayersEditorLoadingStateChanged().RemoveAll(this);
+			UnregisterDelegates();
+
 			ActorReferences.Empty();
 			bLoaded = false;
 
@@ -83,6 +84,16 @@ bool IWorldPartitionActorLoaderInterface::ILoaderAdapter::IsLoaded() const
 	return bLoaded;
 }
 
+void IWorldPartitionActorLoaderInterface::ILoaderAdapter::RegisterDelegates()
+{
+	FDataLayersEditorBroadcast::Get().OnActorDataLayersEditorLoadingStateChanged().AddRaw(this, &IWorldPartitionActorLoaderInterface::ILoaderAdapter::OnActorDataLayersEditorLoadingStateChanged);
+}
+
+void IWorldPartitionActorLoaderInterface::ILoaderAdapter::UnregisterDelegates()
+{
+	FDataLayersEditorBroadcast::Get().OnActorDataLayersEditorLoadingStateChanged().RemoveAll(this);
+}
+
 bool IWorldPartitionActorLoaderInterface::ILoaderAdapter::RefreshLoadedState()
 {
 	if (bLoaded)
diff --git a/Engine/Source/Runtime/Engine/Public/WorldPartition/WorldPartitionActorLoaderInterface.h b/Engine/Source/Runtime/Engine/Public/WorldPartition/WorldPartitionActorLoaderInterface.h
--- a/Engine/Source/Runtime/Engine/Public/WorldPartition/WorldPartitionActorLoaderInterface.h
+++ b/Engine/Source/Runtime/Engine/Public/WorldPartition/WorldPartitionActorLoaderInterface.h
@@ -42,6 +42,8 @@ public:
 		virtual bool Intersect(const FBox& Box) const =0;
 		virtual void ForEachActor(TFunctionRef<void(const FWorldPartitionHandle&)> InOperation) const =0;
 
+		void RegisterDelegates();
+		void UnregisterDelegates();
 		bool RefreshLoadedState();
 		bool ShouldActorBeLoaded(const FWorldPartitionHandle& Actor) const;
 		void PostLoadedStateChanged(bool bUnloadedActors);
