World Partition: added last loaded location volumes support.

#rnx

diff --git a/Engine/Source/Editor/UnrealEd/Private/FileHelpers.cpp b/Engine/Source/Editor/UnrealEd/Private/FileHelpers.cpp
--- a/Engine/Source/Editor/UnrealEd/Private/FileHelpers.cpp
+++ b/Engine/Source/Editor/UnrealEd/Private/FileHelpers.cpp
@@ -754,7 +754,7 @@ static bool SaveWorld(UWorld* World,
 				// When creating a Partitioned Level for a Level Instance, the WorldPartition is not initialized, so no need to do this.
 				if (RenamedWorldPartition && RenamedWorldPartition->IsInitialized())
 				{
-					LoadedEditorRegions = RenamedWorldPartition->GetUserLoadedEditorGridRegions();
+					LoadedEditorRegions = RenamedWorldPartition->GetUserLoadedEditorRegions();
 					RenamedWorldPartition->LoadAllActors(ActorReferences);
 
 					if (bIsTempPackage)
@@ -900,7 +900,7 @@ static bool SaveWorld(UWorld* World,
 			if (RenamedWorldPartition)
 			{
 				// Save Snapshot of loaded Editor regions
-				GetMutableDefault<UWorldPartitionEditorPerProjectUserSettings>()->SetEditorGridLoadedRegions(SaveWorld, LoadedEditorRegions);
+				GetMutableDefault<UWorldPartitionEditorPerProjectUserSettings>()->SetEditorLoadedRegions(SaveWorld, LoadedEditorRegions);
 				RenamedWorldPartition->LoadLastLoadedRegions();
 			}
 		}
diff --git a/Engine/Source/Runtime/Engine/Private/LocationVolume.cpp b/Engine/Source/Runtime/Engine/Private/LocationVolume.cpp
--- a/Engine/Source/Runtime/Engine/Private/LocationVolume.cpp
+++ b/Engine/Source/Runtime/Engine/Private/LocationVolume.cpp
@@ -12,6 +12,7 @@ ALocationVolume::ALocationVolume(const FObjectInitializer& ObjectInitializer)
 
 #if WITH_EDITORONLY_DATA
 	bIsSpatiallyLoaded = false;
+	bIsAutoLoad = false;
 #endif
 
 #if WITH_EDITOR
@@ -22,6 +23,18 @@ ALocationVolume::ALocationVolume(const FObjectInitializer& ObjectInitializer)
 #endif
 }
 
+void ALocationVolume::PostRegisterAllComponents()
+{
+	Super::PostRegisterAllComponents();
+
+#if WITH_EDITOR
+	if (WorldPartitionActorLoader && bIsAutoLoad)
+	{
+		WorldPartitionActorLoader->Load();
+	}
+#endif
+}
+
 void ALocationVolume::BeginDestroy()
 {
 #if WITH_EDITOR
diff --git a/Engine/Source/Runtime/Engine/Private/WorldPartition/LoaderAdapter/LoaderAdapterActor.cpp b/Engine/Source/Runtime/Engine/Private/WorldPartition/LoaderAdapter/LoaderAdapterActor.cpp
--- a/Engine/Source/Runtime/Engine/Private/WorldPartition/LoaderAdapter/LoaderAdapterActor.cpp
+++ b/Engine/Source/Runtime/Engine/Private/WorldPartition/LoaderAdapter/LoaderAdapterActor.cpp
@@ -1,12 +1,15 @@
 // Copyright Epic Games, Inc. All Rights Reserved.
 
 #include "WorldPartition/LoaderAdapter/LoaderAdapterActor.h"
+#include "GameFramework/Actor.h"
 
 #if WITH_EDITOR
 FLoaderAdapterActor::FLoaderAdapterActor(AActor* InActor)
 	: ILoaderAdapterSpatial(InActor->GetWorld())
 	, Actor(InActor)
-{}
+{
+	SetUserCreated(true);
+}
 
 TOptional<FBox> FLoaderAdapterActor::GetBoundingBox() const
 {
diff --git a/Engine/Source/Runtime/Engine/Private/WorldPartition/WorldPartition.cpp b/Engine/Source/Runtime/Engine/Private/WorldPartition/WorldPartition.cpp
--- a/Engine/Source/Runtime/Engine/Private/WorldPartition/WorldPartition.cpp
+++ b/Engine/Source/Runtime/Engine/Private/WorldPartition/WorldPartition.cpp
@@ -32,6 +32,7 @@
 #include "Misc/MessageDialog.h"
 #include "ScopedTransaction.h"
 #include "UnrealEdMisc.h"
+#include "LocationVolume.h"
 #include "WorldPartition/IWorldPartitionEditorModule.h"
 #include "WorldPartition/WorldPartitionLevelHelper.h"
 #include "WorldPartition/WorldPartitionLevelStreamingDynamic.h"
@@ -1154,13 +1155,13 @@ bool UWorldPartition::FinalizeGeneratorPackageForCook(const TArray<ICookPackageS
 	return false;
 }
 
-TArray<FBox> UWorldPartition::GetUserLoadedEditorGridRegions() const
+TArray<FBox> UWorldPartition::GetUserLoadedEditorRegions() const
 {
 	TArray<FBox> Result;
 
 	for (IWorldPartitionActorLoaderInterface::ILoaderAdapter* LoaderAdapter : RegisteredEditorLoaderAdapters)
 	{
-		if (LoaderAdapter->GetUserCreated())
+		if (LoaderAdapter->IsLoaded() && LoaderAdapter->GetUserCreated())
 		{
 			Result.Add(*LoaderAdapter->GetBoundingBox());
 		}
@@ -1173,7 +1174,25 @@ void UWorldPartition::SavePerUserSettings()
 {
 	if (GIsEditor && !World->IsGameWorld() && !IsRunningCommandlet() && !IsEngineExitRequested())
 	{
-		GetMutableDefault<UWorldPartitionEditorPerProjectUserSettings>()->SetEditorGridLoadedRegions(GetWorld(), GetUserLoadedEditorGridRegions());
+		GetMutableDefault<UWorldPartitionEditorPerProjectUserSettings>()->SetEditorLoadedRegions(GetWorld(), GetUserLoadedEditorRegions());
+
+		TArray<FName> EditorLoadedLocationVolumes;
+		for (UActorDescContainer::TConstIterator<> ActorDescIterator(this); ActorDescIterator; ++ActorDescIterator)
+		{
+			if (ALocationVolume* LocationVolume = Cast<ALocationVolume>(ActorDescIterator->GetActor()))
+			{
+				check(LocationVolume->GetClass()->ImplementsInterface(UWorldPartitionActorLoaderInterface::StaticClass()));
+
+				IWorldPartitionActorLoaderInterface::ILoaderAdapter* LoaderAdapter = Cast<IWorldPartitionActorLoaderInterface>(LocationVolume)->GetLoaderAdapter();
+				check(LoaderAdapter);
+
+				if (LoaderAdapter->IsLoaded() && LoaderAdapter->GetUserCreated())
+				{
+					EditorLoadedLocationVolumes.Add(LocationVolume->GetFName());
+				}
+			}
+		}
+		GetMutableDefault<UWorldPartitionEditorPerProjectUserSettings>()->SetEditorLoadedLocationVolumes(GetWorld(), EditorLoadedLocationVolumes);
 	}
 }
 
@@ -1293,14 +1312,24 @@ bool UWorldPartition::IsActorPinned(const FGuid& ActorGuid) const
 
 void UWorldPartition::LoadLastLoadedRegions()
 {
-	TArray<FBox> EditorGridLastLoadedRegions = GetMutableDefault<UWorldPartitionEditorPerProjectUserSettings>()->GetEditorGridLoadedRegions(World);
+	TArray<FBox> EditorLastLoadedRegions = GetMutableDefault<UWorldPartitionEditorPerProjectUserSettings>()->GetEditorLoadedRegions(World);
 
-	for (const FBox& EditorGridLastLoadedRegion : EditorGridLastLoadedRegions)
+	for (const FBox& EditorLastLoadedRegion : EditorLastLoadedRegions)
 	{
-		FLoaderAdapterShape* LoaderAdapter = CreateEditorLoaderAdapter<FLoaderAdapterShape>(World, EditorGridLastLoadedRegion, TEXT("Last Loaded Region"));
+		FLoaderAdapterShape* LoaderAdapter = CreateEditorLoaderAdapter<FLoaderAdapterShape>(World, EditorLastLoadedRegion, TEXT("Last Loaded Region"));
 		LoaderAdapter->SetUserCreated(true);
 		LoaderAdapter->Load();
 	}
+
+	TArray<FName> EditorLoadedLocationVolumes = GetMutableDefault<UWorldPartitionEditorPerProjectUserSettings>()->GetEditorLoadedLocationVolumes(World);
+
+	for (const FName& EditorLoadedLocationVolume : EditorLoadedLocationVolumes)
+	{
+		if (ALocationVolume* LocationVolume = FindObject<ALocationVolume>(World->PersistentLevel, *EditorLoadedLocationVolume.ToString()))
+		{
+			LocationVolume->bIsAutoLoad = true;
+		}
+	}
 }
 
 void UWorldPartition::OnWorldRenamed()
diff --git a/Engine/Source/Runtime/Engine/Private/WorldPartition/WorldPartitionEditorPerProjectUserSettings.cpp b/Engine/Source/Runtime/Engine/Private/WorldPartition/WorldPartitionEditorPerProjectUserSettings.cpp
--- a/Engine/Source/Runtime/Engine/Private/WorldPartition/WorldPartitionEditorPerProjectUserSettings.cpp
+++ b/Engine/Source/Runtime/Engine/Private/WorldPartition/WorldPartitionEditorPerProjectUserSettings.cpp
@@ -19,27 +19,48 @@ void UWorldPartitionEditorPerProjectUserSettings::SetWorldDataLayersNonDefaultEd
 	}
 }
 
-void UWorldPartitionEditorPerProjectUserSettings::SetEditorGridLoadedRegions(UWorld* InWorld, const TArray<FBox>& InEditorGridLoadedRegions)
+void UWorldPartitionEditorPerProjectUserSettings::SetEditorLoadedRegions(UWorld* InWorld, const TArray<FBox>& InEditorLoadedRegions)
 {
 	if (ShouldSaveSettings(InWorld))
 	{
 		FWorldPartitionPerWorldSettings& PerWorldSettings = PerWorldEditorSettings.FindOrAdd(TSoftObjectPtr<UWorld>(InWorld));
-		PerWorldSettings.LoadedEditorGridRegions = InEditorGridLoadedRegions;
+		PerWorldSettings.LoadedEditorRegions = InEditorLoadedRegions;
 		
 		SaveConfig();
 	}
 }
 
-TArray<FBox> UWorldPartitionEditorPerProjectUserSettings::GetEditorGridLoadedRegions(UWorld* InWorld) const
+TArray<FBox> UWorldPartitionEditorPerProjectUserSettings::GetEditorLoadedRegions(UWorld* InWorld) const
 {
 	if (const FWorldPartitionPerWorldSettings* PerWorldSettings = GetWorldPartitionPerWorldSettings(InWorld))
 	{
-		return PerWorldSettings->LoadedEditorGridRegions;
+		return PerWorldSettings->LoadedEditorRegions;
 	}
 
 	return TArray<FBox>();
 }
 
+void UWorldPartitionEditorPerProjectUserSettings::SetEditorLoadedLocationVolumes(UWorld* InWorld, const TArray<FName>& InEditorLoadedLocationVolumes)
+{
+	if (ShouldSaveSettings(InWorld))
+	{
+		FWorldPartitionPerWorldSettings& PerWorldSettings = PerWorldEditorSettings.FindOrAdd(TSoftObjectPtr<UWorld>(InWorld));
+		PerWorldSettings.LoadedEditorLocationVolumes = InEditorLoadedLocationVolumes;
+		
+		SaveConfig();
+	}
+}
+
+TArray<FName> UWorldPartitionEditorPerProjectUserSettings::GetEditorLoadedLocationVolumes(UWorld* InWorld) const
+{
+	if (const FWorldPartitionPerWorldSettings* PerWorldSettings = GetWorldPartitionPerWorldSettings(InWorld))
+	{
+		return PerWorldSettings->LoadedEditorLocationVolumes;
+	}
+
+	return TArray<FName>();
+}
+
 TArray<FName> UWorldPartitionEditorPerProjectUserSettings::GetWorldDataLayersNotLoadedInEditor(UWorld* InWorld) const
 {
 	if (const FWorldPartitionPerWorldSettings* PerWorldSettings = GetWorldPartitionPerWorldSettings(InWorld))
diff --git a/Engine/Source/Runtime/Engine/Private/WorldSettings.cpp b/Engine/Source/Runtime/Engine/Private/WorldSettings.cpp
--- a/Engine/Source/Runtime/Engine/Private/WorldSettings.cpp
+++ b/Engine/Source/Runtime/Engine/Private/WorldSettings.cpp
@@ -511,7 +511,7 @@ void AWorldSettings::SaveDefaultWorldPartitionSettings()
 
 	if (WorldPartition)
 	{
-		DefaultWorldPartitionSettings.LoadedEditorGridRegions = WorldPartition->GetUserLoadedEditorGridRegions();		
+		DefaultWorldPartitionSettings.LoadedEditorRegions = WorldPartition->GetUserLoadedEditorRegions();
 
 		if (const UDataLayerSubsystem* DataLayerSubsystem = UWorld::GetSubsystem<UDataLayerSubsystem>(GetWorld()))
 		{
diff --git a/Engine/Source/Runtime/Engine/Public/LocationVolume.h b/Engine/Source/Runtime/Engine/Public/LocationVolume.h
--- a/Engine/Source/Runtime/Engine/Public/LocationVolume.h
+++ b/Engine/Source/Runtime/Engine/Public/LocationVolume.h
@@ -25,6 +25,7 @@ public:
 
 	//~ Begin AActor Interface
 	virtual bool IsEditorOnly() const override { return !bIsRuntime; }
+	virtual void PostRegisterAllComponents();
 #if WITH_EDITOR
 	virtual bool CanChangeIsSpatiallyLoadedFlag() const override { return false; }
 	virtual bool SupportsDataLayer() const override { return false; }
@@ -43,6 +44,11 @@ public:
 	UPROPERTY(EditAnywhere, Category=LocationVolume)
 	uint8 bIsRuntime : 1;
 
+#if WITH_EDITORONLY_DATA
+	/* To support per-user last loaded location volumes */
+	uint8 bIsAutoLoad : 1;
+#endif
+
 #if WITH_EDITOR
 private:
 	FLoaderAdapterActor* WorldPartitionActorLoader;
diff --git a/Engine/Source/Runtime/Engine/Public/WorldPartition/WorldPartition.h b/Engine/Source/Runtime/Engine/Public/WorldPartition/WorldPartition.h
--- a/Engine/Source/Runtime/Engine/Public/WorldPartition/WorldPartition.h
+++ b/Engine/Source/Runtime/Engine/Public/WorldPartition/WorldPartition.h
@@ -96,7 +96,7 @@ public:
 	FWorldPartitionInitializeDelegate OnWorldPartitionUninitialized;
 
 #if WITH_EDITOR
-	TArray<FBox> GetUserLoadedEditorGridRegions() const;
+	TArray<FBox> GetUserLoadedEditorRegions() const;
 
 public:
 	bool SupportsStreaming() const;
diff --git a/Engine/Source/Runtime/Engine/Public/WorldPartition/WorldPartitionEditorPerProjectUserSettings.h b/Engine/Source/Runtime/Engine/Public/WorldPartition/WorldPartitionEditorPerProjectUserSettings.h
--- a/Engine/Source/Runtime/Engine/Public/WorldPartition/WorldPartitionEditorPerProjectUserSettings.h
+++ b/Engine/Source/Runtime/Engine/Public/WorldPartition/WorldPartitionEditorPerProjectUserSettings.h
@@ -18,31 +18,27 @@ struct FWorldPartitionPerWorldSettings
 	FWorldPartitionPerWorldSettings()
 	{}
 
-	FWorldPartitionPerWorldSettings(TArray<FBox>& InLoadedEditorGridRegions)
-		: LoadedEditorGridRegions(InLoadedEditorGridRegions)
-	{}
-
 	void Reset()
 	{
-		LoadedEditorGridRegions.Empty();
+		LoadedEditorRegions.Empty();
+		LoadedEditorLocationVolumes.Empty();
 		NotLoadedDataLayers.Empty();
 		LoadedDataLayers.Empty();
-		EditorGridConfigHash = 0;
 	}
 #endif
 
 #if WITH_EDITORONLY_DATA
 	UPROPERTY()
-	TArray<FBox> LoadedEditorGridRegions;
+	TArray<FBox> LoadedEditorRegions;
 
 	UPROPERTY()
-	TArray<FName> NotLoadedDataLayers;
+	TArray<FName> LoadedEditorLocationVolumes;
 
 	UPROPERTY()
-	TArray<FName> LoadedDataLayers;
+	TArray<FName> NotLoadedDataLayers;
 
 	UPROPERTY()
-	uint32 EditorGridConfigHash = 0;
+	TArray<FName> LoadedDataLayers;
 #endif
 };
 
@@ -67,10 +63,11 @@ public:
 	{}
 
 #if WITH_EDITOR
-	void UpdateEditorGridConfigHash(UWorld* InWorld);
-	static void UpdateEditorGridConfigHash(UWorld* InWorld, FWorldPartitionPerWorldSettings& PerWorldSettings);
-	TArray<FBox> GetEditorGridLoadedRegions(UWorld* InWorld) const;
-	void SetEditorGridLoadedRegions(UWorld* InWorld, const TArray<FBox>& InEditorGridLoadedRegions);
+	TArray<FBox> GetEditorLoadedRegions(UWorld* InWorld) const;
+	void SetEditorLoadedRegions(UWorld* InWorld, const TArray<FBox>& InEditorLoadedRegions);
+
+	TArray<FName> GetEditorLoadedLocationVolumes(UWorld* InWorld) const;
+	void SetEditorLoadedLocationVolumes(UWorld* InWorld, const TArray<FName>& InEditorLoadedLocationVolumes);
 
 	bool GetEnableLoadingOfLastLoadedRegions() const
 	{
