Fix incorrect comments in UStaticMesh::PostLoad()

diff --git a/Engine/Source/Runtime/Engine/Private/StaticMesh.cpp b/Engine/Source/Runtime/Engine/Private/StaticMesh.cpp
--- a/Engine/Source/Runtime/Engine/Private/StaticMesh.cpp
+++ b/Engine/Source/Runtime/Engine/Private/StaticMesh.cpp
@@ -5545,13 +5545,13 @@ void UStaticMesh::PostLoad()
 	{
 		FModuleManager::Get().LoadModuleChecked<IMeshUtilities>(TEXT("MeshUtilities"));
  
-		// Cache derived data for the running platform.
+		// Load the mesh builder module in order to cache data for the running platform.
 		ITargetPlatformManagerModule& TargetPlatformManager = GetTargetPlatformManagerRef();
 		ITargetPlatform* RunningPlatform = TargetPlatformManager.GetRunningTargetPlatform();
 		check(RunningPlatform);
 		IMeshBuilderModule::GetForPlatform(RunningPlatform);
 
-		// Additionally cache derived data for any other platforms we care about.
+		// Additionally, load the mesh builder modules for any other platforms we care about.
 		const TArray<ITargetPlatform*>& TargetPlatforms = TargetPlatformManager.GetActiveTargetPlatforms();
 		for (ITargetPlatform* Platform : TargetPlatforms)
 		{
