Expose individual output channels for Constant2/3/4 material expressions

diff --git a/Engine/Source/Runtime/Engine/Private/Materials/MaterialExpressions.cpp b/Engine/Source/Runtime/Engine/Private/Materials/MaterialExpressions.cpp
--- a/Engine/Source/Runtime/Engine/Private/Materials/MaterialExpressions.cpp
+++ b/Engine/Source/Runtime/Engine/Private/Materials/MaterialExpressions.cpp
@@ -4968,6 +4968,11 @@ UMaterialExpressionConstant2Vector::UMaterialExpressionConstant2Vector(const FOb
 	MenuCategories.Add(ConstructorStatics.NAME_Vectors);
 
 	bCollapsed = true;
+
+	Outputs.Reset();
+	Outputs.Add(FExpressionOutput(TEXT(""), 1, 1, 1, 0, 0));
+	Outputs.Add(FExpressionOutput(TEXT(""), 1, 1, 0, 0, 0));
+	Outputs.Add(FExpressionOutput(TEXT(""), 1, 0, 1, 0, 0));
 #endif
 }
 
@@ -5013,6 +5018,12 @@ UMaterialExpressionConstant3Vector::UMaterialExpressionConstant3Vector(const FOb
 	MenuCategories.Add(ConstructorStatics.NAME_Vectors);
 
 	bCollapsed = false;
+
+	Outputs.Reset();
+	Outputs.Add(FExpressionOutput(TEXT(""), 1, 1, 1, 1, 0));
+	Outputs.Add(FExpressionOutput(TEXT(""), 1, 1, 0, 0, 0));
+	Outputs.Add(FExpressionOutput(TEXT(""), 1, 0, 1, 0, 0));
+	Outputs.Add(FExpressionOutput(TEXT(""), 1, 0, 0, 1, 0));
 #endif
 }
 
@@ -5058,6 +5069,13 @@ UMaterialExpressionConstant4Vector::UMaterialExpressionConstant4Vector(const FOb
 	MenuCategories.Add(ConstructorStatics.NAME_Vectors);
 
 	bCollapsed = false;
+
+	Outputs.Reset();
+	Outputs.Add(FExpressionOutput(TEXT(""), 1, 1, 1, 1, 1));
+	Outputs.Add(FExpressionOutput(TEXT(""), 1, 1, 0, 0, 0));
+	Outputs.Add(FExpressionOutput(TEXT(""), 1, 0, 1, 0, 0));
+	Outputs.Add(FExpressionOutput(TEXT(""), 1, 0, 0, 1, 0));
+	Outputs.Add(FExpressionOutput(TEXT(""), 1, 0, 0, 0, 1));
 #endif
 }
 
