Fixed crash caused by DataLayerSubsystem assuming it is part of a partitioned world and AWorldDataLayers is valid.

diff --git a/Engine/Source/Runtime/Engine/Public/WorldPartition/DataLayer/DataLayerSubsystem.h b/Engine/Source/Runtime/Engine/Public/WorldPartition/DataLayer/DataLayerSubsystem.h
--- a/Engine/Source/Runtime/Engine/Public/WorldPartition/DataLayer/DataLayerSubsystem.h
+++ b/Engine/Source/Runtime/Engine/Public/WorldPartition/DataLayer/DataLayerSubsystem.h
@@ -229,8 +229,12 @@ TArray<FName> UDataLayerSubsystem::GetDataLayerInstanceNames(const TArray<T>& In
 	DataLayerInstanceNames.Reserve(InDataLayerIdentifiers.Num());
 
 	PRAGMA_DISABLE_DEPRECATION_WARNINGS
-	const AWorldDataLayers* WorldDataLayers = InLevelContext ? InLevelContext->GetWorldDataLayers() : GetWorld()->GetWorldDataLayers();
-	DataLayerInstanceNames = WorldDataLayers->GetDataLayerInstanceNames(InDataLayerIdentifiers);
+	// Non-partitioned worlds don't have an AWorldDataLayers.
+	// This function can be called by a partitioned sub-level that contains Data Layers.
+	if (const AWorldDataLayers* WorldDataLayers = InLevelContext ? InLevelContext->GetWorldDataLayers() : GetWorld()->GetWorldDataLayers())
+	{
+		DataLayerInstanceNames = WorldDataLayers->GetDataLayerInstanceNames(InDataLayerIdentifiers);
+	}
 	PRAGMA_ENABLE_DEPRECATION_WARNINGS
 	return DataLayerInstanceNames;
 }
@@ -242,8 +246,12 @@ TArray<const UDataLayerInstance*> UDataLayerSubsystem::GetDataLayerInstances(con
 	DataLayerInstances.Reserve(InDataLayerIdentifiers.Num());
 
 	PRAGMA_DISABLE_DEPRECATION_WARNINGS
-	const AWorldDataLayers* WorldDataLayers = InLevelContext ? InLevelContext->GetWorldDataLayers() : GetWorld()->GetWorldDataLayers();
-	DataLayerInstances = WorldDataLayers->GetDataLayerInstances(InDataLayerIdentifiers);
+	// Non-partitioned worlds don't have an AWorldDataLayers.
+	// This function can be called by a partitioned sub-level that contains Data Layers.
+	if (const AWorldDataLayers* WorldDataLayers = InLevelContext ? InLevelContext->GetWorldDataLayers() : GetWorld()->GetWorldDataLayers())
+	{
+		DataLayerInstances = WorldDataLayers->GetDataLayerInstances(InDataLayerIdentifiers);
+	}
 	PRAGMA_ENABLE_DEPRECATION_WARNINGS
 	return DataLayerInstances;
 }
\ No newline at end of file
