Fixed case where Level Instance's WorldAsset path was being fixup for PIE. This can happen if the level is also part of the world's levels.

diff --git a/Engine/Source/Runtime/Engine/Private/LevelInstance/LevelInstanceActor.cpp b/Engine/Source/Runtime/Engine/Private/LevelInstance/LevelInstanceActor.cpp
--- a/Engine/Source/Runtime/Engine/Private/LevelInstance/LevelInstanceActor.cpp
+++ b/Engine/Source/Runtime/Engine/Private/LevelInstance/LevelInstanceActor.cpp
@@ -32,6 +32,17 @@ void ALevelInstance::Serialize(FArchive& Ar)
 {
 	Super::Serialize(Ar);
 	Ar << LevelInstanceActorGuid;
+
+#if WITH_EDITORONLY_DATA
+	if (Ar.GetPortFlags() & PPF_DuplicateForPIE)
+	{
+		// Remove PIE Prefix in case the LevelInstance was part of the DuplicateWorldForPIE
+		// This can happen if the level (WorldAsset) is part of the world's levels.
+		// ULevelStreaming::RenameForPIE will call FSoftObjectPath::AddPIEPackageName which will
+		// force this softobjectpath to be processed by FSoftObjectPath::FixupForPIE (even one that comes from a level instance).
+		WorldAsset = FSoftObjectPath(UWorld::RemovePIEPrefix(WorldAsset.ToString()));
+	}
+#endif
 }
 
 void ALevelInstance::PostRegisterAllComponents()
