Allow color adjustment for cubemaps created from longlat sources.

diff --git a/Engine/Source/Developer/TextureCompressor/Private/TextureCompressorModule.cpp b/Engine/Source/Developer/TextureCompressor/Private/TextureCompressorModule.cpp
--- a/Engine/Source/Developer/TextureCompressor/Private/TextureCompressorModule.cpp
+++ b/Engine/Source/Developer/TextureCompressor/Private/TextureCompressorModule.cpp
@@ -3506,13 +3506,6 @@ private:
 				GenerateBaseCubeMipFromLongitudeLatitude2D(Mip, Image, BuildSettings.MaxTextureResolution, BuildSettings.SourceEncodingOverride);
 	
 				check( CopyCount == 1 );
-
-				if (!BuildSettings.bUseNewMipFilter)
-				{
-					// note the break here skips other adjustments:
-					// this is maintained only for old textures
-					break;
-				}
 			}
 			else
 			{
diff --git a/Engine/Source/Runtime/Engine/Private/TextureDerivedData.cpp b/Engine/Source/Runtime/Engine/Private/TextureDerivedData.cpp
--- a/Engine/Source/Runtime/Engine/Private/TextureDerivedData.cpp
+++ b/Engine/Source/Runtime/Engine/Private/TextureDerivedData.cpp
@@ -241,7 +241,15 @@ static void SerializeForKey(FArchive& Ar, const FTextureBuildSettings& Settings)
 		TempGuid = FGuid(0x27B79A99, 0xE1A5458E, 0xAB619475, 0xCD01AD2A);
 		Ar << TempGuid;
 	}
-	
+
+	if (Settings.bLongLatSource)
+	{
+		// @@!! these can go away whenever we bump the overall ddc key
+		// texture processing for cubemaps generated from longlat sources changed, so modify the key :
+		TempGuid = FGuid(0x3D642836, 0xEBF64714, 0x9E8E3241, 0x39F66906);
+		Ar << TempGuid;
+	}
+
 	if (Settings.CompressionCacheId.IsValid())
 	{
 		TempGuid = Settings.CompressionCacheId; Ar << TempGuid;
