World Partition: deprecate WorldPartitionVolume class.

#rnx

diff --git a/Engine/Source/Runtime/Engine/Private/WorldPartition/WorldPartitionVolume.cpp b/Engine/Source/Runtime/Engine/Private/WorldPartition/WorldPartitionVolume.cpp
--- a/Engine/Source/Runtime/Engine/Private/WorldPartition/WorldPartitionVolume.cpp
+++ b/Engine/Source/Runtime/Engine/Private/WorldPartition/WorldPartitionVolume.cpp
@@ -1,47 +1,11 @@
 // Copyright Epic Games, Inc. All Rights Reserved.
 
 #include "WorldPartition/WorldPartitionVolume.h"
-#include "WorldPartition/WorldPartitionEditorHash.h"
-#include "WorldPartition/WorldPartition.h"
-#include "WorldPartition/LoaderAdapter/LoaderAdapterActor.h"
-#include "Engine/World.h"
-#include "Misc/ScopedSlowTask.h"
 
-#define LOCTEXT_NAMESPACE "WorldPartitionEditor"
-
-AWorldPartitionVolume::AWorldPartitionVolume(const FObjectInitializer& ObjectInitializer)
+ADEPRECATED_WorldPartitionVolume::ADEPRECATED_WorldPartitionVolume(const FObjectInitializer& ObjectInitializer)
 	: Super(ObjectInitializer)
 {
 #if WITH_EDITORONLY_DATA
 	bIsSpatiallyLoaded = false;
 #endif
-
-#if WITH_EDITOR
-	if (!IsTemplate() && GetWorld() && GetWorld()->GetWorldPartition())
-	{
-		WorldPartitionActorLoader = new FLoaderAdapterActor(this);
-	}
-#endif
-}
-
-void AWorldPartitionVolume::BeginDestroy()
-{
-#if WITH_EDITOR
-	if (WorldPartitionActorLoader)
-	{
-		delete WorldPartitionActorLoader;
-		WorldPartitionActorLoader = nullptr;
-	}
-#endif
-
-	Super::BeginDestroy();
-}
-
-#if WITH_EDITOR
-IWorldPartitionActorLoaderInterface::ILoaderAdapter* AWorldPartitionVolume::GetLoaderAdapter()
-{
-	return WorldPartitionActorLoader;
-}
-#endif
-
-#undef LOCTEXT_NAMESPACE
\ No newline at end of file
+}
\ No newline at end of file
diff --git a/Engine/Source/Runtime/Engine/Public/WorldPartition/WorldPartitionVolume.h b/Engine/Source/Runtime/Engine/Public/WorldPartition/WorldPartitionVolume.h
--- a/Engine/Source/Runtime/Engine/Public/WorldPartition/WorldPartitionVolume.h
+++ b/Engine/Source/Runtime/Engine/Public/WorldPartition/WorldPartitionVolume.h
@@ -8,21 +8,12 @@
 
 class FLoaderAdapterActor;
 
-/**
- * A world partition volume to allow loading cells inside (editor-only)
- */
-UCLASS()
-class ENGINE_API AWorldPartitionVolume : public AVolume, public IWorldPartitionActorLoaderInterface
+UCLASS(Deprecated, meta = (DeprecationMessage = "WorldPartitionVolume has been replaced by LocationVolume"))
+class ENGINE_API ADEPRECATED_WorldPartitionVolume : public AVolume
 {
-	GENERATED_BODY()
+	GENERATED_UCLASS_BODY()
 
 public:
-	AWorldPartitionVolume(const FObjectInitializer& ObjectInitializer = FObjectInitializer::Get());
-
-	//~ Begin UObject Interface
-	virtual void BeginDestroy() override;
-	//~ End UObject Interface
-
 	//~ Begin AActor Interface
 	virtual bool IsEditorOnly() const override { return true; }
 #if WITH_EDITOR
@@ -30,15 +21,4 @@ public:
 	virtual bool SupportsDataLayer() const override { return false; }
 #endif
 	//~ End AActor Interface
-
-#if WITH_EDITOR
-	//~ Begin IWorldPartitionActorLoaderInterface interface
-	virtual ILoaderAdapter* GetLoaderAdapter() override;
-	//~ End IWorldPartitionActorLoaderInterface interface
-#endif
-
-#if WITH_EDITOR
-private:
-	FLoaderAdapterActor* WorldPartitionActorLoader;
-#endif
 };
\ No newline at end of file
