Fixed bug where a request to unload a LevelInstance that is pending load would trigger checks and would skip the unloading.

diff --git a/Engine/Source/Runtime/Engine/Private/LevelInstance/LevelInstanceSubsystem.cpp b/Engine/Source/Runtime/Engine/Private/LevelInstance/LevelInstanceSubsystem.cpp
--- a/Engine/Source/Runtime/Engine/Private/LevelInstance/LevelInstanceSubsystem.cpp
+++ b/Engine/Source/Runtime/Engine/Private/LevelInstance/LevelInstanceSubsystem.cpp
@@ -158,7 +158,8 @@ void ULevelInstanceSubsystem::RequestLoadLevelInstance(ILevelInstanceInterface*
 void ULevelInstanceSubsystem::RequestUnloadLevelInstance(ILevelInstanceInterface* LevelInstance)
 {
 	const FLevelInstanceID& LevelInstanceID = LevelInstance->GetLevelInstanceID();
-	if (LevelInstances.Contains(LevelInstanceID))
+	// Test whether level instance is loaded or is still loading
+	if (LoadedLevelInstances.Contains(LevelInstanceID) || LoadingLevelInstances.Contains(LevelInstanceID))
 	{
 		// LevelInstancesToUnload uses FLevelInstanceID because LevelInstance* can be destroyed in later Tick and we don't need it.
 		LevelInstancesToUnload.Add(LevelInstanceID);
@@ -168,7 +169,7 @@ void ULevelInstanceSubsystem::RequestUnloadLevelInstance(ILevelInstanceInterface
 
 bool ULevelInstanceSubsystem::IsLoaded(const ILevelInstanceInterface* LevelInstance) const
 {
-	return LevelInstance->HasValidLevelInstanceID() && LevelInstances.Contains(LevelInstance->GetLevelInstanceID());
+	return LevelInstance->HasValidLevelInstanceID() && LoadedLevelInstances.Contains(LevelInstance->GetLevelInstanceID());
 }
 
 void ULevelInstanceSubsystem::UpdateStreamingState()
@@ -200,7 +201,14 @@ void ULevelInstanceSubsystem::UpdateStreamingState()
 #if WITH_EDITOR
 			SlowTask.EnterProgressFrame(1.f, LOCTEXT("UnloadingLevelInstance", "Unloading Level Instance"));
 #endif
-			UnloadLevelInstance(LevelInstanceID);
+			if (LoadingLevelInstances.Contains(LevelInstanceID))
+			{
+				LevelInstancesToUnload.Add(LevelInstanceID);
+			}
+			else
+			{
+				UnloadLevelInstance(LevelInstanceID);
+			}
 		}
 	}
 
@@ -244,13 +252,25 @@ void ULevelInstanceSubsystem::UpdateStreamingState()
 
 void ULevelInstanceSubsystem::RegisterLoadedLevelStreamingLevelInstance(ULevelStreamingLevelInstance* LevelStreaming)
 {
-	ILevelInstanceInterface* LevelInstance = LevelStreaming->GetLevelInstance();
-	const FLevelInstanceID LevelInstanceID = LevelInstance->GetLevelInstanceID();
-	check(!LevelInstances.Contains(LevelInstanceID));
-	FLevelInstance& LevelInstanceEntry = LevelInstances.Add(LevelInstanceID);
+	const FLevelInstanceID LevelInstanceID = LevelStreaming->GetLevelInstanceID();
+	check(LoadingLevelInstances.Contains(LevelInstanceID));
+	LoadingLevelInstances.Remove(LevelInstanceID);
+	check(!LoadedLevelInstances.Contains(LevelInstanceID));
+	FLevelInstance& LevelInstanceEntry = LoadedLevelInstances.Add(LevelInstanceID);
 	LevelInstanceEntry.LevelStreaming = LevelStreaming;
 #if WITH_EDITOR
-	LevelInstance->OnLevelInstanceLoaded();
+	// LevelInstanceID might not be registered anymore in the case where the level instance 
+	// was unloaded while still being load.
+	if (ILevelInstanceInterface* LevelInstance = LevelStreaming->GetLevelInstance())
+	{
+		check(LevelInstance->GetLevelInstanceID() == LevelInstanceID);
+		LevelInstance->OnLevelInstanceLoaded();
+	}
+	else
+	{
+		// Validate that the LevelInstanceID is requested to be unloaded
+		check(LevelInstancesToUnload.Contains(LevelInstanceID));
+	}
 #endif
 }
 
@@ -304,53 +324,68 @@ void ULevelInstanceSubsystem::LoadLevelInstance(ILevelInstanceInterface* LevelIn
 	}
 
 	const FLevelInstanceID& LevelInstanceID = LevelInstance->GetLevelInstanceID();
-	check(!LevelInstances.Contains(LevelInstanceID));
+	check(!LoadedLevelInstances.Contains(LevelInstanceID));
+	check(!LoadingLevelInstances.Contains(LevelInstanceID));
+	LoadingLevelInstances.Add(LevelInstanceID);
 
 	if (ULevelStreamingLevelInstance* LevelStreaming = ULevelStreamingLevelInstance::LoadInstance(LevelInstance))
 	{
 #if WITH_EDITOR
-		check(LevelInstanceActor->GetWorld()->IsGameWorld() || LevelInstances.Contains(LevelInstanceID));
+		check(LevelInstanceActor->GetWorld()->IsGameWorld() || LoadedLevelInstances.Contains(LevelInstanceID));
 #endif
 	}
+	else
+	{
+		LoadingLevelInstances.Remove(LevelInstanceID);
+	}
 }
 
 void ULevelInstanceSubsystem::UnloadLevelInstance(const FLevelInstanceID& LevelInstanceID)
 {
-#if WITH_EDITOR
-	// Create scope if it doesn't exist
-	bool bReleaseScope = false;
-	if (!LevelsToRemoveScope)
+	if (GetWorld()->IsGameWorld())
 	{
-		bReleaseScope = true;
-		LevelsToRemoveScope.Reset(new FLevelsToRemoveScope(this));
+		FLevelInstance LevelInstance;
+		if (LoadedLevelInstances.RemoveAndCopyValue(LevelInstanceID, LevelInstance))
+		{
+			ULevelStreamingLevelInstance::UnloadInstance(LevelInstance.LevelStreaming);
+		}
 	}
-#endif
-				
-	FLevelInstance LevelInstance;
-	if (LevelInstances.RemoveAndCopyValue(LevelInstanceID, LevelInstance))
+#if WITH_EDITOR
+	else
 	{
-		if (ULevel* LoadedLevel = LevelInstance.LevelStreaming->GetLoadedLevel())
+		// Create scope if it doesn't exist
+		bool bReleaseScope = false;
+		if (!LevelsToRemoveScope)
 		{
-			ForEachActorInLevel(LoadedLevel, [this](AActor* LevelActor)
-			{
-				if (ILevelInstanceInterface* LevelInstance = Cast<ILevelInstanceInterface>(LevelActor))
-				{
-					// Make sure to remove from pending loads if we are unloading child can't be loaded
-					LevelInstancesToLoadOrUpdate.Remove(LevelInstance);
-					
-					UnloadLevelInstance(LevelInstance->GetLevelInstanceID());
-				}
-				return true;
-			});
+			bReleaseScope = true;
+			LevelsToRemoveScope.Reset(new FLevelsToRemoveScope(this));
 		}
 
-		ULevelStreamingLevelInstance::UnloadInstance(LevelInstance.LevelStreaming);
-	}
+		FLevelInstance LevelInstance;
+		if (LoadedLevelInstances.RemoveAndCopyValue(LevelInstanceID, LevelInstance))
+		{
+			if (ULevel* LoadedLevel = LevelInstance.LevelStreaming->GetLoadedLevel())
+			{
+				ForEachActorInLevel(LoadedLevel, [this](AActor* LevelActor)
+					{
+						if (ILevelInstanceInterface* LevelInstance = Cast<ILevelInstanceInterface>(LevelActor))
+						{
+							// Make sure to remove from pending loads if we are unloading child can't be loaded
+							LevelInstancesToLoadOrUpdate.Remove(LevelInstance);
 
-#if WITH_EDITOR
-	if (bReleaseScope)
-	{
-		LevelsToRemoveScope.Reset();
+							UnloadLevelInstance(LevelInstance->GetLevelInstanceID());
+						}
+						return true;
+					});
+			}
+
+			ULevelStreamingLevelInstance::UnloadInstance(LevelInstance.LevelStreaming);
+		}
+
+		if (bReleaseScope)
+		{
+			LevelsToRemoveScope.Reset();
+		}
 	}
 #endif
 }
@@ -572,7 +607,7 @@ bool ULevelInstanceSubsystem::GetLevelInstanceBounds(const ILevelInstanceInterfa
 {
 	if (IsLoaded(LevelInstance))
 	{
-		const FLevelInstance& LevelInstanceEntry = LevelInstances.FindChecked(LevelInstance->GetLevelInstanceID());
+		const FLevelInstance& LevelInstanceEntry = LoadedLevelInstances.FindChecked(LevelInstance->GetLevelInstanceID());
 		OutBounds = LevelInstanceEntry.LevelStreaming->GetBounds();
 		return true;
 	}
@@ -1226,7 +1261,7 @@ ULevel* ULevelInstanceSubsystem::GetLevelInstanceLevel(const ILevelInstanceInter
 		{
 			return LevelInstanceEdit->LevelStreaming->GetLoadedLevel();
 		}
-		else if (const FLevelInstance* LevelInstanceEntry = LevelInstances.Find(LevelInstance->GetLevelInstanceID()))
+		else if (const FLevelInstance* LevelInstanceEntry = LoadedLevelInstances.Find(LevelInstance->GetLevelInstanceID()))
 		{
 			return LevelInstanceEntry->LevelStreaming->GetLoadedLevel();
 		}
diff --git a/Engine/Source/Runtime/Engine/Public/LevelInstance/LevelInstanceLevelStreaming.h b/Engine/Source/Runtime/Engine/Public/LevelInstance/LevelInstanceLevelStreaming.h
--- a/Engine/Source/Runtime/Engine/Public/LevelInstance/LevelInstanceLevelStreaming.h
+++ b/Engine/Source/Runtime/Engine/Public/LevelInstance/LevelInstanceLevelStreaming.h
@@ -34,6 +34,7 @@ protected:
 
 	friend class ULevelInstanceSubsystem;
 
+	const FLevelInstanceID& GetLevelInstanceID() const { return LevelInstanceID; }
 private:
 #if WITH_EDITOR
 	void PrepareLevelInstanceLoadedActor(AActor& InActor, ILevelInstanceInterface* InLevelInstance);
diff --git a/Engine/Source/Runtime/Engine/Public/LevelInstance/LevelInstanceSubsystem.h b/Engine/Source/Runtime/Engine/Public/LevelInstance/LevelInstanceSubsystem.h
--- a/Engine/Source/Runtime/Engine/Public/LevelInstance/LevelInstanceSubsystem.h
+++ b/Engine/Source/Runtime/Engine/Public/LevelInstance/LevelInstanceSubsystem.h
@@ -188,7 +188,8 @@ private:
 
 	TMap<ILevelInstanceInterface*, bool> LevelInstancesToLoadOrUpdate;
 	TSet<FLevelInstanceID> LevelInstancesToUnload;
-	TMap<FLevelInstanceID, FLevelInstance> LevelInstances;
+	TSet<FLevelInstanceID> LoadingLevelInstances;
+	TMap<FLevelInstanceID, FLevelInstance> LoadedLevelInstances;
 	TMap<FLevelInstanceID, ILevelInstanceInterface*> RegisteredLevelInstances;
 
 #if WITH_EDITORONLY_DATA
