Add bUseInverseSquaredFalloff setter to UPointLightComponent for blueprint usage.

diff --git a/Engine/Source/Runtime/Engine/Classes/Components/PointLightComponent.h b/Engine/Source/Runtime/Engine/Classes/Components/PointLightComponent.h
--- a/Engine/Source/Runtime/Engine/Classes/Components/PointLightComponent.h
+++ b/Engine/Source/Runtime/Engine/Classes/Components/PointLightComponent.h
@@ -57,6 +57,9 @@ class ENGINE_API UPointLightComponent : public ULocalLightComponent
 	UPROPERTY(EditAnywhere, BlueprintReadOnly, interp, Category=Light)
 	float SourceLength;
 
+	UFUNCTION(BlueprintCallable, Category = "Rendering|Lighting")
+	void SetUseInverseSquaredFalloff(bool bNewValue);
+
 	UFUNCTION(BlueprintCallable, Category="Rendering|Lighting")
 	void SetLightFalloffExponent(float NewLightFalloffExponent);
 
diff --git a/Engine/Source/Runtime/Engine/Private/Components/PointLightComponent.cpp b/Engine/Source/Runtime/Engine/Private/Components/PointLightComponent.cpp
--- a/Engine/Source/Runtime/Engine/Private/Components/PointLightComponent.cpp
+++ b/Engine/Source/Runtime/Engine/Private/Components/PointLightComponent.cpp
@@ -123,6 +123,16 @@ FLightSceneProxy* UPointLightComponent::CreateSceneProxy() const
 	return nullptr;
 }
 
+void UPointLightComponent::SetUseInverseSquaredFalloff(bool bNewValue)
+{
+	if (AreDynamicDataChangesAllowed()
+		&& bUseInverseSquaredFalloff != bNewValue)
+	{
+		bUseInverseSquaredFalloff = bNewValue;
+		MarkRenderStateDirty();
+	}
+}
+
 void UPointLightComponent::SetLightFalloffExponent(float NewLightFalloffExponent)
 {
 	if (AreDynamicDataChangesAllowed()
