World Partition: fix crash when creating a Level Instance.

#rnx

diff --git a/Engine/Source/Editor/UnrealEd/Private/FileHelpers.cpp b/Engine/Source/Editor/UnrealEd/Private/FileHelpers.cpp
--- a/Engine/Source/Editor/UnrealEd/Private/FileHelpers.cpp
+++ b/Engine/Source/Editor/UnrealEd/Private/FileHelpers.cpp
@@ -901,7 +901,7 @@ static bool SaveWorld(UWorld* World,
 			{
 				// Save Snapshot of loaded Editor regions
 				GetMutableDefault<UWorldPartitionEditorPerProjectUserSettings>()->SetEditorLoadedRegions(SaveWorld, LoadedEditorRegions);
-				RenamedWorldPartition->LoadLastLoadedRegions();
+				RenamedWorldPartition->LoadLastLoadedRegions(LoadedEditorRegions);
 			}
 		}
 
diff --git a/Engine/Source/Runtime/Engine/Private/WorldPartition/WorldPartition.cpp b/Engine/Source/Runtime/Engine/Private/WorldPartition/WorldPartition.cpp
--- a/Engine/Source/Runtime/Engine/Private/WorldPartition/WorldPartition.cpp
+++ b/Engine/Source/Runtime/Engine/Private/WorldPartition/WorldPartition.cpp
@@ -387,7 +387,7 @@ void UWorldPartition::Initialize(UWorld* InWorld, const FTransform& InTransform)
 
 		EditorHash->Initialize();
 
-		AlwaysLoadedActors = new FLoaderAdapterAlwaysLoadedActors(World);
+		AlwaysLoadedActors = new FLoaderAdapterAlwaysLoadedActors(OuterWorld);
 	}
 
 	check(RuntimeHash);
@@ -1310,19 +1310,22 @@ bool UWorldPartition::IsActorPinned(const FGuid& ActorGuid) const
 	return PinnedActors.Contains(ActorGuid);
 }
 
-void UWorldPartition::LoadLastLoadedRegions()
+void UWorldPartition::LoadLastLoadedRegions(const TArray<FBox>& EditorLastLoadedRegions)
 {
-	TArray<FBox> EditorLastLoadedRegions = GetMutableDefault<UWorldPartitionEditorPerProjectUserSettings>()->GetEditorLoadedRegions(World);
-
 	for (const FBox& EditorLastLoadedRegion : EditorLastLoadedRegions)
 	{
 		FLoaderAdapterShape* LoaderAdapter = CreateEditorLoaderAdapter<FLoaderAdapterShape>(World, EditorLastLoadedRegion, TEXT("Last Loaded Region"));
 		LoaderAdapter->SetUserCreated(true);
 		LoaderAdapter->Load();
 	}
+}
 
-	TArray<FName> EditorLoadedLocationVolumes = GetMutableDefault<UWorldPartitionEditorPerProjectUserSettings>()->GetEditorLoadedLocationVolumes(World);
+void UWorldPartition::LoadLastLoadedRegions()
+{
+	TArray<FBox> EditorLastLoadedRegions = GetMutableDefault<UWorldPartitionEditorPerProjectUserSettings>()->GetEditorLoadedRegions(World);
+	LoadLastLoadedRegions(EditorLastLoadedRegions);
 
+	TArray<FName> EditorLoadedLocationVolumes = GetMutableDefault<UWorldPartitionEditorPerProjectUserSettings>()->GetEditorLoadedLocationVolumes(World);
 	for (const FName& EditorLoadedLocationVolume : EditorLoadedLocationVolumes)
 	{
 		if (ALocationVolume* LocationVolume = FindObject<ALocationVolume>(World->PersistentLevel, *EditorLoadedLocationVolume.ToString()))
diff --git a/Engine/Source/Runtime/Engine/Public/WorldPartition/WorldPartition.h b/Engine/Source/Runtime/Engine/Public/WorldPartition/WorldPartition.h
--- a/Engine/Source/Runtime/Engine/Public/WorldPartition/WorldPartition.h
+++ b/Engine/Source/Runtime/Engine/Public/WorldPartition/WorldPartition.h
@@ -170,6 +170,7 @@ public:
 	void UnpinActor(const FGuid& ActorGuid);
 	bool IsActorPinned(const FGuid& ActorGuid) const;
 
+	void LoadLastLoadedRegions(const TArray<FBox>& EditorLastLoadedRegions);
 	void LoadLastLoadedRegions();
 #endif
 
