WorldPartition: Avoid broadcasting GEngine->BroadcastActorFoldersUpdated for nothing as it impacts performance.
The typical case was during PIE while streaming cells, each newly created runtime levels would trigger this event and update the sceneoutliner for nothing.

diff --git a/Engine/Source/Runtime/Engine/Private/Level.cpp b/Engine/Source/Runtime/Engine/Private/Level.cpp
--- a/Engine/Source/Runtime/Engine/Private/Level.cpp
+++ b/Engine/Source/Runtime/Engine/Private/Level.cpp
@@ -3229,6 +3229,8 @@ void ULevel::CreateOrUpdateActorFolders()
 		return;
 	}
 
+	const bool bHadActorFolders = !ActorFolders.IsEmpty();
+
 	// Here we also do a cleanup of all actor folders marked as deleted.
 	// Find and fixup actors and folders that are referencing them.
 	// Also update actor folders packaging mode.
@@ -3266,7 +3268,12 @@ void ULevel::CreateOrUpdateActorFolders()
 		ActorFolders.Remove(ActorFolderToDelete->GetGuid());
 	}
 
-	GEngine->BroadcastActorFoldersUpdated(this);
+	// Avoid broadcasting if no actor folder were/are part of this level
+	const bool bHasActorFolders = !ActorFolders.IsEmpty();
+	if (bHadActorFolders || bHasActorFolders)
+	{
+		GEngine->BroadcastActorFoldersUpdated(this);
+	}
 }
 
 void ULevel::ForEachActorFolder(TFunctionRef<bool(UActorFolder*)> Operation, bool bSkipDeleted)
