Reapply changes that enable hotfixing of device profile.

diff --git a/Engine/Source/Runtime/Engine/Classes/DeviceProfiles/DeviceProfileManager.h b/Engine/Source/Runtime/Engine/Classes/DeviceProfiles/DeviceProfileManager.h
--- a/Engine/Source/Runtime/Engine/Classes/DeviceProfiles/DeviceProfileManager.h
+++ b/Engine/Source/Runtime/Engine/Classes/DeviceProfiles/DeviceProfileManager.h
@@ -48,13 +48,13 @@ public:
 	 * Startup and select the active device profile
 	 * Then Init the CVars from this profile and it's Device profile parent tree.
 	 */
-	static void InitializeCVarsForActiveDeviceProfile();
+	static void InitializeCVarsForActiveDeviceProfile(bool bPushSettings = false, bool bIsDeviceProfilePreview = false, bool bForceReload = false);
 
 	/**
 	 * Reapplies the device profile. Useful when configs have changed (i.e. hotfix)
 	 * Applies base and then any overridden device profile.
 	 */
-	void ReapplyDeviceProfile();
+	void ReapplyDeviceProfile(bool bForceReload = false);
 
 	/**
 	 * Examine the currently active or overridden profile for references to entries in DeviceProfilesToQuery
diff --git a/Engine/Source/Runtime/Engine/Private/DeviceProfiles/DeviceProfileManager.cpp b/Engine/Source/Runtime/Engine/Private/DeviceProfiles/DeviceProfileManager.cpp
--- a/Engine/Source/Runtime/Engine/Private/DeviceProfiles/DeviceProfileManager.cpp
+++ b/Engine/Source/Runtime/Engine/Private/DeviceProfiles/DeviceProfileManager.cpp
@@ -491,13 +491,17 @@ void UDeviceProfileManager::ChangeTaggedFragmentState(FName FragmentTag, bool bN
 	}
 }
 
-void UDeviceProfileManager::InitializeCVarsForActiveDeviceProfile()
+void UDeviceProfileManager::InitializeCVarsForActiveDeviceProfile(bool bPushSettings, bool bIsDeviceProfilePreview, bool bForceReload)
 {
 	FString ActiveProfileName;
 
 	if (DeviceProfileManagerSingleton)
 	{
 		ActiveProfileName = DeviceProfileManagerSingleton->ActiveDeviceProfile->GetName();
+		//Ensure we've loaded the device profiles for the active platform.
+		//This can be needed when overriding the device profile.
+		FConfigCacheIni::LoadGlobalIniFile(GDeviceProfilesIni, TEXT("DeviceProfiles"), NULL, bForceReload);
+
 	}
 	else
 	{
@@ -553,7 +557,7 @@ bool UDeviceProfileManager::DoActiveProfilesReference(const TSet<FString>& Devic
 	return bDoActiveProfilesReferenceQuerySet;
 }
 
-void UDeviceProfileManager::ReapplyDeviceProfile()
+void UDeviceProfileManager::ReapplyDeviceProfile(bool bForceReload)
 {	
 	UDeviceProfile* OverrideProfile = DeviceProfileManagerSingleton->BaseDeviceProfile ? DeviceProfileManagerSingleton->GetActiveProfile() : nullptr;
 	UDeviceProfile* BaseProfile = DeviceProfileManagerSingleton->BaseDeviceProfile ? DeviceProfileManagerSingleton->BaseDeviceProfile : DeviceProfileManagerSingleton->GetActiveProfile();
@@ -574,7 +578,7 @@ void UDeviceProfileManager::ReapplyDeviceProfile()
 		RestorePushedState(PushedSettings);
 
 		// Apply the active DP. 
-		InitializeCVarsForActiveDeviceProfile();
+		InitializeCVarsForActiveDeviceProfile(false, false, bForceReload);
 
 		// broadcast cvar sinks now that we are done
 		IConsoleManager::Get().CallAllConsoleVariableSinks();
