Fix for texture arrays with mip chains and streaming. This is just modifying the texture array code to be more like the base texture code, which required adding an access point to texture arrays to get the mip data.

diff --git a/Engine/Source/Runtime/Engine/Classes/Engine/Texture2DArray.cpp b/Engine/Source/Runtime/Engine/Classes/Engine/Texture2DArray.cpp
--- a/Engine/Source/Runtime/Engine/Classes/Engine/Texture2DArray.cpp
+++ b/Engine/Source/Runtime/Engine/Classes/Engine/Texture2DArray.cpp
@@ -5,6 +5,8 @@
 =============================================================================*/
 
 #include "Engine/Texture2DArray.h"
+
+#include "AsyncCompilationHelpers.h"
 #include "RenderUtils.h"
 #include "TextureResource.h"
 #include "EngineUtils.h"
@@ -16,6 +18,9 @@
 #include "Streaming/TextureStreamIn.h"
 #include "Streaming/TextureStreamOut.h"
 #include "Engine/TextureMipDataProviderFactory.h"
+#include "Misc/ScopedSlowTask.h"
+
+#define LOCTEXT_NAMESPACE "UTexture2DArray"
 
 // Master switch to control whether streaming is enabled for texture2darray. 
 bool GSupportsTexture2DArrayStreaming = true;
@@ -67,6 +72,30 @@ FTextureResource* UTexture2DArray::CreateResource()
 	return nullptr;
 }
 
+
+// Any direct access to GetPlatformDataOrWait() will stall until the structure
+// is safe to use. It is advisable to replace those use case with
+// async aware code to avoid stalls where possible.
+FTexturePlatformData* UTexture2DArray::GetPlatformDataOrWait()
+{
+#if WITH_EDITOR
+	if (PlatformData && !PlatformData->IsAsyncWorkComplete())
+	{
+		TRACE_CPUPROFILER_EVENT_SCOPE(UTexture2D::GetPlatformDataStall);
+		UE_LOG(LogTexture, Log, TEXT("Call to GetPlatformDataOrWait() is forcing a wait on data that is not yet ready."));
+
+		FText Msg = FText::Format(LOCTEXT("WaitOnTextureCompilation", "Waiting on texture array compilation {0} ..."), FText::FromString(GetName()));
+		FScopedSlowTask Progress(1.f, Msg, true);
+		Progress.MakeDialog(true);
+		uint64 StartTime = FPlatformTime::Cycles64();
+		PlatformData->FinishCache();
+		AsyncCompilationHelpers::SaveStallStack(FPlatformTime::Cycles64() - StartTime);
+	}
+#endif
+
+	return PlatformData;
+}
+
 void UTexture2DArray::UpdateResource()
 {
 #if WITH_EDITOR
@@ -409,3 +438,5 @@ bool UTexture2DArray::StreamIn(int32 NewMipCount, bool bHighPrio)
 	}
 	return false;
 }
+
+#undef LOCTEXT_NAMESPACE
\ No newline at end of file
diff --git a/Engine/Source/Runtime/Engine/Classes/Engine/Texture2DArray.h b/Engine/Source/Runtime/Engine/Classes/Engine/Texture2DArray.h
--- a/Engine/Source/Runtime/Engine/Classes/Engine/Texture2DArray.h
+++ b/Engine/Source/Runtime/Engine/Classes/Engine/Texture2DArray.h
@@ -47,6 +47,21 @@ public:
 		return PlatformData ? PlatformData->PixelFormat : PF_Unknown;
 	}
 
+	/**
+	*  Makes a copy of the mip chain for the texture array starting with InFirstMipToLoad.
+	* 
+	*  In the editor this can potentially wait for texture builds to complete.
+	*/
+	bool GetMipData(int32 InFirstMipToLoad, TArray<FUniqueBuffer>& OutMipData);
+
+	/** 
+	*  Returns the platform data, potentially waiting on texture build if not ready.
+	*  Generally waits should be avoided by structuring callers to expect async behavior.
+	* 
+	*  Waits only occur in the editor (as builds only occur in the editor).
+	*/
+	FTexturePlatformData* GetPlatformDataOrWait();
+
 	//~ Begin UTexture Interface
 	virtual ETextureClass GetTextureClass() const override { return ETextureClass::Array; }
 	virtual void Serialize(FArchive& Ar) override;
diff --git a/Engine/Source/Runtime/Engine/Private/Rendering/Texture2DArrayResource.cpp b/Engine/Source/Runtime/Engine/Private/Rendering/Texture2DArrayResource.cpp
--- a/Engine/Source/Runtime/Engine/Private/Rendering/Texture2DArrayResource.cpp
+++ b/Engine/Source/Runtime/Engine/Private/Rendering/Texture2DArrayResource.cpp
@@ -20,49 +20,32 @@ FTexture2DArrayResource::FTexture2DArrayResource(UTexture2DArray* InOwner, const
 	AddressU = InOwner->AddressX == TA_Wrap ? AM_Wrap : (InOwner->AddressX == TA_Clamp ? AM_Clamp : AM_Mirror);
 	AddressV = InOwner->AddressY == TA_Wrap ? AM_Wrap : (InOwner->AddressY == TA_Clamp ? AM_Clamp : AM_Mirror);
 	AddressW = InOwner->AddressZ == TA_Wrap ? AM_Wrap : (InOwner->AddressZ == TA_Clamp ? AM_Clamp : AM_Mirror);
-
-	const int32 MaxLoadableMipIndex = State.MaxNumLODs - FMath::Max(1, PlatformData->GetNumMipsInTail()) + 1;
-
-	TArray<uint64> MipOffsets;
-	MipOffsets.AddZeroed(State.MaxNumLODs);
-
-	uint64 InitialMipDataSize = 0;
-	TArrayView<const FTexture2DMipMap*> MipsView = GetPlatformMipsView();
-	for (int32 MipIdx = State.RequestedFirstLODIdx(); MipIdx < MaxLoadableMipIndex; ++MipIdx)
+	
+	if (InOwner->GetMipData(State.RequestedFirstLODIdx(), MipData) == false)
 	{
-		const FTexture2DMipMap& Mip = *MipsView[MipIdx];
-		MipOffsets[MipIdx] = InitialMipDataSize;
-		InitialMipDataSize += Mip.BulkData.GetBulkDataSize();
+		// This is fatal as we will crash trying to upload the data below, this way we crash at the cause.
+		UE_LOG(LogTexture, Fatal, TEXT("Corrupt texture [%s]! Unable to load mips (bulk data missing)"), *TextureName.ToString());
+		return;
 	}
-	InitialMipData.Reset(new uint8[InitialMipDataSize]);
 
 	// Resize each structure correctly per slice, since we access it per mip after.
 	SliceMipDataViews.AddDefaulted(SizeZ);
 	for (uint32 SliceIdx = 0; SliceIdx < SizeZ; ++SliceIdx)
 	{
-		SliceMipDataViews[SliceIdx].AddDefaulted(MaxLoadableMipIndex);
+		SliceMipDataViews[SliceIdx].AddDefaulted(State.MaxNumLODs);
 	}
 
-	for (int32 MipIdx = State.RequestedFirstLODIdx(); MipIdx < MaxLoadableMipIndex; ++MipIdx)
+	const int32 RequestedFirstLODIdx = State.RequestedFirstLODIdx();
+	for (int32 RHIMipIdx = 0; RHIMipIdx < State.NumRequestedLODs; ++RHIMipIdx)
 	{
-		FTexture2DMipMap& Mip = const_cast<FTexture2DMipMap&>(*MipsView[MipIdx]);
-		if (Mip.BulkData.GetBulkDataSize() > 0)
-		{
-			const uint64 SliceMipDataSize = Mip.BulkData.GetBulkDataSize() / SizeZ;
-			const uint8* SrcData = (const uint8*)Mip.BulkData.Lock(LOCK_READ_ONLY);
-
-			for (uint32 SliceIdx = 0; SliceIdx < SizeZ; ++SliceIdx)
-			{
-				TArrayView<uint8>& SliceMipDataView = SliceMipDataViews[SliceIdx][MipIdx];
-				SliceMipDataView = TArrayView<uint8>(InitialMipData.Get() + MipOffsets[MipIdx] + SliceIdx * SliceMipDataSize, SliceMipDataSize);
-				FMemory::Memcpy(SliceMipDataView.GetData(), SrcData + SliceIdx * SliceMipDataSize, SliceMipDataSize);
-			}
+		const int32 MipIdx = RHIMipIdx + RequestedFirstLODIdx;
+		const uint64 SliceMipDataSize = MipData[RHIMipIdx].GetSize() / SizeZ;
+		uint8* SrcData = (uint8*)MipData[RHIMipIdx].GetData();
 
-			Mip.BulkData.Unlock();
-		}
-		else
+		for (uint32 SliceIdx = 0; SliceIdx < SizeZ; ++SliceIdx)
 		{
-			UE_LOG(LogTexture, Error, TEXT("Corrupt texture [%s]! Missing bulk data for MipIndex=%d"), *TextureName.ToString(), MipIdx);
+			TArrayView<uint8>& SliceMipDataView = SliceMipDataViews[SliceIdx][MipIdx];
+			SliceMipDataView = TArrayView<uint8>(SrcData + SliceIdx * SliceMipDataSize, SliceMipDataSize);
 		}
 	}
 }
@@ -99,7 +82,7 @@ void FTexture2DArrayResource::CreateTexture()
 	}
 		
 	SliceMipDataViews.Empty();
-	InitialMipData.Reset();
+	MipData.Empty();
 }
 
 void FTexture2DArrayResource::CreatePartiallyResidentTexture()
diff --git a/Engine/Source/Runtime/Engine/Private/Rendering/Texture2DArrayResource.h b/Engine/Source/Runtime/Engine/Private/Rendering/Texture2DArrayResource.h
--- a/Engine/Source/Runtime/Engine/Private/Rendering/Texture2DArrayResource.h
+++ b/Engine/Source/Runtime/Engine/Private/Rendering/Texture2DArrayResource.h
@@ -9,6 +9,7 @@
 #include "CoreMinimal.h"
 #include "Rendering/StreamableTextureResource.h"
 #include "Containers/ResourceArray.h"
+#include "Memory/SharedBuffer.h"
 
 class UTexture2DArray;
 
@@ -33,10 +34,9 @@ protected:
 
 	void GetData(uint32 SliceIndex, uint32 MipIndex, void* Dest, uint32 DestPitch);
 
-	/** The initial data for all mips of a single slice. */
+	/** The initial data for all mips of a single slice. This is all discarded once the texture has been uploaded. */
 	typedef TArray<TArrayView<uint8>, TInlineAllocator<MAX_TEXTURE_MIP_COUNT> > FSingleSliceMipDataView;
 	/** All slices initial data. */
 	TArray<FSingleSliceMipDataView> SliceMipDataViews;
-	/** The single allocation holding initial data. */
-	TUniquePtr<uint8[]> InitialMipData;
+	TArray<FUniqueBuffer> MipData;
 };
diff --git a/Engine/Source/Runtime/Engine/Private/TextureDerivedData.cpp b/Engine/Source/Runtime/Engine/Private/TextureDerivedData.cpp
--- a/Engine/Source/Runtime/Engine/Private/TextureDerivedData.cpp
+++ b/Engine/Source/Runtime/Engine/Private/TextureDerivedData.cpp
@@ -1347,9 +1347,8 @@ typedef TArray<uint32, TInlineAllocator<MAX_TEXTURE_MIP_COUNT> > FAsyncMipHandle
 typedef TArray<uint32> FAsyncVTChunkHandles;
 
 /**
- * Executes async DDC gets for mips stored in the derived data cache.
- * @param Mip - Mips to retrieve.
- * @param FirstMipToLoad - Index of the first mip to retrieve.
+ * Executes DDC gets for mips stored in the derived data cache.
+ * @param FirstMipToLoad - Index of the first mip to retrieve - starting with this mip and all smaller mips.
  * @param Callback - Callback invoked for each mip as it loads.
  * 
  * This function must be called after the initial DDC fetch is complete,
@@ -1806,6 +1805,8 @@ TFuture<TTuple<uint64, uint64>> FTexturePlatformData::LaunchEstimateOnDiskSizeTa
 	return ResultFuture;
 }
 
+// Fills the mip's bulk data from the DDC starting with FirstMipToLoad.
+// "Inline" meaning "Load in to bulk data".
 bool FTexturePlatformData::TryInlineMipData(int32 FirstMipToLoad, FStringView DebugContext)
 {
 	TRACE_CPUPROFILER_EVENT_SCOPE(FTexturePlatformData::TryInlineMipData);
@@ -2566,6 +2567,42 @@ void FTexturePlatformData::SerializeCooked(FArchive& Ar, UTexture* Owner, bool b
 	Texture derived data interface.
 ------------------------------------------------------------------------------*/
 
+bool UTexture2DArray::GetMipData(int32 InFirstMipToLoad, TArray<FUniqueBuffer>& OutMipData)
+{
+	FTexturePlatformData* LocalPlatformData = GetPlatformDataOrWait();
+	int32 OutputMipCount = LocalPlatformData->Mips.Num() - InFirstMipToLoad;
+
+	check(OutputMipCount <= MAX_TEXTURE_MIP_COUNT);
+
+	void* MipData[MAX_TEXTURE_MIP_COUNT] = {};
+	int64 MipSizes[MAX_TEXTURE_MIP_COUNT];
+	if (LocalPlatformData->TryLoadMipsWithSizes(InFirstMipToLoad, MipData, MipSizes, GetPathName()) == false)
+	{
+		// Unable to load mips from the cache. Rebuild the texture and try again.
+		UE_LOG(LogTexture, Warning, TEXT("GetMipData failed for %s (%s)"),
+			*GetPathName(), GPixelFormats[GetPixelFormat()].Name);
+#if WITH_EDITOR
+		if (!GetOutermost()->bIsCookedForEditor)
+		{
+			ForceRebuildPlatformData();
+			if (LocalPlatformData->TryLoadMipsWithSizes(InFirstMipToLoad, MipData, MipSizes, GetPathName()) == false)
+			{
+				UE_LOG(LogTexture, Error, TEXT("TryLoadMipsWithSizes still failed after ForceRebuildPlatformData %s."), *GetPathName());
+				return false;
+			}
+		}
+#else // #if WITH_EDITOR
+		return false;
+#endif // WITH_EDITOR
+	}
+
+	for (int32 MipIndex = 0; MipIndex < OutputMipCount; MipIndex++)
+	{
+		OutMipData.Emplace(FUniqueBuffer::TakeOwnership(MipData[MipIndex], MipSizes[MipIndex], [](void* Ptr) {FMemory::Free(Ptr);} ));
+	}
+	return true;
+}
+
 void UTexture2D::GetMipData(int32 FirstMipToLoad, void** OutMipData)
 {
 	if (GetPlatformData()->TryLoadMips(FirstMipToLoad, OutMipData, GetPathName()) == false)
