Fix static analysis

diff --git a/Engine/Source/Runtime/Engine/Private/PhysicsEngine/AsyncPhysicsInputComponent.cpp b/Engine/Source/Runtime/Engine/Private/PhysicsEngine/AsyncPhysicsInputComponent.cpp
--- a/Engine/Source/Runtime/Engine/Private/PhysicsEngine/AsyncPhysicsInputComponent.cpp
+++ b/Engine/Source/Runtime/Engine/Private/PhysicsEngine/AsyncPhysicsInputComponent.cpp
@@ -402,6 +402,10 @@ void UAsyncPhysicsInputComponent::ServerRPCBufferInput_Implementation(UAsyncPhys
 void UAsyncPhysicsInputComponent::OnDispatchPhysicsTick(int32 PhysicsStep, int32 NumSteps, int32 ServerFrame)
 {
 	ensureMsgf(DataClass != nullptr, TEXT("You must call SetDataClass after creating the component"));
+	if (!ensure(DataToWrite != nullptr))
+	{
+		return;
+	}
 
 	// It would be better if we only registered while we had a local controller,
 	// but this is simpler to implement
