When streaming-out a partitioned sub-level, world partition's streaming cells (levels) are now unloaded.

diff --git a/Engine/Source/Runtime/Engine/Private/Level.cpp b/Engine/Source/Runtime/Engine/Private/Level.cpp
--- a/Engine/Source/Runtime/Engine/Private/Level.cpp
+++ b/Engine/Source/Runtime/Engine/Private/Level.cpp
@@ -2553,7 +2553,17 @@ void ULevel::OnLevelLoaded()
 					Transform = LevelStreaming->LevelTransform;
 				}
 
-				WorldPartition->Initialize(GetWorld(), Transform);
+				// It is allowed for a WorldPartition to already be initialized in the case where a partitioned sub-level 
+				// was streamed-out and is streamed-in again, without a CleanupLevel (GC).
+				if (!WorldPartition->IsInitialized())
+				{
+					WorldPartition->Initialize(GetWorld(), Transform);
+				}
+				else
+				{
+					check(GetWorld()->IsGameWorld());
+					check(WorldPartition->GetInstanceTransform().Equals(Transform));
+				}
 			}
 		}
 	}
diff --git a/Engine/Source/Runtime/Engine/Private/WorldPartition/WorldPartition.cpp b/Engine/Source/Runtime/Engine/Private/WorldPartition/WorldPartition.cpp
--- a/Engine/Source/Runtime/Engine/Private/WorldPartition/WorldPartition.cpp
+++ b/Engine/Source/Runtime/Engine/Private/WorldPartition/WorldPartition.cpp
@@ -609,6 +609,13 @@ bool UWorldPartition::IsInitialized() const
 	return InitState == EWorldPartitionInitState::Initialized;
 }
 
+bool UWorldPartition::CanStream() const
+{
+	// WorldPartition can't stream if not initialized or it it's part of a partitioned 
+	// sub-level that was removed from its owning world.
+	return IsInitialized() && (GetTypedOuter<UWorld>()->PersistentLevel->GetWorld() != nullptr);
+}
+
 bool UWorldPartition::IsMainWorldPartition() const
 {
 	check(World);
@@ -662,13 +669,20 @@ void UWorldPartition::RegisterDelegates()
 	}
 #endif
 
-	if (IsMainWorldPartition() && World->IsGameWorld())
+	if (World->IsGameWorld())
 	{
-		World->OnWorldMatchStarting.AddUObject(this, &UWorldPartition::OnWorldMatchStarting);
+		if (IsMainWorldPartition())
+		{
+			World->OnWorldMatchStarting.AddUObject(this, &UWorldPartition::OnWorldMatchStarting);
 
 #if !UE_BUILD_SHIPPING
-		FCoreDelegates::OnGetOnScreenMessages.AddUObject(this, &UWorldPartition::GetOnScreenMessages);
+			FCoreDelegates::OnGetOnScreenMessages.AddUObject(this, &UWorldPartition::GetOnScreenMessages);
 #endif
+		}
+		else
+		{
+			FWorldDelegates::LevelRemovedFromWorld.AddUObject(this, &UWorldPartition::OnLevelRemovedFromWorld);
+		}
 	}
 }
 
@@ -696,13 +710,30 @@ void UWorldPartition::UnregisterDelegates()
 	}
 #endif
 
-	if (IsMainWorldPartition() && World->IsGameWorld())
+	if (World->IsGameWorld())
 	{
-		World->OnWorldMatchStarting.RemoveAll(this);
+		if (IsMainWorldPartition())
+		{
+			World->OnWorldMatchStarting.RemoveAll(this);
 
 #if !UE_BUILD_SHIPPING
-		FCoreDelegates::OnGetOnScreenMessages.RemoveAll(this);
+			FCoreDelegates::OnGetOnScreenMessages.RemoveAll(this);
 #endif
+		}
+		else
+		{
+			FWorldDelegates::LevelRemovedFromWorld.RemoveAll(this);
+		}
+	}
+}
+
+void UWorldPartition::OnLevelRemovedFromWorld(ULevel* InLevel, UWorld* InWorld)
+{
+	check(!IsMainWorldPartition());
+	if ((World == InWorld) && (InLevel == GetTypedOuter<UWorld>()->PersistentLevel))
+	{
+		check(!CanStream());
+		UpdateStreamingState();
 	}
 }
 
diff --git a/Engine/Source/Runtime/Engine/Private/WorldPartition/WorldPartitionStreamingPolicy.cpp b/Engine/Source/Runtime/Engine/Private/WorldPartition/WorldPartitionStreamingPolicy.cpp
--- a/Engine/Source/Runtime/Engine/Private/WorldPartition/WorldPartitionStreamingPolicy.cpp
+++ b/Engine/Source/Runtime/Engine/Private/WorldPartition/WorldPartitionStreamingPolicy.cpp
@@ -88,7 +88,7 @@ void UWorldPartitionStreamingPolicy::UpdateStreamingSources()
 
 	StreamingSources.Reset();
 
-	if (!WorldPartition->IsInitialized())
+	if (!WorldPartition->CanStream())
 	{
 		return;
 	}
@@ -243,8 +243,8 @@ void UWorldPartitionStreamingPolicy::UpdateStreamingState()
 	bool bClient = NetMode == NM_Standalone || NetMode == NM_Client || AWorldPartitionReplay::IsEnabled(World);
 	if (bClient)
 	{
-		// When uninitializing, UpdateStreamingState is called, but we don't want any cells to be loaded
-		if (WorldPartition->IsInitialized())
+		// When world partition can't stream, all cells must be unloaded
+		if (WorldPartition->CanStream())
 		{
 			UWorldPartitionRuntimeCell::DirtyStreamingSourceCacheEpoch();
 			WorldPartition->RuntimeHash->GetStreamingCells(StreamingSources, FrameActivateCells, FrameLoadCells);
diff --git a/Engine/Source/Runtime/Engine/Public/WorldPartition/WorldPartition.h b/Engine/Source/Runtime/Engine/Public/WorldPartition/WorldPartition.h
--- a/Engine/Source/Runtime/Engine/Public/WorldPartition/WorldPartition.h
+++ b/Engine/Source/Runtime/Engine/Public/WorldPartition/WorldPartition.h
@@ -31,6 +31,7 @@ class IStreamingGenerationErrorHandler;
 class FLoaderAdapterAlwaysLoadedActors;
 class FHLODActorDesc;
 class UCanvas;
+class ULevel;
 
 struct IWorldPartitionStreamingSourceProvider;
 
@@ -181,13 +182,14 @@ public:
 	bool IsInitialized() const;
 	virtual void Uninitialize() override;
 
+	bool CanStream() const;
 	bool IsMainWorldPartition() const;
 	const FTransform& GetInstanceTransform() const;
 	const FTransform* GetInstanceTransformPtr() const;
 
 	void Tick(float DeltaSeconds);
 	void UpdateStreamingState();
-	bool CanAddLoadedLevelToWorld(class ULevel* InLevel) const;
+	bool CanAddLoadedLevelToWorld(ULevel* InLevel) const;
 	bool IsStreamingCompleted(EWorldPartitionRuntimeCellState QueryState, const TArray<FWorldPartitionStreamingQuerySource>& QuerySources, bool bExactState) const;
 
 	const TArray<FWorldPartitionStreamingSource>& GetStreamingSources() const;
@@ -270,7 +272,7 @@ private:
 #endif
 
 	void OnWorldMatchStarting();
-
+	void OnLevelRemovedFromWorld(ULevel* InLevel, UWorld* InWorld);
 	void OnPostBugItGoCalled(const FVector& Loc, const FRotator& Rot);
 
 	// Delegates registration
