Fix for ::ForceRebuild and DDC1 texture builds

diff --git a/Engine/Source/Runtime/Engine/Private/TextureDerivedDataTask.cpp b/Engine/Source/Runtime/Engine/Private/TextureDerivedDataTask.cpp
--- a/Engine/Source/Runtime/Engine/Private/TextureDerivedDataTask.cpp
+++ b/Engine/Source/Runtime/Engine/Private/TextureDerivedDataTask.cpp
@@ -833,12 +833,16 @@ void FTextureCacheDerivedDataWorker::DoWork()
 		LocalDerivedDataKeySuffix = MoveTemp(FetchOrBuildKeySuffix);
 		GetTextureDerivedDataKeyFromSuffix(LocalDerivedDataKeySuffix, LocalDerivedDataKey);
 
-		FRequestOwner BlockingOwner(EPriority::Blocking);
-		GetCache().GetValue({{{TexturePathName}, ConvertLegacyCacheKey(LocalDerivedDataKey)}}, BlockingOwner, [&RawDerivedData](FCacheGetValueResponse&& Response)
+		if (bForceRebuild == false)
 		{
-			RawDerivedData = Response.Value.GetData().Decompress();
-		});
-		BlockingOwner.Wait();
+			FRequestOwner BlockingOwner(EPriority::Blocking);
+			GetCache().GetValue({{{TexturePathName}, ConvertLegacyCacheKey(LocalDerivedDataKey)}}, BlockingOwner, [&RawDerivedData](FCacheGetValueResponse&& Response)
+			{
+				RawDerivedData = Response.Value.GetData().Decompress();
+			});
+			BlockingOwner.Wait();
+		}
+
 		bLoadedFromDDC = !RawDerivedData.IsNull();
 	}
 
