Fix for Collision boxes always have Navigation Area Class reset to NavArea_Obstacle on project load

diff --git a/Engine/Source/Runtime/Engine/Classes/Components/ShapeComponent.h b/Engine/Source/Runtime/Engine/Classes/Components/ShapeComponent.h
--- a/Engine/Source/Runtime/Engine/Classes/Components/ShapeComponent.h
+++ b/Engine/Source/Runtime/Engine/Classes/Components/ShapeComponent.h
@@ -29,9 +29,15 @@ class ENGINE_API UShapeComponent : public UPrimitiveComponent
 	UPROPERTY(transient, duplicatetransient)
 	TObjectPtr<class UBodySetup> ShapeBodySetup;
 
-	/** Navigation area type (empty = default obstacle) */
-	UPROPERTY(EditAnywhere, Category = Navigation)
+#if WITH_EDITORONLY_DATA
+	/** Navigation area type (empty / none has no effect on underlying area type).
+	 *  Deprecated! Use AreaClassOverride / bUseSystemDefaultObstacleAreaClass instead. 
+	 *  NOTE Adding _DEPRECATED to this variable can causes known bugs when patching it up.
+	 */
+	UE_DEPRECATED(5.0, "AreaClass is deprecated, use AreaClassOverride / bUseSystemDefaultObstacleAreaClass instead!.")
+	UPROPERTY()
 	TSubclassOf<class UNavAreaBase> AreaClass;
+#endif // WITH_EDITORONLY_DATA
 
 	/** Color used to draw the shape. */
 	UPROPERTY(EditAnywhere, AdvancedDisplay, BlueprintReadOnly, Category=Shape)
@@ -50,6 +56,16 @@ class ENGINE_API UShapeComponent : public UPrimitiveComponent
 	uint8 bDynamicObstacle : 1;
 
 protected:
+	/** Navigation area type override, null / none = no change to nav mesh.
+	 *  bDynamicObstacle must be true and bUseSystemDefaultAreaClass false to use this.
+	 */
+	UPROPERTY(EditAnywhere, Category = Navigation, meta = (EditCondition = "bDynamicObstacle && !bUseSystemDefaultObstacleAreaClass"))
+	TSubclassOf<class UNavAreaBase> AreaClassOverride;
+
+	/** Uses FNavigationSystem::GetDefaultObstacleArea() by default instead of AreaClassOverride, bDynamicObstacle must be true to use this.  */
+	UPROPERTY(EditAnywhere, Category = Navigation, meta = (EditCondition = "bDynamicObstacle"))
+	uint8 bUseSystemDefaultObstacleAreaClass : 1;
+
 	/** If the body setup can be shared (i.e. there have been no alterations compared to the CDO)*/
 	uint8 bUseArchetypeBodySetup : 1;
 
@@ -93,6 +109,8 @@ public:
 	//~ End USceneComponent Interface
 
 	//~ Begin UObject Interface.
+	virtual void Serialize(FArchive& Ar) override;
+
 #if WITH_EDITOR
 	virtual void PostEditChangeProperty(FPropertyChangedEvent& PropertyChangedEvent) override;
 	virtual bool IgnoreBoundsForEditorFocus() const override { return true; }
@@ -101,6 +119,10 @@ public:
 
 	/** Update the body setup parameters based on shape information*/
 	virtual void UpdateBodySetup();
+
+	TSubclassOf<class UNavAreaBase> GetDesiredAreaClass() const;
+	void SetAreaClassOverride(TSubclassOf<class UNavAreaBase> InAreaClassOverride);
+	void SetUseSystemDefaultObstacleAreaClass();
 };
 
 enum class EShapeBodySetupHelper
diff --git a/Engine/Source/Runtime/Engine/Private/Components/ShapeComponent.cpp b/Engine/Source/Runtime/Engine/Private/Components/ShapeComponent.cpp
--- a/Engine/Source/Runtime/Engine/Private/Components/ShapeComponent.cpp
+++ b/Engine/Source/Runtime/Engine/Private/Components/ShapeComponent.cpp
@@ -13,7 +13,33 @@
 	#include "PhysXPublic.h"
 #endif // WITH_PHYSX
 
+// Custom serialization version
+struct FShapeComponentCustomVersion
+{
+	enum Type
+	{
+		// Before any version changes were made
+		BeforeCustomVersionWasAdded = 0,
+
+		// Reworked AreaClass Implementation
+		AreaClassRework,
+
+		// -----<new versions can be added above this line>-------------------------------------------------
+		VersionPlusOne,
+		LatestVersion = VersionPlusOne - 1
+	};
+
+	// The GUID for this custom version number
+	const static FGuid GUID;
 
+private:
+	FShapeComponentCustomVersion() {}
+};
+
+const FGuid FShapeComponentCustomVersion::GUID(0xB6E31B1C, 0xD29F11EC, 0x857E9F85, 0x6F9970E2);
+
+// Register the custom version with core
+FCustomVersionRegistration GRegisterShapeComponentCustomVersion(FShapeComponentCustomVersion::GUID, FShapeComponentCustomVersion::LatestVersion, TEXT("ShapeComponentVer"));
 
 UShapeComponent::UShapeComponent(const FObjectInitializer& ObjectInitializer)
 	: Super(ObjectInitializer)
@@ -33,10 +59,18 @@ UShapeComponent::UShapeComponent(const FObjectInitializer& ObjectInitializer)
 	bHasCustomNavigableGeometry = EHasCustomNavigableGeometry::Yes;
 	bCanEverAffectNavigation = true;
 	bDynamicObstacle = false;
+	
+#if WITH_EDITOR
+	PRAGMA_DISABLE_DEPRECATION_WARNINGS
 	AreaClass = FNavigationSystem::GetDefaultObstacleArea();
+	PRAGMA_ENABLE_DEPRECATION_WARNINGS
+#endif // WITH_EDITOR
 
 	// Ignore streaming updates since GetUsedMaterials() is not implemented.
 	bIgnoreStreamingManagerUpdate = true;
+
+	bUseSystemDefaultObstacleAreaClass = true;
+	AreaClassOverride = nullptr;
 }
 
 FPrimitiveSceneProxy* UShapeComponent::CreateSceneProxy()
@@ -62,6 +96,34 @@ UBodySetup* UShapeComponent::GetBodySetup()
 	return ShapeBodySetup;
 }
 
+void UShapeComponent::Serialize(FArchive& Ar)
+{
+	Super::Serialize(Ar);
+
+#if WITH_EDITOR
+	Ar.UsingCustomVersion(FShapeComponentCustomVersion::GUID);
+
+	if (Ar.IsLoading())
+	{
+		const int32 Version = Ar.CustomVer(FShapeComponentCustomVersion::GUID);
+
+		// Note this has to be done during Serialize() not on PostLoad(), otherwise a blueprint class (with this component) CDO,
+		// can become out of sync with patching AreaClassOverride with AreaClass and causes bugs. This occurs if a level has an instance of the 
+		// blueprint class (with this component) and is modified then saved after patching up (but the BP class is not saved and is the 
+		// older version still), on the next time the level is loaded.
+		if (Version < FShapeComponentCustomVersion::AreaClassRework)
+		{
+			// If we are loading this object prior to the AreaClass rework then we just use whatever the current AreaClass is as the override.
+			bUseSystemDefaultObstacleAreaClass = false;
+
+			PRAGMA_DISABLE_DEPRECATION_WARNINGS
+			AreaClassOverride = AreaClass;
+			PRAGMA_ENABLE_DEPRECATION_WARNINGS
+		}
+	}
+#endif // WITH_EDITOR
+}
+
 #if WITH_EDITOR
 void UShapeComponent::PostEditChangeProperty(FPropertyChangedEvent& PropertyChangedEvent)
 {
@@ -85,7 +147,7 @@ void UShapeComponent::GetNavigationData(FNavigationRelevantData& Data) const
 
 	if (bDynamicObstacle)
 	{
-		Data.Modifiers.CreateAreaModifiers(this, AreaClass);
+		Data.Modifiers.CreateAreaModifiers(this, GetDesiredAreaClass());
 	}
 }
 
@@ -97,6 +159,22 @@ bool UShapeComponent::IsNavigationRelevant() const
 	return (bDynamicObstacle && CanEverAffectNavigation()) || Super::IsNavigationRelevant();
 }
 
+TSubclassOf<class UNavAreaBase> UShapeComponent::GetDesiredAreaClass() const
+{
+	return bUseSystemDefaultObstacleAreaClass ?  FNavigationSystem::GetDefaultObstacleArea() : AreaClassOverride;
+}
+
+void UShapeComponent::SetAreaClassOverride(TSubclassOf<class UNavAreaBase> InAreaClassOverride)
+{
+	AreaClassOverride = InAreaClassOverride;
+	bUseSystemDefaultObstacleAreaClass = false;
+}
+
+void UShapeComponent::SetUseSystemDefaultObstacleAreaClass()
+{
+	bUseSystemDefaultObstacleAreaClass = true;
+}
+
 template <> void UShapeComponent::AddShapeToGeomArray<FKBoxElem>() { ShapeBodySetup->AggGeom.BoxElems.Add(FKBoxElem()); }
 template <> void UShapeComponent::AddShapeToGeomArray<FKSphereElem>() { ShapeBodySetup->AggGeom.SphereElems.Add(FKSphereElem()); }
 template <> void UShapeComponent::AddShapeToGeomArray<FKSphylElem>() { ShapeBodySetup->AggGeom.SphylElems.Add(FKSphylElem()); }
