Removed SingleLayerWaterRefractionFullPrecision ShowFlag

#rnx

diff --git a/Engine/Source/Editor/DetailCustomizations/Private/SceneCaptureDetails.cpp b/Engine/Source/Editor/DetailCustomizations/Private/SceneCaptureDetails.cpp
--- a/Engine/Source/Editor/DetailCustomizations/Private/SceneCaptureDetails.cpp
+++ b/Engine/Source/Editor/DetailCustomizations/Private/SceneCaptureDetails.cpp
@@ -103,7 +103,6 @@ void FSceneCaptureDetails::CustomizeDetails( IDetailLayoutBuilder& DetailLayout
 	ShowFlagsToAllowForCaptures.Add(FEngineShowFlags::EShowFlag::SF_EyeAdaptation);
 	ShowFlagsToAllowForCaptures.Add(FEngineShowFlags::EShowFlag::SF_Game);
 	ShowFlagsToAllowForCaptures.Add(FEngineShowFlags::EShowFlag::SF_ToneCurve); 
-	ShowFlagsToAllowForCaptures.Add(FEngineShowFlags::EShowFlag::SF_SingleLayerWaterRefractionFullPrecision);
 	ShowFlagsToAllowForCaptures.Add(FEngineShowFlags::EShowFlag::SF_PathTracing);
 
 	// Create array of flag name strings for each group
diff --git a/Engine/Source/Runtime/Engine/Public/ShowFlags.h b/Engine/Source/Runtime/Engine/Public/ShowFlags.h
--- a/Engine/Source/Runtime/Engine/Public/ShowFlags.h
+++ b/Engine/Source/Runtime/Engine/Public/ShowFlags.h
@@ -476,7 +476,6 @@ private:
 		SetVisualizeVolumetricCloudConservativeDensity(false);
 		SetVisualizeStrataMaterial(false);
 		SetDrawOnlyVSMInvalidatingGeo(false);
-		SetSingleLayerWaterRefractionFullPrecision(false);
 
 		SetLumenScreenTraces(true);
 		SetLumenDetailTraces(true);
diff --git a/Engine/Source/Runtime/Engine/Public/ShowFlagsValues.inl b/Engine/Source/Runtime/Engine/Public/ShowFlagsValues.inl
--- a/Engine/Source/Runtime/Engine/Public/ShowFlagsValues.inl
+++ b/Engine/Source/Runtime/Engine/Public/ShowFlagsValues.inl
@@ -401,9 +401,6 @@ SHOWFLAG_FIXED_IN_SHIPPING(1, VirtualShadowMapCaching, SFG_Developer, NSLOCTEXT(
 /** Debug the Strata material buffer content */
 SHOWFLAG_FIXED_IN_SHIPPING(0, DrawOnlyVSMInvalidatingGeo, SFG_Visualize, NSLOCTEXT("UnrealEd", "DrawOnlyVSMInvalidatingGeo", "Draw Only Geometry Causing VSM Invalidation"))
 
-/** If single layer water should use a full precision depth buffer for SceneDepthWithoutSingleLayerWater. For now SHOWFLAG_ALWAYS_ACCESSIBLE because it's exposed in SceneCapture */
-SHOWFLAG_ALWAYS_ACCESSIBLE(SingleLayerWaterRefractionFullPrecision, SFG_Advanced, NSLOCTEXT("UnrealEd", "SingleLayerWaterRefractionFullPrecision", "Single Layer Water Refraction Full Precision"))
-
 /** Debug the Virtual Texture System */
 SHOWFLAG_FIXED_IN_SHIPPING(0, VirtualTextureResidency, SFG_Hidden, NSLOCTEXT("UnrealEd", "VirtualTextureResidency", "Virtual Texture Residency"))
 
