Exposed some cloud properties from sky and dir light to be interpolable for timeline.

diff --git a/Engine/Source/Runtime/Engine/Classes/Components/DirectionalLightComponent.h b/Engine/Source/Runtime/Engine/Classes/Components/DirectionalLightComponent.h
--- a/Engine/Source/Runtime/Engine/Classes/Components/DirectionalLightComponent.h
+++ b/Engine/Source/Runtime/Engine/Classes/Components/DirectionalLightComponent.h
@@ -199,17 +199,17 @@ class ENGINE_API UDirectionalLightComponent : public ULightComponent
 	/**
 	 * The overall strength of the cloud shadow, higher value will block more light.
 	 */
-	UPROPERTY(EditAnywhere, BlueprintReadOnly, Category = AtmosphereAndCloud, AdvancedDisplay, meta = (UIMin = "0", UIMax = "1", ClampMin = "0", SliderExponent = 1.0))
+	UPROPERTY(EditAnywhere, BlueprintReadOnly, interp, Category = AtmosphereAndCloud, AdvancedDisplay, meta = (UIMin = "0", UIMax = "1", ClampMin = "0", SliderExponent = 1.0))
 	float CloudShadowStrength;
 	/**
 	 * The strength of the shadow on atmosphere. Disabled when 0.
 	 */
-	UPROPERTY(EditAnywhere, BlueprintReadOnly, Category = AtmosphereAndCloud, AdvancedDisplay, meta = (UIMin = "0", UIMax = "1", ClampMin = "0", SliderExponent = 1.0))
+	UPROPERTY(EditAnywhere, BlueprintReadOnly, interp, Category = AtmosphereAndCloud, AdvancedDisplay, meta = (UIMin = "0", UIMax = "1", ClampMin = "0", SliderExponent = 1.0))
 	float CloudShadowOnAtmosphereStrength;
 	/**
 	 * The strength of the shadow on opaque and transparent meshes. Disabled when 0.
 	 */
-	UPROPERTY(EditAnywhere, BlueprintReadOnly, Category = AtmosphereAndCloud, AdvancedDisplay, meta = (UIMin = "0", UIMax = "1", ClampMin = "0", SliderExponent = 1.0))
+	UPROPERTY(EditAnywhere, BlueprintReadOnly, interp, Category = AtmosphereAndCloud, AdvancedDisplay, meta = (UIMin = "0", UIMax = "1", ClampMin = "0", SliderExponent = 1.0))
 	float CloudShadowOnSurfaceStrength;
 	/**
 	 * The bias applied to the shadow front depth of the volumetric cloud shadow map.
diff --git a/Engine/Source/Runtime/Engine/Classes/Components/SkyLightComponent.h b/Engine/Source/Runtime/Engine/Classes/Components/SkyLightComponent.h
--- a/Engine/Source/Runtime/Engine/Classes/Components/SkyLightComponent.h
+++ b/Engine/Source/Runtime/Engine/Classes/Components/SkyLightComponent.h
@@ -180,7 +180,7 @@ class ENGINE_API USkyLightComponent : public ULightComponentBase
 	/**
 	 * The strength of the ambient occlusion, higher value will block more light.
 	 */
-	UPROPERTY(EditAnywhere, BlueprintReadOnly, Category = AtmosphereAndCloud, meta = (UIMin = "0", UIMax = "1", ClampMin = "0", SliderExponent = 1.0))
+	UPROPERTY(EditAnywhere, BlueprintReadOnly, interp, Category = AtmosphereAndCloud, meta = (UIMin = "0", UIMax = "1", ClampMin = "0", SliderExponent = 1.0))
 	float CloudAmbientOcclusionStrength;
 	/**
 	 * The world space radius of the cloud ambient occlusion map around the camera in kilometers.
