Fix for ensure firing when enabling UseMaterialAttribute. When not using attributes, we can check the refraction input.
Otherwise we only rely on the cached bUsesDistortion from UMaterial::PostEditChangePropertyInternal.

diff --git a/Engine/Source/Runtime/Engine/Private/Materials/Material.cpp b/Engine/Source/Runtime/Engine/Private/Materials/Material.cpp
--- a/Engine/Source/Runtime/Engine/Private/Materials/Material.cpp
+++ b/Engine/Source/Runtime/Engine/Private/Materials/Material.cpp
@@ -3284,10 +3284,6 @@ void UMaterial::ConvertMaterialToStrataMaterial()
 	{
 		// We might have moved connections above so update the CachedExpressionData from the EditorOnly connection data (ground truth).
 		UpdateCachedExpressionData();
-
-		// Make sure our cached data is consistent
-		ensureMsgf(!(GetEditorOnlyData()->Refraction.IsConnected() && !GetCachedExpressionData().IsPropertyConnected(MP_Refraction)),
-			TEXT("GetCachedExpressionData() says refraction isn't connected, but GetEditorOnlyData() says it is"));
 	}
 #endif
 }
@@ -6173,9 +6169,13 @@ bool UMaterial::IsPropertyActiveInEditor(EMaterialProperty InProperty) const
 bool UMaterial::IsPropertyActiveInDerived(EMaterialProperty InProperty, const UMaterialInterface* DerivedMaterial) const
 {
 #if WITH_EDITOR
-	// Make sure our cached data is consistent
-	ensureMsgf(!(GetEditorOnlyData()->Refraction.IsConnected() && !GetCachedExpressionData().IsPropertyConnected(MP_Refraction)),
-		TEXT("GetCachedExpressionData() says refraction isn't connected, but GetEditorOnlyData() says it is"));
+	if (!bUseMaterialAttributes)
+	{
+		// We can only check connected pins if we are not using material attributes because the connections are kept.
+		// This with the goal to bring back the connections when bUseMaterialAttributes is disabled again.
+		ensureMsgf(!(GetEditorOnlyData()->Refraction.IsConnected() && !GetCachedExpressionData().IsPropertyConnected(MP_Refraction)),
+			TEXT("GetCachedExpressionData() says refraction isn't connected, but GetEditorOnlyData() says it is"));
+	}
 #endif
 	return IsPropertyActive_Internal(InProperty,
 		MaterialDomain,
@@ -6184,7 +6184,7 @@ bool UMaterial::IsPropertyActiveInDerived(EMaterialProperty InProperty, const UM
 		DerivedMaterial->GetShadingModels(),
 		TranslucencyLightingMode,
 		BlendableOutputAlpha,
-		GetCachedExpressionData().IsPropertyConnected(MP_Refraction),
+		!bUseMaterialAttributes ? GetCachedExpressionData().IsPropertyConnected(MP_Refraction) : bUsesDistortion,
 		DerivedMaterial->IsShadingModelFromMaterialExpression(),
 		IsTranslucencyWritingVelocity(),
 		IsPropertySupported(InProperty));
