Fixed AActor::UseWorldPartitionRuntimeCellDataLayers, must return false actor owning world is null.
#rnx

diff --git a/Engine/Source/Runtime/Engine/Private/Actor.cpp b/Engine/Source/Runtime/Engine/Private/Actor.cpp
--- a/Engine/Source/Runtime/Engine/Private/Actor.cpp
+++ b/Engine/Source/Runtime/Engine/Private/Actor.cpp
@@ -5944,7 +5944,8 @@ TArray<const UDataLayerInstance*> AActor::GetDataLayerInstancesInternal(bool bUs
 {
 	if (UseWorldPartitionRuntimeCellDataLayers())
 	{
-		return GetLevel()->GetWorldPartitionRuntimeCell()->GetDataLayerInstances();
+		const IWorldPartitionCell* Cell = GetLevel()->GetWorldPartitionRuntimeCell();
+		return Cell ? Cell->GetDataLayerInstances() : TArray<const UDataLayerInstance*>();
 	}
 
 #if WITH_EDITOR
@@ -5981,7 +5982,8 @@ bool AActor::ContainsDataLayer(const UDataLayerInstance* DataLayerInstance) cons
 {
 	if (UseWorldPartitionRuntimeCellDataLayers())
 	{
-		return GetLevel()->GetWorldPartitionRuntimeCell()->ContainsDataLayer(DataLayerInstance);
+		const IWorldPartitionCell* Cell = GetLevel()->GetWorldPartitionRuntimeCell();
+		return Cell ? Cell->ContainsDataLayer(DataLayerInstance) : false;
 	}
 
 #if WITH_EDITOR
@@ -6000,7 +6002,8 @@ bool AActor::ContainsDataLayer(const UDataLayerAsset* DataLayerAsset) const
 
 	if (UseWorldPartitionRuntimeCellDataLayers())
 	{
-		return GetLevel()->GetWorldPartitionRuntimeCell()->ContainsDataLayer(DataLayerAsset);
+		const IWorldPartitionCell* Cell = GetLevel()->GetWorldPartitionRuntimeCell();
+		return Cell ? Cell->ContainsDataLayer(DataLayerAsset) : false;
 	}
 
 #if WITH_EDITOR
@@ -6022,7 +6025,8 @@ bool AActor::HasDataLayers() const
 {
 	if (UseWorldPartitionRuntimeCellDataLayers())
 	{
-		return GetLevel()->GetWorldPartitionRuntimeCell()->HasDataLayers();
+		const IWorldPartitionCell* Cell = GetLevel()->GetWorldPartitionRuntimeCell();
+		return Cell ? Cell->HasDataLayers() : false;
 	}
 
 #if WITH_EDITOR
