World Partition: added cvar LevelStreaming.ShouldReuseUnloadedButStillAroundLevels to allow enabling the old behavior of level streaming to reuse unloaded sublevels still around waiting for GC.

#rnx

diff --git a/Engine/Source/Runtime/Engine/Classes/Engine/LevelStreaming.h b/Engine/Source/Runtime/Engine/Classes/Engine/LevelStreaming.h
--- a/Engine/Source/Runtime/Engine/Classes/Engine/LevelStreaming.h
+++ b/Engine/Source/Runtime/Engine/Classes/Engine/LevelStreaming.h
@@ -528,6 +528,9 @@ public:
 	/** If true server will wait for client acknowledgment before making treating streaming levels as visible for the client */
 	static bool ShouldServerUseMakingVisibleTransactionRequest();
 
+	/** If true level streaming will reuse the unloaded levels that aren't GC'd yet. */
+	static bool ShouldReuseUnloadedButStillAroundLevels();
+
 	/** 
 	 * Traverses all streaming level objects in the persistent world and in all inner worlds and calls appropriate delegate for streaming objects that refer specified level 
 	 *
diff --git a/Engine/Source/Runtime/Engine/Private/LevelStreaming.cpp b/Engine/Source/Runtime/Engine/Private/LevelStreaming.cpp
--- a/Engine/Source/Runtime/Engine/Private/LevelStreaming.cpp
+++ b/Engine/Source/Runtime/Engine/Private/LevelStreaming.cpp
@@ -59,6 +59,14 @@ namespace LevelStreamingCVars
 		TEXT("Whether server should wait for client to acknowledge visibility update before treating streaming levels as visible by the client.\n")
 		TEXT("0: Disable, 1: Enable"),
 		ECVF_Default);
+
+	static bool bShouldReuseUnloadedButStillAroundLevels = false;
+	FAutoConsoleVariableRef CVarShouldReuseUnloadedButStillAroundLevels(
+		TEXT("LevelStreaming.ShouldReuseUnloadedButStillAroundLevels"),
+		bShouldReuseUnloadedButStillAroundLevels,
+		TEXT("Whether level streaming will reuse the unloaded levels that aren't GC'd yet.\n")
+		TEXT("0: Disable, 1: Enable"),
+		ECVF_ReadOnly);
 }
 
 bool ULevelStreaming::ShouldClientUseMakingInvisibleTransactionRequest()
@@ -71,6 +79,11 @@ bool ULevelStreaming::ShouldServerUseMakingVisibleTransactionRequest()
 	return LevelStreamingCVars::bShouldServerUseMakingVisibleTransactionRequest;
 }
 
+bool ULevelStreaming::ShouldReuseUnloadedButStillAroundLevels()
+{
+	return LevelStreamingCVars::bShouldReuseUnloadedButStillAroundLevels;
+}
+
 int32 ULevelStreamingDynamic::UniqueLevelInstanceId = 0;
 
 /**
@@ -1021,6 +1034,33 @@ void ULevelStreaming::SetLoadedLevel(ULevel* Level)
 	if (PendingUnloadLevel)
 	{
 		RemoveLevelAnnotation(PendingUnloadLevel);
+
+		if (!ULevelStreaming::ShouldReuseUnloadedButStillAroundLevels())
+		{
+			// Rename so it isn't found again
+			auto TrashPackage = [](UPackage* Package)
+			{
+				ForEachObjectWithPackage(Package, [](UObject* Object)
+				{
+					Object->ClearFlags(RF_Standalone);
+					return true;
+				}, false);
+
+				FName NewPackageName = MakeUniqueObjectName(nullptr, UPackage::StaticClass(), FName(*FString::Printf(TEXT("%s_Trashed"), *Package->GetName())));
+				Package->Rename(*NewPackageName.ToString(), nullptr, REN_ForceNoResetLoaders | REN_DontCreateRedirectors | REN_NonTransactional | REN_DoNotDirty);
+			};
+		
+			UPackage* LevelPackage = PendingUnloadLevel->GetPackage();
+			TrashPackage(LevelPackage);
+
+			ForEachObjectWithOuter(PendingUnloadLevel, [LevelPackage,  TrashPackage](UObject* InObject)
+			{
+				if (UPackage* ObjectPackage = InObject->GetPackage(); ObjectPackage && LevelPackage != ObjectPackage)
+				{
+					TrashPackage(ObjectPackage);
+				}
+			}, false);
+		}
 	}
 
 	if (LoadedLevel)
diff --git a/Engine/Source/Runtime/Engine/Private/WorldPartition/WorldPartitionLevelStreamingDynamic.cpp b/Engine/Source/Runtime/Engine/Private/WorldPartition/WorldPartitionLevelStreamingDynamic.cpp
--- a/Engine/Source/Runtime/Engine/Private/WorldPartition/WorldPartitionLevelStreamingDynamic.cpp
+++ b/Engine/Source/Runtime/Engine/Private/WorldPartition/WorldPartitionLevelStreamingDynamic.cpp
@@ -197,17 +197,26 @@ bool UWorldPartitionLevelStreamingDynamic::RequestLevel(UWorld* InPersistentWorl
 	const FName DesiredPackageName = GetWorldAssetPackageFName();
 	UPackage* LevelPackage = (UPackage*)StaticFindObjectFast(UPackage::StaticClass(), nullptr, DesiredPackageName, 0, 0, RF_NoFlags, EInternalObjectFlags::Garbage);
 	UWorld* FoundWorld = LevelPackage ? UWorld::FindWorldInPackage(LevelPackage) : nullptr;
-	check(!FoundWorld || IsValidChecked(FoundWorld));
-	check(!FoundWorld || FoundWorld->PersistentLevel);
-	if (FoundWorld && FoundWorld->PersistentLevel != RuntimeLevel)
+
+	if (FoundWorld)
 	{
-		check(RuntimeLevel == nullptr);
-		check(LoadedLevel == nullptr);
-		RuntimeLevel = FoundWorld->PersistentLevel;
+		check(IsValid(FoundWorld));
+		check(FoundWorld->PersistentLevel);
+		check(ULevelStreaming::ShouldReuseUnloadedButStillAroundLevels());
+
+		if (FoundWorld->PersistentLevel != RuntimeLevel)
+		{
+			check(!RuntimeLevel);
+			check(!LoadedLevel);
+
+			RuntimeLevel = FoundWorld->PersistentLevel;
+		}
 	}
 
 	if (RuntimeLevel)
 	{
+		check(ULevelStreaming::ShouldReuseUnloadedButStillAroundLevels());
+
 		// Reuse existing Level
 		UPackage* CellLevelPackage = RuntimeLevel->GetPackage();
 		check(CellLevelPackage);
@@ -272,6 +281,22 @@ bool UWorldPartitionLevelStreamingDynamic::RequestLevel(UWorld* InPersistentWorl
 	return true;
 }
 
+void UWorldPartitionLevelStreamingDynamic::SetLoadedLevel(ULevel* Level)
+{
+	// Should only support changing a valid to invalid level or vice versa
+	check((RuntimeLevel == Level) || (!!RuntimeLevel != !!Level));
+
+	if (!ULevelStreaming::ShouldReuseUnloadedButStillAroundLevels())
+	{
+		if (RuntimeLevel && !Level)
+		{
+			OnCleanupLevel();
+		}
+	}
+
+	Super::SetLoadedLevel(Level);
+}
+
 /**
  * Loads all objects of a runtime Level
  */
@@ -431,32 +456,7 @@ void UWorldPartitionLevelStreamingDynamic::OnCleanupLevel()
 	if (RuntimeLevel)
 	{
 		PackageReferencer.RemoveReferences();
-
 		RuntimeLevel->OnCleanupLevel.Remove(OnCleanupLevelDelegateHandle);
-
-		auto TrashPackage = [](UPackage* Package)
-		{
-			// Clears RF_Standalone flag on objects in package (UMetaData)
-			ForEachObjectWithPackage(Package, [](UObject* Object)
-			{
-				Object->ClearFlags(RF_Standalone);
-				return true;
-			}, false);
-
-			// Rename package to avoid having to deal with pending kill objects in subsequent RequestLevel calls
-			FName NewPackageName = MakeUniqueObjectName(nullptr, UPackage::StaticClass(), FName(*FString::Printf(TEXT("%s_Trashed"), *Package->GetName())));
-			Package->Rename(*NewPackageName.ToString(), nullptr, REN_ForceNoResetLoaders | REN_DontCreateRedirectors | REN_NonTransactional | REN_DoNotDirty);
-		};
-		
-		TrashPackage(RuntimeLevel->GetPackage());
-		for (AActor* Actor : RuntimeLevel->Actors)
-		{
-			if (UPackage* ActorPackage = Actor ? Actor->GetExternalPackage() : nullptr)
-			{
-				TrashPackage(ActorPackage);
-			}
-		}
-
 		RuntimeLevel = nullptr;
 	}
 }
diff --git a/Engine/Source/Runtime/Engine/Public/WorldPartition/WorldPartitionLevelStreamingDynamic.h b/Engine/Source/Runtime/Engine/Public/WorldPartition/WorldPartitionLevelStreamingDynamic.h
--- a/Engine/Source/Runtime/Engine/Public/WorldPartition/WorldPartitionLevelStreamingDynamic.h
+++ b/Engine/Source/Runtime/Engine/Public/WorldPartition/WorldPartitionLevelStreamingDynamic.h
@@ -44,6 +44,7 @@ class ENGINE_API UWorldPartitionLevelStreamingDynamic : public ULevelStreamingDy
 
 	// Override ULevelStreaming
 	virtual bool RequestLevel(UWorld* PersistentWorld, bool bAllowLevelLoadRequests, EReqLevelBlock BlockPolicy) override;
+	virtual void SetLoadedLevel(ULevel* Level) override;
 	virtual void BeginDestroy() override;
 	virtual TOptional<FFolder::FRootObject> GetFolderRootObject() const override;
 
