WorldPartitionLevelHelper: Call GetExternalPackage() and check that it is valid (doesn't fix the Jira but fails earlier in case of regression)

#rnx

diff --git a/Engine/Source/Runtime/Engine/Private/WorldPartition/WorldPartitionLevelHelper.cpp b/Engine/Source/Runtime/Engine/Private/WorldPartition/WorldPartitionLevelHelper.cpp
--- a/Engine/Source/Runtime/Engine/Private/WorldPartition/WorldPartitionLevelHelper.cpp
+++ b/Engine/Source/Runtime/Engine/Private/WorldPartition/WorldPartitionLevelHelper.cpp
@@ -124,7 +124,8 @@ void FWorldPartitionLevelHelper::MoveExternalActorsToLevel(const TArray<FWorldPa
 		AActor* Actor = FindObject<AActor>(nullptr, *PackageObjectMapping.LoadedPath.ToString());
 		if (Actor)
 		{
-			UPackage* ActorExternalPackage = Actor->GetPackage();
+			UPackage* ActorExternalPackage = Actor->GetExternalPackage();
+			check(ActorExternalPackage);
 
 			const bool bSameOuter = (InLevel == Actor->GetOuter());
 			Actor->SetPackageExternal(false, false);
@@ -143,17 +144,14 @@ void FWorldPartitionLevelHelper::MoveExternalActorsToLevel(const TArray<FWorldPa
 			}
 
 			// Include objects found in the source actor package in the destination level package
-			if (ensure(ActorExternalPackage))
+			TArray<UObject*> Objects;
+			const bool bIncludeNestedSubobjects = false;
+			GetObjectsWithOuter(ActorExternalPackage, Objects, bIncludeNestedSubobjects);
+			for (UObject* Object : Objects)
 			{
-				TArray<UObject*> Objects;
-				const bool bIncludeNestedSubobjects = false;
-				GetObjectsWithOuter(ActorExternalPackage, Objects, bIncludeNestedSubobjects);
-				for (UObject* Object : Objects)
+				if (Object->GetFName() != NAME_PackageMetaData)
 				{
-					if (Object->GetFName() != NAME_PackageMetaData)
-					{
-						Object->Rename(nullptr, LevelPackage, REN_ForceNoResetLoaders);
-					}
+					Object->Rename(nullptr, LevelPackage, REN_ForceNoResetLoaders);
 				}
 			}
 		}
