Add sorted device profile fragment list to analytics
Alphabetically sort crash report's fragment list.

diff --git a/Engine/Source/Runtime/Engine/Classes/DeviceProfiles/DeviceProfileManager.h b/Engine/Source/Runtime/Engine/Classes/DeviceProfiles/DeviceProfileManager.h
--- a/Engine/Source/Runtime/Engine/Classes/DeviceProfiles/DeviceProfileManager.h
+++ b/Engine/Source/Runtime/Engine/Classes/DeviceProfiles/DeviceProfileManager.h
@@ -181,6 +181,18 @@ public:
 	*/
 	const FString GetActiveDeviceProfileName();
 
+	/**
+	* Get a string containing the current matched fragment list.
+	*
+	* e.g. "Fragment1,Fragment2,[tag]Fragment3"
+	* 
+	* @param bEnabledOnly				- Only Enabled fragments will be present in the returned string.
+	* @param bIncludeTags				- If true the Fragment string will include the tag.
+	* @param bAlphaSort					- If true the Fragments will be in alphabetical order as opposed to application order.
+	* @return csv string of the current matched fragments.
+	*/
+	const FString GetActiveDeviceProfileMatchedFragmentsString(bool bEnabledOnly, bool bIncludeTags, bool bAlphaSort);
+
 	/**
 	* Get the selected device profile name, either the platform name, or the name
 	* provided by a Device Profile Selector Module.
@@ -256,6 +268,9 @@ private:
 	/** Read and process all of the fragment matching rules. Returns an array containing the names of fragments selected. */
 	static TArray<FSelectedFragmentProperties> FindMatchingFragments(const FString& ParentDP, class FConfigCacheIni* PreviewConfigSystem);
 
+	/** Convert a FSelectedFragmentProperties array to a string */
+	static const FString FragmentPropertyArrayToFragmentString(const TArray<FSelectedFragmentProperties>& FragmentProperties, bool bEnabledOnly, bool bIncludeTags, bool bAlphaSort);
+
 	/** Get the current platform's the selector module. Can return null */
 	static class IDeviceProfileSelectorModule* GetDeviceProfileSelectorModule();
 
diff --git a/Engine/Source/Runtime/Engine/Private/DeviceProfiles/DeviceProfileManager.cpp b/Engine/Source/Runtime/Engine/Private/DeviceProfiles/DeviceProfileManager.cpp
--- a/Engine/Source/Runtime/Engine/Private/DeviceProfiles/DeviceProfileManager.cpp
+++ b/Engine/Source/Runtime/Engine/Private/DeviceProfiles/DeviceProfileManager.cpp
@@ -188,6 +188,12 @@ TMap<FName, FString> UDeviceProfileManager::GatherDeviceProfileCVars(const FStri
 		if (PlatformFragmentsSelected.Num() == 0)
 		{
 			PlatformFragmentsSelected = FindMatchingFragments(DeviceProfileName, GConfig);
+			if(PlatformFragmentsSelected.Num())
+			{
+				// Store the fragment string:
+				FString MatchedFragmentString = FragmentPropertyArrayToFragmentString(PlatformFragmentsSelected, false, false, true);
+				FGenericCrashContext::SetEngineData(TEXT("DeviceProfile.MatchedFragmentsSorted"), MatchedFragmentString);
+			}
 		}
 		FragmentsSelected = PlatformFragmentsSelected;
 	}
@@ -1288,10 +1294,10 @@ void UDeviceProfileManager::GetAllPossibleParentProfiles(const UDeviceProfile* C
 	}
 }
 
-
-
-
-
+const FString UDeviceProfileManager::GetActiveDeviceProfileMatchedFragmentsString(bool bEnabledOnly, bool bIncludeTags, bool bAlphaSort)
+{
+	return FragmentPropertyArrayToFragmentString(PlatformFragmentsSelected, bEnabledOnly, bIncludeTags, bAlphaSort);
+}
 
 #if ALLOW_OTHER_PLATFORM_CONFIG
 static bool GetCVarForDeviceProfile( FOutputDevice& Ar, FString DPName, FString CVarName)
diff --git a/Engine/Source/Runtime/Engine/Private/DeviceProfiles/DeviceProfileMatching.cpp b/Engine/Source/Runtime/Engine/Private/DeviceProfiles/DeviceProfileMatching.cpp
--- a/Engine/Source/Runtime/Engine/Private/DeviceProfiles/DeviceProfileMatching.cpp
+++ b/Engine/Source/Runtime/Engine/Private/DeviceProfiles/DeviceProfileMatching.cpp
@@ -540,6 +540,37 @@ static TArray<FSelectedFragmentProperties> FragmentStringToFragmentProperties(co
 	return FragmentPropertiesList;
 }
 
+const FString UDeviceProfileManager::FragmentPropertyArrayToFragmentString(const TArray<FSelectedFragmentProperties>& FragmentProperties, bool bEnabledOnly, bool bIncludeTags,bool bAlphaSort)
+{
+	FString MatchedFragmentString;
+	if(!FragmentProperties.IsEmpty())
+	{
+		TArray<FString> FragmentNames;
+		for (const FSelectedFragmentProperties& FragmentProperty : FragmentProperties)
+		{
+			if (bEnabledOnly == false || FragmentProperty.bEnabled)
+			{
+				FString FragmentName;
+				if (bIncludeTags && FragmentProperty.Tag != NAME_None)
+				{
+					FragmentName = TEXT("[") + FragmentProperty.Tag.ToString() + TEXT("]");
+				}
+				FragmentNames.Add(FragmentProperty.Fragment);
+			}
+		}
+		if (bAlphaSort)
+		{
+			FragmentNames.Sort();
+		}
+
+		for (FString& Name : FragmentNames)
+		{
+			MatchedFragmentString += MatchedFragmentString.IsEmpty() ? Name : TEXT(",") + Name;
+		}
+	}
+	return MatchedFragmentString;
+}
+
 class FDeviceProfileMatchingErrorContext: public FOutputDevice
 {
 public:
@@ -641,10 +672,6 @@ TArray<FSelectedFragmentProperties> UDeviceProfileManager::FindMatchingFragments
 	
 	}
 	SelectedFragments = RemoveAllWhiteSpace(SelectedFragments);
-	if(!SelectedFragments.IsEmpty())
-	{
-		FGenericCrashContext::SetEngineData(TEXT("DeviceProfile.MatchedFragments"), SelectedFragments);
-	}
 
 	UE_CLOG(!SelectedFragments.IsEmpty(), LogInit, Log, TEXT("MatchesRules:Fragment string %s"), *SelectedFragments);
 	TArray<FSelectedFragmentProperties> MatchedFragments = FragmentStringToFragmentProperties(SelectedFragments);
