- Fixed non-unity compile error

diff --git a/Engine/Source/Runtime/Engine/Private/Scalability.cpp b/Engine/Source/Runtime/Engine/Private/Scalability.cpp
--- a/Engine/Source/Runtime/Engine/Private/Scalability.cpp
+++ b/Engine/Source/Runtime/Engine/Private/Scalability.cpp
@@ -2,6 +2,7 @@
 
 #include "Scalability.h"
 #include "GenericPlatform/GenericPlatformSurvey.h"
+#include "Misc/ConfigCacheIni.h"
 #include "Misc/ConfigUtilities.h"
 #include "HAL/IConsoleManager.h"
 #include "SynthBenchmark.h"
