#ue5 - Add virtual specifiers to GetActorFeetLocation() and StartSwimming().

github #7320

diff --git a/Engine/Source/Runtime/Engine/Classes/GameFramework/CharacterMovementComponent.h b/Engine/Source/Runtime/Engine/Classes/GameFramework/CharacterMovementComponent.h
--- a/Engine/Source/Runtime/Engine/Classes/GameFramework/CharacterMovementComponent.h
+++ b/Engine/Source/Runtime/Engine/Classes/GameFramework/CharacterMovementComponent.h
@@ -1517,7 +1517,7 @@ public:
 	 * @param remainingTime - DeltaTime to complete transition to swimming
 	 * @param Iterations - physics iteration count
 	 */
-	void StartSwimming(FVector OldLocation, FVector OldVelocity, float timeTick, float remainingTime, int32 Iterations);
+	virtual void StartSwimming(FVector OldLocation, FVector OldVelocity, float timeTick, float remainingTime, int32 Iterations);
 
 	/* Swimming uses gravity - but scaled by (1.f - buoyancy) */
 	float Swim(FVector Delta, FHitResult& Hit);
diff --git a/Engine/Source/Runtime/Engine/Classes/GameFramework/NavMovementComponent.h b/Engine/Source/Runtime/Engine/Classes/GameFramework/NavMovementComponent.h
--- a/Engine/Source/Runtime/Engine/Classes/GameFramework/NavMovementComponent.h
+++ b/Engine/Source/Runtime/Engine/Classes/GameFramework/NavMovementComponent.h
@@ -82,7 +82,7 @@ public:
 	/** Returns location of controlled actor - meaning center of collision bounding box */
 	FORCEINLINE FVector GetActorLocation() const { return UpdatedComponent ? UpdatedComponent->GetComponentLocation() : FVector(FLT_MAX); }
 	/** Returns location of controlled actor's "feet" meaning center of bottom of collision bounding box */
-	FORCEINLINE FVector GetActorFeetLocation() const { return UpdatedComponent ? (UpdatedComponent->GetComponentLocation() - FVector(0,0,UpdatedComponent->Bounds.BoxExtent.Z)) : FNavigationSystem::InvalidLocation; }
+	virtual FVector GetActorFeetLocation() const { return UpdatedComponent ? (UpdatedComponent->GetComponentLocation() - FVector(0,0,UpdatedComponent->Bounds.BoxExtent.Z)) : FNavigationSystem::InvalidLocation; }
 	/** Returns based location of controlled actor */
 	virtual FBasedPosition GetActorFeetLocationBased() const;
 	/** Returns navigation location of controlled actor */
