#ue5 - Restore state of JumpMaxCount and JumpMaxHold time after replaying saved moves on a client correction, in case those values are modified by saved moves.

github #8722

diff --git a/Engine/Source/Runtime/Engine/Private/Components/CharacterMovementComponent.cpp b/Engine/Source/Runtime/Engine/Private/Components/CharacterMovementComponent.cpp
--- a/Engine/Source/Runtime/Engine/Private/Components/CharacterMovementComponent.cpp
+++ b/Engine/Source/Runtime/Engine/Private/Components/CharacterMovementComponent.cpp
@@ -8096,6 +8096,8 @@ bool UCharacterMovementComponent::ClientUpdatePositionAfterServerUpdate()
 	const FRootMotionMovementParams BackupRootMotionParams = RootMotionParams; // For animation root motion
 	const FRootMotionSourceGroup BackupRootMotion = CurrentRootMotion;
 	const bool bRealPressedJump = CharacterOwner->bPressedJump;
+	const float RealJumpMaxHoldTime = CharacterOwner->JumpMaxHoldTime;
+	const int32 RealJumpMaxCount = CharacterOwner->JumpMaxCount;
 	const bool bRealCrouch = bWantsToCrouch;
 	const bool bRealForceMaxAccel = bForceMaxAccel;
 	CharacterOwner->bClientWasFalling = (MovementMode == MOVE_Falling);
@@ -8154,6 +8156,8 @@ bool UCharacterMovementComponent::ClientUpdatePositionAfterServerUpdate()
 	CharacterOwner->bClientResimulateRootMotion = false;
 	CharacterOwner->bClientUpdating = false;
 	CharacterOwner->bPressedJump = bRealPressedJump || bPostReplayPressedJump;
+	CharacterOwner->JumpMaxHoldTime = RealJumpMaxHoldTime;
+	CharacterOwner->JumpMaxCount = RealJumpMaxCount;
 	bWantsToCrouch = bRealCrouch;
 	bForceMaxAccel = bRealForceMaxAccel;
 	bForceNextFloorCheck = true;
