Refactor component instance data cache to be able to store actor or component instance data

diff --git a/Engine/Config/BaseEngine.ini b/Engine/Config/BaseEngine.ini
--- a/Engine/Config/BaseEngine.ini
+++ b/Engine/Config/BaseEngine.ini
@@ -1058,6 +1058,8 @@ DefaultWorldPartitionHLODLayer=/Engine/Maps/Templates/HLODs/HLODLayer_Instanced.
 +FunctionRedirects=(OldName="KismetMathLibrary.NotEqual_FloatFloat", NewName="NotEqual_DoubleDouble")
 
 ;5.1
++ClassRedirects=(OldName="/Script/Engine.ActorComponentDuplicatedObjectData", NewName="/Script/Engine.DataCacheDuplicatedObjectData")
+
 +ClassRedirects=(OldName="RigVMTemplateNode", NewName="/Script/RigVMDeveloper.RigVMUnitNode")
 
 ; Blueprint real number support
diff --git a/Engine/Source/Runtime/Engine/Private/ComponentInstanceDataCache.cpp b/Engine/Source/Runtime/Engine/Private/ComponentInstanceDataCache.cpp
--- a/Engine/Source/Runtime/Engine/Private/ComponentInstanceDataCache.cpp
+++ b/Engine/Source/Runtime/Engine/Private/ComponentInstanceDataCache.cpp
@@ -16,42 +16,40 @@
 #include "Components/SceneComponent.h"
 #include "GameFramework/Actor.h"
 
-class FComponentPropertyWriter : public FObjectWriter
+class FDataCachePropertyWriter : public FObjectWriter
 {
 public:
-	FComponentPropertyWriter(const UActorComponent* InComponent, FActorComponentInstanceData& InActorInstanceData)
-		: FObjectWriter(InActorInstanceData.SavedProperties)
-		, Component(InComponent)
-		, ActorInstanceData(InActorInstanceData)
+	FDataCachePropertyWriter(const UObject* InCacheObject, FInstanceCacheDataBase& InInstanceData)
+		: FObjectWriter(InInstanceData.SavedProperties)
+		, CacheObject(InCacheObject)
+		, InstanceData(InInstanceData)
 	{
 		// Include properties that would normally skip tagged serialization (e.g. bulk serialization of array properties).
 		ArPortFlags |= PPF_ForceTaggedSerialization;
 
 		// Nested subobjects should be recursed in to
 		ArPortFlags |= PPF_DeepCompareInstances;
+	}
 
-		if (Component)
+	void SerializeProperties()
+	{
+		if (CacheObject)
 		{
-			UClass* ComponentClass = Component->GetClass();
-
-			Component->GetUCSModifiedProperties(PropertiesToSkip);
+			UClass* CacheObjectClass = CacheObject->GetClass();
+			CacheObjectClass->SerializeTaggedProperties(*this, (uint8*)CacheObject, CacheObjectClass, (uint8*)CacheObject->GetArchetype());
 
-			if (AActor* ComponentOwner = Component->GetOwner())
+			// Sort duplicated objects, if any, this ensures that lower depth duplicated objects are first in the array, which ensures proper creation order when deserializing
+			if (InstanceData.DuplicatedObjects.Num() > 0)
 			{
-				// If this is the owning Actor's root scene component, don't include relative transform properties. This is handled elsewhere.
-				if (Component == ComponentOwner->GetRootComponent())
+				InstanceData.DuplicatedObjects.StableSort([](const FDataCacheDuplicatedObjectData& One, const FDataCacheDuplicatedObjectData& Two) -> bool
 				{
-					PropertiesToSkip.Add(ComponentClass->FindPropertyByName(USceneComponent::GetRelativeLocationPropertyName()));
-					PropertiesToSkip.Add(ComponentClass->FindPropertyByName(USceneComponent::GetRelativeRotationPropertyName()));
-					PropertiesToSkip.Add(ComponentClass->FindPropertyByName(USceneComponent::GetRelativeScale3DPropertyName()));
-				}
+					return One.ObjectPathDepth < Two.ObjectPathDepth;
+				});
 			}
-
-			ComponentClass->SerializeTaggedProperties(*this, (uint8*)Component, ComponentClass, (uint8*)Component->GetArchetype());
 		}
 	}
 
-	virtual ~FComponentPropertyWriter()
+	virtual ~FDataCachePropertyWriter()
 	{
 		DuplicatedObjectAnnotation.RemoveAllAnnotations();
 	}
@@ -82,14 +80,14 @@ public:
 			{
 				Result = DupObjectInfo.DuplicatedObject.GetEvenIfUnreachable();
 			}
-			else if (Object->GetOuter() == Component)
+			else if (Object->GetOuter() == CacheObject)
 			{
-				Result = DuplicateObject(Object, ActorInstanceData.GetUniqueTransientPackage());
-				ActorInstanceData.DuplicatedObjects.Emplace(Result);
+				Result = DuplicateObject(Object, InstanceData.GetUniqueTransientPackage());
+				InstanceData.DuplicatedObjects.Emplace(Result);
 			}
 			else
 			{
-				check(Object->IsIn(Component));
+				check(Object->IsIn(CacheObject));
 
 				// Check to see if the object's outer is being duplicated.
 				UObject* DupOuter = GetDuplicatedObject(Object->GetOuter());
@@ -115,7 +113,7 @@ public:
 	virtual FArchive& operator<<(FName& Name) override
 	{
 		// store the reference to this name in the array instead of the global table, this allow to for persistence
-		int32 ReferenceIndex = ActorInstanceData.ReferencedNames.AddUnique(Name);
+		int32 ReferenceIndex = InstanceData.ReferencedNames.AddUnique(Name);
 		// save the name as an index in the referenced name array
 		*this << ReferenceIndex;
 
@@ -125,7 +123,7 @@ public:
 	virtual FArchive& operator<<(UObject*& Object) override
 	{
 		UObject* SerializedObject = Object;
-		if (Object && Component && Object->IsIn(Component))
+		if (Object && CacheObject && Object->IsIn(CacheObject))
 		{
 			SerializedObject = GetDuplicatedObject(Object);
 		}
@@ -134,8 +132,8 @@ public:
 		int32 ReferenceIndex = INDEX_NONE;
 		if (SerializedObject)
 		{
-			ReferenceIndex = ActorInstanceData.ReferencedObjects.Num();
-			ActorInstanceData.ReferencedObjects.Add(SerializedObject);
+			ReferenceIndex = InstanceData.ReferencedObjects.Num();
+			InstanceData.ReferencedObjects.Add(SerializedObject);
 		}
 		// save the pointer as an index in the referenced object array
 		*this << ReferenceIndex;
@@ -155,27 +153,66 @@ public:
 		return FArchiveUObject::SerializeObjectPtr(*this, ObjectPtr);
 	}
 
-private:
-	const UActorComponent* Component;
-	FActorComponentInstanceData& ActorInstanceData;
+protected:
 	TSet<const FProperty*> PropertiesToSkip;
+
+private:
+	const UObject* CacheObject;
 	FUObjectAnnotationSparse<FDuplicatedObject,false> DuplicatedObjectAnnotation;
+	FInstanceCacheDataBase& InstanceData;
 };
 
-class FComponentPropertyReader : public FObjectReader
+class FComponentPropertyWriter : public FDataCachePropertyWriter
 {
 public:
-	FComponentPropertyReader(UActorComponent* InComponent, FActorComponentInstanceData& InActorInstanceData)
-		: FObjectReader(InActorInstanceData.SavedProperties)
-		, ActorInstanceData(InActorInstanceData)
+
+	FComponentPropertyWriter(const UActorComponent* Component, FActorComponentInstanceData& InInstanceData)
+		: FDataCachePropertyWriter(Component, InInstanceData)
 	{
-		// Include properties that would normally skip tagged serialization (e.g. bulk serialization of array properties).
-		ArPortFlags |= PPF_ForceTaggedSerialization;
+		if (Component)
+		{
+			Component->GetUCSModifiedProperties(PropertiesToSkip);
 
-		InComponent->GetUCSModifiedProperties(PropertiesToSkip);
+			if (AActor* ComponentOwner = Component->GetOwner())
+			{
+				// If this is the owning Actor's root scene component, don't include relative transform properties. This is handled elsewhere.
+				if (Component == ComponentOwner->GetRootComponent())
+				{
+					UClass* ComponentClass = Component->GetClass();
+					PropertiesToSkip.Add(ComponentClass->FindPropertyByName(USceneComponent::GetRelativeLocationPropertyName()));
+					PropertiesToSkip.Add(ComponentClass->FindPropertyByName(USceneComponent::GetRelativeRotationPropertyName()));
+					PropertiesToSkip.Add(ComponentClass->FindPropertyByName(USceneComponent::GetRelativeScale3DPropertyName()));
+				}
+			}
 
-		UClass* Class = InComponent->GetClass();
-		Class->SerializeTaggedProperties(*this, (uint8*)InComponent, Class, (uint8*)InComponent->GetArchetype());
+			SerializeProperties();
+		}
+	}
+};
+
+class FActorPropertyWriter : public FDataCachePropertyWriter
+{
+public:
+
+	FActorPropertyWriter(const AActor* Actor, FActorInstanceData& InInstanceData)
+		: FDataCachePropertyWriter(Actor, InInstanceData)
+	{
+		if (Actor)
+		{
+			SerializeProperties();
+		}
+	}
+};
+
+class FDataCachePropertyReader : public FObjectReader
+{
+public:
+	FDataCachePropertyReader(FInstanceCacheDataBase& InInstanceData)
+		: FObjectReader(InInstanceData.SavedProperties)
+		, InstanceData(InInstanceData)
+	{
+		// Include properties that would normally skip tagged serialization (e.g. bulk serialization of array properties).
+		ArPortFlags |= PPF_ForceTaggedSerialization;
 	}
 
 	virtual bool ShouldSkipProperty(const FProperty* InProperty) const override
@@ -188,7 +225,7 @@ public:
 		// FName are serialized as Index in ActorInstanceData instead of the normal FName table
 		int32 ReferenceIndex = INDEX_NONE;
 		*this << ReferenceIndex;
-		Name = ActorInstanceData.ReferencedNames.IsValidIndex(ReferenceIndex) ? ActorInstanceData.ReferencedNames[ReferenceIndex] : FName();
+		Name = InstanceData.ReferencedNames.IsValidIndex(ReferenceIndex) ? InstanceData.ReferencedNames[ReferenceIndex] : FName();
 		return *this;
 	}
 
@@ -197,7 +234,7 @@ public:
 		// UObject pointer are serialized as Index in ActorInstanceData
 		int32 ReferenceIndex = INDEX_NONE;
 		*this << ReferenceIndex;
-		Object = ActorInstanceData.ReferencedObjects.IsValidIndex(ReferenceIndex) ? ToRawPtr(ActorInstanceData.ReferencedObjects[ReferenceIndex]) : nullptr;
+		Object = InstanceData.ReferencedObjects.IsValidIndex(ReferenceIndex) ? ToRawPtr(InstanceData.ReferencedObjects[ReferenceIndex]) : nullptr;
 		return *this;
 	}
 
@@ -213,12 +250,58 @@ public:
 		return FArchiveUObject::SerializeObjectPtr(*this, ObjectPtr);
 	}
 
-	FActorComponentInstanceData& ActorInstanceData;
+	FInstanceCacheDataBase& InstanceData;
 	TSet<const FProperty*> PropertiesToSkip;
 };
 
+class FComponentPropertyReader : public FDataCachePropertyReader
+{
+public:
+	FComponentPropertyReader(UActorComponent* InComponent, FActorComponentInstanceData& InInstanceData)
+		: FDataCachePropertyReader(InInstanceData)
+	{
+		InComponent->GetUCSModifiedProperties(PropertiesToSkip);
 
-FActorComponentDuplicatedObjectData::FActorComponentDuplicatedObjectData(UObject* InObject)
+		UClass* Class = InComponent->GetClass();
+		Class->SerializeTaggedProperties(*this, (uint8*)InComponent, Class, (uint8*)InComponent->GetArchetype());
+	}
+};
+
+class FActorPropertyReader : public FDataCachePropertyReader
+{
+public:
+	FActorPropertyReader(AActor* InActor, FActorInstanceData& InInstanceData)
+		: FDataCachePropertyReader(InInstanceData)
+	{
+		UClass* Class = InActor->GetClass();
+		Class->SerializeTaggedProperties(*this, (uint8*)InActor, Class, (uint8*)InActor->GetArchetype());
+	}
+
+	virtual bool ShouldSkipProperty(const FProperty* InProperty) const override
+	{
+		// We want to skip component properties, but not other instanced subobjects
+		if (InProperty->HasAnyPropertyFlags(CPF_InstancedReference | CPF_ContainsInstancedReference))
+		{
+			const FObjectProperty* ObjProp = CastField<const FObjectProperty>(InProperty);
+			if (!ObjProp)
+			{
+				if (const FArrayProperty* ArrayProp = CastField<const FArrayProperty>(InProperty))
+				{
+					ObjProp = CastField<const FObjectProperty>(ArrayProp->Inner);
+				}
+			}
+
+			if (ObjProp && ObjProp->PropertyClass->IsChildOf<UActorComponent>())
+			{
+				return true;
+			}
+		}
+
+		return FDataCachePropertyReader::ShouldSkipProperty(InProperty);
+	}
+};
+
+FDataCacheDuplicatedObjectData::FDataCacheDuplicatedObjectData(UObject* InObject)
 	: DuplicatedObject(InObject)
 	, ObjectPathDepth(0)
 {
@@ -231,7 +314,7 @@ FActorComponentDuplicatedObjectData::FActorComponentDuplicatedObjectData(UObject
 	}
 }
 
-bool FActorComponentDuplicatedObjectData::Serialize(FArchive& Ar)
+bool FDataCacheDuplicatedObjectData::Serialize(FArchive& Ar)
 {
 	enum class EVersion : uint8
 	{
@@ -301,6 +384,49 @@ bool FActorComponentDuplicatedObjectData::Serialize(FArchive& Ar)
 	return true;
 }
 
+FActorInstanceData::FActorInstanceData(const AActor* SourceActor)
+	: ActorClass(SourceActor->GetClass())
+{
+	FActorPropertyWriter PropertyWriter(SourceActor, *this);
+
+	// Cache off the length of an array that will come from SerializeTaggedProperties that had no properties saved in to it.
+	auto GetSizeOfEmptyArchive = []() -> int32
+	{
+		const AActor* DummyActor = GetDefault<AActor>();
+		FActorInstanceData DummyInstanceData;
+		FDataCachePropertyWriter NullWriter(nullptr, DummyInstanceData);
+		UClass* DummyActorClass = DummyActor->GetClass();
+
+		// By serializing the component with itself as its defaults we guarantee that no properties will be written out
+		DummyActorClass->SerializeTaggedProperties(NullWriter, (uint8*)DummyActor, DummyActorClass, (uint8*)DummyActor);
+
+		check(DummyInstanceData.GetDuplicatedObjects().Num() == 0 && DummyInstanceData.GetReferencedObjects().Num() == 0);
+		return DummyInstanceData.GetSavedProperties().Num();
+	};
+
+	static const int32 SizeOfEmptyArchive = GetSizeOfEmptyArchive();
+
+	// SerializeTaggedProperties will always put a sentinel NAME_None at the end of the Archive. 
+	// If that is the only thing in the buffer then empty it because we want to know that we haven't stored anything.
+	if (SavedProperties.Num() == SizeOfEmptyArchive)
+	{
+		SavedProperties.Empty();
+	}
+}
+
+const UClass* FActorInstanceData::GetActorClass() const
+{
+	return *ActorClass;
+}
+
+void FActorInstanceData::ApplyToActor(AActor* Actor, const ECacheApplyPhase CacheApplyPhase)
+{
+	if ((CacheApplyPhase == ECacheApplyPhase::PostUserConstructionScript || CacheApplyPhase == ECacheApplyPhase::NonConstructionScript) && SavedProperties.Num() > 0)
+	{
+		FActorPropertyReader ActorPropertyReader(Actor, *this);
+	}
+}
+
 FActorComponentInstanceData::FActorComponentInstanceData()
 	: SourceComponentTemplate(nullptr)
 	, SourceComponentCreationMethod(EComponentCreationMethod::Native)
@@ -330,8 +456,8 @@ FActorComponentInstanceData::FActorComponentInstanceData(const UActorComponent*
 			// By serializing the component with itself as its defaults we guarantee that no properties will be written out
 			ComponentClass->SerializeTaggedProperties(NullWriter, (uint8*)DummyComponent, ComponentClass, (uint8*)DummyComponent);
 
-			check(DummyInstanceData.DuplicatedObjects.Num() == 0 && DummyInstanceData.ReferencedObjects.Num() == 0);
-			return DummyInstanceData.SavedProperties.Num();
+			check(DummyInstanceData.GetDuplicatedObjects().Num() == 0 && DummyInstanceData.GetReferencedObjects().Num() == 0);
+			return DummyInstanceData.GetSavedProperties().Num();
 		};
 
 		static const int32 SizeOfEmptyArchive = GetSizeOfEmptyArchive();
@@ -343,15 +469,6 @@ FActorComponentInstanceData::FActorComponentInstanceData(const UActorComponent*
 			SavedProperties.Empty();
 		}
 	}
-
-	// Sort duplicated objects, if any, this ensures that lower depth duplicated objects are first in the array, which ensures proper creation order when deserializing
-	if (DuplicatedObjects.Num() > 0)
-	{
-		DuplicatedObjects.StableSort([](const FActorComponentDuplicatedObjectData& One, const FActorComponentDuplicatedObjectData& Two) -> bool
-		{
-			return One.ObjectPathDepth < Two.ObjectPathDepth;
-		});
-	}
 }
 
 bool FActorComponentInstanceData::MatchesComponent(const UActorComponent* Component, const UObject* ComponentTemplate, const TMap<UActorComponent*, const UObject*>& ComponentToArchetypeMap) const
@@ -389,7 +506,7 @@ void FActorComponentInstanceData::ApplyToComponent(UActorComponent* Component, c
 			Component->ClearUCSModifiedProperties();
 		}
 
-		for (const FActorComponentDuplicatedObjectData& DuplicatedObjectData : DuplicatedObjects)
+		for (const FDataCacheDuplicatedObjectData& DuplicatedObjectData : GetDuplicatedObjects())
 		{
 			if (DuplicatedObjectData.DuplicatedObject)
 			{
@@ -407,19 +524,24 @@ void FActorComponentInstanceData::ApplyToComponent(UActorComponent* Component, c
 	}
 }
 
-UObject* FActorComponentInstanceData::GetUniqueTransientPackage()
+UObject* FInstanceCacheDataBase::GetUniqueTransientPackage()
 {
 	if (UniqueTransientPackage.DuplicatedObject == nullptr)
 	{
-		UniqueTransientPackage = FActorComponentDuplicatedObjectData(NewObject<UActorComponentInstanceDataTransientOuter>(GetTransientPackage()));
+		UniqueTransientPackage = FDataCacheDuplicatedObjectData(NewObject<UActorComponentInstanceDataTransientOuter>(GetTransientPackage()));
 	}
 	return UniqueTransientPackage.DuplicatedObject;
 }
 
+void FInstanceCacheDataBase::AddReferencedObjects(FReferenceCollector& Collector)
+{
+	Collector.AddReferencedObjects(ReferencedObjects);
+}
+
 void FActorComponentInstanceData::AddReferencedObjects(FReferenceCollector& Collector)
 {
+	Super::AddReferencedObjects(Collector);
 	Collector.AddReferencedObject(SourceComponentTemplate);
-	Collector.AddReferencedObjects(ReferencedObjects);
 }
 
 FComponentInstanceDataCache::FComponentInstanceDataCache(const AActor* Actor)
diff --git a/Engine/Source/Runtime/Engine/Public/ComponentInstanceDataCache.h b/Engine/Source/Runtime/Engine/Public/ComponentInstanceDataCache.h
--- a/Engine/Source/Runtime/Engine/Public/ComponentInstanceDataCache.h
+++ b/Engine/Source/Runtime/Engine/Public/ComponentInstanceDataCache.h
@@ -3,6 +3,7 @@
 #pragma once
 
 #include "CoreMinimal.h"
+#include "Templates/SubclassOf.h"
 #include "UObject/Object.h"
 #include "UObject/StructOnScope.h"
 #include "ComponentInstanceDataCache.generated.h"
@@ -33,12 +34,15 @@ enum class EComponentCreationMethod : uint8
 	Instance,
 };
 
+UE_DEPRECATED(5.1, "FActorComponentDuplicatedObjectData has been renamed to FDataCacheDuplicatedObjectData")
+typedef FDataCacheDuplicatedObjectData FActorComponentDuplicatedObjectData;
+
 USTRUCT()
-struct FActorComponentDuplicatedObjectData
+struct FDataCacheDuplicatedObjectData
 {
 	GENERATED_BODY()
 
-	FActorComponentDuplicatedObjectData(UObject* InObject = nullptr);
+	FDataCacheDuplicatedObjectData(UObject* InObject = nullptr);
 
 	bool Serialize(FArchive& Ar);
 
@@ -51,7 +55,7 @@ struct FActorComponentDuplicatedObjectData
 
 // Trait to signal ActorCompomentInstanceData duplicated objects uses a serialize function
 template<>
-struct TStructOpsTypeTraits<FActorComponentDuplicatedObjectData> : public TStructOpsTypeTraitsBase2<FActorComponentDuplicatedObjectData>
+struct TStructOpsTypeTraits<FDataCacheDuplicatedObjectData> : public TStructOpsTypeTraitsBase2<FDataCacheDuplicatedObjectData>
 {
 	enum
 	{
@@ -66,17 +70,60 @@ class UActorComponentInstanceDataTransientOuter : public UObject
 	GENERATED_BODY()
 };
 
+/** Base class for instance cached data of a particular type. */
+USTRUCT()
+struct ENGINE_API FInstanceCacheDataBase
+{
+	GENERATED_BODY()
+
+	virtual ~FInstanceCacheDataBase() = default;
+
+	virtual void AddReferencedObjects(FReferenceCollector& Collector);
+
+	/** Get (or create) the unique transient outer for the duplicated objects created for this object */
+	UObject* GetUniqueTransientPackage();
+
+	const TArray<FDataCacheDuplicatedObjectData>& GetDuplicatedObjects() const { return DuplicatedObjects; }
+	const TArray<TObjectPtr<UObject>>& GetReferencedObjects() const { return ReferencedObjects; }
+	const TArray<uint8>& GetSavedProperties() const { return SavedProperties; }
+
+protected:
+	UPROPERTY()
+	TArray<uint8> SavedProperties;
+
+private:
+	friend class FDataCachePropertyWriter;
+	friend class FDataCachePropertyReader;
+
+	/** 
+	 * A unique outer created in the transient package to act as outer for this object's duplicated objects 
+	 * to avoid name conflicts of objects that already exist in the transient package
+	 */
+	UPROPERTY()
+	FDataCacheDuplicatedObjectData UniqueTransientPackage;
+
+	// Duplicated objects created when saving instance properties
+	UPROPERTY()
+	TArray<FDataCacheDuplicatedObjectData> DuplicatedObjects;
+
+	// Referenced objects in instance saved properties
+	UPROPERTY()
+	TArray<TObjectPtr<UObject>> ReferencedObjects;
+
+	// Referenced names in instance saved properties
+	UPROPERTY()
+	TArray<FName> ReferencedNames;
+};
+
 /** Base class for component instance cached data of a particular type. */
 USTRUCT()
-struct ENGINE_API FActorComponentInstanceData
+struct ENGINE_API FActorComponentInstanceData : public FInstanceCacheDataBase
 {
 	GENERATED_BODY()
 public:
 	FActorComponentInstanceData();
 	FActorComponentInstanceData(const UActorComponent* SourceComponent);
 
-	virtual ~FActorComponentInstanceData() = default;
-
 	/** Determines whether this component instance data matches the component */
 	bool MatchesComponent(const UActorComponent* Component, const UObject* ComponentTemplate, const TMap<UActorComponent*, const UObject*>& ComponentToArchetypeMap) const;
 
@@ -89,18 +136,13 @@ public:
 	/** Replaces any references to old instances during Actor reinstancing */
 	virtual void FindAndReplaceInstances(const TMap<UObject*, UObject*>& OldToNewInstanceMap) { };
 
-	virtual void AddReferencedObjects(FReferenceCollector& Collector);
+	virtual void AddReferencedObjects(FReferenceCollector& Collector) override;
 
 	const UClass* GetComponentClass() const { return SourceComponentTemplate ? SourceComponentTemplate->GetClass() : nullptr; }
 
 	const UObject* GetComponentTemplate() const { return SourceComponentTemplate; }
 
-	/** Get (or create) the unique transient outer for the duplicated objects created for this component */
-	UObject* GetUniqueTransientPackage();
-
 protected:
-	friend class FComponentPropertyWriter;
-	friend class FComponentPropertyReader;
 
 	/** The template used to create the source component */
 	UPROPERTY()
@@ -114,28 +156,29 @@ protected:
 	when filtered to just that component type */
 	UPROPERTY()
 	int32 SourceComponentTypeSerializedIndex;
+};
 
-	UPROPERTY()
-	TArray<uint8> SavedProperties;
+/** Per instance data to be persisted for a given actor */
+USTRUCT()
+struct ENGINE_API FActorInstanceData : public FInstanceCacheDataBase
+{
+	GENERATED_BODY()
+public:
+	FActorInstanceData() = default;
+	FActorInstanceData(const AActor* SourceActor);
 
-	/** 
-	 * A unique outer created in the transient package to act as outer for this component's duplicated objects 
-	 * to avoid name conflicts of objects that already exist in the transient package
-	 */
-	UPROPERTY()
-	FActorComponentDuplicatedObjectData UniqueTransientPackage;
+	const UClass* GetActorClass() const;
 
-	// Duplicated objects created when saving component instance properties
-	UPROPERTY()
-	TArray<FActorComponentDuplicatedObjectData> DuplicatedObjects;
+	bool HasInstanceData() const { return GetSavedProperties().Num() > 0; }
 
-	// Referenced objects in component instance saved properties
-	UPROPERTY()
-	TArray<TObjectPtr<UObject>> ReferencedObjects;
+	/** Iterates over an Actor's components and applies the stored component instance data to each */
+	void ApplyToActor(AActor* Actor, const ECacheApplyPhase CacheApplyPhase);
 
-	// Referenced names in component instance saved properties
+protected:
+
+	/** The class of the actor that the instance data is for */
 	UPROPERTY()
-	TArray<FName> ReferencedNames;
+	TSubclassOf<AActor> ActorClass;
 };
 
 /** 
