Remove ensure from FSimpleSuspensionHelpers::ComputeSprungMasses, opting instead for populating an optional error string and relying on success/failure return bool.

diff --git a/Engine/Source/Runtime/Engine/Private/PhysicsEngine/SimpleSuspension.cpp b/Engine/Source/Runtime/Engine/Private/PhysicsEngine/SimpleSuspension.cpp
--- a/Engine/Source/Runtime/Engine/Private/PhysicsEngine/SimpleSuspension.cpp
+++ b/Engine/Source/Runtime/Engine/Private/PhysicsEngine/SimpleSuspension.cpp
@@ -44,14 +44,14 @@ void FSimpleSuspension::Update(const FTransform& LocalToWorld, const FVector& Li
 	FSimpleSuspensionHelpers::ComputeSuspensionForces(LinearVelocity, AngularVelocityRad, SuspensionState, SuspensionParams.SpringParams, SuspensionState);
 }
 
-bool FSimpleSuspensionHelpers::ComputeSingleAxisLambda(const FVector::FReal AxisDot, const FVector::FReal SumAxis, const uint32 Count, TArray<FVector::FReal, TFixedAllocator<2>>& Lambdas)
+bool FSimpleSuspensionHelpers::ComputeSingleAxisLambda(const FVector::FReal AxisDot, const FVector::FReal SumAxis, const uint32 Count, TArray<FVector::FReal, TFixedAllocator<2>>& Lambdas, FString* ErrMsg)
 {
 	using Chaos::FReal;
 
 	//compute determinant
 	const FReal DetLL = AxisDot * Count - SumAxis * SumAxis;
-	if (!ensureMsgf(!FMath::IsNearlyZero(DetLL, SUSPENSION_SMALL_NUMBER),
-		TEXT("Spring configuration is invalid! Please make sure no two springs are at the same location.")))
+	if (!ErrCheck(!FMath::IsNearlyZero(DetLL, SUSPENSION_SMALL_NUMBER),
+		TEXT("Spring configuration is invalid! Please make sure no two springs are at the same location."), ErrMsg))
 	{
 		return false;
 	}
@@ -73,7 +73,7 @@ bool FSimpleSuspensionHelpers::ComputeSingleAxisLambda(const FVector::FReal Axis
 	return true;
 }
 
-bool FSimpleSuspensionHelpers::ComputeSprungMasses(const TArray<FVector>& MassSpringPositions, const float TotalMass, TArray<float>& OutSprungMasses)
+bool FSimpleSuspensionHelpers::ComputeSprungMasses(const TArray<FVector>& MassSpringPositions, const float TotalMass, TArray<float>& OutSprungMasses, FString* ErrMsg)
 {
 	using Chaos::FReal;
 	/*
@@ -95,12 +95,12 @@ bool FSimpleSuspensionHelpers::ComputeSprungMasses(const TArray<FVector>& MassSp
 	OutSprungMasses.Reserve(Count);
 
 	// Check essential values
-	if (!ensureMsgf(Count > 0, TEXT("Must have at least one spring to compute sprung masses.")))
+	if (!ErrCheck(Count > 0, TEXT("Must have at least one spring to compute sprung masses."), ErrMsg))
 	{
 		return false;
 	}
 
-	if (!ensureMsgf(TotalMass > UE_SMALL_NUMBER, TEXT("Total mass must be greater than zero to compute sprung masses.")))
+	if (!ErrCheck(TotalMass > SMALL_NUMBER, TEXT("Total mass must be greater than zero to compute sprung masses."), ErrMsg))
 	{
 		return false;
 	}
@@ -156,8 +156,8 @@ bool FSimpleSuspensionHelpers::ComputeSprungMasses(const TArray<FVector>& MassSp
 		const FReal DirDotA = (DirX * AX) + (DirY * AY);
 		OutSprungMasses[0] = -TotalMass * DirDotA * DistInv;
 		OutSprungMasses[1] = TotalMass - OutSprungMasses[0];
-		if (ensureMsgf(OutSprungMasses[0] >= 0.f, TEXT("Spring configuration is invalid! Please make sure the center of mass is located between the springs.")) &&
-			ensureMsgf(OutSprungMasses[1] >= 0.f, TEXT("Spring configuration is invalid! Please make sure the center of mass is located between the springs.")))
+		if (ErrCheck(OutSprungMasses[0] >= 0.f, TEXT("Spring configuration is invalid! Please make sure the center of mass is located between the springs."), ErrMsg) &&
+			ErrCheck(OutSprungMasses[1] >= 0.f, TEXT("Spring configuration is invalid! Please make sure the center of mass is located between the springs."), ErrMsg))
 		{
 			return true;
 		}
@@ -303,8 +303,8 @@ bool FSimpleSuspensionHelpers::ComputeSprungMasses(const TArray<FVector>& MassSp
 			- (XDotY * XDotY * Count);
 
 		// Make sure the matrix is invertible!
-		if (!ensureMsgf(!FMath::IsNearlyZero(DetLL, SUSPENSION_SMALL_NUMBER),
-			TEXT("Spring configuration is invalid! Please make sure no two springs are at the same location.")))
+		if (!ErrCheck(!FMath::IsNearlyZero(DetLL, SUSPENSION_SMALL_NUMBER),
+			TEXT("Spring configuration is invalid! Please make sure no two springs are at the same location."), ErrMsg))
 		{
 			return false;
 		}
@@ -334,7 +334,7 @@ bool FSimpleSuspensionHelpers::ComputeSprungMasses(const TArray<FVector>& MassSp
 		const FReal Y = MassSpringPositions[Index].Y;
 		const FReal LLambda = (X * Lambda0) + (Y * Lambda1) + Lambda2;
 		OutSprungMasses[Index] = AverageMass - (0.5f * LLambda);
-		if (!ensureMsgf(OutSprungMasses[Index] >= 0.f, TEXT("Spring configuration is invalid! Please make sure the center of mass is located inside the area covered by the springs.")))
+		if (!ErrCheck(OutSprungMasses[Index] >= 0.f, TEXT("Spring configuration is invalid! Please make sure the center of mass is located inside the area covered by the springs."), ErrMsg))
 		{
 			return false;
 		}
@@ -343,7 +343,7 @@ bool FSimpleSuspensionHelpers::ComputeSprungMasses(const TArray<FVector>& MassSp
 	return true;
 }
 
-bool FSimpleSuspensionHelpers::ComputeSprungMasses(const TArray<FVector>& LocalSpringPositions, const FVector& LocalCenterOfMass, const float TotalMass, TArray<float>& OutSprungMasses)
+bool FSimpleSuspensionHelpers::ComputeSprungMasses(const TArray<FVector>& LocalSpringPositions, const FVector& LocalCenterOfMass, const float TotalMass, TArray<float>& OutSprungMasses, FString* ErrMsg)
 {
 	// Compute support origin's in center of mass space
 	const int32 SpringCount = LocalSpringPositions.Num();
@@ -566,3 +566,9 @@ void FSimpleSuspensionHelpers::IntegrateSprings(const float DeltaTime, const TAr
 		IntegrateSpring(DeltaTime, SpringDisplacements[Index], SpringVelocities[Index], SuspensionParams[Index], SprungMasses[Index], OutNewSpringDisplacements[Index], OutNewSpringVelocities[Index]);
 	}
 }
+
+bool FSimpleSuspensionHelpers::ErrCheck(const bool bCondition, const FString& Message, FString* ErrMsg)
+{
+	if (ErrMsg != nullptr && !bCondition) { *ErrMsg = Message; }
+	return bCondition;
+}
diff --git a/Engine/Source/Runtime/Engine/Public/Physics/SimpleSuspension.h b/Engine/Source/Runtime/Engine/Public/Physics/SimpleSuspension.h
--- a/Engine/Source/Runtime/Engine/Public/Physics/SimpleSuspension.h
+++ b/Engine/Source/Runtime/Engine/Public/Physics/SimpleSuspension.h
@@ -131,15 +131,15 @@ struct FSimpleSuspensionHelpers
 		configuration may not result in a stable suspension system -
 		a bicycle or pogostick, for example, which is not perfectly centered
 		may have a valid sprung mass configuration without being stable. */
-	ENGINE_API static bool ComputeSprungMasses(const TArray<FVector>& MassSpringPositions, const float TotalMass, TArray<float>& OutSprungMasses);
+	ENGINE_API static bool ComputeSprungMasses(const TArray<FVector>& MassSpringPositions, const float TotalMass, TArray<float>& OutSprungMasses, FString* ErrMsg = nullptr);
 
 	/** Same as above, but allows the caller to specify spring locations
 		in a local space which is not necessarily originated at the center
 		of mass. */
-	ENGINE_API static bool ComputeSprungMasses(const TArray<FVector>& LocalSpringPositions, const FVector& LocalCenterOfMass, const float TotalMass, TArray<float>& OutSprungMasses);
+	ENGINE_API static bool ComputeSprungMasses(const TArray<FVector>& LocalSpringPositions, const FVector& LocalCenterOfMass, const float TotalMass, TArray<float>& OutSprungMasses, FString* ErrMsg = nullptr);
 
 	/** Calculates the Lambdas in the single axis case*/
-	static bool ComputeSingleAxisLambda(const FVector::FReal AxisDot, const FVector::FReal SumAxis, const uint32 Count, TArray<FVector::FReal, TFixedAllocator<2>>& Lambdas);
+	static bool ComputeSingleAxisLambda(const FVector::FReal AxisDot, const FVector::FReal SumAxis, const uint32 Count, TArray<FVector::FReal, TFixedAllocator<2>>& Lambdas, FString* ErrMsg = nullptr);
 
 	/** Given a sprung mass and a spring stiffness, compute the natural
 		frequency of the spring */
@@ -200,4 +200,8 @@ struct FSimpleSuspensionHelpers
 	/** Performs a fully implicit damped spring integration to compute next
 		positions and velocities. */
 	static void IntegrateSprings(const float DeltaTime, const TArray<float>& SpringDisplacements, const TArray<float>& SpringVelocities, const TArray<FSimpleSuspensionSpringParams>& SuspensionParams, const TArray<float>& SprungMasses, TArray<float>& OutNewSpringDisplacements, TArray<float>& OutNewSpringVelocities);
+
+private:
+
+	static bool ErrCheck(const bool bCondition, const FString& Message, FString* ErrMsg);
 };
