Fix merge error - replace SMALL_NUMBER with UE_SMALL_NUMBER from 20262621

diff --git a/Engine/Source/Runtime/Engine/Private/PhysicsEngine/SimpleSuspension.cpp b/Engine/Source/Runtime/Engine/Private/PhysicsEngine/SimpleSuspension.cpp
--- a/Engine/Source/Runtime/Engine/Private/PhysicsEngine/SimpleSuspension.cpp
+++ b/Engine/Source/Runtime/Engine/Private/PhysicsEngine/SimpleSuspension.cpp
@@ -100,7 +100,7 @@ bool FSimpleSuspensionHelpers::ComputeSprungMasses(const TArray<FVector>& MassSp
 		return false;
 	}
 
-	if (!ErrCheck(TotalMass > SMALL_NUMBER, TEXT("Total mass must be greater than zero to compute sprung masses."), ErrMsg))
+	if (!ErrCheck(TotalMass > UE_SMALL_NUMBER, TEXT("Total mass must be greater than zero to compute sprung masses."), ErrMsg))
 	{
 		return false;
 	}
