Add object names to check
#rnx

diff --git a/Engine/Source/Runtime/Engine/Private/PackageMapClient.cpp b/Engine/Source/Runtime/Engine/Private/PackageMapClient.cpp
--- a/Engine/Source/Runtime/Engine/Private/PackageMapClient.cpp
+++ b/Engine/Source/Runtime/Engine/Private/PackageMapClient.cpp
@@ -3251,7 +3251,8 @@ UObject* FNetGUIDCache::GetObjectFromNetGUID( const FNetworkGUID& NetGUID, const
 	if ( Object != NULL )
 	{
 		// Either the name should match, or this is dynamic, or we're on the server
-		check( Object->GetFName() == CacheObjectPtr->PathName || NetGUID.IsDynamic() || IsNetGUIDAuthority() );
+		checkf( Object->GetFName() == CacheObjectPtr->PathName || NetGUID.IsDynamic() || IsNetGUIDAuthority(),
+		        TEXT("ObjectName: '%s', CacheObjectPathName: '%s'"), *Object->GetName(), *CacheObjectPtr->PathName.ToString());
 		return Object;
 	}
 
