Improved anim BP re-instancing

This performs a more targeted re-initialization of animation systems, rather than just tearing down and re-creating the skeletal mesh component.
New behavior is opt-out on CVar in case there are any corner cases where we were relying on the old behavior

Also fixes external instances getting disconnected on anim BP recompilation

diff --git a/Engine/Plugins/Animation/ControlRig/Source/ControlRig/Private/AnimNode_ControlRig.cpp b/Engine/Plugins/Animation/ControlRig/Source/ControlRig/Private/AnimNode_ControlRig.cpp
--- a/Engine/Plugins/Animation/ControlRig/Source/ControlRig/Private/AnimNode_ControlRig.cpp
+++ b/Engine/Plugins/Animation/ControlRig/Source/ControlRig/Private/AnimNode_ControlRig.cpp
@@ -43,20 +43,9 @@ void FAnimNode_ControlRig::OnInitializeAnimInstance(const FAnimInstanceProxy* In
 
 	FAnimNode_ControlRigBase::OnInitializeAnimInstance(InProxy, InAnimInstance);
 
-#if WITH_EDITOR
-	FCoreUObjectDelegates::OnObjectsReplaced.AddRaw(this, &FAnimNode_ControlRig::OnObjectsReplaced);
-#endif // WITH_EDITOR
-
 	InitializeProperties(InAnimInstance, GetTargetClass());
 }
 
-FAnimNode_ControlRig::~FAnimNode_ControlRig()
-{
-#if WITH_EDITOR
-	FCoreUObjectDelegates::OnObjectsReplaced.RemoveAll(this);
-#endif // WITH_EDITOR
-}
-
 void FAnimNode_ControlRig::GatherDebugData(FNodeDebugData& DebugData)
 {
 	DECLARE_SCOPE_HIERARCHICAL_COUNTER_FUNC()
@@ -591,19 +580,3 @@ void FAnimNode_ControlRig::PropagateInputProperties(const UObject* InSourceInsta
 		}
 	}
 }
-
-#if WITH_EDITOR
-void FAnimNode_ControlRig::OnObjectsReplaced(const TMap<UObject*, UObject*>& OldToNewInstanceMap)
-{
-	if (ControlRig)
-	{
-		UObject* const* NewFound = OldToNewInstanceMap.Find(ControlRig);
-
-		if (NewFound)
-		{
-			// recache the properties
-			bReinitializeProperties = true;
-		}
-	}
-}
-#endif	// #if WITH_EDITOR
\ No newline at end of file
diff --git a/Engine/Plugins/Animation/ControlRig/Source/ControlRig/Public/AnimNode_ControlRig.h b/Engine/Plugins/Animation/ControlRig/Source/ControlRig/Public/AnimNode_ControlRig.h
--- a/Engine/Plugins/Animation/ControlRig/Source/ControlRig/Public/AnimNode_ControlRig.h
+++ b/Engine/Plugins/Animation/ControlRig/Source/ControlRig/Public/AnimNode_ControlRig.h
@@ -18,7 +18,6 @@ struct CONTROLRIG_API FAnimNode_ControlRig : public FAnimNode_ControlRigBase
 	GENERATED_BODY()
 
 	FAnimNode_ControlRig();
-	virtual ~FAnimNode_ControlRig();
 
 	UControlRig* GetControlRig() const { return ControlRig; }
 
@@ -93,9 +92,6 @@ private:
 	UPROPERTY(EditAnywhere, Category = Performance, meta = (DisplayName = "LOD Threshold"))
 	int32 LODThreshold;
 
-#if WITH_EDITOR
-	void OnObjectsReplaced(const TMap<UObject*, UObject*>& OldToNewInstanceMap);
-#endif // WITH_EDITOR
 protected:
 	virtual UClass* GetTargetClass() const override { return *ControlRigClass; }
 	virtual void UpdateInput(UControlRig* InControlRig, const FPoseContext& InOutput) override;
diff --git a/Engine/Source/Editor/Persona/Private/PreviewSceneCustomizations.cpp b/Engine/Source/Editor/Persona/Private/PreviewSceneCustomizations.cpp
--- a/Engine/Source/Editor/Persona/Private/PreviewSceneCustomizations.cpp
+++ b/Engine/Source/Editor/Persona/Private/PreviewSceneCustomizations.cpp
@@ -30,6 +30,7 @@
 #include "Algo/Sort.h"
 #include "ScopedTransaction.h"
 #include "Features/IModularFeatures.h"
+#include "AnimPreviewInstance.h"
 
 #define LOCTEXT_NAMESPACE "PreviewSceneCustomizations"
 
@@ -642,7 +643,13 @@ ECheckBoxState FPreviewSceneDescriptionCustomization::HandleUseCustomAnimBPIsChe
 
 void FPreviewSceneDescriptionCustomization::HandleAnimBlueprintCompiled(UBlueprint*)
 {
-	ReinitializePreviewController();
+	// Only re-initialize controller if we are not debugging an external instance.
+	// If we switch at this point then we will disconnect from the external instance
+	TSharedPtr<FAnimationEditorPreviewScene> PreviewScenePtr = PreviewScene.Pin();
+	if(PreviewScenePtr->GetPreviewMeshComponent()->PreviewInstance == nullptr || PreviewScenePtr->GetPreviewMeshComponent()->PreviewInstance->GetDebugSkeletalMeshComponent() == nullptr)
+	{
+		ReinitializePreviewController();
+	}
 }
 
 void FPreviewSceneDescriptionCustomization::ReinitializePreviewController()
diff --git a/Engine/Source/Editor/UnrealEd/Private/Kismet2/KismetReinstanceUtilities.cpp b/Engine/Source/Editor/UnrealEd/Private/Kismet2/KismetReinstanceUtilities.cpp
--- a/Engine/Source/Editor/UnrealEd/Private/Kismet2/KismetReinstanceUtilities.cpp
+++ b/Engine/Source/Editor/UnrealEd/Private/Kismet2/KismetReinstanceUtilities.cpp
@@ -55,6 +55,13 @@ DECLARE_CYCLE_STAT(TEXT("Reinstance Objects"), EKismetCompilerStats_ReinstanceOb
 DECLARE_CYCLE_STAT(TEXT("Refresh Dependent Blueprints In Reinstancer"), EKismetCompilerStats_RefreshDependentBlueprintsInReinstancer, STATGROUP_KismetCompiler);
 DECLARE_CYCLE_STAT(TEXT("Recreate UberGraphPersistentFrame"), EKismetCompilerStats_RecreateUberGraphPersistentFrame, STATGROUP_KismetCompiler);
 
+bool GUseLegacyAnimInstanceReinstancingBehavior = false;
+static FAutoConsoleVariableRef CVarUseLegacyAnimInstanceReinstancingBehavior(
+	TEXT("bp.UseLegacyAnimInstanceReinstancingBehavior"),
+	GUseLegacyAnimInstanceReinstancingBehavior,
+	TEXT("Use the legacy re-instancing behavior for anim instances where the instance is destroyed and re-created.")
+);
+
 struct FReplaceReferenceHelper
 {
 	static void IncludeCDO(UClass* OldClass, UClass* NewClass, TMap<UObject*, UObject*> &OldToNewInstanceMap, TArray<UObject*> &SourceObjects, UObject* OriginalCDO)
@@ -2558,14 +2565,17 @@ void FBlueprintCompileReinstancer::ReplaceInstancesOfClass_Inner(TMap<UClass*, U
 
 			if (UAnimInstance* AnimTree = Cast<UAnimInstance>(*NewObject))
 			{
-				// Initialising the anim instance isn't enough to correctly set up the skeletal mesh again in a
-				// paused world, need to initialise the skeletal mesh component that contains the anim instance.
 				if (USkeletalMeshComponent* SkelComponent = Cast<USkeletalMeshComponent>(AnimTree->GetOuter()))
 				{
-					SkelComponent->ClearAnimScriptInstance();
-					SkelComponent->InitAnim(true);
-					// compile change ignores motion vector, so ignore this. 
-					SkelComponent->ClearMotionVector();
+					if(GUseLegacyAnimInstanceReinstancingBehavior)
+					{
+						// Legacy behavior - destroy and re-create the anim instance
+						SkelComponent->ClearAnimScriptInstance();
+						SkelComponent->InitAnim(true);
+
+						// compile change ignores motion vector, so ignore this. 
+						SkelComponent->ClearMotionVector();
+					}
 				}
 			}
 		}
diff --git a/Engine/Source/Runtime/Engine/Classes/Animation/AnimInstance.h b/Engine/Source/Runtime/Engine/Classes/Animation/AnimInstance.h
--- a/Engine/Source/Runtime/Engine/Classes/Animation/AnimInstance.h
+++ b/Engine/Source/Runtime/Engine/Classes/Animation/AnimInstance.h
@@ -1635,4 +1635,9 @@ public:
 public:
 	/** Dispatch AnimEvents (AnimNotifies, Montage Events) queued during UpdateAnimation() */
 	void DispatchQueuedAnimEvents();
+
+#if WITH_EDITOR
+	// Helper function to handle reinstancing
+	void HandleAnimInstanceReplaced();
+#endif
 };
diff --git a/Engine/Source/Runtime/Engine/Classes/Animation/AnimNode_CustomProperty.h b/Engine/Source/Runtime/Engine/Classes/Animation/AnimNode_CustomProperty.h
--- a/Engine/Source/Runtime/Engine/Classes/Animation/AnimNode_CustomProperty.h
+++ b/Engine/Source/Runtime/Engine/Classes/Animation/AnimNode_CustomProperty.h
@@ -29,6 +29,7 @@ struct ENGINE_API FAnimNode_CustomProperty : public FAnimNode_Base
 public:
 
 	FAnimNode_CustomProperty();
+	~FAnimNode_CustomProperty();
 
 	/* Set Target Instance */
 	void SetTargetInstance(UObject* InInstance);
@@ -89,6 +90,11 @@ protected:
 	/** Get Target Class */
 	virtual UClass* GetTargetClass() const PURE_VIRTUAL(FAnimNode_CustomProperty::GetTargetClass, return nullptr;);
 
+#if WITH_EDITOR
+	// Handle object replacement in editor
+	virtual void HandleObjectsReplaced(const TMap<UObject*, UObject*>& OldToNewInstanceMap);
+#endif	// #if WITH_EDITOR
+	
 	friend class UAnimGraphNode_CustomProperty;
 };
 
diff --git a/Engine/Source/Runtime/Engine/Classes/Animation/AnimNode_LinkedAnimGraph.h b/Engine/Source/Runtime/Engine/Classes/Animation/AnimNode_LinkedAnimGraph.h
--- a/Engine/Source/Runtime/Engine/Classes/Animation/AnimNode_LinkedAnimGraph.h
+++ b/Engine/Source/Runtime/Engine/Classes/Animation/AnimNode_LinkedAnimGraph.h
@@ -113,12 +113,16 @@ protected:
 	// Shutdown the currently running instance
 	void TeardownInstance(const UAnimInstance* InOwningAnimInstance);
 
-	// Get Target Class
+	// FAnimNode_CustomProperty interface
 	virtual UClass* GetTargetClass() const override 
 	{
 		return *InstanceClass;
 	}
 
+#if WITH_EDITOR
+	virtual void HandleObjectsReplaced(const TMap<UObject*, UObject*>& OldToNewInstanceMap) override;
+#endif	// #if WITH_EDITOR
+
 	/** Link up pose links dynamically with linked instance */
 	void DynamicLink(UAnimInstance* InOwningAnimInstance);
 
diff --git a/Engine/Source/Runtime/Engine/Classes/Animation/AnimNode_LinkedAnimLayer.h b/Engine/Source/Runtime/Engine/Classes/Animation/AnimNode_LinkedAnimLayer.h
--- a/Engine/Source/Runtime/Engine/Classes/Animation/AnimNode_LinkedAnimLayer.h
+++ b/Engine/Source/Runtime/Engine/Classes/Animation/AnimNode_LinkedAnimLayer.h
@@ -60,6 +60,9 @@ public:
 protected:
 	void InitializeSelfLayer(const UAnimInstance* SelfAnimInstance);
 
+	// Initialize the source properties to copy from
+	void InitializeSourceProperties(const UAnimInstance* InAnimInstance);
+
 #if WITH_EDITOR
 	// Event fired when the instance we are running has changed
 	FSimpleMulticastDelegate OnInstanceChangedEvent;
diff --git a/Engine/Source/Runtime/Engine/Classes/Components/SkeletalMeshComponent.h b/Engine/Source/Runtime/Engine/Classes/Components/SkeletalMeshComponent.h
--- a/Engine/Source/Runtime/Engine/Classes/Components/SkeletalMeshComponent.h
+++ b/Engine/Source/Runtime/Engine/Classes/Components/SkeletalMeshComponent.h
@@ -2551,7 +2551,12 @@ private:
 	 * Cooking does not guarantee skeleton containing all names
 	 */
 	bool AreRequiredCurvesUpToDate() const;
-
+	
+#if WITH_EDITOR
+	// Handle reinstancing in editor
+	void HandleObjectsReplaced(const TMap<UObject*, UObject*>& OldToNewInstanceMap);
+#endif
+	
 public:
 	void ConditionallyDispatchQueuedAnimEvents();
 
diff --git a/Engine/Source/Runtime/Engine/Private/Animation/AnimInstance.cpp b/Engine/Source/Runtime/Engine/Private/Animation/AnimInstance.cpp
--- a/Engine/Source/Runtime/Engine/Private/Animation/AnimInstance.cpp
+++ b/Engine/Source/Runtime/Engine/Private/Animation/AnimInstance.cpp
@@ -3785,6 +3785,19 @@ void UAnimInstance::QueueRootMotionBlend(const FTransform& RootTransform, const
 	RootMotionBlendQueue.Add(FQueuedRootMotionBlend(RootTransform, SlotName, Weight));
 }
 
+#if WITH_EDITOR
+void UAnimInstance::HandleAnimInstanceReplaced()
+{
+	RecalcRequiredBones();
+	
+	// Reinit proxy
+	FAnimInstanceProxy& Proxy = GetProxyOnGameThread<FAnimInstanceProxy>();
+	Proxy.Initialize(this);
+	Proxy.InitializeCachedClassData();
+	Proxy.InitializeRootNode_WithRoot(Proxy.RootNode);
+}
+#endif
+
 bool UAnimInstance::RequestTransitionEvent(const FName EventName, const double RequestTimeout, const ETransitionRequestQueueMode QueueMode, const ETransitionRequestOverwriteMode OverwriteMode)
 {
 	return GetProxyOnAnyThread<FAnimInstanceProxy>().RequestTransitionEvent(EventName, RequestTimeout, QueueMode, OverwriteMode);
diff --git a/Engine/Source/Runtime/Engine/Private/Animation/AnimInstanceProxy.cpp b/Engine/Source/Runtime/Engine/Private/Animation/AnimInstanceProxy.cpp
--- a/Engine/Source/Runtime/Engine/Private/Animation/AnimInstanceProxy.cpp
+++ b/Engine/Source/Runtime/Engine/Private/Animation/AnimInstanceProxy.cpp
@@ -250,7 +250,7 @@ void FAnimInstanceProxy::Initialize(UAnimInstance* InAnimInstance)
 	}
 }
 
-void FAnimInstanceProxy::InitializeRootNode(bool bInDeferRootNodeInitialization)
+void FAnimInstanceProxy::InitializeCachedClassData()
 {
 	DECLARE_SCOPE_HIERARCHICAL_COUNTER_FUNC()
 
@@ -266,15 +266,7 @@ void FAnimInstanceProxy::InitializeRootNode(bool bInDeferRootNodeInitialization)
 			FAnimNode_StateMachine* StateMachine = Property->ContainerPtrToValuePtr<FAnimNode_StateMachine>(AnimInstanceObject);
 			StateMachine->CacheMachineDescription(AnimClassInterface);
 		}
-
-		// Init any nodes that need non-relevancy based initialization
-		UAnimInstance* AnimInstance = CastChecked<UAnimInstance>(GetAnimInstanceObject());
-		for (const FStructProperty* Property : AnimClassInterface->GetInitializationNodeProperties())
-		{
-			FAnimNode_Base* AnimNode = Property->ContainerPtrToValuePtr<FAnimNode_Base>(AnimInstanceObject);
-			AnimNode->OnInitializeAnimInstance(this, AnimInstance);
-		}
-
+		
 		// Cache any preupdate nodes
 		for (const FStructProperty* Property : AnimClassInterface->GetPreUpdateNodeProperties())
 		{
@@ -306,6 +298,24 @@ void FAnimInstanceProxy::InitializeRootNode(bool bInDeferRootNodeInitialization)
 			}
 		}
 	}
+}
+
+void FAnimInstanceProxy::InitializeRootNode(bool bInDeferRootNodeInitialization)
+{
+	DECLARE_SCOPE_HIERARCHICAL_COUNTER_FUNC()
+
+	InitializeCachedClassData();
+	
+	if(AnimClassInterface)
+	{
+		// Init any nodes that need non-relevancy based initialization
+		UAnimInstance* AnimInstance = CastChecked<UAnimInstance>(GetAnimInstanceObject());
+		for (const FStructProperty* Property : AnimClassInterface->GetInitializationNodeProperties())
+		{
+			FAnimNode_Base* AnimNode = Property->ContainerPtrToValuePtr<FAnimNode_Base>(AnimInstanceObject);
+			AnimNode->OnInitializeAnimInstance(this, AnimInstance);
+		}
+	}
 	else
 	{
 		auto InitializeNode = [this](FAnimNode_Base* AnimNode)
diff --git a/Engine/Source/Runtime/Engine/Private/Animation/AnimNode_CustomProperty.cpp b/Engine/Source/Runtime/Engine/Private/Animation/AnimNode_CustomProperty.cpp
--- a/Engine/Source/Runtime/Engine/Private/Animation/AnimNode_CustomProperty.cpp
+++ b/Engine/Source/Runtime/Engine/Private/Animation/AnimNode_CustomProperty.cpp
@@ -12,7 +12,16 @@ FAnimNode_CustomProperty::FAnimNode_CustomProperty()
 	, bReinitializeProperties(false)
 #endif // WITH_EDITOR
 {
+#if WITH_EDITOR
+	FCoreUObjectDelegates::OnObjectsReplaced.AddRaw(this, &FAnimNode_CustomProperty::HandleObjectsReplaced);
+#endif // WITH_EDITOR
+}
 
+FAnimNode_CustomProperty::~FAnimNode_CustomProperty()
+{
+#if WITH_EDITOR
+	FCoreUObjectDelegates::OnObjectsReplaced.RemoveAll(this);
+#endif // WITH_EDITOR
 }
 
 void FAnimNode_CustomProperty::SetTargetInstance(UObject* InInstance)
@@ -96,3 +105,17 @@ void FAnimNode_CustomProperty::InitializeProperties(const UObject* InSourceInsta
 	}
 }
 
+#if WITH_EDITOR
+void FAnimNode_CustomProperty::HandleObjectsReplaced(const TMap<UObject*, UObject*>& OldToNewInstanceMap)
+{
+	if (UObject* ThisTargetInstance = GetTargetInstance<UObject>())
+	{
+		UObject* const* ReinstancedTarget = OldToNewInstanceMap.Find(ThisTargetInstance);
+		if (ReinstancedTarget)
+		{
+			// recache the properties
+			bReinitializeProperties = true;
+		}
+	}
+}
+#endif	// #if WITH_EDITOR
diff --git a/Engine/Source/Runtime/Engine/Private/Animation/AnimNode_LinkedAnimGraph.cpp b/Engine/Source/Runtime/Engine/Private/Animation/AnimNode_LinkedAnimGraph.cpp
--- a/Engine/Source/Runtime/Engine/Private/Animation/AnimNode_LinkedAnimGraph.cpp
+++ b/Engine/Source/Runtime/Engine/Private/Animation/AnimNode_LinkedAnimGraph.cpp
@@ -470,3 +470,23 @@ void FAnimNode_LinkedAnimGraph::RequestBlend(const IAnimClassInterface* PriorAni
 		PendingBlendInProfile = nullptr;
 	}
 }
+
+#if WITH_EDITOR
+void FAnimNode_LinkedAnimGraph::HandleObjectsReplaced(const TMap<UObject*, UObject*>& OldToNewInstanceMap)
+{
+	FAnimNode_CustomProperty::HandleObjectsReplaced(OldToNewInstanceMap);
+	
+	if (UAnimInstance* ThisTargetInstance = GetTargetInstance<UAnimInstance>())
+	{
+		UObject* const* ReinstancedTarget = OldToNewInstanceMap.Find(ThisTargetInstance);
+		if (ReinstancedTarget)
+		{
+			if(UAnimInstance* ReinstancedTargetInstance = Cast<UAnimInstance>(*ReinstancedTarget))
+			{
+				DynamicUnlink(ThisTargetInstance);
+				DynamicLink(ReinstancedTargetInstance);
+			}
+		}
+	}
+}
+#endif
diff --git a/Engine/Source/Runtime/Engine/Private/Animation/AnimNode_LinkedAnimLayer.cpp b/Engine/Source/Runtime/Engine/Private/Animation/AnimNode_LinkedAnimLayer.cpp
--- a/Engine/Source/Runtime/Engine/Private/Animation/AnimNode_LinkedAnimLayer.cpp
+++ b/Engine/Source/Runtime/Engine/Private/Animation/AnimNode_LinkedAnimLayer.cpp
@@ -20,11 +20,8 @@ UAnimInstance* FAnimNode_LinkedAnimLayer::GetDynamicLinkTarget(UAnimInstance* In
 	}
 }
 
-void FAnimNode_LinkedAnimLayer::OnInitializeAnimInstance(const FAnimInstanceProxy* InProxy, const UAnimInstance* InAnimInstance)
+void FAnimNode_LinkedAnimLayer::InitializeSourceProperties(const UAnimInstance* InAnimInstance)
 {
-	check(SourcePropertyNames.Num() == DestPropertyNames.Num());
-	
-	// Initialize source properties here as they do not change
 	const int32 NumSourceProperties = SourcePropertyNames.Num();
 	SourceProperties.SetNumZeroed(NumSourceProperties);
 
@@ -33,7 +30,15 @@ void FAnimNode_LinkedAnimLayer::OnInitializeAnimInstance(const FAnimInstanceProx
 	{
 		SourceProperties[SourcePropertyIndex] = ThisClass->FindPropertyByName(SourcePropertyNames[SourcePropertyIndex]);
 	}
+}
+
+void FAnimNode_LinkedAnimLayer::OnInitializeAnimInstance(const FAnimInstanceProxy* InProxy, const UAnimInstance* InAnimInstance)
+{
+	check(SourcePropertyNames.Num() == DestPropertyNames.Num());
 	
+	// Initialize source properties here as they do not change
+	InitializeSourceProperties(InAnimInstance);
+
 	// We only initialize here if we are running a 'self' layer. Layers that use external instances need to be 
 	// initialized by the owning anim instance as they may share linked instances via grouping.
 	if(Interface.Get() == nullptr || InstanceClass.Get() == nullptr)
@@ -98,6 +103,14 @@ void FAnimNode_LinkedAnimLayer::SetLinkedLayerInstance(const UAnimInstance* InOw
 void FAnimNode_LinkedAnimLayer::InitializeProperties(const UObject* InSourceInstance, UClass* InTargetClass)
 {
 	check(SourcePropertyNames.Num() == DestPropertyNames.Num());
+
+#if WITH_EDITOR
+	// When reinstancing we need to init source properties here too as the class may have changed
+	if(GIsReinstancing)
+	{
+		InitializeSourceProperties(CastChecked<UAnimInstance>(InSourceInstance));
+	}
+#endif
 	
 	// Build dest property list - source is set up when we initialize
 	DestProperties.SetNumZeroed(SourcePropertyNames.Num());
diff --git a/Engine/Source/Runtime/Engine/Private/Animation/AnimNode_LinkedInputPose.cpp b/Engine/Source/Runtime/Engine/Private/Animation/AnimNode_LinkedInputPose.cpp
--- a/Engine/Source/Runtime/Engine/Private/Animation/AnimNode_LinkedInputPose.cpp
+++ b/Engine/Source/Runtime/Engine/Private/Animation/AnimNode_LinkedInputPose.cpp
@@ -73,7 +73,7 @@ void FAnimNode_LinkedInputPose::GatherDebugData(FNodeDebugData& DebugData)
 
 void FAnimNode_LinkedInputPose::DynamicLink(FAnimInstanceProxy* InInputProxy, FPoseLinkBase* InPoseLink, int32 InOuterGraphNodeIndex)
 {
-	check(InputProxy == nullptr);			// Must be unlinked before re-linking
+	check(GIsReinstancing || InputProxy == nullptr);			// Must be unlinked before re-linking
 
 	InputProxy = InInputProxy;
 	InputPose.SetDynamicLinkNode(InPoseLink);
@@ -82,7 +82,7 @@ void FAnimNode_LinkedInputPose::DynamicLink(FAnimInstanceProxy* InInputProxy, FP
 
 void FAnimNode_LinkedInputPose::DynamicUnlink()
 {
-	check(InputProxy != nullptr);			// Must be linked before unlinking
+	check(GIsReinstancing || InputProxy != nullptr);			// Must be linked before unlinking
 
 	InputProxy = nullptr;
 	InputPose.SetDynamicLinkNode(nullptr);
diff --git a/Engine/Source/Runtime/Engine/Private/Components/SkeletalMeshComponent.cpp b/Engine/Source/Runtime/Engine/Private/Components/SkeletalMeshComponent.cpp
--- a/Engine/Source/Runtime/Engine/Private/Components/SkeletalMeshComponent.cpp
+++ b/Engine/Source/Runtime/Engine/Private/Components/SkeletalMeshComponent.cpp
@@ -44,6 +44,7 @@
 #include "Animation/BuiltInAttributeTypes.h"
 #include "HAL/LowLevelMemTracker.h"
 #include "HAL/LowLevelMemStats.h"
+#include "HAL/IConsoleManager.h"
 
 LLM_DEFINE_TAG(SkeletalMesh_TransformData);
 
@@ -274,6 +275,10 @@ USkeletalMeshComponent::USkeletalMeshComponent(const FObjectInitializer& ObjectI
 	bSkipBoundsUpdateWhenInterpolating = false;
 
 	DeferredKinematicUpdateIndex = INDEX_NONE;
+	
+#if WITH_EDITOR
+	FCoreUObjectDelegates::OnObjectsReplaced.AddUObject(this, &USkeletalMeshComponent::HandleObjectsReplaced);
+#endif // WITH_EDITOR
 }
 
 void USkeletalMeshComponent::Serialize(FArchive& Ar)
@@ -4606,4 +4611,35 @@ TArray<FTransform> USkeletalMeshComponent::GetBoneSpaceTransforms()
 	return BoneSpaceTransforms;
 	PRAGMA_ENABLE_DEPRECATION_WARNINGS
 }
+
+
+#if WITH_EDITOR
+void USkeletalMeshComponent::HandleObjectsReplaced(const TMap<UObject*, UObject*>& OldToNewInstanceMap)
+{
+	IConsoleVariable* UseLegacyAnimInstanceReinstancingBehavior = IConsoleManager::Get().FindConsoleVariable(TEXT("bp.UseLegacyAnimInstanceReinstancingBehavior"));
+	if(UseLegacyAnimInstanceReinstancingBehavior == nullptr || !UseLegacyAnimInstanceReinstancingBehavior->GetBool())
+	{
+		bool bReinstancedAnimInstance = false;
+		
+		ForEachAnimInstance([&OldToNewInstanceMap, &bReinstancedAnimInstance](UAnimInstance* InAnimInstance)
+		{
+			UObject* const* ReinstancedTarget = OldToNewInstanceMap.Find(InAnimInstance);
+			if(ReinstancedTarget)
+			{
+				if (UAnimInstance* ReinstancedTargetInstance = Cast<UAnimInstance>(*ReinstancedTarget))
+				{
+					ReinstancedTargetInstance->HandleAnimInstanceReplaced();
+					bReinstancedAnimInstance = true;
+				}
+			}
+		});
+
+		if(bReinstancedAnimInstance)
+		{
+			ClearMotionVector();
+		}
+	}
+}
+#endif	// #if WITH_EDITOR
+
 #undef LOCTEXT_NAMESPACE
diff --git a/Engine/Source/Runtime/Engine/Public/Animation/AnimInstanceProxy.h b/Engine/Source/Runtime/Engine/Public/Animation/AnimInstanceProxy.h
--- a/Engine/Source/Runtime/Engine/Public/Animation/AnimInstanceProxy.h
+++ b/Engine/Source/Runtime/Engine/Public/Animation/AnimInstanceProxy.h
@@ -936,6 +936,9 @@ private:
 	/** Evaluate the slot when there are blend spaces involved in any of the active anim montages. */
 	void SlotEvaluatePoseWithBlendProfiles(const FName& SlotNodeName, const FAnimationPoseData& SourceAnimationPoseData, float InSourceWeight, FAnimationPoseData& OutBlendedAnimationPoseData, float InBlendWeight);
 
+	/** Initialize cached class data - state machines, pre-update nodes etc. */
+	void InitializeCachedClassData();
+	
 	/** The component to world transform of the component we are running on */
 	FTransform ComponentTransform;
 
