World Partition:
- Converted pinned actors management to the new loader adapters.
- Fixes the bug where you would pin an actor, modify it, unpin, losing the changes silently.
- Brings filtering of pinned actors by active data layers instead of making them always visible.

#rnx

diff --git a/Engine/Source/Runtime/Engine/Private/WorldPartition/WorldPartition.cpp b/Engine/Source/Runtime/Engine/Private/WorldPartition/WorldPartition.cpp
--- a/Engine/Source/Runtime/Engine/Private/WorldPartition/WorldPartition.cpp
+++ b/Engine/Source/Runtime/Engine/Private/WorldPartition/WorldPartition.cpp
@@ -209,12 +209,69 @@ public:
 		bIncludeNonSpatiallyLoadedActors = true;
 	}
 };
+
+class FLoaderAdapterPinnedActors: public IWorldPartitionActorLoaderInterface::ILoaderAdapterList
+{
+public:
+	FLoaderAdapterPinnedActors(UWorld* InWorld)
+		: IWorldPartitionActorLoaderInterface::ILoaderAdapterList(InWorld)
+	{
+		Load();
+	}
+
+	void AddActor(FWorldPartitionHandle& ActorHandle)
+	{
+		if (!ContainsActor(ActorHandle))
+		{
+			Actors.Add(ActorHandle);
+			RefreshLoadedState();
+		}
+	}
+
+	void RemoveActor(FWorldPartitionHandle& ActorHandle)
+	{
+		if (ContainsActor(ActorHandle))
+		{
+			ActorToRemove = ActorHandle;
+
+			if (RefreshLoadedState())
+			{
+				Actors.Remove(ActorHandle);
+			}
+
+			ActorToRemove.Reset();
+		}
+	}
+
+	bool ContainsActor(const FWorldPartitionHandle& ActorHandle) const
+	{
+		return Actors.Contains(ActorHandle);
+	}
+
+protected:
+	//~ Begin IWorldPartitionActorLoaderInterface::ILoaderAdapterList interface
+	virtual bool ShouldActorBeLoaded(const FWorldPartitionHandle& ActorHandle) const override
+	{
+		if (ContainsActor(ActorHandle) && (ActorHandle == ActorToRemove))
+		{
+			return false;
+		}
+
+		return IWorldPartitionActorLoaderInterface::ILoaderAdapterList::ShouldActorBeLoaded(ActorHandle);
+	}
+	//~ End IWorldPartitionActorLoaderInterface::ILoaderAdapterList interface
+
+private:
+	FWorldPartitionHandle ActorToRemove;
+};
 #endif
 
 UWorldPartition::UWorldPartition(const FObjectInitializer& ObjectInitializer)
 	: Super(ObjectInitializer)
 #if WITH_EDITOR
 	, EditorHash(nullptr)
+	, AlwaysLoadedActors(nullptr)
+	, PinnedActors(nullptr)
 	, WorldPartitionEditor(nullptr)
 	, bEnableStreaming(true)
 	, bStreamingWasEnabled(true)
@@ -388,6 +445,7 @@ void UWorldPartition::Initialize(UWorld* InWorld, const FTransform& InTransform)
 		EditorHash->Initialize();
 
 		AlwaysLoadedActors = new FLoaderAdapterAlwaysLoadedActors(OuterWorld);
+		PinnedActors = new FLoaderAdapterPinnedActors(OuterWorld);
 	}
 
 	check(RuntimeHash);
@@ -575,6 +633,12 @@ void UWorldPartition::Uninitialize()
 			AlwaysLoadedActors = nullptr;
 		}
 
+		if (PinnedActors)
+		{
+			delete PinnedActors;
+			PinnedActors = nullptr;
+		}
+
 		if (RegisteredEditorLoaderAdapters.Num())
 		{
 			for (IWorldPartitionActorLoaderInterface::ILoaderAdapter* RegisteredLoaderAdapter : RegisteredEditorLoaderAdapters)
@@ -587,9 +651,6 @@ void UWorldPartition::Uninitialize()
 
 		WorldDataLayersActor = FWorldPartitionReference();
 
-		PinnedActors.Empty();
-		PinnedActorRefs.Empty();
-
 		EditorHash = nullptr;
 #endif		
 
@@ -1285,60 +1346,24 @@ void UWorldPartition::SetEditorWantedCellSize(uint32 InCellSize)
 
 AActor* UWorldPartition::PinActor(const FGuid& ActorGuid)
 {
-	if (FWorldPartitionReference* PinnedActor = PinnedActors.Find(ActorGuid))
-	{
-		return (*PinnedActor).IsValid() ? (*PinnedActor)->GetActor() : nullptr;
-	}
-
-	TFunction<void(const FGuid&, TMap<FGuid, FWorldPartitionReference>&)> AddReferences = [this, &AddReferences](const FGuid& ActorGuid, TMap<FGuid, FWorldPartitionReference>& ReferenceMap)
-	{
-		if (ReferenceMap.Contains(ActorGuid))
-		{
-			return;
-		}
-
-		if (TUniquePtr<FWorldPartitionActorDesc>* const ActorDescPtr = GetActorDescriptor(ActorGuid); ActorDescPtr != nullptr && ActorDescPtr->IsValid())
-		{
-			ReferenceMap.Emplace(ActorGuid, ActorDescPtr);
-			
-			const FWorldPartitionActorDesc* const ActorDesc = ActorDescPtr->Get(); 
-			for (const FGuid& ReferencedActor : ActorDesc->GetReferences())
-			{
-				AddReferences(ReferencedActor, ReferenceMap);
-			}
-		}
-	};
-	
-	if (TUniquePtr<FWorldPartitionActorDesc>* const ActorDescPtr = GetActorDescriptor(ActorGuid))
-	{
-		if (const FWorldPartitionActorDesc* const ActorDesc = ActorDescPtr->Get())
-		{
-			FWorldPartitionReference& PinnedActor = PinnedActors.Emplace(ActorGuid, ActorDescPtr);
-
-			// If the pinned actor has references, we must also create references to those to ensure they are added
-			TMap<FGuid, FWorldPartitionReference>& References = PinnedActorRefs.Emplace(ActorGuid);
-
-			for (const FGuid& ReferencedActor : ActorDesc->GetReferences())
-			{
-				AddReferences(ReferencedActor, References);
-			}
-
-			return PinnedActor.IsValid() ? PinnedActor->GetActor() : nullptr;
-		}
-	}
-
-	return nullptr;
+	check(PinnedActors);
+	FWorldPartitionHandle ActorHandle(this, ActorGuid);
+	PinnedActors->AddActor(ActorHandle);
+	return ActorHandle->GetActor();
 }
 
 void UWorldPartition::UnpinActor(const FGuid& ActorGuid)
 {
-	PinnedActors.Remove(ActorGuid);
-	PinnedActorRefs.Remove(ActorGuid);
+	check(PinnedActors);
+	FWorldPartitionHandle ActorHandle(this, ActorGuid);
+	PinnedActors->RemoveActor(ActorHandle);
 }
 
 bool UWorldPartition::IsActorPinned(const FGuid& ActorGuid) const
 {
-	return PinnedActors.Contains(ActorGuid);
+	check(PinnedActors);
+	FWorldPartitionHandle ActorHandle(const_cast<UWorldPartition*>(this), ActorGuid);
+	return PinnedActors->ContainsActor(ActorHandle);
 }
 
 void UWorldPartition::LoadLastLoadedRegions(const TArray<FBox>& EditorLastLoadedRegions)
diff --git a/Engine/Source/Runtime/Engine/Private/WorldPartition/WorldPartitionActorLoaderInterface.cpp b/Engine/Source/Runtime/Engine/Private/WorldPartition/WorldPartitionActorLoaderInterface.cpp
--- a/Engine/Source/Runtime/Engine/Private/WorldPartition/WorldPartitionActorLoaderInterface.cpp
+++ b/Engine/Source/Runtime/Engine/Private/WorldPartition/WorldPartitionActorLoaderInterface.cpp
@@ -137,7 +137,7 @@ bool IWorldPartitionActorLoaderInterface::ILoaderAdapter::RefreshLoadedState()
 			for (FWorldPartitionHandle& ActorToUnload : ActorsToUnload)
 			{
 				SlowTask.EnterProgressFrame(1);
-				ActorReferences.Remove(ActorToUnload->GetGuid());
+				RemoveReferenceToActor(ActorToUnload);
 			}
 
 			PostLoadedStateChanged(ActorsToUnload.Num() > 0);
@@ -292,6 +292,11 @@ void IWorldPartitionActorLoaderInterface::ILoaderAdapter::AddReferenceToActor(FW
 	AddReferences(ActorHandle, ActorReferences.Emplace(ActorHandle->GetGuid()));
 }
 
+void IWorldPartitionActorLoaderInterface::ILoaderAdapter::RemoveReferenceToActor(FWorldPartitionHandle& ActorHandle)
+{
+	ActorReferences.Remove(ActorHandle->GetGuid());
+}
+
 void IWorldPartitionActorLoaderInterface::ILoaderAdapter::OnActorDataLayersEditorLoadingStateChanged(bool bFromUserOperation)
 {
 	if (!RefreshLoadedState())
@@ -303,11 +308,11 @@ void IWorldPartitionActorLoaderInterface::ILoaderAdapter::OnActorDataLayersEdito
 	}
 }
 
-IWorldPartitionActorLoaderInterface::FLoaderAdapterList::FLoaderAdapterList(UWorld* InWorld)
+IWorldPartitionActorLoaderInterface::ILoaderAdapterList::ILoaderAdapterList(UWorld* InWorld)
 	: ILoaderAdapter(InWorld)
 {}
 
-void IWorldPartitionActorLoaderInterface::FLoaderAdapterList::ForEachActor(TFunctionRef<void(const FWorldPartitionHandle&)> InOperation) const
+void IWorldPartitionActorLoaderInterface::ILoaderAdapterList::ForEachActor(TFunctionRef<void(const FWorldPartitionHandle&)> InOperation) const
 {
 	for (const FWorldPartitionHandle& Actor : Actors)
 	{
@@ -323,7 +328,7 @@ IWorldPartitionActorLoaderInterface::ILoaderAdapterSpatial::ILoaderAdapterSpatia
 
 void IWorldPartitionActorLoaderInterface::ILoaderAdapterSpatial::ForEachActor(TFunctionRef<void(const FWorldPartitionHandle&)> InOperation) const
 {
-	if (UWorldPartition* WorldPartition = World->GetWorldPartition())
+	if (UWorldPartition* WorldPartition = GetWorld()->GetWorldPartition())
 	{
 		WorldPartition->EditorHash->ForEachIntersectingActor(*GetBoundingBox(), [this, WorldPartition, &InOperation](FWorldPartitionActorDesc* ActorDesc)
 		{
diff --git a/Engine/Source/Runtime/Engine/Public/WorldPartition/LoaderAdapter/LoaderAdapterShape.h b/Engine/Source/Runtime/Engine/Public/WorldPartition/LoaderAdapter/LoaderAdapterShape.h
--- a/Engine/Source/Runtime/Engine/Public/WorldPartition/LoaderAdapter/LoaderAdapterShape.h
+++ b/Engine/Source/Runtime/Engine/Public/WorldPartition/LoaderAdapter/LoaderAdapterShape.h
@@ -16,6 +16,7 @@ public:
 	virtual TOptional<FString> GetLabel() const override;
 	//~ End IWorldPartitionActorLoaderInterface::ILoader interface
 
+protected:
 	//~ Begin IWorldPartitionActorLoaderInterface::ILoaderAdapterSpatial interface
 	virtual bool Intersect(const FBox& Box) const override;
 	//~ End IWorldPartitionActorLoaderInterface::ILoaderAdapterSpatial interface
diff --git a/Engine/Source/Runtime/Engine/Public/WorldPartition/WorldPartition.h b/Engine/Source/Runtime/Engine/Public/WorldPartition/WorldPartition.h
--- a/Engine/Source/Runtime/Engine/Public/WorldPartition/WorldPartition.h
+++ b/Engine/Source/Runtime/Engine/Public/WorldPartition/WorldPartition.h
@@ -29,6 +29,7 @@ class UWorldPartitionRuntimeHash;
 class UWorldPartitionStreamingPolicy;
 class IStreamingGenerationErrorHandler;
 class FLoaderAdapterAlwaysLoadedActors;
+class FLoaderAdapterPinnedActors;
 class FHLODActorDesc;
 class UCanvas;
 class ULevel;
@@ -209,6 +210,7 @@ public:
 	TObjectPtr<UWorldPartitionEditorHash> EditorHash;
 
 	FLoaderAdapterAlwaysLoadedActors* AlwaysLoadedActors;
+	FLoaderAdapterPinnedActors* PinnedActors;
 
 	IWorldPartitionEditor* WorldPartitionEditor;
 
@@ -283,9 +285,6 @@ private:
 	void HashActorDesc(FWorldPartitionActorDesc* ActorDesc);
 	void UnhashActorDesc(FWorldPartitionActorDesc* ActorDesc);
 
-	TMap<FGuid, FWorldPartitionReference> PinnedActors;
-	TMap<FGuid, TMap<FGuid, FWorldPartitionReference>> PinnedActorRefs;
-
 public:
 	// Editor loader adapters management
 	template <typename T, typename... ArgsType>
diff --git a/Engine/Source/Runtime/Engine/Public/WorldPartition/WorldPartitionActorLoaderInterface.h b/Engine/Source/Runtime/Engine/Public/WorldPartition/WorldPartitionActorLoaderInterface.h
--- a/Engine/Source/Runtime/Engine/Public/WorldPartition/WorldPartitionActorLoaderInterface.h
+++ b/Engine/Source/Runtime/Engine/Public/WorldPartition/WorldPartitionActorLoaderInterface.h
@@ -31,27 +31,40 @@ public:
 		bool Unload();
 		bool IsLoaded() const;
 
+		UWorld* GetWorld() const { return World; }
+
 		bool GetUserCreated() const { return bUserCreated; }
 		void SetUserCreated(bool bValue) { bUserCreated = bValue; }
 
-		// Interface
+		// Public interface
 		virtual TOptional<FBox> GetBoundingBox() const { return TOptional<FBox>(); }
 		virtual TOptional<FString> GetLabel() const { return TOptional<FString>(); }
 		virtual TOptional<FColor> GetColor() const { return TOptional<FColor>(); }
 
 	protected:
+		// Private interface
 		virtual bool Intersect(const FBox& Box) const =0;
 		virtual void ForEachActor(TFunctionRef<void(const FWorldPartitionHandle&)> InOperation) const =0;
+		virtual bool ShouldActorBeLoaded(const FWorldPartitionHandle& Actor) const;
 
 		void RegisterDelegates();
 		void UnregisterDelegates();
+
+		/**
+		 * Refresh the loaded state of this loader adatper.
+		 *
+		 * @return true if the operation was successful, false otherwise (operation was cancelled by the user
+		 *         if actors needed resaving, etc.
+		 */
 		bool RefreshLoadedState();
-		bool ShouldActorBeLoaded(const FWorldPartitionHandle& Actor) const;
+
 		void PostLoadedStateChanged(bool bUnloadedActors);
 		bool AllowUnloadingActors(const TArray<FWorldPartitionHandle>& ActorsToUnload) const;
 		void AddReferenceToActor(FWorldPartitionHandle& Actor);
+		void RemoveReferenceToActor(FWorldPartitionHandle& Actor);
 		void OnActorDataLayersEditorLoadingStateChanged(bool bFromUserOperation);		
 
+	private:
 		UWorld* World;
 
 		uint8 bLoaded : 1;
@@ -61,18 +74,18 @@ public:
 	};
 
 	/** Base class for actor loaders that contains a specific list of actors */
-	class ENGINE_API FLoaderAdapterList : public ILoaderAdapter
+	class ENGINE_API ILoaderAdapterList : public ILoaderAdapter
 	{
 	public:
-		FLoaderAdapterList(UWorld* InWorld);
-		virtual ~FLoaderAdapterList() {}
+		ILoaderAdapterList(UWorld* InWorld);
+		virtual ~ILoaderAdapterList() {}
 
+	protected:
 		//~ Begin ILoaderAdapter interface
 		virtual bool Intersect(const FBox& Box) const { return true; }
 		virtual void ForEachActor(TFunctionRef<void(const FWorldPartitionHandle&)> InOperation) const;
 		//~ End ILoaderAdapter interface
 
-	protected:
 		TSet<FWorldPartitionHandle> Actors;
 	};
 
@@ -83,6 +96,7 @@ public:
 		ILoaderAdapterSpatial(UWorld* InWorld);
 		virtual ~ILoaderAdapterSpatial() {}
 
+	protected:
 		//~ Begin ILoaderAdapter interface
 		virtual void ForEachActor(TFunctionRef<void(const FWorldPartitionHandle&)> InOperation) const;
 		//~ End ILoaderAdapter interface
