World Partition: missing last-minute changes done from previous submit.

#rnx

diff --git a/Engine/Source/Runtime/Engine/Private/WorldPartition/WorldPartition.cpp b/Engine/Source/Runtime/Engine/Private/WorldPartition/WorldPartition.cpp
--- a/Engine/Source/Runtime/Engine/Private/WorldPartition/WorldPartition.cpp
+++ b/Engine/Source/Runtime/Engine/Private/WorldPartition/WorldPartition.cpp
@@ -1346,24 +1346,32 @@ void UWorldPartition::SetEditorWantedCellSize(uint32 InCellSize)
 
 AActor* UWorldPartition::PinActor(const FGuid& ActorGuid)
 {
-	check(PinnedActors);
-	FWorldPartitionHandle ActorHandle(this, ActorGuid);
-	PinnedActors->AddActor(ActorHandle);
-	return ActorHandle->GetActor();
+	if (PinnedActors)
+	{
+		FWorldPartitionHandle ActorHandle(this, ActorGuid);
+		PinnedActors->AddActor(ActorHandle);
+		return ActorHandle->GetActor();
+	}
+	return nullptr;
 }
 
 void UWorldPartition::UnpinActor(const FGuid& ActorGuid)
 {
-	check(PinnedActors);
-	FWorldPartitionHandle ActorHandle(this, ActorGuid);
-	PinnedActors->RemoveActor(ActorHandle);
+	if (PinnedActors)
+	{
+		FWorldPartitionHandle ActorHandle(this, ActorGuid);
+		PinnedActors->RemoveActor(ActorHandle);
+	}
 }
 
 bool UWorldPartition::IsActorPinned(const FGuid& ActorGuid) const
 {
-	check(PinnedActors);
-	FWorldPartitionHandle ActorHandle(const_cast<UWorldPartition*>(this), ActorGuid);
-	return PinnedActors->ContainsActor(ActorHandle);
+	if (PinnedActors)
+	{
+		FWorldPartitionHandle ActorHandle(const_cast<UWorldPartition*>(this), ActorGuid);
+		return PinnedActors->ContainsActor(ActorHandle);
+	}
+	return false;
 }
 
 void UWorldPartition::LoadLastLoadedRegions(const TArray<FBox>& EditorLastLoadedRegions)
