World Partition: only save per-user settings for the main world partition.

#rnx

diff --git a/Engine/Source/Runtime/Engine/Private/WorldPartition/WorldPartition.cpp b/Engine/Source/Runtime/Engine/Private/WorldPartition/WorldPartition.cpp
--- a/Engine/Source/Runtime/Engine/Private/WorldPartition/WorldPartition.cpp
+++ b/Engine/Source/Runtime/Engine/Private/WorldPartition/WorldPartition.cpp
@@ -554,8 +554,8 @@ void UWorldPartition::Initialize(UWorld* InWorld, const FTransform& InTransform)
 		AlwaysLoadedActors->Load();
 
 		// Load more cells depending on the user's settings
-		// Skipped when running from a commandlet
-		if (!IsRunningCommandlet())
+		// Skipped when running from a commandlet and for subpartitions
+		if (IsMainWorldPartition() && !IsRunningCommandlet())
 		{
 			// Load last loaded regions
 			if (GetMutableDefault<UWorldPartitionEditorPerProjectUserSettings>()->GetEnableLoadingOfLastLoadedRegions())
@@ -563,11 +563,6 @@ void UWorldPartition::Initialize(UWorld* InWorld, const FTransform& InTransform)
 				LoadLastLoadedRegions();
 			}
 		}
-		
-		if (GEditor)
-		{
-			GEditor->OnEditorClose().AddUObject(this, &UWorldPartition::SavePerUserSettings);
-		}
 	}
 #endif //WITH_EDITOR
 
@@ -615,11 +610,9 @@ void UWorldPartition::Uninitialize()
 		}
 
 #if WITH_EDITOR
-		SavePerUserSettings();
-
-		if (GEditor && !World->IsGameWorld())
+		if (IsMainWorldPartition())
 		{
-			GEditor->OnEditorClose().RemoveAll(this);
+			SavePerUserSettings();
 		}
 
 		if (World->IsGameWorld())
@@ -726,6 +719,7 @@ void UWorldPartition::RegisterDelegates()
 			FGameDelegates::Get().GetEndPlayMapDelegate().AddUObject(this, &UWorldPartition::OnEndPlay);
 			FCoreUObjectDelegates::PostReachabilityAnalysis.AddUObject(this, &UWorldPartition::OnGCPostReachabilityAnalysis);
 			GEditor->OnPostBugItGoCalled().AddUObject(this, &UWorldPartition::OnPostBugItGoCalled);
+			GEditor->OnEditorClose().AddUObject(this, &UWorldPartition::SavePerUserSettings);
 		}
 	}
 #endif
@@ -767,6 +761,7 @@ void UWorldPartition::UnregisterDelegates()
 			}
 
 			GEditor->OnPostBugItGoCalled().RemoveAll(this);
+			GEditor->OnEditorClose().RemoveAll(this);
 		}
 	}
 #endif
@@ -1264,6 +1259,8 @@ TArray<FBox> UWorldPartition::GetUserLoadedEditorRegions() const
 
 void UWorldPartition::SavePerUserSettings()
 {
+	check(IsMainWorldPartition());
+
 	if (GIsEditor && !World->IsGameWorld() && !IsRunningCommandlet() && !IsEngineExitRequested())
 	{
 		GetMutableDefault<UWorldPartitionEditorPerProjectUserSettings>()->SetEditorLoadedRegions(GetWorld(), GetUserLoadedEditorRegions());
@@ -1386,6 +1383,8 @@ void UWorldPartition::LoadLastLoadedRegions(const TArray<FBox>& EditorLastLoaded
 
 void UWorldPartition::LoadLastLoadedRegions()
 {
+	check(IsMainWorldPartition());
+
 	TArray<FBox> EditorLastLoadedRegions = GetMutableDefault<UWorldPartitionEditorPerProjectUserSettings>()->GetEditorLoadedRegions(World);
 	LoadLastLoadedRegions(EditorLastLoadedRegions);
 
