FModelSceneProxy fills in velocity in view relevance.
This allows it to write velocity when moving or when r.Velocity.ForceOutput is set.

diff --git a/Engine/Source/Runtime/Engine/Private/ModelRender.cpp b/Engine/Source/Runtime/Engine/Private/ModelRender.cpp
--- a/Engine/Source/Runtime/Engine/Private/ModelRender.cpp
+++ b/Engine/Source/Runtime/Engine/Private/ModelRender.cpp
@@ -565,6 +565,7 @@ public:
 		}
 		Result.bShadowRelevance = IsShadowCast(View);
 		MaterialRelevance.SetPrimitiveViewRelevance(Result);
+		Result.bVelocityRelevance = DrawsVelocity() && Result.bOpaque && Result.bRenderInMainPass;
 		return Result;
 	}
 
