Remove obsolete platform bug workaround
#rnx

diff --git a/Engine/Source/Runtime/Engine/Private/SkeletalRenderGPUSkin.cpp b/Engine/Source/Runtime/Engine/Private/SkeletalRenderGPUSkin.cpp
--- a/Engine/Source/Runtime/Engine/Private/SkeletalRenderGPUSkin.cpp
+++ b/Engine/Source/Runtime/Engine/Private/SkeletalRenderGPUSkin.cpp
@@ -98,12 +98,7 @@ void FMorphVertexBuffer::InitDynamicRHI()
 	const bool bUseGPUMorphTargets = UseGPUMorphTargets(FeatureLevel);
 	bUsesComputeShader = bUseGPUMorphTargets;
 
-#if PLATFORM_PS4
-	// PS4 requires non-static buffers in order to be updated on the GPU while the CPU is writing into it
-	EBufferUsageFlags Flags = bUseGPUMorphTargets ? (EBufferUsageFlags)(BUF_Dynamic | BUF_UnorderedAccess) : BUF_Dynamic;
-#else
 	EBufferUsageFlags Flags = bUseGPUMorphTargets ? (EBufferUsageFlags)(BUF_Static | BUF_UnorderedAccess) : BUF_Dynamic;
-#endif
 
 	// BUF_ShaderResource is needed for Morph support of the SkinCache
 	Flags = (EBufferUsageFlags)(Flags | BUF_ShaderResource);
