Default EvaluateWPO primitive flag to true.
It is always set by primitive proxy types that care, but for types that currently don't (ie niagara) it should default to on.

diff --git a/Engine/Source/Runtime/Engine/Public/PrimitiveUniformShaderParameters.h b/Engine/Source/Runtime/Engine/Public/PrimitiveUniformShaderParameters.h
--- a/Engine/Source/Runtime/Engine/Public/PrimitiveUniformShaderParameters.h
+++ b/Engine/Source/Runtime/Engine/Public/PrimitiveUniformShaderParameters.h
@@ -95,6 +95,7 @@ public:
 		// Flags defaulted on
 		bCastShadow									= true;
 		bCastContactShadow							= true;
+		bEvaluateWorldPositionOffset				= true;
 		bVisibleInGame								= true;
 		bVisibleInEditor							= true;
 		bVisibleInReflectionCaptures				= true;
@@ -107,7 +108,6 @@ public:
 		bUseVolumetricLightmap						= false;
 		bShouldCacheShadow							= false;
 		bOutputVelocity								= false;
-		bEvaluateWorldPositionOffset				= false;
 		bHasCapsuleRepresentation					= false;
 		bHasPreSkinnedLocalBounds					= false;
 		bHasPreviousLocalToWorld					= false;
