Fix static analysis and RPC version of async input crashing

diff --git a/Engine/Source/Runtime/Engine/Private/PhysicsEngine/AsyncPhysicsInputComponent.cpp b/Engine/Source/Runtime/Engine/Private/PhysicsEngine/AsyncPhysicsInputComponent.cpp
--- a/Engine/Source/Runtime/Engine/Private/PhysicsEngine/AsyncPhysicsInputComponent.cpp
+++ b/Engine/Source/Runtime/Engine/Private/PhysicsEngine/AsyncPhysicsInputComponent.cpp
@@ -370,7 +370,8 @@ void UAsyncPhysicsInputComponent::AsyncPhysicsTickComponent(float DeltaTime, flo
 		{
 			if (PC->IsLocalController() && !InputCmdCVars::ForceInputDrop)
 			{
-				ServerRPCBufferInput(Data);
+				//TODO: fix rpc version
+				//ServerRPCBufferInput(Data);
 
 				//If we are the local player then we need to keep this data around to send redundant RPCs to deal with potential packet loss
 				bFreeData = --Data->ReplicationRedundancy <= 0;
@@ -396,7 +397,8 @@ void UAsyncPhysicsInputComponent::ServerRPCBufferInput_Implementation(UAsyncPhys
 		}
 	}
 
-	BufferedData.Add(NewData);
+	//TODO: fix rpc version
+	//BufferedData.Add(NewData);
 }
 
 void UAsyncPhysicsInputComponent::OnDispatchPhysicsTick(int32 PhysicsStep, int32 NumSteps, int32 ServerFrame)
@@ -418,7 +420,7 @@ void UAsyncPhysicsInputComponent::OnDispatchPhysicsTick(int32 PhysicsStep, int32
 			DataToWrite = DuplicateObject((UAsyncPhysicsData*)DataClass->GetDefaultObject(), nullptr);
 			for (int32 Step = 0; Step < NumSteps; ++Step)
 			{
-				if(Step > 0 && DataToSend)
+				if(Step > 0)
 				{
 					//Make sure each sub-step gets its own unique data so we don't have to worry about ref count
 					//TODO: should probably given user opportunity to modify this per sub-step. For example a jump instruction should only happen on first sub-step
@@ -463,4 +465,4 @@ const UAsyncPhysicsData* UAsyncPhysicsInputComponent::GetDataToConsume() const
 	}
 	
 	return (UAsyncPhysicsData*)DataClass->GetDefaultObject();
-}
\ No newline at end of file
+}
