World Partition: remove code that was preventing updating the streaming state while there's an active transaction.

#rnx

diff --git a/Engine/Source/Runtime/Engine/Private/WorldPartition/WorldPartitionSubsystem.cpp b/Engine/Source/Runtime/Engine/Private/WorldPartition/WorldPartitionSubsystem.cpp
--- a/Engine/Source/Runtime/Engine/Private/WorldPartition/WorldPartitionSubsystem.cpp
+++ b/Engine/Source/Runtime/Engine/Private/WorldPartition/WorldPartitionSubsystem.cpp
@@ -291,14 +291,6 @@ void UWorldPartitionSubsystem::DumpStreamingSources(FOutputDevice& OutputDevice)
 
 void UWorldPartitionSubsystem::UpdateStreamingState()
 {
-#if WITH_EDITOR
-	// Do not update during transaction
-	if (GUndo)
-	{
-		return;
-	}
-#endif
-
 	for (UWorldPartition* RegisteredWorldPartition : RegisteredWorldPartitions)
 	{
 		RegisteredWorldPartition->UpdateStreamingState();
