Remove platform specific defines

diff --git a/Engine/Source/Runtime/Engine/Private/GameUserSettings.cpp b/Engine/Source/Runtime/Engine/Private/GameUserSettings.cpp
--- a/Engine/Source/Runtime/Engine/Private/GameUserSettings.cpp
+++ b/Engine/Source/Runtime/Engine/Private/GameUserSettings.cpp
@@ -973,7 +973,7 @@ void UGameUserSettings::EnableHDRDisplayOutput(bool bEnable, int32 DisplayNits /
 		}
 
 		// Update final requested state for saved config
-#if !PLATFORM_PS4 && !PLATFORM_USES_FIXED_HDR_SETTING
+#if !PLATFORM_USES_FIXED_HDR_SETTING
 		// Do not override the user setting on console (we rely on the OS setting)
 		bUseHDRDisplayOutput = bEnable;
 #endif
