CAC per instance properties

diff --git a/Engine/Source/Editor/LevelEditor/Private/Elements/Actor/ActorElementLevelEditorSelectionCustomization.cpp b/Engine/Source/Editor/LevelEditor/Private/Elements/Actor/ActorElementLevelEditorSelectionCustomization.cpp
--- a/Engine/Source/Editor/LevelEditor/Private/Elements/Actor/ActorElementLevelEditorSelectionCustomization.cpp
+++ b/Engine/Source/Editor/LevelEditor/Private/Elements/Actor/ActorElementLevelEditorSelectionCustomization.cpp
@@ -50,9 +50,13 @@ FTypedElementHandle FActorElementLevelEditorSelectionCustomization::GetSelection
 {
 	if (AActor* ConsideredActor = ActorElementDataUtil::GetActorFromHandle(InElementSelectionHandle))
 	{
-		while (ConsideredActor->IsChildActor())
+		extern ENGINE_API int32 GExperimentalAllowPerInstanceChildActorProperties;
+		if (!GExperimentalAllowPerInstanceChildActorProperties)
 		{
-			ConsideredActor = ConsideredActor->GetParentActor();
+			while (ConsideredActor->IsChildActor())
+			{
+				ConsideredActor = ConsideredActor->GetParentActor();
+			}
 		}
 		return UEngineElementsLibrary::AcquireEditorActorElementHandle(ConsideredActor);
 	}
diff --git a/Engine/Source/Runtime/Engine/Classes/Components/ChildActorComponent.h b/Engine/Source/Runtime/Engine/Classes/Components/ChildActorComponent.h
--- a/Engine/Source/Runtime/Engine/Classes/Components/ChildActorComponent.h
+++ b/Engine/Source/Runtime/Engine/Classes/Components/ChildActorComponent.h
@@ -56,6 +56,9 @@ public:
 	FGuid ChildActorGUID;
 #endif
 
+	// The saved properties for the ChildActor itself
+	TSharedPtr<FActorInstanceData> ActorInstanceData;
+
 	// The component instance data cache for the ChildActor spawned by this component
 	TSharedPtr<FComponentInstanceDataCache> ComponentInstanceData;
 };
diff --git a/Engine/Source/Runtime/Engine/Private/Actor.cpp b/Engine/Source/Runtime/Engine/Private/Actor.cpp
--- a/Engine/Source/Runtime/Engine/Private/Actor.cpp
+++ b/Engine/Source/Runtime/Engine/Private/Actor.cpp
@@ -1631,6 +1631,15 @@ bool AActor::Modify( bool bAlwaysMarkDirty/*=true*/ )
 		return false;
 	}
 
+	extern int32 GExperimentalAllowPerInstanceChildActorProperties;
+	if (GExperimentalAllowPerInstanceChildActorProperties)
+	{
+		if (AActor* ParentActor = GetParentActor())
+		{
+			return ParentActor->Modify(bAlwaysMarkDirty);
+		}
+	}
+
 	// Any properties that reference a blueprint constructed component needs to avoid creating a reference to the component from the transaction
 	// buffer, so we temporarily switch the property to non-transactional while the modify occurs
 	TArray<FObjectProperty*> TemporarilyNonTransactionalProperties;
@@ -4903,7 +4912,8 @@ bool AActor::IsSelectionChild() const
 
 AActor* AActor::GetSelectionParent() const
 {
-	if (IsChildActor())
+	extern int32 GExperimentalAllowPerInstanceChildActorProperties;
+	if (!GExperimentalAllowPerInstanceChildActorProperties && IsChildActor())
 	{
 		return GetParentActor();
 	}
diff --git a/Engine/Source/Runtime/Engine/Private/Components/ActorComponent.cpp b/Engine/Source/Runtime/Engine/Private/Components/ActorComponent.cpp
--- a/Engine/Source/Runtime/Engine/Private/Components/ActorComponent.cpp
+++ b/Engine/Source/Runtime/Engine/Private/Components/ActorComponent.cpp
@@ -812,9 +812,13 @@ bool UActorComponent::Modify( bool bAlwaysMarkDirty/*=true*/ )
 
 	// If this is a construction script component we don't store them in the transaction buffer.  Instead, mark
 	// the Actor as modified so that we store of the transaction annotation that has the component properties stashed
-	if (MyOwner && IsCreatedByConstructionScript())
+	if (MyOwner)
 	{
-		return MyOwner->Modify(bAlwaysMarkDirty);
+		extern int32 GExperimentalAllowPerInstanceChildActorProperties;
+		if (IsCreatedByConstructionScript() || (GExperimentalAllowPerInstanceChildActorProperties && MyOwner->IsChildActor()))
+		{
+			return MyOwner->Modify(bAlwaysMarkDirty);
+		}
 	}
 
 	return Super::Modify(bAlwaysMarkDirty);
diff --git a/Engine/Source/Runtime/Engine/Private/Components/ChildActorComponent.cpp b/Engine/Source/Runtime/Engine/Private/Components/ChildActorComponent.cpp
--- a/Engine/Source/Runtime/Engine/Private/Components/ChildActorComponent.cpp
+++ b/Engine/Source/Runtime/Engine/Private/Components/ChildActorComponent.cpp
@@ -11,6 +11,15 @@
 
 DEFINE_LOG_CATEGORY_STATIC(LogChildActorComponent, Warning, All);
 
+ENGINE_API int32 GExperimentalAllowPerInstanceChildActorProperties = 0;
+
+static FAutoConsoleVariableRef CVarExperimentalAllowPerInstanceChildActorProperties(
+	TEXT("cac.ExperimentalAllowPerInstanceChildActorProperties"),
+	GExperimentalAllowPerInstanceChildActorProperties,
+	TEXT("[EXPERIMENTAL] If true, allows properties to be modified on a per-instance basis for child actors."),
+	ECVF_Default
+);
+
 UChildActorComponent::UChildActorComponent(const FObjectInitializer& ObjectInitializer)
 	: Super(ObjectInitializer)
 	, ActorOuter(nullptr)
@@ -326,6 +335,16 @@ FChildActorComponentInstanceData::FChildActorComponentInstanceData(const UChildA
 			ChildActorName = ChildActor->GetFName();
 		}
 
+		if (GExperimentalAllowPerInstanceChildActorProperties)
+		{
+			ActorInstanceData = MakeShared<FActorInstanceData>(ChildActor);
+			// If it is empty dump it
+			if (!ActorInstanceData->HasInstanceData())
+			{
+				ActorInstanceData.Reset();
+			}
+		}
+
 		ComponentInstanceData = MakeShared<FComponentInstanceDataCache>(ChildActor);
 		// If it is empty dump it
 		if (!ComponentInstanceData->HasInstanceData())
@@ -446,6 +465,11 @@ void UChildActorComponent::ApplyComponentInstanceData(FChildActorComponentInstan
 			}
 		}
 
+		if (ChildActorInstanceData->ActorInstanceData)
+		{
+			ChildActorInstanceData->ActorInstanceData->ApplyToActor(ChildActor, CacheApplyPhase);
+		}
+
 		if (ChildActorInstanceData->ComponentInstanceData)
 		{
 			ChildActorInstanceData->ComponentInstanceData->ApplyToActor(ChildActor, CacheApplyPhase);
