nop just comment cleanups

diff --git a/Engine/Source/Developer/Android/AndroidTargetPlatform/Private/AndroidTargetPlatform.cpp b/Engine/Source/Developer/Android/AndroidTargetPlatform/Private/AndroidTargetPlatform.cpp
--- a/Engine/Source/Developer/Android/AndroidTargetPlatform/Private/AndroidTargetPlatform.cpp
+++ b/Engine/Source/Developer/Android/AndroidTargetPlatform/Private/AndroidTargetPlatform.cpp
@@ -66,6 +66,7 @@ namespace AndroidTexFormat
 	const static FName NameETC2_RGBA(TEXT("ETC2_RGBA"));
 	const static FName NameETC2_R11(TEXT("ETC2_R11"));
 	const static FName NameAutoETC2(TEXT("AutoETC2"));
+	// AFAICT none of these explicitly sized ASTC names are actually valid TextureFormat names
 	const static FName NameASTC_4x4(TEXT("ASTC_4x4"));
 	const static FName NameASTC_6x6(TEXT("ASTC_6x6"));
 	const static FName NameASTC_8x8(TEXT("ASTC_8x8"));
diff --git a/Engine/Source/Developer/Android/AndroidTargetPlatform/Private/AndroidTargetPlatform.h b/Engine/Source/Developer/Android/AndroidTargetPlatform/Private/AndroidTargetPlatform.h
--- a/Engine/Source/Developer/Android/AndroidTargetPlatform/Private/AndroidTargetPlatform.h
+++ b/Engine/Source/Developer/Android/AndroidTargetPlatform/Private/AndroidTargetPlatform.h
@@ -180,7 +180,7 @@ protected:
 
 	virtual FAndroidTargetDevicePtr CreateTargetDevice(const ITargetPlatform& InTargetPlatform, const FString& InSerialNumber, const FString& InAndroidVariant) const;
 
-	// query for rene3ring mode support
+	// query for rendering mode support
 	bool SupportsES31() const;
 	bool SupportsVulkan() const;
 	bool SupportsLandscapeMeshLODStreaming() const;
diff --git a/Engine/Source/Developer/Android/AndroidTargetPlatform/Private/AndroidTargetPlatformModule.cpp b/Engine/Source/Developer/Android/AndroidTargetPlatform/Private/AndroidTargetPlatformModule.cpp
--- a/Engine/Source/Developer/Android/AndroidTargetPlatform/Private/AndroidTargetPlatformModule.cpp
+++ b/Engine/Source/Developer/Android/AndroidTargetPlatform/Private/AndroidTargetPlatformModule.cpp
@@ -33,7 +33,7 @@ public:
 				SinglePlatforms.Add(new FAndroid_DXTTargetPlatform(bIsClient));
 				SinglePlatforms.Add(new FAndroid_ETC2TargetPlatform(bIsClient));
 
-				// thse are used in NotifyMultiSelectedFormatsChanged, so track in another array
+				// these are used in NotifyMultiSelectedFormatsChanged, so track in another array
 				MultiPlatforms.Add(new FAndroid_MultiTargetPlatform(bIsClient));
 			}
 
diff --git a/Engine/Source/Developer/TextureFormatIntelISPCTexComp/Private/TextureFormatIntelISPCTexComp.cpp b/Engine/Source/Developer/TextureFormatIntelISPCTexComp/Private/TextureFormatIntelISPCTexComp.cpp
--- a/Engine/Source/Developer/TextureFormatIntelISPCTexComp/Private/TextureFormatIntelISPCTexComp.cpp
+++ b/Engine/Source/Developer/TextureFormatIntelISPCTexComp/Private/TextureFormatIntelISPCTexComp.cpp
@@ -638,9 +638,8 @@ public:
 	/*
 	virtual FString GetDerivedDataKeyString(const FTextureBuildSettings& BuildSettings) const override
 	{
-		// @@!! no GetDerivedDataKeyString !?
+		// @todo Oodle : should log settings here :
 
-		// should log settings here :
 		// should store
 		// GetDefaultCompressionBySizeValue() 
 		//return FString::Printf(TEXT("ISPC_%d_%d_%d"), GetQualityVersion(BuildSettings.FormatConfigOverride, BuildSettings.CompressionQuality), GASTCCompressor, int32(GASTCHDRProfile > 0 && BuildSettings.bHDRSource));
@@ -754,7 +753,9 @@ public:
 			{
 				F16Value.Encoded = 0;
 			}
-			// @@!! also clamp +Inf to max non-inf value?
+			// @todo Oodle : should also clamp +Inf to max non-inf value?
+			//	 note that behavior of the basic F32->F16 conversion was changed,
+			//	 it used to clamp to max val if out of range, now stores inf
 		}
 	}
 
diff --git a/Engine/Source/Runtime/Engine/Private/ImageUtils.cpp b/Engine/Source/Runtime/Engine/Private/ImageUtils.cpp
--- a/Engine/Source/Runtime/Engine/Private/ImageUtils.cpp
+++ b/Engine/Source/Runtime/Engine/Private/ImageUtils.cpp
@@ -364,7 +364,7 @@ bool FImageUtils::GetRenderTargetImage(UTextureRenderTarget2D* TexRT, FImage & I
 		return false;
 	}
 	
-	// @@!! in theory the RT knows its gamma
+	// @todo Oodle : in theory the RenderTarget knows its gamma
 	//   but from what I've seen it's usually wrong?
 	// TexRT->GetDisplayGamma() or TexRT->SRGB or TexRT->IsSRGB() , OMG
 
@@ -1214,7 +1214,7 @@ bool FImageUtils::GetTexture2DSourceImage(UTexture2D* Texture, FImage & OutImage
 
 	UE_LOG(LogImageUtils,Warning,TEXT("GetTexture2DSourceImage from PlatformData not implemented yet"));
 
-	// @@!! could export texture from platformdata
+	// @todo Oodle : could export texture from platformdata
 	// only a few formats would be possible to grab to an FImage
 	//   (for DDS export we could do all formats)
 	// eg. see ExportHDR below
diff --git a/Engine/Source/Runtime/Engine/Private/Rendering/Texture2DResource.cpp b/Engine/Source/Runtime/Engine/Private/Rendering/Texture2DResource.cpp
--- a/Engine/Source/Runtime/Engine/Private/Rendering/Texture2DResource.cpp
+++ b/Engine/Source/Runtime/Engine/Private/Rendering/Texture2DResource.cpp
@@ -268,7 +268,7 @@ void FTexture2DResource::GetData( uint32 MipIndex, void* Dest, uint32 DestPitch
 		#if WITH_EDITORONLY_DATA
 		// in Editor, Mip doesn't come from BulkData, it may be null
 		// MipData[] was set from Editor data
-		// @@!! check MipData[MipIndex] size ! but it's not stored
+		// would be nice to check MipData[MipIndex] size ! but it's not stored
 		if ( BulkDataSize == 0 )
 		{
 			BulkDataSize = EffectiveSize;
diff --git a/Engine/Source/Runtime/Engine/Private/Streaming/Texture2DStreamIn.cpp b/Engine/Source/Runtime/Engine/Private/Streaming/Texture2DStreamIn.cpp
--- a/Engine/Source/Runtime/Engine/Private/Streaming/Texture2DStreamIn.cpp
+++ b/Engine/Source/Runtime/Engine/Private/Streaming/Texture2DStreamIn.cpp
@@ -39,7 +39,7 @@ void FTexture2DStreamIn::DoAllocateNewMips(const FContext& Context)
 
 			check(MipData[MipIndex].Data == nullptr);
 			MipData[MipIndex].Data = FMemory::Malloc(MipSize);
-			// @@!! store Size !
+			// would be nice to store Size here !
 			//MipData[MipIndex].Size = MipSize;
 			MipData[MipIndex].Pitch = 0; // 0 means tight packed
 		}
@@ -75,7 +75,7 @@ void FTexture2DStreamIn::DoLockNewMips(const FContext& Context)
 			uint32 DestPitch = -1;
 			MipData[MipIndex].Data = RHILockTexture2D(IntermediateTextureRHI, MipIndex - MipOffset, RLM_WriteOnly, DestPitch, false, CVarFlushRHIThreadOnSTreamingTextureLocks.GetValueOnAnyThread() > 0);
 			MipData[MipIndex].Pitch = DestPitch;
-			// @@!! should store Size but RHILockTexture2D doesn't tell us size
+			// note: should store Size but RHILockTexture2D doesn't tell us size
 			//MipData[MipIndex].Size
 
 			UE_LOG(LogTextureUpload,Verbose,TEXT("FTexture2DStreamIn::DoLockNewMips( : Lock Mip %d Pitch=%d"),MipIndex,DestPitch);
diff --git a/Engine/Source/Runtime/Engine/Private/Streaming/Texture2DStreamIn.h b/Engine/Source/Runtime/Engine/Private/Streaming/Texture2DStreamIn.h
--- a/Engine/Source/Runtime/Engine/Private/Streaming/Texture2DStreamIn.h
+++ b/Engine/Source/Runtime/Engine/Private/Streaming/Texture2DStreamIn.h
@@ -24,7 +24,7 @@ protected:
 	{
 		void * Data = nullptr;
 		uint32 Pitch = 0;
-		// @@!! Size would be nice but LockTexture doesn't return size
+		// note: Size would be nice but LockTexture doesn't return size
 	};
 	TArray<FStreamMipData, TInlineAllocator<MAX_TEXTURE_MIP_COUNT> > MipData;
 	void * InitialMipDataForAsyncCreate[MAX_TEXTURE_MIP_COUNT];
diff --git a/Engine/Source/Runtime/Engine/Private/Streaming/Texture2DStreamIn_IO.cpp b/Engine/Source/Runtime/Engine/Private/Streaming/Texture2DStreamIn_IO.cpp
--- a/Engine/Source/Runtime/Engine/Private/Streaming/Texture2DStreamIn_IO.cpp
+++ b/Engine/Source/Runtime/Engine/Private/Streaming/Texture2DStreamIn_IO.cpp
@@ -34,7 +34,7 @@ FTexture2DStreamIn_IO::~FTexture2DStreamIn_IO()
 
 static void ValidateMipBulkDataSize(const UTexture2D& Texture, int32 MipSizeX, int32 MipSizeY, int32 MipIndex, int64& BulkDataSize)
 {
-// @@!! why is this not done on all platforms?
+	// why is this not done on all platforms?
 #if PLATFORM_ANDROID
 	const int64 ExpectedMipSize = CalcTextureMipMapSize((uint32)MipSizeX, (uint32)MipSizeY, Texture.GetPixelFormat(), 0);
 	if (BulkDataSize != ExpectedMipSize)
@@ -66,7 +66,7 @@ void FTexture2DStreamIn_IO::SetIORequests(const FContext& Context)
 			TaskSynchronization.Increment();
 
 			// Validate buffer size for the mip, so we don't overrun it on streaming
-			// @@!! MipData[] should have size
+			// note: MipData[] should have size
 			// ValidateMipBulkDataSize only does anything on Android
 			ValidateMipBulkDataSize(*Context.Texture, MipMap.SizeX, MipMap.SizeY, MipIndex, BulkDataSize);
 			
diff --git a/Engine/Source/Runtime/Engine/Private/Texture.cpp b/Engine/Source/Runtime/Engine/Private/Texture.cpp
--- a/Engine/Source/Runtime/Engine/Private/Texture.cpp
+++ b/Engine/Source/Runtime/Engine/Private/Texture.cpp
@@ -2968,7 +2968,6 @@ FName GetDefaultTextureFormatName( const ITargetPlatform* TargetPlatform, const
 
 				// NOTE: this function is called like 10 times for each texture
 				//	it makes any log here a bit annoying
-				//  @@!! why is this called over and over?  ; make this log "Display" not "Verbose" when it's less spewy
 
 				UE_LOG(LogTexture, Verbose, TEXT("Texture forced to uncompressed because size is not a multiple of 4 : %dx%d: was %s : %s"), SizeX,SizeY, 
 					*(StaticEnum<TextureCompressionSettings>()->GetNameStringByValue(FormatSettings.CompressionSettings)),
diff --git a/Engine/Source/Runtime/Engine/Private/TextureDerivedData.cpp b/Engine/Source/Runtime/Engine/Private/TextureDerivedData.cpp
--- a/Engine/Source/Runtime/Engine/Private/TextureDerivedData.cpp
+++ b/Engine/Source/Runtime/Engine/Private/TextureDerivedData.cpp
@@ -61,7 +61,8 @@
 // In case of merge conflicts with DDC versions, you MUST generate a new GUID and set this new
 // guid as version
 // this is put in the DDC1 and the DDC2 key
-// @@!! next time this changes clean up SerializeForKey todo marks
+
+// next time this changes clean up SerializeForKey todo marks , search "@todo SerializeForKey"
 #define TEXTURE_DERIVEDDATA_VER		TEXT("95BCE5A0BFB949539A18684748C633C9")
 
 // This GUID is mixed into DDC version for virtual textures only, this allows updating DDC version for VT without invalidating DDC for all textures
@@ -204,7 +205,7 @@ static void SerializeForKey(FArchive& Ar, const FTextureBuildSettings& Settings)
 
 	// weird/unncessary way to do this :
 	// should have just serialized the bool bForceAlphaChannel
-	// @@!! these can go away whenever we bump the overall ddc key
+	// @todo SerializeForKey these can go away whenever we bump the overall ddc key
 	if (Settings.bForceAlphaChannel)
 	{
 		TempGuid = FGuid(0x2C9DF7E3, 0xBC9D413B, 0xBF963C7A, 0x3F27E8B1);
@@ -214,7 +215,7 @@ static void SerializeForKey(FArchive& Ar, const FTextureBuildSettings& Settings)
 	if ( Settings.MaxTextureResolution != FTextureBuildSettings::MaxTextureResolutionDefault &&
 		( Settings.MipGenSettings == TMGS_LeaveExistingMips || Settings.bDoScaleMipsForAlphaCoverage ) )
 	{
-		// @@!! these can go away whenever we bump the overall ddc key
+		// @todo SerializeForKey these can go away whenever we bump the overall ddc key
 		// behavior of MaxTextureResolution + LeaveExistingMips or bDoScaleMipsForAlphaCoverage changed, so modify the key :
 		TempGuid = FGuid(0x418B8584, 0x72D54EA5, 0xBA8E8C2B, 0xECC880DE);
 		Ar << TempGuid;
@@ -222,20 +223,20 @@ static void SerializeForKey(FArchive& Ar, const FTextureBuildSettings& Settings)
 
 	if ( Settings.bVolume )
 	{
-		// @@!! these can go away whenever we bump the overall ddc key
+		// @todo SerializeForKey these can go away whenever we bump the overall ddc key
 		TempGuid = FGuid(0xCC4348B8,0x84714993,0xAB1E2C93,0x8EA6C9E0);
 		Ar << TempGuid;
 	}
 
 	if ( Settings.bVirtualStreamable && Settings.bSRGB && Settings.bUseLegacyGamma )
 	{
-		// @@!! these can go away whenever we bump the overall ddc key
+		// @todo SerializeForKey these can go away whenever we bump the overall ddc key
 		TempGuid = FGuid(0xCAEDDFB6,0xEDC2455D,0x8D45B90C,0x3A1B7783);
 		Ar << TempGuid;
 	}
 
 	// do not change key if old mip filter is used for old textures
-	// @@!! these can go away whenever we bump the overall ddc key
+	// @todo SerializeForKey these can go away whenever we bump the overall ddc key
 	if (Settings.bUseNewMipFilter)
 	{
 		TempGuid = FGuid(0x27B79A99, 0xE1A5458E, 0xAB619475, 0xCD01AD2A);
@@ -244,7 +245,7 @@ static void SerializeForKey(FArchive& Ar, const FTextureBuildSettings& Settings)
 
 	if (Settings.bLongLatSource)
 	{
-		// @@!! these can go away whenever we bump the overall ddc key
+		// @todo SerializeForKey these can go away whenever we bump the overall ddc key
 		// texture processing for cubemaps generated from longlat sources changed, so modify the key :
 		TempGuid = FGuid(0x3D642836, 0xEBF64714, 0x9E8E3241, 0x39F66906);
 		Ar << TempGuid;
@@ -288,9 +289,9 @@ void GetTextureDerivedDataKeySuffix(const UTexture& Texture, const FTextureBuild
 		{
 			Version = TextureFormat->GetVersion(BuildSettings.TextureFormatName, &BuildSettings);
 		}
-		// @@!! else error !?
+		// else error !?
 	}
-	// @@!! else error !?
+	// else error !?
 	
 	FString CompositeTextureStr;
 
diff --git a/Engine/Source/Runtime/Engine/Private/TextureRenderTarget2D.cpp b/Engine/Source/Runtime/Engine/Private/TextureRenderTarget2D.cpp
--- a/Engine/Source/Runtime/Engine/Private/TextureRenderTarget2D.cpp
+++ b/Engine/Source/Runtime/Engine/Private/TextureRenderTarget2D.cpp
@@ -305,7 +305,7 @@ UTexture2D* UTextureRenderTarget2D::ConstructTexture2D(UObject* Outer, const FSt
 	UpdateTexture2D(Result, TextureFormat, Flags, AlphaOverride);
 
 	// if render target gamma used was 1.0 then disable SRGB for the static texture
-	// @@!! UTextureRenderTarget2D also has an explicit SRGB flag in the UTexture parent class
+	// note: UTextureRenderTarget2D also has an explicit SRGB flag in the UTexture parent class
 	//	  these are NOT correctly kept in sync
 	//		I see SRGB = 1 but Gamma = 1.0
 	//	 see also IsSRGB() which is yet another query that has different ideas
diff --git a/Engine/Source/Runtime/Engine/Public/DDSLoader.h b/Engine/Source/Runtime/Engine/Public/DDSLoader.h
--- a/Engine/Source/Runtime/Engine/Public/DDSLoader.h
+++ b/Engine/Source/Runtime/Engine/Public/DDSLoader.h
@@ -93,8 +93,7 @@ struct FDDS10FileHeader
 };
 #pragma pack(pop)
 
-// @@!! still used in InterchangeDDSTranslator
-//   deprecate after that use is removed
+// @todo Oodle : deprecate old DDSLoader after it is removed from InterchangeDDSTranslator 
 // 
 //class UE_DEPRECATED(5.2, "Use DDSFile instead") FDDSLoadHelper;
 class FDDSLoadHelper
diff --git a/Engine/Source/Runtime/ImageWrapper/Private/Formats/BmpImageWrapper.cpp b/Engine/Source/Runtime/ImageWrapper/Private/Formats/BmpImageWrapper.cpp
--- a/Engine/Source/Runtime/ImageWrapper/Private/Formats/BmpImageWrapper.cpp
+++ b/Engine/Source/Runtime/ImageWrapper/Private/Formats/BmpImageWrapper.cpp
@@ -353,7 +353,7 @@ ERawImageFormat::Type FBmpImageWrapper::GetSupportedRawFormat(const ERawImageFor
 	case ERawImageFormat::RGBA16F:
 	case ERawImageFormat::RGBA32F:
 	case ERawImageFormat::R16F:
-	case ERawImageFormat::R32F: // @@!! can these go to G8 ?
+	case ERawImageFormat::R32F:
 		return ERawImageFormat::BGRA8; // needs conversion
 	default:
 		check(0);
diff --git a/Engine/Source/Runtime/ImageWrapper/Private/Formats/JpegImageWrapper.cpp b/Engine/Source/Runtime/ImageWrapper/Private/Formats/JpegImageWrapper.cpp
--- a/Engine/Source/Runtime/ImageWrapper/Private/Formats/JpegImageWrapper.cpp
+++ b/Engine/Source/Runtime/ImageWrapper/Private/Formats/JpegImageWrapper.cpp
@@ -152,7 +152,7 @@ ERawImageFormat::Type FJpegImageWrapper::GetSupportedRawFormat(const ERawImageFo
 	case ERawImageFormat::RGBA16:
 	case ERawImageFormat::RGBA16F:
 	case ERawImageFormat::RGBA32F:
-	case ERawImageFormat::G16: // @@!! use G8 ?
+	case ERawImageFormat::G16:
 	case ERawImageFormat::R16F:
 	case ERawImageFormat::R32F:
 		return ERawImageFormat::BGRA8; // needs conversion
