SubObject
* Added CVar that controls if actors and actorcomponent classes use the new subobject replication by default or not
* net.SubObjects.DefaultUseSubObjectReplicationList needs to be enabled to opt-in to the new method by default.
* Classes can still choose their own method by setting bReplicateUsingRegisteredSubObjectList in their constructor

diff --git a/Engine/Source/Runtime/Engine/Private/Actor.cpp b/Engine/Source/Runtime/Engine/Private/Actor.cpp
--- a/Engine/Source/Runtime/Engine/Private/Actor.cpp
+++ b/Engine/Source/Runtime/Engine/Private/Actor.cpp
@@ -82,6 +82,9 @@ DECLARE_CYCLE_STAT(TEXT("PostActorConstruction"), STAT_PostActorConstruction, ST
 
 FMakeNoiseDelegate AActor::MakeNoiseDelegate = FMakeNoiseDelegate::CreateStatic(&AActor::MakeNoiseImpl);
 
+/** Selects if actor and actorcomponents will replicate subobjects using the registration list by default. */
+extern bool GDefaultUseSubObjectReplicationList;
+
 #if !UE_BUILD_SHIPPING
 FOnProcessEvent AActor::ProcessEventDelegate;
 #endif
@@ -126,7 +129,7 @@ void AActor::InitializeDefaults()
 	bReplicates = false;
 	bCallPreReplication = true;
 	bCallPreReplicationForReplay = true;
-	bReplicateUsingRegisteredSubObjectList = DEFAULT_USE_SUBOBJECT_LIST;
+	bReplicateUsingRegisteredSubObjectList = GDefaultUseSubObjectReplicationList;
 	NetPriority = 1.0f;
 	NetUpdateFrequency = 100.0f;
 	MinNetUpdateFrequency = 2.0f;
diff --git a/Engine/Source/Runtime/Engine/Private/Components/ActorComponent.cpp b/Engine/Source/Runtime/Engine/Private/Components/ActorComponent.cpp
--- a/Engine/Source/Runtime/Engine/Private/Components/ActorComponent.cpp
+++ b/Engine/Source/Runtime/Engine/Private/Components/ActorComponent.cpp
@@ -80,6 +80,13 @@ FAutoConsoleVariableRef GTickComponentLatentActionsWithTheComponentCVar(
 /** Static var indicating activity of reregister context */
 int32 FGlobalComponentReregisterContext::ActiveGlobalReregisterContextCount = 0;
 
+bool GDefaultUseSubObjectReplicationList = false;
+static FAutoConsoleVariableRef CVarDefaultUseSubObjectReplicationList(
+	TEXT("net.SubObjects.DefaultUseSubObjectReplicationList"),
+	GDefaultUseSubObjectReplicationList,
+	TEXT("Do actors and actorcomponents replicate subobjects using the registration method by default."));
+
+
 #if WITH_CHAOS
 // Allows for CreatePhysicsState to be deferred, to batch work and parallelize.
 int32 GEnableDeferredPhysicsCreation = 0;
@@ -333,7 +340,7 @@ UActorComponent::UActorComponent(const FObjectInitializer& ObjectInitializer /*=
 	bMarkedForPreEndOfFrameSync = false;
 	bAsyncPhysicsTickEnabled = false;
 
-	bReplicateUsingRegisteredSubObjectList = DEFAULT_USE_SUBOBJECT_LIST;
+	bReplicateUsingRegisteredSubObjectList = GDefaultUseSubObjectReplicationList;
 }
 
 void UActorComponent::PostInitProperties()
diff --git a/Engine/Source/Runtime/Net/Core/Public/Net/Core/Misc/NetSubObjectRegistry.h b/Engine/Source/Runtime/Net/Core/Public/Net/Core/Misc/NetSubObjectRegistry.h
--- a/Engine/Source/Runtime/Net/Core/Public/Net/Core/Misc/NetSubObjectRegistry.h
+++ b/Engine/Source/Runtime/Net/Core/Public/Net/Core/Misc/NetSubObjectRegistry.h
@@ -7,15 +7,6 @@
 
 class UActorComponent;
 
-/** 
-* Chooses the default value of Actor and ActorComponent's bReplicateUsingRegisteredSubObjectList.  
-* If a project wants to replicate subobjects using the registration list by default it should define it to true in it's TargetRules.
-* Classes can also chose to use the list or not by setting bReplicateUsingRegisteredSubObjectList in their config or constructor.
-*/
-#ifndef DEFAULT_USE_SUBOBJECT_LIST
-	#define DEFAULT_USE_SUBOBJECT_LIST false
-#endif
-
 namespace UE::Net 
 {
 
