no functionality change, misc comments

diff --git a/Engine/Source/Developer/DerivedDataCache/Public/DerivedDataBuildFunctionFactory.h b/Engine/Source/Developer/DerivedDataCache/Public/DerivedDataBuildFunctionFactory.h
--- a/Engine/Source/Developer/DerivedDataCache/Public/DerivedDataBuildFunctionFactory.h
+++ b/Engine/Source/Developer/DerivedDataCache/Public/DerivedDataBuildFunctionFactory.h
@@ -32,7 +32,7 @@ public:
 template <typename FunctionType>
 class TBuildFunctionFactory final : public IBuildFunctionFactory
 {
-	static_assert(sizeof(FunctionType) == sizeof(IBuildFunction), "IBuildFunction must be pure and maintain no state.");
+	static_assert(sizeof(FunctionType) == sizeof(IBuildFunction), "Derivations of IBuildFunction must be pure and maintain no state.");
 
 public:
 	template <typename... ArgTypes>
diff --git a/Engine/Source/Runtime/Engine/Private/TextureDerivedDataTask.h b/Engine/Source/Runtime/Engine/Private/TextureDerivedDataTask.h
--- a/Engine/Source/Runtime/Engine/Private/TextureDerivedDataTask.h
+++ b/Engine/Source/Runtime/Engine/Private/TextureDerivedDataTask.h
@@ -40,7 +40,12 @@ enum class ETextureCacheFlags : uint32
 	None			= 0x00,
 	Async			= 0x01,
 	ForceRebuild	= 0x02,
-	InlineMips		= 0x08,
+	
+	/** 
+	* If specified, all mips will be loaded in to BulkData at the end of the Cache() operation, prior to return. Without this, 
+	* only "Inline" mips will be in BulkData (for streaming textures, typically the smallest NUM_INLINE_DERIVED_MIPS mips).
+	*/
+	InlineMips		= 0x08,	
 	AllowAsyncBuild	= 0x10,
 	ForDDCBuild		= 0x20,
 	RemoveSourceMipDataAfterCache = 0x40,
@@ -297,6 +302,8 @@ FTextureAsyncCacheDerivedDataTask* CreateTextureBuildTask(
 	EQueuedWorkPriority Priority,
 	ETextureCacheFlags Flags);
 
+// This needs to match the key generated by the build, otherwise the build will get continually kicked
+// by CachePlatformData().
 FTexturePlatformData::FStructuredDerivedDataKey CreateTextureDerivedDataKey(
 	UTexture& Texture,
 	ETextureCacheFlags CacheFlags,
diff --git a/Engine/Source/Runtime/Engine/Private/VolumeTexture.cpp b/Engine/Source/Runtime/Engine/Private/VolumeTexture.cpp
--- a/Engine/Source/Runtime/Engine/Private/VolumeTexture.cpp
+++ b/Engine/Source/Runtime/Engine/Private/VolumeTexture.cpp
@@ -22,7 +22,7 @@
 // Master switch to control whether streaming is enabled for volume texture. 
 bool GSupportsVolumeTextureStreaming = true;
 
-// Limit the possible depth of volume texture otherwise when the user converts 2D textures, he can crash the engine.
+// Limit the possible depth of volume texture otherwise when the user converts 2D textures, they can crash the engine.
 const int32 MAX_VOLUME_TEXTURE_DEPTH = 512;
 
 //*****************************************************************************
