Incorrect relative Transform data on client when using USceneComponent::AttachToComponent with EAttachmentRule::KeepWorld under certain conditions.

diff --git a/Engine/Source/Runtime/Engine/Classes/Components/SceneComponent.h b/Engine/Source/Runtime/Engine/Classes/Components/SceneComponent.h
--- a/Engine/Source/Runtime/Engine/Classes/Components/SceneComponent.h
+++ b/Engine/Source/Runtime/Engine/Classes/Components/SceneComponent.h
@@ -186,6 +186,9 @@ private:
 	UPROPERTY(Transient, Replicated)
 	uint8 bShouldSnapRotationWhenAttached : 1;
 
+	UPROPERTY(Transient, Replicated)
+	uint8 bShouldSnapScaleWhenAttached : 1;
+
 	/** Sets bShouldBeAttached, push model aware. */
 	void SetShouldBeAttached(bool bNewShouldBeAttached);
 
@@ -195,6 +198,9 @@ private:
 	/** Sets bShouldSnapRotationWhenAttached, push model aware. */
 	void SetShouldSnapRotationWhenAttached(bool bShouldSnapRotation);
 
+	/** Sets bShouldSnapScaleWhenAttached, push model aware. */
+	void SetShouldSnapScaleWhenAttached(bool bShouldSnapScale);
+
 	/**
 	 * Whether or not the cached PhysicsVolume this component overlaps should be updated when the component is moved.
 	 * @see GetPhysicsVolume()
diff --git a/Engine/Source/Runtime/Engine/Private/Components/SceneComponent.cpp b/Engine/Source/Runtime/Engine/Private/Components/SceneComponent.cpp
--- a/Engine/Source/Runtime/Engine/Private/Components/SceneComponent.cpp
+++ b/Engine/Source/Runtime/Engine/Private/Components/SceneComponent.cpp
@@ -689,6 +689,7 @@ void USceneComponent::OnRegister()
 			SetShouldBeAttached(false);
 			SetShouldSnapLocationWhenAttached(false);
 			SetShouldSnapRotationWhenAttached(false);
+			SetShouldSnapScaleWhenAttached(false);
 		}
 	}
 	
@@ -1870,6 +1871,7 @@ bool USceneComponent::AttachToComponent(USceneComponent* Parent, const FAttachme
 		SetupAttachment(Parent, SocketName);
 		SetShouldSnapLocationWhenAttached(false);
 		SetShouldSnapRotationWhenAttached(false);
+		SetShouldSnapScaleWhenAttached(false);
 
 		return true;
 	}
@@ -2034,8 +2036,10 @@ bool USceneComponent::AttachToComponent(USceneComponent* Parent, const FAttachme
 		SetAttachSocketName(SocketName);
 		SetShouldBeAttached(AttachParent != nullptr);
 
-		SetShouldSnapLocationWhenAttached(AttachmentRules.LocationRule == EAttachmentRule::SnapToTarget);
-		SetShouldSnapRotationWhenAttached(AttachmentRules.RotationRule == EAttachmentRule::SnapToTarget);
+		// Tell Client to snap to target if we use SnapToTarget rule or if we are using KeepWorld and Transform is same as Parent
+		SetShouldSnapLocationWhenAttached(AttachmentRules.LocationRule == EAttachmentRule::SnapToTarget || (AttachmentRules.LocationRule == EAttachmentRule::KeepWorld && GetRelativeLocation() == Parent->GetRelativeLocation()));
+		SetShouldSnapRotationWhenAttached(AttachmentRules.RotationRule == EAttachmentRule::SnapToTarget || (AttachmentRules.RotationRule == EAttachmentRule::KeepWorld && GetRelativeRotation() == Parent->GetRelativeRotation()));
+		SetShouldSnapScaleWhenAttached(   AttachmentRules.ScaleRule    == EAttachmentRule::SnapToTarget || (AttachmentRules.ScaleRule    == EAttachmentRule::KeepWorld && GetRelativeScale3D()  == Parent->GetRelativeScale3D()));
 
 		OnAttachmentChanged();
 
@@ -2239,6 +2243,7 @@ void USceneComponent::DetachFromComponent(const FDetachmentTransformRules& Detac
 
 		SetShouldSnapLocationWhenAttached(false);
 		SetShouldSnapRotationWhenAttached(false);
+		SetShouldSnapScaleWhenAttached(false);
 
 		OnAttachmentChanged();
 
@@ -3344,6 +3349,10 @@ void USceneComponent::PostRepNotifies()
 		{
 			SetRelativeRotation_Direct(FRotator::ZeroRotator);
 		}
+		if (bShouldSnapScaleWhenAttached && !bNetUpdateTransform)
+		{
+			SetRelativeScale3D_Direct(FVector::OneVector);
+		}
 
 		// Check if this is a detach
 		if (AttachParent && !bShouldBeAttached)
@@ -3356,6 +3365,7 @@ void USceneComponent::PostRepNotifies()
 			const bool bOldShouldBeAttached = bShouldBeAttached;
 			const bool bOldShouldSnapLocationWhenAttached = bShouldSnapLocationWhenAttached;
 			const bool bOldShouldSnapRotationWhenAttached = bShouldSnapRotationWhenAttached;
+			const bool bOldShouldSnapScaleWhenAttached    = bShouldSnapScaleWhenAttached;
 
 			AttachToComponent(NetOldAttachParent, FAttachmentTransformRules::KeepRelativeTransform, NetOldAttachSocketName);
 
@@ -3363,6 +3373,7 @@ void USceneComponent::PostRepNotifies()
 			SetShouldBeAttached(bOldShouldBeAttached);
 			SetShouldSnapLocationWhenAttached(bOldShouldSnapLocationWhenAttached);
 			SetShouldSnapRotationWhenAttached(bOldShouldSnapRotationWhenAttached);
+			SetShouldSnapScaleWhenAttached(bOldShouldSnapScaleWhenAttached);
 		}
 
 		bNetUpdateAttachment = false;
@@ -3410,6 +3421,7 @@ void USceneComponent::GetLifetimeReplicatedProps(TArray< FLifetimeProperty > & O
 	DOREPLIFETIME_WITH_PARAMS_FAST(USceneComponent, bShouldBeAttached, SharedParams);
 	DOREPLIFETIME_WITH_PARAMS_FAST(USceneComponent, bShouldSnapLocationWhenAttached, SharedParams);
 	DOREPLIFETIME_WITH_PARAMS_FAST(USceneComponent, bShouldSnapRotationWhenAttached, SharedParams);
+	DOREPLIFETIME_WITH_PARAMS_FAST(USceneComponent, bShouldSnapScaleWhenAttached, SharedParams);
 	DOREPLIFETIME_WITH_PARAMS_FAST(USceneComponent, AttachParent, SharedParams);
 	DOREPLIFETIME_WITH_PARAMS_FAST(USceneComponent, AttachChildren, SharedParams);
 	DOREPLIFETIME_WITH_PARAMS_FAST(USceneComponent, AttachSocketName, SharedParams);
@@ -3802,4 +3814,10 @@ void USceneComponent::SetShouldSnapRotationWhenAttached(bool bShouldSnapRotation
 	MARK_PROPERTY_DIRTY_FROM_NAME(USceneComponent, bShouldSnapRotationWhenAttached, this);
 }
 
+void USceneComponent::SetShouldSnapScaleWhenAttached(bool bShouldSnapScale)
+{
+	bShouldSnapScaleWhenAttached = bShouldSnapScale;
+	MARK_PROPERTY_DIRTY_FROM_NAME(USceneComponent, bShouldSnapScaleWhenAttached, this);
+}
+
 #undef LOCTEXT_NAMESPACE
