Call NotifyControllerChanged from ClientRestart in case it is received with the valid pawn before APawn::OnRep_Controller is called

diff --git a/Engine/Source/Runtime/Engine/Private/PlayerController.cpp b/Engine/Source/Runtime/Engine/Private/PlayerController.cpp
--- a/Engine/Source/Runtime/Engine/Private/PlayerController.cpp
+++ b/Engine/Source/Runtime/Engine/Private/PlayerController.cpp
@@ -748,7 +748,13 @@ void APlayerController::ClientRestart_Implementation(APawn* NewPawn)
 	// but we may receive the function call before Pawn is replicated over, so it will resolve to NULL.
 	AcknowledgePossession(GetPawn());
 
+	AController* OldController = GetPawn()->Controller;
 	GetPawn()->Controller = this;
+	if (OldController != this)
+	{
+		// In case this is received before APawn::OnRep_Controller is called
+		GetPawn()->NotifyControllerChanged();
+	}
 	GetPawn()->DispatchRestart(true);
 	
 	if (GetLocalRole() < ROLE_Authority)
