PR #8692: Convert second to millisecond (Contributed by [REDACTED_PERSON])

diff --git a/Engine/Source/Runtime/Engine/Private/UnrealEngine.cpp b/Engine/Source/Runtime/Engine/Private/UnrealEngine.cpp
--- a/Engine/Source/Runtime/Engine/Private/UnrealEngine.cpp
+++ b/Engine/Source/Runtime/Engine/Private/UnrealEngine.cpp
@@ -2237,7 +2237,7 @@ void UEngine::TickDeferredCommands()
 		}
 	}
 
-	const double ElapsedTimeMS = (FPlatformTime::Seconds() - StartTime) / 1000.0;
+	const double ElapsedTimeMS = (FPlatformTime::Seconds() - StartTime) * 1000.0;
 
 	// If we're not in the editor, and commands took more than our target frame time to execute, print them out so there's a paper trail
 	if (GIsEditor == false && ElapsedTimeMS >= FEnginePerformanceTargets::GetTargetFrameTimeThresholdMS())
