Fix a bug of alpha detection in virtual texture builder causing tiles in the same VT to use different formats.

diff --git a/Engine/Source/Runtime/Engine/Private/VT/VirtualTextureDataBuilder.cpp b/Engine/Source/Runtime/Engine/Private/VT/VirtualTextureDataBuilder.cpp
--- a/Engine/Source/Runtime/Engine/Private/VT/VirtualTextureDataBuilder.cpp
+++ b/Engine/Source/Runtime/Engine/Private/VT/VirtualTextureDataBuilder.cpp
@@ -580,7 +580,7 @@ void FVirtualTextureDataBuilder::BuildTiles(const TArray<FVTSourceTileEntry>& Ti
 		TBSettings.bUseLegacyGamma = BuildSettingsForLayer.bUseLegacyGamma;
 		TBSettings.MipGenSettings = TMGS_NoMipmaps;
 		TBSettings.bForceAlphaChannel = BuildSettingsForLayer.bForceAlphaChannel || LayerData.bHasAlpha;
-		TBSettings.bForceNoAlphaChannel = BuildSettingsForLayer.bForceNoAlphaChannel;
+		TBSettings.bForceNoAlphaChannel = BuildSettingsForLayer.bForceNoAlphaChannel || !LayerData.bHasAlpha;
 		TBSettings.bHDRSource = BuildSettingsForLayer.bHDRSource;
 		TBSettings.bVirtualStreamable = true;
 
