World Partition: cleaned up loader adapters code, splitted classes into separated files.

#rnx

diff --git a/Engine/Source/Runtime/Engine/Private/WorldPartition/LoaderAdapter/LoaderAdapterList.cpp b/Engine/Source/Runtime/Engine/Private/WorldPartition/LoaderAdapter/LoaderAdapterList.cpp
new file mode 100644
--- /dev/null
+++ b/Engine/Source/Runtime/Engine/Private/WorldPartition/LoaderAdapter/LoaderAdapterList.cpp
@@ -0,0 +1,17 @@
+// Copyright Epic Games, Inc. All Rights Reserved.
+
+#include "WorldPartition/LoaderAdapter/LoaderAdapterList.h"
+
+#if WITH_EDITOR
+FLoaderAdapterList::FLoaderAdapterList(UWorld* InWorld)
+	: ILoaderAdapter(InWorld)
+{}
+
+void FLoaderAdapterList::ForEachActor(TFunctionRef<void(const FWorldPartitionHandle&)> InOperation) const
+{
+	for (const FWorldPartitionHandle& Actor : Actors)
+	{
+		InOperation(Actor);
+	}
+}
+#endif
diff --git a/Engine/Source/Runtime/Engine/Private/WorldPartition/LoaderAdapter/LoaderAdapterSpatial.cpp b/Engine/Source/Runtime/Engine/Private/WorldPartition/LoaderAdapter/LoaderAdapterSpatial.cpp
new file mode 100644
--- /dev/null
+++ b/Engine/Source/Runtime/Engine/Private/WorldPartition/LoaderAdapter/LoaderAdapterSpatial.cpp
@@ -0,0 +1,29 @@
+// Copyright Epic Games, Inc. All Rights Reserved.
+
+#include "WorldPartition/LoaderAdapter/LoaderAdapterSpatial.h"
+#include "WorldPartition/WorldPartition.h"
+#include "WorldPartition/WorldPartitionEditorHash.h"
+#include "Engine/World.h"
+
+#if WITH_EDITOR
+ILoaderAdapterSpatial::ILoaderAdapterSpatial(UWorld* InWorld)
+	: ILoaderAdapter(InWorld)
+	, bIncludeSpatiallyLoadedActors(true)
+	, bIncludeNonSpatiallyLoadedActors(false)
+{}
+
+void ILoaderAdapterSpatial::ForEachActor(TFunctionRef<void(const FWorldPartitionHandle&)> InOperation) const
+{
+	if (UWorldPartition* WorldPartition = GetWorld()->GetWorldPartition())
+	{
+		WorldPartition->EditorHash->ForEachIntersectingActor(*GetBoundingBox(), [this, WorldPartition, &InOperation](FWorldPartitionActorDesc* ActorDesc)
+		{
+			if (Intersect(ActorDesc->GetBounds()))
+			{
+				FWorldPartitionHandle ActorHandle(WorldPartition, ActorDesc->GetGuid());
+				InOperation(ActorHandle);
+			}
+		}, bIncludeSpatiallyLoadedActors, bIncludeNonSpatiallyLoadedActors);
+	}
+}
+#endif
diff --git a/Engine/Source/Runtime/Engine/Private/WorldPartition/WorldPartition.cpp b/Engine/Source/Runtime/Engine/Private/WorldPartition/WorldPartition.cpp
--- a/Engine/Source/Runtime/Engine/Private/WorldPartition/WorldPartition.cpp
+++ b/Engine/Source/Runtime/Engine/Private/WorldPartition/WorldPartition.cpp
@@ -45,6 +45,7 @@
 #include "WorldPartition/DataLayer/DataLayerInstance.h"
 #include "WorldPartition/DataLayer/DataLayerSubsystem.h"
 #include "WorldPartition/DataLayer/WorldDataLayers.h"
+#include "WorldPartition/LoaderAdapter/LoaderAdapterList.h"
 #include "WorldPartition/LoaderAdapter/LoaderAdapterShape.h"
 #include "WorldPartition/WorldPartitionActorDescViewProxy.h"
 #include "WorldPartition/HLOD/HLODLayer.h"
@@ -210,11 +211,11 @@ public:
 	}
 };
 
-class FLoaderAdapterPinnedActors: public IWorldPartitionActorLoaderInterface::ILoaderAdapterList
+class FLoaderAdapterPinnedActors: public FLoaderAdapterList
 {
 public:
 	FLoaderAdapterPinnedActors(UWorld* InWorld)
-		: IWorldPartitionActorLoaderInterface::ILoaderAdapterList(InWorld)
+		: FLoaderAdapterList(InWorld)
 	{
 		Load();
 	}
@@ -252,12 +253,7 @@ protected:
 	//~ Begin IWorldPartitionActorLoaderInterface::ILoaderAdapterList interface
 	virtual bool ShouldActorBeLoaded(const FWorldPartitionHandle& ActorHandle) const override
 	{
-		if (ContainsActor(ActorHandle) && (ActorHandle == ActorToRemove))
-		{
-			return false;
-		}
-
-		return IWorldPartitionActorLoaderInterface::ILoaderAdapterList::ShouldActorBeLoaded(ActorHandle);
+		return (ActorHandle != ActorToRemove) && FLoaderAdapterList::ShouldActorBeLoaded(ActorHandle);
 	}
 	//~ End IWorldPartitionActorLoaderInterface::ILoaderAdapterList interface
 
diff --git a/Engine/Source/Runtime/Engine/Private/WorldPartition/WorldPartitionActorLoaderInterface.cpp b/Engine/Source/Runtime/Engine/Private/WorldPartition/WorldPartitionActorLoaderInterface.cpp
--- a/Engine/Source/Runtime/Engine/Private/WorldPartition/WorldPartitionActorLoaderInterface.cpp
+++ b/Engine/Source/Runtime/Engine/Private/WorldPartition/WorldPartitionActorLoaderInterface.cpp
@@ -307,39 +307,6 @@ void IWorldPartitionActorLoaderInterface::ILoaderAdapter::OnActorDataLayersEdito
 		GEditor->ResetTransaction(LOCTEXT("LoadingEditorActorResetTrans", "Editor Actors Loading State Changed"));
 	}
 }
-
-IWorldPartitionActorLoaderInterface::ILoaderAdapterList::ILoaderAdapterList(UWorld* InWorld)
-	: ILoaderAdapter(InWorld)
-{}
-
-void IWorldPartitionActorLoaderInterface::ILoaderAdapterList::ForEachActor(TFunctionRef<void(const FWorldPartitionHandle&)> InOperation) const
-{
-	for (const FWorldPartitionHandle& Actor : Actors)
-	{
-		InOperation(Actor);
-	}
-}
-
-IWorldPartitionActorLoaderInterface::ILoaderAdapterSpatial::ILoaderAdapterSpatial(UWorld* InWorld)
-	: ILoaderAdapter(InWorld)
-	, bIncludeSpatiallyLoadedActors(true)
-	, bIncludeNonSpatiallyLoadedActors(false)
-{}
-
-void IWorldPartitionActorLoaderInterface::ILoaderAdapterSpatial::ForEachActor(TFunctionRef<void(const FWorldPartitionHandle&)> InOperation) const
-{
-	if (UWorldPartition* WorldPartition = GetWorld()->GetWorldPartition())
-	{
-		WorldPartition->EditorHash->ForEachIntersectingActor(*GetBoundingBox(), [this, WorldPartition, &InOperation](FWorldPartitionActorDesc* ActorDesc)
-		{
-			if (Intersect(ActorDesc->GetBounds()))
-			{
-				FWorldPartitionHandle ActorHandle(WorldPartition, ActorDesc->GetGuid());
-				InOperation(ActorHandle);
-			}
-		}, bIncludeSpatiallyLoadedActors, bIncludeNonSpatiallyLoadedActors);
-	}
-}
 #endif
 
 #undef LOCTEXT_NAMESPACE
\ No newline at end of file
diff --git a/Engine/Source/Runtime/Engine/Public/WorldPartition/LoaderAdapter/LoaderAdapterActor.h b/Engine/Source/Runtime/Engine/Public/WorldPartition/LoaderAdapter/LoaderAdapterActor.h
--- a/Engine/Source/Runtime/Engine/Public/WorldPartition/LoaderAdapter/LoaderAdapterActor.h
+++ b/Engine/Source/Runtime/Engine/Public/WorldPartition/LoaderAdapter/LoaderAdapterActor.h
@@ -2,11 +2,11 @@
 #pragma once
 
 #include "CoreMinimal.h"
-#include "WorldPartition/WorldPartitionActorLoaderInterface.h"
+#include "WorldPartition/LoaderAdapter/LoaderAdapterSpatial.h"
 
 #if WITH_EDITOR
 /** Actor loader adapter that will do a spatial query based on the actor's brush  */
-class ENGINE_API FLoaderAdapterActor : public IWorldPartitionActorLoaderInterface::ILoaderAdapterSpatial
+class ENGINE_API FLoaderAdapterActor : public ILoaderAdapterSpatial
 {
 public:
 	FLoaderAdapterActor(AActor* InActor);
diff --git a/Engine/Source/Runtime/Engine/Public/WorldPartition/LoaderAdapter/LoaderAdapterList.h b/Engine/Source/Runtime/Engine/Public/WorldPartition/LoaderAdapter/LoaderAdapterList.h
new file mode 100644
--- /dev/null
+++ b/Engine/Source/Runtime/Engine/Public/WorldPartition/LoaderAdapter/LoaderAdapterList.h
@@ -0,0 +1,22 @@
+// Copyright Epic Games, Inc. All Rights Reserved.
+#pragma once
+
+#include "CoreMinimal.h"
+#include "WorldPartition/WorldPartitionActorLoaderInterface.h"
+
+#if WITH_EDITOR
+/** Base class for actor loaders that contains a specific list of actors */
+class ENGINE_API FLoaderAdapterList : public IWorldPartitionActorLoaderInterface::ILoaderAdapter
+{
+public:
+	FLoaderAdapterList(UWorld* InWorld);
+	virtual ~FLoaderAdapterList() {}
+
+protected:
+	//~ Begin ILoaderAdapter interface
+	virtual void ForEachActor(TFunctionRef<void(const FWorldPartitionHandle&)> InOperation) const override;
+	//~ End ILoaderAdapter interface
+
+	TSet<FWorldPartitionHandle> Actors;
+};
+#endif
\ No newline at end of file
diff --git a/Engine/Source/Runtime/Engine/Public/WorldPartition/LoaderAdapter/LoaderAdapterShape.h b/Engine/Source/Runtime/Engine/Public/WorldPartition/LoaderAdapter/LoaderAdapterShape.h
--- a/Engine/Source/Runtime/Engine/Public/WorldPartition/LoaderAdapter/LoaderAdapterShape.h
+++ b/Engine/Source/Runtime/Engine/Public/WorldPartition/LoaderAdapter/LoaderAdapterShape.h
@@ -2,11 +2,11 @@
 #pragma once
 
 #include "CoreMinimal.h"
-#include "WorldPartition/WorldPartitionActorLoaderInterface.h"
+#include "WorldPartition/LoaderAdapter/LoaderAdapterSpatial.h"
 
 #if WITH_EDITOR
 /** Actor loader adapter that will do a spatial query based on a shape */
-class ENGINE_API FLoaderAdapterShape : public IWorldPartitionActorLoaderInterface::ILoaderAdapterSpatial
+class ENGINE_API FLoaderAdapterShape : public ILoaderAdapterSpatial
 {
 public:
 	FLoaderAdapterShape(UWorld* InWorld, const FBox& InBoundingBox, const FString& InLabel);
diff --git a/Engine/Source/Runtime/Engine/Public/WorldPartition/LoaderAdapter/LoaderAdapterSpatial.h b/Engine/Source/Runtime/Engine/Public/WorldPartition/LoaderAdapter/LoaderAdapterSpatial.h
new file mode 100644
--- /dev/null
+++ b/Engine/Source/Runtime/Engine/Public/WorldPartition/LoaderAdapter/LoaderAdapterSpatial.h
@@ -0,0 +1,26 @@
+// Copyright Epic Games, Inc. All Rights Reserved.
+#pragma once
+
+#include "CoreMinimal.h"
+#include "WorldPartition/WorldPartitionActorLoaderInterface.h"
+
+#if WITH_EDITOR
+/** Base class for actor loaders that requires spatial queries */
+class ENGINE_API ILoaderAdapterSpatial : public IWorldPartitionActorLoaderInterface::ILoaderAdapter
+{
+public:
+	ILoaderAdapterSpatial(UWorld* InWorld);
+	virtual ~ILoaderAdapterSpatial() {}
+
+protected:
+	//~ Begin ILoaderAdapter interface
+	virtual void ForEachActor(TFunctionRef<void(const FWorldPartitionHandle&)> InOperation) const override;
+	//~ End ILoaderAdapter interface
+
+	// Private interface
+	virtual bool Intersect(const FBox& Box) const =0;
+
+	bool bIncludeSpatiallyLoadedActors;
+	bool bIncludeNonSpatiallyLoadedActors;
+};
+#endif
\ No newline at end of file
diff --git a/Engine/Source/Runtime/Engine/Public/WorldPartition/WorldPartitionActorLoaderInterface.h b/Engine/Source/Runtime/Engine/Public/WorldPartition/WorldPartitionActorLoaderInterface.h
--- a/Engine/Source/Runtime/Engine/Public/WorldPartition/WorldPartitionActorLoaderInterface.h
+++ b/Engine/Source/Runtime/Engine/Public/WorldPartition/WorldPartitionActorLoaderInterface.h
@@ -43,7 +43,6 @@ public:
 
 	protected:
 		// Private interface
-		virtual bool Intersect(const FBox& Box) const =0;
 		virtual void ForEachActor(TFunctionRef<void(const FWorldPartitionHandle&)> InOperation) const =0;
 		virtual bool ShouldActorBeLoaded(const FWorldPartitionHandle& Actor) const;
 
@@ -73,41 +72,6 @@ public:
 		TMap<FGuid, TMap<FGuid, FWorldPartitionReference>> ActorReferences;
 	};
 
-	/** Base class for actor loaders that contains a specific list of actors */
-	class ENGINE_API ILoaderAdapterList : public ILoaderAdapter
-	{
-	public:
-		ILoaderAdapterList(UWorld* InWorld);
-		virtual ~ILoaderAdapterList() {}
-
-	protected:
-		//~ Begin ILoaderAdapter interface
-		virtual bool Intersect(const FBox& Box) const { return true; }
-		virtual void ForEachActor(TFunctionRef<void(const FWorldPartitionHandle&)> InOperation) const;
-		//~ End ILoaderAdapter interface
-
-		TSet<FWorldPartitionHandle> Actors;
-	};
-
-	/** Base class for actor loaders that requires spatial queries */
-	class ENGINE_API ILoaderAdapterSpatial : public ILoaderAdapter
-	{
-	public:
-		ILoaderAdapterSpatial(UWorld* InWorld);
-		virtual ~ILoaderAdapterSpatial() {}
-
-	protected:
-		//~ Begin ILoaderAdapter interface
-		virtual void ForEachActor(TFunctionRef<void(const FWorldPartitionHandle&)> InOperation) const;
-		//~ End ILoaderAdapter interface
-
-		// Interface
-		virtual bool Intersect(const FBox& Box) const =0;
-
-		bool bIncludeSpatiallyLoadedActors;
-		bool bIncludeNonSpatiallyLoadedActors;
-	};
-
 	virtual ILoaderAdapter* GetLoaderAdapter() =0;
 #endif
 };
\ No newline at end of file
