Fix primitive disappearing when EvaluateWPO is only turned on programatically and we are using velocity in depth pass.
Issue was that static draw command recorded velocity relevance with EvaluateWPO disabled.
Static mesh relevance needs to track if we CAN write velocity not if the current settings DO write velocity.
#rnx

diff --git a/Engine/Source/Runtime/Engine/Public/PrimitiveSceneProxy.h b/Engine/Source/Runtime/Engine/Public/PrimitiveSceneProxy.h
--- a/Engine/Source/Runtime/Engine/Public/PrimitiveSceneProxy.h
+++ b/Engine/Source/Runtime/Engine/Public/PrimitiveSceneProxy.h
@@ -721,10 +721,10 @@ public:
 	
 	/** Returns true if this proxy can change transform so that we should cache previous transform for calculating velocity. */
 	inline bool HasDynamicTransform() const { return IsMovable() || bIsBeingMovedByEditor; }
-	/** Returns true if this proxy should always write velocity, even if transform doesn't change. */
-	inline bool AlwaysHasVelocity() const {	return bAlwaysHasVelocity || (bEvaluateWorldPositionOffset && bHasWorldPositionOffsetVelocity);	}
-	/** Returns true if this proxy should set velocity relevance. */
-	inline bool DrawsVelocity() const { return HasDynamicTransform() || AlwaysHasVelocity(); }
+	/** Returns true if this proxy can write velocity. This is used for setting velocity relevance. */
+	inline bool DrawsVelocity() const { return HasDynamicTransform() || bAlwaysHasVelocity || bHasWorldPositionOffsetVelocity; }
+	/** Returns true if this proxy should write velocity even when the transform isn't changing. Usually this is combined with a check for the transform changing. */
+	inline bool AlwaysHasVelocity() const {	return bAlwaysHasVelocity || (bHasWorldPositionOffsetVelocity && EvaluateWorldPositionOffset()); }
 
 #if WITH_EDITOR
 	inline int32 GetNumUncachedStaticLightingInteractions() { return NumUncachedStaticLightingInteractions; }
