UAnimDataModel fix for single-frame animations which were incorrectly upgraded. As we now expect uniform key data, the model should contain 2 keys for a single-frame anim.

diff --git a/Engine/Source/Runtime/Core/Public/UObject/FortniteMainBranchObjectVersion.h b/Engine/Source/Runtime/Core/Public/UObject/FortniteMainBranchObjectVersion.h
--- a/Engine/Source/Runtime/Core/Public/UObject/FortniteMainBranchObjectVersion.h
+++ b/Engine/Source/Runtime/Core/Public/UObject/FortniteMainBranchObjectVersion.h
@@ -211,6 +211,9 @@ struct CORE_API FFortniteMainBranchObjectVersion
 		// Serialize LevelInstanceActorGuid on new ILevelInstanceInterface implementation
 		LevelInstanceActorGuidSerialize,
 
+		// Single-frame/key AnimDataModel patch-up
+		SingleFrameAndKeyAnimModel,
+
 		// -----<new versions can be added above this line>-------------------------------------------------
 		VersionPlusOne,
 		LatestVersion = VersionPlusOne - 1
diff --git a/Engine/Source/Runtime/Engine/Private/Animation/AnimDataModel.cpp b/Engine/Source/Runtime/Engine/Private/Animation/AnimDataModel.cpp
--- a/Engine/Source/Runtime/Engine/Private/Animation/AnimDataModel.cpp
+++ b/Engine/Source/Runtime/Engine/Private/Animation/AnimDataModel.cpp
@@ -9,7 +9,7 @@
 #include "Algo/Accumulate.h"
 #include "Animation/SmartName.h"
 #include "UObject/UE5MainStreamObjectVersion.h"
-
+#include "UObject/FortniteMainBranchObjectVersion.h"
 
 void UAnimDataModel::PostLoad()
 {
@@ -28,6 +28,34 @@ void UAnimDataModel::PostLoad()
 		}
 		Notify(EAnimDataModelNotifyType::BracketClosed);
 	}
+
+	if (GetLinkerCustomVersion(FFortniteMainBranchObjectVersion::GUID) < FFortniteMainBranchObjectVersion::SingleFrameAndKeyAnimModel)
+	{
+		const bool bHasBoneTracks = BoneAnimationTracks.Num() > 0;
+		if (bHasBoneTracks)
+		{
+			// Number of keys was used directly rather than Max(Value,2), as a single _frame_ animation should always have two _keys_ 
+			const int32 ActualNumKeys = BoneAnimationTracks[0].InternalTrackData.PosKeys.Num();
+			if (ActualNumKeys == 1 && NumberOfKeys == 2)
+			{
+				auto AddKey = [this](auto& Keys)
+				{
+					const auto KeyZero = Keys[0];
+					Keys.Add(KeyZero);
+					ensure(Keys.Num() == NumberOfKeys);
+				};
+				
+				for (FBoneAnimationTrack& BoneTrack : BoneAnimationTracks)
+				{
+					AddKey(BoneTrack.InternalTrackData.PosKeys);
+					AddKey(BoneTrack.InternalTrackData.RotKeys);
+					AddKey(BoneTrack.InternalTrackData.ScaleKeys);
+				}
+
+				Notify(EAnimDataModelNotifyType::TrackChanged);
+			}
+		}
+	}	
 }
 
 void UAnimDataModel::PostDuplicate(bool bDuplicateForPIE)
diff --git a/Engine/Source/Runtime/Engine/Private/Animation/AnimSequence.cpp b/Engine/Source/Runtime/Engine/Private/Animation/AnimSequence.cpp
--- a/Engine/Source/Runtime/Engine/Private/Animation/AnimSequence.cpp
+++ b/Engine/Source/Runtime/Engine/Private/Animation/AnimSequence.cpp
@@ -5155,7 +5155,7 @@ void UAnimSequence::PopulateModel()
 	check(!HasAnyFlags(EObjectFlags::RF_ClassDefaultObject));
 	
 	PRAGMA_DISABLE_DEPRECATION_WARNINGS
-	const int32 NumKeys = NumberOfKeys;
+	const int32 NumKeys = FMath::Max(NumberOfKeys, 2);
 	const float PlayLength = SequenceLength;
 	// Reset target framerate to current sampling frame rate
 	TargetFrameRate = SamplingFrameRate;
