Various crash fixes for material editor only data for the cooked editor.

# Crash Fixes
- Fix crash due to unserialized data in FMaterialLayersFunctions.  The editor only layer state data did not match that of the runtime data resulting in crashes.  FMaterialLayersFunctions cannot have UPROPERTIES wrapped in WITH_EDITORONLY_DATA if FMaterialLayersFunctions is a UPROPERTY on an optional object (which it is).  The simple fix is to remove WITH_EDITORONLY_DATA which has the side effect of cooking this data into the optional package.
- Fix crash trying to access null GameThreadShaderMap by calling AssociateWithAsset() _after_ we assign to the pointer.

# Missing code
I re-editgrated Material.cpp and MaterialInstanceConstant.cpp from CL# 20043286.  This caught some code that didn't cleanly make it over in the original merge.
- UMaterial::UpdateCachedExpressionData now validates data, and will properly update cached expression data.
- UMaterialInstanceConstant::UpdateCachedData properly updates cached expression data.

diff --git a/Engine/Source/Runtime/Engine/Classes/Materials/MaterialLayersFunctions.h b/Engine/Source/Runtime/Engine/Classes/Materials/MaterialLayersFunctions.h
--- a/Engine/Source/Runtime/Engine/Classes/Materials/MaterialLayersFunctions.h
+++ b/Engine/Source/Runtime/Engine/Classes/Materials/MaterialLayersFunctions.h
@@ -199,10 +199,8 @@ struct ENGINE_API FMaterialLayersFunctions : public FMaterialLayersFunctionsRunt
 	FMaterialLayersFunctionsRuntimeData& GetRuntime() { return *this; }
 	const FMaterialLayersFunctionsRuntimeData& GetRuntime() const { return *this; }
 
-#if WITH_EDITORONLY_DATA
 	UPROPERTY(EditAnywhere, Category = MaterialLayers)
 	FMaterialLayersFunctionsEditorOnlyData EditorOnly;
-#endif
 
 	void Empty()
 	{
diff --git a/Engine/Source/Runtime/Engine/Private/Materials/Material.cpp b/Engine/Source/Runtime/Engine/Private/Materials/Material.cpp
--- a/Engine/Source/Runtime/Engine/Private/Materials/Material.cpp
+++ b/Engine/Source/Runtime/Engine/Private/Materials/Material.cpp
@@ -1899,7 +1899,7 @@ void UMaterial::UpdateCachedExpressionData()
 {
 	COOK_STAT(FScopedDurationTimer BlockingTimer(MaterialCookStats::UpdateCachedExpressionDataSec));
 
-	if (bLoadedCachedExpressionData)
+	if (bLoadedCachedExpressionData && !GetPackage()->bIsCookedForEditor)
 	{
 		// Don't need to rebuild cached data if it was serialized
 		return;
@@ -3549,11 +3549,7 @@ void UMaterial::PostLoad()
 	PropagateDataToMaterialProxy();
 
 #if WITH_EDITOR
-	// cooked materials will not have any expressions in them, so this will obliterate the saved cached expression data
-	if (!GetOutermost()->bIsCookedForEditor)
-	{
-		UpdateCachedExpressionData();
-	}
+	UpdateCachedExpressionData();
 #endif // WITH_EDITOR
 
 	// Strata materials conversion needs to be done after expressions are cached, otherwise material function won't have 
