Configured an analytics app ID for VK editor.

diff --git a/Engine/Source/Runtime/Engine/Private/EngineAnalytics.cpp b/Engine/Source/Runtime/Engine/Private/EngineAnalytics.cpp
--- a/Engine/Source/Runtime/Engine/Private/EngineAnalytics.cpp
+++ b/Engine/Source/Runtime/Engine/Private/EngineAnalytics.cpp
@@ -52,7 +52,10 @@ static TSharedPtr<IAnalyticsProviderET> CreateEpicAnalyticsProvider()
 		FString UETypeOverride;
 		bool bHasOverride = GConfig->GetString(TEXT("Analytics"), TEXT("UE4TypeOverride"), UETypeOverride, GEngineIni);
 		const TCHAR* UETypeStr = bHasOverride ? *UETypeOverride : FEngineBuildSettings::IsPerforceBuild() ? TEXT("Perforce") : TEXT("UnrealEngine");
-		Config.APIKeyET = FString::Printf(TEXT("UEEditor.%s.%s"), UETypeStr, BuildTypeStr);
+
+		FString AppID;
+		GConfig->GetString(TEXT("Analytics"), TEXT("AppIdOverride"), AppID, GEditorIni);
+		Config.APIKeyET = FString::Printf(TEXT("%s.%s.%s"), AppID.IsEmpty() ? TEXT("UEEditor") : *AppID, UETypeStr, BuildTypeStr);
 	}
 	Config.APIServerET = TEXT("https://datarouter.ol.epicgames.com/");
 	Config.AppEnvironment = TEXT("datacollector-binary");
