Enabling ISPC on more supporting platforms, but defaulting all per-system toggles to off besides cloth, since most of those show a regression so we'll start conservative and enable only for cloth for now.

This saves about ~40% of the cloth simulation cost.

diff --git a/Engine/Plugins/Experimental/ChaosCloth/Source/Chaos/Private/ChaosCloth/ChaosClothingSimulationMesh.cpp b/Engine/Plugins/Experimental/ChaosCloth/Source/Chaos/Private/ChaosCloth/ChaosClothingSimulationMesh.cpp
--- a/Engine/Plugins/Experimental/ChaosCloth/Source/Chaos/Private/ChaosCloth/ChaosClothingSimulationMesh.cpp
+++ b/Engine/Plugins/Experimental/ChaosCloth/Source/Chaos/Private/ChaosCloth/ChaosClothingSimulationMesh.cpp
@@ -14,7 +14,7 @@
 #endif
 
 #if INTEL_ISPC && !UE_BUILD_SHIPPING
-bool bChaos_SkinPhysicsMesh_ISPC_Enabled = true;
+bool bChaos_SkinPhysicsMesh_ISPC_Enabled = CHAOS_SKIN_PHYSICS_MESH_ISPC_ENABLED_DEFAULT;
 FAutoConsoleVariableRef CVarChaosSkinPhysicsMeshISPCEnabled(TEXT("p.Chaos.SkinPhysicsMesh.ISPC"), bChaos_SkinPhysicsMesh_ISPC_Enabled, TEXT("Whether to use ISPC optimizations on skinned physics meshes"));
 
 static_assert(sizeof(ispc::FVector3f) == sizeof(Chaos::Softs::FSolverVec3), "sizeof(ispc::FVector3f) != sizeof(Chaos::Softs::FSolverVec3)");
diff --git a/Engine/Plugins/Experimental/ChaosCloth/Source/Chaos/Private/ChaosCloth/ChaosClothingSimulationSolver.cpp b/Engine/Plugins/Experimental/ChaosCloth/Source/Chaos/Private/ChaosCloth/ChaosClothingSimulationSolver.cpp
--- a/Engine/Plugins/Experimental/ChaosCloth/Source/Chaos/Private/ChaosCloth/ChaosClothingSimulationSolver.cpp
+++ b/Engine/Plugins/Experimental/ChaosCloth/Source/Chaos/Private/ChaosCloth/ChaosClothingSimulationSolver.cpp
@@ -27,9 +27,9 @@ DECLARE_CYCLE_STAT(TEXT("Chaos Cloth Solver Particle Pre Simulation Transforms")
 DECLARE_CYCLE_STAT(TEXT("Chaos Cloth Solver Collision Pre Simulation Transforms"), STAT_ChaosClothCollisionPreSimulationTransforms, STATGROUP_ChaosCloth);
 
 #if INTEL_ISPC && !UE_BUILD_SHIPPING
-bool bChaos_PreSimulationTransforms_ISPC_Enabled = true;
+bool bChaos_PreSimulationTransforms_ISPC_Enabled = CHAOS_PRE_SIMULATION_TRANSFORMS_ISPC_ENABLED_DEFAULT;
 FAutoConsoleVariableRef CVarChaosPreSimulationTransformsISPCEnabled(TEXT("p.Chaos.PreSimulationTransforms.ISPC"), bChaos_PreSimulationTransforms_ISPC_Enabled, TEXT("Whether to use ISPC optimizations in ApplySimulationTransforms"));
-bool bChaos_CalculateBounds_ISPC_Enabled = bChaos_CalculateBounds_ISPC_Enable;  // Disabled by default
+bool bChaos_CalculateBounds_ISPC_Enabled = CHAOS_CALCULATE_BOUNDS_ISPC_ENABLED_DEFAULT;
 FAutoConsoleVariableRef CVarChaosCalculateBoundsISPCEnabled(TEXT("p.Chaos.CalculateBounds.ISPC"), bChaos_CalculateBounds_ISPC_Enabled, TEXT("Whether to use ISPC optimizations in CalculateBounds"));
 
 static_assert(sizeof(ispc::FVector3f) == sizeof(Chaos::Softs::FSolverVec3), "sizeof(ispc::FVector) != sizeof(Chaos::Softs::FSolverVec3)");
diff --git a/Engine/Plugins/Experimental/ChaosCloth/Source/Chaos/Public/ChaosCloth/ChaosClothingSimulation.h b/Engine/Plugins/Experimental/ChaosCloth/Source/Chaos/Public/ChaosCloth/ChaosClothingSimulation.h
--- a/Engine/Plugins/Experimental/ChaosCloth/Source/Chaos/Public/ChaosCloth/ChaosClothingSimulation.h
+++ b/Engine/Plugins/Experimental/ChaosCloth/Source/Chaos/Public/ChaosCloth/ChaosClothingSimulation.h
@@ -178,11 +178,13 @@ namespace Chaos
 	};
 } // namespace Chaos
 
-// Support ISPC enable/disable in non-shipping builds
-#if !INTEL_ISPC
-const bool bChaos_GetSimData_ISPC_Enabled = false;
-#elif UE_BUILD_SHIPPING
-const bool bChaos_GetSimData_ISPC_Enabled = true;
+#if !defined(CHAOS_GET_SIM_DATA_ISPC_ENABLED_DEFAULT)
+#define CHAOS_GET_SIM_DATA_ISPC_ENABLED_DEFAULT 1
+#endif
+
+// Support run-time toggling on supported platforms in non-shipping configurations
+#if !INTEL_ISPC || UE_BUILD_SHIPPING
+static constexpr bool bChaos_GetSimData_ISPC_Enabled = INTEL_ISPC && CHAOS_GET_SIM_DATA_ISPC_ENABLED_DEFAULT;
 #else
 extern bool bChaos_GetSimData_ISPC_Enabled;
 #endif
diff --git a/Engine/Plugins/Experimental/ChaosCloth/Source/Chaos/Public/ChaosCloth/ChaosClothingSimulationMesh.h b/Engine/Plugins/Experimental/ChaosCloth/Source/Chaos/Public/ChaosCloth/ChaosClothingSimulationMesh.h
--- a/Engine/Plugins/Experimental/ChaosCloth/Source/Chaos/Public/ChaosCloth/ChaosClothingSimulationMesh.h
+++ b/Engine/Plugins/Experimental/ChaosCloth/Source/Chaos/Public/ChaosCloth/ChaosClothingSimulationMesh.h
@@ -76,12 +76,13 @@ namespace Chaos
 	};
 } // namespace Chaos
 
-// Support ISPC enable/disable in non-shipping builds
-constexpr bool bChaos_SkinPhysicsMesh_ISPC_Enable = true;
-#if !INTEL_ISPC
-const bool bChaos_SkinPhysicsMesh_ISPC_Enabled = false;
-#elif UE_BUILD_SHIPPING
-const bool bChaos_SkinPhysicsMesh_ISPC_Enabled = true;
+#if !defined(CHAOS_SKIN_PHYSICS_MESH_ISPC_ENABLED_DEFAULT)
+#define CHAOS_SKIN_PHYSICS_MESH_ISPC_ENABLED_DEFAULT 1
+#endif
+
+// Support run-time toggling on supported platforms in non-shipping configurations
+#if !INTEL_ISPC || UE_BUILD_SHIPPING
+static constexpr bool bChaos_SkinPhysicsMesh_ISPC_Enabled = INTEL_ISPC && CHAOS_SKIN_PHYSICS_MESH_ISPC_ENABLED_DEFAULT;
 #else
 extern bool bChaos_SkinPhysicsMesh_ISPC_Enabled;
 #endif
diff --git a/Engine/Plugins/Experimental/ChaosCloth/Source/Chaos/Public/ChaosCloth/ChaosClothingSimulationSolver.h b/Engine/Plugins/Experimental/ChaosCloth/Source/Chaos/Public/ChaosCloth/ChaosClothingSimulationSolver.h
--- a/Engine/Plugins/Experimental/ChaosCloth/Source/Chaos/Public/ChaosCloth/ChaosClothingSimulationSolver.h
+++ b/Engine/Plugins/Experimental/ChaosCloth/Source/Chaos/Public/ChaosCloth/ChaosClothingSimulationSolver.h
@@ -248,14 +248,18 @@ namespace Chaos
 
 } // namespace Chaos
 
-// Support ISPC enable/disable in non-shipping builds
-constexpr bool bChaos_CalculateBounds_ISPC_Enable = true;
-#if !INTEL_ISPC
-const bool bChaos_PreSimulationTransforms_ISPC_Enabled = false;
-const bool bChaos_CalculateBounds_ISPC_Enabled = false;
-#elif UE_BUILD_SHIPPING
-const bool bChaos_PreSimulationTransforms_ISPC_Enabled = true;
-const bool bChaos_CalculateBounds_ISPC_Enabled = bChaos_CalculateBounds_ISPC_Enable;
+#if !defined(CHAOS_CALCULATE_BOUNDS_ISPC_ENABLED_DEFAULT)
+#define CHAOS_CALCULATE_BOUNDS_ISPC_ENABLED_DEFAULT 1
+#endif
+
+#if !defined(CHAOS_PRE_SIMULATION_TRANSFORMS_ISPC_ENABLED_DEFAULT)
+#define CHAOS_PRE_SIMULATION_TRANSFORMS_ISPC_ENABLED_DEFAULT 1
+#endif
+
+// Support run-time toggling on supported platforms in non-shipping configurations
+#if !INTEL_ISPC || UE_BUILD_SHIPPING
+static constexpr bool bChaos_CalculateBounds_ISPC_Enabled = INTEL_ISPC && CHAOS_CALCULATE_BOUNDS_ISPC_ENABLED_DEFAULT;
+static constexpr bool bChaos_PreSimulationTransforms_ISPC_Enabled = INTEL_ISPC && CHAOS_PRE_SIMULATION_TRANSFORMS_ISPC_ENABLED_DEFAULT;
 #else
 extern bool bChaos_PreSimulationTransforms_ISPC_Enabled;
 extern bool bChaos_CalculateBounds_ISPC_Enabled;
diff --git a/Engine/Source/Runtime/Engine/Classes/Components/SkeletalMeshComponent.h b/Engine/Source/Runtime/Engine/Classes/Components/SkeletalMeshComponent.h
--- a/Engine/Source/Runtime/Engine/Classes/Components/SkeletalMeshComponent.h
+++ b/Engine/Source/Runtime/Engine/Classes/Components/SkeletalMeshComponent.h
@@ -2590,12 +2590,3 @@ struct FLinkedInstancesAdapter
 		}
 	}
 };
-
-// Support ISPC enable/disable in non-shipping builds
-#if !INTEL_ISPC
-const bool bAnim_SkeletalMesh_ISPC_Enabled = false;
-#elif UE_BUILD_SHIPPING
-const bool bAnim_SkeletalMesh_ISPC_Enabled = true;
-#else
-extern bool bAnim_SkeletalMesh_ISPC_Enabled;
-#endif
diff --git a/Engine/Source/Runtime/Engine/Private/Animation/AnimEncoding_ConstantKeyLerp.cpp b/Engine/Source/Runtime/Engine/Private/Animation/AnimEncoding_ConstantKeyLerp.cpp
--- a/Engine/Source/Runtime/Engine/Private/Animation/AnimEncoding_ConstantKeyLerp.cpp
+++ b/Engine/Source/Runtime/Engine/Private/Animation/AnimEncoding_ConstantKeyLerp.cpp
@@ -14,9 +14,16 @@ static_assert(sizeof(ispc::FTransform) == sizeof(FTransform), "sizeof(ispc::FTra
 static_assert(sizeof(ispc::BoneTrackPair) == sizeof(BoneTrackPair), "sizeof(ispc::BoneTrackPair) != sizeof(BoneTrackPair)");
 #endif
 
-#if INTEL_ISPC && !UE_BUILD_SHIPPING
-bool bAnim_ConstantKeyLerp_ISPC_Enabled = true;
-FAutoConsoleVariableRef CVarAnimConstantKeyLerpISPCEnabled(TEXT("a.ConstantKeyLerp.ISPC"), bAnim_ConstantKeyLerp_ISPC_Enabled, TEXT("Whether to use ISPC optimizations in constant key anim encoding"));
+#if !defined(ANIM_CONSTANT_KEY_LERP_ISPC_ENABLED_DEFAULT)
+#define ANIM_CONSTANT_KEY_LERP_ISPC_ENABLED_DEFAULT 1
+#endif
+
+// Support run-time toggling on supported platforms in non-shipping configurations
+#if !INTEL_ISPC || UE_BUILD_SHIPPING
+static constexpr bool bAnim_ConstantKeyLerp_ISPC_Enabled = INTEL_ISPC && ANIM_CONSTANT_KEY_LERP_ISPC_ENABLED_DEFAULT;
+#else
+static bool bAnim_ConstantKeyLerp_ISPC_Enabled = ANIM_CONSTANT_KEY_LERP_ISPC_ENABLED_DEFAULT;
+static FAutoConsoleVariableRef CVarAnimConstantKeyLerpISPCEnabled(TEXT("a.ConstantKeyLerp.ISPC"), bAnim_ConstantKeyLerp_ISPC_Enabled, TEXT("Whether to use ISPC optimizations in constant key anim encoding"));
 #endif
 
 /**
diff --git a/Engine/Source/Runtime/Engine/Private/Animation/AnimEncoding_PerTrackCompression.cpp b/Engine/Source/Runtime/Engine/Private/Animation/AnimEncoding_PerTrackCompression.cpp
--- a/Engine/Source/Runtime/Engine/Private/Animation/AnimEncoding_PerTrackCompression.cpp
+++ b/Engine/Source/Runtime/Engine/Private/Animation/AnimEncoding_PerTrackCompression.cpp
@@ -16,9 +16,16 @@ static_assert(sizeof(ispc::FTransform) == sizeof(FTransform), "sizeof(ispc::FTra
 static_assert(sizeof(ispc::BoneTrackPair) == sizeof(BoneTrackPair), "sizeof(ispc::BoneTrackPair) != sizeof(BoneTrackPair)");
 #endif
 
-#if INTEL_ISPC && !UE_BUILD_SHIPPING
-bool bAnim_PerTrackCompression_ISPC_Enabled = true;
-FAutoConsoleVariableRef CVarAnimPerTrackCompressionISPCEnabled(TEXT("a.PerTrackCompression.ISPC"), bAnim_PerTrackCompression_ISPC_Enabled, TEXT("Whether to use ISPC optimizations in per track anim encoding"));
+#if !defined(ANIM_PER_TRACK_COMPRESSION_ISPC_ENABLED_DEFAULT)
+#define ANIM_PER_TRACK_COMPRESSION_ISPC_ENABLED_DEFAULT 1
+#endif
+
+// Support run-time toggling on supported platforms in non-shipping configurations
+#if !INTEL_ISPC || UE_BUILD_SHIPPING
+static constexpr bool bAnim_PerTrackCompression_ISPC_Enabled = INTEL_ISPC && ANIM_PER_TRACK_COMPRESSION_ISPC_ENABLED_DEFAULT;
+#else
+static bool bAnim_PerTrackCompression_ISPC_Enabled = ANIM_PER_TRACK_COMPRESSION_ISPC_ENABLED_DEFAULT;
+static FAutoConsoleVariableRef CVarAnimPerTrackCompressionISPCEnabled(TEXT("a.PerTrackCompression.ISPC"), bAnim_PerTrackCompression_ISPC_Enabled, TEXT("Whether to use ISPC optimizations in per track anim encoding"));
 #endif
 
 // This define controls whether scalar or vector code is used to decompress keys.  Note that not all key decompression code
diff --git a/Engine/Source/Runtime/Engine/Private/Animation/AnimEncoding_VariableKeyLerp.cpp b/Engine/Source/Runtime/Engine/Private/Animation/AnimEncoding_VariableKeyLerp.cpp
--- a/Engine/Source/Runtime/Engine/Private/Animation/AnimEncoding_VariableKeyLerp.cpp
+++ b/Engine/Source/Runtime/Engine/Private/Animation/AnimEncoding_VariableKeyLerp.cpp
@@ -14,9 +14,16 @@ static_assert(sizeof(ispc::FTransform) == sizeof(FTransform), "sizeof(ispc::FTra
 static_assert(sizeof(ispc::BoneTrackPair) == sizeof(BoneTrackPair), "sizeof(ispc::BoneTrackPair) != sizeof(BoneTrackPair)");
 #endif
 
-#if INTEL_ISPC && !UE_BUILD_SHIPPING
-bool bAnim_VariableKeyLerp_ISPC_Enabled = true;
-FAutoConsoleVariableRef CVarAnimVariableKeyLerpISPCEnabled(TEXT("a.VariableKeyLerp.ISPC"), bAnim_VariableKeyLerp_ISPC_Enabled, TEXT("Whether to use ISPC optimizations in variable key anim encoding"));
+#if !defined(ANIM_VARIABLE_KEY_LERP_ISPC_ENABLED_DEFAULT)
+#define ANIM_VARIABLE_KEY_LERP_ISPC_ENABLED_DEFAULT 1
+#endif
+
+// Support run-time toggling on supported platforms in non-shipping configurations
+#if !INTEL_ISPC || UE_BUILD_SHIPPING
+static constexpr bool bAnim_VariableKeyLerp_ISPC_Enabled = INTEL_ISPC && ANIM_VARIABLE_KEY_LERP_ISPC_ENABLED_DEFAULT;
+#else
+static bool bAnim_VariableKeyLerp_ISPC_Enabled = ANIM_VARIABLE_KEY_LERP_ISPC_ENABLED_DEFAULT;
+static FAutoConsoleVariableRef CVarAnimVariableKeyLerpISPCEnabled(TEXT("a.VariableKeyLerp.ISPC"), bAnim_VariableKeyLerp_ISPC_Enabled, TEXT("Whether to use ISPC optimizations in variable key anim encoding"));
 #endif
 
 /**
diff --git a/Engine/Source/Runtime/Engine/Private/Animation/AnimationRuntime.cpp b/Engine/Source/Runtime/Engine/Private/Animation/AnimationRuntime.cpp
--- a/Engine/Source/Runtime/Engine/Private/Animation/AnimationRuntime.cpp
+++ b/Engine/Source/Runtime/Engine/Private/Animation/AnimationRuntime.cpp
@@ -36,9 +36,16 @@ static_assert(sizeof(ispc::FVector4) == sizeof(FQuat), "sizeof(ispc::FVector4) !
 static_assert(sizeof(ispc::FPerBoneBlendWeight) == sizeof(FPerBoneBlendWeight), "sizeof(ispc::FPerBoneBlendWeight) != sizeof(FPerBoneBlendWeight)");
 #endif
 
-#if INTEL_ISPC && !UE_BUILD_SHIPPING
-bool bAnim_Runtime_ISPC_Enabled = true;
-FAutoConsoleVariableRef CVarAnimRuntimeISPCEnabled(TEXT("a.Runtime.ISPC"), bAnim_Runtime_ISPC_Enabled, TEXT("Whether to use ISPC optimizations in animation runtime"));
+#if !defined(ANIM_RUNTIME_ISPC_ENABLED_DEFAULT)
+#define ANIM_RUNTIME_ISPC_ENABLED_DEFAULT 1
+#endif
+
+// Support run-time toggling on supported platforms in non-shipping configurations
+#if !INTEL_ISPC || UE_BUILD_SHIPPING
+static constexpr bool bAnim_Runtime_ISPC_Enabled = INTEL_ISPC && ANIM_RUNTIME_ISPC_ENABLED_DEFAULT;
+#else
+static bool bAnim_Runtime_ISPC_Enabled = ANIM_RUNTIME_ISPC_ENABLED_DEFAULT;
+static FAutoConsoleVariableRef CVarAnimRuntimeISPCEnabled(TEXT("a.Runtime.ISPC"), bAnim_Runtime_ISPC_Enabled, TEXT("Whether to use ISPC optimizations in animation runtime"));
 #endif
 
 void FAnimationRuntime::NormalizeRotations(const FBoneContainer& RequiredBones, /*inout*/ FTransformArrayA2& Atoms)
diff --git a/Engine/Source/Runtime/Engine/Private/Animation/BonePose.cpp b/Engine/Source/Runtime/Engine/Private/Animation/BonePose.cpp
--- a/Engine/Source/Runtime/Engine/Private/Animation/BonePose.cpp
+++ b/Engine/Source/Runtime/Engine/Private/Animation/BonePose.cpp
@@ -10,9 +10,16 @@
 static_assert(sizeof(ispc::FTransform) == sizeof(FTransform), "sizeof(ispc::FTransform) != sizeof(FTransform)");
 #endif
 
-#if INTEL_ISPC && !UE_BUILD_SHIPPING
-bool bAnim_BonePose_ISPC_Enabled = true;
-FAutoConsoleVariableRef CVarAnimBonePoseISPCEnabled(TEXT("a.BonePose.ISPC"), bAnim_BonePose_ISPC_Enabled, TEXT("Whether to use ISPC optimizations in bone pose calculations"));
+#if !defined(ANIM_BONE_POSE_ISPC_ENABLED_DEFAULT)
+#define ANIM_BONE_POSE_ISPC_ENABLED_DEFAULT 1
+#endif
+
+// Support run-time toggling on supported platforms in non-shipping configurations
+#if !INTEL_ISPC || UE_BUILD_SHIPPING
+static constexpr bool bAnim_BonePose_ISPC_Enabled = INTEL_ISPC && ANIM_BONE_POSE_ISPC_ENABLED_DEFAULT;
+#else
+static bool bAnim_BonePose_ISPC_Enabled = ANIM_BONE_POSE_ISPC_ENABLED_DEFAULT;
+static FAutoConsoleVariableRef CVarAnimBonePoseISPCEnabled(TEXT("a.BonePose.ISPC"), bAnim_BonePose_ISPC_Enabled, TEXT("Whether to use ISPC optimizations in bone pose calculations"));
 #endif
 
 // Normalizes all rotations in this pose
diff --git a/Engine/Source/Runtime/Engine/Private/Collision/KAggregateGeom.cpp b/Engine/Source/Runtime/Engine/Private/Collision/KAggregateGeom.cpp
--- a/Engine/Source/Runtime/Engine/Private/Collision/KAggregateGeom.cpp
+++ b/Engine/Source/Runtime/Engine/Private/Collision/KAggregateGeom.cpp
@@ -35,11 +35,16 @@ static_assert(sizeof(ispc::FRotator) == sizeof(FRotator), "sizeof(ispc::FRotator
 static_assert(sizeof(ispc::FVector) == sizeof(FVector), "sizeof(ispc::FVector) != sizeof(FVector)");
 #endif
 
-#if INTEL_ISPC && !UE_BUILD_SHIPPING
-bool bPhysics_AggregateGeom_ISPC_Enabled = true;
-FAutoConsoleVariableRef CVarPhysicsAggregateGeomISPCEnabled(TEXT("p.AggregateGeom.ISPC"), bPhysics_AggregateGeom_ISPC_Enabled, TEXT("Whether to use ISPC optimizations in physics aggregate geometry calculations"));
+#if !defined(PHYSICS_AGGREGATE_GEOMETRY_ISPC_ENABLED_DEFAULT)
+#define PHYSICS_AGGREGATE_GEOMETRY_ISPC_ENABLED_DEFAULT 1
+#endif
+
+// Support run-time toggling on supported platforms in non-shipping configurations
+#if !INTEL_ISPC || UE_BUILD_SHIPPING
+static constexpr bool bPhysics_AggregateGeom_ISPC_Enabled = INTEL_ISPC && PHYSICS_AGGREGATE_GEOMETRY_ISPC_ENABLED_DEFAULT;
 #else
-const bool bPhysics_AggregateGeom_ISPC_Enabled = INTEL_ISPC;
+static bool bPhysics_AggregateGeom_ISPC_Enabled = PHYSICS_AGGREGATE_GEOMETRY_ISPC_ENABLED_DEFAULT;
+static FAutoConsoleVariableRef CVarPhysicsAggregateGeomISPCEnabled(TEXT("p.AggregateGeom.ISPC"), bPhysics_AggregateGeom_ISPC_Enabled, TEXT("Whether to use ISPC optimizations in physics aggregate geometry calculations"));
 #endif
 
 ///////////////////////////////////////
diff --git a/Engine/Source/Runtime/Engine/Private/Components/SkeletalMeshComponent.cpp b/Engine/Source/Runtime/Engine/Private/Components/SkeletalMeshComponent.cpp
--- a/Engine/Source/Runtime/Engine/Private/Components/SkeletalMeshComponent.cpp
+++ b/Engine/Source/Runtime/Engine/Private/Components/SkeletalMeshComponent.cpp
@@ -65,9 +65,16 @@ static TAutoConsoleVariable<float> CVarStallParallelAnimation(
 	0.0f,
 	TEXT("Sleep for the given time in each parallel animation task. Time is given in ms. This is a debug option used for critical path analysis and forcing a change in the critical path."));
 
-#if INTEL_ISPC && !UE_BUILD_SHIPPING
-bool bAnim_SkeletalMesh_ISPC_Enabled = true;
-FAutoConsoleVariableRef CVarAnimSkeletalMeshISPCEnabled(TEXT("a.SkeletalMesh.ISPC"), bAnim_SkeletalMesh_ISPC_Enabled, TEXT("Whether to use ISPC optimizations in animation skeletal mesh components"));
+#if !defined(ANIM_SKELETAL_MESH_ISPC_ENABLED_DEFAULT)
+#define ANIM_SKELETAL_MESH_ISPC_ENABLED_DEFAULT 1
+#endif
+
+// Support run-time toggling on supported platforms in non-shipping configurations
+#if !INTEL_ISPC || UE_BUILD_SHIPPING
+static constexpr bool bAnim_SkeletalMesh_ISPC_Enabled = INTEL_ISPC && ANIM_SKELETAL_MESH_ISPC_ENABLED_DEFAULT;
+#else
+static bool bAnim_SkeletalMesh_ISPC_Enabled = ANIM_SKELETAL_MESH_ISPC_ENABLED_DEFAULT;
+static FAutoConsoleVariableRef CVarAnimSkeletalMeshISPCEnabled(TEXT("a.SkeletalMesh.ISPC"), bAnim_SkeletalMesh_ISPC_Enabled, TEXT("Whether to use ISPC optimizations in animation skeletal mesh components"));
 #endif
 
 static TAutoConsoleVariable<int32> CVarCacheLocalSpaceBounds(
@@ -75,7 +82,6 @@ static TAutoConsoleVariable<int32> CVarCacheLocalSpaceBounds(
 	1,
 	TEXT("If 1 (default) local-space bounds are calculated and cached, otherwise worldspace bounds are built and cached (and inverse transformed to produce local bounds)."));
 
-
 DECLARE_CYCLE_STAT_EXTERN(TEXT("Anim Instance Spawn Time"), STAT_AnimSpawnTime, STATGROUP_Anim, );
 DEFINE_STAT(STAT_AnimSpawnTime);
 DEFINE_STAT(STAT_PostAnimEvaluation);
diff --git a/Engine/Source/Runtime/Engine/Private/GPUSkinVertexFactory.cpp b/Engine/Source/Runtime/Engine/Private/GPUSkinVertexFactory.cpp
--- a/Engine/Source/Runtime/Engine/Private/GPUSkinVertexFactory.cpp
+++ b/Engine/Source/Runtime/Engine/Private/GPUSkinVertexFactory.cpp
@@ -102,9 +102,16 @@ static TAutoConsoleVariable<int32> CVarVelocityTest(
 	ECVF_Cheat | ECVF_RenderThreadSafe);
 #endif // if !(UE_BUILD_SHIPPING || UE_BUILD_TEST)
 
-#if INTEL_ISPC && !UE_BUILD_SHIPPING
-bool bGPUSkin_CopyBones_ISPC_Enabled = true;
-FAutoConsoleVariableRef CVarGPUSkinCopyBonesISPCEnabled(TEXT("r.GPUSkin.CopyBones.ISPC"), bGPUSkin_CopyBones_ISPC_Enabled, TEXT("Whether to use ISPC optimizations when copying bones for GPU skinning"));
+#if !defined(GPU_SKIN_COPY_BONES_ISPC_ENABLED_DEFAULT)
+#define GPU_SKIN_COPY_BONES_ISPC_ENABLED_DEFAULT 1
+#endif
+
+// Support run-time toggling on supported platforms in non-shipping configurations
+#if !INTEL_ISPC || UE_BUILD_SHIPPING
+static constexpr bool bGPUSkin_CopyBones_ISPC_Enabled = INTEL_ISPC && GPU_SKIN_COPY_BONES_ISPC_ENABLED_DEFAULT;
+#else
+static bool bGPUSkin_CopyBones_ISPC_Enabled = GPU_SKIN_COPY_BONES_ISPC_ENABLED_DEFAULT;
+static FAutoConsoleVariableRef CVarGPUSkinCopyBonesISPCEnabled(TEXT("r.GPUSkin.CopyBones.ISPC"), bGPUSkin_CopyBones_ISPC_Enabled, TEXT("Whether to use ISPC optimizations when copying bones for GPU skinning"));
 #endif
 
 #if INTEL_ISPC
diff --git a/Engine/Source/Runtime/Engine/Public/AnimEncoding_ConstantKeyLerp.h b/Engine/Source/Runtime/Engine/Public/AnimEncoding_ConstantKeyLerp.h
--- a/Engine/Source/Runtime/Engine/Public/AnimEncoding_ConstantKeyLerp.h
+++ b/Engine/Source/Runtime/Engine/Public/AnimEncoding_ConstantKeyLerp.h
@@ -322,12 +322,3 @@ FORCEINLINE void AEFConstantKeyLerp<FORMAT>::GetBoneAtomScale(FTransform& OutAto
 		OutAtom.SetScale3D((FVector)P0);
 	}
 }
-
-// Support ISPC enable/disable in non-shipping builds
-#if !INTEL_ISPC
-const bool bAnim_ConstantKeyLerp_ISPC_Enabled = false;
-#elif UE_BUILD_SHIPPING
-const bool bAnim_ConstantKeyLerp_ISPC_Enabled = true;
-#else
-extern bool bAnim_ConstantKeyLerp_ISPC_Enabled;
-#endif
diff --git a/Engine/Source/Runtime/Engine/Public/AnimEncoding_PerTrackCompression.h b/Engine/Source/Runtime/Engine/Public/AnimEncoding_PerTrackCompression.h
--- a/Engine/Source/Runtime/Engine/Public/AnimEncoding_PerTrackCompression.h
+++ b/Engine/Source/Runtime/Engine/Public/AnimEncoding_PerTrackCompression.h
@@ -133,12 +133,3 @@ protected:
 
 	friend class UAnimCompress_PerTrackCompression;
 };
-
-// Support ISPC enable/disable in non-shipping builds
-#if !INTEL_ISPC
-const bool bAnim_PerTrackCompression_ISPC_Enabled = false;
-#elif UE_BUILD_SHIPPING
-const bool bAnim_PerTrackCompression_ISPC_Enabled = true;
-#else
-extern bool bAnim_PerTrackCompression_ISPC_Enabled;
-#endif
diff --git a/Engine/Source/Runtime/Engine/Public/AnimEncoding_VariableKeyLerp.h b/Engine/Source/Runtime/Engine/Public/AnimEncoding_VariableKeyLerp.h
--- a/Engine/Source/Runtime/Engine/Public/AnimEncoding_VariableKeyLerp.h
+++ b/Engine/Source/Runtime/Engine/Public/AnimEncoding_VariableKeyLerp.h
@@ -329,12 +329,3 @@ FORCEINLINE_DEBUGGABLE void AEFVariableKeyLerp<FORMAT>::GetBoneAtomScale(FTransf
 		OutAtom.SetScale3D( (FVector)P0 );
 	}
 }
-
-// Support ISPC enable/disable in non-shipping builds
-#if !INTEL_ISPC
-const bool bAnim_VariableKeyLerp_ISPC_Enabled = false;
-#elif UE_BUILD_SHIPPING
-const bool bAnim_VariableKeyLerp_ISPC_Enabled = true;
-#else
-extern bool bAnim_VariableKeyLerp_ISPC_Enabled;
-#endif
diff --git a/Engine/Source/Runtime/Engine/Public/AnimationRuntime.h b/Engine/Source/Runtime/Engine/Public/AnimationRuntime.h
--- a/Engine/Source/Runtime/Engine/Public/AnimationRuntime.h
+++ b/Engine/Source/Runtime/Engine/Public/AnimationRuntime.h
@@ -621,12 +621,3 @@ public:
 };
 
 ENUM_CLASS_FLAGS(FAnimationRuntime::EBlendPosesPerBoneFilterFlags);
-
-// Support ISPC enable/disable in non-shipping builds
-#if !INTEL_ISPC
-const bool bAnim_Runtime_ISPC_Enabled = false;
-#elif UE_BUILD_SHIPPING
-const bool bAnim_Runtime_ISPC_Enabled = true;
-#else
-extern bool bAnim_Runtime_ISPC_Enabled;
-#endif
diff --git a/Engine/Source/Runtime/Engine/Public/BonePose.h b/Engine/Source/Runtime/Engine/Public/BonePose.h
--- a/Engine/Source/Runtime/Engine/Public/BonePose.h
+++ b/Engine/Source/Runtime/Engine/Public/BonePose.h
@@ -958,12 +958,3 @@ extern ENGINE_API void BuildPoseFromRawData(
 	const TArray<FTransform>& RetargetTransforms,
 	const TMap<int32, const FTransformCurve*>* AdditiveBoneTransformCurves = nullptr
 	);
-
-// Support ISPC enable/disable in non-shipping builds
-#if !INTEL_ISPC
-const bool bAnim_BonePose_ISPC_Enabled = false;
-#elif UE_BUILD_SHIPPING
-const bool bAnim_BonePose_ISPC_Enabled = true;
-#else
-extern bool bAnim_BonePose_ISPC_Enabled;
-#endif
diff --git a/Engine/Source/Runtime/Engine/Public/GPUSkinVertexFactory.h b/Engine/Source/Runtime/Engine/Public/GPUSkinVertexFactory.h
--- a/Engine/Source/Runtime/Engine/Public/GPUSkinVertexFactory.h
+++ b/Engine/Source/Runtime/Engine/Public/GPUSkinVertexFactory.h
@@ -989,12 +989,3 @@ public:
 	static void ModifyCompilationEnvironment(const FVertexFactoryShaderPermutationParameters& Parameters, FShaderCompilerEnvironment& OutEnvironment);
 	static bool ShouldCompilePermutation(const FVertexFactoryShaderPermutationParameters& Parameters);
 };
-
-// Support ISPC enable/disable in non-shipping builds
-#if !INTEL_ISPC
-const bool bGPUSkin_CopyBones_ISPC_Enabled = false;
-#elif UE_BUILD_SHIPPING
-const bool bGPUSkin_CopyBones_ISPC_Enabled = true;
-#else
-extern bool bGPUSkin_CopyBones_ISPC_Enabled;
-#endif
diff --git a/Engine/Source/Runtime/Experimental/Chaos/Private/Chaos/AABB.cpp b/Engine/Source/Runtime/Experimental/Chaos/Private/Chaos/AABB.cpp
--- a/Engine/Source/Runtime/Experimental/Chaos/Private/Chaos/AABB.cpp
+++ b/Engine/Source/Runtime/Experimental/Chaos/Private/Chaos/AABB.cpp
@@ -10,9 +10,16 @@ static_assert(sizeof(ispc::FTransform) == sizeof(FTransform), "sizeof(ispc::FTra
 static_assert(sizeof(ispc::FVector) == sizeof(Chaos::TVector<Chaos::FReal, 3>), "sizeof(ispc::FVector) != sizeof(Chaos::TVector<Chaos::FReal, 3>)");
 #endif
 
-#if INTEL_ISPC && !UE_BUILD_SHIPPING
-bool bChaos_AABBTransform_ISPC_Enabled = true;
-FAutoConsoleVariableRef CVarChaosAABBTransformISPCEnabled(TEXT("p.Chaos.AABBTransform.ISPC"), bChaos_AABBTransform_ISPC_Enabled, TEXT("Whether to use ISPC optimizations when computing AABB transforms"));
+#if !defined(CHAOS_AABB_TRANSFORM_ISPC_ENABLED_DEFAULT)
+#define CHAOS_AABB_TRANSFORM_ISPC_ENABLED_DEFAULT 1
+#endif
+
+// Support run-time toggling on supported platforms in non-shipping configurations
+#if !INTEL_ISPC || UE_BUILD_SHIPPING
+static constexpr bool bChaos_AABBTransform_ISPC_Enabled = INTEL_ISPC && CHAOS_AABB_TRANSFORM_ISPC_ENABLED_DEFAULT;
+#else
+static bool bChaos_AABBTransform_ISPC_Enabled = CHAOS_AABB_TRANSFORM_ISPC_ENABLED_DEFAULT;
+static FAutoConsoleVariableRef CVarChaosAABBTransformISPCEnabled(TEXT("p.Chaos.AABBTransform.ISPC"), bChaos_AABBTransform_ISPC_Enabled, TEXT("Whether to use ISPC optimizations when computing AABB transforms"));
 #endif
 
 namespace Chaos
diff --git a/Engine/Source/Runtime/Experimental/Chaos/Private/Chaos/Joint/JointConstraintsCVars.cpp b/Engine/Source/Runtime/Experimental/Chaos/Private/Chaos/Joint/JointConstraintsCVars.cpp
--- a/Engine/Source/Runtime/Experimental/Chaos/Private/Chaos/Joint/JointConstraintsCVars.cpp
+++ b/Engine/Source/Runtime/Experimental/Chaos/Private/Chaos/Joint/JointConstraintsCVars.cpp
@@ -3,8 +3,8 @@
 #include "HAL/IConsoleManager.h"
 
 #if INTEL_ISPC && !UE_BUILD_SHIPPING
-bool bChaos_Joint_ISPC_Enabled = false;
-FAutoConsoleVariableRef CVarChaosJointISPCEnabled(TEXT("p.Chaos.Joint.ISPC"), bChaos_Joint_ISPC_Enabled, TEXT("Whether to use ISPC optimizations in the Joint Solver"));
+bool bChaos_Joint_ISPC_Enabled = CHAOS_JOINT_SOLVER_ISPC_ENABLED_DEFAULT;
+static FAutoConsoleVariableRef CVarChaosJointISPCEnabled(TEXT("p.Chaos.Joint.ISPC"), bChaos_Joint_ISPC_Enabled, TEXT("Whether to use ISPC optimizations in the Joint Solver"));
 #endif
 
 float Chaos_Joint_DegenerateRotationLimit = -0.998f;	// Cos(176deg)
diff --git a/Engine/Source/Runtime/Experimental/Chaos/Private/Chaos/Joint/JointConstraintsCVars.h b/Engine/Source/Runtime/Experimental/Chaos/Private/Chaos/Joint/JointConstraintsCVars.h
--- a/Engine/Source/Runtime/Experimental/Chaos/Private/Chaos/Joint/JointConstraintsCVars.h
+++ b/Engine/Source/Runtime/Experimental/Chaos/Private/Chaos/Joint/JointConstraintsCVars.h
@@ -3,13 +3,15 @@
 
 #include "CoreMinimal.h"
 
-// Support ISPC enable/disable in non-shipping builds
-#if !INTEL_ISPC
-	const bool bChaos_Joint_ISPC_Enabled = false;
-#elif UE_BUILD_SHIPPING
-	const bool bChaos_Joint_ISPC_Enabled = false;
+#if !defined(CHAOS_JOINT_SOLVER_ISPC_ENABLED_DEFAULT)
+#define CHAOS_JOINT_SOLVER_ISPC_ENABLED_DEFAULT 0
+#endif
+
+// Support run-time toggling on supported platforms in non-shipping configurations
+#if !INTEL_ISPC || UE_BUILD_SHIPPING
+static constexpr bool bChaos_Joint_ISPC_Enabled = INTEL_ISPC && CHAOS_JOINT_SOLVER_ISPC_ENABLED_DEFAULT;
 #else
-	extern bool bChaos_Joint_ISPC_Enabled;
+extern bool bChaos_Joint_ISPC_Enabled;
 #endif
 
 extern bool bChaos_Joint_EarlyOut_Enabled;
diff --git a/Engine/Source/Runtime/Experimental/Chaos/Private/Chaos/PerParticlePBDCollisionConstraint.cpp b/Engine/Source/Runtime/Experimental/Chaos/Private/Chaos/PerParticlePBDCollisionConstraint.cpp
--- a/Engine/Source/Runtime/Experimental/Chaos/Private/Chaos/PerParticlePBDCollisionConstraint.cpp
+++ b/Engine/Source/Runtime/Experimental/Chaos/Private/Chaos/PerParticlePBDCollisionConstraint.cpp
@@ -104,8 +104,8 @@ static_assert(sizeof(ispc::VerticesS) == sizeof(Chaos::FConvexStructureData::FCo
 static_assert(sizeof(ispc::VerticesM) == sizeof(Chaos::FConvexStructureData::FConvexStructureDataMedium::FVertexData));
 static_assert(sizeof(ispc::VerticesL) == sizeof(Chaos::FConvexStructureData::FConvexStructureDataLarge::FVertexData));
 
-bool bChaos_PerParticleCollision_ISPC_Enabled = true;
-FAutoConsoleVariableRef CVarChaosPerParticleCollisionISPCEnabled(TEXT("p.Chaos.PerParticleCollision.ISPC"), bChaos_PerParticleCollision_ISPC_Enabled, TEXT("Whether to use ISPC optimizations in per particle collisions"));
+bool bChaos_PerParticleCollision_ISPC_Enabled = CHAOS_PER_PARTICLE_COLLISION_ISPC_ENABLED_DEFAULT;
+static FAutoConsoleVariableRef CVarChaosPerParticleCollisionISPCEnabled(TEXT("p.Chaos.PerParticleCollision.ISPC"), bChaos_PerParticleCollision_ISPC_Enabled, TEXT("Whether to use ISPC optimizations in per particle collisions"));
 #endif
 
 static int32 Chaos_PerParticleCollision_ISPC_ParallelBatchSize = 128;
diff --git a/Engine/Source/Runtime/Experimental/Chaos/Private/Chaos/TriangleMesh.cpp b/Engine/Source/Runtime/Experimental/Chaos/Private/Chaos/TriangleMesh.cpp
--- a/Engine/Source/Runtime/Experimental/Chaos/Private/Chaos/TriangleMesh.cpp
+++ b/Engine/Source/Runtime/Experimental/Chaos/Private/Chaos/TriangleMesh.cpp
@@ -18,14 +18,20 @@
 
 #if INTEL_ISPC
 #include "TriangleMesh.ispc.generated.h"
-#endif
-
-#if INTEL_ISPC && !UE_BUILD_SHIPPING
 static_assert(sizeof(ispc::FVector3f) == sizeof(Chaos::TVec3<Chaos::FRealSingle>), "sizeof(ispc::FVector3f) != sizeof(Chaos::TVec3<Chaos::FRealSingle>)");
 static_assert(sizeof(ispc::TArrayInt) == sizeof(TArray<int32>), "sizeof(ispc::TArrayInt) != sizeof(TArray<int32>)");
+#endif
+
+#if !defined(CHAOS_TRIANGLE_MESH_ISPC_ENABLED_DEFAULT)
+#define CHAOS_TRIANGLE_MESH_ISPC_ENABLED_DEFAULT 1
+#endif
 
-bool bChaos_TriangleMesh_ISPC_Enabled = true;
-FAutoConsoleVariableRef CVarChaosTriangleMeshISPCEnabled(TEXT("p.Chaos.TriangleMesh.ISPC"), bChaos_TriangleMesh_ISPC_Enabled, TEXT("Whether to use ISPC optimizations in triangle mesh calculations"));
+// Support run-time toggling on supported platforms in non-shipping configurations
+#if !INTEL_ISPC || UE_BUILD_SHIPPING
+static constexpr bool bChaos_TriangleMesh_ISPC_Enabled = INTEL_ISPC && CHAOS_TRIANGLE_MESH_ISPC_ENABLED_DEFAULT;
+#else
+static bool bChaos_TriangleMesh_ISPC_Enabled = CHAOS_TRIANGLE_MESH_ISPC_ENABLED_DEFAULT;
+static FAutoConsoleVariableRef CVarChaosTriangleMeshISPCEnabled(TEXT("p.Chaos.TriangleMesh.ISPC"), bChaos_TriangleMesh_ISPC_Enabled, TEXT("Whether to use ISPC optimizations in triangle mesh calculations"));
 #endif
 
 namespace Chaos
diff --git a/Engine/Source/Runtime/Experimental/Chaos/Public/Chaos/AABB.h b/Engine/Source/Runtime/Experimental/Chaos/Public/Chaos/AABB.h
--- a/Engine/Source/Runtime/Experimental/Chaos/Public/Chaos/AABB.h
+++ b/Engine/Source/Runtime/Experimental/Chaos/Public/Chaos/AABB.h
@@ -691,13 +691,4 @@ namespace Chaos
 	};
 
 	using FAABB3 = TAABB<FReal, 3>;
-}
-
-// Support ISPC enable/disable in non-shipping builds
-#if !INTEL_ISPC
-const bool bChaos_AABBTransform_ISPC_Enabled = false;
-#elif UE_BUILD_SHIPPING
-const bool bChaos_AABBTransform_ISPC_Enabled = true;
-#else
-extern bool bChaos_AABBTransform_ISPC_Enabled;
-#endif
\ No newline at end of file
+}
\ No newline at end of file
diff --git a/Engine/Source/Runtime/Experimental/Chaos/Public/Chaos/PBDAxialSpringConstraints.h b/Engine/Source/Runtime/Experimental/Chaos/Public/Chaos/PBDAxialSpringConstraints.h
--- a/Engine/Source/Runtime/Experimental/Chaos/Public/Chaos/PBDAxialSpringConstraints.h
+++ b/Engine/Source/Runtime/Experimental/Chaos/Public/Chaos/PBDAxialSpringConstraints.h
@@ -40,11 +40,13 @@ private:
 
 }  // End namespace Chaos::Softs
 
-// Support ISPC enable/disable in non-shipping builds
-#if !INTEL_ISPC
-const bool bChaos_AxialSpring_ISPC_Enabled = false;
-#elif UE_BUILD_SHIPPING
-const bool bChaos_AxialSpring_ISPC_Enabled = true;
+#if !defined(CHAOS_AXIAL_SPRING_ISPC_ENABLED_DEFAULT)
+#define CHAOS_AXIAL_SPRING_ISPC_ENABLED_DEFAULT 1
+#endif
+
+// Support run-time toggling on supported platforms in non-shipping configurations
+#if !INTEL_ISPC || UE_BUILD_SHIPPING
+static constexpr bool bChaos_AxialSpring_ISPC_Enabled = INTEL_ISPC && CHAOS_AXIAL_SPRING_ISPC_ENABLED_DEFAULT;
 #else
 extern CHAOS_API bool bChaos_AxialSpring_ISPC_Enabled;
 #endif
diff --git a/Engine/Source/Runtime/Experimental/Chaos/Public/Chaos/PBDEvolution.h b/Engine/Source/Runtime/Experimental/Chaos/Public/Chaos/PBDEvolution.h
--- a/Engine/Source/Runtime/Experimental/Chaos/Public/Chaos/PBDEvolution.h
+++ b/Engine/Source/Runtime/Experimental/Chaos/Public/Chaos/PBDEvolution.h
@@ -205,11 +205,13 @@ private:
 
 }  // End namespace Chaos::Softs
 
+#if !defined(CHAOS_POST_ITERATION_UPDATES_ISPC_ENABLED)
+#define CHAOS_POST_ITERATION_UPDATES_ISPC_ENABLED 1
+#endif
+
 // Support ISPC enable/disable in non-shipping builds
-#if !INTEL_ISPC
-const bool bChaos_PostIterationUpdates_ISPC_Enabled = false;
-#elif UE_BUILD_SHIPPING
-const bool bChaos_PostIterationUpdates_ISPC_Enabled = true;
+#if !INTEL_ISPC || UE_BUILD_SHIPPING
+static constexpr bool bChaos_PostIterationUpdates_ISPC_Enabled = INTEL_ISPC && CHAOS_POST_ITERATION_UPDATES_ISPC_ENABLED;
 #else
 extern CHAOS_API bool bChaos_PostIterationUpdates_ISPC_Enabled;
 #endif
diff --git a/Engine/Source/Runtime/Experimental/Chaos/Public/Chaos/PBDLongRangeConstraints.h b/Engine/Source/Runtime/Experimental/Chaos/Public/Chaos/PBDLongRangeConstraints.h
--- a/Engine/Source/Runtime/Experimental/Chaos/Public/Chaos/PBDLongRangeConstraints.h
+++ b/Engine/Source/Runtime/Experimental/Chaos/Public/Chaos/PBDLongRangeConstraints.h
@@ -40,11 +40,13 @@ private:
 
 }  // End namespace Chaos::Softs
 
-// Support ISPC enable/disable in non-shipping builds
-#if !INTEL_ISPC
-const bool bChaos_LongRange_ISPC_Enabled = false;
-#elif UE_BUILD_SHIPPING
-const bool bChaos_LongRange_ISPC_Enabled = true;
+#if !defined(CHAOS_LONG_RANGE_ISPC_ENABLED_DEFAULT)
+#define CHAOS_LONG_RANGE_ISPC_ENABLED_DEFAULT 1
+#endif
+
+// Support run-time toggling on supported platforms in non-shipping configurations
+#if !INTEL_ISPC || UE_BUILD_SHIPPING
+static constexpr bool bChaos_LongRange_ISPC_Enabled = INTEL_ISPC && CHAOS_LONG_RANGE_ISPC_ENABLED_DEFAULT;
 #else
 extern CHAOS_API bool bChaos_LongRange_ISPC_Enabled;
 #endif
diff --git a/Engine/Source/Runtime/Experimental/Chaos/Public/Chaos/PBDSphericalConstraint.h b/Engine/Source/Runtime/Experimental/Chaos/Public/Chaos/PBDSphericalConstraint.h
--- a/Engine/Source/Runtime/Experimental/Chaos/Public/Chaos/PBDSphericalConstraint.h
+++ b/Engine/Source/Runtime/Experimental/Chaos/Public/Chaos/PBDSphericalConstraint.h
@@ -11,11 +11,13 @@
 DECLARE_CYCLE_STAT(TEXT("Chaos PBD Spherical Constraints"), STAT_PBD_Spherical, STATGROUP_Chaos);
 DECLARE_CYCLE_STAT(TEXT("Chaos PBD Spherical Backstop Constraints"), STAT_PBD_SphericalBackstop, STATGROUP_Chaos);
 
-// Support ISPC enable/disable in non-shipping builds
-#if !INTEL_ISPC
-const bool bChaos_Spherical_ISPC_Enabled = false;
-#elif UE_BUILD_SHIPPING
-const bool bChaos_Spherical_ISPC_Enabled = true;
+#if !defined(CHAOS_SPHERICAL_ISPC_ENABLED_DEFAULT)
+#define CHAOS_SPHERICAL_ISPC_ENABLED_DEFAULT 1
+#endif
+
+// Support run-time toggling on supported platforms in non-shipping configurations
+#if !INTEL_ISPC || UE_BUILD_SHIPPING
+static constexpr bool bChaos_Spherical_ISPC_Enabled = INTEL_ISPC && CHAOS_SPHERICAL_ISPC_ENABLED_DEFAULT;
 #else
 extern CHAOS_API bool bChaos_Spherical_ISPC_Enabled;
 #endif
diff --git a/Engine/Source/Runtime/Experimental/Chaos/Public/Chaos/PBDSpringConstraints.h b/Engine/Source/Runtime/Experimental/Chaos/Public/Chaos/PBDSpringConstraints.h
--- a/Engine/Source/Runtime/Experimental/Chaos/Public/Chaos/PBDSpringConstraints.h
+++ b/Engine/Source/Runtime/Experimental/Chaos/Public/Chaos/PBDSpringConstraints.h
@@ -44,11 +44,13 @@ private:
 
 }  // End namespace Chaos::Softs
 
-// Support ISPC enable/disable in non-shipping builds
-#if !INTEL_ISPC
-const bool bChaos_Spring_ISPC_Enabled = false;
-#elif UE_BUILD_SHIPPING
-const bool bChaos_Spring_ISPC_Enabled = true;
+#if !defined(CHAOS_SPRING_ISPC_ENABLED_DEFAULT)
+#define CHAOS_SPRING_ISPC_ENABLED_DEFAULT 1
+#endif
+
+// Support run-time toggling on supported platforms in non-shipping configurations
+#if !INTEL_ISPC || UE_BUILD_SHIPPING
+static constexpr bool bChaos_Spring_ISPC_Enabled = INTEL_ISPC && CHAOS_SPRING_ISPC_ENABLED_DEFAULT;
 #else
 extern CHAOS_API bool bChaos_Spring_ISPC_Enabled;
 #endif
diff --git a/Engine/Source/Runtime/Experimental/Chaos/Public/Chaos/PerParticleDampVelocity.h b/Engine/Source/Runtime/Experimental/Chaos/Public/Chaos/PerParticleDampVelocity.h
--- a/Engine/Source/Runtime/Experimental/Chaos/Public/Chaos/PerParticleDampVelocity.h
+++ b/Engine/Source/Runtime/Experimental/Chaos/Public/Chaos/PerParticleDampVelocity.h
@@ -34,11 +34,13 @@ private:
 
 }  // End namespace Chaos::Softs
 
-// Support ISPC enable/disable in non-shipping builds
-#if !INTEL_ISPC
-const bool bChaos_DampVelocity_ISPC_Enabled = false;
-#elif UE_BUILD_SHIPPING
-const bool bChaos_DampVelocity_ISPC_Enabled = true;
+#if !defined(CHAOS_DAMP_VELOCITY_ISPC_ENABLED_DEFAULT)
+#define CHAOS_DAMP_VELOCITY_ISPC_ENABLED_DEFAULT 1
+#endif
+
+// Support run-time toggling on supported platforms in non-shipping configurations
+#if !INTEL_ISPC || UE_BUILD_SHIPPING
+static constexpr bool bChaos_DampVelocity_ISPC_Enabled = INTEL_ISPC && CHAOS_DAMP_VELOCITY_ISPC_ENABLED_DEFAULT;
 #else
 extern CHAOS_API bool bChaos_DampVelocity_ISPC_Enabled;
 #endif
diff --git a/Engine/Source/Runtime/Experimental/Chaos/Public/Chaos/PerParticlePBDCollisionConstraint.h b/Engine/Source/Runtime/Experimental/Chaos/Public/Chaos/PerParticlePBDCollisionConstraint.h
--- a/Engine/Source/Runtime/Experimental/Chaos/Public/Chaos/PerParticlePBDCollisionConstraint.h
+++ b/Engine/Source/Runtime/Experimental/Chaos/Public/Chaos/PerParticlePBDCollisionConstraint.h
@@ -8,11 +8,13 @@
 #include "Chaos/KinematicGeometryParticles.h"
 #include "HAL/PlatformMath.h"
 
-// Support ISPC enable/disable in non-shipping builds
-#if !INTEL_ISPC
-const bool bChaos_PerParticleCollision_ISPC_Enabled = false;
-#elif UE_BUILD_SHIPPING
-const bool bChaos_PerParticleCollision_ISPC_Enabled = true;
+#if !defined(CHAOS_PER_PARTICLE_COLLISION_ISPC_ENABLED_DEFAULT)
+#define CHAOS_PER_PARTICLE_COLLISION_ISPC_ENABLED_DEFAULT 1
+#endif
+
+// Support run-time toggling on supported platforms in non-shipping configurations
+#if !INTEL_ISPC || UE_BUILD_SHIPPING
+const bool bChaos_PerParticleCollision_ISPC_Enabled = INTEL_ISPC && CHAOS_PER_PARTICLE_COLLISION_ISPC_ENABLED_DEFAULT;
 #else
 extern CHAOS_API bool bChaos_PerParticleCollision_ISPC_Enabled;
 #endif
diff --git a/Engine/Source/Runtime/Experimental/Chaos/Public/Chaos/TriangleMesh.h b/Engine/Source/Runtime/Experimental/Chaos/Public/Chaos/TriangleMesh.h
--- a/Engine/Source/Runtime/Experimental/Chaos/Public/Chaos/TriangleMesh.h
+++ b/Engine/Source/Runtime/Experimental/Chaos/Public/Chaos/TriangleMesh.h
@@ -342,12 +342,3 @@ namespace Chaos
 	template <typename T>
 	using TTriangleMesh = FTriangleMesh;
 }
-
-// Support ISPC enable/disable in non-shipping builds
-#if !INTEL_ISPC
-const bool bChaos_TriangleMesh_ISPC_Enabled = false;
-#elif UE_BUILD_SHIPPING
-const bool bChaos_TriangleMesh_ISPC_Enabled = true;
-#else
-extern CHAOS_API bool bChaos_TriangleMesh_ISPC_Enabled;
-#endif
diff --git a/Engine/Source/Runtime/Experimental/Chaos/Public/Chaos/VelocityField.h b/Engine/Source/Runtime/Experimental/Chaos/Public/Chaos/VelocityField.h
--- a/Engine/Source/Runtime/Experimental/Chaos/Public/Chaos/VelocityField.h
+++ b/Engine/Source/Runtime/Experimental/Chaos/Public/Chaos/VelocityField.h
@@ -115,11 +115,13 @@ using FVelocityField UE_DEPRECATED(5.1, "Chaos::Softs::FVelocityField has been r
 
 }  // End namespace Chaos::Softs
 
-// Support ISPC enable/disable in non-shipping builds
-#if !INTEL_ISPC
-const bool bChaos_VelocityField_ISPC_Enabled = false;
-#elif UE_BUILD_SHIPPING
-const bool bChaos_VelocityField_ISPC_Enabled = true;
+#if !defined(CHAOS_VELOCITY_FIELD_ISPC_ENABLED_DEFAULT)
+#define CHAOS_VELOCITY_FIELD_ISPC_ENABLED_DEFAULT 1
+#endif
+
+// Support run-time toggling on supported platforms in non-shipping configurations
+#if !INTEL_ISPC || UE_BUILD_SHIPPING
+static constexpr bool bChaos_VelocityField_ISPC_Enabled = INTEL_ISPC && CHAOS_VELOCITY_FIELD_ISPC_ENABLED_DEFAULT;
 #else
 extern CHAOS_API bool bChaos_VelocityField_ISPC_Enabled;
 #endif
diff --git a/Engine/Source/Runtime/Experimental/GeometryCollectionEngine/Private/GeometryCollection/GeometryCollectionComponent.cpp b/Engine/Source/Runtime/Experimental/GeometryCollectionEngine/Private/GeometryCollection/GeometryCollectionComponent.cpp
--- a/Engine/Source/Runtime/Experimental/GeometryCollectionEngine/Private/GeometryCollection/GeometryCollectionComponent.cpp
+++ b/Engine/Source/Runtime/Experimental/GeometryCollectionEngine/Private/GeometryCollection/GeometryCollectionComponent.cpp
@@ -65,9 +65,16 @@ static_assert(sizeof(ispc::FMatrix) == sizeof(FMatrix), "sizeof(ispc::FMatrix) !
 static_assert(sizeof(ispc::FBox) == sizeof(FBox), "sizeof(ispc::FBox) != sizeof(FBox)");
 #endif
 
-#if INTEL_ISPC && !UE_BUILD_SHIPPING
-bool bChaos_BoxCalcBounds_ISPC_Enabled = true;
-FAutoConsoleVariableRef CVarChaosBoxCalcBoundsISPCEnabled(TEXT("p.Chaos.BoxCalcBounds.ISPC"), bChaos_BoxCalcBounds_ISPC_Enabled, TEXT("Whether to use ISPC optimizations in calculating box bounds in geometry collections"));
+#if !defined(CHAOS_BOX_CALC_BOUNDS_ISPC_ENABLED_DEFAULT)
+#define CHAOS_BOX_CALC_BOUNDS_ISPC_ENABLED_DEFAULT 1
+#endif
+
+// Support run-time toggling on supported platforms in non-shipping configurations
+#if !INTEL_ISPC || UE_BUILD_SHIPPING
+static constexpr bool bChaos_BoxCalcBounds_ISPC_Enabled = INTEL_ISPC && CHAOS_BOX_CALC_BOUNDS_ISPC_ENABLED_DEFAULT;
+#else
+static bool bChaos_BoxCalcBounds_ISPC_Enabled = CHAOS_BOX_CALC_BOUNDS_ISPC_ENABLED_DEFAULT;
+static FAutoConsoleVariableRef CVarChaosBoxCalcBoundsISPCEnabled(TEXT("p.Chaos.BoxCalcBounds.ISPC"), bChaos_BoxCalcBounds_ISPC_Enabled, TEXT("Whether to use ISPC optimizations in calculating box bounds in geometry collections"));
 #endif
 
 DEFINE_LOG_CATEGORY_STATIC(UGCC_LOG, Error, All);
diff --git a/Engine/Source/Runtime/Experimental/GeometryCollectionEngine/Private/GeometryCollection/GeometryCollectionSceneProxy.cpp b/Engine/Source/Runtime/Experimental/GeometryCollectionEngine/Private/GeometryCollection/GeometryCollectionSceneProxy.cpp
--- a/Engine/Source/Runtime/Experimental/GeometryCollectionEngine/Private/GeometryCollection/GeometryCollectionSceneProxy.cpp
+++ b/Engine/Source/Runtime/Experimental/GeometryCollectionEngine/Private/GeometryCollection/GeometryCollectionSceneProxy.cpp
@@ -84,9 +84,16 @@ FAutoConsoleVariableRef CVarRayTracingGeometryCollectionProxyMeshes(
 	ECVF_RenderThreadSafe
 );
 
-#if INTEL_ISPC && !UE_BUILD_SHIPPING
-bool bGeometryCollection_SetDynamicData_ISPC_Enabled = true;
-FAutoConsoleVariableRef CVarGeometryCollectionSetDynamicDataISPCEnabled(TEXT("r.GeometryCollectionSetDynamicData.ISPC"), bGeometryCollection_SetDynamicData_ISPC_Enabled, TEXT("Whether to use ISPC optimizations to set dynamic data in geometry collections"));
+#if !defined(CHAOS_GEOMETRY_COLLECTION_SET_DYNAMIC_DATA_ISPC_ENABLED_DEFAULT)
+#define CHAOS_GEOMETRY_COLLECTION_SET_DYNAMIC_DATA_ISPC_ENABLED_DEFAULT 1
+#endif
+
+// Support run-time toggling on supported platforms in non-shipping configurations
+#if !INTEL_ISPC || UE_BUILD_SHIPPING
+static constexpr bool bGeometryCollection_SetDynamicData_ISPC_Enabled = INTEL_ISPC && CHAOS_GEOMETRY_COLLECTION_SET_DYNAMIC_DATA_ISPC_ENABLED_DEFAULT;
+#else
+static bool bGeometryCollection_SetDynamicData_ISPC_Enabled = CHAOS_GEOMETRY_COLLECTION_SET_DYNAMIC_DATA_ISPC_ENABLED_DEFAULT;
+static FAutoConsoleVariableRef CVarGeometryCollectionSetDynamicDataISPCEnabled(TEXT("r.GeometryCollectionSetDynamicData.ISPC"), bGeometryCollection_SetDynamicData_ISPC_Enabled, TEXT("Whether to use ISPC optimizations to set dynamic data in geometry collections"));
 #endif
 
 DEFINE_LOG_CATEGORY_STATIC(FGeometryCollectionSceneProxyLogging, Log, All);
diff --git a/Engine/Source/Runtime/Experimental/GeometryCollectionEngine/Private/GeometryCollection/GeometryCollectionSceneProxy.h b/Engine/Source/Runtime/Experimental/GeometryCollectionEngine/Private/GeometryCollection/GeometryCollectionSceneProxy.h
--- a/Engine/Source/Runtime/Experimental/GeometryCollectionEngine/Private/GeometryCollection/GeometryCollectionSceneProxy.h
+++ b/Engine/Source/Runtime/Experimental/GeometryCollectionEngine/Private/GeometryCollection/GeometryCollectionSceneProxy.h
@@ -446,12 +446,3 @@ protected:
 	uint32 bCurrentlyInMotion : 1;
 	uint32 bRequiresGPUSceneUpdate : 1;
 };
-
-// Support ISPC enable/disable in non-shipping builds
-#if !INTEL_ISPC
-const bool bGeometryCollection_SetDynamicData_ISPC_Enabled = false;
-#elif UE_BUILD_SHIPPING
-const bool bGeometryCollection_SetDynamicData_ISPC_Enabled = true;
-#else
-extern bool bGeometryCollection_SetDynamicData_ISPC_Enabled;
-#endif
diff --git a/Engine/Source/Runtime/Experimental/GeometryCollectionEngine/Public/GeometryCollection/GeometryCollectionComponent.h b/Engine/Source/Runtime/Experimental/GeometryCollectionEngine/Public/GeometryCollection/GeometryCollectionComponent.h
--- a/Engine/Source/Runtime/Experimental/GeometryCollectionEngine/Public/GeometryCollection/GeometryCollectionComponent.h
+++ b/Engine/Source/Runtime/Experimental/GeometryCollectionEngine/Public/GeometryCollection/GeometryCollectionComponent.h
@@ -896,12 +896,3 @@ private:
 	bool bIsMoving;
 
 };
-
-// Support ISPC enable/disable in non-shipping builds
-#if !INTEL_ISPC
-const bool bChaos_BoxCalcBounds_ISPC_Enabled = false;
-#elif UE_BUILD_SHIPPING
-const bool bChaos_BoxCalcBounds_ISPC_Enabled = true;
-#else
-extern bool bChaos_BoxCalcBounds_ISPC_Enabled;
-#endif
