Added blueprint accessor for getting texture memory size on Texture.

diff --git a/Engine/Source/Runtime/Engine/Classes/Engine/Texture.h b/Engine/Source/Runtime/Engine/Classes/Engine/Texture.h
--- a/Engine/Source/Runtime/Engine/Classes/Engine/Texture.h
+++ b/Engine/Source/Runtime/Engine/Classes/Engine/Texture.h
@@ -1352,6 +1352,12 @@ public:
 	UPROPERTY(transient)
 	uint8 bNotOfflineProcessed : 1;
 
+	/**
+	 * Gets the memory size of the texture, in bytes
+	 */
+	UFUNCTION(BlueprintCallable, meta = (DisplayName = "GetMemorySize"), Category = "Rendering|Texture")
+	int64 Blueprint_GetMemorySize() const;
+
 private:
 	/** Whether the async resource release process has already been kicked off or not */
 	UPROPERTY(transient)
diff --git a/Engine/Source/Runtime/Engine/Private/Texture.cpp b/Engine/Source/Runtime/Engine/Private/Texture.cpp
--- a/Engine/Source/Runtime/Engine/Private/Texture.cpp
+++ b/Engine/Source/Runtime/Engine/Private/Texture.cpp
@@ -3285,6 +3285,11 @@ void UTexture::NotifyMaterials(const ENotifyMaterialsEffectOnShaders EffectOnSha
 
 #endif //WITH_EDITOR
 
+int64 UTexture::Blueprint_GetMemorySize() const
+{
+	return CalcTextureMemorySizeEnum(TMC_AllMips);
+}
+
 #if WITH_EDITOR
 
 FTextureSource::FMipData::FMipData(const FTextureSource& InSource, FSharedBuffer InData)
