Converting hardcoded short class/enum names to pathnames ahead of ANY_PACKAGE removal

diff --git a/Engine/Plugins/Runtime/ZoneGraph/Source/ZoneGraph/Public/ZoneGraphTypes.h b/Engine/Plugins/Runtime/ZoneGraph/Source/ZoneGraph/Public/ZoneGraphTypes.h
--- a/Engine/Plugins/Runtime/ZoneGraph/Source/ZoneGraph/Public/ZoneGraphTypes.h
+++ b/Engine/Plugins/Runtime/ZoneGraph/Source/ZoneGraph/Public/ZoneGraphTypes.h
@@ -844,7 +844,7 @@ struct ZONEGRAPH_API FZoneShapePoint
 	bool bReverseLaneProfile = false;
 
 	/** Lane connection restrictions */
-	UPROPERTY(Category = Zone, EditAnywhere, BlueprintReadWrite, meta = (Bitmask, BitmaskEnum = EZoneShapeLaneConnectionRestrictions))
+	UPROPERTY(Category = Zone, EditAnywhere, BlueprintReadWrite, meta = (Bitmask, BitmaskEnum = "/Script/ZoneGraph.EZoneShapeLaneConnectionRestrictions"))
 	int32 LaneConnectionRestrictions = 0;
 };
 
@@ -947,7 +947,7 @@ struct ZONEGRAPH_API FZoneGraphLaneRoutingRule
 	UPROPERTY(Category = Rule, EditAnywhere)
 	EZoneGraphLaneRoutingCountRule DestinationIncomingConnections = EZoneGraphLaneRoutingCountRule::Any;
 
-	UPROPERTY(Category = Rule, EditAnywhere, meta = (Bitmask, BitmaskEnum = EZoneShapeLaneConnectionRestrictions))
+	UPROPERTY(Category = Rule, EditAnywhere, meta = (Bitmask, BitmaskEnum = "/Script/ZoneGraph.EZoneShapeLaneConnectionRestrictions"))
 	int32 ConnectionRestrictions = 0;
 };
 
diff --git a/Engine/Plugins/VirtualProduction/CameraCalibrationCore/Source/CameraCalibrationCore/Public/LensComponent.h b/Engine/Plugins/VirtualProduction/CameraCalibrationCore/Source/CameraCalibrationCore/Public/LensComponent.h
--- a/Engine/Plugins/VirtualProduction/CameraCalibrationCore/Source/CameraCalibrationCore/Public/LensComponent.h
+++ b/Engine/Plugins/VirtualProduction/CameraCalibrationCore/Source/CameraCalibrationCore/Public/LensComponent.h
@@ -54,7 +54,7 @@ private:
 
 protected:
 	/** The CineCameraComponent on which to apply the post-process distortion effect */
-	UPROPERTY(EditInstanceOnly, Category = "Default", meta = (UseComponentPicker, AllowedClasses = "CineCameraComponent"))
+	UPROPERTY(EditInstanceOnly, Category = "Default", meta = (UseComponentPicker, AllowedClasses = "/Script/CinematicCamera.CineCameraComponent"))
 	FComponentReference TargetCameraComponent;
 
 	/** Structure used to query the camera calibration subsystem for a lens distortion model handler */
diff --git a/Engine/Plugins/VirtualProduction/DMX/DMXEngine/Source/DMXEditor/Public/Widgets/SAssetPickerButton.h b/Engine/Plugins/VirtualProduction/DMX/DMXEngine/Source/DMXEditor/Public/Widgets/SAssetPickerButton.h
--- a/Engine/Plugins/VirtualProduction/DMX/DMXEngine/Source/DMXEditor/Public/Widgets/SAssetPickerButton.h
+++ b/Engine/Plugins/VirtualProduction/DMX/DMXEngine/Source/DMXEditor/Public/Widgets/SAssetPickerButton.h
@@ -37,11 +37,11 @@ public:
 		/** Allows more than a single class on the picker list.
 		 * Overrides AssetClass (doesn't affect the Use Selected on Content Browser button)
 		 */
-		SLATE_ARGUMENT(TArray<FName>, AllowedClasses)
+		SLATE_ARGUMENT(TArray<FTopLevelAssetPath>, AllowedClasses)
 		/** Disallows subclasses from the allowed asset class(es).
 		 * Doesn't affect the Use Selected on Content Browser button.
 		 */
-		SLATE_ARGUMENT(TArray<FName>, DisallowedClasses)
+		SLATE_ARGUMENT(TArray<FTopLevelAssetPath>, DisallowedClasses)
 
 		/** Called to decide the hovered state using the parent's interact able area */
 		SLATE_EVENT(FOnParentIsHovered, OnParentIsHovered)
@@ -56,9 +56,9 @@ public:
 	/** The class to filter the list of assets in the picker menu */
 	TWeakObjectPtr<UClass> AssetClass;
 	/** Allows more than a single class. Overrides AssetClass */
-	TArray<FName> AllowedClasses;
+	TArray<FTopLevelAssetPath> AllowedClasses;
 	/** Disallows subclasses from the allowed asset class(es) */
-	TArray<FName> DisallowedClasses;
+	TArray<FTopLevelAssetPath> DisallowedClasses;
 
 protected:
 	/** Clicked Use button */
diff --git a/Engine/Plugins/VirtualProduction/Takes/Source/TakeTrackRecorders/Public/TrackRecorders/IMovieSceneTrackRecorderHost.h b/Engine/Plugins/VirtualProduction/Takes/Source/TakeTrackRecorders/Public/TrackRecorders/IMovieSceneTrackRecorderHost.h
--- a/Engine/Plugins/VirtualProduction/Takes/Source/TakeTrackRecorders/Public/TrackRecorders/IMovieSceneTrackRecorderHost.h
+++ b/Engine/Plugins/VirtualProduction/Takes/Source/TakeTrackRecorders/Public/TrackRecorders/IMovieSceneTrackRecorderHost.h
@@ -28,7 +28,7 @@ struct FTakeRecorderTrackSettings
 	GENERATED_BODY()
 
 	/** The Actor class to create movie scene tracks for. */
-	UPROPERTY(config, EditAnywhere, Category = TrackSettings, meta = (MetaClass = "Actor"))
+	UPROPERTY(config, EditAnywhere, Category = TrackSettings, meta = (MetaClass = "/Script/Engine.Actor"))
 	FSoftClassPath MatchingActorClass;
 
 	/** List of property names for which movie scene tracks will be created automatically. */
diff --git a/Engine/Source/Developer/AnimationDataController/Private/AnimDataController.cpp b/Engine/Source/Developer/AnimationDataController/Private/AnimDataController.cpp
--- a/Engine/Source/Developer/AnimationDataController/Private/AnimDataController.cpp
+++ b/Engine/Source/Developer/AnimationDataController/Private/AnimDataController.cpp
@@ -1269,7 +1269,7 @@ FString UAnimDataController::GetCurveTypeValueName(ERawCurveTrackTypes InType) c
 {
 	FString ValueString;
 
-	const UEnum* Enum = FindObject<UEnum>(ANY_PACKAGE, TEXT("ERawCurveTrackTypes"));
+	const UEnum* Enum = FindObject<UEnum>(nullptr, TEXT("/Script/Engine.ERawCurveTrackTypes"));
 	if (Enum)
 	{
 		ValueString = Enum->GetNameStringByValue((int64)InType);
diff --git a/Engine/Source/Developer/AssetTools/Public/AssetToolsSettings.h b/Engine/Source/Developer/AssetTools/Public/AssetToolsSettings.h
--- a/Engine/Source/Developer/AssetTools/Public/AssetToolsSettings.h
+++ b/Engine/Source/Developer/AssetTools/Public/AssetToolsSettings.h
@@ -16,11 +16,11 @@ struct FAdvancedCopyMap
 
 public:
 	/** When copying this class, use a particular set of dependency and destination rules */
-	UPROPERTY(EditAnywhere, Category = "Asset Tools", meta = (MetaClass = "Object"))
+	UPROPERTY(EditAnywhere, Category = "Asset Tools", meta = (MetaClass = "/Script/CoreUObject.Object"))
 	FSoftClassPath ClassToCopy;
 
 	/** The set of dependency and destination rules to use for advanced copy */
-	UPROPERTY(EditAnywhere, Category = "Asset Tools", meta = (MetaClass = "AdvancedCopyCustomization"))
+	UPROPERTY(EditAnywhere, Category = "Asset Tools", meta = (MetaClass = "/Script/AssetTools.AdvancedCopyCustomization"))
 	FSoftClassPath AdvancedCopyCustomization;
 };
 
diff --git a/Engine/Source/Developer/IOS/IOSPlatformEditor/Private/IOSTargetSettingsCustomization.cpp b/Engine/Source/Developer/IOS/IOSPlatformEditor/Private/IOSTargetSettingsCustomization.cpp
--- a/Engine/Source/Developer/IOS/IOSPlatformEditor/Private/IOSTargetSettingsCustomization.cpp
+++ b/Engine/Source/Developer/IOS/IOSPlatformEditor/Private/IOSTargetSettingsCustomization.cpp
@@ -1837,7 +1837,7 @@ TSharedRef<SWidget> FIOSTargetSettingsCustomization::OnGetShaderVersionContent()
 {
 	FMenuBuilder MenuBuilder(true, NULL);
 	
-	UEnum* Enum = FindObjectChecked<UEnum>(ANY_PACKAGE, TEXT("EIOSMetalShaderStandard"), true);
+	UEnum* Enum = FindObjectChecked<UEnum>(nullptr, TEXT("/Script/IOSRuntimeSettings.EIOSMetalShaderStandard"), true);
 	
 	for (int32 i = 0; i < Enum->GetMaxEnumValue(); i++)
 	{
@@ -1856,7 +1856,7 @@ FText FIOSTargetSettingsCustomization::GetShaderVersionDesc() const
 	uint8 EnumValue;
 	ShaderVersionPropertyHandle->GetValue(EnumValue);
 	
-	UEnum* Enum = FindObjectChecked<UEnum>(ANY_PACKAGE, TEXT("EIOSMetalShaderStandard"), true);
+	UEnum* Enum = FindObjectChecked<UEnum>(nullptr, TEXT("/Script/IOSRuntimeSettings.EIOSMetalShaderStandard"), true);
 	
 	if (EnumValue < Enum->GetMaxEnumValue() && Enum->IsValidEnumValue(EnumValue))
 	{
@@ -1870,7 +1870,7 @@ TSharedRef<SWidget> FIOSTargetSettingsCustomization::OnGetMinVersionContent()
 {
 	FMenuBuilder MenuBuilder(true, NULL);
 
-	UEnum* Enum = FindObjectChecked<UEnum>(ANY_PACKAGE, TEXT("EIOSVersion"), true);
+	UEnum* Enum = FindObjectChecked<UEnum>(nullptr, TEXT("/Script/IOSRuntimeSettings.EIOSVersion"), true);
 
 	for (int32 i = 0; i < Enum->GetMaxEnumValue(); i++)
 	{
@@ -1889,7 +1889,7 @@ FText FIOSTargetSettingsCustomization::GetMinVersionDesc() const
 	uint8 EnumValue;
 	MinOSPropertyHandle->GetValue(EnumValue);
 
-	UEnum* Enum = FindObjectChecked<UEnum>(ANY_PACKAGE, TEXT("EIOSVersion"), true);
+	UEnum* Enum = FindObjectChecked<UEnum>(nullptr, TEXT("/Script/IOSRuntimeSettings.EIOSVersion"), true);
 
 	if (EnumValue < Enum->GetMaxEnumValue() && Enum->IsValidEnumValue(EnumValue))
 	{
diff --git a/Engine/Source/Developer/Localization/Private/LocalizationConfigurationScript.cpp b/Engine/Source/Developer/Localization/Private/LocalizationConfigurationScript.cpp
--- a/Engine/Source/Developer/Localization/Private/LocalizationConfigurationScript.cpp
+++ b/Engine/Source/Developer/Localization/Private/LocalizationConfigurationScript.cpp
@@ -565,11 +565,11 @@ namespace LocalizationConfigurationScript
 
 			// Import-specific settings.
 			{
-				UEnum* LocalizedTextCollapseModeEnum = FindObjectChecked<UEnum>(ANY_PACKAGE, TEXT("ELocalizedTextCollapseMode"));
+				UEnum* LocalizedTextCollapseModeEnum = FindObjectChecked<UEnum>(nullptr, TEXT("/Script/Localization.ELocalizedTextCollapseMode"));
 				const FName CollapseModeName = LocalizedTextCollapseModeEnum->GetNameByValue((int64)Target->Settings.ExportSettings.CollapseMode);
 				ConfigSection.Add(TEXT("LocalizedTextCollapseMode"), CollapseModeName.ToString());
 
-				UEnum* POFormatEnum = FindObjectChecked<UEnum>(ANY_PACKAGE, TEXT("EPortableObjectFormat"));
+				UEnum* POFormatEnum = FindObjectChecked<UEnum>(nullptr, TEXT("/Script/Localization.EPortableObjectFormat"));
 				const FName POFormatName = POFormatEnum->GetNameByValue((int64)Target->Settings.ExportSettings.POFormat);
 				ConfigSection.Add(TEXT("POFormat"), POFormatName.ToString());
 			}
@@ -690,11 +690,11 @@ namespace LocalizationConfigurationScript
 
 			// Export-specific settings.
 			{
-				UEnum* LocalizedTextCollapseModeEnum = FindObjectChecked<UEnum>(ANY_PACKAGE, TEXT("ELocalizedTextCollapseMode"));
+				UEnum* LocalizedTextCollapseModeEnum = FindObjectChecked<UEnum>(nullptr, TEXT("/Script/Localization.ELocalizedTextCollapseMode"));
 				const FName CollapseModeName = LocalizedTextCollapseModeEnum->GetNameByValue((int64)Target->Settings.ExportSettings.CollapseMode);
 				ConfigSection.Add(TEXT("LocalizedTextCollapseMode"), CollapseModeName.ToString());
 
-				UEnum* POFormatEnum = FindObjectChecked<UEnum>(ANY_PACKAGE, TEXT("EPortableObjectFormat"));
+				UEnum* POFormatEnum = FindObjectChecked<UEnum>(nullptr, TEXT("/Script/Localization.EPortableObjectFormat"));
 				const FName POFormatName = POFormatEnum->GetNameByValue((int64)Target->Settings.ExportSettings.POFormat);
 				ConfigSection.Add(TEXT("POFormat"), POFormatName.ToString());
 
diff --git a/Engine/Source/Developer/Localization/Public/LocalizationTargetTypes.h b/Engine/Source/Developer/Localization/Public/LocalizationTargetTypes.h
--- a/Engine/Source/Developer/Localization/Public/LocalizationTargetTypes.h
+++ b/Engine/Source/Developer/Localization/Public/LocalizationTargetTypes.h
@@ -179,7 +179,7 @@ struct FGatherTextFromPackagesConfiguration
 	TArray<FName> Collections;
 
 	/* Classes that should be excluded from gathering. */
-	UPROPERTY(config, EditAnywhere, Category = "Filter", meta=(MetaClass="Object", AllowAbstract="True"))
+	UPROPERTY(config, EditAnywhere, Category = "Filter", meta=(MetaClass="/Script/CoreUObject.Object", AllowAbstract="True"))
 	TArray<FSoftClassPath> ExcludeClasses;
 
 	/* Should classes derived from those in the exclude classes list also be excluded from gathering? */
diff --git a/Engine/Source/Developer/ToolMenus/Public/ToolMenuEntryScript.h b/Engine/Source/Developer/ToolMenus/Public/ToolMenuEntryScript.h
--- a/Engine/Source/Developer/ToolMenus/Public/ToolMenuEntryScript.h
+++ b/Engine/Source/Developer/ToolMenus/Public/ToolMenuEntryScript.h
@@ -21,7 +21,7 @@
 struct FToolMenuEntry;
 struct FToolMenuSection;
 
-USTRUCT(BlueprintType, meta=(HasNativeBreak="ToolMenus.ToolMenuEntryExtensions.BreakScriptSlateIcon", HasNativeMake="ToolMenus.ToolMenuEntryExtensions.MakeScriptSlateIcon"))
+USTRUCT(BlueprintType, meta=(HasNativeBreak="/Script/ToolMenus.ToolMenuEntryExtensions.BreakScriptSlateIcon", HasNativeMake="/Script/ToolMenus.ToolMenuEntryExtensions.MakeScriptSlateIcon"))
 struct TOOLMENUS_API FScriptSlateIcon
 {
 	GENERATED_BODY()
diff --git a/Engine/Source/Developer/ToolMenus/Public/ToolMenuMisc.h b/Engine/Source/Developer/ToolMenus/Public/ToolMenuMisc.h
--- a/Engine/Source/Developer/ToolMenus/Public/ToolMenuMisc.h
+++ b/Engine/Source/Developer/ToolMenus/Public/ToolMenuMisc.h
@@ -17,7 +17,7 @@ enum class EToolMenuStringCommandType : uint8
 	Custom
 };
 
-USTRUCT(BlueprintType, meta=(HasNativeBreak="ToolMenus.ToolMenuEntryExtensions.BreakStringCommand", HasNativeMake="ToolMenus.ToolMenuEntryExtensions.MakeStringCommand"))
+USTRUCT(BlueprintType, meta=(HasNativeBreak="/Script/ToolMenus.ToolMenuEntryExtensions.BreakStringCommand", HasNativeMake="/Script/ToolMenus.ToolMenuEntryExtensions.MakeStringCommand"))
 struct TOOLMENUS_API FToolMenuStringCommand
 {
 	GENERATED_BODY()
diff --git a/Engine/Source/Developer/ToolMenus/Public/ToolMenuOwner.h b/Engine/Source/Developer/ToolMenus/Public/ToolMenuOwner.h
--- a/Engine/Source/Developer/ToolMenus/Public/ToolMenuOwner.h
+++ b/Engine/Source/Developer/ToolMenus/Public/ToolMenuOwner.h
@@ -7,7 +7,7 @@
 #include "ToolMenuOwner.generated.h"
 
 
-USTRUCT(BlueprintType, meta=(HasNativeBreak="ToolMenus.ToolMenuEntryExtensions.BreakToolMenuOwner", HasNativeMake="ToolMenus.ToolMenuEntryExtensions.MakeToolMenuOwner"))
+USTRUCT(BlueprintType, meta=(HasNativeBreak="/Script/ToolMenus.ToolMenuEntryExtensions.BreakToolMenuOwner", HasNativeMake="/Script/ToolMenus.ToolMenuEntryExtensions.MakeToolMenuOwner"))
 struct TOOLMENUS_API FToolMenuOwner
 {
 	GENERATED_BODY()
diff --git a/Engine/Source/Editor/AddContentDialog/Private/ContentSourceProviders/FeaturePack/FeaturePackContentSource.cpp b/Engine/Source/Editor/AddContentDialog/Private/ContentSourceProviders/FeaturePack/FeaturePackContentSource.cpp
--- a/Engine/Source/Editor/AddContentDialog/Private/ContentSourceProviders/FeaturePack/FeaturePackContentSource.cpp
+++ b/Engine/Source/Editor/AddContentDialog/Private/ContentSourceProviders/FeaturePack/FeaturePackContentSource.cpp
@@ -771,7 +771,7 @@ bool FFeaturePackContentSource::ParseManifestString(const FString& ManifestStrin
 		}
 	}
 
-	UEnum* Enum = FindObjectChecked<UEnum>(ANY_PACKAGE, TEXT("EContentSourceCategory"));
+	UEnum* Enum = FindObjectChecked<UEnum>(nullptr, TEXT("/Script/AddContentDialog.EContentSourceCategory"));
 	for (const FString& CategoryString : CategoryStrings)
 	{
 		int32 EnumValue = Enum->GetValueByName(FName(*CategoryString));
@@ -808,7 +808,7 @@ bool FFeaturePackContentSource::ParseManifestString(const FString& ManifestStrin
 	// Parse additional packs data
 	if (ManifestObject->HasTypedField<EJson::Array>("AdditionalFeaturePacks") == true)
 	{
-		UEnum* DetailEnum = FindObjectChecked<UEnum>(ANY_PACKAGE, TEXT("EFeaturePackDetailLevel"));
+		UEnum* DetailEnum = FindObjectChecked<UEnum>(nullptr, TEXT("/Script/AddContentDialog.EFeaturePackDetailLevel"));
 		for (TSharedPtr<FJsonValue> AdditionalFeaturePackValue : ManifestObject->GetArrayField("AdditionalFeaturePacks"))
 		{
 			TSharedPtr<FJsonObject> EachAdditionalPack = AdditionalFeaturePackValue->AsObject();
diff --git a/Engine/Source/Editor/AnimGraph/Private/PoseDriverDetails.cpp b/Engine/Source/Editor/AnimGraph/Private/PoseDriverDetails.cpp
--- a/Engine/Source/Editor/AnimGraph/Private/PoseDriverDetails.cpp
+++ b/Engine/Source/Editor/AnimGraph/Private/PoseDriverDetails.cpp
@@ -88,7 +88,7 @@ void SPDD_TargetRow::Construct(const FArguments& InArgs, const TSharedRef<STable
 
 	if (DistanceMethodOptions.Num() == 0)
 	{
-		UEnum * Enum = FindObjectChecked<UEnum>(ANY_PACKAGE, TEXT("ERBFDistanceMethod"));
+		UEnum* Enum = FindObjectChecked<UEnum>(nullptr, TEXT("/Script/AnimGraphRuntime.ERBFDistanceMethod"));
 		for (int32 i = 0; i < Enum->NumEnums() - 1; i++)
 		{
 			DistanceMethodOptions.Add(MakeShareable(new FString(Enum->GetDisplayNameTextByIndex(i).ToString())));
@@ -97,7 +97,7 @@ void SPDD_TargetRow::Construct(const FArguments& InArgs, const TSharedRef<STable
 
 	if (FunctionTypeOptions.Num() == 0)
 	{
-		UEnum * Enum = FindObjectChecked<UEnum>(ANY_PACKAGE, TEXT("ERBFFunctionType"));
+		UEnum* Enum = FindObjectChecked<UEnum>(nullptr, TEXT("/Script/AnimGraphRuntime.ERBFFunctionType"));
 		for (int32 i = 0; i < Enum->NumEnums() - 1; i++)
 		{
 			FunctionTypeOptions.Add(MakeShareable(new FString(Enum->GetDisplayNameTextByIndex(i).ToString())));
diff --git a/Engine/Source/Editor/BlueprintGraph/Private/K2Node_FormatText.cpp b/Engine/Source/Editor/BlueprintGraph/Private/K2Node_FormatText.cpp
--- a/Engine/Source/Editor/BlueprintGraph/Private/K2Node_FormatText.cpp
+++ b/Engine/Source/Editor/BlueprintGraph/Private/K2Node_FormatText.cpp
@@ -426,7 +426,7 @@ void UK2Node_FormatText::ExpandNode(class FKismetCompilerContext& CompilerContex
 			}
 			else if (ArgumentPinCategory == UEdGraphSchema_K2::PC_Byte || ArgumentPinCategory == UEdGraphSchema_K2::PC_Enum)
 			{
-				static UEnum* TextGenderEnum = FindObjectChecked<UEnum>(ANY_PACKAGE, TEXT("ETextGender"), /*ExactClass*/true);
+				static UEnum* TextGenderEnum = FindObjectChecked<UEnum>(nullptr, TEXT("/Script/Engine.ETextGender"), /*ExactClass*/true);
 				if (ArgumentPin->PinType.PinSubCategoryObject == TextGenderEnum)
 				{
 					MakeFormatArgumentDataStruct->GetSchema()->TrySetDefaultValue(*ArgumentTypePin, TEXT("Gender"));
@@ -562,7 +562,7 @@ bool UK2Node_FormatText::IsConnectionDisallowed(const UEdGraphPin* MyPin, const
 		}
 		else if (OtherPinCategory == UEdGraphSchema_K2::PC_Byte || OtherPinCategory == UEdGraphSchema_K2::PC_Enum)
 		{
-			static UEnum* TextGenderEnum = FindObjectChecked<UEnum>(ANY_PACKAGE, TEXT("ETextGender"), /*ExactClass*/true);
+			static UEnum* TextGenderEnum = FindObjectChecked<UEnum>(nullptr, TEXT("/Script/Engine.ETextGender"), /*ExactClass*/true);
 			if (OtherPin->PinType.PinSubCategoryObject == TextGenderEnum)
 			{
 				bIsValidType = true;
diff --git a/Engine/Source/Editor/BlueprintGraph/Private/K2Node_SpawnActorFromClass.cpp b/Engine/Source/Editor/BlueprintGraph/Private/K2Node_SpawnActorFromClass.cpp
--- a/Engine/Source/Editor/BlueprintGraph/Private/K2Node_SpawnActorFromClass.cpp
+++ b/Engine/Source/Editor/BlueprintGraph/Private/K2Node_SpawnActorFromClass.cpp
@@ -55,7 +55,7 @@ void UK2Node_SpawnActorFromClass::AllocateDefaultPins()
 	UEdGraphPin* TransformPin = CreatePin(EGPD_Input, UEdGraphSchema_K2::PC_Struct, TransformStruct, FK2Node_SpawnActorFromClassHelper::SpawnTransformPinName);
 
 	// Collision handling method pin
-	UEnum* const MethodEnum = FindObjectChecked<UEnum>(ANY_PACKAGE, TEXT("ESpawnActorCollisionHandlingMethod"), true);
+	UEnum* const MethodEnum = FindObjectChecked<UEnum>(nullptr, TEXT("/Script/Engine.ESpawnActorCollisionHandlingMethod"), true);
 	UEdGraphPin* const CollisionHandlingOverridePin = CreatePin(EGPD_Input, UEdGraphSchema_K2::PC_Byte, MethodEnum, FK2Node_SpawnActorFromClassHelper::CollisionHandlingOverridePinName);
 	CollisionHandlingOverridePin->DefaultValue = MethodEnum->GetNameStringByValue(static_cast<int>(ESpawnActorCollisionHandlingMethod::Undefined));
 
@@ -83,7 +83,7 @@ void UK2Node_SpawnActorFromClass::MaybeUpdateCollisionPin(TArray<UEdGraphPin*>&
 				UEdGraphPin* const CollisionHandlingOverridePin = GetCollisionHandlingOverridePin();
 				if (CollisionHandlingOverridePin)
 				{
-					UEnum const* const MethodEnum = FindObjectChecked<UEnum>(ANY_PACKAGE, TEXT("ESpawnActorCollisionHandlingMethod"), true);
+					UEnum const* const MethodEnum = FindObjectChecked<UEnum>(nullptr, TEXT("/Script/Engine.ESpawnActorCollisionHandlingMethod"), true);
 					CollisionHandlingOverridePin->DefaultValue =
 						bOldCollisionPinValue
 						? MethodEnum->GetNameStringByValue(static_cast<int>(ESpawnActorCollisionHandlingMethod::AlwaysSpawn))
@@ -98,7 +98,7 @@ void UK2Node_SpawnActorFromClass::MaybeUpdateCollisionPin(TArray<UEdGraphPin*>&
 				UEdGraphPin* const CollisionHandlingOverridePin = GetCollisionHandlingOverridePin();
 				check(CollisionHandlingOverridePin);
 
-				UEnum* const MethodEnum = FindObjectChecked<UEnum>(ANY_PACKAGE, TEXT("ESpawnActorCollisionHandlingMethod"), true);
+				UEnum* const MethodEnum = FindObjectChecked<UEnum>(nullptr, TEXT("/Script/Engine.ESpawnActorCollisionHandlingMethod"), true);
 				
 				FGraphNodeCreator<UK2Node_EnumLiteral> AlwaysSpawnLiteralCreator(*GetGraph());
 				UK2Node_EnumLiteral* const AlwaysSpawnLiteralNode = AlwaysSpawnLiteralCreator.CreateNode();
diff --git a/Engine/Source/Editor/BlueprintGraph/Private/K2Node_TemporaryVariable.cpp b/Engine/Source/Editor/BlueprintGraph/Private/K2Node_TemporaryVariable.cpp
--- a/Engine/Source/Editor/BlueprintGraph/Private/K2Node_TemporaryVariable.cpp
+++ b/Engine/Source/Editor/BlueprintGraph/Private/K2Node_TemporaryVariable.cpp
@@ -228,7 +228,7 @@ void UK2Node_TemporaryVariable::GetMenuActions(FBlueprintActionDatabaseRegistrar
 	ActionRegistrar.AddBlueprintAction(ActionKey, MakeTempVarNodeSpawner(FEdGraphPinType(UEdGraphSchema_K2::PC_Struct, TEXT("Transform"), TransformStruct, EPinContainerType::None, /*bIsReference =*/ false, /*InValueTerminalType=*/FEdGraphTerminalType()), /*bIsPersistent =*/false));
 	ActionRegistrar.AddBlueprintAction(ActionKey, MakeTempVarNodeSpawner(FEdGraphPinType(UEdGraphSchema_K2::PC_Struct, TEXT("Transform"), TransformStruct, EPinContainerType::Array, /*bIsReference =*/ false, /*InValueTerminalType=*/FEdGraphTerminalType()), /*bIsPersistent =*/false));
 
-	UScriptStruct* BlendSampleStruct = FindObjectChecked<UScriptStruct>(ANY_PACKAGE, TEXT("BlendSampleData"));
+	UScriptStruct* BlendSampleStruct = FindObjectChecked<UScriptStruct>(nullptr, TEXT("/Script/Engine.BlendSampleData"));
 	ActionRegistrar.AddBlueprintAction(ActionKey, MakeTempVarNodeSpawner(FEdGraphPinType(UEdGraphSchema_K2::PC_Struct, TEXT("BlendSampleData"), BlendSampleStruct, EPinContainerType::None, /*bIsReference =*/ false, /*InValueTerminalType=*/FEdGraphTerminalType()), /*bIsPersistent =*/false));
 	ActionRegistrar.AddBlueprintAction(ActionKey, MakeTempVarNodeSpawner(FEdGraphPinType(UEdGraphSchema_K2::PC_Struct, TEXT("BlendSampleData"), BlendSampleStruct, EPinContainerType::Array, /*bIsReference =*/ false, /*InValueTerminalType=*/FEdGraphTerminalType()), /*bIsPersistent =*/false));
 
diff --git a/Engine/Source/Editor/ClassViewer/Public/ClassViewerProjectSettings.h b/Engine/Source/Editor/ClassViewer/Public/ClassViewerProjectSettings.h
--- a/Engine/Source/Editor/ClassViewer/Public/ClassViewerProjectSettings.h
+++ b/Engine/Source/Editor/ClassViewer/Public/ClassViewerProjectSettings.h
@@ -24,7 +24,7 @@ class CLASSVIEWER_API UClassViewerProjectSettings : public UObject
 	TArray<FDirectoryPath> InternalOnlyPaths;
 
 	/** The base classes to be considered Internal Only for the class picker.*/
-	UPROPERTY(EditAnywhere, config, Category = ClassVisibilityManagement, meta = (MetaClass = "Object", DisplayName = "List of base classes to consider Internal Only.", AllowAbstract = "true", ShowTreeView, HideViewOptions))
+	UPROPERTY(EditAnywhere, config, Category = ClassVisibilityManagement, meta = (MetaClass = "/Script/CoreUObject.Object", DisplayName = "List of base classes to consider Internal Only.", AllowAbstract = "true", ShowTreeView, HideViewOptions))
 	TArray<FSoftClassPath> InternalOnlyClasses;
 #endif
 };
diff --git a/Engine/Source/Editor/ContentBrowserData/Public/ContentBrowserItemPath.h b/Engine/Source/Editor/ContentBrowserData/Public/ContentBrowserItemPath.h
--- a/Engine/Source/Editor/ContentBrowserData/Public/ContentBrowserItemPath.h
+++ b/Engine/Source/Editor/ContentBrowserData/Public/ContentBrowserItemPath.h
@@ -13,7 +13,7 @@
  * 
  * Path conversion each time Set is called
  */
-USTRUCT(BlueprintType, meta=(HasNativeBreak="ContentBrowserData.ContentBrowserItemPathExtensions.BreakContentBrowserItemPath", HasNativeMake="ContentBrowserData.ContentBrowserItemPathExtensions.MakeContentBrowserItemPath"))
+USTRUCT(BlueprintType, meta=(HasNativeBreak="/Script/ContentBrowserData.ContentBrowserItemPathExtensions.BreakContentBrowserItemPath", HasNativeMake="/Script/ContentBrowserData.ContentBrowserItemPathExtensions.MakeContentBrowserItemPath"))
 struct CONTENTBROWSERDATA_API FContentBrowserItemPath
 {
 	GENERATED_BODY()
diff --git a/Engine/Source/Editor/DetailCustomizations/Private/MacTargetSettingsDetails.cpp b/Engine/Source/Editor/DetailCustomizations/Private/MacTargetSettingsDetails.cpp
--- a/Engine/Source/Editor/DetailCustomizations/Private/MacTargetSettingsDetails.cpp
+++ b/Engine/Source/Editor/DetailCustomizations/Private/MacTargetSettingsDetails.cpp
@@ -350,7 +350,7 @@ TSharedRef<SWidget> FMacTargetSettingsDetails::OnGetShaderVersionContent()
 {
 	FMenuBuilder MenuBuilder(true, NULL);
 	
-	UEnum* Enum = FindObjectChecked<UEnum>(ANY_PACKAGE, TEXT("EMacMetalShaderStandard"), true);
+	UEnum* Enum = FindObjectChecked<UEnum>(nullptr, TEXT("/Script/MacTargetPlatform.EMacMetalShaderStandard"), true);
 	
 	for (int32 i = 0; i < Enum->GetMaxEnumValue(); i++)
 	{
@@ -373,7 +373,7 @@ FText FMacTargetSettingsDetails::GetShaderVersionDesc() const
     int32 EnumValue;
     ShaderVersionPropertyHandle->GetValue(EnumValue);
 	
-	UEnum* Enum = FindObjectChecked<UEnum>(ANY_PACKAGE, TEXT("EMacMetalShaderStandard"), true);
+	UEnum* Enum = FindObjectChecked<UEnum>(nullptr, TEXT("/Script/MacTargetPlatform.EMacMetalShaderStandard"), true);
 	
 	if (EnumValue < Enum->GetMaxEnumValue() && Enum->IsValidEnumValue(EnumValue))
 	{
diff --git a/Engine/Source/Editor/Kismet/Private/BlueprintWarningsConfigurationPanel.cpp b/Engine/Source/Editor/Kismet/Private/BlueprintWarningsConfigurationPanel.cpp
--- a/Engine/Source/Editor/Kismet/Private/BlueprintWarningsConfigurationPanel.cpp
+++ b/Engine/Source/Editor/Kismet/Private/BlueprintWarningsConfigurationPanel.cpp
@@ -69,7 +69,7 @@ class SBlueprintWarningRow : public SMultiColumnTableRow< FBlueprintWarningListE
 		{
 			const auto& GetWarningText = [this]() -> FText
 			{
-				UEnum* const BlueprintWarningBehaviorEnum = FindObjectChecked<UEnum>(ANY_PACKAGE, TEXT("EBlueprintWarningBehavior"));
+				UEnum* const BlueprintWarningBehaviorEnum = FindObjectChecked<UEnum>(nullptr, TEXT("/Script/BlueprintRuntime.EBlueprintWarningBehavior"));
 				EBlueprintWarningBehavior Behavior = EBlueprintWarningBehavior::Warn;
 				FName WarningIdentifier = this->WarningInfo->WarningIdentifier;
 				if (FBlueprintSupport::ShouldTreatWarningAsError(WarningIdentifier))
@@ -103,7 +103,7 @@ class SBlueprintWarningRow : public SMultiColumnTableRow< FBlueprintWarningListE
 					FBlueprintWarningBehaviorComboBox::FOnGenerateWidget::CreateStatic(
 						[]( TSharedPtr<EBlueprintWarningBehavior> Behavior )->TSharedRef<SWidget>
 						{
-							UEnum* const BlueprintWarningBehaviorEnum = FindObjectChecked<UEnum>(ANY_PACKAGE, TEXT("EBlueprintWarningBehavior"));
+							UEnum* const BlueprintWarningBehaviorEnum = FindObjectChecked<UEnum>(nullptr, TEXT("/Script/BlueprintRuntime.EBlueprintWarningBehavior"));
 							return SNew(STextBlock)
 								.Text(BlueprintWarningBehaviorEnum->GetDisplayNameTextByValue(static_cast<int64>(*Behavior)));
 						}
diff --git a/Engine/Source/Editor/LandscapeEditor/Private/LandscapeEditorDetailCustomization_Layers.cpp b/Engine/Source/Editor/LandscapeEditor/Private/LandscapeEditorDetailCustomization_Layers.cpp
--- a/Engine/Source/Editor/LandscapeEditor/Private/LandscapeEditorDetailCustomization_Layers.cpp
+++ b/Engine/Source/Editor/LandscapeEditor/Private/LandscapeEditorDetailCustomization_Layers.cpp
@@ -378,7 +378,7 @@ void FLandscapeEditorCustomNodeBuilder_Layers::FillClearPaintLayerMenu(FMenuBuil
 void FLandscapeEditorCustomNodeBuilder_Layers::FillClearLayerMenu(FMenuBuilder& MenuBuilder, int32 InLayerIndex)
 {
 	TSharedRef<FLandscapeEditorCustomNodeBuilder_Layers> SharedThis = AsShared();
-	const UEnum* EnumPtr = FindObject<UEnum>(ANY_PACKAGE, TEXT("ELandscapeClearMode"), true);
+	const UEnum* EnumPtr = FindObject<UEnum>(nullptr, TEXT("/Script/Landscape.ELandscapeClearMode"), true);
 	if (ensure(EnumPtr != nullptr))
 	{
 		// NumEnums()-1 to exclude Enum Max Value
diff --git a/Engine/Source/Editor/LocalizationDashboard/Private/LocalizationTargetDetailCustomization.cpp b/Engine/Source/Editor/LocalizationDashboard/Private/LocalizationTargetDetailCustomization.cpp
--- a/Engine/Source/Editor/LocalizationDashboard/Private/LocalizationTargetDetailCustomization.cpp
+++ b/Engine/Source/Editor/LocalizationDashboard/Private/LocalizationTargetDetailCustomization.cpp
@@ -409,7 +409,7 @@ void FLocalizationTargetDetailCustomization::CustomizeDetails(IDetailLayoutBuild
 
 		static const TArray< TSharedPtr<ELocalizationTargetLoadingPolicy> > LoadingPolicies = []()
 		{
-			UEnum* const LoadingPolicyEnum = FindObjectChecked<UEnum>(ANY_PACKAGE, TEXT("ELocalizationTargetLoadingPolicy"));
+			UEnum* const LoadingPolicyEnum = FindObjectChecked<UEnum>(nullptr, TEXT("/Script/Localization.ELocalizationTargetLoadingPolicy"));
 			TArray< TSharedPtr<ELocalizationTargetLoadingPolicy> > Array;
 			for (int32 i = 0; i < LoadingPolicyEnum->NumEnums() - 1; ++i)
 			{
@@ -437,7 +437,7 @@ void FLocalizationTargetDetailCustomization::CustomizeDetails(IDetailLayoutBuild
 					.Font(DetailLayoutBuilder->GetDetailFont())
 					.Text_Lambda([this]() 
 					{
-						UEnum* const LoadingPolicyEnum = FindObjectChecked<UEnum>(ANY_PACKAGE, TEXT("ELocalizationTargetLoadingPolicy"));
+						UEnum* const LoadingPolicyEnum = FindObjectChecked<UEnum>(nullptr, TEXT("/Script/Localization.ELocalizationTargetLoadingPolicy"));
 						return LoadingPolicyEnum->GetDisplayNameTextByValue(static_cast<int64>(GetLoadingPolicy()));
 					})
 				]
@@ -700,7 +700,7 @@ void FLocalizationTargetDetailCustomization::OnLoadingPolicySelectionChanged(TSh
 
 TSharedRef<SWidget> FLocalizationTargetDetailCustomization::GenerateWidgetForLoadingPolicy(TSharedPtr<ELocalizationTargetLoadingPolicy> LoadingPolicy)
 {
-	UEnum* const LoadingPolicyEnum = FindObjectChecked<UEnum>(ANY_PACKAGE, TEXT("ELocalizationTargetLoadingPolicy"));
+	UEnum* const LoadingPolicyEnum = FindObjectChecked<UEnum>(nullptr, TEXT("/Script/Localization.ELocalizationTargetLoadingPolicy"));
 	return SNew(STextBlock)
 		.Font(DetailLayoutBuilder->GetDetailFont())
 		.Text(LoadingPolicyEnum->GetDisplayNameTextByValue(static_cast<int64>(*LoadingPolicy.Get())));
diff --git a/Engine/Source/Editor/MaterialEditor/Private/MaterialExpressionClasses.cpp b/Engine/Source/Editor/MaterialEditor/Private/MaterialExpressionClasses.cpp
--- a/Engine/Source/Editor/MaterialEditor/Private/MaterialExpressionClasses.cpp
+++ b/Engine/Source/Editor/MaterialEditor/Private/MaterialExpressionClasses.cpp
@@ -52,7 +52,7 @@ MaterialExpressionClasses* MaterialExpressionClasses::Get()
 const UStruct* MaterialExpressionClasses::GetExpressionInputStruct()
 {
 	static const UStruct* ExpressionInputStruct =
-		CastChecked<UStruct>( StaticFindObject( UStruct::StaticClass(), ANY_PACKAGE, TEXT("ExpressionInput")) );
+		CastChecked<UStruct>( StaticFindObject( UStruct::StaticClass(), nullptr, TEXT("/Script/Engine.ExpressionInput")) );
 	check( ExpressionInputStruct );
 	return ExpressionInputStruct;
 }
diff --git a/Engine/Source/Editor/MovieSceneTools/Public/K2Node_GetSequenceBinding.h b/Engine/Source/Editor/MovieSceneTools/Public/K2Node_GetSequenceBinding.h
--- a/Engine/Source/Editor/MovieSceneTools/Public/K2Node_GetSequenceBinding.h
+++ b/Engine/Source/Editor/MovieSceneTools/Public/K2Node_GetSequenceBinding.h
@@ -24,7 +24,7 @@ public:
 	GENERATED_BODY()
 
 	/** The sequence from which to choose a binding identifier */
-	UPROPERTY(EditAnywhere, Category="Sequence", meta=(AllowedClasses="MovieSceneSequence"))
+	UPROPERTY(EditAnywhere, Category="Sequence", meta=(AllowedClasses="/Script/MovieScene.MovieSceneSequence"))
 	FSoftObjectPath SourceSequence;
 
 	/** The user-selected literal binding identifier from the sequence to use */
diff --git a/Engine/Source/Editor/Sequencer/Private/SSequencerSection.cpp b/Engine/Source/Editor/Sequencer/Private/SSequencerSection.cpp
--- a/Engine/Source/Editor/Sequencer/Private/SSequencerSection.cpp
+++ b/Engine/Source/Editor/Sequencer/Private/SSequencerSection.cpp
@@ -284,7 +284,7 @@ struct FSequencerSectionPainterImpl : FSequencerSectionPainter
 
 		TSharedRef<FSlateFontCache> FontCache = FSlateApplication::Get().GetRenderer()->GetFontCache();
 
-		UEnum* Enum = FindObjectChecked<UEnum>(ANY_PACKAGE, TEXT("EMovieSceneBlendType"), true);
+		UEnum* Enum = FindObjectChecked<UEnum>(nullptr, TEXT("/Script/MovieScene.EMovieSceneBlendType"), true);
 		FText DisplayText = Enum->GetDisplayNameTextByValue((int64)Section.GetBlendType().Get());
 
 		FSlateFontInfo FontInfo = FAppStyle::GetFontStyle("Sequencer.Section.BackgroundText");
diff --git a/Engine/Source/Editor/Sequencer/Private/SequencerTrackFilters.h b/Engine/Source/Editor/Sequencer/Private/SequencerTrackFilters.h
--- a/Engine/Source/Editor/Sequencer/Private/SequencerTrackFilters.h
+++ b/Engine/Source/Editor/Sequencer/Private/SequencerTrackFilters.h
@@ -228,8 +228,8 @@ class FSequencerTrackFilter_AudioTracks : public FSequencerTrackFilter_ClassType
 			return false;
 		}
 
-		static UClass* LevelSequenceClass = FindObject<UClass>(ANY_PACKAGE, TEXT("LevelSequence"), true);
-		static UClass* WidgetAnimationClass = FindObject<UClass>(ANY_PACKAGE, TEXT("WidgetAnimation"), true);
+		static UClass* LevelSequenceClass = FindObject<UClass>(nullptr, TEXT("/Script/LevelSequence.LevelSequence"), true);
+		static UClass* WidgetAnimationClass = FindObject<UClass>(nullptr, TEXT("/Script/UMG.WidgetAnimation"), true);
 		return InSequence != nullptr &&
 			((LevelSequenceClass != nullptr && InSequence->GetClass()->IsChildOf(LevelSequenceClass)) ||
 			(WidgetAnimationClass != nullptr && InSequence->GetClass()->IsChildOf(WidgetAnimationClass)));
@@ -250,8 +250,8 @@ class FSequencerTrackFilter_EventTracks : public FSequencerTrackFilter_ClassType
 			return false;
 		}
 
-		static UClass* LevelSequenceClass = FindObject<UClass>(ANY_PACKAGE, TEXT("LevelSequence"), true);
-		static UClass* WidgetAnimationClass = FindObject<UClass>(ANY_PACKAGE, TEXT("WidgetAnimation"), true);
+		static UClass* LevelSequenceClass = FindObject<UClass>(nullptr, TEXT("/Script/LevelSequence.LevelSequence"), true);
+		static UClass* WidgetAnimationClass = FindObject<UClass>(nullptr, TEXT("/Script/UMG.WidgetAnimation"), true);
 		return InSequence != nullptr &&
 			((LevelSequenceClass != nullptr && InSequence->GetClass()->IsChildOf(LevelSequenceClass)) ||
 			(WidgetAnimationClass != nullptr && InSequence->GetClass()->IsChildOf(WidgetAnimationClass)));
diff --git a/Engine/Source/Editor/Sequencer/Private/SequencerUtilities.cpp b/Engine/Source/Editor/Sequencer/Private/SequencerUtilities.cpp
--- a/Engine/Source/Editor/Sequencer/Private/SequencerUtilities.cpp
+++ b/Engine/Source/Editor/Sequencer/Private/SequencerUtilities.cpp
@@ -263,7 +263,7 @@ void FSequencerUtilities::PopulateMenu_CreateNewSection(FMenuBuilder& MenuBuilde
 	FText NameOverride		= Track->GetSupportedBlendTypes().Num() == 1 ? LOCTEXT("AddSectionText", "Add New Section") : FText();
 	FText TooltipOverride	= Track->GetSupportedBlendTypes().Num() == 1 ? LOCTEXT("AddSectionToolTip", "Adds a new section at the current time") : FText();
 
-	const UEnum* MovieSceneBlendType = FindObjectChecked<UEnum>(ANY_PACKAGE, TEXT("EMovieSceneBlendType"));
+	const UEnum* MovieSceneBlendType = FindObjectChecked<UEnum>(nullptr, TEXT("/Script/MovieScene.EMovieSceneBlendType"));
 	for (EMovieSceneBlendType BlendType : Track->GetSupportedBlendTypes())
 	{
 		FText DisplayName = MovieSceneBlendType->GetDisplayNameTextByValue((int64)BlendType);
@@ -349,7 +349,7 @@ void FSequencerUtilities::PopulateMenu_SetBlendType(FMenuBuilder& MenuBuilder, c
 		}
 	};
 
-	const UEnum* MovieSceneBlendType = FindObjectChecked<UEnum>(ANY_PACKAGE, TEXT("EMovieSceneBlendType"));
+	const UEnum* MovieSceneBlendType = FindObjectChecked<UEnum>(nullptr, TEXT("/Script/MovieScene.EMovieSceneBlendType"));
 	for (int32 NameIndex = 0; NameIndex < MovieSceneBlendType->NumEnums() - 1; ++NameIndex)
 	{
 		EMovieSceneBlendType BlendType = (EMovieSceneBlendType)MovieSceneBlendType->GetValueByIndex(NameIndex);
diff --git a/Engine/Source/Editor/Sequencer/Public/SequencerTrackFilterBase.h b/Engine/Source/Editor/Sequencer/Public/SequencerTrackFilterBase.h
--- a/Engine/Source/Editor/Sequencer/Public/SequencerTrackFilterBase.h
+++ b/Engine/Source/Editor/Sequencer/Public/SequencerTrackFilterBase.h
@@ -34,7 +34,7 @@ public:
 	/** Returns whether the filter supports the sequence type */
 	virtual bool SupportsSequence(UMovieSceneSequence* InSequence) const
 	{
-		static UClass* LevelSequenceClass = FindObject<UClass>(ANY_PACKAGE, TEXT("LevelSequence"), true);
+		static UClass* LevelSequenceClass = FindObject<UClass>(nullptr, TEXT("/Script/LevelSequence.LevelSequence"), true);
 		return InSequence != nullptr && (LevelSequenceClass != nullptr && InSequence->GetClass()->IsChildOf(LevelSequenceClass));
 	}
 
diff --git a/Engine/Source/Editor/UMGEditor/Public/UMGEditorProjectSettings.h b/Engine/Source/Editor/UMGEditor/Public/UMGEditorProjectSettings.h
--- a/Engine/Source/Editor/UMGEditor/Public/UMGEditorProjectSettings.h
+++ b/Engine/Source/Editor/UMGEditor/Public/UMGEditorProjectSettings.h
@@ -159,7 +159,7 @@ public:
 	UPROPERTY(EditAnywhere, config, AdvancedDisplay, Category="Class Filtering")
 	TArray<FString> CategoriesToHide;
 
-	UPROPERTY(EditAnywhere, config, Category = "Class Filtering", meta = (MetaClass = "Widget"))
+	UPROPERTY(EditAnywhere, config, Category = "Class Filtering", meta = (MetaClass = "/Script/UMG.Widget"))
 	TArray<FSoftClassPath> WidgetClassesToHide;
 
 public:
diff --git a/Engine/Source/Editor/UnrealEd/Classes/Editor/EditorEngine.h b/Engine/Source/Editor/UnrealEd/Classes/Editor/EditorEngine.h
--- a/Engine/Source/Editor/UnrealEd/Classes/Editor/EditorEngine.h
+++ b/Engine/Source/Editor/UnrealEd/Classes/Editor/EditorEngine.h
@@ -3229,7 +3229,7 @@ protected:
 
 protected:
 
-	UPROPERTY(EditAnywhere, config, Category = Advanced, meta = (MetaClass = "ActorGroupingUtils"))
+	UPROPERTY(EditAnywhere, config, Category = Advanced, meta = (MetaClass = "/Script/UnrealEd.ActorGroupingUtils"))
 	FSoftClassPath ActorGroupingUtilsClassName;
 
 	UPROPERTY()
diff --git a/Engine/Source/Editor/UnrealEd/Classes/Editor/PropertyEditorTestObject.h b/Engine/Source/Editor/UnrealEd/Classes/Editor/PropertyEditorTestObject.h
--- a/Engine/Source/Editor/UnrealEd/Classes/Editor/PropertyEditorTestObject.h
+++ b/Engine/Source/Editor/UnrealEd/Classes/Editor/PropertyEditorTestObject.h
@@ -307,22 +307,22 @@ class UPropertyEditorTestObject : public UObject
 	UPROPERTY(EditAnywhere, Category=Classes)
 	TObjectPtr<UClass> ClassProperty;
 
-	UPROPERTY(EditAnywhere, Category=Classes, meta=(AllowedClasses="Texture2D"))
+	UPROPERTY(EditAnywhere, Category=Classes, meta=(AllowedClasses="/Script/Engine.Texture2D"))
 	TObjectPtr<UClass> ClassPropertyWithAllowed;
 
-	UPROPERTY(EditAnywhere, Category=Classes, meta=(DisallowedClasses="Texture2D"))
+	UPROPERTY(EditAnywhere, Category=Classes, meta=(DisallowedClasses="/Script/Engine.Texture2D"))
 	TObjectPtr<UClass> ClassPropertyWithDisallowed;
 
 	UPROPERTY(EditAnywhere, Category=Classes)
 	TSubclassOf<UTexture> SubclassOfTexture;
 
-	UPROPERTY(EditAnywhere, Category=Classes, meta=(AllowedClasses="Texture2D"))
+	UPROPERTY(EditAnywhere, Category=Classes, meta=(AllowedClasses="/Script/Engine.Texture2D"))
 	TSubclassOf<UTexture> SubclassOfWithAllowed;
 
-	UPROPERTY(EditAnywhere, Category=Classes, meta=(DisallowedClasses="Texture2D"))
+	UPROPERTY(EditAnywhere, Category=Classes, meta=(DisallowedClasses="/Script/Engine.Texture2D"))
 	TSubclassOf<UTexture> SubclassOfWithDisallowed;
 
-	UPROPERTY(EditAnywhere, Category=Classes, meta=(AllowedClasses="StaticMesh,  SkeletalMesh	"))
+	UPROPERTY(EditAnywhere, Category=Classes, meta=(AllowedClasses="/Script/Engine.StaticMesh,  /Script/Engine.SkeletalMesh	"))
 	TSoftObjectPtr<UObject> AssetPointerWithAllowedAndWhitespace;
 
 	// Integer
@@ -427,7 +427,7 @@ class UPropertyEditorTestObject : public UObject
 	UPROPERTY(VisibleAnywhere, Category=AdvancedProperties)
 	TObjectPtr<UPrimitiveComponent> ObjectThatCannotBeChanged;
 
-	UPROPERTY(EditAnywhere, Category=AdvancedProperties, meta=(Bitmask, BitmaskEnum="EPropertyEditorTestBitflags"))
+	UPROPERTY(EditAnywhere, Category=AdvancedProperties, meta=(Bitmask, BitmaskEnum="/Script/UnrealEd.EPropertyEditorTestBitflags"))
 	int32 EnumBitflags=0;
 
 	UPROPERTY(EditAnywhere, Category=AdvancedProperties, meta=(PasswordField=true))
@@ -463,26 +463,26 @@ class UPropertyEditorTestObject : public UObject
 	UPROPERTY(EditAnywhere, Category=Assets, meta=(DisplayThumbnail="true"))
 	FSoftObjectPath AssetReferenceCustomStructWithThumbnail;
 
-	UPROPERTY(EditAnywhere, Category=Assets, meta=(AllowedClasses="PointLight", ExactClass))
+	UPROPERTY(EditAnywhere, Category=Assets, meta=(AllowedClasses="/Script/Engine.PointLight", ExactClass))
 	FSoftObjectPath ExactlyPointLightActorReference;
 
-	UPROPERTY(EditAnywhere, Category=Assets, meta=(AllowedClasses="Light"))
+	UPROPERTY(EditAnywhere, Category=Assets, meta=(AllowedClasses="/Script/Engine.Light"))
 	FSoftObjectPath LightActorReference;
 
-	UPROPERTY(EditAnywhere, Category=Assets, meta=(AllowedClasses="PointLight, SpotLight", ExactClass=true))
+	UPROPERTY(EditAnywhere, Category=Assets, meta=(AllowedClasses="/Script/Engine.PointLight, /Script/Engine.SpotLight", ExactClass=true))
 	FSoftObjectPath ExactPointOrSpotLightActorReference;
 
 	// NOTE: intentionally misplaced space in AllowedClasses
-	UPROPERTY(EditAnywhere, Category=Assets, meta=(AllowedClasses="Light ,StaticMeshActor", DisplayThumbnail))
+	UPROPERTY(EditAnywhere, Category=Assets, meta=(AllowedClasses="/Script/Engine.Light ,/Script/Engine.StaticMeshActor", DisplayThumbnail))
 	FSoftObjectPath LightOrStaticMeshActorReference;
 
-	UPROPERTY(EditAnywhere, Category=Assets, meta=(AllowedClasses="Actor", DisallowedClasses="Light"))
+	UPROPERTY(EditAnywhere, Category=Assets, meta=(AllowedClasses="/Script/Engine.Actor", DisallowedClasses="/Script/Engine.Light"))
 	FSoftObjectPath NotLightActorReference;
 
-	UPROPERTY(EditAnywhere, Category=Assets, meta=(AllowedClasses="Material,Texture"))
+	UPROPERTY(EditAnywhere, Category=Assets, meta=(AllowedClasses="/Script/Engine.Material,/Script/Engine.Texture"))
 	FSoftObjectPath MaterialOrTextureAssetReference;
 
-	UPROPERTY(EditAnywhere, Category=Assets, meta=(MetaClass="Actor"))
+	UPROPERTY(EditAnywhere, Category=Assets, meta=(MetaClass="/Script/Engine.Actor"))
 	FSoftObjectPath ActorWithMetaClass;
 
 	UPROPERTY(EditAnywhere, Category=Assets)
@@ -626,7 +626,7 @@ class UPropertyEditorTestObject : public UObject
 
 	// Allows either an object that's derived from UTexture or IBlendableInterface, to ensure that Object Property handles know how to
 	// filter for AllowedClasses correctly.
-	UPROPERTY(EditAnywhere, Category=ObjectPropertyAllowedClasses, meta=(AllowedClasses="Texture,BlendableInterface"))
+	UPROPERTY(EditAnywhere, Category=ObjectPropertyAllowedClasses, meta=(AllowedClasses="/Script/Engine.Texture,/Script/Engine.BlendableInterface"))
 	TObjectPtr<UObject> TextureOrBlendableInterface;
 
 	UPROPERTY(EditAnywhere, Category="Subcategory")
diff --git a/Engine/Source/Editor/UnrealEd/Classes/Factories/FbxTextureImportData.h b/Engine/Source/Editor/UnrealEd/Classes/Factories/FbxTextureImportData.h
--- a/Engine/Source/Editor/UnrealEd/Classes/Factories/FbxTextureImportData.h
+++ b/Engine/Source/Editor/UnrealEd/Classes/Factories/FbxTextureImportData.h
@@ -26,7 +26,7 @@ class UNREALED_API UFbxTextureImportData : public UFbxAssetImportData
 	EMaterialSearchLocation MaterialSearchLocation;
 
 	/** Base material to instance from when importing materials. */
-	UPROPERTY(EditAnywhere, BlueprintReadWrite, config, Category = Material, meta = (ImportType = "Mesh", AllowedClasses = "MaterialInterface"))
+	UPROPERTY(EditAnywhere, BlueprintReadWrite, config, Category = Material, meta = (ImportType = "Mesh", AllowedClasses = "/Script/Engine.MaterialInterface"))
 	FSoftObjectPath BaseMaterialName;
 
 	/** transient, ImportUI customize helper to store if we must show or not the BaseMaterialName property. */
diff --git a/Engine/Source/Editor/UnrealEd/Classes/Settings/LevelEditorViewportSettings.h b/Engine/Source/Editor/UnrealEd/Classes/Settings/LevelEditorViewportSettings.h
--- a/Engine/Source/Editor/UnrealEd/Classes/Settings/LevelEditorViewportSettings.h
+++ b/Engine/Source/Editor/UnrealEd/Classes/Settings/LevelEditorViewportSettings.h
@@ -505,7 +505,7 @@ public:
 	float BackgroundDropDistance;
 
 	/** A list of meshes that can be used as preview mesh in the editor view port by holding down the backslash key */
-	UPROPERTY(EditAnywhere, config, Category=Preview, meta=(AllowedClasses = "StaticMesh"))
+	UPROPERTY(EditAnywhere, config, Category=Preview, meta=(AllowedClasses = "/Script/Engine.StaticMesh"))
 	TArray<FSoftObjectPath> PreviewMeshes;
 
 	UPROPERTY(EditAnywhere, config, AdvancedDisplay, Category=LookAndFeel, meta=(ClampMin = "0.01", UIMin = "0.01", UIMax = "5"))
diff --git a/Engine/Source/Editor/UnrealEd/Classes/ThumbnailRendering/SceneThumbnailInfoWithPrimitive.h b/Engine/Source/Editor/UnrealEd/Classes/ThumbnailRendering/SceneThumbnailInfoWithPrimitive.h
--- a/Engine/Source/Editor/UnrealEd/Classes/ThumbnailRendering/SceneThumbnailInfoWithPrimitive.h
+++ b/Engine/Source/Editor/UnrealEd/Classes/ThumbnailRendering/SceneThumbnailInfoWithPrimitive.h
@@ -25,7 +25,7 @@ class USceneThumbnailInfoWithPrimitive : public USceneThumbnailInfo
 	TEnumAsByte<EThumbnailPrimType> PrimitiveType;
 
 	/** The custom mesh used when the primitive type is TPT_None */
-	UPROPERTY(EditAnywhere, Category=Thumbnail, meta = (AllowedClasses = "StaticMesh,SkeletalMesh", ExactClass = "true"))
+	UPROPERTY(EditAnywhere, Category=Thumbnail, meta = (AllowedClasses = "/Script/Engine.StaticMesh,/Script/Engine.SkeletalMesh", ExactClass = "true"))
 	FSoftObjectPath PreviewMesh;
 
 	UPROPERTY()
diff --git a/Engine/Source/Editor/UnrealEd/Private/Commandlets/InternationalizationExportCommandlet.cpp b/Engine/Source/Editor/UnrealEd/Private/Commandlets/InternationalizationExportCommandlet.cpp
--- a/Engine/Source/Editor/UnrealEd/Private/Commandlets/InternationalizationExportCommandlet.cpp
+++ b/Engine/Source/Editor/UnrealEd/Private/Commandlets/InternationalizationExportCommandlet.cpp
@@ -111,7 +111,7 @@ int32 UInternationalizationExportCommandlet::Main(const FString& Params)
 		FString TextCollapseModeName;
 		if (GetStringFromConfig(*SectionName, TEXT("LocalizedTextCollapseMode"), TextCollapseModeName, ConfigPath))
 		{
-			UEnum* LocalizedTextCollapseModeEnum = FindObjectChecked<UEnum>(ANY_PACKAGE, TEXT("ELocalizedTextCollapseMode"));
+			UEnum* LocalizedTextCollapseModeEnum = FindObjectChecked<UEnum>(nullptr, TEXT("/Script/Localization.ELocalizedTextCollapseMode"));
 			const int64 TextCollapseModeInt = LocalizedTextCollapseModeEnum->GetValueByName(*TextCollapseModeName);
 			if (TextCollapseModeInt != INDEX_NONE)
 			{
@@ -126,7 +126,7 @@ int32 UInternationalizationExportCommandlet::Main(const FString& Params)
 		FString POFormatName;
 		if (GetStringFromConfig(*SectionName, TEXT("POFormat"), POFormatName, ConfigPath))
 		{
-			UEnum* POFormatEnum = FindObjectChecked<UEnum>(ANY_PACKAGE, TEXT("EPortableObjectFormat"));
+			UEnum* POFormatEnum = FindObjectChecked<UEnum>(nullptr, TEXT("/Script/Localization.EPortableObjectFormat"));
 			const int64 POFormatInt = POFormatEnum->GetValueByName(*POFormatName);
 			if (POFormatInt != INDEX_NONE)
 			{
diff --git a/Engine/Source/Editor/UnrealEd/Private/Commandlets/WorldPartitionConvertCommandlet.cpp b/Engine/Source/Editor/UnrealEd/Private/Commandlets/WorldPartitionConvertCommandlet.cpp
--- a/Engine/Source/Editor/UnrealEd/Private/Commandlets/WorldPartitionConvertCommandlet.cpp
+++ b/Engine/Source/Editor/UnrealEd/Private/Commandlets/WorldPartitionConvertCommandlet.cpp
@@ -761,8 +761,8 @@ int32 UWorldPartitionConvertCommandlet::Main(const FString& Params)
 		}
 		else
 		{
-			EditorHashClass = FindObject<UClass>(ANY_PACKAGE, TEXT("WorldPartitionEditorSpatialHash"));
-			RuntimeHashClass = FindObject<UClass>(ANY_PACKAGE, TEXT("WorldPartitionRuntimeSpatialHash"));
+			EditorHashClass = FindObject<UClass>(nullptr, TEXT("/Script/Engine.WorldPartitionEditorSpatialHash"));
+			RuntimeHashClass = FindObject<UClass>(nullptr, TEXT("/Script/Engine.WorldPartitionRuntimeSpatialHash"));
 		}
 	}
 
diff --git a/Engine/Source/Editor/UnrealEd/Private/CookOnTheFlyServer.cpp b/Engine/Source/Editor/UnrealEd/Private/CookOnTheFlyServer.cpp
--- a/Engine/Source/Editor/UnrealEd/Private/CookOnTheFlyServer.cpp
+++ b/Engine/Source/Editor/UnrealEd/Private/CookOnTheFlyServer.cpp
@@ -1336,7 +1336,7 @@ void UCookOnTheFlyServer::SetSaveBusy(bool bInBusy)
 			TSet<UPackage*> ExpectedPackages;
 			FPackageDataQueue& SaveQueue = PackageDatas->GetSaveQueue();
 			const TArray<FPendingCookedPlatformData>& PendingCookedPlatformDatas = PackageDatas->GetPendingCookedPlatformDatas();
-			TArray<UClass*> CompilationUsers({ UMaterialInterface::StaticClass(), FindObject<UClass>(ANY_PACKAGE, TEXT("NiagaraScript")) });
+			TArray<UClass*> CompilationUsers({ UMaterialInterface::StaticClass(), FindObject<UClass>(nullptr, TEXT("/Script/Niagara.NiagaraScript")) });
 
 			for (const FPendingCookedPlatformData& Data : PendingCookedPlatformDatas)
 			{
diff --git a/Engine/Source/Editor/UnrealEd/Private/EditorBuildUtils.cpp b/Engine/Source/Editor/UnrealEd/Private/EditorBuildUtils.cpp
--- a/Engine/Source/Editor/UnrealEd/Private/EditorBuildUtils.cpp
+++ b/Engine/Source/Editor/UnrealEd/Private/EditorBuildUtils.cpp
@@ -1280,7 +1280,7 @@ void FEditorBuildUtils::TriggerHierarchicalLODBuilder(UWorld* InWorld, FName Id)
 	if (InWorld->IsPartitionedWorld())
 	{
 		IWorldPartitionEditorModule& WorldPartitionEditorModule = FModuleManager::LoadModuleChecked<IWorldPartitionEditorModule>("WorldPartitionEditor");
-		TSubclassOf<UWorldPartitionBuilder> WorldPartitionHLODsBuilder = FindObjectChecked<UClass>(ANY_PACKAGE, TEXT("WorldPartitionHLODsBuilder"), true);
+		TSubclassOf<UWorldPartitionBuilder> WorldPartitionHLODsBuilder = FindObjectChecked<UClass>(nullptr, TEXT("/Script/UnrealEd.WorldPartitionHLODsBuilder"), true);
 		WorldPartitionEditorModule.RunBuilder(WorldPartitionHLODsBuilder, InWorld->GetPackage()->GetName());
 	}
 	else
@@ -1295,7 +1295,7 @@ void FEditorBuildUtils::TriggerMinimapBuilder(UWorld* InWorld, FName Id)
 	if (InWorld->IsPartitionedWorld())
 	{
 		IWorldPartitionEditorModule& WorldPartitionEditorModule = FModuleManager::LoadModuleChecked<IWorldPartitionEditorModule>("WorldPartitionEditor");
-		TSubclassOf<UWorldPartitionBuilder> WorldPartitionMiniMapBuilder = FindObjectChecked<UClass>(ANY_PACKAGE, TEXT("WorldPartitionMiniMapBuilder"), true);
+		TSubclassOf<UWorldPartitionBuilder> WorldPartitionMiniMapBuilder = FindObjectChecked<UClass>(nullptr, TEXT("/Script/UnrealEd.WorldPartitionMiniMapBuilder"), true);
 		WorldPartitionEditorModule.RunBuilder(WorldPartitionMiniMapBuilder, InWorld->GetPackage()->GetName());
 	}
 }
diff --git a/Engine/Source/Editor/UnrealEd/Private/Kismet2/DebuggerCommands.cpp b/Engine/Source/Editor/UnrealEd/Private/Kismet2/DebuggerCommands.cpp
--- a/Engine/Source/Editor/UnrealEd/Private/Kismet2/DebuggerCommands.cpp
+++ b/Engine/Source/Editor/UnrealEd/Private/Kismet2/DebuggerCommands.cpp
@@ -880,7 +880,7 @@ TSharedRef< SWidget > FPlayWorldCommands::GeneratePlayMenuContent(TSharedRef<FUI
 					LOCTEXT("NetworkModeToolTip", "Which network mode should the clients launch in? A server will automatically be started if needed."),
 					FNewMenuDelegate::CreateLambda([](FMenuBuilder& InMenuBuilder)
 						{
-							const UEnum* PlayNetModeEnum = FindObject<UEnum>(ANY_PACKAGE, TEXT("EPlayNetMode"));
+							const UEnum* PlayNetModeEnum = FindObject<UEnum>(nullptr, TEXT("/Script/UnrealEd.EPlayNetMode"));
 
 							for (int32 i = 0; i < PlayNetModeEnum->NumEnums() - 1; i++)
 							{
diff --git a/Engine/Source/Editor/UnrealEd/Private/Kismet2/Kismet2.cpp b/Engine/Source/Editor/UnrealEd/Private/Kismet2/Kismet2.cpp
--- a/Engine/Source/Editor/UnrealEd/Private/Kismet2/Kismet2.cpp
+++ b/Engine/Source/Editor/UnrealEd/Private/Kismet2/Kismet2.cpp
@@ -514,8 +514,8 @@ UBlueprint* FKismetEditorUtilities::CreateBlueprint(UClass* ParentClass, UObject
 		{
 			// Based on the Blueprint type we are constructing, place some starting events.
 			// Note, this cannot happen in the Factories for constructing these Blueprint types due to the fact that creating child BPs circumvent the factories
-			UClass* WidgetClass = FindObject<UClass>(ANY_PACKAGE, TEXT("UserWidget"));
-			UClass* GameplayAbilityClass = FindObject<UClass>(ANY_PACKAGE, TEXT("GameplayAbility"));
+			UClass* WidgetClass = FindObject<UClass>(nullptr, TEXT("/Script/UMG.UserWidget"));
+			UClass* GameplayAbilityClass = FindObject<UClass>(nullptr, TEXT("/Script/GameplayAbilities.GameplayAbility"));
 
 			TArray<FName> AutoSpawnedEventNames;
 			int32 NodePositionY = 0;
diff --git a/Engine/Source/Runtime/AIModule/Classes/AISystem.h b/Engine/Source/Runtime/AIModule/Classes/AISystem.h
--- a/Engine/Source/Runtime/AIModule/Classes/AISystem.h
+++ b/Engine/Source/Runtime/AIModule/Classes/AISystem.h
@@ -29,15 +29,15 @@ class AIMODULE_API UAISystem : public UAISystemBase
 
 protected:
 	/** Class that will be used to spawn the perception system, can be game-specific */
-	UPROPERTY(globalconfig, EditAnywhere, Category = "AISystem", meta = (MetaClass = "AIPerceptionSystem", DisplayName = "Perception System Class"))
+	UPROPERTY(globalconfig, EditAnywhere, Category = "AISystem", meta = (MetaClass = "/Script/AIModule.AIPerceptionSystem", DisplayName = "Perception System Class"))
 	FSoftClassPath PerceptionSystemClassName;
 
 	/** Class that will be used to spawn the hot spot manager, can be game-specific */
-	UPROPERTY(globalconfig, EditAnywhere, Category = "AISystem", meta = (MetaClass = "AIHotSpotManager", DisplayName = "AIHotSpotManager Class"))
+	UPROPERTY(globalconfig, EditAnywhere, Category = "AISystem", meta = (MetaClass = "/Script/AIModule.AIHotSpotManager", DisplayName = "AIHotSpotManager Class"))
 	FSoftClassPath HotSpotManagerClassName;
 
 	/** Class that will be used to spawn the env query manager, can be game-specific */
-	UPROPERTY(globalconfig, EditAnywhere, Category = "AISystem", meta = (MetaClass = "EnvQueryManager", DisplayName = "EnvQueryManager Class"))
+	UPROPERTY(globalconfig, EditAnywhere, Category = "AISystem", meta = (MetaClass = "/Script/AIModule.EnvQueryManager", DisplayName = "EnvQueryManager Class"))
 	FSoftClassPath EnvQueryManagerClassName;
 public:
 	/** Default AI movement's acceptance radius used to determine whether 
diff --git a/Engine/Source/Runtime/AnimGraphRuntime/Private/BoneControllers/BoneControllerTypes.cpp b/Engine/Source/Runtime/AnimGraphRuntime/Private/BoneControllers/BoneControllerTypes.cpp
--- a/Engine/Source/Runtime/AnimGraphRuntime/Private/BoneControllers/BoneControllerTypes.cpp
+++ b/Engine/Source/Runtime/AnimGraphRuntime/Private/BoneControllers/BoneControllerTypes.cpp
@@ -21,7 +21,7 @@ FVector FWarpingVectorValue::AsComponentSpaceDirection(const FAnimInstanceProxy*
 		return IKFootRootTransform.TransformVectorNoScale(Value.GetSafeNormal());
 	}
 #if WITH_EDITORONLY_DATA
-	const UEnum* TypeEnum = FindObject<UEnum>(ANY_PACKAGE, TEXT("EWarpingVectorMode"));
+	const UEnum* TypeEnum = FindObject<UEnum>(nullptr, TEXT("/Script/AnimGraphRuntime.EWarpingVectorMode"));
 	check(!!TypeEnum);
 	ensureMsgf(false, TEXT("AsComponentSpaceDirection specified an unhandled EWarpingVectorMode: %s"), *(TypeEnum->GetNameStringByIndex(static_cast<int32>(Mode))));
 #endif
diff --git a/Engine/Source/Runtime/AutomationTest/Public/AutomationTestExcludelist.h b/Engine/Source/Runtime/AutomationTest/Public/AutomationTestExcludelist.h
--- a/Engine/Source/Runtime/AutomationTest/Public/AutomationTestExcludelist.h
+++ b/Engine/Source/Runtime/AutomationTest/Public/AutomationTestExcludelist.h
@@ -48,7 +48,7 @@ struct FAutomationTestExcludeOptions
 	FName Reason;
 
 	/* Option to target specific RHI. No option means it should be applied to all RHI */
-	UPROPERTY(EditAnywhere, Meta = (Bitmask, BitmaskEnum = "ERHI_Flags"), Category = ExcludeTestOptions)
+	UPROPERTY(EditAnywhere, Meta = (Bitmask, BitmaskEnum = "/Script/AutomationTest.ERHI_Flags"), Category = ExcludeTestOptions)
 	int8 RHIs = 0;
 
 	/* Should the Reason be reported as a warning in the log */
diff --git a/Engine/Source/Runtime/CoreUObject/Public/UObject/NoExportTypes.h b/Engine/Source/Runtime/CoreUObject/Public/UObject/NoExportTypes.h
--- a/Engine/Source/Runtime/CoreUObject/Public/UObject/NoExportTypes.h
+++ b/Engine/Source/Runtime/CoreUObject/Public/UObject/NoExportTypes.h
@@ -505,7 +505,7 @@ struct FVector3d
  * A point or direction FVector in 3d space.
  * @note The full C++ class is located here: Engine\Source\Runtime\Core\Public\Math\Vector.h
  */
-USTRUCT(immutable, noexport, BlueprintType, IsAlwaysAccessible, HasDefaults, HasNoOpConstructor, IsCoreType, meta = (HasNativeMake = "Engine.KismetMathLibrary.MakeVector", HasNativeBreak = "Engine.KismetMathLibrary.BreakVector"))
+USTRUCT(immutable, noexport, BlueprintType, IsAlwaysAccessible, HasDefaults, HasNoOpConstructor, IsCoreType, meta = (HasNativeMake = "/Script/Engine.KismetMathLibrary.MakeVector", HasNativeBreak = "/Script/Engine.KismetMathLibrary.BreakVector"))
 struct FVector
 {
 	UPROPERTY(EditAnywhere, BlueprintReadWrite, Category = Vector, SaveGame)
@@ -564,7 +564,7 @@ struct FVector4d
 * A 4-D homogeneous vector.
 * @note The full C++ class is located here: Engine\Source\Runtime\Core\Public\Math\Vector4.h
 */
-USTRUCT(immutable, noexport, BlueprintType, IsAlwaysAccessible, HasDefaults, HasNoOpConstructor, IsCoreType, meta = (HasNativeMake = "Engine.KismetMathLibrary.MakeVector4", HasNativeBreak = "Engine.KismetMathLibrary.BreakVector4"))
+USTRUCT(immutable, noexport, BlueprintType, IsAlwaysAccessible, HasDefaults, HasNoOpConstructor, IsCoreType, meta = (HasNativeMake = "/Script/Engine.KismetMathLibrary.MakeVector4", HasNativeBreak = "/Script/Engine.KismetMathLibrary.BreakVector4"))
 struct FVector4
 {
 	UPROPERTY(EditAnywhere, BlueprintReadWrite, Category=Vector4, SaveGame)
@@ -614,7 +614,7 @@ struct FVector2f
  * A vector in 2-D space composed of components (X, Y) with floating point precision.
  * @note The full C++ class is located here: Engine\Source\Runtime\Core\Public\Math\Vector2D.h
  */
-USTRUCT(immutable, noexport, BlueprintType, IsAlwaysAccessible, HasDefaults, HasNoOpConstructor, IsCoreType, meta=(HasNativeMake="Engine.KismetMathLibrary.MakeVector2D", HasNativeBreak="Engine.KismetMathLibrary.BreakVector2D"))
+USTRUCT(immutable, noexport, BlueprintType, IsAlwaysAccessible, HasDefaults, HasNoOpConstructor, IsCoreType, meta=(HasNativeMake="/Script/Engine.KismetMathLibrary.MakeVector2D", HasNativeBreak="/Script/Engine.KismetMathLibrary.BreakVector2D"))
 struct FVector2D
 {
 	UPROPERTY(EditAnywhere, BlueprintReadWrite, Category=Vector2D, SaveGame)
@@ -713,7 +713,7 @@ struct FRotator3d
  * An orthogonal rotation in 3d space.
  * @note The full C++ class is located here: Engine\Source\Runtime\Core\Public\Math\Rotator.h
  */
-USTRUCT(immutable, noexport, BlueprintType, IsAlwaysAccessible, HasDefaults, HasNoOpConstructor, IsCoreType, meta=(HasNativeMake="Engine.KismetMathLibrary.MakeRotator", HasNativeBreak="Engine.KismetMathLibrary.BreakRotator"))
+USTRUCT(immutable, noexport, BlueprintType, IsAlwaysAccessible, HasDefaults, HasNoOpConstructor, IsCoreType, meta=(HasNativeMake="/Script/Engine.KismetMathLibrary.MakeRotator", HasNativeBreak="/Script/Engine.KismetMathLibrary.BreakRotator"))
 struct FRotator
 {
 	/** Pitch (degrees) around Y axis */
@@ -777,7 +777,7 @@ struct FQuat4d
  * Quaternion.
  * @note The full C++ class is located here: Engine\Source\Runtime\Core\Public\Math\Quat.h
  */
-USTRUCT(immutable, noexport, BlueprintType, IsAlwaysAccessible, HasDefaults, HasNoOpConstructor, IsCoreType, meta=(HasNativeMake ="Engine.KismetMathLibrary.MakeQuat", HasNativeBreak="Engine.KismetMathLibrary.BreakQuat"))
+USTRUCT(immutable, noexport, BlueprintType, IsAlwaysAccessible, HasDefaults, HasNoOpConstructor, IsCoreType, meta=(HasNativeMake ="/Script/Engine.KismetMathLibrary.MakeQuat", HasNativeBreak="/Script/Engine.KismetMathLibrary.BreakQuat"))
 struct FQuat
 {
 	UPROPERTY(EditAnywhere, BlueprintReadWrite, Category = Quat, SaveGame)
@@ -973,7 +973,7 @@ struct FBox3d
  * A bounding box.
  * @note The full C++ class is located here: Engine\Source\Runtime\Core\Public\Math\Box.h
  */
-USTRUCT(immutable, noexport, BlueprintType, IsAlwaysAccessible, HasDefaults, HasNoOpConstructor, IsCoreType, meta=(HasNativeMake="Engine.KismetMathLibrary.MakeBox"))
+USTRUCT(immutable, noexport, BlueprintType, IsAlwaysAccessible, HasDefaults, HasNoOpConstructor, IsCoreType, meta=(HasNativeMake="/Script/Engine.KismetMathLibrary.MakeBox"))
 struct FBox
 {
 	UPROPERTY(EditAnywhere, BlueprintReadWrite, Category=Box, SaveGame)
@@ -1025,7 +1025,7 @@ struct FBox2f
 * A rectangular 2D Box.
 * @note The full C++ class is located here: Engine\Source\Runtime\Core\Public\Math\Box2D.h
 */
-USTRUCT(immutable, noexport, BlueprintType, IsAlwaysAccessible, HasDefaults, HasNoOpConstructor, IsCoreType, meta=(HasNativeMake="Engine.KismetMathLibrary.MakeBox2D"))
+USTRUCT(immutable, noexport, BlueprintType, IsAlwaysAccessible, HasDefaults, HasNoOpConstructor, IsCoreType, meta=(HasNativeMake="/Script/Engine.KismetMathLibrary.MakeBox2D"))
 struct FBox2D
 {
 	UPROPERTY(EditAnywhere, BlueprintReadWrite, Category=Box2D, SaveGame)
@@ -1082,7 +1082,7 @@ struct FBoxSphereBounds3d
  * A bounding box and bounding sphere with the same origin.
  * @note The full C++ class is located here : Engine\Source\Runtime\Core\Public\Math\BoxSphereBounds.h
  */
-USTRUCT(noexport, BlueprintType, IsAlwaysAccessible, HasDefaults, HasNoOpConstructor, IsCoreType, meta = (HasNativeMake = "Engine.KismetMathLibrary.MakeBoxSphereBounds", HasNativeBreak = "Engine.KismetMathLibrary.BreakBoxSphereBounds"))
+USTRUCT(noexport, BlueprintType, IsAlwaysAccessible, HasDefaults, HasNoOpConstructor, IsCoreType, meta = (HasNativeMake = "/Script/Engine.KismetMathLibrary.MakeBoxSphereBounds", HasNativeBreak = "/Script/Engine.KismetMathLibrary.BreakBoxSphereBounds"))
 struct FBoxSphereBounds
 {
 	/** Holds the origin of the bounding box and sphere. */
@@ -1530,7 +1530,7 @@ struct FTransform3d
  * Transform composed of Quat/Translation/Scale.
  * @note This is implemented in either TransformVectorized.h or TransformNonVectorized.h depending on the platform.
  */
-USTRUCT(noexport, BlueprintType, IsAlwaysAccessible, HasDefaults, IsCoreType, meta=(HasNativeMake="Engine.KismetMathLibrary.MakeTransform", HasNativeBreak="Engine.KismetMathLibrary.BreakTransform"))
+USTRUCT(noexport, BlueprintType, IsAlwaysAccessible, HasDefaults, IsCoreType, meta=(HasNativeMake="/Script/Engine.KismetMathLibrary.MakeTransform", HasNativeBreak="/Script/Engine.KismetMathLibrary.BreakTransform"))
 struct FTransform
 {
 	/** Rotation of this transformation, as a quaternion. */
@@ -1550,7 +1550,7 @@ struct FTransform
  * Thread-safe random number generator that can be manually seeded.
  * @note The full C++ class is located here: Engine\Source\Runtime\Core\Public\Math\RandomStream.h
  */
-USTRUCT(noexport, BlueprintType, IsAlwaysAccessible, HasDefaults, HasNoOpConstructor, meta = (HasNativeMake = "Engine.KismetMathLibrary.MakeRandomStream", HasNativeBreak = "Engine.KismetMathLibrary.BreakRandomStream"))
+USTRUCT(noexport, BlueprintType, IsAlwaysAccessible, HasDefaults, HasNoOpConstructor, meta = (HasNativeMake = "/Script/Engine.KismetMathLibrary.MakeRandomStream", HasNativeBreak = "/Script/Engine.KismetMathLibrary.BreakRandomStream"))
 struct FRandomStream
 {
 public:
@@ -1567,7 +1567,7 @@ public:
  * A value representing a specific point date and time over a wide range of years.
  * @note The full C++ class is located here: Engine\Source\Runtime\Core\Public\Misc\DateTime.h
  */
-USTRUCT(immutable, noexport, BlueprintType, IsAlwaysAccessible, HasDefaults, meta=(HasNativeMake="Engine.KismetMathLibrary.MakeDateTime", HasNativeBreak="Engine.KismetMathLibrary.BreakDateTime"))
+USTRUCT(immutable, noexport, BlueprintType, IsAlwaysAccessible, HasDefaults, meta=(HasNativeMake="/Script/Engine.KismetMathLibrary.MakeDateTime", HasNativeBreak="/Script/Engine.KismetMathLibrary.BreakDateTime"))
 struct FDateTime
 {
 	int64 Ticks;
@@ -1588,7 +1588,7 @@ struct FFrameNumber
  * A frame rate represented as a fraction comprising 2 integers: a numerator (number of frames), and a denominator (per second).
  * @note The full C++ class is located here: Engine\Source\Runtime\Core\Public\Misc\FrameRate.h
  */
-USTRUCT(noexport, BlueprintType, meta=(HasNativeMake="Engine.KismetMathLibrary.MakeFrameRate", HasNativeBreak="Engine.KismetMathLibrary.BreakFrameRate"))
+USTRUCT(noexport, BlueprintType, meta=(HasNativeMake="/Script/Engine.KismetMathLibrary.MakeFrameRate", HasNativeBreak="/Script/Engine.KismetMathLibrary.BreakFrameRate"))
 struct FFrameRate
 {
 	/** The numerator of the framerate represented as a number of frames per second (e.g. 60 for 60 fps) */
@@ -1622,7 +1622,7 @@ private:
  * A frame time qualified by a frame rate context.
  * @note The full C++ class is located here: Engine\Source\Runtime\Core\Public\Misc\QualifiedFrameTime.h
  */
-USTRUCT(noexport, BlueprintType, meta=(ScriptName="QualifiedTime", HasNativeMake="Engine.KismetMathLibrary.MakeQualifiedFrameTime", HasNativeBreak="Engine.KismetMathLibrary.BreakQualifiedFrameTime"))
+USTRUCT(noexport, BlueprintType, meta=(ScriptName="QualifiedTime", HasNativeMake="/Script/Engine.KismetMathLibrary.MakeQualifiedFrameTime", HasNativeBreak="/Script/Engine.KismetMathLibrary.BreakQualifiedFrameTime"))
 struct FQualifiedFrameTime
 {
 	/** The frame time */
@@ -1662,7 +1662,7 @@ struct FTimecode
  * A time span value, which is the difference between two dates and times.
  * @note The full C++ class is located here: Engine\Source\Runtime\Core\Public\Misc\Timespan.h
  */
-USTRUCT(immutable, noexport, BlueprintType, IsAlwaysAccessible, HasDefaults, meta=(HasNativeMake="Engine.KismetMathLibrary.MakeTimespan", HasNativeBreak="Engine.KismetMathLibrary.BreakTimespan"))
+USTRUCT(immutable, noexport, BlueprintType, IsAlwaysAccessible, HasDefaults, meta=(HasNativeMake="/Script/Engine.KismetMathLibrary.MakeTimespan", HasNativeBreak="/Script/Engine.KismetMathLibrary.BreakTimespan"))
 struct FTimespan
 {
 	int64 Ticks;
@@ -1672,7 +1672,7 @@ struct FTimespan
  * A struct that contains a string reference to an object, either a top level asset or a subobject.
  * @note The full C++ class is located here: Engine\Source\Runtime\CoreUObject\Public\UObject\SoftObjectPath.h
  */
-USTRUCT(noexport, BlueprintType, IsAlwaysAccessible, HasDefaults, meta=(HasNativeMake="Engine.KismetSystemLibrary.MakeSoftObjectPath", HasNativeBreak="Engine.KismetSystemLibrary.BreakSoftObjectPath"))
+USTRUCT(noexport, BlueprintType, IsAlwaysAccessible, HasDefaults, meta=(HasNativeMake="/Script/Engine.KismetSystemLibrary.MakeSoftObjectPath", HasNativeBreak="/Script/Engine.KismetSystemLibrary.BreakSoftObjectPath"))
 struct FSoftObjectPath
 {
 	/** Asset path, patch to a top level object in a package */
@@ -1688,7 +1688,7 @@ struct FSoftObjectPath
  * A struct that contains a string reference to a class, can be used to make soft references to classes.
  * @note The full C++ class is located here: Engine\Source\Runtime\CoreUObject\Public\UObject\SoftObjectPath.h
  */
-USTRUCT(noexport, BlueprintType, IsAlwaysAccessible, HasDefaults, meta=(HasNativeMake="Engine.KismetSystemLibrary.MakeSoftClassPath", HasNativeBreak="Engine.KismetSystemLibrary.BreakSoftClassPath"))
+USTRUCT(noexport, BlueprintType, IsAlwaysAccessible, HasDefaults, meta=(HasNativeMake="/Script/Engine.KismetSystemLibrary.MakeSoftClassPath", HasNativeBreak="/Script/Engine.KismetSystemLibrary.BreakSoftClassPath"))
 struct FSoftClassPath : public FSoftObjectPath
 {
 };
diff --git a/Engine/Source/Runtime/Engine/Classes/AI/AISystemBase.h b/Engine/Source/Runtime/Engine/Classes/AI/AISystemBase.h
--- a/Engine/Source/Runtime/Engine/Classes/AI/AISystemBase.h
+++ b/Engine/Source/Runtime/Engine/Classes/AI/AISystemBase.h
@@ -53,11 +53,11 @@ class ENGINE_API UAISystemBase : public UObject
 
 private:
 	/** List of specific AI system class to create, can be game-specific */
-	UPROPERTY(globalconfig, EditAnywhere, Category = "AISystem", noclear, meta = (MetaClass = "AISystem", DisplayName = "AISystem Class"))
+	UPROPERTY(globalconfig, EditAnywhere, Category = "AISystem", noclear, meta = (MetaClass = "/Script/AIModule.AISystem", DisplayName = "AISystem Class"))
 	FSoftClassPath AISystemClassName;
 
 	/** Name of a module used to spawn the AI system. If not empty, this module has to implement IAISystemModule */
-	UPROPERTY(globalconfig, EditAnywhere, Category = "AISystem", noclear, meta = (MetaClass = "AISystem", DisplayName = "AISystem Module"))
+	UPROPERTY(globalconfig, EditAnywhere, Category = "AISystem", noclear, meta = (MetaClass = "/Script/AIModule.AISystem", DisplayName = "AISystem Module"))
 	FName AISystemModuleName;
 
 	FDelegateHandle OnMatchStateSetHandle;
diff --git a/Engine/Source/Runtime/Engine/Classes/AI/Navigation/NavigationTypes.h b/Engine/Source/Runtime/Engine/Classes/AI/Navigation/NavigationTypes.h
--- a/Engine/Source/Runtime/Engine/Classes/AI/Navigation/NavigationTypes.h
+++ b/Engine/Source/Runtime/Engine/Classes/AI/Navigation/NavigationTypes.h
@@ -476,7 +476,7 @@ struct ENGINE_API FNavAgentProperties : public FMovementProperties
 	float NavWalkingSearchHeightScale;
 
 	/** Type of navigation data used by agent, null means "any" */
-	UPROPERTY(EditAnywhere, BlueprintReadWrite, Category=MovementProperties, meta=(MetaClass = "NavigationData"))
+	UPROPERTY(EditAnywhere, BlueprintReadWrite, Category=MovementProperties, meta=(MetaClass = "/Script/NavigationSystem.NavigationData"))
 	FSoftClassPath PreferredNavData;
 	
 	FNavAgentProperties(float Radius = -1.f, float Height = -1.f)
@@ -554,7 +554,7 @@ struct ENGINE_API FNavDataConfig : public FNavAgentProperties
 
 protected:
 	/** Class to use when spawning navigation data instance */
-	UPROPERTY(EditAnywhere, BlueprintReadWrite, Category = Navigation, meta = (MetaClass = "NavigationData"))
+	UPROPERTY(EditAnywhere, BlueprintReadWrite, Category = Navigation, meta = (MetaClass = "/Script/NavigationSystem.NavigationData"))
 	TSoftClassPtr<AActor> NavDataClass;
 
 public:	
diff --git a/Engine/Source/Runtime/Engine/Classes/AI/NavigationSystemConfig.h b/Engine/Source/Runtime/Engine/Classes/AI/NavigationSystemConfig.h
--- a/Engine/Source/Runtime/Engine/Classes/AI/NavigationSystemConfig.h
+++ b/Engine/Source/Runtime/Engine/Classes/AI/NavigationSystemConfig.h
@@ -15,7 +15,7 @@ class ENGINE_API UNavigationSystemConfig : public UObject
 	GENERATED_BODY()
 
 public:
-	UPROPERTY(EditAnywhere, Category=Navigation, meta = (MetaClass = "NavigationSystemBase", NoResetToDefault))
+	UPROPERTY(EditAnywhere, Category=Navigation, meta = (MetaClass = "/Script/Engine.NavigationSystemBase", NoResetToDefault))
 	FSoftClassPath NavigationSystemClass;
 
 	/** NavigationSystem's properties in Project Settings define all possible supported agents,
diff --git a/Engine/Source/Runtime/Engine/Classes/Animation/AnimNode_SequencePlayer.h b/Engine/Source/Runtime/Engine/Classes/Animation/AnimNode_SequencePlayer.h
--- a/Engine/Source/Runtime/Engine/Classes/Animation/AnimNode_SequencePlayer.h
+++ b/Engine/Source/Runtime/Engine/Classes/Animation/AnimNode_SequencePlayer.h
@@ -106,7 +106,7 @@ protected:
 	bool bIgnoreForRelevancyTest = false;
 	
 	// The animation sequence asset to play
-	UPROPERTY(EditAnywhere, Category = Settings, meta = (PinHiddenByDefault, DisallowedClasses="AnimMontage", FoldProperty))
+	UPROPERTY(EditAnywhere, Category = Settings, meta = (PinHiddenByDefault, DisallowedClasses="/Script/Engine.AnimMontage", FoldProperty))
 	TObjectPtr<UAnimSequenceBase> Sequence = nullptr;
 
 	// The Basis in which the PlayRate is expressed in. This is used to rescale PlayRate inputs.
@@ -188,7 +188,7 @@ protected:
 	bool bIgnoreForRelevancyTest = false;
 
 	// The animation sequence asset to play
-	UPROPERTY(EditAnywhere, Category = Settings, meta = (PinHiddenByDefault, DisallowedClasses="AnimMontage"))
+	UPROPERTY(EditAnywhere, Category = Settings, meta = (PinHiddenByDefault, DisallowedClasses="/Script/Engine.AnimMontage"))
 	TObjectPtr<UAnimSequenceBase> Sequence = nullptr;
 
 	// The Basis in which the PlayRate is expressed in. This is used to rescale PlayRate inputs.
diff --git a/Engine/Source/Runtime/Engine/Classes/Animation/AnimSequence.h b/Engine/Source/Runtime/Engine/Classes/Animation/AnimSequence.h
--- a/Engine/Source/Runtime/Engine/Classes/Animation/AnimSequence.h
+++ b/Engine/Source/Runtime/Engine/Classes/Animation/AnimSequence.h
@@ -313,7 +313,7 @@ public:
 
 #if WITH_EDITORONLY_DATA
 	/** If RetargetSource is set to Default (None), this is asset for the base pose to use when retargeting. Transform data will be saved in RetargetSourceAssetReferencePose. */
-	UPROPERTY(EditAnywhere, AssetRegistrySearchable, Category=Animation, meta = (DisallowedClasses = "DestructibleMesh"))
+	UPROPERTY(EditAnywhere, AssetRegistrySearchable, Category=Animation, meta = (DisallowedClasses = "/Script/ApexDestruction.DestructibleMesh"))
 	TSoftObjectPtr<USkeletalMesh> RetargetSourceAsset;
 #endif
 
diff --git a/Engine/Source/Runtime/Engine/Classes/Animation/PoseAsset.h b/Engine/Source/Runtime/Engine/Classes/Animation/PoseAsset.h
--- a/Engine/Source/Runtime/Engine/Classes/Animation/PoseAsset.h
+++ b/Engine/Source/Runtime/Engine/Classes/Animation/PoseAsset.h
@@ -175,7 +175,7 @@ public:
 
 #if WITH_EDITORONLY_DATA
 	/** If RetargetSource is set to Default (None), this is asset for the base pose to use when retargeting. Transform data will be saved in RetargetSourceAssetReferencePose. */
-	UPROPERTY(EditAnywhere, AssetRegistrySearchable, Category=Animation, meta = (DisallowedClasses = "DestructibleMesh"))
+	UPROPERTY(EditAnywhere, AssetRegistrySearchable, Category=Animation, meta = (DisallowedClasses = "/Script/ApexDestruction.DestructibleMesh"))
 	TSoftObjectPtr<USkeletalMesh> RetargetSourceAsset;
 #endif
 
diff --git a/Engine/Source/Runtime/Engine/Classes/Components/SkeletalMeshComponent.h b/Engine/Source/Runtime/Engine/Classes/Components/SkeletalMeshComponent.h
--- a/Engine/Source/Runtime/Engine/Classes/Components/SkeletalMeshComponent.h
+++ b/Engine/Source/Runtime/Engine/Classes/Components/SkeletalMeshComponent.h
@@ -2511,7 +2511,7 @@ private:
 #if WITH_EDITORONLY_DATA
 
 	/** Default Animating Rig to Use if bOverrideDefaultAnimatingRig is true */
-	UPROPERTY(EditAnywhere, Category = AnimationRig, BlueprintGetter = GetDefaultAnimatingRigOverride, BlueprintSetter = SetDefaultAnimatingRigOverride, meta = (AllowedClasses = "ControlRigBlueprint"), meta = (editcondition = "bOverrideDefaultAnimatingRig"))
+	UPROPERTY(EditAnywhere, Category = AnimationRig, BlueprintGetter = GetDefaultAnimatingRigOverride, BlueprintSetter = SetDefaultAnimatingRigOverride, meta = (AllowedClasses = "/Script/ControlRigDeveloper.ControlRigBlueprint"), meta = (editcondition = "bOverrideDefaultAnimatingRig"))
 	TSoftObjectPtr<UObject> DefaultAnimatingRigOverride;
 
 #endif
diff --git a/Engine/Source/Runtime/Engine/Classes/Components/SkinnedMeshComponent.h b/Engine/Source/Runtime/Engine/Classes/Components/SkinnedMeshComponent.h
--- a/Engine/Source/Runtime/Engine/Classes/Components/SkinnedMeshComponent.h
+++ b/Engine/Source/Runtime/Engine/Classes/Components/SkinnedMeshComponent.h
@@ -139,7 +139,7 @@ struct FActiveMorphTarget
 };
 
 /** Vertex skin weight info supplied for a component override. */
-USTRUCT(BlueprintType, meta = (HasNativeMake = "Engine.KismetRenderingLibrary.MakeSkinWeightInfo", HasNativeBreak = "Engine.KismetRenderingLibrary.BreakSkinWeightInfo"))
+USTRUCT(BlueprintType, meta = (HasNativeMake = "/Script/Engine.KismetRenderingLibrary.MakeSkinWeightInfo", HasNativeBreak = "/Script/Engine.KismetRenderingLibrary.BreakSkinWeightInfo"))
 struct FSkelMeshSkinWeightInfo
 {
 	GENERATED_USTRUCT_BODY()
@@ -192,7 +192,7 @@ struct FVertexOffsetUsage
 {
 	GENERATED_BODY()
 
-	UPROPERTY(EditAnywhere, Category = "Mesh", meta = (Bitmask, BitmaskEnum = EVertexOffsetUsageType))
+	UPROPERTY(EditAnywhere, Category = "Mesh", meta = (Bitmask, BitmaskEnum = "/Script/Engine.EVertexOffsetUsageType"))
 	int32 Usage = 0;
 };
 
@@ -213,7 +213,7 @@ class ENGINE_API USkinnedMeshComponent : public UMeshComponent, public ILODSyncI
 	friend class FSkinnedMeshComponentRecreateRenderStateContext;
 
 	/** The skeletal mesh used by this component. */
-	UPROPERTY(EditAnywhere, BlueprintReadOnly, Category="Mesh", meta = (DisallowedClasses = "DestructibleMesh"))
+	UPROPERTY(EditAnywhere, BlueprintReadOnly, Category="Mesh", meta = (DisallowedClasses = "/Script/ApexDestruction.DestructibleMesh"))
 	TObjectPtr<class USkeletalMesh> SkeletalMesh;
 
 	//
diff --git a/Engine/Source/Runtime/Engine/Classes/Engine/Engine.h b/Engine/Source/Runtime/Engine/Classes/Engine/Engine.h
--- a/Engine/Source/Runtime/Engine/Classes/Engine/Engine.h
+++ b/Engine/Source/Runtime/Engine/Classes/Engine/Engine.h
@@ -667,7 +667,7 @@ private:
 
 public:
 	/** Sets the font used for the smallest engine text */
-	UPROPERTY(globalconfig, EditAnywhere, Category=Fonts, meta=(AllowedClasses="Font", DisplayName="Tiny Font", ConfigRestartRequired=true))
+	UPROPERTY(globalconfig, EditAnywhere, Category=Fonts, meta=(AllowedClasses="/Script/Engine.Font", DisplayName="Tiny Font", ConfigRestartRequired=true))
 	FSoftObjectPath TinyFontName;
 
 private:
@@ -676,7 +676,7 @@ private:
 
 public:
 	/** Sets the font used for small engine text, used for most debug displays */
-	UPROPERTY(globalconfig, EditAnywhere, Category=Fonts, meta=(AllowedClasses="Font", DisplayName="Small Font", ConfigRestartRequired=true))
+	UPROPERTY(globalconfig, EditAnywhere, Category=Fonts, meta=(AllowedClasses="/Script/Engine.Font", DisplayName="Small Font", ConfigRestartRequired=true))
 	FSoftObjectPath SmallFontName;
 
 private:
@@ -685,7 +685,7 @@ private:
 
 public:
 	/** Sets the font used for medium engine text */
-	UPROPERTY(globalconfig, EditAnywhere, Category=Fonts, meta=(AllowedClasses="Font", DisplayName="Medium Font", ConfigRestartRequired=true))
+	UPROPERTY(globalconfig, EditAnywhere, Category=Fonts, meta=(AllowedClasses="/Script/Engine.Font", DisplayName="Medium Font", ConfigRestartRequired=true))
 	FSoftObjectPath MediumFontName;
 
 private:
@@ -694,7 +694,7 @@ private:
 
 public:
 	/** Sets the font used for large engine text */
-	UPROPERTY(globalconfig, EditAnywhere, Category=Fonts, meta=(AllowedClasses="Font", DisplayName="Large Font", ConfigRestartRequired=true))
+	UPROPERTY(globalconfig, EditAnywhere, Category=Fonts, meta=(AllowedClasses="/Script/Engine.Font", DisplayName="Large Font", ConfigRestartRequired=true))
 	FSoftObjectPath LargeFontName;
 
 private:
@@ -703,7 +703,7 @@ private:
 
 public:
 	/** Sets the font used by the default Subtitle Manager */
-	UPROPERTY(globalconfig, EditAnywhere, Category=Fonts, meta=(AllowedClasses="Font", DisplayName="Subtitle Font", ConfigRestartRequired=true), AdvancedDisplay)
+	UPROPERTY(globalconfig, EditAnywhere, Category=Fonts, meta=(AllowedClasses="/Script/Engine.Font", DisplayName="Subtitle Font", ConfigRestartRequired=true), AdvancedDisplay)
 	FSoftObjectPath SubtitleFontName;
 
 private:
@@ -719,31 +719,31 @@ public:
 	TSubclassOf<class UConsole>  ConsoleClass;
 
 	/** Sets the class to use for the game console summoned with ~ */
-	UPROPERTY(globalconfig, noclear, EditAnywhere, Category=DefaultClasses, meta=(MetaClass="Console", DisplayName="Console Class", ConfigRestartRequired=true))
+	UPROPERTY(globalconfig, noclear, EditAnywhere, Category=DefaultClasses, meta=(MetaClass="/Script/Engine.Console", DisplayName="Console Class", ConfigRestartRequired=true))
 	FSoftClassPath ConsoleClassName;
 
 	UPROPERTY()
 	TSubclassOf<class UGameViewportClient>  GameViewportClientClass;
 
 	/** Sets the class to use for the game viewport client, which can be overridden to change game-specific input and display behavior. */
-	UPROPERTY(globalconfig, noclear, EditAnywhere, Category=DefaultClasses, meta=(MetaClass="GameViewportClient", DisplayName="Game Viewport Client Class", ConfigRestartRequired=true))
+	UPROPERTY(globalconfig, noclear, EditAnywhere, Category=DefaultClasses, meta=(MetaClass="/Script/Engine.GameViewportClient", DisplayName="Game Viewport Client Class", ConfigRestartRequired=true))
 	FSoftClassPath GameViewportClientClassName;
 
 	UPROPERTY()
 	TSubclassOf<class ULocalPlayer>  LocalPlayerClass;
 
 	/** Sets the class to use for local players, which can be overridden to store game-specific information for a local player. */
-	UPROPERTY(globalconfig, noclear, EditAnywhere, Category=DefaultClasses, meta=(MetaClass="LocalPlayer", DisplayName="Local Player Class", ConfigRestartRequired=true))
+	UPROPERTY(globalconfig, noclear, EditAnywhere, Category=DefaultClasses, meta=(MetaClass="/Script/Engine.LocalPlayer", DisplayName="Local Player Class", ConfigRestartRequired=true))
 	FSoftClassPath LocalPlayerClassName;
 
 	UPROPERTY()
 	TSubclassOf<class AWorldSettings>  WorldSettingsClass;
 
 	/** Sets the class to use for WorldSettings, which can be overridden to store game-specific information on map/world. */
-	UPROPERTY(globalconfig, noclear, EditAnywhere, Category=DefaultClasses, meta=(MetaClass="WorldSettings", DisplayName="World Settings Class", ConfigRestartRequired=true))
+	UPROPERTY(globalconfig, noclear, EditAnywhere, Category=DefaultClasses, meta=(MetaClass="/Script/Engine.WorldSettings", DisplayName="World Settings Class", ConfigRestartRequired=true))
 	FSoftClassPath WorldSettingsClassName;
 
-	UPROPERTY(globalconfig, noclear, meta=(MetaClass="NavigationSystem", DisplayName="Navigation System Class"))
+	UPROPERTY(globalconfig, noclear, meta=(MetaClass="/Script/NavigationSystem.NavigationSystem", DisplayName="Navigation System Class"))
 	FSoftClassPath NavigationSystemClassName;
 
 	/** Sets the class to use for NavigationSystem, which can be overridden to change game-specific navigation/AI behavior. */
@@ -751,32 +751,32 @@ public:
 	TSubclassOf<class UNavigationSystemBase>  NavigationSystemClass;
 
 	/** Sets the Navigation System Config class, which can be overridden to change game-specific navigation/AI behavior. */
-	UPROPERTY(globalconfig, noclear, meta = (MetaClass = "NavigationSystem", DisplayName = "Navigation System Config Class"))
+	UPROPERTY(globalconfig, noclear, meta = (MetaClass = "/Script/NavigationSystem.NavigationSystem", DisplayName = "Navigation System Config Class"))
 	FSoftClassPath NavigationSystemConfigClassName;
 
 	UPROPERTY()
 	TSubclassOf<class UNavigationSystemConfig>  NavigationSystemConfigClass;
 	
 	/** Sets the AvoidanceManager class, which can be overridden to change AI crowd behavior. */
-	UPROPERTY(globalconfig, noclear, meta=(MetaClass="AvoidanceManager", DisplayName="Avoidance Manager Class"))
+	UPROPERTY(globalconfig, noclear, meta=(MetaClass="/Script/Engine.AvoidanceManager", DisplayName="Avoidance Manager Class"))
 	FSoftClassPath AvoidanceManagerClassName;
 	
 	UPROPERTY()
 	TSubclassOf<class UAvoidanceManager>  AvoidanceManagerClass;
 
 	/** Sets the class to be used as the default AIController class for pawns. */
-	UPROPERTY(globalconfig, noclear, meta = (MetaClass = "AI", DisplayName = "Default AIController class for all Pawns"))
+	UPROPERTY(globalconfig, noclear, meta = (MetaClass = "/Script/AIModule.AIController", DisplayName = "Default AIController class for all Pawns"))
 	FSoftClassPath AIControllerClassName;
 
 	UPROPERTY()
 	TSubclassOf<class UPhysicsCollisionHandler>	PhysicsCollisionHandlerClass;
 
 	/** Sets the PhysicsCollisionHandler class to use by default, which can be overridden to change game-specific behavior when objects collide using physics. */
-	UPROPERTY(globalconfig, noclear, EditAnywhere, Category=DefaultClasses, meta=(MetaClass="PhysicsCollisionHandler", DisplayName="Physics Collision Handler Class", ConfigRestartRequired=true), AdvancedDisplay)
+	UPROPERTY(globalconfig, noclear, EditAnywhere, Category=DefaultClasses, meta=(MetaClass="/Script/Engine.PhysicsCollisionHandler", DisplayName="Physics Collision Handler Class", ConfigRestartRequired=true), AdvancedDisplay)
 	FSoftClassPath PhysicsCollisionHandlerClassName;
 
 	/** Sets the GameUserSettings class, which can be overridden to support game-specific options for Graphics/Sound/Gameplay. */
-	UPROPERTY(globalconfig, noclear, EditAnywhere, Category=DefaultClasses, meta=(MetaClass="GameUserSettings", DisplayName="Game User Settings Class", ConfigRestartRequired=true), AdvancedDisplay)
+	UPROPERTY(globalconfig, noclear, EditAnywhere, Category=DefaultClasses, meta=(MetaClass="/Script/Engine.GameUserSettings", DisplayName="Game User Settings Class", ConfigRestartRequired=true), AdvancedDisplay)
 	FSoftClassPath GameUserSettingsClassName;
 
 	UPROPERTY()
@@ -790,15 +790,15 @@ public:
 	TSubclassOf<class ALevelScriptActor>  LevelScriptActorClass;
 
 	/** Sets the Level Script Actor class, which can be overridden to allow game-specific behavior in per-map blueprint scripting */
-	UPROPERTY(globalconfig, noclear, EditAnywhere, Category=DefaultClasses, meta=(MetaClass="LevelScriptActor", DisplayName="Level Script Actor Class", ConfigRestartRequired=true))
+	UPROPERTY(globalconfig, noclear, EditAnywhere, Category=DefaultClasses, meta=(MetaClass="/Script/Engine.LevelScriptActor", DisplayName="Level Script Actor Class", ConfigRestartRequired=true))
 	FSoftClassPath LevelScriptActorClassName;
 	
 	/** Sets the base class to use for new blueprints created in the editor, configurable on a per-game basis */
-	UPROPERTY(globalconfig, noclear, EditAnywhere, Category=DefaultClasses, meta=(MetaClass="Object", DisplayName="Default Blueprint Base Class", AllowAbstract, BlueprintBaseOnly), AdvancedDisplay)
+	UPROPERTY(globalconfig, noclear, EditAnywhere, Category=DefaultClasses, meta=(MetaClass="/Script/CoreUObject.Object", DisplayName="Default Blueprint Base Class", AllowAbstract, BlueprintBaseOnly), AdvancedDisplay)
 	FSoftClassPath DefaultBlueprintBaseClassName;
 
 	/** Sets the class for a global object spawned at startup to handle game-specific data. If empty, it will not spawn one */
-	UPROPERTY(globalconfig, noclear, EditAnywhere, Category=DefaultClasses, meta=(MetaClass="Object", DisplayName="Game Singleton Class", ConfigRestartRequired=true), AdvancedDisplay)
+	UPROPERTY(globalconfig, noclear, EditAnywhere, Category=DefaultClasses, meta=(MetaClass="/Script/CoreUObject.Object", DisplayName="Game Singleton Class", ConfigRestartRequired=true), AdvancedDisplay)
 	FSoftClassPath GameSingletonClassName;
 
 	/** A UObject spawned at initialization time to handle game-specific data */
@@ -806,7 +806,7 @@ public:
 	TObjectPtr<UObject> GameSingleton;
 
 	/** Sets the class to spawn as the global AssetManager, configurable per game. If empty, it will not spawn one */
-	UPROPERTY(globalconfig, noclear, EditAnywhere, Category=DefaultClasses, meta=(MetaClass="Object", DisplayName="Asset Manager Class", ConfigRestartRequired=true), AdvancedDisplay)
+	UPROPERTY(globalconfig, noclear, EditAnywhere, Category=DefaultClasses, meta=(MetaClass="/Script/CoreUObject.Object", DisplayName="Asset Manager Class", ConfigRestartRequired=true), AdvancedDisplay)
 	FSoftClassPath AssetManagerClassName;
 
 	/** A UObject spawned at initialization time to handle runtime asset loading and management */
@@ -1099,7 +1099,7 @@ public:
 	TObjectPtr<class UMaterial> PreviewShadowsIndicatorMaterial;
 
 	/** Path of the material that renders a message about preview shadows being used. */
-	UPROPERTY(globalconfig, EditAnywhere, Category=DefaultMaterials, meta=(AllowedClasses="Material", DisplayName="Preview Shadows Indicator Material"))
+	UPROPERTY(globalconfig, EditAnywhere, Category=DefaultMaterials, meta=(AllowedClasses="/Script/Engine.Material", DisplayName="Preview Shadows Indicator Material"))
 	FSoftObjectPath PreviewShadowsIndicatorMaterialName;
 
 	/** Material that 'fakes' lighting, used for arrows, widgets. */
@@ -1386,7 +1386,7 @@ public:
 	 * This class will be responsible of updating the application Time and DeltaTime.
 	 * Can be used to synchronize the engine with another process (gen-lock).
 	 */
-	UPROPERTY(AdvancedDisplay, config, EditAnywhere, Category=Framerate, meta=(MetaClass="EngineCustomTimeStep", DisplayName="Custom TimeStep"))
+	UPROPERTY(AdvancedDisplay, config, EditAnywhere, Category=Framerate, meta=(MetaClass="/Script/Engine.EngineCustomTimeStep", DisplayName="Custom TimeStep"))
 	FSoftClassPath CustomTimeStepClassName;
 
 private:
@@ -1402,7 +1402,7 @@ private:
 
 public:
 	/** Set TimecodeProvider when the engine is started. */
-	UPROPERTY(config, EditAnywhere, Category=Timecode, meta=(MetaClass="TimecodeProvider", DisplayName="Timecode Provider"))
+	UPROPERTY(config, EditAnywhere, Category=Timecode, meta=(MetaClass="/Script/Engine.TimecodeProvider", DisplayName="Timecode Provider"))
 	FSoftClassPath TimecodeProviderClassName;
 
 	/**
diff --git a/Engine/Source/Runtime/Engine/Classes/Engine/MeshMerging.h b/Engine/Source/Runtime/Engine/Classes/Engine/MeshMerging.h
--- a/Engine/Source/Runtime/Engine/Classes/Engine/MeshMerging.h
+++ b/Engine/Source/Runtime/Engine/Classes/Engine/MeshMerging.h
@@ -670,7 +670,7 @@ struct ENGINE_API FMeshInstancingSettings
 	/**
 	 * Whether to use the Instanced Static Mesh Compoment or the Hierarchical Instanced Static Mesh Compoment
 	 */
-	UPROPERTY(BlueprintReadWrite, EditAnywhere, Category = "Instancing", meta = (DisplayName = "Select the type of Instanced Component", DisallowedClasses = "FoliageInstancedStaticMeshComponent"))
+	UPROPERTY(BlueprintReadWrite, EditAnywhere, Category = "Instancing", meta = (DisplayName = "Select the type of Instanced Component", DisallowedClasses = "/Script/Foliage.FoliageInstancedStaticMeshComponent"))
 	TSubclassOf<UInstancedStaticMeshComponent> ISMComponentToUse;
 };
 
diff --git a/Engine/Source/Runtime/Engine/Classes/Engine/NetSerialization.h b/Engine/Source/Runtime/Engine/Classes/Engine/NetSerialization.h
--- a/Engine/Source/Runtime/Engine/Classes/Engine/NetSerialization.h
+++ b/Engine/Source/Runtime/Engine/Classes/Engine/NetSerialization.h
@@ -401,7 +401,7 @@ bool SerializeFixedVector(FVector3d &Vector, FArchive& Ar)
  *	Note: this is the historical UE format for vector net serialization
  *
  */
-USTRUCT(meta = (HasNativeMake = "Engine.KismetMathLibrary.MakeVector_NetQuantize", HasNativeBreak = "Engine.KismetMathLibrary.BreakVector_NetQuantize"))
+USTRUCT(meta = (HasNativeMake = "/Script/Engine.KismetMathLibrary.MakeVector_NetQuantize", HasNativeBreak = "/Script/Engine.KismetMathLibrary.BreakVector_NetQuantize"))
 struct FVector_NetQuantize : public FVector
 {
 	GENERATED_USTRUCT_BODY()
@@ -447,7 +447,7 @@ struct TStructOpsTypeTraits< FVector_NetQuantize > : public TStructOpsTypeTraits
  *	Valid range: 2^24 / 10 = +/- 1,677,721.6
  *
  */
-USTRUCT(meta = (HasNativeMake = "Engine.KismetMathLibrary.MakeVector_NetQuantize10", HasNativeBreak = "Engine.KismetMathLibrary.BreakVector_NetQuantize10"))
+USTRUCT(meta = (HasNativeMake = "/Script/Engine.KismetMathLibrary.MakeVector_NetQuantize10", HasNativeBreak = "/Script/Engine.KismetMathLibrary.BreakVector_NetQuantize10"))
 struct FVector_NetQuantize10 : public FVector
 {
 	GENERATED_USTRUCT_BODY()
@@ -493,7 +493,7 @@ struct TStructOpsTypeTraits< FVector_NetQuantize10 > : public TStructOpsTypeTrai
  *	Valid range: 2^30 / 100 = +/- 10,737,418.24
  *
  */
-USTRUCT(meta = (HasNativeMake = "Engine.KismetMathLibrary.MakeVector_NetQuantize100", HasNativeBreak = "Engine.KismetMathLibrary.BreakVector_NetQuantize100"))
+USTRUCT(meta = (HasNativeMake = "/Script/Engine.KismetMathLibrary.MakeVector_NetQuantize100", HasNativeBreak = "/Script/Engine.KismetMathLibrary.BreakVector_NetQuantize100"))
 struct FVector_NetQuantize100 : public FVector
 {
 	GENERATED_USTRUCT_BODY()
@@ -535,7 +535,7 @@ struct TStructOpsTypeTraits< FVector_NetQuantize100 > : public TStructOpsTypeTra
  *	16 bits per component
  *	Valid range: -1..+1 inclusive
  */
-USTRUCT(meta = (HasNativeMake = "Engine.KismetMathLibrary.MakeVector_NetQuantizeNormal", HasNativeBreak = "Engine.KismetMathLibrary.BreakVector_NetQuantizeNormal"))
+USTRUCT(meta = (HasNativeMake = "/Script/Engine.KismetMathLibrary.MakeVector_NetQuantizeNormal", HasNativeBreak = "/Script/Engine.KismetMathLibrary.BreakVector_NetQuantizeNormal"))
 struct FVector_NetQuantizeNormal : public FVector
 {
 	GENERATED_USTRUCT_BODY()
diff --git a/Engine/Source/Runtime/Engine/Classes/Engine/PreviewMeshCollection.h b/Engine/Source/Runtime/Engine/Classes/Engine/PreviewMeshCollection.h
--- a/Engine/Source/Runtime/Engine/Classes/Engine/PreviewMeshCollection.h
+++ b/Engine/Source/Runtime/Engine/Classes/Engine/PreviewMeshCollection.h
@@ -30,7 +30,7 @@ struct FPreviewMeshCollectionEntry
 	}
 
 	/** The skeletal mesh to display */
-	UPROPERTY(EditAnywhere, Category = "Skeletal Mesh", meta=(DisplayThumbnail=true, DisallowedClasses = "DestructibleMesh"))
+	UPROPERTY(EditAnywhere, Category = "Skeletal Mesh", meta=(DisplayThumbnail=true, DisallowedClasses = "/Script/ApexDestruction.DestructibleMesh"))
 	TSoftObjectPtr<USkeletalMesh> SkeletalMesh;
 };
 
diff --git a/Engine/Source/Runtime/Engine/Classes/Engine/RendererSettings.h b/Engine/Source/Runtime/Engine/Classes/Engine/RendererSettings.h
--- a/Engine/Source/Runtime/Engine/Classes/Engine/RendererSettings.h
+++ b/Engine/Source/Runtime/Engine/Classes/Engine/RendererSettings.h
@@ -1126,19 +1126,19 @@ class ENGINE_API URendererSettings : public UDeveloperSettings
 	/**
 	" Visualize calibration material settings for post process calibration materials, used for setting full-screen images used for monitor calibration."
 	*/
-	UPROPERTY(config, EditAnywhere, Category = PostProcessCalibrationMaterials, meta = (AllowedClasses = "Material",
+	UPROPERTY(config, EditAnywhere, Category = PostProcessCalibrationMaterials, meta = (AllowedClasses = "/Script/Engine.Material",
 		DisplayName = "Visualize Calibration Color Material Path",
 		ToolTip = "When the VisualizeCalibrationColor show flag is enabled, this path will be used as the post-process material to render. The post-process material's Blendable Location property must be set to \"After Tonemapping\" for proper calibration display.",
 		ConfigRestartRequired = false))
 	FSoftObjectPath VisualizeCalibrationColorMaterialPath;
 
-	UPROPERTY(config, EditAnywhere, Category = PostProcessCalibrationMaterials, meta = (AllowedClasses = "Material",
+	UPROPERTY(config, EditAnywhere, Category = PostProcessCalibrationMaterials, meta = (AllowedClasses = "/Script/Engine.Material",
 		DisplayName = "Visualize Calibration Custom Material Path",
 		ToolTip = "When the VisualizeCalibrationCustom show flag is enabled, this path will be used as the post-process material to render. The post-process material's Blendable Location property must be set to \"After Tonemapping\" for proper calibration display.",
 		ConfigRestartRequired = false))
 	FSoftObjectPath VisualizeCalibrationCustomMaterialPath;
 
-	UPROPERTY(config, EditAnywhere, Category = PostProcessCalibrationMaterials, meta = (AllowedClasses = "Material",
+	UPROPERTY(config, EditAnywhere, Category = PostProcessCalibrationMaterials, meta = (AllowedClasses = "/Script/Engine.Material",
 		DisplayName = "Visualize Calibration Grayscale Material Path",
 		ToolTip = "When the VisualizeCalibrationGrayscale show flag is enabled, this path will be used as the post-process material to render. The post-process material's Blendable Location property must be set to \"After Tonemapping\" for proper calibration display.",
 		ConfigRestartRequired = false))
diff --git a/Engine/Source/Runtime/Engine/Classes/Engine/Scene.h b/Engine/Source/Runtime/Engine/Classes/Engine/Scene.h
--- a/Engine/Source/Runtime/Engine/Classes/Engine/Scene.h
+++ b/Engine/Source/Runtime/Engine/Classes/Engine/Scene.h
@@ -635,7 +635,7 @@ struct FWeightedBlendable
 	float Weight;
 
 	/** should be of the IBlendableInterface* type but UProperties cannot express that */
-	UPROPERTY(EditAnywhere, BlueprintReadWrite, Category=FWeightedBlendable, meta=( AllowedClasses="BlendableInterface", Keywords="PostProcess" ))
+	UPROPERTY(EditAnywhere, BlueprintReadWrite, Category=FWeightedBlendable, meta=( AllowedClasses="/Script/Engine.BlendableInterface", Keywords="PostProcess" ))
 	TObjectPtr<UObject> Object;
 
 	// default constructor
diff --git a/Engine/Source/Runtime/Engine/Classes/Engine/SkeletalMesh.h b/Engine/Source/Runtime/Engine/Classes/Engine/SkeletalMesh.h
--- a/Engine/Source/Runtime/Engine/Classes/Engine/SkeletalMesh.h
+++ b/Engine/Source/Runtime/Engine/Classes/Engine/SkeletalMesh.h
@@ -1432,7 +1432,7 @@ public:
 
 	/** The Default Control Rig To Animate with when used in Sequnecer. */
 	UE_DEPRECATED(4.27, "Please do not access this member directly; use USkeletalMesh::GetDefaultAnimatingRig() or USkeletalMesh::SetDefaultAnimatingRig().")
-	UPROPERTY(EditAnywhere, Category = AnimationRig, BlueprintGetter = GetDefaultAnimatingRig, BlueprintSetter = SetDefaultAnimatingRig, meta = (AllowedClasses = "ControlRigBlueprint"))
+	UPROPERTY(EditAnywhere, Category = AnimationRig, BlueprintGetter = GetDefaultAnimatingRig, BlueprintSetter = SetDefaultAnimatingRig, meta = (AllowedClasses = "/Script/ControlRigDeveloper.ControlRigBlueprint"))
 	TSoftObjectPtr<UObject> DefaultAnimatingRig;
 	static FName GetDefaultAnimatingRigMemberName()
 	{
diff --git a/Engine/Source/Runtime/Engine/Classes/Engine/UserInterfaceSettings.h b/Engine/Source/Runtime/Engine/Classes/Engine/UserInterfaceSettings.h
--- a/Engine/Source/Runtime/Engine/Classes/Engine/UserInterfaceSettings.h
+++ b/Engine/Source/Runtime/Engine/Classes/Engine/UserInterfaceSettings.h
@@ -117,7 +117,7 @@ public:
 	UPROPERTY(config, EditAnywhere, Category = "Hardware Cursors")
 	TMap<TEnumAsByte<EMouseCursor::Type>, FHardwareCursorReference> HardwareCursors;
 
-	UPROPERTY(config, EditAnywhere, Category = "Software Cursors", meta = ( MetaClass = "UserWidget" ))
+	UPROPERTY(config, EditAnywhere, Category = "Software Cursors", meta = ( MetaClass = "/Script/UMG.UserWidget" ))
 	TMap<TEnumAsByte<EMouseCursor::Type>, FSoftClassPath> SoftwareCursors;
 
 	// DEPRECATED 4.16
@@ -164,7 +164,7 @@ public:
 	/**
 	 * Set DPI Scale Rule to Custom, and this class will be used instead of any of the built-in rules.
 	 */
-	UPROPERTY(config, EditAnywhere, Category="DPI Scaling", meta=( MetaClass="DPICustomScalingRule" ))
+	UPROPERTY(config, EditAnywhere, Category="DPI Scaling", meta=( MetaClass="/Script/Engine.DPICustomScalingRule" ))
 	FSoftClassPath CustomScalingRuleClass;
 
 	/**
