Converting hardcoded short class/enum names to pathnames ahead of ANY_PACKAGE removal

diff --git a/Engine/Source/Runtime/Engine/Classes/GameFramework/Actor.h b/Engine/Source/Runtime/Engine/Classes/GameFramework/Actor.h
--- a/Engine/Source/Runtime/Engine/Classes/GameFramework/Actor.h
+++ b/Engine/Source/Runtime/Engine/Classes/GameFramework/Actor.h
@@ -3366,7 +3366,7 @@ public:
 	virtual UActorComponent* FindComponentByClass(const TSubclassOf<UActorComponent> ComponentClass) const;
 	
 	/** Searches components array and returns first encountered component of the specified class */
-	UFUNCTION(BlueprintCallable, Category = "Actor", meta = (ComponentClass = "ActorComponent"), meta = (DeterminesOutputType = "ComponentClass"))
+	UFUNCTION(BlueprintCallable, Category = "Actor", meta = (ComponentClass = "/Script/Engine.ActorComponent"), meta = (DeterminesOutputType = "ComponentClass"))
 	UActorComponent* GetComponentByClass(TSubclassOf<UActorComponent> ComponentClass) const;
 
 	/**
@@ -3374,11 +3374,11 @@ public:
 	 * Currently returns an array of UActorComponent which must be cast to the correct type.
 	 * This intended to only be used by blueprints. Use GetComponents() in C++.
 	 */
-	UFUNCTION(BlueprintCallable, Category = "Actor", meta = (ComponentClass = "ActorComponent", DisplayName = "Get Components By Class", ScriptName = "GetComponentsByClass", DeterminesOutputType = "ComponentClass"))
+	UFUNCTION(BlueprintCallable, Category = "Actor", meta = (ComponentClass = "/Script/Engine.ActorComponent", DisplayName = "Get Components By Class", ScriptName = "GetComponentsByClass", DeterminesOutputType = "ComponentClass"))
 	TArray<UActorComponent*> K2_GetComponentsByClass(TSubclassOf<UActorComponent> ComponentClass) const;
 
 	/** Gets all the components that inherit from the given class with a given tag. */
-	UFUNCTION(BlueprintCallable, Category = "Actor", meta = (ComponentClass = "ActorComponent"), meta = (DeterminesOutputType = "ComponentClass"))
+	UFUNCTION(BlueprintCallable, Category = "Actor", meta = (ComponentClass = "/Script/Engine.ActorComponent"), meta = (DeterminesOutputType = "ComponentClass"))
 	TArray<UActorComponent*> GetComponentsByTag(TSubclassOf<UActorComponent> ComponentClass, FName Tag) const;
 
 	/** Gets all the components that implements the given interface. */
diff --git a/Engine/Source/Runtime/Engine/Classes/GameFramework/InputSettings.h b/Engine/Source/Runtime/Engine/Classes/GameFramework/InputSettings.h
--- a/Engine/Source/Runtime/Engine/Classes/GameFramework/InputSettings.h
+++ b/Engine/Source/Runtime/Engine/Classes/GameFramework/InputSettings.h
@@ -141,7 +141,7 @@ private:
 
 public:
 	/** The default on-screen touch input interface for the game (can be null to disable the onscreen interface) */
-	UPROPERTY(config, EditAnywhere, Category="Mobile", meta=(AllowedClasses="TouchInterface"))
+	UPROPERTY(config, EditAnywhere, Category="Mobile", meta=(AllowedClasses="/Script/Engine.TouchInterface"))
 	FSoftObjectPath DefaultTouchInterface;
 
 	/** The key which opens the console. */
diff --git a/Engine/Source/Runtime/Engine/Classes/GameFramework/WorldSettings.h b/Engine/Source/Runtime/Engine/Classes/GameFramework/WorldSettings.h
--- a/Engine/Source/Runtime/Engine/Classes/GameFramework/WorldSettings.h
+++ b/Engine/Source/Runtime/Engine/Classes/GameFramework/WorldSettings.h
@@ -541,7 +541,7 @@ public:
 	uint8 bIncludeGridSizeInNameForPartitionedActors : 1;
 #endif
 
-	UPROPERTY(EditAnywhere, BlueprintReadOnly, Category=AI, meta=(MetaClass="AISystemBase", editcondition="bEnableAISystem"), AdvancedDisplay)
+	UPROPERTY(EditAnywhere, BlueprintReadOnly, Category=AI, meta=(MetaClass="/Script/Engine.AISystemBase", editcondition="bEnableAISystem"), AdvancedDisplay)
 	TSoftClassPtr<UAISystemBase> AISystemClass;
 
 	/** Additional transform applied when applying LevelStreaming Transform to LevelInstance */
@@ -913,7 +913,7 @@ private:
 	 * Class that will be used when creating new bookmarks.
 	 * Old bookmarks may be recreated with the new class where possible.
 	 */
-	UPROPERTY(Config, Meta=(AllowAbstract="false", ExactClass="false", AllowedClasses="BookmarkBase"))
+	UPROPERTY(Config, Meta=(AllowAbstract="false", ExactClass="false", AllowedClasses="/Script/Engine.BookmarkBase"))
 	TSubclassOf<class UBookmarkBase> DefaultBookmarkClass;
 
 	UPROPERTY()
diff --git a/Engine/Source/Runtime/Engine/Classes/Kismet/ImportanceSamplingLibrary.h b/Engine/Source/Runtime/Engine/Classes/Kismet/ImportanceSamplingLibrary.h
--- a/Engine/Source/Runtime/Engine/Classes/Kismet/ImportanceSamplingLibrary.h
+++ b/Engine/Source/Runtime/Engine/Classes/Kismet/ImportanceSamplingLibrary.h
@@ -34,7 +34,7 @@ namespace EImportanceWeight
 * Texture processed for importance sampling
 * Holds marginal PDF of the rows, as well as the PDF of each row
 */
-USTRUCT(BlueprintType, meta = (HasNativeMake = "Engine.ImportanceSamplingLibrary.MakeImportanceTexture", HasNativeBreak = "Engine.ImportanceSamplingLibrary.BreakImportanceTexture"))
+USTRUCT(BlueprintType, meta = (HasNativeMake = "/Script/Engine.ImportanceSamplingLibrary.MakeImportanceTexture", HasNativeBreak = "/Script/Engine.ImportanceSamplingLibrary.BreakImportanceTexture"))
 struct FImportanceTexture
 {
 	GENERATED_BODY();
diff --git a/Engine/Source/Runtime/Engine/Classes/Kismet/KismetMathLibrary.h b/Engine/Source/Runtime/Engine/Classes/Kismet/KismetMathLibrary.h
--- a/Engine/Source/Runtime/Engine/Classes/Kismet/KismetMathLibrary.h
+++ b/Engine/Source/Runtime/Engine/Classes/Kismet/KismetMathLibrary.h
@@ -896,27 +896,27 @@ class ENGINE_API UKismetMathLibrary : public UBlueprintFunctionLibrary
 	//
 	
 	/** Zero Int Point (0, 0) */
-	UFUNCTION(BlueprintPure, meta = (ScriptConstant = "Zero", ScriptConstantHost = "IntPoint"), Category = "Math|IntPoint|Constants")
+	UFUNCTION(BlueprintPure, meta = (ScriptConstant = "Zero", ScriptConstantHost = "/Script/CoreUObject.IntPoint"), Category = "Math|IntPoint|Constants")
 	static FIntPoint IntPoint_Zero();
 	
 	/** One Int Point (1, 1) */
-	UFUNCTION(BlueprintPure, meta = (ScriptConstant = "One", ScriptConstantHost = "IntPoint"), Category = "Math|IntPoint|Constants")
+	UFUNCTION(BlueprintPure, meta = (ScriptConstant = "One", ScriptConstantHost = "/Script/CoreUObject.IntPoint"), Category = "Math|IntPoint|Constants")
 	static FIntPoint IntPoint_One();
 	
 	/** Up Int Point (0, -1) */
-	UFUNCTION(BlueprintPure, meta = (ScriptConstant = "Up", ScriptConstantHost = "IntPoint"), Category = "Math|IntPoint|Constants")
+	UFUNCTION(BlueprintPure, meta = (ScriptConstant = "Up", ScriptConstantHost = "/Script/CoreUObject.IntPoint"), Category = "Math|IntPoint|Constants")
 	static FIntPoint IntPoint_Up();
 	
 	/** Left Int Point (-1, 0) */
-	UFUNCTION(BlueprintPure, meta = (ScriptConstant = "Left", ScriptConstantHost = "IntPoint"), Category = "Math|IntPoint|Constants")
+	UFUNCTION(BlueprintPure, meta = (ScriptConstant = "Left", ScriptConstantHost = "/Script/CoreUObject.IntPoint"), Category = "Math|IntPoint|Constants")
 	static FIntPoint IntPoint_Left();
 	
 	/** Right Int Point (1, 0) */
-	UFUNCTION(BlueprintPure, meta = (ScriptConstant = "Right", ScriptConstantHost = "IntPoint"), Category = "Math|IntPoint|Constants")
+	UFUNCTION(BlueprintPure, meta = (ScriptConstant = "Right", ScriptConstantHost = "/Script/CoreUObject.IntPoint"), Category = "Math|IntPoint|Constants")
 	static FIntPoint IntPoint_Right();
 	
 	/** Down Int Point (0, 1) */
-	UFUNCTION(BlueprintPure, meta = (ScriptConstant = "Down", ScriptConstantHost = "IntPoint"), Category = "Math|IntPoint|Constants")
+	UFUNCTION(BlueprintPure, meta = (ScriptConstant = "Down", ScriptConstantHost = "/Script/CoreUObject.IntPoint"), Category = "Math|IntPoint|Constants")
 	static FIntPoint IntPoint_Down();
 
 	//
@@ -973,15 +973,15 @@ class ENGINE_API UKismetMathLibrary : public UBlueprintFunctionLibrary
 	//
 
 	/** 2D one vector constant (1,1) */
-	UFUNCTION(BlueprintPure, meta = (ScriptConstant = "One", ScriptConstantHost = "Vector2D"), Category = "Math|Vector2D")
+	UFUNCTION(BlueprintPure, meta = (ScriptConstant = "One", ScriptConstantHost = "/Script/CoreUObject.Vector2D"), Category = "Math|Vector2D")
 	static FVector2D Vector2D_One();
 
 	/** 2D unit vector constant along the 45 degree angle or symmetrical positive axes (sqrt(.5),sqrt(.5)) or (.707,.707). https://en.wikipedia.org/wiki/Unit_vector */
-	UFUNCTION(BlueprintPure, meta = (ScriptConstant = "Unit45Deg", ScriptConstantHost = "Vector2D"), Category = "Math|Vector2D")
+	UFUNCTION(BlueprintPure, meta = (ScriptConstant = "Unit45Deg", ScriptConstantHost = "/Script/CoreUObject.Vector2D"), Category = "Math|Vector2D")
 	static FVector2D Vector2D_Unit45Deg();
 
 	/** 2D zero vector constant (0,0) */
-	UFUNCTION(BlueprintPure, meta = (ScriptConstant = "Zero", ScriptConstantHost = "Vector2D"), Category = "Math|Vector2D")
+	UFUNCTION(BlueprintPure, meta = (ScriptConstant = "Zero", ScriptConstantHost = "/Script/CoreUObject.Vector2D"), Category = "Math|Vector2D")
 	static FVector2D Vector2D_Zero();
 
 
@@ -1256,35 +1256,35 @@ class ENGINE_API UKismetMathLibrary : public UBlueprintFunctionLibrary
 	//
 
 	/** 3D vector zero constant (0,0,0) */
-	UFUNCTION(BlueprintPure, meta = (ScriptConstant = "Zero", ScriptConstantHost = "Vector"), Category = "Math|Vector")
+	UFUNCTION(BlueprintPure, meta = (ScriptConstant = "Zero", ScriptConstantHost = "/Script/CoreUObject.Vector"), Category = "Math|Vector")
 	static FVector Vector_Zero();
 
 	/** 3D vector one constant (1,1,1) */
-	UFUNCTION(BlueprintPure, meta = (ScriptConstant = "One", ScriptConstantHost = "Vector"), Category = "Math|Vector")
+	UFUNCTION(BlueprintPure, meta = (ScriptConstant = "One", ScriptConstantHost = "/Script/CoreUObject.Vector"), Category = "Math|Vector")
 	static FVector Vector_One();
 
 	/** 3D vector Unreal forward direction constant (1,0,0) */
-	UFUNCTION(BlueprintPure, meta = (ScriptConstant = "Forward", ScriptConstantHost = "Vector"), Category = "Math|Vector")
+	UFUNCTION(BlueprintPure, meta = (ScriptConstant = "Forward", ScriptConstantHost = "/Script/CoreUObject.Vector"), Category = "Math|Vector")
 	static FVector Vector_Forward();
 
 	/** 3D vector Unreal backward direction constant (-1,0,0) */
-	UFUNCTION(BlueprintPure, meta = (ScriptConstant = "Backward", ScriptConstantHost = "Vector"), Category = "Math|Vector")
+	UFUNCTION(BlueprintPure, meta = (ScriptConstant = "Backward", ScriptConstantHost = "/Script/CoreUObject.Vector"), Category = "Math|Vector")
 	static FVector Vector_Backward();
 
 	/** 3D vector Unreal up direction constant (0,0,1) */
-	UFUNCTION(BlueprintPure, meta = (ScriptConstant = "Up", ScriptConstantHost = "Vector"), Category = "Math|Vector")
+	UFUNCTION(BlueprintPure, meta = (ScriptConstant = "Up", ScriptConstantHost = "/Script/CoreUObject.Vector"), Category = "Math|Vector")
 	static FVector Vector_Up();
 
 	/** 3D vector Unreal down direction constant (0,0,-1) */
-	UFUNCTION(BlueprintPure, meta = (ScriptConstant = "Down", ScriptConstantHost = "Vector"), Category = "Math|Vector")
+	UFUNCTION(BlueprintPure, meta = (ScriptConstant = "Down", ScriptConstantHost = "/Script/CoreUObject.Vector"), Category = "Math|Vector")
 	static FVector Vector_Down();
 
 	/** 3D vector Unreal right direction constant (0,1,0) */
-	UFUNCTION(BlueprintPure, meta = (ScriptConstant = "Right", ScriptConstantHost = "Vector"), Category = "Math|Vector")
+	UFUNCTION(BlueprintPure, meta = (ScriptConstant = "Right", ScriptConstantHost = "/Script/CoreUObject.Vector"), Category = "Math|Vector")
 	static FVector Vector_Right();
 
 	/** 3D vector Unreal left direction constant (0,-1,0) */
-	UFUNCTION(BlueprintPure, meta = (ScriptConstant = "Left", ScriptConstantHost = "Vector"), Category = "Math|Vector")
+	UFUNCTION(BlueprintPure, meta = (ScriptConstant = "Left", ScriptConstantHost = "/Script/CoreUObject.Vector"), Category = "Math|Vector")
 	static FVector Vector_Left();
 
 	//
@@ -2028,7 +2028,7 @@ class ENGINE_API UKismetMathLibrary : public UBlueprintFunctionLibrary
 	//
 
 	/** 4D vector zero constant (0,0,0) */
-	UFUNCTION(BlueprintPure, meta = (ScriptConstant = "Zero", ScriptConstantHost = "Vector4"), Category = "Math|Vector4")
+	UFUNCTION(BlueprintPure, meta = (ScriptConstant = "Zero", ScriptConstantHost = "/Script/CoreUObject.Vector4"), Category = "Math|Vector4")
 	static FVector4 Vector4_Zero();
 
 	//
@@ -2428,7 +2428,7 @@ class ENGINE_API UKismetMathLibrary : public UBlueprintFunctionLibrary
 	static FVector Matrix_GetOrigin(const FMatrix& InMatrix);
 
 	// Identity matrix
-	UFUNCTION(BlueprintPure, meta = (DisplayName = "Identity (Matrix)", ScriptConstant = "Identity", ScriptConstantHost = "Matrix"), Category = "Math|Matrix")
+	UFUNCTION(BlueprintPure, meta = (DisplayName = "Identity (Matrix)", ScriptConstant = "Identity", ScriptConstantHost = "/Script/CoreUObject.Matrix"), Category = "Math|Matrix")
 	static FMatrix Matrix_Identity();
 
 	/**
@@ -2727,7 +2727,7 @@ class ENGINE_API UKismetMathLibrary : public UBlueprintFunctionLibrary
 	//
 
 	/** Identity quaternion constant */
-	UFUNCTION(BlueprintPure, meta = (ScriptConstant = "Identity", ScriptConstantHost = "Quat"), Category = "Math|Quat")
+	UFUNCTION(BlueprintPure, meta = (ScriptConstant = "Identity", ScriptConstantHost = "/Script/CoreUObject.Quat"), Category = "Math|Quat")
 	static FQuat Quat_Identity();
 
 	//
@@ -2949,35 +2949,35 @@ class ENGINE_API UKismetMathLibrary : public UBlueprintFunctionLibrary
 	//
 
 	/** White linear color */
-	UFUNCTION(BlueprintPure, meta = (ScriptConstant = "White", ScriptConstantHost = "LinearColor"), Category = "Math|Color")
+	UFUNCTION(BlueprintPure, meta = (ScriptConstant = "White", ScriptConstantHost = "/Script/CoreUObject.LinearColor"), Category = "Math|Color")
 	static FLinearColor LinearColor_White();
 
 	/** Grey linear color */
-	UFUNCTION(BlueprintPure, meta = (ScriptConstant = "Gray", ScriptConstantHost = "LinearColor"), Category = "Math|Color")
+	UFUNCTION(BlueprintPure, meta = (ScriptConstant = "Gray", ScriptConstantHost = "/Script/CoreUObject.LinearColor"), Category = "Math|Color")
 	static FLinearColor LinearColor_Gray();
 
 	/** Black linear color */
-	UFUNCTION(BlueprintPure, meta = (ScriptConstant = "Black", ScriptConstantHost = "LinearColor"), Category = "Math|Color")
+	UFUNCTION(BlueprintPure, meta = (ScriptConstant = "Black", ScriptConstantHost = "/Script/CoreUObject.LinearColor"), Category = "Math|Color")
 	static FLinearColor LinearColor_Black();
 
 	/** Red linear color */
-	UFUNCTION(BlueprintPure, meta = (ScriptConstant = "Red", ScriptConstantHost = "LinearColor"), Category = "Math|Color")
+	UFUNCTION(BlueprintPure, meta = (ScriptConstant = "Red", ScriptConstantHost = "/Script/CoreUObject.LinearColor"), Category = "Math|Color")
 	static FLinearColor LinearColor_Red();
 
 	/** Green linear color */
-	UFUNCTION(BlueprintPure, meta = (ScriptConstant = "Green", ScriptConstantHost = "LinearColor"), Category = "Math|Color")
+	UFUNCTION(BlueprintPure, meta = (ScriptConstant = "Green", ScriptConstantHost = "/Script/CoreUObject.LinearColor"), Category = "Math|Color")
 	static FLinearColor LinearColor_Green();
 
 	/** Blue linear color */
-	UFUNCTION(BlueprintPure, meta = (ScriptConstant = "Blue", ScriptConstantHost = "LinearColor"), Category = "Math|Color")
+	UFUNCTION(BlueprintPure, meta = (ScriptConstant = "Blue", ScriptConstantHost = "/Script/CoreUObject.LinearColor"), Category = "Math|Color")
 	static FLinearColor LinearColor_Blue();
 
 	/** Yellow linear color */
-	UFUNCTION(BlueprintPure, meta = (ScriptConstant = "Yellow", ScriptConstantHost = "LinearColor"), Category = "Math|Color")
+	UFUNCTION(BlueprintPure, meta = (ScriptConstant = "Yellow", ScriptConstantHost = "/Script/CoreUObject.LinearColor"), Category = "Math|Color")
 	static FLinearColor LinearColor_Yellow();
 
 	/** Transparent linear color - black with 0 opacity/alpha */
-	UFUNCTION(BlueprintPure, meta = (ScriptConstant = "Transparent", ScriptConstantHost = "LinearColor"), Category = "Math|Color")
+	UFUNCTION(BlueprintPure, meta = (ScriptConstant = "Transparent", ScriptConstantHost = "/Script/CoreUObject.LinearColor"), Category = "Math|Color")
 	static FLinearColor LinearColor_Transparent();
 
 
@@ -3342,15 +3342,15 @@ class ENGINE_API UKismetMathLibrary : public UBlueprintFunctionLibrary
 	//
 
 	/** Returns the maximum time span value */
-	UFUNCTION(BlueprintPure, meta=(DisplayName="Max Value (Timespan)", ScriptConstant = "MaxValue", ScriptConstantHost = "Timespan"), Category="Math|Timespan")
+	UFUNCTION(BlueprintPure, meta=(DisplayName="Max Value (Timespan)", ScriptConstant = "MaxValue", ScriptConstantHost = "/Script/CoreUObject.Timespan"), Category="Math|Timespan")
 	static FTimespan TimespanMaxValue( );
 
 	/** Returns the minimum time span value */
-	UFUNCTION(BlueprintPure, meta=(DisplayName="Min Value (Timespan)", ScriptConstant = "MinValue", ScriptConstantHost = "Timespan"), Category="Math|Timespan")
+	UFUNCTION(BlueprintPure, meta=(DisplayName="Min Value (Timespan)", ScriptConstant = "MinValue", ScriptConstantHost = "/Script/CoreUObject.Timespan"), Category="Math|Timespan")
 	static FTimespan TimespanMinValue( );
 
 	/** Returns a zero time span value */
-	UFUNCTION(BlueprintPure, meta=(DisplayName="Zero Value (Timespan)", ScriptConstant = "Zero", ScriptConstantHost = "Timespan"), Category="Math|Timespan")
+	UFUNCTION(BlueprintPure, meta=(DisplayName="Zero Value (Timespan)", ScriptConstant = "Zero", ScriptConstantHost = "/Script/CoreUObject.Timespan"), Category="Math|Timespan")
 	static FTimespan TimespanZeroValue( );
 
 	//
diff --git a/Engine/Source/Runtime/Engine/Classes/Materials/MaterialExpressionMaterialFunctionCall.h b/Engine/Source/Runtime/Engine/Classes/Materials/MaterialExpressionMaterialFunctionCall.h
--- a/Engine/Source/Runtime/Engine/Classes/Materials/MaterialExpressionMaterialFunctionCall.h
+++ b/Engine/Source/Runtime/Engine/Classes/Materials/MaterialExpressionMaterialFunctionCall.h
@@ -80,7 +80,7 @@ class UMaterialExpressionMaterialFunctionCall : public UMaterialExpression
 	GENERATED_UCLASS_BODY()
 
 	/** The function to call. */
-	UPROPERTY(EditAnywhere, Category=MaterialExpressionMaterialFunctionCall, meta=(DisallowedClasses="MaterialFunctionInstance"))
+	UPROPERTY(EditAnywhere, Category=MaterialExpressionMaterialFunctionCall, meta=(DisallowedClasses="/Script/Engine.MaterialFunctionInstance"))
 	TObjectPtr<class UMaterialFunctionInterface> MaterialFunction;
 
 	/** Array of all the function inputs that this function exposes. */
diff --git a/Engine/Source/Runtime/Engine/Classes/Materials/MaterialInterface.h b/Engine/Source/Runtime/Engine/Classes/Materials/MaterialInterface.h
--- a/Engine/Source/Runtime/Engine/Classes/Materials/MaterialInterface.h
+++ b/Engine/Source/Runtime/Engine/Classes/Materials/MaterialInterface.h
@@ -287,7 +287,7 @@ public:
 	TArray<FMaterialTextureInfo> TextureStreamingDataMissingEntries;
 
 	/** The mesh used by the material editor to preview the material.*/
-	UPROPERTY(EditAnywhere, Category=Previewing, meta=(AllowedClasses="StaticMesh,SkeletalMesh", ExactClass="true"))
+	UPROPERTY(EditAnywhere, Category=Previewing, meta=(AllowedClasses="/Script/Engine.StaticMesh,/Script/Engine.SkeletalMesh", ExactClass="true"))
 	FSoftObjectPath PreviewMesh;
 
 	/** Information for thumbnail rendering */
diff --git a/Engine/Source/Runtime/Engine/Classes/Particles/ParticleEmitter.h b/Engine/Source/Runtime/Engine/Classes/Particles/ParticleEmitter.h
--- a/Engine/Source/Runtime/Engine/Classes/Particles/ParticleEmitter.h
+++ b/Engine/Source/Runtime/Engine/Classes/Particles/ParticleEmitter.h
@@ -203,7 +203,7 @@ class UParticleEmitter : public UObject
 	float QualityLevelSpawnRateScale;
 
 	/** Detail mode: Set flags reflecting which system detail mode you want the emitter to be ticked and rendered in */
-	UPROPERTY(EditAnywhere, Category = Particle, meta = (Bitmask, BitmaskEnum = EParticleDetailMode))
+	UPROPERTY(EditAnywhere, Category = Particle, meta = (Bitmask, BitmaskEnum = "/Script/Engine.EParticleDetailMode"))
 	uint32 DetailModeBitmask;
 
 #if WITH_EDITORONLY_DATA
diff --git a/Engine/Source/Runtime/Engine/Classes/Sound/AudioSettings.h b/Engine/Source/Runtime/Engine/Classes/Sound/AudioSettings.h
--- a/Engine/Source/Runtime/Engine/Classes/Sound/AudioSettings.h
+++ b/Engine/Source/Runtime/Engine/Classes/Sound/AudioSettings.h
@@ -96,7 +96,7 @@ struct ENGINE_API FSoundDebugEntry
 	FName DebugName;
 
 	/** Reference to a Debug Sound */
-	UPROPERTY(config, EditAnywhere, Category="Debug", meta=(AllowedClasses="SoundBase"))
+	UPROPERTY(config, EditAnywhere, Category="Debug", meta=(AllowedClasses="/Script/Engine.SoundBase"))
 	FSoftObjectPath Sound;
 };
 
@@ -106,7 +106,7 @@ struct ENGINE_API FDefaultAudioBusSettings
 	GENERATED_BODY()
 
 	/** The audio bus to start up by default on init. */
-	UPROPERTY(EditAnywhere, Category = "Mix", meta = (AllowedClasses = "AudioBus"))
+	UPROPERTY(EditAnywhere, Category = "Mix", meta = (AllowedClasses = "/Script/Engine.AudioBus"))
 	FSoftObjectPath AudioBus;
 };
 
@@ -127,39 +127,39 @@ class ENGINE_API UAudioSettings : public UDeveloperSettings
 #endif // WITH_EDITOR
 
 	/** The SoundClass assigned to newly created sounds */
-	UPROPERTY(config, EditAnywhere, Category="Audio", meta=(AllowedClasses="SoundClass", DisplayName="Default Sound Class"))
+	UPROPERTY(config, EditAnywhere, Category="Audio", meta=(AllowedClasses="/Script/Engine.SoundClass", DisplayName="Default Sound Class"))
 	FSoftObjectPath DefaultSoundClassName;
 
 	/** The SoundClass assigned to media player assets */
-	UPROPERTY(config, EditAnywhere, Category = "Audio", meta = (AllowedClasses = "SoundClass", DisplayName = "Default Media Sound Class"))
+	UPROPERTY(config, EditAnywhere, Category = "Audio", meta = (AllowedClasses = "/Script/Engine.SoundClass", DisplayName = "Default Media Sound Class"))
 	FSoftObjectPath DefaultMediaSoundClassName;
 
 	/** The SoundConcurrency assigned to newly created sounds */
-	UPROPERTY(config, EditAnywhere, Category = "Audio", meta = (AllowedClasses = "SoundConcurrency", DisplayName = "Default Sound Concurrency"))
+	UPROPERTY(config, EditAnywhere, Category = "Audio", meta = (AllowedClasses = "/Script/Engine.SoundConcurrency", DisplayName = "Default Sound Concurrency"))
 	FSoftObjectPath DefaultSoundConcurrencyName;
 
 	/** The SoundMix to use as base when no other system has specified a Base SoundMix */
-	UPROPERTY(config, EditAnywhere, Category="Audio", meta=(AllowedClasses="SoundMix"))
+	UPROPERTY(config, EditAnywhere, Category="Audio", meta=(AllowedClasses="/Script/Engine.SoundMix"))
 	FSoftObjectPath DefaultBaseSoundMix;
 
 	/** Sound class to be used for the VOIP audio component */
-	UPROPERTY(config, EditAnywhere, Category="Audio", meta=(AllowedClasses="SoundClass", DisplayName = "VOIP Sound Class"))
+	UPROPERTY(config, EditAnywhere, Category="Audio", meta=(AllowedClasses="/Script/Engine.SoundClass", DisplayName = "VOIP Sound Class"))
 	FSoftObjectPath VoiPSoundClass;
 
 	/** The default submix through which all sounds are routed to. The root submix that outputs to audio hardware. */
-	UPROPERTY(config, EditAnywhere, Category="Mix", meta=(AllowedClasses="SoundSubmix"))
+	UPROPERTY(config, EditAnywhere, Category="Mix", meta=(AllowedClasses="/Script/Engine.SoundSubmix"))
 	FSoftObjectPath MasterSubmix;
 
 	/** The default submix to use for implicit submix sends (i.e. if the base submix send is null or if a submix parent is null) */
-	UPROPERTY(config, EditAnywhere, Category = "Mix", meta = (AllowedClasses = "SoundSubmix"), AdvancedDisplay)
+	UPROPERTY(config, EditAnywhere, Category = "Mix", meta = (AllowedClasses = "/Script/Engine.SoundSubmix"), AdvancedDisplay)
 	FSoftObjectPath BaseDefaultSubmix;
 
 	/** The submix through which all sounds set to use reverb are routed */
-	UPROPERTY(config, EditAnywhere, Category="Mix", meta=(AllowedClasses="SoundSubmix"))
+	UPROPERTY(config, EditAnywhere, Category="Mix", meta=(AllowedClasses="/Script/Engine.SoundSubmix"))
 	FSoftObjectPath ReverbSubmix;
 
 	/** The submix through which all sounds set to use legacy EQ system are routed */
-	UPROPERTY(config, EditAnywhere, Category="Mix", meta=(AllowedClasses="SoundSubmix", DisplayName = "EQ Submix (Legacy)"), AdvancedDisplay)
+	UPROPERTY(config, EditAnywhere, Category="Mix", meta=(AllowedClasses="/Script/Engine.SoundSubmix", DisplayName = "EQ Submix (Legacy)"), AdvancedDisplay)
 	FSoftObjectPath EQSubmix;
 
 	/** Sample rate used for voice over IP. VOIP audio is resampled to the application's sample rate on the receiver side. */
diff --git a/Engine/Source/Runtime/Engine/Classes/Tests/AutomationTestSettings.h b/Engine/Source/Runtime/Engine/Classes/Tests/AutomationTestSettings.h
--- a/Engine/Source/Runtime/Engine/Classes/Tests/AutomationTestSettings.h
+++ b/Engine/Source/Runtime/Engine/Classes/Tests/AutomationTestSettings.h
@@ -306,7 +306,7 @@ struct FEditorMapPerformanceTestDefinition
 	GENERATED_USTRUCT_BODY()
 	
 	/** Map to be used for the Performance Capture **/
-	UPROPERTY(config, EditAnywhere, Category = Automation, meta=( AllowedClasses="World" ))
+	UPROPERTY(config, EditAnywhere, Category = Automation, meta=( AllowedClasses="/Script/Engine.World" ))
 	FSoftObjectPath PerformanceTestmap;
 
 	/** How long is this test expected to run before stopping **/
@@ -358,7 +358,7 @@ public:
 	/**
 	 * The automation test map to be used for several of the automation tests.
 	 */
-	UPROPERTY(config, EditAnywhere, Category = Automation, meta=( AllowedClasses="World" ))
+	UPROPERTY(config, EditAnywhere, Category = Automation, meta=( AllowedClasses="/Script/Engine.World" ))
 	FSoftObjectPath AutomationTestmap;
 
 	/**
diff --git a/Engine/Source/Runtime/Engine/Private/Materials/HLSLMaterialTranslator.cpp b/Engine/Source/Runtime/Engine/Private/Materials/HLSLMaterialTranslator.cpp
--- a/Engine/Source/Runtime/Engine/Private/Materials/HLSLMaterialTranslator.cpp
+++ b/Engine/Source/Runtime/Engine/Private/Materials/HLSLMaterialTranslator.cpp
@@ -1570,7 +1570,7 @@ void FHLSLMaterialTranslator::ValidateShadingModelsForFeatureLevel(const FMateri
 				GetFeatureLevelName(FeatureLevel, FeatureLevelName);
 
 				FString ShadingModelName;
-				const UEnum* EnumPtr = FindObject<UEnum>(ANY_PACKAGE, TEXT("EMaterialShadingModel"), true);
+				const UEnum* EnumPtr = FindObject<UEnum>(nullptr, TEXT("/Script/Engine.EMaterialShadingModel"), true);
 				if (EnumPtr)
 				{
 					ShadingModelName = EnumPtr->GetNameStringByValue(InvalidShadingModel);
diff --git a/Engine/Source/Runtime/Engine/Private/PackedLevelActor/PackedLevelActorBuilder.cpp b/Engine/Source/Runtime/Engine/Private/PackedLevelActor/PackedLevelActorBuilder.cpp
--- a/Engine/Source/Runtime/Engine/Private/PackedLevelActor/PackedLevelActorBuilder.cpp
+++ b/Engine/Source/Runtime/Engine/Private/PackedLevelActor/PackedLevelActorBuilder.cpp
@@ -142,7 +142,7 @@ TSharedPtr<FPackedLevelActorBuilder> FPackedLevelActorBuilder::CreateDefaultBuil
 	// Components or Actor classes in this set will not generate warnings
 	Builder->ClassDiscards.Add(ALevelBounds::StaticClass());
 	// Avoid dependency (find class)
-	UClass* ChaosDebugClass = FindObject<UClass>(ANY_PACKAGE, TEXT("ChaosDebugDrawComponent"));
+	UClass* ChaosDebugClass = FindObject<UClass>(nullptr, TEXT("/Script/ChaosSolverEngine.ChaosDebugDrawComponent"));
 	if (ChaosDebugClass)
 	{
 		Builder->ClassDiscards.Add(ChaosDebugClass);
diff --git a/Engine/Source/Runtime/Engine/Private/Particles/ParticleSystemManager.cpp b/Engine/Source/Runtime/Engine/Private/Particles/ParticleSystemManager.cpp
--- a/Engine/Source/Runtime/Engine/Private/Particles/ParticleSystemManager.cpp
+++ b/Engine/Source/Runtime/Engine/Private/Particles/ParticleSystemManager.cpp
@@ -174,7 +174,7 @@ void FParticleSystemWorldManager::OnWorldInit(UWorld* World, const UWorld::Initi
 #if !UE_BUILD_SHIPPING
 	if (TickGroupEnum == nullptr)
 	{
-		TickGroupEnum = FindObjectChecked<UEnum>(ANY_PACKAGE, TEXT("ETickingGroup"));
+		TickGroupEnum = FindObjectChecked<UEnum>(nullptr, TEXT("/Script/Engine.ETickingGroup"));
 	}
 #endif
 	FParticleSystemWorldManager* NewWorldMan = new FParticleSystemWorldManager(World);
@@ -831,7 +831,7 @@ void FParticleSystemWorldManagerTickFunction::ExecuteTick(float DeltaTime, enum
 
 FString FParticleSystemWorldManagerTickFunction::DiagnosticMessage()
 {
-	static const UEnum* EnumType = FindObjectChecked<UEnum>(ANY_PACKAGE, TEXT("ETickingGroup"));
+	static const UEnum* EnumType = FindObjectChecked<UEnum>(nullptr, TEXT("/Script/Engine.ETickingGroup"));
 
 	return TEXT("FParticleSystemManager::Tick(") + EnumType->GetNameStringByIndex(static_cast<uint32>(TickGroup)) + TEXT(")");
 }
diff --git a/Engine/Source/Runtime/Engine/Private/TickTaskManager.cpp b/Engine/Source/Runtime/Engine/Private/TickTaskManager.cpp
--- a/Engine/Source/Runtime/Engine/Private/TickTaskManager.cpp
+++ b/Engine/Source/Runtime/Engine/Private/TickTaskManager.cpp
@@ -1205,7 +1205,7 @@ public:
 	/** Dumps all tick functions to output device. */
 	void DumpAllTickFunctions(FOutputDevice& Ar, int32& EnabledCount, int32& DisabledCount, bool bEnabled, bool bDisabled)
 	{
-		UEnum* TickGroupEnum = CastChecked<UEnum>(StaticFindObject(UEnum::StaticClass(), ANY_PACKAGE, TEXT("ETickingGroup"), true));
+		UEnum* TickGroupEnum = CastChecked<UEnum>(StaticFindObject(UEnum::StaticClass(), nullptr, TEXT("/Script/Engine.ETickingGroup"), true));
 		if (bEnabled)
 		{
 			for (TSet<FTickFunction*>::TIterator It(AllEnabledTickFunctions); It; ++It)
diff --git a/Engine/Source/Runtime/Engine/Private/WorldPartition/WorldPartition.cpp b/Engine/Source/Runtime/Engine/Private/WorldPartition/WorldPartition.cpp
--- a/Engine/Source/Runtime/Engine/Private/WorldPartition/WorldPartition.cpp
+++ b/Engine/Source/Runtime/Engine/Private/WorldPartition/WorldPartition.cpp
@@ -845,7 +845,7 @@ UWorldPartition* UWorldPartition::CreateOrRepairWorldPartition(AWorldSettings* W
 	{
 		if (!EditorHashClass)
 		{
-			EditorHashClass = FindObject<UClass>(ANY_PACKAGE, TEXT("WorldPartitionEditorSpatialHash"));
+			EditorHashClass = FindObject<UClass>(nullptr, TEXT("/Script/Engine.WorldPartitionEditorSpatialHash"));
 		}
 
 		WorldPartition->EditorHash = NewObject<UWorldPartitionEditorHash>(WorldPartition, EditorHashClass);
@@ -856,7 +856,7 @@ UWorldPartition* UWorldPartition::CreateOrRepairWorldPartition(AWorldSettings* W
 	{
 		if (!RuntimeHashClass)
 		{
-			RuntimeHashClass = FindObject<UClass>(ANY_PACKAGE, TEXT("WorldPartitionRuntimeSpatialHash"));
+			RuntimeHashClass = FindObject<UClass>(nullptr, TEXT("/Script/Engine.WorldPartitionRuntimeSpatialHash"));
 		}
 
 		WorldPartition->RuntimeHash = NewObject<UWorldPartitionRuntimeHash>(WorldPartition, RuntimeHashClass, NAME_None, RF_Transactional);
diff --git a/Engine/Source/Runtime/Engine/Private/WorldPartition/WorldPartitionSubsystem.cpp b/Engine/Source/Runtime/Engine/Private/WorldPartition/WorldPartitionSubsystem.cpp
--- a/Engine/Source/Runtime/Engine/Private/WorldPartition/WorldPartitionSubsystem.cpp
+++ b/Engine/Source/Runtime/Engine/Private/WorldPartition/WorldPartitionSubsystem.cpp
@@ -118,7 +118,7 @@ const UWorldPartition* UWorldPartitionSubsystem::GetWorldPartition() const
 #if WITH_EDITOR
 bool UWorldPartitionSubsystem::IsRunningConvertWorldPartitionCommandlet() const
 {
-	static UClass* WorldPartitionConvertCommandletClass = FindObject<UClass>(ANY_PACKAGE, TEXT("WorldPartitionConvertCommandlet"), true);
+	static UClass* WorldPartitionConvertCommandletClass = FindObject<UClass>(nullptr, TEXT("/Script/UnrealEd.WorldPartitionConvertCommandlet"), true);
 	check(WorldPartitionConvertCommandletClass);
 	static const bool bIsRunningWorldPartitionConvertCommandlet = GetRunningCommandletClass() && GetRunningCommandletClass()->IsChildOf(WorldPartitionConvertCommandletClass);
 	return bIsRunningWorldPartitionConvertCommandlet;
diff --git a/Engine/Source/Runtime/Engine/Public/WorldPartition/WorldPartitionEditorSpatialHash.h b/Engine/Source/Runtime/Engine/Public/WorldPartition/WorldPartitionEditorSpatialHash.h
--- a/Engine/Source/Runtime/Engine/Public/WorldPartition/WorldPartitionEditorSpatialHash.h
+++ b/Engine/Source/Runtime/Engine/Public/WorldPartition/WorldPartitionEditorSpatialHash.h
@@ -220,7 +220,7 @@ private:
 	bool bBoundsDirty;	
 
 public:
-	UPROPERTY(Config, meta = (AllowedClasses = "Texture2D, MaterialInterface"))
+	UPROPERTY(Config, meta = (AllowedClasses = "/Script/Engine.Texture2D, /Script/Engine.MaterialInterface"))
 	FSoftObjectPath WorldImage;
 
 	UPROPERTY(Config)
diff --git a/Engine/Source/Runtime/EngineSettings/Classes/GameMapsSettings.h b/Engine/Source/Runtime/EngineSettings/Classes/GameMapsSettings.h
--- a/Engine/Source/Runtime/EngineSettings/Classes/GameMapsSettings.h
+++ b/Engine/Source/Runtime/EngineSettings/Classes/GameMapsSettings.h
@@ -48,11 +48,11 @@ struct FGameModeName
 	GENERATED_USTRUCT_BODY()
 
 	/** Abbreviation/prefix that can be used as an alias for the class name */
-	UPROPERTY(EditAnywhere, Category = DefaultModes, meta = (MetaClass = "GameModeBase"))
+	UPROPERTY(EditAnywhere, Category = DefaultModes, meta = (MetaClass = "/Script/Engine.GameModeBase"))
 	FString Name;
 
 	/** GameMode class to load */
-	UPROPERTY(EditAnywhere, Category = DefaultModes, meta = (MetaClass = "GameModeBase"))
+	UPROPERTY(EditAnywhere, Category = DefaultModes, meta = (MetaClass = "/Script/Engine.GameModeBase"))
 	FSoftClassPath GameMode;
 };
 
@@ -148,7 +148,7 @@ public:
 
 #if WITH_EDITORONLY_DATA
 	/** If set, this map will be loaded when the Editor starts up. */
-	UPROPERTY(config, EditAnywhere, Category=DefaultMaps, meta=(AllowedClasses="World"))
+	UPROPERTY(config, EditAnywhere, Category=DefaultMaps, meta=(AllowedClasses="/Script/Engine.World"))
 	FSoftObjectPath EditorStartupMap;
 
 	/** Map templates that should show up in the new level dialog. This will completely override the default maps chosen by the default editor */
@@ -161,7 +161,7 @@ public:
 	FString LocalMapOptions;
 
 	/** The map loaded when transition from one map to another. */
-	UPROPERTY(config, EditAnywhere, Category=DefaultMaps, AdvancedDisplay, meta=(AllowedClasses="World"))
+	UPROPERTY(config, EditAnywhere, Category=DefaultMaps, AdvancedDisplay, meta=(AllowedClasses="/Script/Engine.World"))
 	FSoftObjectPath TransitionMap;
 
 	/** Whether the screen should be split or not when multiple local players are present */
@@ -188,28 +188,28 @@ public:
 	bool bOffsetPlayerGamepadIds;
 
 	/** The class to use when instantiating the transient GameInstance class */
-	UPROPERTY(config, noclear, EditAnywhere, Category=GameInstance, meta=(MetaClass="GameInstance"))
+	UPROPERTY(config, noclear, EditAnywhere, Category=GameInstance, meta=(MetaClass="/Script/Engine.GameInstance"))
 	FSoftClassPath GameInstanceClass;
 
 private:
 
 	/** The map that will be loaded by default when no other map is loaded. */
-	UPROPERTY(config, EditAnywhere, Category=DefaultMaps, meta=(AllowedClasses="World"))
+	UPROPERTY(config, EditAnywhere, Category=DefaultMaps, meta=(AllowedClasses="/Script/Engine.World"))
 	FSoftObjectPath GameDefaultMap;
 
 	/** The map that will be loaded by default when no other map is loaded (DEDICATED SERVER). */
-	UPROPERTY(config, EditAnywhere, Category=DefaultMaps, AdvancedDisplay, meta=(AllowedClasses="World"))
+	UPROPERTY(config, EditAnywhere, Category=DefaultMaps, AdvancedDisplay, meta=(AllowedClasses="/Script/Engine.World"))
 	FSoftObjectPath ServerDefaultMap;
 
 	/** GameMode to use if not specified in any other way. (e.g. per-map DefaultGameMode or on the URL). */
-	UPROPERTY(config, noclear, EditAnywhere, Category=DefaultModes, meta=(MetaClass="GameModeBase", DisplayName="Default GameMode"))
+	UPROPERTY(config, noclear, EditAnywhere, Category=DefaultModes, meta=(MetaClass="/Script/Engine.GameModeBase", DisplayName="Default GameMode"))
 	FSoftClassPath GlobalDefaultGameMode;
 
 	/**
 	 * GameMode to use if not specified in any other way. (e.g. per-map DefaultGameMode or on the URL) (DEDICATED SERVERS)
 	 * If not set, the GlobalDefaultGameMode value will be used.
 	 */
-	UPROPERTY(config, EditAnywhere, Category=DefaultModes, meta=(MetaClass="GameModeBase"), AdvancedDisplay)
+	UPROPERTY(config, EditAnywhere, Category=DefaultModes, meta=(MetaClass="/Script/Engine.GameModeBase"), AdvancedDisplay)
 	FSoftClassPath GlobalDefaultServerGameMode;
 
 	/** Overrides the GameMode to use when loading a map that starts with a specific prefix */
diff --git a/Engine/Source/Runtime/Experimental/ChaosSolverEngine/Public/Chaos/ChaosSolverSettings.h b/Engine/Source/Runtime/Experimental/ChaosSolverEngine/Public/Chaos/ChaosSolverSettings.h
--- a/Engine/Source/Runtime/Experimental/ChaosSolverEngine/Public/Chaos/ChaosSolverSettings.h
+++ b/Engine/Source/Runtime/Experimental/ChaosSolverEngine/Public/Chaos/ChaosSolverSettings.h
@@ -23,7 +23,7 @@ public:
 	virtual UClass* GetSolverActorClass() const;
 
 	/** The class to use when auto-creating a default chaos solver actor */
-	UPROPERTY(config, noclear, EditAnywhere, Category = GameInstance, meta = (MetaClass = "ChaosSolverActor"))
+	UPROPERTY(config, noclear, EditAnywhere, Category = GameInstance, meta = (MetaClass = "/Script/ChaosSolverEngine.ChaosSolverActor"))
 	FSoftClassPath DefaultChaosSolverActorClass;
 
 #if WITH_EDITOR
diff --git a/Engine/Source/Runtime/Experimental/GeometryCollectionEngine/Public/GeometryCollection/GeometryCollectionObject.h b/Engine/Source/Runtime/Experimental/GeometryCollectionEngine/Public/GeometryCollection/GeometryCollectionObject.h
--- a/Engine/Source/Runtime/Experimental/GeometryCollectionEngine/Public/GeometryCollection/GeometryCollectionObject.h
+++ b/Engine/Source/Runtime/Experimental/GeometryCollectionEngine/Public/GeometryCollection/GeometryCollectionObject.h
@@ -24,7 +24,7 @@ struct GEOMETRYCOLLECTIONENGINE_API FGeometryCollectionSource
 {
 	GENERATED_BODY()
 
-	UPROPERTY(BlueprintReadWrite, EditAnywhere, Category = "GeometrySource", meta=(AllowedClasses="StaticMesh, SkeletalMesh, GeometryCollection"))
+	UPROPERTY(BlueprintReadWrite, EditAnywhere, Category = "GeometrySource", meta=(AllowedClasses="/Script/Engine.StaticMesh, /Script/Engine.SkeletalMesh, /Script/GeometryCollectionEngine.GeometryCollection"))
 	FSoftObjectPath SourceGeometryObject;
 
 	UPROPERTY(BlueprintReadWrite, EditAnywhere, Category = "GeometrySource")
@@ -53,7 +53,7 @@ struct GEOMETRYCOLLECTIONENGINE_API FGeometryCollectionEmbeddedExemplar
 		, InstanceCount(0)
 	{ }
 
-	UPROPERTY(EditAnywhere, Category = "EmbeddedExemplar", meta = (AllowedClasses = "StaticMesh"))
+	UPROPERTY(EditAnywhere, Category = "EmbeddedExemplar", meta = (AllowedClasses = "/Script/Engine.StaticMesh"))
 	FSoftObjectPath StaticMeshExemplar;
 
 	UPROPERTY(EditAnywhere, Category = "EmbeddedExemplar")
diff --git a/Engine/Source/Runtime/Foliage/Public/FoliageTypeObject.h b/Engine/Source/Runtime/Foliage/Public/FoliageTypeObject.h
--- a/Engine/Source/Runtime/Foliage/Public/FoliageTypeObject.h
+++ b/Engine/Source/Runtime/Foliage/Public/FoliageTypeObject.h
@@ -42,7 +42,7 @@ struct FOLIAGE_API FFoliageTypeObject
 
 private:
 	/** The foliage type that will be spawned by the procedural foliage simulation */
-	UPROPERTY(Category = ProceduralFoliageSimulation, EditAnywhere, meta=(AllowedClasses="FoliageType_InstancedStaticMesh,FoliageType_Actor,Blueprint", DisplayThumbnail="true", ThumbnailSize="X=40 Y=40"))
+	UPROPERTY(Category = ProceduralFoliageSimulation, EditAnywhere, meta=(AllowedClasses="/Script/Foliage.FoliageType_InstancedStaticMesh,/Script/Foliage.FoliageType_Actor,/Script/Engine.Blueprint", DisplayThumbnail="true", ThumbnailSize="X=40 Y=40"))
 	TObjectPtr<UObject> FoliageTypeObject;
 
 	/** The actual instance of the foliage type that is used for spawning */
diff --git a/Engine/Source/Runtime/GameplayTags/Classes/GameplayTagContainer.h b/Engine/Source/Runtime/GameplayTags/Classes/GameplayTagContainer.h
--- a/Engine/Source/Runtime/GameplayTags/Classes/GameplayTagContainer.h
+++ b/Engine/Source/Runtime/GameplayTags/Classes/GameplayTagContainer.h
@@ -46,7 +46,7 @@ typedef uint16 FGameplayTagNetIndex;
  * A single gameplay tag, which represents a hierarchical name of the form x.y that is registered in the GameplayTagsManager
  * You can filter the gameplay tags displayed in the editor using, meta = (Categories = "Tag1.Tag2.Tag3"))
  */
-USTRUCT(BlueprintType, meta = (HasNativeMake = "GameplayTags.BlueprintGameplayTagLibrary.MakeLiteralGameplayTag", HasNativeBreak = "GameplayTags.BlueprintGameplayTagLibrary.GetTagName", DisableSplitPin))
+USTRUCT(BlueprintType, meta = (HasNativeMake = "/Script/GameplayTags.BlueprintGameplayTagLibrary.MakeLiteralGameplayTag", HasNativeBreak = "/Script/GameplayTags.BlueprintGameplayTagLibrary.GetTagName", DisableSplitPin))
 struct GAMEPLAYTAGS_API FGameplayTag
 {
 	GENERATED_USTRUCT_BODY()
@@ -230,7 +230,7 @@ struct TStructOpsTypeTraits< FGameplayTag > : public TStructOpsTypeTraitsBase2<
 };
 
 /** A Tag Container holds a collection of FGameplayTags, tags are included explicitly by adding them, and implicitly from adding child tags */
-USTRUCT(BlueprintType, meta = (HasNativeMake = "GameplayTags.BlueprintGameplayTagLibrary.MakeGameplayTagContainerFromArray", HasNativeBreak = "GameplayTags.BlueprintGameplayTagLibrary.BreakGameplayTagContainer"))
+USTRUCT(BlueprintType, meta = (HasNativeMake = "/Script/GameplayTags.BlueprintGameplayTagLibrary.MakeGameplayTagContainerFromArray", HasNativeBreak = "/Script/GameplayTags.BlueprintGameplayTagLibrary.BreakGameplayTagContainer"))
 struct GAMEPLAYTAGS_API FGameplayTagContainer
 {
 	GENERATED_USTRUCT_BODY()
@@ -793,7 +793,7 @@ namespace EGameplayTagQueryStreamVersion
  * Queries are internally represented as a byte stream that is memory-efficient and can be evaluated quickly at runtime.
  * Note: these have an extensive details and graph pin customization for editing, so there is no need to expose the internals to Blueprints.
  */
-USTRUCT(BlueprintType, meta=(HasNativeMake="GameplayTags.BlueprintGameplayTagLibrary.MakeGameplayTagQuery"))
+USTRUCT(BlueprintType, meta=(HasNativeMake="/Script/GameplayTags.BlueprintGameplayTagLibrary.MakeGameplayTagQuery"))
 struct GAMEPLAYTAGS_API FGameplayTagQuery
 {
 	GENERATED_BODY();
diff --git a/Engine/Source/Runtime/GameplayTags/Classes/GameplayTagsSettings.h b/Engine/Source/Runtime/GameplayTags/Classes/GameplayTagsSettings.h
--- a/Engine/Source/Runtime/GameplayTags/Classes/GameplayTagsSettings.h
+++ b/Engine/Source/Runtime/GameplayTags/Classes/GameplayTagsSettings.h
@@ -125,7 +125,7 @@ class GAMEPLAYTAGS_API UGameplayTagsSettings : public UGameplayTagsList
 	TArray<FGameplayTagCategoryRemap> CategoryRemapping;
 
 	/** List of data tables to load tags from */
-	UPROPERTY(config, EditAnywhere, Category = GameplayTags, meta = (AllowedClasses = "DataTable"))
+	UPROPERTY(config, EditAnywhere, Category = GameplayTags, meta = (AllowedClasses = "/Script/Engine.DataTable"))
 	TArray<FSoftObjectPath> GameplayTagTableList;
 
 	/** List of active tag redirects */
diff --git a/Engine/Source/Runtime/LevelSequence/Public/LevelSequenceActor.h b/Engine/Source/Runtime/LevelSequence/Public/LevelSequenceActor.h
--- a/Engine/Source/Runtime/LevelSequence/Public/LevelSequenceActor.h
+++ b/Engine/Source/Runtime/LevelSequence/Public/LevelSequenceActor.h
@@ -42,7 +42,7 @@ public:
 	UPROPERTY(config, EditAnywhere, BlueprintReadWrite, Category="General")
 	bool bUseBurnIn;
 
-	UPROPERTY(config, EditAnywhere, BlueprintReadWrite, Category="General", meta=(EditCondition=bUseBurnIn, MetaClass="LevelSequenceBurnIn"))
+	UPROPERTY(config, EditAnywhere, BlueprintReadWrite, Category="General", meta=(EditCondition=bUseBurnIn, MetaClass="/Script/LevelSequence.LevelSequenceBurnIn"))
 	FSoftClassPath BurnInClass;
 
 	UPROPERTY(Instanced, EditAnywhere, BlueprintReadWrite, Category="General", meta=(EditCondition=bUseBurnIn))
@@ -82,7 +82,7 @@ public:
 	UPROPERTY(Instanced, transient, replicated, BlueprintReadOnly, BlueprintGetter=GetSequencePlayer, Category="Playback", meta=(ExposeFunctionCategories="Sequencer|Player"))
 	TObjectPtr<ULevelSequencePlayer> SequencePlayer;
 
-	UPROPERTY(EditAnywhere, replicated, BlueprintReadOnly, Category="General", meta=(AllowedClasses="LevelSequence"))
+	UPROPERTY(EditAnywhere, replicated, BlueprintReadOnly, Category="General", meta=(AllowedClasses="/Script/LevelSequence.LevelSequence"))
 	TObjectPtr<ULevelSequence> LevelSequenceAsset;
 
 #if WITH_EDITORONLY_DATA
diff --git a/Engine/Source/Runtime/LiveLinkInterface/Public/LiveLinkCurveRemapSettings.h b/Engine/Source/Runtime/LiveLinkInterface/Public/LiveLinkCurveRemapSettings.h
--- a/Engine/Source/Runtime/LiveLinkInterface/Public/LiveLinkCurveRemapSettings.h
+++ b/Engine/Source/Runtime/LiveLinkInterface/Public/LiveLinkCurveRemapSettings.h
@@ -13,7 +13,7 @@ struct FLiveLinkCurveConversionSettings
 {
 	GENERATED_USTRUCT_BODY()
 
-	UPROPERTY(EditAnywhere, Category = Settings, meta = (AllowedClasses = "PoseAsset"))
+	UPROPERTY(EditAnywhere, Category = Settings, meta = (AllowedClasses = "/Script/Engine.PoseAsset"))
 	TMap<FString, FSoftObjectPath> CurveConversionAssetMap;
 };
 
diff --git a/Engine/Source/Runtime/MovieScene/Public/Evaluation/MovieSceneSequenceHierarchy.h b/Engine/Source/Runtime/MovieScene/Public/Evaluation/MovieSceneSequenceHierarchy.h
--- a/Engine/Source/Runtime/MovieScene/Public/Evaluation/MovieSceneSequenceHierarchy.h
+++ b/Engine/Source/Runtime/MovieScene/Public/Evaluation/MovieSceneSequenceHierarchy.h
@@ -64,7 +64,7 @@ struct FMovieSceneSubSequenceData
 	MOVIESCENE_API FMovieSceneSectionParameters ToSubSectionParameters() const;
 
 	/** The sequence that the sub section references */
-	UPROPERTY(meta=(AllowedClasses="MovieSceneSequence"))
+	UPROPERTY(meta=(AllowedClasses="/Script/MovieScene.MovieSceneSequence"))
 	FSoftObjectPath Sequence;
 
 	/** The transform from this sub sequence's parent to its own play space. */
diff --git a/Engine/Source/Runtime/MovieScene/Public/Sections/MovieSceneSubSection.h b/Engine/Source/Runtime/MovieScene/Public/Sections/MovieSceneSubSection.h
--- a/Engine/Source/Runtime/MovieScene/Public/Sections/MovieSceneSubSection.h
+++ b/Engine/Source/Runtime/MovieScene/Public/Sections/MovieSceneSubSection.h
@@ -149,7 +149,7 @@ private:
 	UPROPERTY()
 	float PrerollTime_DEPRECATED;
 
-	UPROPERTY(EditAnywhere, Category="Networking", meta=(Bitmask, BitmaskEnum=EMovieSceneServerClientMask))
+	UPROPERTY(EditAnywhere, Category="Networking", meta=(Bitmask, BitmaskEnum="/Script/MovieScene.EMovieSceneServerClientMask"))
 	uint8 NetworkMask;
 
 protected:
diff --git a/Engine/Source/Runtime/MovieSceneCapture/Public/MovieSceneCapture.h b/Engine/Source/Runtime/MovieSceneCapture/Public/MovieSceneCapture.h
--- a/Engine/Source/Runtime/MovieSceneCapture/Public/MovieSceneCapture.h
+++ b/Engine/Source/Runtime/MovieSceneCapture/Public/MovieSceneCapture.h
@@ -80,11 +80,11 @@ protected:
 public:
 
 	/** The type of capture protocol to use for image data */
-	UPROPERTY(config, EditAnywhere, NoClear, Category=CaptureSettings, DisplayName="Image Output Format", meta=(MetaClass="MovieSceneImageCaptureProtocolBase", HideViewOptions, ShowDisplayNames))
+	UPROPERTY(config, EditAnywhere, NoClear, Category=CaptureSettings, DisplayName="Image Output Format", meta=(MetaClass="/Script/MovieSceneCapture.MovieSceneImageCaptureProtocolBase", HideViewOptions, ShowDisplayNames))
 	FSoftClassPath ImageCaptureProtocolType;
 	
 	/** The type of capture protocol to use for audio data. Requires experimental audio mixer (launch editor via with -audiomixer). */
-	UPROPERTY(config, EditAnywhere, NoClear, Category=CaptureSettings, DisplayName="Audio Output Format", meta=(MetaClass="MovieSceneAudioCaptureProtocolBase", HideViewOptions, ShowDisplayNames))
+	UPROPERTY(config, EditAnywhere, NoClear, Category=CaptureSettings, DisplayName="Audio Output Format", meta=(MetaClass="/Script/MovieSceneCapture.MovieSceneAudioCaptureProtocolBase", HideViewOptions, ShowDisplayNames))
 	FSoftClassPath AudioCaptureProtocolType;
 
 	/** Capture protocol responsible for actually capturing frame data */
diff --git a/Engine/Source/Runtime/MovieSceneTracks/Public/Sections/MovieSceneSkeletalAnimationSection.h b/Engine/Source/Runtime/MovieSceneTracks/Public/Sections/MovieSceneSkeletalAnimationSection.h
--- a/Engine/Source/Runtime/MovieSceneTracks/Public/Sections/MovieSceneSkeletalAnimationSection.h
+++ b/Engine/Source/Runtime/MovieSceneTracks/Public/Sections/MovieSceneSkeletalAnimationSection.h
@@ -28,7 +28,7 @@ struct FMovieSceneSkeletalAnimationParams
 	float GetSequenceLength() const { return Animation != nullptr ? Animation->GetPlayLength() : 0.f; }
 
 	/** The animation this section plays */
-	UPROPERTY(BlueprintReadWrite, EditAnywhere, Category="Animation", meta=(AllowedClasses = "AnimSequence,AnimComposite,AnimStreamable"))
+	UPROPERTY(BlueprintReadWrite, EditAnywhere, Category="Animation", meta=(AllowedClasses = "/Script/Engine.AnimSequence,/Script/Engine.AnimComposite,/Script/Engine.AnimStreamable"))
 	TObjectPtr<UAnimSequenceBase> Animation;
 
 	/** The offset into the beginning of the animation clip for the first loop of play. */
diff --git a/Engine/Source/Runtime/Net/Core/Private/Net/Core/Connection/NetConnectionFaultRecoveryBase.cpp b/Engine/Source/Runtime/Net/Core/Private/Net/Core/Connection/NetConnectionFaultRecoveryBase.cpp
--- a/Engine/Source/Runtime/Net/Core/Private/Net/Core/Connection/NetConnectionFaultRecoveryBase.cpp
+++ b/Engine/Source/Runtime/Net/Core/Private/Net/Core/Connection/NetConnectionFaultRecoveryBase.cpp
@@ -28,7 +28,7 @@ const TCHAR* FNetFaultState::GetConfigSection()
 
 UClass* FNetFaultState::GetBaseConfigClass()
 {
-	return FindObject<UClass>(ANY_PACKAGE, TEXT("NetFaultConfig"));
+	return FindObject<UClass>(nullptr, TEXT("/Script/Engine.NetFaultConfig"));
 }
 
 
diff --git a/Engine/Source/Runtime/SlateCore/Public/Fonts/SlateFontInfo.h b/Engine/Source/Runtime/SlateCore/Public/Fonts/SlateFontInfo.h
--- a/Engine/Source/Runtime/SlateCore/Public/Fonts/SlateFontInfo.h
+++ b/Engine/Source/Runtime/SlateCore/Public/Fonts/SlateFontInfo.h
@@ -50,7 +50,7 @@ struct SLATECORE_API FFontOutlineSettings
 	bool bApplyOutlineToDropShadows;
 
 	/** Optional material to apply to the outline */
-	UPROPERTY(EditAnywhere, BlueprintReadWrite, Category=SlateStyleRules, meta=(AllowedClasses="MaterialInterface"))
+	UPROPERTY(EditAnywhere, BlueprintReadWrite, Category=SlateStyleRules, meta=(AllowedClasses="/Script/Engine.MaterialInterface"))
 	TObjectPtr<UObject> OutlineMaterial;
 
 	/** The color of the outline for any character in this font */
@@ -134,11 +134,11 @@ struct SLATECORE_API FSlateFontInfo
 	GENERATED_USTRUCT_BODY()
 
 	/** The font object (valid when used from UMG or a Slate widget style asset) */
-	UPROPERTY(EditAnywhere, BlueprintReadWrite, Category=SlateStyleRules, meta=(AllowedClasses="Font", DisplayName="Font Family"))
+	UPROPERTY(EditAnywhere, BlueprintReadWrite, Category=SlateStyleRules, meta=(AllowedClasses="/Script/Engine.Font", DisplayName="Font Family"))
 	TObjectPtr<const UObject> FontObject;
 
 	/** The material to use when rendering this font */
-	UPROPERTY(EditAnywhere, BlueprintReadWrite, Category=SlateStyleRules, meta=(AllowedClasses="MaterialInterface"))
+	UPROPERTY(EditAnywhere, BlueprintReadWrite, Category=SlateStyleRules, meta=(AllowedClasses="/Script/Engine.MaterialInterface"))
 	TObjectPtr<UObject> FontMaterial;
 
 	/** Settings for applying an outline to a font */
diff --git a/Engine/Source/Runtime/SlateCore/Public/Sound/SlateSound.h b/Engine/Source/Runtime/SlateCore/Public/Sound/SlateSound.h
--- a/Engine/Source/Runtime/SlateCore/Public/Sound/SlateSound.h
+++ b/Engine/Source/Runtime/SlateCore/Public/Sound/SlateSound.h
@@ -67,7 +67,7 @@ protected:
 	 * Pointer to the USoundBase. Holding onto it as a UObject because USoundBase is not available in Slate core.
 	 * Edited via FSlateSoundStructCustomization to ensure you can only set USoundBase assets on it.
 	 */
-	UPROPERTY(EditAnywhere, BlueprintReadWrite, Category=Sound, meta=(DisplayName="Sound", AllowedClasses="SoundBase"))
+	UPROPERTY(EditAnywhere, BlueprintReadWrite, Category=Sound, meta=(DisplayName="Sound", AllowedClasses="/Script/Engine.SoundBase"))
 	TObjectPtr<UObject> ResourceObject;
 
 	/** The legacy resource name; only used by sounds that have been set-up in code, or otherwise upgraded from old FName properties, set to NAME_None in non-legacy instances */
diff --git a/Engine/Source/Runtime/SlateCore/Public/Styling/SlateBrush.h b/Engine/Source/Runtime/SlateCore/Public/Styling/SlateBrush.h
--- a/Engine/Source/Runtime/SlateCore/Public/Styling/SlateBrush.h
+++ b/Engine/Source/Runtime/SlateCore/Public/Styling/SlateBrush.h
@@ -472,7 +472,7 @@ private:
 	 * The image to render for this brush, can be a UTexture or UMaterialInterface or an object implementing 
 	 * the AtlasedTextureInterface. 
 	 */
-	UPROPERTY(EditAnywhere, BlueprintReadWrite, Category=Brush, meta=( AllowPrivateAccess="true", DisplayThumbnail="true", DisplayName="Image", AllowedClasses="Texture,MaterialInterface,SlateTextureAtlasInterface", DisallowedClasses = "MediaTexture"))
+	UPROPERTY(EditAnywhere, BlueprintReadWrite, Category=Brush, meta=( AllowPrivateAccess="true", DisplayThumbnail="true", DisplayName="Image", AllowedClasses="/Script/Engine.Texture,/Script/Engine.MaterialInterface,/Script/Engine.SlateTextureAtlasInterface", DisallowedClasses = "/Script/MediaAssets.MediaTexture"))
 	TObjectPtr<UObject> ResourceObject;
 
 protected:
diff --git a/Engine/Source/Runtime/UMG/Public/Components/RichTextBlock.h b/Engine/Source/Runtime/UMG/Public/Components/RichTextBlock.h
--- a/Engine/Source/Runtime/UMG/Public/Components/RichTextBlock.h
+++ b/Engine/Source/Runtime/UMG/Public/Components/RichTextBlock.h
@@ -196,7 +196,7 @@ protected:
 	FText Text;
 
 	/**  */
-	UPROPERTY(EditAnywhere, Category=Appearance, meta=(RequiredAssetDataTags = "RowStructure=RichTextStyleRow"))
+	UPROPERTY(EditAnywhere, Category=Appearance, meta=(RequiredAssetDataTags = "RowStructure=/Script/UMG.RichTextStyleRow"))
 	TObjectPtr<UDataTable> TextStyleSet;
 
 	/**  */
diff --git a/Engine/Source/Runtime/UMG/Public/Components/RichTextBlockImageDecorator.h b/Engine/Source/Runtime/UMG/Public/Components/RichTextBlockImageDecorator.h
--- a/Engine/Source/Runtime/UMG/Public/Components/RichTextBlockImageDecorator.h
+++ b/Engine/Source/Runtime/UMG/Public/Components/RichTextBlockImageDecorator.h
@@ -51,6 +51,6 @@ protected:
 
 	FRichImageRow* FindImageRow(FName TagOrId, bool bWarnIfMissing);
 
-	UPROPERTY(EditAnywhere, Category=Appearance, meta = (RequiredAssetDataTags = "RowStructure=RichImageRow"))
+	UPROPERTY(EditAnywhere, Category=Appearance, meta = (RequiredAssetDataTags = "RowStructure=/Script/UMG.RichImageRow"))
 	TObjectPtr<class UDataTable> ImageSet;
 };
