Fix typo causing inaccurate server Net Close Reason analytics.

diff --git a/Engine/Source/Runtime/Engine/Private/Analytics/EngineNetAnalytics.cpp b/Engine/Source/Runtime/Engine/Private/Analytics/EngineNetAnalytics.cpp
--- a/Engine/Source/Runtime/Engine/Private/Analytics/EngineNetAnalytics.cpp
+++ b/Engine/Source/Runtime/Engine/Private/Analytics/EngineNetAnalytics.cpp
@@ -523,7 +523,7 @@ void FNetConnAnalyticsData::SendAnalytics()
 			EZAttrib_OutOfOrderPacketsRecoveredCount, OutOfOrderPacketsRecoveredCount,
 			EZAttrib_OutOfOrderPacketsDuplicateCount, OutOfOrderPacketsDuplicateCount);
 
-		if (ClientCloseReasonsJsonStr.Len() > 0)
+		if (CloseReasonsJsonStr.Len() > 0)
 		{
 			AppendAnalyticsEventAttributeArray(Attrs, EZAttrib_CloseReasons, FJsonFragment(MoveTemp(CloseReasonsJsonStr)));
 		}
