Expose screenshot enable functionality via PlatformFeatures

diff --git a/Engine/Source/Runtime/Engine/Private/PlatformFeatures.cpp b/Engine/Source/Runtime/Engine/Private/PlatformFeatures.cpp
--- a/Engine/Source/Runtime/Engine/Private/PlatformFeatures.cpp
+++ b/Engine/Source/Runtime/Engine/Private/PlatformFeatures.cpp
@@ -5,6 +5,10 @@
 #include "DVRStreaming.h"
 #include "VideoRecordingSystem.h"
 
+#if PLATFORM_ANDROID
+extern bool AndroidThunkCpp_IsScreenCaptureDisabled();
+extern void AndroidThunkCpp_DisableScreenCapture(bool bDisable);
+#endif
 
 ISaveGameSystem* IPlatformFeaturesModule::GetSaveGameSystem()
 {
@@ -29,3 +33,13 @@ IVideoRecordingSystem* IPlatformFeaturesModule::GetVideoRecordingSystem()
 	static FGenericVideoRecordingSystem GenericVideoRecordingSystem;
 	return &GenericVideoRecordingSystem;
 }
+
+void IPlatformFeaturesModule::SetScreenshotEnableState(bool bEnabled)
+{
+#if PLATFORM_ANDROID
+	if (AndroidThunkCpp_IsScreenCaptureDisabled() != !bEnabled)
+	{
+		AndroidThunkCpp_DisableScreenCapture(!bEnabled);
+	}
+#endif
+}
\ No newline at end of file
diff --git a/Engine/Source/Runtime/Engine/Public/PlatformFeatures.h b/Engine/Source/Runtime/Engine/Public/PlatformFeatures.h
--- a/Engine/Source/Runtime/Engine/Public/PlatformFeatures.h
+++ b/Engine/Source/Runtime/Engine/Public/PlatformFeatures.h
@@ -45,4 +45,6 @@ public:
 	virtual FString GetUniqueAppId();
 
 	virtual IVideoRecordingSystem* GetVideoRecordingSystem();
+
+	virtual void SetScreenshotEnableState(bool bEnabled);
 };
