TextureDerivedData : fix bForceNoAlphaChannel not in DDC key

diff --git a/Engine/Source/Runtime/Engine/Private/TextureDerivedData.cpp b/Engine/Source/Runtime/Engine/Private/TextureDerivedData.cpp
--- a/Engine/Source/Runtime/Engine/Private/TextureDerivedData.cpp
+++ b/Engine/Source/Runtime/Engine/Private/TextureDerivedData.cpp
@@ -203,14 +203,21 @@ static void SerializeForKey(FArchive& Ar, const FTextureBuildSettings& Settings)
 		TempByte = Settings.DownscaleOptions; Ar << TempByte;
 	}
 
-	// weird/unncessary way to do this :
-	// should have just serialized the bool bForceAlphaChannel
-	// @todo SerializeForKey these can go away whenever we bump the overall ddc key
+	// this is done in a funny way to add the bool that wasn't being serialized before
+	//  without changing DDC keys where the bool is not set
+	// @todo SerializeForKey these can go away whenever we bump the overall ddc key - just serialize the bool
 	if (Settings.bForceAlphaChannel)
 	{
 		TempGuid = FGuid(0x2C9DF7E3, 0xBC9D413B, 0xBF963C7A, 0x3F27E8B1);
 		Ar << TempGuid;
 	}
+	// fix - bForceNoAlphaChannel is not in key !
+	// @todo SerializeForKey these can go away whenever we bump the overall ddc key - just serialize the bool
+	if (Settings.bForceNoAlphaChannel)
+	{
+		TempGuid = FGuid(0x748fc0d4, 0x62004afa, 0x9530460a, 0xf8149d02);
+		Ar << TempGuid;
+	}
 
 	if ( Settings.MaxTextureResolution != FTextureBuildSettings::MaxTextureResolutionDefault &&
 		( Settings.MipGenSettings == TMGS_LeaveExistingMips || Settings.bDoScaleMipsForAlphaCoverage ) )
