Sequencer: Remove interp from RelativeScale3D as it's redundant - same as adding an additive transform section that only keys scale. This only prevents users from creating new RelativeScale3D tracks and doesn't attempt to fix any content (ie. remove existing RelativeScale3D tracks)

diff --git a/Engine/Source/Runtime/Engine/Classes/Components/SceneComponent.h b/Engine/Source/Runtime/Engine/Classes/Components/SceneComponent.h
--- a/Engine/Source/Runtime/Engine/Classes/Components/SceneComponent.h
+++ b/Engine/Source/Runtime/Engine/Classes/Components/SceneComponent.h
@@ -140,7 +140,7 @@ private:
 	*	Non-uniform scaling of the component relative to its parent.
 	*	Note that scaling is always applied in local space (no shearing etc)
 	*/
-	UPROPERTY(BlueprintReadOnly, ReplicatedUsing=OnRep_Transform, interp, Category=Transform, meta=(AllowPrivateAccess="true"))
+	UPROPERTY(BlueprintReadOnly, ReplicatedUsing=OnRep_Transform, Category=Transform, meta=(AllowPrivateAccess="true"))
 	FVector RelativeScale3D;
 
 public:
