Update UGameplayStatics::GetPlatformName() tooltip

diff --git a/Engine/Source/Runtime/Engine/Classes/Kismet/GameplayStatics.h b/Engine/Source/Runtime/Engine/Classes/Kismet/GameplayStatics.h
--- a/Engine/Source/Runtime/Engine/Classes/Kismet/GameplayStatics.h
+++ b/Engine/Source/Runtime/Engine/Classes/Kismet/GameplayStatics.h
@@ -1192,7 +1192,7 @@ public:
 
 	/**
 	 * Returns the string name of the current platform, to perform different behavior based on platform. 
-	 * (Platform names include Windows, Mac, IOS, Android, PS4, XboxOne, Linux) */
+	 * (Platform names include Windows, Mac, Linux, IOS, Android, consoles, etc.). */
 	UFUNCTION(BlueprintPure, Category="Game")
 	static FString GetPlatformName();
 
