Removing no-longer-used 'replay' network smoothing mode. Note that this will potentially invalidate replay data captured long-ago (~2018).

diff --git a/Engine/Source/Runtime/Engine/Classes/Engine/EngineTypes.h b/Engine/Source/Runtime/Engine/Classes/Engine/EngineTypes.h
--- a/Engine/Source/Runtime/Engine/Classes/Engine/EngineTypes.h
+++ b/Engine/Source/Runtime/Engine/Classes/Engine/EngineTypes.h
@@ -850,8 +850,6 @@ enum class ENetworkSmoothingMode : uint8
 	/** Exponential. Faster as you are further from target. */
 	Exponential		UMETA(DisplayName="Exponential"),
 
-	/** Special linear interpolation designed specifically for replays. Not intended as a selectable mode in-editor. */
-	Replay			UMETA(Hidden, DisplayName="Replay"),
 };
 
 // Number of bits used currently from FMaskFilter.
diff --git a/Engine/Source/Runtime/Engine/Private/Character.cpp b/Engine/Source/Runtime/Engine/Private/Character.cpp
--- a/Engine/Source/Runtime/Engine/Private/Character.cpp
+++ b/Engine/Source/Runtime/Engine/Private/Character.cpp
@@ -1164,11 +1164,6 @@ void ACharacter::OnRep_ReplicatedBasedMovement()
 
 void ACharacter::OnRep_ReplicatedMovement()
 {
-	if (CharacterMovement && (CharacterMovement->NetworkSmoothingMode == ENetworkSmoothingMode::Replay))
-	{
-		return;
-	}
-
 	if (GetWorld()->IsPlayingReplay())
 	{
 		Super::OnRep_ReplicatedMovement();
@@ -1198,11 +1193,6 @@ FAnimMontageInstance * ACharacter::GetRootMotionAnimMontageInstance() const
 
 void ACharacter::OnRep_RootMotion()
 {
-	if (CharacterMovement && (CharacterMovement->NetworkSmoothingMode == ENetworkSmoothingMode::Replay))
-	{
-		return;
-	}
-
 	if (GetLocalRole() == ROLE_SimulatedProxy)
 	{
 
diff --git a/Engine/Source/Runtime/Engine/Private/Components/CharacterMovementComponent.cpp b/Engine/Source/Runtime/Engine/Private/Components/CharacterMovementComponent.cpp
--- a/Engine/Source/Runtime/Engine/Private/Components/CharacterMovementComponent.cpp
+++ b/Engine/Source/Runtime/Engine/Private/Components/CharacterMovementComponent.cpp
@@ -198,12 +198,6 @@ namespace CharacterMovementCVars
 		TEXT("0: Disable, 1: Enable"),
 		ECVF_Default);
 
-	static int32 ReplayUseInterpolation = 0;
-	FAutoConsoleVariableRef CVarReplayUseInterpolation(
-		TEXT( "p.ReplayUseInterpolation" ),
-		ReplayUseInterpolation,
-		TEXT( "" ),
-		ECVF_Default);
 
 	static int32 ReplayLerpAcceleration = 0;
 	FAutoConsoleVariableRef CVarReplayLerpAcceleration(
@@ -737,39 +731,10 @@ void UCharacterMovementComponent::OnRegister()
 	// Force linear smoothing for replays.
 	const UWorld* MyWorld = GetWorld();
 	const bool bIsReplay = (MyWorld && MyWorld->IsPlayingReplay());
+
 	if (bIsReplay)
 	{
-		// At least one of these conditions will be true
-		const UDemoNetDriver* DemoNetDriver = MyWorld ? MyWorld->GetDemoNetDriver() : nullptr;
-		const bool bHasInterpolationData = DemoNetDriver && (DemoNetDriver->GetPlaybackDemoVersion() < HISTORY_CHARACTER_MOVEMENT_NOINTERP);
-		const bool bHasRepMovement = DemoNetDriver && (DemoNetDriver->GetPlaybackDemoVersion() >= HISTORY_CHARACTER_MOVEMENT);
-
-		if (CharacterMovementCVars::ReplayUseInterpolation == 1)
-		{
-			if (bHasInterpolationData)
-			{
-				NetworkSmoothingMode = ENetworkSmoothingMode::Replay;
-			}
-			else
-			{
-				UE_LOG(LogCharacterMovement, Warning, TEXT("p.ReplayUseInterpolation is enabled, but the replay was not recorded with interpolation data."));
-				ensure(bHasRepMovement);
-				NetworkSmoothingMode = ENetworkSmoothingMode::Linear;
-			}
-		}
-		else
-		{
-			if (bHasRepMovement)
-			{
-				NetworkSmoothingMode = ENetworkSmoothingMode::Linear;
-			}
-			else
-			{
-				UE_LOG(LogCharacterMovement, Warning, TEXT("p.ReplayUseInterpolation is disabled, but the replay was not recorded with rep movement data."));
-				ensure(bHasInterpolationData);
-				NetworkSmoothingMode = ENetworkSmoothingMode::Replay;
-			}
-		}
+		NetworkSmoothingMode = ENetworkSmoothingMode::Linear;
 	}
 	else if (NetMode == NM_ListenServer)
 	{
@@ -7398,12 +7363,7 @@ void UCharacterMovementComponent::SmoothCorrection(const FVector& OldLocation, c
 	bNetworkSmoothingComplete = false;
 
 	// Handle selected smoothing mode.
-	if (NetworkSmoothingMode == ENetworkSmoothingMode::Replay)
-	{
-		// Replays use pure interpolation in this mode, all of the work is done in SmoothClientPosition_Interpolate
-		return;
-	}
-	else if (NetworkSmoothingMode == ENetworkSmoothingMode::Disabled)
+	if (NetworkSmoothingMode == ENetworkSmoothingMode::Disabled)
 	{
 		UpdatedComponent->SetWorldLocationAndRotation(NewLocation, NewRotation, false, nullptr, ETeleportType::TeleportPhysics);
 		bNetworkSmoothingComplete = true;
@@ -7712,232 +7672,6 @@ void UCharacterMovementComponent::SmoothClientPosition_Interpolate(float DeltaSe
 				ClientData->MeshRotationOffset = MeshRotationTarget;
 			}
 		}
-		else if ( NetworkSmoothingMode == ENetworkSmoothingMode::Replay )
-		{
-			const UWorld* MyWorld = GetWorld();
-			UDemoNetDriver* DemoNetDriver = MyWorld ? MyWorld->GetDemoNetDriver() : nullptr;
-
-			if ( !MyWorld || !DemoNetDriver )
-			{
-				return;
-			}
-
-			const float CurrentTime = DemoNetDriver->GetDemoCurrentTime();
-
-			// Remove old samples
-			while ( ClientData->ReplaySamples.Num() > 0 )
-			{
-				if ( ClientData->ReplaySamples[0].Time > CurrentTime - 1.0f )
-				{
-					break;
-				}
-
-				ClientData->ReplaySamples.RemoveAt( 0 );
-			}
-
-			FReplayExternalDataArray* ExternalReplayData = DemoNetDriver->GetExternalDataArrayForObject( CharacterOwner );
-
-			// Grab any samples available, deserialize them, then clear originals
-			if ( ExternalReplayData && ExternalReplayData->Num() > 0 )
-			{
-				for ( int i = 0; i < ExternalReplayData->Num(); i++ )
-				{
-					FCharacterReplaySample ReplaySample;
-
-					( *ExternalReplayData )[i].Reader << ReplaySample;
-
-					if (FMath::IsNearlyZero(ReplaySample.Time))
-					{
-						ReplaySample.Time = ( *ExternalReplayData )[i].TimeSeconds;
-					}
-
-					ClientData->ReplaySamples.Add( ReplaySample );
-				}
-
-				if ( CharacterMovementCVars::FixReplayOverSampling > 0 )
-				{
-					// Remove invalid replay samples that can occur due to oversampling (sampling at higher rate than physics is being ticked)
-					// We detect this by finding samples that have the same location but have a velocity that says the character should be moving
-					// If we don't do this, then characters will look like they are skipping around, which looks really bad
-					for ( int i = 1; i < ClientData->ReplaySamples.Num(); i++ )
-					{
-						if ( ClientData->ReplaySamples[i].Location.Equals( ClientData->ReplaySamples[i - 1].Location, UE_KINDA_SMALL_NUMBER ) )
-						{
-							if ( ClientData->ReplaySamples[i - 1].Velocity.SizeSquared() > FMath::Square(UE_KINDA_SMALL_NUMBER ) && ClientData->ReplaySamples[i].Velocity.SizeSquared() > FMath::Square(UE_KINDA_SMALL_NUMBER ) )
-							{
-								ClientData->ReplaySamples.RemoveAt( i );
-								i--;
-							}
-						}
-					}
-				}
-
-				ExternalReplayData->Empty();
-			}
-
-			bool bFoundSample = false;
-
-			for ( int i = 0; i < ClientData->ReplaySamples.Num() - 1; i++ )
-			{
-				if ( CurrentTime >= ClientData->ReplaySamples[i].Time && CurrentTime <= ClientData->ReplaySamples[i + 1].Time )
-				{
-					const float EPSILON		= UE_SMALL_NUMBER;
-					const float Delta		= ( ClientData->ReplaySamples[i + 1].Time - ClientData->ReplaySamples[i].Time );
-					const float LerpPercent = Delta > EPSILON ? FMath::Clamp<float>( ( float )( CurrentTime - ClientData->ReplaySamples[i].Time ) / Delta, 0.0f, 1.0f ) : 1.0f;
-
-					const FCharacterReplaySample& ReplaySample1 = ClientData->ReplaySamples[i];
-					const FCharacterReplaySample& ReplaySample2 = ClientData->ReplaySamples[i + 1];
-
-					const FVector Location = FMath::Lerp( ReplaySample1.Location, ReplaySample2.Location, LerpPercent );
-					const FQuat Rotation = FQuat::FastLerp( FQuat( ReplaySample1.Rotation ), FQuat( ReplaySample2.Rotation ), LerpPercent ).GetNormalized();
-					Velocity = FMath::Lerp( ReplaySample1.Velocity, ReplaySample2.Velocity, LerpPercent );
-
-					if (CharacterMovementCVars::ReplayLerpAcceleration)
-					{
-						Acceleration = FMath::Lerp( ReplaySample1.Acceleration, ReplaySample2.Acceleration, LerpPercent );
-					}
-					else
-					{
-						Acceleration = ReplaySample2.Acceleration;
-					}
-
-					const FRotator Rotator1( FRotator::DecompressAxisFromByte( ReplaySample1.RemoteViewPitch ), 0.0f, 0.0f );
-					const FRotator Rotator2( FRotator::DecompressAxisFromByte( ReplaySample2.RemoteViewPitch ), 0.0f, 0.0f );
-					const FRotator FinalPitch = FQuat::FastLerp( FQuat( Rotator1 ), FQuat( Rotator2 ), LerpPercent ).GetNormalized().Rotator();
-					CharacterOwner->BlendedReplayViewPitch = FinalPitch.Pitch;
-
-					UpdateComponentVelocity();
-
-					USkeletalMeshComponent* Mesh = CharacterOwner->GetMesh();
-
-					if ( Mesh )
-					{
-						Mesh->SetRelativeLocation_Direct(CharacterOwner->GetBaseTranslationOffset());
-						Mesh->SetRelativeRotation_Direct(CharacterOwner->GetBaseRotationOffset().Rotator());
-					}
-
-					ClientData->MeshTranslationOffset = Location;
-					ClientData->MeshRotationOffset = Rotation;
-
-#if !(UE_BUILD_SHIPPING || UE_BUILD_TEST)
-					if ( CharacterMovementCVars::NetVisualizeSimulatedCorrections >= 1 )
-					{
-						const float Radius		= 4.0f;
-						const int32	Sides		= 8;
-						const float ArrowSize	= 4.0f;
-						const FColor DebugColor = FColor::White;
-
-						const FVector DebugLocation = CharacterOwner->GetMesh()->GetComponentLocation() + FVector( 0.f, 0.f, 300.0f ) - CharacterOwner->GetBaseTranslationOffset();
-
-						FString DebugText = FString::Printf( TEXT( "Lerp: %2.2f, %2.2f" ), LerpPercent, CharacterOwner->BlendedReplayViewPitch);
-						DrawDebugString( GetWorld(), DebugLocation, DebugText, nullptr, DebugColor, 0.f, true );
-						DrawDebugBox( GetWorld(), DebugLocation, FVector( 45, 45, 45 ), CharacterOwner->GetMesh()->GetComponentQuat(), FColor( 0, 255, 0 ) );
-
-						DrawDebugDirectionalArrow( GetWorld(), DebugLocation, DebugLocation + Velocity, 20.0f, FColor( 255, 0, 0, 255 ) );
-					}
-#endif
-					bFoundSample = true;
-					break;
-				}
-			}
-
-			if ( !bFoundSample && ClientData->ReplaySamples.Num() > 0)
-			{
-				int32 BestSample = 0;
-				float BestDelta = FMath::Abs(ClientData->ReplaySamples[BestSample].Time - CurrentTime);
-
-				for ( int i = 1; i < ClientData->ReplaySamples.Num(); i++ )
-				{
-					const float Delta = FMath::Abs( ClientData->ReplaySamples[i].Time - CurrentTime );
-					if ( Delta < BestDelta )
-					{
-						BestDelta = Delta;
-						BestSample = i;
-					}
-				}
-
-				const FCharacterReplaySample& ReplaySample = ClientData->ReplaySamples[BestSample];
-
-				Velocity								= ReplaySample.Velocity;
-				Acceleration							= ReplaySample.Acceleration;
-				CharacterOwner->BlendedReplayViewPitch	= FRotator::DecompressAxisFromByte(ReplaySample.RemoteViewPitch);
-
-				UpdateComponentVelocity();
-
-				USkeletalMeshComponent* const Mesh = CharacterOwner->GetMesh();
-
-				if ( Mesh )
-				{
-					Mesh->SetRelativeLocation_Direct(CharacterOwner->GetBaseTranslationOffset());
-					Mesh->SetRelativeRotation_Direct(CharacterOwner->GetBaseRotationOffset().Rotator());
-				}
-
-				ClientData->MeshTranslationOffset	= ReplaySample.Location;
-				ClientData->MeshRotationOffset		= ReplaySample.Rotation.Quaternion();
-					
-#if !(UE_BUILD_SHIPPING || UE_BUILD_TEST)
-				if ( Mesh && CharacterMovementCVars::NetVisualizeSimulatedCorrections >= 1 )
-				{
-					const float ArrowSize	= 4.0f;
-
-					const FVector DebugLocation = Mesh->GetComponentLocation() + FVector( 0.f, 0.f, 300.0f ) - CharacterOwner->GetBaseTranslationOffset();
-
-					DrawDebugBox( GetWorld(), DebugLocation, FVector( 45, 45, 45 ), Mesh->GetComponentQuat(), FColor( 255, 0, 0 ) );
-					DrawDebugDirectionalArrow( GetWorld(), DebugLocation, DebugLocation + Velocity, 20.0f, FColor( 255, 0, 0, 255 ) );
-				}
-#endif
-			}
-
-#if !(UE_BUILD_SHIPPING || UE_BUILD_TEST)
-			// Show future samples
-			if ( CharacterMovementCVars::NetVisualizeSimulatedCorrections >= 1 )
-			{
-				const float Radius		= 4.0f;
-				const int32	Sides		= 8;
-				const float ArrowSize	= 4.0f;
-				const FColor DebugColor = FColor::White;
-
-				// Draw points ahead up to a few seconds
-				for ( int i = 0; i < ClientData->ReplaySamples.Num(); i++ )
-				{
-					const bool bHasMorePoints = i < ClientData->ReplaySamples.Num() - 1;
-					const bool bActiveSamples = ( bHasMorePoints && CurrentTime >= ClientData->ReplaySamples[i].Time && CurrentTime <= ClientData->ReplaySamples[i + 1].Time );
-
-					if ( ClientData->ReplaySamples[i].Time >= CurrentTime || bActiveSamples )
-					{
-						//const FVector Adjust = FVector( 0.f, 0.f, 300.0f + i * 15.0f );
-						const FVector Adjust = FVector( 0.f, 0.f, 300.0f );
-						const FVector Location = ClientData->ReplaySamples[i].Location + Adjust;
-
-						if ( bHasMorePoints )
-						{
-							const FVector NextLocation = ClientData->ReplaySamples[i + 1].Location + Adjust;
-							DrawDebugDirectionalArrow( GetWorld(), Location, NextLocation, 4.0f, FColor( 0, 255, 0, 255 ) );
-						}
-
-						DrawCircle( GetWorld(), Location, FVector( 1, 0, 0 ), FVector( 0, 1, 0 ), FColor( 255, 0, 0, 255 ), Radius, Sides, false, 0.0f );
-
-						if ( CharacterMovementCVars::NetVisualizeSimulatedCorrections >= 2 )
-						{
-							DrawDebugDirectionalArrow( GetWorld(), Location, Location + ClientData->ReplaySamples[i].Velocity, 20.0f, FColor( 255, 0, 0, 255 ) );
-						}
-
-						if ( CharacterMovementCVars::NetVisualizeSimulatedCorrections >= 3 )
-						{
-							DrawDebugDirectionalArrow( GetWorld(), Location, Location + ClientData->ReplaySamples[i].Acceleration, 20.0f, FColor( 255, 255, 255, 255 ) );
-						}
-					}
-					
-					if ( ClientData->ReplaySamples[i].Time - CurrentTime > 2.0f )
-					{
-						break;
-					}
-				}
-			}
-#endif
-
-			bNetworkSmoothingComplete = false;
-		}
 		else
 		{
 			// Unhandled mode
@@ -7947,7 +7681,7 @@ void UCharacterMovementComponent::SmoothClientPosition_Interpolate(float DeltaSe
 		//UE_LOG(LogCharacterNetSmoothing, VeryVerbose, TEXT("SmoothClientPosition_Interpolate %s: Translation: %s Rotation: %s"),
 		//	*GetNameSafe(CharacterOwner), *ClientData->MeshTranslationOffset.ToString(), *ClientData->MeshRotationOffset.ToString());
 
-		if ( CharacterMovementCVars::NetVisualizeSimulatedCorrections >= 1 && NetworkSmoothingMode != ENetworkSmoothingMode::Replay )
+		if ( CharacterMovementCVars::NetVisualizeSimulatedCorrections >= 1 )
 		{
 			const FVector DebugLocation = CharacterOwner->GetMesh()->GetComponentLocation() + FVector( 0.f, 0.f, 300.0f ) - CharacterOwner->GetBaseTranslationOffset();
 			DrawDebugBox( GetWorld(), DebugLocation, FVector( 45, 45, 45 ), CharacterOwner->GetMesh()->GetComponentQuat(), FColor( 0, 255, 0 ) );
@@ -8034,13 +7768,6 @@ void UCharacterMovementComponent::SmoothClientPosition_UpdateVisuals()
 			const FQuat NewRelRotation = ClientData->MeshRotationOffset * CharacterOwner->GetBaseRotationOffset();
 			Mesh->SetRelativeLocationAndRotation(NewRelTranslation, NewRelRotation, false, nullptr, GetTeleportType());
 		}
-		else if ( NetworkSmoothingMode == ENetworkSmoothingMode::Replay )
-		{
-			if ( !UpdatedComponent->GetComponentQuat().Equals( ClientData->MeshRotationOffset, SCENECOMPONENT_QUAT_TOLERANCE ) || !UpdatedComponent->GetComponentLocation().Equals( ClientData->MeshTranslationOffset, UE_KINDA_SMALL_NUMBER ) )
-			{
-				UpdatedComponent->SetWorldLocationAndRotation( ClientData->MeshTranslationOffset, ClientData->MeshRotationOffset, false, nullptr, GetTeleportType());
-			}
-		}
 		else
 		{
 			// Unhandled mode
diff --git a/Engine/Source/Runtime/Engine/Public/ReplayTypes.h b/Engine/Source/Runtime/Engine/Public/ReplayTypes.h
--- a/Engine/Source/Runtime/Engine/Public/ReplayTypes.h
+++ b/Engine/Source/Runtime/Engine/Public/ReplayTypes.h
@@ -114,11 +114,9 @@ enum ENetworkVersionHistory
 	HISTORY_LATEST = HISTORY_PLUS_ONE - 1
 };
 
-static const uint32 MIN_SUPPORTED_VERSION = HISTORY_EXTRA_VERSION;
-
 static const uint32 NETWORK_DEMO_MAGIC = 0x2CF5A13D;
 static const uint32 NETWORK_DEMO_VERSION = HISTORY_LATEST;
-static const uint32 MIN_NETWORK_DEMO_VERSION = HISTORY_EXTRA_VERSION;
+static const uint32 MIN_NETWORK_DEMO_VERSION = HISTORY_CHARACTER_MOVEMENT;
 
 static const uint32 NETWORK_DEMO_METADATA_MAGIC = 0x3D06B24E;
 static const uint32 NETWORK_DEMO_METADATA_VERSION = 0;
