Increasing default fx budget decay rate

diff --git a/Engine/Source/Runtime/Engine/Private/Particles/FXBudget.cpp b/Engine/Source/Runtime/Engine/Private/Particles/FXBudget.cpp
--- a/Engine/Source/Runtime/Engine/Private/Particles/FXBudget.cpp
+++ b/Engine/Source/Runtime/Engine/Private/Particles/FXBudget.cpp
@@ -67,7 +67,7 @@ static FAutoConsoleVariableRef CVarFXRenderThreadTimeOverride(
 );
 
 
-static float GFXBudget_AdjustedUsageDecayRate = 0.005f;
+static float GFXBudget_AdjustedUsageDecayRate = 0.1f;
 static FAutoConsoleVariableRef CVarFXBudget_AdjustedUsageDecayRate(
 	TEXT("fx.Budget.AdjustedUsageDecayRate"),
 	GFXBudget_AdjustedUsageDecayRate,
