Make async compilers always read editor settings from the game-thread as early as possible to avoid lazy-init from non game-thread

#rnx

diff --git a/Engine/Source/Runtime/Engine/Private/AsyncCompilationHelpers.cpp b/Engine/Source/Runtime/Engine/Private/AsyncCompilationHelpers.cpp
--- a/Engine/Source/Runtime/Engine/Private/AsyncCompilationHelpers.cpp
+++ b/Engine/Source/Runtime/Engine/Private/AsyncCompilationHelpers.cpp
@@ -192,32 +192,42 @@ namespace AsyncCompilationHelpers
 	{
 		if (InExperimentalSettingsName != NAME_None)
 		{
-			auto UpdateCVarFromSettings =
-				[&InCVarAsyncCompilation, InExperimentalSettingsName]
+			// We cannot try to get the UEditorExperimentalSettings until all the class properties
+			// have been linked otherwise the CDO wouldn't be able to initialize properties
+			// from the ini files which would most likely end up with the default value of
+			// async compilation being disabled. This obviously has the effect of async compilation
+			// being disabled during boot until the object system is ready and settings can be read.
+			FDelayedAutoRegisterHelper(EDelayedRegisterRunPhase::ObjectSystemReady,
+				[&InCVarAsyncCompilation, InExperimentalSettingsName]()
 				{
-					const UEditorExperimentalSettings* Settings = GetDefault<UEditorExperimentalSettings>();
-					FProperty* Property = Settings->GetClass()->FindPropertyByName(InExperimentalSettingsName);
-					if (Property)
-					{
-						if (FBoolProperty* BoolProperty = CastField<FBoolProperty>(Property))
+					auto UpdateCVarFromSettings =
+						[&InCVarAsyncCompilation, InExperimentalSettingsName]
+						{
+							const UEditorExperimentalSettings* Settings = GetDefault<UEditorExperimentalSettings>();
+							FProperty* Property = Settings->GetClass()->FindPropertyByName(InExperimentalSettingsName);
+							if (Property)
+							{
+								if (FBoolProperty* BoolProperty = CastField<FBoolProperty>(Property))
+								{
+									bool bIsEnabled = BoolProperty->GetPropertyValue(BoolProperty->ContainerPtrToValuePtr<void>(Settings));
+									InCVarAsyncCompilation->Set(bIsEnabled ? 1 : 0, ECVF_SetByProjectSetting);
+								}
+							}
+						};
+
+					GetMutableDefault<UEditorExperimentalSettings>()->OnSettingChanged().AddLambda(
+						[UpdateCVarFromSettings, InExperimentalSettingsName](FName Name)
 						{
-							bool bIsEnabled = BoolProperty->GetPropertyValue(BoolProperty->ContainerPtrToValuePtr<void>(Settings));
-							InCVarAsyncCompilation->Set(bIsEnabled ? 1 : 0, ECVF_SetByProjectSetting);
+							if (Name == InExperimentalSettingsName)
+							{
+								UpdateCVarFromSettings();
+							}
 						}
-					}
-				};
+					);
 
-			GetMutableDefault<UEditorExperimentalSettings>()->OnSettingChanged().AddLambda(
-				[UpdateCVarFromSettings, InExperimentalSettingsName](FName Name)
-				{
-					if (Name == InExperimentalSettingsName)
-					{
-						UpdateCVarFromSettings();
-					}
+					UpdateCVarFromSettings();
 				}
 			);
-
-			UpdateCVarFromSettings();
 		}
 
 		FString Value;
diff --git a/Engine/Source/Runtime/Engine/Private/SkeletalMeshCompiler.cpp b/Engine/Source/Runtime/Engine/Private/SkeletalMeshCompiler.cpp
--- a/Engine/Source/Runtime/Engine/Private/SkeletalMeshCompiler.cpp
+++ b/Engine/Source/Runtime/Engine/Private/SkeletalMeshCompiler.cpp
@@ -62,6 +62,7 @@ namespace SkeletalMeshCompilingManagerImpl
 FSkeletalMeshCompilingManager::FSkeletalMeshCompilingManager()
 	: Notification(GetAssetNameFormat())
 {
+	SkeletalMeshCompilingManagerImpl::EnsureInitializedCVars();
 }
 
 FName FSkeletalMeshCompilingManager::GetAssetTypeName() const
@@ -101,8 +102,6 @@ FQueuedThreadPool* FSkeletalMeshCompilingManager::GetThreadPool() const
 	static FQueuedThreadPoolDynamicWrapper* GSkeletalMeshThreadPool = nullptr;
 	if (GSkeletalMeshThreadPool == nullptr && FAssetCompilingManager::Get().GetThreadPool() != nullptr)
 	{
-		SkeletalMeshCompilingManagerImpl::EnsureInitializedCVars();
-
 		// For now, skeletal mesh have almost no high-level awareness of their async behavior.
 		// Let them build first to avoid game-thread stalls as much as possible.
 		TFunction<EQueuedWorkPriority(EQueuedWorkPriority)> PriorityMapper = [](EQueuedWorkPriority) { return EQueuedWorkPriority::Highest; };
@@ -166,8 +165,6 @@ bool FSkeletalMeshCompilingManager::IsAsyncCompilationEnabled() const
 		return false;
 	}
 
-	SkeletalMeshCompilingManagerImpl::EnsureInitializedCVars();
-
 	return CVarAsyncSkeletalMeshStandard.AsyncCompilation.GetValueOnAnyThread() != 0;
 }
 
diff --git a/Engine/Source/Runtime/Engine/Private/SoundWaveCompiler.cpp b/Engine/Source/Runtime/Engine/Private/SoundWaveCompiler.cpp
--- a/Engine/Source/Runtime/Engine/Private/SoundWaveCompiler.cpp
+++ b/Engine/Source/Runtime/Engine/Private/SoundWaveCompiler.cpp
@@ -60,6 +60,7 @@ namespace SoundWaveCompilingManagerImpl
 FSoundWaveCompilingManager::FSoundWaveCompilingManager()
 	: Notification(GetAssetNameFormat())
 {
+	SoundWaveCompilingManagerImpl::EnsureInitializedCVars();
 }
 
 FName FSoundWaveCompilingManager::GetStaticAssetTypeName()
@@ -92,8 +93,6 @@ FQueuedThreadPool* FSoundWaveCompilingManager::GetThreadPool() const
 	static FQueuedThreadPoolWrapper* GSoundWaveThreadPool = nullptr;
 	if (GSoundWaveThreadPool == nullptr && FAssetCompilingManager::Get().GetThreadPool() != nullptr)
 	{
-		SoundWaveCompilingManagerImpl::EnsureInitializedCVars();
-
 		const auto SoundWavePriorityMapper = [](EQueuedWorkPriority SoundWavePriority) { return FMath::Max(SoundWavePriority, EQueuedWorkPriority::Low); };
 
 		// SoundWaves will be scheduled on the asset thread pool, where concurrency limits might by dynamically adjusted depending on memory constraints.
@@ -138,8 +137,6 @@ bool FSoundWaveCompilingManager::IsAsyncSoundWaveCompilationEnabled() const
 		return false;
 	}
 
-	SoundWaveCompilingManagerImpl::EnsureInitializedCVars();
-
 	return CVarAsyncSoundWaveStandard.AsyncCompilation.GetValueOnAnyThread() != 0;
 }
 
diff --git a/Engine/Source/Runtime/Engine/Private/StaticMeshCompiler.cpp b/Engine/Source/Runtime/Engine/Private/StaticMeshCompiler.cpp
--- a/Engine/Source/Runtime/Engine/Private/StaticMeshCompiler.cpp
+++ b/Engine/Source/Runtime/Engine/Private/StaticMeshCompiler.cpp
@@ -90,6 +90,8 @@ namespace StaticMeshCompilingManagerImpl
 FStaticMeshCompilingManager::FStaticMeshCompilingManager()
 	: Notification(GetAssetNameFormat())
 {
+	StaticMeshCompilingManagerImpl::EnsureInitializedCVars();
+
 	PostReachabilityAnalysisHandle = FCoreUObjectDelegates::PostReachabilityAnalysis.AddRaw(this, &FStaticMeshCompilingManager::OnPostReachabilityAnalysis);
 }
 
@@ -161,8 +163,6 @@ FQueuedThreadPool* FStaticMeshCompilingManager::GetThreadPool() const
 	static FQueuedThreadPoolDynamicWrapper* GStaticMeshThreadPool = nullptr;
 	if (GStaticMeshThreadPool == nullptr && FAssetCompilingManager::Get().GetThreadPool() != nullptr)
 	{
-		StaticMeshCompilingManagerImpl::EnsureInitializedCVars();
-
 		// Static meshes will be scheduled on the asset thread pool, where concurrency limits might by dynamically adjusted depending on memory constraints.
 		GStaticMeshThreadPool = new FQueuedThreadPoolDynamicWrapper(FAssetCompilingManager::Get().GetThreadPool(), -1, [](EQueuedWorkPriority) { return EQueuedWorkPriority::Low; });
 
@@ -213,8 +213,6 @@ bool FStaticMeshCompilingManager::IsAsyncStaticMeshCompilationEnabled() const
 		return false;
 	}
 
-	StaticMeshCompilingManagerImpl::EnsureInitializedCVars();
-
 	return CVarAsyncStaticMeshStandard.AsyncCompilation.GetValueOnAnyThread() != 0;
 }
 
diff --git a/Engine/Source/Runtime/Engine/Private/TextureCompiler.cpp b/Engine/Source/Runtime/Engine/Private/TextureCompiler.cpp
--- a/Engine/Source/Runtime/Engine/Private/TextureCompiler.cpp
+++ b/Engine/Source/Runtime/Engine/Private/TextureCompiler.cpp
@@ -80,6 +80,7 @@ namespace TextureCompilingManagerImpl
 FTextureCompilingManager::FTextureCompilingManager()
 	: Notification(GetAssetNameFormat())
 {
+	TextureCompilingManagerImpl::EnsureInitializedCVars();
 }
 
 FName FTextureCompilingManager::GetStaticAssetTypeName()
@@ -134,8 +135,6 @@ FQueuedThreadPool* FTextureCompilingManager::GetThreadPool() const
 	static FQueuedThreadPoolWrapper* GTextureThreadPool = nullptr;
 	if (GTextureThreadPool == nullptr && FAssetCompilingManager::Get().GetThreadPool() != nullptr)
 	{
-		TextureCompilingManagerImpl::EnsureInitializedCVars();
-
 		const auto TexturePriorityMapper = [](EQueuedWorkPriority TexturePriority) { return FMath::Max(TexturePriority, EQueuedWorkPriority::Low); };
 
 		// Textures will be scheduled on the asset thread pool, where concurrency limits might by dynamically adjusted depending on memory constraints.
@@ -188,8 +187,6 @@ bool FTextureCompilingManager::IsAsyncTextureCompilationEnabled() const
 		return false;
 	}
 
-	TextureCompilingManagerImpl::EnsureInitializedCVars();
-
 	return CVarAsyncTextureStandard.AsyncCompilation.GetValueOnAnyThread() != 0;
 }
 
