Make Serialize and PostSerialize for various classes thread-safe

#rnx

diff --git a/Engine/Source/Runtime/Engine/Private/PhysicsEngine/ConstraintInstance.cpp b/Engine/Source/Runtime/Engine/Private/PhysicsEngine/ConstraintInstance.cpp
--- a/Engine/Source/Runtime/Engine/Private/PhysicsEngine/ConstraintInstance.cpp
+++ b/Engine/Source/Runtime/Engine/Private/PhysicsEngine/ConstraintInstance.cpp
@@ -1154,14 +1154,14 @@ void FConstraintInstance::PostSerialize(const FArchive& Ar)
 		}
 
 		//Now handle the new linear spring stiffness and damping coefficient
-		if(CVarConstraintAngularStiffnessScale.GetValueOnGameThread() > 0.f)
+		if(CVarConstraintAngularStiffnessScale.GetValueOnAnyThread() > 0.f)
 		{
-			ProfileInstance.LinearLimit.Stiffness *= CVarConstraintAngularStiffnessScale.GetValueOnGameThread() / CVarConstraintLinearStiffnessScale.GetValueOnGameThread();
+			ProfileInstance.LinearLimit.Stiffness *= CVarConstraintAngularStiffnessScale.GetValueOnAnyThread() / CVarConstraintLinearStiffnessScale.GetValueOnAnyThread();
 		}
 
-		if (CVarConstraintAngularDampingScale.GetValueOnGameThread() > 0.f)
+		if (CVarConstraintAngularDampingScale.GetValueOnAnyThread() > 0.f)
 		{
-			ProfileInstance.LinearLimit.Damping *= CVarConstraintAngularDampingScale.GetValueOnGameThread() / CVarConstraintLinearDampingScale.GetValueOnGameThread();
+			ProfileInstance.LinearLimit.Damping *= CVarConstraintAngularDampingScale.GetValueOnAnyThread() / CVarConstraintLinearDampingScale.GetValueOnAnyThread();
 		}
 	}
 }
diff --git a/Engine/Source/Runtime/Engine/Private/Texture2D.cpp b/Engine/Source/Runtime/Engine/Private/Texture2D.cpp
--- a/Engine/Source/Runtime/Engine/Private/Texture2D.cpp
+++ b/Engine/Source/Runtime/Engine/Private/Texture2D.cpp
@@ -444,18 +444,6 @@ void UTexture2D::Serialize(FArchive& Ar)
 	{
 		SerializeCookedPlatformData(Ar);
 	}
-
-#if WITH_EDITOR	
-	if (Ar.IsLoading() && !Ar.IsTransacting() && !bCooked && !GetOutermost()->HasAnyPackageFlags(PKG_ReloadingForCooker))
-	{
-		// The composite texture may not have been loaded yet. We have to defer caching platform
-		// data until post load.
-		if (CompositeTexture == NULL || CompositeTextureMode == CTM_Disabled || CompositeTexture->HasAnyFlags(RF_LoadCompleted))
-		{
-			BeginCachePlatformData();
-		}
-	}
-#endif // #if WITH_EDITOR
 }
 
 int32 UTexture2D::GetNumResidentMips() const
diff --git a/Engine/Source/Runtime/Engine/Private/TextureCube.cpp b/Engine/Source/Runtime/Engine/Private/TextureCube.cpp
--- a/Engine/Source/Runtime/Engine/Private/TextureCube.cpp
+++ b/Engine/Source/Runtime/Engine/Private/TextureCube.cpp
@@ -146,13 +146,6 @@ void UTextureCube::Serialize(FArchive& Ar)
 		UE_LOG(LogTexture, Log, TEXT("Default maximum texture size for cubemaps generated from long-lat sources has been changed from 512 to unlimited. In order to preserve old behavior for '%s', its maximum texture size has been explicitly set to 512."), *GetPathName());
 	}
 #endif // #if WITH_EDITORONLY_DATA
-
-#if WITH_EDITOR
-	if (Ar.IsLoading() && !Ar.IsTransacting() && !bCooked)
-	{
-		BeginCachePlatformData();
-	}
-#endif // #if WITH_EDITOR
 }
 
 void UTextureCube::PostLoad()
