Trace screenshots from the Shot and HighResShot commands when the Screenshot channel is on.

diff --git a/Engine/Source/Editor/UnrealEd/Private/EditorViewportClient.cpp b/Engine/Source/Editor/UnrealEd/Private/EditorViewportClient.cpp
--- a/Engine/Source/Editor/UnrealEd/Private/EditorViewportClient.cpp
+++ b/Engine/Source/Editor/UnrealEd/Private/EditorViewportClient.cpp
@@ -60,6 +60,8 @@
 #include "UnrealWidget.h"
 #include "EdModeInteractiveToolsContext.h"
 #include "Engine/World.h"
+#include "ProfilingDebugging/MiscTrace.h"
+#include "ProfilingDebugging/TraceScreenshot.h"
 
 #include "CustomEditorStaticScreenPercentage.h"
 #include "IImageWrapperModule.h"
@@ -5853,6 +5855,11 @@ bool FEditorViewportClient::ProcessScreenShots(FViewport* InViewport)
 				}
 			}
 
+			if (SHOULD_TRACE_SCREENSHOT())
+			{
+				FTraceScreenshot::TraceScreenshot(BitmapSize.X, BitmapSize.Y, Bitmap, FScreenshotRequest::GetFilename());
+			}
+
 			// Save the bitmap to disk
 			TFuture<bool> CompletionFuture = HighResScreenshotConfig.ImageWriteQueue->Enqueue(MoveTemp(ImageTask));
 			if (CompletionFuture.IsValid())
diff --git a/Engine/Source/Runtime/Core/Private/ProfilingDebugging/MiscTrace.cpp b/Engine/Source/Runtime/Core/Private/ProfilingDebugging/MiscTrace.cpp
--- a/Engine/Source/Runtime/Core/Private/ProfilingDebugging/MiscTrace.cpp
+++ b/Engine/Source/Runtime/Core/Private/ProfilingDebugging/MiscTrace.cpp
@@ -134,4 +134,9 @@ void FMiscTrace::OutputScreenshot(const TCHAR* Name, uint64 Cycle, uint32 Width,
 
 	check(RemainingSize == 0);
 }
+
+bool FMiscTrace::ShouldTraceScreenshot()
+{
+	return UE_TRACE_CHANNELEXPR_IS_ENABLED(ScreenshotChannel);
+}
 #endif
diff --git a/Engine/Source/Runtime/Core/Public/ProfilingDebugging/MiscTrace.h b/Engine/Source/Runtime/Core/Public/ProfilingDebugging/MiscTrace.h
--- a/Engine/Source/Runtime/Core/Public/ProfilingDebugging/MiscTrace.h
+++ b/Engine/Source/Runtime/Core/Public/ProfilingDebugging/MiscTrace.h
@@ -62,6 +62,7 @@ struct FMiscTrace
 	CORE_API static void OutputEndFrame(ETraceFrameType FrameType);
 
 	CORE_API static void OutputScreenshot(const TCHAR* Name, uint64 Cycle, uint32 Width, uint32 Height, TArray64<uint8> Data);
+	CORE_API static bool ShouldTraceScreenshot();
 
 private:
 	CORE_API static void OutputBookmarkInternal(const void* BookmarkPoint, uint16 EncodedFormatArgsSize, uint8* EncodedFormatArgs);
@@ -85,11 +86,15 @@ private:
 #define TRACE_SCREENSHOT(Name, Cycle, Width, Height, Data) \
 	FMiscTrace::OutputScreenshot(Name, Cycle, Width, Height, Data);
 
+#define SHOULD_TRACE_SCREENSHOT() \
+	FMiscTrace::ShouldTraceScreenshot()
+
 #else
 
 #define TRACE_BOOKMARK(...)
 #define TRACE_BEGIN_FRAME(...)
 #define TRACE_END_FRAME(...)
 #define TRACE_SCREENSHOT(...)
+#define SHOULD_TRACE_SCREENSHOT(...) false
 
 #endif
diff --git a/Engine/Source/Runtime/Engine/Private/GameViewportClient.cpp b/Engine/Source/Runtime/Engine/Private/GameViewportClient.cpp
--- a/Engine/Source/Runtime/Engine/Private/GameViewportClient.cpp
+++ b/Engine/Source/Runtime/Engine/Private/GameViewportClient.cpp
@@ -61,6 +61,8 @@
 #include "ComponentRecreateRenderStateContext.h"
 #include "DynamicResolutionState.h"
 #include "ProfilingDebugging/CsvProfiler.h"
+#include "ProfilingDebugging/MiscTrace.h"
+#include "ProfilingDebugging/TraceScreenshot.h"
 #include "IImageWrapper.h"
 #include "IImageWrapperModule.h"
 #include "HAL/PlatformApplicationMisc.h"
@@ -2002,6 +2004,11 @@ bool UGameViewportClient::ProcessScreenShots(FViewport* InViewport)
 					ScreenShotName += TEXT(".png");
 				}
 
+				if (SHOULD_TRACE_SCREENSHOT())
+				{
+					FTraceScreenshot::TraceScreenshot(Size.X, Size.Y, Bitmap, ScreenShotName);
+				}
+
 				// Save the contents of the array to a png file.
 				
 				FImageView Image((const FColor *)Bitmap.GetData(),Size.X,Size.Y);
diff --git a/Engine/Source/Runtime/Engine/Private/ProfilingDebugging/TraceScreenshot.cpp b/Engine/Source/Runtime/Engine/Private/ProfilingDebugging/TraceScreenshot.cpp
--- a/Engine/Source/Runtime/Engine/Private/ProfilingDebugging/TraceScreenshot.cpp
+++ b/Engine/Source/Runtime/Engine/Private/ProfilingDebugging/TraceScreenshot.cpp
@@ -69,26 +69,10 @@ void FTraceScreenshot::RequestScreenshot(FString Name)
 
 void FTraceScreenshot::HandleScreenshotData(int32 InSizeX, int32 InSizeY, const TArray<FColor>& InImageData)
 {
-	if (ScreenshotName.IsEmpty())
+	if (SHOULD_TRACE_SCREENSHOT())
 	{
-		ScreenshotName = FDateTime::Now().ToString(TEXT("Screenshot_%Y%m%d_%H%M%S"));
+		TraceScreenshot(InSizeX, InSizeY, InImageData, ScreenshotName, 640);
 	}
-
-	uint64 Cycles = FPlatformTime::Cycles64();
-	UE_LOG(LogCore, Display, TEXT("Tracing Screenshot \"%s\" taken with size: %d x %d"), *ScreenshotName, InSizeX, InSizeY);
-
-	int32 ResizedX = FMath::Min(640, InSizeX);
-	int32 ResizedY = (InSizeY * ResizedX) / InSizeX;
-
-	TArray<FColor> ResizedImage;
-	ResizedImage.SetNum(ResizedX * ResizedY);
-	FImageUtils::ImageResize(InSizeX, InSizeY, InImageData, ResizedX, ResizedY, ResizedImage, false);
-
-	TArray64<uint8> CompressedBitmap;
-	FImageUtils::PNGCompressImageArray(ResizedX, ResizedY, TArrayView64<const FColor>(ResizedImage.GetData(), ResizedImage.Num()), CompressedBitmap);
-
-	TRACE_SCREENSHOT(*ScreenshotName, Cycles, ResizedX, ResizedY, CompressedBitmap);
-	ScreenshotName.Empty();
 }
 
 void FTraceScreenshot::WorldDestroyed(ULevel* InLevel, UWorld* InWorld)
@@ -105,4 +89,37 @@ void FTraceScreenshot::WorldDestroyed(ULevel* InLevel, UWorld* InWorld)
 		}
 		FWorldDelegates::LevelRemovedFromWorld.RemoveAll(this);
 	}
-}
\ No newline at end of file
+}
+
+void FTraceScreenshot::TraceScreenshot(int32 InSizeX, int32 InSizeY, const TArray<FColor>& InImageData, const FString& InScreenshotName, int32 DesiredX)
+{
+	FString ScreenshotName = InScreenshotName;
+	if (ScreenshotName.IsEmpty())
+	{
+		ScreenshotName = FDateTime::Now().ToString(TEXT("Screenshot_%Y%m%d_%H%M%S"));
+	}
+
+	uint64 Cycles = FPlatformTime::Cycles64();
+	UE_LOG(LogCore, Display, TEXT("Tracing Screenshot \"%s\" taken with size: %d x %d"), *ScreenshotName, InSizeX, InSizeY);
+
+	TArray64<uint8> CompressedBitmap;
+	if (DesiredX > 0 && InSizeX != DesiredX)
+	{
+		int32 ResizedX = FMath::Min(640, InSizeX);
+		int32 ResizedY = (InSizeY * ResizedX) / InSizeX;
+
+		TArray<FColor> ResizedImage;
+		ResizedImage.SetNum(ResizedX * ResizedY);
+		FImageUtils::ImageResize(InSizeX, InSizeY, InImageData, ResizedX, ResizedY, ResizedImage, false);
+
+		FImageUtils::PNGCompressImageArray(ResizedX, ResizedY, TArrayView64<const FColor>(ResizedImage.GetData(), ResizedImage.Num()), CompressedBitmap);
+		TRACE_SCREENSHOT(*ScreenshotName, Cycles, ResizedX, ResizedY, CompressedBitmap);
+	}
+	else
+	{
+		FImageUtils::PNGCompressImageArray(InSizeX, InSizeY, TArrayView64<const FColor>(InImageData.GetData(), InImageData.Num()), CompressedBitmap);
+		TRACE_SCREENSHOT(*ScreenshotName, Cycles, InSizeX, InSizeY, CompressedBitmap);
+	}
+
+	ScreenshotName.Empty();
+}
diff --git a/Engine/Source/Runtime/Engine/Public/ProfilingDebugging/TraceScreenshot.h b/Engine/Source/Runtime/Engine/Public/ProfilingDebugging/TraceScreenshot.h
--- a/Engine/Source/Runtime/Engine/Public/ProfilingDebugging/TraceScreenshot.h
+++ b/Engine/Source/Runtime/Engine/Public/ProfilingDebugging/TraceScreenshot.h
@@ -7,7 +7,7 @@
 
 class ULevel;
 
-class FTraceScreenshot
+class ENGINE_API FTraceScreenshot
 {
 public:
 	static FTraceScreenshot* Get() { return Instance; }
@@ -16,6 +16,15 @@ public:
 	void HandleScreenshotData(int32 InSizeX, int32 InSizeY, const TArray<FColor>& InImageData);
 	void WorldDestroyed(ULevel* InLevel, UWorld* InWorld);
 
+	/* 
+	* Add the provided screenshot to the trace.
+	* @param InSizeX - The width of the screenshot.
+	* @param InSizeY - The heigth of the screenshot.
+	* @param InImageData - The data of the screenshot.
+	* @param InScreenshotName - The name of the screenshot.
+	* @param DesiredX - Optionally resize the image to the desired width before tracing. Aspect ratio is preserved.
+	*/
+	static void TraceScreenshot(int32 InSizeX, int32 InSizeY, const TArray<FColor>& InImageData, const FString& InScreenshotName, int32 DesiredX = -1);
 private:
 	FTraceScreenshot();
 	FTraceScreenshot(const FTraceScreenshot& Other) {}
