Store bounds in FDeferredDecalProxy.

diff --git a/Engine/Source/Runtime/Engine/Private/Components/DecalComponent.cpp b/Engine/Source/Runtime/Engine/Private/Components/DecalComponent.cpp
--- a/Engine/Source/Runtime/Engine/Private/Components/DecalComponent.cpp
+++ b/Engine/Source/Runtime/Engine/Private/Components/DecalComponent.cpp
@@ -41,7 +41,7 @@ FDeferredDecalProxy::FDeferredDecalProxy(const UDecalComponent* InComponent)
 
 	Component = InComponent;
 	DecalMaterial = EffectiveMaterial;
-	SetTransformIncludingDecalSize(InComponent->GetTransformIncludingDecalSize());
+	SetTransformIncludingDecalSize(InComponent->GetTransformIncludingDecalSize(), InComponent->CalcBounds(InComponent->GetComponentTransform()));
 	bOwnerSelected = InComponent->IsOwnerSelected();
 	SortOrder = InComponent->SortOrder;
 
@@ -62,9 +62,10 @@ FDeferredDecalProxy::FDeferredDecalProxy(const UDecalComponent* InComponent)
 	}
 }
 
-void FDeferredDecalProxy::SetTransformIncludingDecalSize(const FTransform& InComponentToWorldIncludingDecalSize)
+void FDeferredDecalProxy::SetTransformIncludingDecalSize(const FTransform& InComponentToWorldIncludingDecalSize, const FBoxSphereBounds& InBounds)
 {
 	ComponentTrans = InComponentToWorldIncludingDecalSize;
+	Bounds = InBounds;
 }
 
 void FDeferredDecalProxy::InitializeFadingParameters(float AbsSpawnTime, float FadeDuration, float FadeStartDelay, float FadeInDuration, float FadeInStartDelay)
diff --git a/Engine/Source/Runtime/Engine/Public/SceneManagement.h b/Engine/Source/Runtime/Engine/Public/SceneManagement.h
--- a/Engine/Source/Runtime/Engine/Public/SceneManagement.h
+++ b/Engine/Source/Runtime/Engine/Public/SceneManagement.h
@@ -1814,16 +1814,19 @@ public:
 	FDeferredDecalProxy(const UDecalComponent* InComponent);
 
 	/**
-	 * Updates the decal proxy's cached transform.
+	 * Updates the decal proxy's cached transform and bounds.
 	 * @param InComponentToWorldIncludingDecalSize - The new component-to-world transform including the DecalSize
+	 * @param InBounds - The new world-space bounds including the DecalSize
 	 */
-	void SetTransformIncludingDecalSize(const FTransform& InComponentToWorldIncludingDecalSize);
+	void SetTransformIncludingDecalSize(const FTransform& InComponentToWorldIncludingDecalSize, const FBoxSphereBounds& InBounds);
 
 	void InitializeFadingParameters(float AbsSpawnTime, float FadeDuration, float FadeStartDelay, float FadeInDuration, float FadeInStartDelay);
 
 	/** @return True if the decal is visible in the given view. */
 	bool IsShown( const FSceneView* View ) const;
 
+	inline const FBoxSphereBounds& GetBounds() const { return Bounds; }
+
 	/** Pointer back to the game thread decal component. */
 	const UDecalComponent* Component;
 
@@ -1839,6 +1842,8 @@ private:
 	/** Whether or not the decal should be drawn in the editor. */
 	bool DrawInEditor;
 
+	FBoxSphereBounds Bounds;
+
 public:
 
 	bool bOwnerSelected;
diff --git a/Engine/Source/Runtime/Renderer/Private/RendererScene.cpp b/Engine/Source/Runtime/Renderer/Private/RendererScene.cpp
--- a/Engine/Source/Runtime/Renderer/Private/RendererScene.cpp
+++ b/Engine/Source/Runtime/Renderer/Private/RendererScene.cpp
@@ -2249,11 +2249,12 @@ void FScene::UpdateDecalTransform(UDecalComponent* Decal)
 		//Send command to the rendering thread to update the decal's transform.
 		FDeferredDecalProxy* DecalSceneProxy = Decal->SceneProxy;
 		FTransform ComponentToWorldIncludingDecalSize = Decal->GetTransformIncludingDecalSize();
+		FBoxSphereBounds Bounds = Decal->CalcBounds(Decal->GetComponentTransform());
 		ENQUEUE_RENDER_COMMAND(UpdateTransformCommand)(
-			[DecalSceneProxy, ComponentToWorldIncludingDecalSize](FRHICommandListImmediate& RHICmdList)
+			[DecalSceneProxy, ComponentToWorldIncludingDecalSize, Bounds](FRHICommandListImmediate& RHICmdList)
 			{
 				// Update the primitive's transform.
-				DecalSceneProxy->SetTransformIncludingDecalSize(ComponentToWorldIncludingDecalSize);
+				DecalSceneProxy->SetTransformIncludingDecalSize(ComponentToWorldIncludingDecalSize, Bounds);
 			});
 	}
 }
