World Partition:
- Added FWorldPartitionLoadingContext for editor loading with two modes: immediate and deferred.
- Immediate mode acts as before, registering actor components and running construction scripts as they are loaded.
- Deferred mode will buffer all actors loaded in the scope, and register all components and then run construction scripts.
- Loading in the editor through new loader adapters will now use deferred mode (loading regions, actor pinning, etc).

#rnx

diff --git a/Engine/Source/Runtime/Engine/Classes/Engine/Level.h b/Engine/Source/Runtime/Engine/Classes/Engine/Level.h
--- a/Engine/Source/Runtime/Engine/Classes/Engine/Level.h
+++ b/Engine/Source/Runtime/Engine/Classes/Engine/Level.h
@@ -1026,6 +1026,7 @@ public:
 	 * @param TransformToApply	The transform to apply to this actor if it's not already in the level
 	 */
 	ENGINE_API void AddLoadedActor(AActor* Actor, const FTransform* TransformToApply = nullptr);
+	ENGINE_API void AddLoadedActors(const TArray<AActor*>& ActorList, const FTransform* TransformToApply = nullptr);
 
 	/**
 	 * Remove a dynamically loaded actor from this level.
@@ -1034,6 +1035,7 @@ public:
 	 * @param TransformToRemove	The transform that was applied to this actor when it was added to the level
 	 */
 	ENGINE_API void RemoveLoadedActor(AActor* Actor, const FTransform* TransformToRemove = nullptr);
+	ENGINE_API void RemoveLoadedActors(const TArray<AActor*>& ActorList, const FTransform* TransformToRemove = nullptr);
 
 	/** Called when dynamically loaded actor is added to this level */
 	DECLARE_EVENT_OneParam(ULevel, FLoadedActorAddedToLevelEvent, AActor&);
diff --git a/Engine/Source/Runtime/Engine/Private/Level.cpp b/Engine/Source/Runtime/Engine/Private/Level.cpp
--- a/Engine/Source/Runtime/Engine/Private/Level.cpp
+++ b/Engine/Source/Runtime/Engine/Private/Level.cpp
@@ -817,15 +817,46 @@ bool ULevel::IsNetActor(const AActor* Actor)
 #if WITH_EDITOR
 void ULevel::AddLoadedActor(AActor* Actor, const FTransform* TransformToApply)
 {
-	check(Actor->GetLevel() == this);
-	check(IsValidChecked(Actor));
+	AddLoadedActors({ Actor }, TransformToApply);
+}
+
+void ULevel::AddLoadedActors(const TArray<AActor*>& ActorList, const FTransform* TransformToApply)
+{
+	TArray<AActor*> ActorsQueue;
+	ActorsQueue.Reserve(ActorList.Num());
+
+	TFunction<void(AActor* Actor)> QueueActor = [this, &ActorsQueue, &QueueActor](AActor* Actor)
+	{
+		check(Actor);
+		check(Actor->GetLevel() == this);
+		check(IsValidChecked(Actor));
+
+		int32 ActorIndex;
+		if (!Actors.Find(Actor, ActorIndex))
+		{
+			Actors.Add(Actor);
+			ActorsForGC.Add(Actor);
+			ActorsQueue.Add(Actor);
+
+			// Handle child actors
+			Actor->ForEachComponent<UChildActorComponent>(false, [this, &QueueActor](UChildActorComponent* ChildActorComponent)
+			{
+				if (AActor* ChildActor = ChildActorComponent->GetChildActor())
+				{
+					QueueActor(ChildActor);
+				}
+			});
+		}
+	};
 
-	int32 ActorIndex;
-	if (!Actors.Find(Actor, ActorIndex))
+	for (AActor* Actor : ActorList)
 	{
-		Actors.Add(Actor);
-		ActorsForGC.Add(Actor);
+		QueueActor(Actor);
+	}
 
+	// Register all components
+	for (AActor* Actor : ActorsQueue)
+	{
 		if (TransformToApply)
 		{
 			FLevelUtils::FApplyLevelTransformParams TransformParams(this, *TransformToApply);
@@ -834,12 +865,26 @@ void ULevel::AddLoadedActor(AActor* Actor, const FTransform* TransformToApply)
 			FLevelUtils::ApplyLevelTransform(TransformParams);
 		}
 
-		// If components from actors belonging to this level were already registered, do the same for the newly loaded actor.
-		// Otherwise, the new actor components will be registered later once the world initialization is completed, through UpdateWorldComponents()
 		if (bAreComponentsCurrentlyRegistered)
 		{
 			Actor->RegisterAllComponents();
+		}
+	}
+
+	// Rerun construction scripts
+	for (AActor* Actor : ActorsQueue)
+	{
+		if (bAreComponentsCurrentlyRegistered)
+		{
 			Actor->RerunConstructionScripts();
+		}
+	}
+
+	// Finalize actors
+	for (AActor* Actor : ActorsQueue)
+	{
+		if (bAreComponentsCurrentlyRegistered)
+		{
 			GetWorld()->UpdateCullDistanceVolumes(Actor);
 			Actor->MarkComponentsRenderStateDirty();
 		}
@@ -850,51 +895,63 @@ void ULevel::AddLoadedActor(AActor* Actor, const FTransform* TransformToApply)
 		}
 
 		OnLoadedActorAddedToLevelEvent.Broadcast(*Actor);
-
-		// Handle child actors
-		Actor->ForEachComponent<UChildActorComponent>(false, [this](UChildActorComponent* ChildActorComponent)
-		{
-			if (AActor* ChildActor = ChildActorComponent->GetChildActor())
-			{
-				AddLoadedActor(ChildActor);
-			}
-		});
 	}
 }
 
 void ULevel::RemoveLoadedActor(AActor* Actor, const FTransform* TransformToRemove)
 {
-	check(Actor);
-	check(Actor->GetLevel() == this);
-	check(IsValidChecked(Actor));
+	RemoveLoadedActors({ Actor }, TransformToRemove);
+}
+
+void ULevel::RemoveLoadedActors(const TArray<AActor*>& ActorList, const FTransform* TransformToRemove)
+{
+	TArray<AActor*> ActorsQueue;
+	ActorsQueue.Reserve(ActorList.Num());
 
-	// Handle child actors
-	Actor->ForEachComponent<UChildActorComponent>(false, [this](UChildActorComponent* ChildActorComponent)
+	TFunction<void(AActor* Actor)> QueueActor = [this, &ActorsQueue, &QueueActor](AActor* Actor)
 	{
-		if (AActor* ChildActor = ChildActorComponent->GetChildActor())
+		check(Actor);
+		check(Actor->GetLevel() == this);
+		check(IsValidChecked(Actor));
+
+		// Handle child actors
+		Actor->ForEachComponent<UChildActorComponent>(false, [this, &QueueActor](UChildActorComponent* ChildActorComponent)
 		{
-			RemoveLoadedActor(ChildActor);
-		}
-	});
+			if (AActor* ChildActor = ChildActorComponent->GetChildActor())
+			{
+				QueueActor(ChildActor);
+			}
+		});
 
-	Actor->UnregisterAllComponents();
-	Actor->RegisterAllActorTickFunctions(false, true);	
+		int32 ActorIndex;
+		verify(Actors.Find(Actor, ActorIndex));
 
-	int32 ActorIndex;
-	verify(Actors.Find(Actor, ActorIndex));
+		Actors[ActorIndex] = nullptr;
+		ActorsForGC.Remove(Actor);
 
-	Actors[ActorIndex] = nullptr;
-	ActorsForGC.Remove(Actor);
+		ActorsQueue.Add(Actor);
+	};
 
-	if (TransformToRemove)
+	for (AActor* Actor : ActorList)
 	{
-		FLevelUtils::FApplyLevelTransformParams TransformParams(this, TransformToRemove->Inverse());
-		TransformParams.Actor = Actor;
-		TransformParams.bDoPostEditMove = true;
-		FLevelUtils::ApplyLevelTransform(TransformParams);
+		QueueActor(Actor);
 	}
 
-	OnLoadedActorRemovedFromLevelEvent.Broadcast(*Actor);
+	for (AActor* Actor : ActorsQueue)
+	{
+		Actor->UnregisterAllComponents();
+		Actor->RegisterAllActorTickFunctions(false, true);
+
+		if (TransformToRemove)
+		{
+			FLevelUtils::FApplyLevelTransformParams TransformParams(this, TransformToRemove->Inverse());
+			TransformParams.Actor = Actor;
+			TransformParams.bDoPostEditMove = true;
+			FLevelUtils::ApplyLevelTransform(TransformParams);
+		}
+
+		OnLoadedActorRemovedFromLevelEvent.Broadcast(*Actor);
+	}
 }
 #endif
 
diff --git a/Engine/Source/Runtime/Engine/Private/WorldPartition/ActorDescContainer.cpp b/Engine/Source/Runtime/Engine/Private/WorldPartition/ActorDescContainer.cpp
--- a/Engine/Source/Runtime/Engine/Private/WorldPartition/ActorDescContainer.cpp
+++ b/Engine/Source/Runtime/Engine/Private/WorldPartition/ActorDescContainer.cpp
@@ -267,6 +267,7 @@ void UActorDescContainer::RemoveActor(const FGuid& ActorGuid)
 
 void UActorDescContainer::LoadAllActors(TArray<FWorldPartitionReference>& OutReferences)
 {
+	FWorldPartitionLoadingContext::FDeferred LoadingContext;
 	OutReferences.Reserve(OutReferences.Num() + GetActorDescCount());
 	for (FActorDescList::TIterator<> ActorDescIterator(this); ActorDescIterator; ++ActorDescIterator)
 	{
diff --git a/Engine/Source/Runtime/Engine/Private/WorldPartition/WorldPartition.cpp b/Engine/Source/Runtime/Engine/Private/WorldPartition/WorldPartition.cpp
--- a/Engine/Source/Runtime/Engine/Private/WorldPartition/WorldPartition.cpp
+++ b/Engine/Source/Runtime/Engine/Private/WorldPartition/WorldPartition.cpp
@@ -526,7 +526,7 @@ void UWorldPartition::Initialize(UWorld* InWorld, const FTransform& InTransform)
 		}
 
 		// Apply level transform on actors already part of the level
-		if (GetInstanceTransformPtr())
+		if (!GetInstanceTransform().Equals(FTransform::Identity))
 		{
 			check(!OuterWorld->PersistentLevel->bAlreadyMovedActors);
 			for (AActor* Actor : OuterWorld->PersistentLevel->Actors)
@@ -672,20 +672,6 @@ bool UWorldPartition::IsMainWorldPartition() const
 	return World == GetTypedOuter<UWorld>();
 }
 
-const FTransform& UWorldPartition::GetInstanceTransform() const
-{
-	if (const FTransform* TransformPtr = GetInstanceTransformPtr())
-	{
-		return *TransformPtr;
-	}
-	return FTransform::Identity;
-}
-
-const FTransform* UWorldPartition::GetInstanceTransformPtr() const
-{
-	return InstanceTransform.IsSet() ? &InstanceTransform.GetValue() : nullptr;
-}
-
 void UWorldPartition::OnPostBugItGoCalled(const FVector& Loc, const FRotator& Rot)
 {
 #if WITH_EDITOR
@@ -932,22 +918,6 @@ void UWorldPartition::OnActorDescUpdated(FWorldPartitionActorDesc* ActorDesc)
 	}
 }
 
-void UWorldPartition::OnActorDescRegistered(const FWorldPartitionActorDesc& ActorDesc) 
-{
-	AActor* Actor = ActorDesc.GetActor();
-	check(Actor);
-	Actor->GetLevel()->AddLoadedActor(Actor, GetInstanceTransformPtr());
-}
-
-void UWorldPartition::OnActorDescUnregistered(const FWorldPartitionActorDesc& ActorDesc) 
-{
-	AActor* Actor = ActorDesc.GetActor();
-	if (IsValidChecked(Actor))
-	{
-		Actor->GetLevel()->RemoveLoadedActor(Actor, GetInstanceTransformPtr());
-	}
-}
-
 bool UWorldPartition::GetInstancingContext(const FLinkerInstancingContext*& OutInstancingContext, FSoftObjectPathFixupArchive*& OutSoftObjectPathFixupArchive) const
 {
 	if (InstancingContext.IsInstanced())
@@ -958,7 +928,14 @@ bool UWorldPartition::GetInstancingContext(const FLinkerInstancingContext*& OutI
 	}
 	return false;
 }
+#endif
 
+const FTransform& UWorldPartition::GetInstanceTransform() const
+{
+	return InstanceTransform.IsSet() ? InstanceTransform.GetValue() : Super::GetInstanceTransform();
+}
+
+#if WITH_EDITOR
 bool UWorldPartition::SupportsStreaming() const
 {
 	return World ? World->GetWorldSettings()->SupportsWorldPartitionStreaming() : false;
diff --git a/Engine/Source/Runtime/Engine/Private/WorldPartition/WorldPartitionActorDesc.cpp b/Engine/Source/Runtime/Engine/Private/WorldPartition/WorldPartitionActorDesc.cpp
--- a/Engine/Source/Runtime/Engine/Private/WorldPartition/WorldPartitionActorDesc.cpp
+++ b/Engine/Source/Runtime/Engine/Private/WorldPartition/WorldPartitionActorDesc.cpp
@@ -494,22 +494,4 @@ void FWorldPartitionActorDesc::Unload()
 		ActorPtr = nullptr;
 	}
 }
-
-void FWorldPartitionActorDesc::RegisterActor()
-{
-	if (AActor* Actor = GetActor())
-	{
-		check(Container);
-		Container->OnActorDescRegistered(*this);
-	}
-}
-
-void FWorldPartitionActorDesc::UnregisterActor()
-{
-	if (AActor* Actor = GetActor())
-	{
-		check(Container);
-		Container->OnActorDescUnregistered(*this);
-	}
-}
 #endif
\ No newline at end of file
diff --git a/Engine/Source/Runtime/Engine/Private/WorldPartition/WorldPartitionActorLoaderInterface.cpp b/Engine/Source/Runtime/Engine/Private/WorldPartition/WorldPartitionActorLoaderInterface.cpp
--- a/Engine/Source/Runtime/Engine/Private/WorldPartition/WorldPartitionActorLoaderInterface.cpp
+++ b/Engine/Source/Runtime/Engine/Private/WorldPartition/WorldPartitionActorLoaderInterface.cpp
@@ -69,7 +69,11 @@ bool IWorldPartitionActorLoaderInterface::ILoaderAdapter::Unload()
 
 			UnregisterDelegates();
 
-			ActorReferences.Empty();
+			{
+				FWorldPartitionLoadingContext::FDeferred LoadingContext;
+				ActorReferences.Empty();
+			}
+
 			bLoaded = false;
 
 			PostLoadedStateChanged(true);
@@ -127,11 +131,15 @@ bool IWorldPartitionActorLoaderInterface::ILoaderAdapter::RefreshLoadedState()
 			FScopedSlowTask SlowTask(ActorsToLoad.Num() + ActorsToUnload.Num(), LOCTEXT("UpdatingLoading", "Updating loading..."));
 			SlowTask.MakeDialog();
 
-			for (FWorldPartitionHandle& ActorToLoad : ActorsToLoad)
 			{
-				SlowTask.EnterProgressFrame(1);
-				FWorldPartitionHandle ActorHandle(WorldPartition, ActorToLoad->GetGuid());
-				AddReferenceToActor(ActorHandle);
+				FWorldPartitionLoadingContext::FDeferred LoadingContext;
+
+				for (FWorldPartitionHandle& ActorToLoad : ActorsToLoad)
+				{
+					SlowTask.EnterProgressFrame(1);
+					FWorldPartitionHandle ActorHandle(WorldPartition, ActorToLoad->GetGuid());
+					AddReferenceToActor(ActorHandle);
+				}
 			}
 
 			for (FWorldPartitionHandle& ActorToUnload : ActorsToUnload)
diff --git a/Engine/Source/Runtime/Engine/Private/WorldPartition/WorldPartitionHandle.cpp b/Engine/Source/Runtime/Engine/Private/WorldPartition/WorldPartitionHandle.cpp
--- a/Engine/Source/Runtime/Engine/Private/WorldPartition/WorldPartitionHandle.cpp
+++ b/Engine/Source/Runtime/Engine/Private/WorldPartition/WorldPartitionHandle.cpp
@@ -5,6 +5,9 @@
 #include "WorldPartition/ActorDescContainer.h"
 
 #if WITH_EDITOR
+/**
+* FWorldPartitionHandleUtils
+*/
 TUniquePtr<FWorldPartitionActorDesc>* FWorldPartitionHandleUtils::GetActorDesc(UActorDescContainer* Container, const FGuid& ActorGuid)
 {
 	if (TUniquePtr<FWorldPartitionActorDesc>** ActorDescPtr = Container->ActorsByGuid.Find(ActorGuid))
@@ -25,6 +28,168 @@ bool FWorldPartitionHandleUtils::IsActorDescLoaded(FWorldPartitionActorDesc* Act
 	return ActorDesc->IsLoaded();
 }
 
+/**
+* FWorldPartitionLoadingContext
+*/
+FWorldPartitionLoadingContext::FImmediate FWorldPartitionLoadingContext::DefaultContext;
+FWorldPartitionLoadingContext::IContext* FWorldPartitionLoadingContext::ActiveContext = &DefaultContext;
+
+void FWorldPartitionLoadingContext::LoadAndRegisterActor(FWorldPartitionActorDesc* ActorDesc)
+{
+	ActiveContext->RegisterActor(ActorDesc);
+}
+
+void FWorldPartitionLoadingContext::UnloadAndUnregisterActor(FWorldPartitionActorDesc* ActorDesc)
+{
+	ActiveContext->UnregisterActor(ActorDesc);
+}
+
+/**
+* IContext
+*/
+FWorldPartitionLoadingContext::IContext::IContext()
+{
+	check(ActiveContext == &DefaultContext);
+	ActiveContext = this;
+}
+
+FWorldPartitionLoadingContext::IContext::~IContext()
+{
+	check(ActiveContext == this);
+	ActiveContext = &DefaultContext;
+}
+
+/**
+* FImmediate
+*/
+void FWorldPartitionLoadingContext::FImmediate::RegisterActor(FWorldPartitionActorDesc* ActorDesc)
+{
+	TGuardValue<bool> IsEditorLoadingPackageGuard(GIsEditorLoadingPackage, true);
+
+	if (AActor* Actor = ActorDesc->Load())
+	{
+		UActorDescContainer* Container = ActorDesc->GetContainer();
+		check(Container);
+
+		const FTransform& ContainerTransform = Container->GetInstanceTransform();
+		const FTransform* ContainerTransformPtr = ContainerTransform.Equals(FTransform::Identity) ? nullptr : &ContainerTransform;
+
+		Actor->GetLevel()->AddLoadedActor(Actor, ContainerTransformPtr);
+	}
+}
+
+void FWorldPartitionLoadingContext::FImmediate::UnregisterActor(FWorldPartitionActorDesc* ActorDesc)
+{
+	if (AActor* Actor = ActorDesc->GetActor())
+	{
+		UActorDescContainer* Container = ActorDesc->GetContainer();
+		check(Container);
+
+		const FTransform& ContainerTransform = Container->GetInstanceTransform();
+		const FTransform* ContainerTransformPtr = ContainerTransform.Equals(FTransform::Identity) ? nullptr : &ContainerTransform;
+
+		Actor->GetLevel()->RemoveLoadedActor(Actor, ContainerTransformPtr);
+
+		ActorDesc->Unload();
+	}	
+}
+
+/**
+* FDeferred
+*/
+FWorldPartitionLoadingContext::FDeferred::~FDeferred()
+{
+	for (auto [Container, ContainerOp] : ContainerOps)
+	{
+		const FTransform& ContainerTransform = Container->GetInstanceTransform();
+		const FTransform* ContainerTransformPtr = ContainerTransform.Equals(FTransform::Identity) ? nullptr : &ContainerTransform;
+
+		if (ContainerOp.Registrations.Num())
+		{
+			TGuardValue<bool> IsEditorLoadingPackageGuard(GIsEditorLoadingPackage, true);
+
+			TArray<AActor*> ActorList;
+			ActorList.Reserve(ContainerOp.Registrations.Num());
+
+			ULevel* Level = nullptr;
+			for (FWorldPartitionActorDesc* ActorDesc : ContainerOp.Registrations)
+			{
+				if (AActor* Actor = ActorDesc->GetActor())
+				{
+					ActorList.Add(Actor);
+
+					ULevel* ActorLevel = Actor->GetLevel();
+					check(!Level || (Level == ActorLevel));
+					Level = ActorLevel;
+				}
+			}
+			
+			if (ActorList.Num())
+			{
+				check(Level);
+				Level->AddLoadedActors(ActorList, ContainerTransformPtr);
+			}
+		}
+
+		if (ContainerOp.Unregistrations.Num())
+		{
+			TArray<AActor*> ActorList;
+			ActorList.Reserve(ContainerOp.Registrations.Num());
+
+			ULevel* Level = nullptr;
+			for (FWorldPartitionActorDesc* ActorDesc : ContainerOp.Unregistrations)
+			{
+				if (AActor* Actor = ActorDesc->GetActor())
+				{
+					ActorList.Add(Actor);
+
+					ULevel* ActorLevel = Actor->GetLevel();
+					check(!Level || (Level == ActorLevel));
+					Level = ActorLevel;
+				}
+			}
+			
+			if (ActorList.Num())
+			{
+				check(Level);
+				Level->RemoveLoadedActors(ActorList, ContainerTransformPtr);
+			}
+
+			for (FWorldPartitionActorDesc* ActorDesc : ContainerOp.Unregistrations)
+			{
+				ActorDesc->Unload();
+			}
+		}
+	}
+}
+
+void FWorldPartitionLoadingContext::FDeferred::RegisterActor(FWorldPartitionActorDesc* ActorDesc)
+{
+	check(ActorDesc);
+
+	UActorDescContainer* Container = ActorDesc->GetContainer();
+	check(Container);
+
+	ContainerOps.FindOrAdd(Container).Registrations.Add(ActorDesc);
+	check(!ContainerOps.FindChecked(Container).Unregistrations.Contains(ActorDesc));
+
+	ActorDesc->Load();
+}
+
+void FWorldPartitionLoadingContext::FDeferred::UnregisterActor(FWorldPartitionActorDesc* ActorDesc)
+{
+	check(ActorDesc);
+
+	UActorDescContainer* Container = ActorDesc->GetContainer();
+	check(Container);
+
+	ContainerOps.FindOrAdd(Container).Unregistrations.Add(ActorDesc);
+	check(!ContainerOps.FindChecked(Container).Registrations.Contains(ActorDesc));
+}
+
+/**
+* FWorldPartitionHandleImpl
+*/
 void FWorldPartitionHandleImpl::IncRefCount(FWorldPartitionActorDesc* ActorDesc)
 {
 	ActorDesc->IncSoftRefCount();
@@ -35,13 +200,14 @@ void FWorldPartitionHandleImpl::DecRefCount(FWorldPartitionActorDesc* ActorDesc)
 	ActorDesc->DecSoftRefCount();
 }
 
+/**
+* FWorldPartitionReferenceImpl
+*/
 void FWorldPartitionReferenceImpl::IncRefCount(FWorldPartitionActorDesc* ActorDesc)
 {
 	if (ActorDesc->IncHardRefCount() == 1)
 	{
-		TGuardValue<bool> IsEditorLoadingPackageGuard(GIsEditorLoadingPackage, true);
-		ActorDesc->Load();
-		ActorDesc->RegisterActor();
+		FWorldPartitionLoadingContext::LoadAndRegisterActor(ActorDesc);
 	}
 }
 
@@ -49,8 +215,7 @@ void FWorldPartitionReferenceImpl::DecRefCount(FWorldPartitionActorDesc* ActorDe
 {
 	if (ActorDesc->DecHardRefCount() == 0)
 	{
-		ActorDesc->UnregisterActor();
-		ActorDesc->Unload();
+		FWorldPartitionLoadingContext::UnloadAndUnregisterActor(ActorDesc);
 	}
 }
 #endif
\ No newline at end of file
diff --git a/Engine/Source/Runtime/Engine/Private/WorldPartition/WorldPartitionHelpers.cpp b/Engine/Source/Runtime/Engine/Private/WorldPartition/WorldPartitionHelpers.cpp
--- a/Engine/Source/Runtime/Engine/Private/WorldPartition/WorldPartitionHelpers.cpp
+++ b/Engine/Source/Runtime/Engine/Private/WorldPartition/WorldPartitionHelpers.cpp
@@ -64,6 +64,8 @@ namespace WorldPartitionHelpers
 {
 	void LoadReferences(UWorldPartition* WorldPartition, const FGuid& ActorGuid, TMap<FGuid, FWorldPartitionReference>& InOutActorReferences)
 	{
+		FWorldPartitionLoadingContext::FDeferred LoadingContext;
+
 		if (InOutActorReferences.Contains(ActorGuid))
 		{
 			return;
diff --git a/Engine/Source/Runtime/Engine/Private/WorldPartition/WorldPartitionRuntimeSpatialHash.cpp b/Engine/Source/Runtime/Engine/Private/WorldPartition/WorldPartitionRuntimeSpatialHash.cpp
--- a/Engine/Source/Runtime/Engine/Private/WorldPartition/WorldPartitionRuntimeSpatialHash.cpp
+++ b/Engine/Source/Runtime/Engine/Private/WorldPartition/WorldPartitionRuntimeSpatialHash.cpp
@@ -924,6 +924,7 @@ bool UWorldPartitionRuntimeSpatialHash::CreateStreamingGrid(const FSpatialHashRu
 				FilteredActors.Reset(GridCellDataChunk.GetActors().Num());
 				if (GridCellDataChunk.GetActors().Num())
 				{
+					FWorldPartitionLoadingContext::FDeferred LoadingContext;
 					for (const FActorInstance& ActorInstance : GridCellDataChunk.GetActors())
 					{
 						if (bIsMainWorldPartition && !IsRunningCookCommandlet())
@@ -941,7 +942,7 @@ bool UWorldPartitionRuntimeSpatialHash::CreateStreamingGrid(const FSpatialHashRu
 								{
 									// This will load the actor if it isn't already loaded
 									FWorldPartitionReference Reference(WorldPartition, ActorInstance.Actor);
-									AActor* AlwaysLoadedActor = FindObject<AActor>(nullptr, *ActorDescView.GetActorPath().ToString());
+									AActor* AlwaysLoadedActor = Reference->GetActor();
 									AlwaysLoadedActorsForPIE.Emplace(Reference, AlwaysLoadedActor);
 
 									// Handle child actors
diff --git a/Engine/Source/Runtime/Engine/Public/WorldPartition/ActorDescContainer.h b/Engine/Source/Runtime/Engine/Public/WorldPartition/ActorDescContainer.h
--- a/Engine/Source/Runtime/Engine/Public/WorldPartition/ActorDescContainer.h
+++ b/Engine/Source/Runtime/Engine/Public/WorldPartition/ActorDescContainer.h
@@ -37,7 +37,11 @@ public:
 
 	/** Instancing support */
 	virtual bool GetInstancingContext(const FLinkerInstancingContext*& OutInstancingContext, FSoftObjectPathFixupArchive*& OutSoftObjectPathFixupArchive) const { return false; }
+#endif
+
+	virtual const FTransform& GetInstanceTransform() const { return FTransform::Identity; }
 
+#if WITH_EDITOR
 public:
 	DECLARE_EVENT_OneParam(UWorldPartition, FActorDescAddedEvent, FWorldPartitionActorDesc*);
 	FActorDescAddedEvent OnActorDescAddedEvent;
@@ -72,9 +76,6 @@ protected:
 	virtual void OnActorDescUpdating(FWorldPartitionActorDesc* ActorDesc) {}
 	virtual void OnActorDescUpdated(FWorldPartitionActorDesc* ActorDesc) {}
 
-	virtual void OnActorDescRegistered(const FWorldPartitionActorDesc&) {}
-	virtual void OnActorDescUnregistered(const FWorldPartitionActorDesc&) {}
-
 	bool ShouldHandleActorEvent(const AActor* Actor);
 
 	bool bContainerInitialized;
diff --git a/Engine/Source/Runtime/Engine/Public/WorldPartition/WorldPartition.h b/Engine/Source/Runtime/Engine/Public/WorldPartition/WorldPartition.h
--- a/Engine/Source/Runtime/Engine/Public/WorldPartition/WorldPartition.h
+++ b/Engine/Source/Runtime/Engine/Public/WorldPartition/WorldPartition.h
@@ -128,14 +128,13 @@ private:
 	virtual void OnActorDescUpdating(FWorldPartitionActorDesc* ActorDesc) override;
 	virtual void OnActorDescUpdated(FWorldPartitionActorDesc* ActorDesc) override;
 
-	virtual void OnActorDescRegistered(const FWorldPartitionActorDesc&) override;
-	virtual void OnActorDescUnregistered(const FWorldPartitionActorDesc&) override;
-
 	virtual bool GetInstancingContext(const FLinkerInstancingContext*& OutInstancingContext, FSoftObjectPathFixupArchive*& OutSoftObjectPathFixupArchive) const override;
-	//~ End UActorDescContainer Interface
 #endif
 
 public:
+	virtual const FTransform& GetInstanceTransform() const override;
+	//~ End UActorDescContainer Interface
+
 	//~ Begin UObject Interface
 	virtual bool ResolveSubobject(const TCHAR* SubObjectPath, UObject*& OutObject, bool bLoadIfExists) override;
 	//~ End UObject Interface
@@ -182,8 +181,6 @@ public:
 
 	bool CanStream() const;
 	bool IsMainWorldPartition() const;
-	const FTransform& GetInstanceTransform() const;
-	const FTransform* GetInstanceTransformPtr() const;
 
 	void Tick(float DeltaSeconds);
 	void UpdateStreamingState();
diff --git a/Engine/Source/Runtime/Engine/Public/WorldPartition/WorldPartitionActorDesc.h b/Engine/Source/Runtime/Engine/Public/WorldPartition/WorldPartitionActorDesc.h
--- a/Engine/Source/Runtime/Engine/Public/WorldPartition/WorldPartitionActorDesc.h
+++ b/Engine/Source/Runtime/Engine/Public/WorldPartition/WorldPartitionActorDesc.h
@@ -153,9 +153,6 @@ public:
 
 	virtual bool ShouldBeLoadedByEditor() const { return true; }
 
-	void RegisterActor();
-	void UnregisterActor();
-
 	virtual void Init(const AActor* InActor);
 	virtual void Init(const FWorldPartitionActorDescInitData& DescData);
 
diff --git a/Engine/Source/Runtime/Engine/Public/WorldPartition/WorldPartitionHandle.h b/Engine/Source/Runtime/Engine/Public/WorldPartition/WorldPartitionHandle.h
--- a/Engine/Source/Runtime/Engine/Public/WorldPartition/WorldPartitionHandle.h
+++ b/Engine/Source/Runtime/Engine/Public/WorldPartition/WorldPartitionHandle.h
@@ -17,6 +17,66 @@ struct ENGINE_API FWorldPartitionHandleUtils
 	static bool IsActorDescLoaded(FWorldPartitionActorDesc* ActorDesc);
 };
 
+class ENGINE_API FWorldPartitionLoadingContext
+{
+public:
+	/**
+	 * Base class for loading contexts
+	 */
+	class ENGINE_API IContext
+	{
+		friend class FWorldPartitionLoadingContext;
+
+	public:
+		IContext();
+		virtual ~IContext();
+
+	private:
+		virtual void RegisterActor(FWorldPartitionActorDesc* ActorDesc) =0;
+		virtual void UnregisterActor(FWorldPartitionActorDesc* ActorDesc) =0;
+	};
+
+	/**
+	 * Immediate loading context, which will register and unregister actors on demand.
+	 */
+	class ENGINE_API FImmediate : public IContext
+	{
+	private:
+		virtual void RegisterActor(FWorldPartitionActorDesc* ActorDesc) override;
+		virtual void UnregisterActor(FWorldPartitionActorDesc* ActorDesc) override;
+	};
+
+	/**
+	 * Deferred loading context, which will gather actor registrations and unregistrations and
+	 * execute them at the end of scope. This respects ULevel components registrations and
+	 * construction scripts execution logic.
+	 */
+	class ENGINE_API FDeferred : public IContext
+	{
+	public:
+		~FDeferred();
+
+	private:
+		virtual void RegisterActor(FWorldPartitionActorDesc* ActorDesc) override;
+		virtual void UnregisterActor(FWorldPartitionActorDesc* ActorDesc) override;
+
+		struct FContainerOps
+		{
+			TArray<FWorldPartitionActorDesc*> Registrations;
+			TArray<FWorldPartitionActorDesc*> Unregistrations;
+		};
+
+		TMap<UActorDescContainer*, FContainerOps> ContainerOps;
+	};
+
+	static void LoadAndRegisterActor(FWorldPartitionActorDesc* ActorDesc);
+	static void UnloadAndUnregisterActor(FWorldPartitionActorDesc* ActorDesc);
+
+private:
+	static FImmediate DefaultContext;
+	static IContext* ActiveContext;	
+};
+
 template <typename Impl>
 class TWorldPartitionHandle
 {
