explicit description of ETextureLossyCompressionAmount rdo settings
nop just comment changes

diff --git a/Engine/Source/Runtime/Engine/Classes/Engine/TextureDefines.h b/Engine/Source/Runtime/Engine/Classes/Engine/TextureDefines.h
--- a/Engine/Source/Runtime/Engine/Classes/Engine/TextureDefines.h
+++ b/Engine/Source/Runtime/Engine/Classes/Engine/TextureDefines.h
@@ -219,12 +219,12 @@ UENUM()
 enum ETextureLossyCompressionAmount
 {
 	TLCA_Default		UMETA(DisplayName = "Default"),
-	TLCA_None			UMETA(DisplayName = "No lossy compression"),
-	TLCA_Lowest			UMETA(DisplayName = "Lowest (Best image quality, largest filesize)"),
-	TLCA_Low			UMETA(DisplayName = "Low"),
-	TLCA_Medium			UMETA(DisplayName = "Medium"),
-	TLCA_High			UMETA(DisplayName = "High"),
-	TLCA_Highest		UMETA(DisplayName = "Highest (Worst image quality, smallest filesize)"),
+	TLCA_None			UMETA(DisplayName = "No lossy compression (Oodle RDO disabled)"),
+	TLCA_Lowest			UMETA(DisplayName = "Lowest (Best image quality, largest filesize) (Oodle RDO 1)"),
+	TLCA_Low			UMETA(DisplayName = "Low (Oodle RDO 10)"),
+	TLCA_Medium			UMETA(DisplayName = "Medium (Oodle RDO 20)"),
+	TLCA_High			UMETA(DisplayName = "High (Oodle RDO 30)"),
+	TLCA_Highest		UMETA(DisplayName = "Highest (Worst image quality, smallest filesize) (Oodle RDO 40)"),
 };
 
 // Certain settings can be changed to facilitate how fast a texture build takes. This
diff --git a/Engine/Source/Runtime/Engine/Private/TextureDerivedData.cpp b/Engine/Source/Runtime/Engine/Private/TextureDerivedData.cpp
--- a/Engine/Source/Runtime/Engine/Private/TextureDerivedData.cpp
+++ b/Engine/Source/Runtime/Engine/Private/TextureDerivedData.cpp
@@ -595,6 +595,7 @@ static void FinalizeBuildSettingsForLayer(
 		if (Options.bUsesRDO)
 		{
 			// If this mapping changes, update the tooltip in TextureEncodingSettings.h
+			// this is an ETextureLossyCompressionAmount
 			switch (OutSettings.LossyCompressionAmount)
 			{
 			default:
