World Partition: fix crash during conversion commandlet.

#rnx

diff --git a/Engine/Source/Runtime/Engine/Private/World.cpp b/Engine/Source/Runtime/Engine/Private/World.cpp
--- a/Engine/Source/Runtime/Engine/Private/World.cpp
+++ b/Engine/Source/Runtime/Engine/Private/World.cpp
@@ -8473,8 +8473,8 @@ void UWorld::GetAssetRegistryTags(TArray<FAssetRegistryTag>& OutTags) const
 			OutTags.Add(FAssetRegistryTag(NAME_LevelIsUsingActorFolders, TEXT("1"), FAssetRegistryTag::TT_Hidden));
 		}
 
+		if (AWorldSettings* WorldSettings = GetWorldSettings(/*bCheckStreamingPersistent*/false, /*bChecked*/false))
 		{
-			AWorldSettings* WorldSettings = GetWorldSettings();
 			FVector LevelInstancePivotOffset = WorldSettings ? WorldSettings->LevelInstancePivotOffset : FVector::ZeroVector;
 			if (!LevelInstancePivotOffset.IsNearlyZero())
 			{
