Change r.Streaming.UseTextureStreamingBuiltData to a cvar so it can be set in config files. Also set it to 0 to workaround blurry texture issues. We seem to have outdated or broken pre-built streaming data which needs more investigation.

#localization none
#tests A/B tested on a console platform

diff --git a/Engine/Source/Runtime/Engine/Classes/Engine/TextureStreamingTypes.h b/Engine/Source/Runtime/Engine/Classes/Engine/TextureStreamingTypes.h
--- a/Engine/Source/Runtime/Engine/Classes/Engine/TextureStreamingTypes.h
+++ b/Engine/Source/Runtime/Engine/Classes/Engine/TextureStreamingTypes.h
@@ -253,10 +253,6 @@ class ENGINE_API FStreamingTextureLevelContext
 	int32* GetBuildDataIndexRef(UTexture* Texture, bool bForceUpdate = false);
 	void UpdateQualityAndFeatureLevel(EMaterialQualityLevel::Type InQualityLevel, ERHIFeatureLevel::Type InFeatureLevel, const ULevel* InLevel = nullptr);
 
-	/** Console command used to turn on/off usage of texture streaming built data. */
-	static class FAutoConsoleCommand UseTextureStreamingBuiltDataCommand;
-	static bool UseTextureStreamingBuiltData;
-
 public:
 
 	// Needs InLevel to use precomputed data from 
diff --git a/Engine/Source/Runtime/Engine/Private/Streaming/TextureStreamingBuild.cpp b/Engine/Source/Runtime/Engine/Private/Streaming/TextureStreamingBuild.cpp
--- a/Engine/Source/Runtime/Engine/Private/Streaming/TextureStreamingBuild.cpp
+++ b/Engine/Source/Runtime/Engine/Private/Streaming/TextureStreamingBuild.cpp
@@ -437,11 +437,11 @@ uint32 FStreamingTextureBuildInfo::ComputeHash() const
 }
 #endif
 
-bool FStreamingTextureLevelContext::UseTextureStreamingBuiltData = true;
-FAutoConsoleCommand FStreamingTextureLevelContext::UseTextureStreamingBuiltDataCommand(
+static int32 GUseTextureStreamingBuiltData = 1;
+static FAutoConsoleVariableRef CVarUseTextureStreamingBuiltData(
 	TEXT("r.Streaming.UseTextureStreamingBuiltData"),
-	TEXT("Turn on/off usage of texture streaming built data (0 to turn off)."),
-	FConsoleCommandWithArgsDelegate::CreateLambda([](const TArray<FString>& Args) { FStreamingTextureLevelContext::UseTextureStreamingBuiltData = (Args.Num() != 1) || (Args[0] != TEXT("0")); }));
+	GUseTextureStreamingBuiltData,
+	TEXT("Turn on/off usage of texture streaming built data (0 to turn off)."));
 
 FStreamingTextureLevelContext::FStreamingTextureLevelContext(EMaterialQualityLevel::Type InQualityLevel, const UPrimitiveComponent* Primitive)
 	: TextureGuidToLevelIndex(nullptr)
@@ -544,7 +544,7 @@ void FStreamingTextureLevelContext::UpdateQualityAndFeatureLevel(EMaterialQualit
 
 bool FStreamingTextureLevelContext::CanUseTextureStreamingBuiltData() const
 {
-	return bIsBuiltDataValid && FStreamingTextureLevelContext::UseTextureStreamingBuiltData;
+	return bIsBuiltDataValid && GUseTextureStreamingBuiltData != 0;
 }
 
 FStreamingTextureLevelContext::~FStreamingTextureLevelContext()
