Fix potential replay crash in load checkpoint

diff --git a/Engine/Source/Runtime/Engine/Private/DemoNetDriver.cpp b/Engine/Source/Runtime/Engine/Private/DemoNetDriver.cpp
--- a/Engine/Source/Runtime/Engine/Private/DemoNetDriver.cpp
+++ b/Engine/Source/Runtime/Engine/Private/DemoNetDriver.cpp
@@ -3989,11 +3989,12 @@ bool UDemoNetDriver::LoadCheckpoint(const FGotoResult& GotoResult)
 
 	// Determine if an Actor has a reference to a spectator in some way.
 	// This prevents garbage collection on splitscreen playercontrollers
-	auto HasPlayerSpectatorRef = [this](const AActor* InActor) {
+	auto HasPlayerSpectatorRef = [this](const AActor* InActor) -> bool 
+	{
 		for (const APlayerController* CurSpectator : SpectatorControllers)
 		{
-			if (InActor == CurSpectator || InActor == CurSpectator->GetSpectatorPawn()
-				|| InActor->IsOwnedBy(CurSpectator))
+			if (IsValid(CurSpectator) &&
+				(InActor == CurSpectator || InActor == CurSpectator->GetSpectatorPawn()	|| InActor->IsOwnedBy(CurSpectator)))
 			{
 				return true;
 			}
