World Partition: properly handle pinning items from the scene outliner selection that aren't external actors.

#rnx

diff --git a/Engine/Source/Runtime/Engine/Private/WorldPartition/WorldPartition.cpp b/Engine/Source/Runtime/Engine/Private/WorldPartition/WorldPartition.cpp
--- a/Engine/Source/Runtime/Engine/Private/WorldPartition/WorldPartition.cpp
+++ b/Engine/Source/Runtime/Engine/Private/WorldPartition/WorldPartition.cpp
@@ -1309,8 +1309,11 @@
 	if (PinnedActors)
 	{
 		FWorldPartitionHandle ActorHandle(this, ActorGuid);
-		PinnedActors->AddActor(ActorHandle);
-		return ActorHandle->GetActor();
+		if (ActorHandle.IsValid())
+		{
+			PinnedActors->AddActor(ActorHandle);
+			return ActorHandle->GetActor();
+		}		
 	}
 	return nullptr;
 }
@@ -1320,7 +1323,10 @@
 	if (PinnedActors)
 	{
 		FWorldPartitionHandle ActorHandle(this, ActorGuid);
-		PinnedActors->RemoveActor(ActorHandle);
+		if (ActorHandle.IsValid())
+		{
+			PinnedActors->RemoveActor(ActorHandle);
+		}
 	}
 }
 
@@ -1329,7 +1335,10 @@
 	if (PinnedActors)
 	{
 		FWorldPartitionHandle ActorHandle(const_cast<UWorldPartition*>(this), ActorGuid);
-		return PinnedActors->ContainsActor(ActorHandle);
+		if(ActorHandle.IsValid())
+		{
+			return PinnedActors->ContainsActor(ActorHandle);
+		}
 	}
 	return false;
 }
