Fix for DemoNetDriver causing ClearVoicePackets to be called multiple times in one tick

diff --git a/Engine/Source/Runtime/Engine/Classes/Engine/NetDriver.h b/Engine/Source/Runtime/Engine/Classes/Engine/NetDriver.h
--- a/Engine/Source/Runtime/Engine/Classes/Engine/NetDriver.h
+++ b/Engine/Source/Runtime/Engine/Classes/Engine/NetDriver.h
@@ -927,6 +927,8 @@
 
 	uint8 bSkipServerReplicateActors : 1;
 
+	uint8 bSkipClearVoicePackets : 1;
+
 public:
 	/**
 	 * If true, ignore timeouts completely.  Should be used only in development
diff --git a/Engine/Source/Runtime/Engine/Private/DemoNetDriver.cpp b/Engine/Source/Runtime/Engine/Private/DemoNetDriver.cpp
--- a/Engine/Source/Runtime/Engine/Private/DemoNetDriver.cpp
+++ b/Engine/Source/Runtime/Engine/Private/DemoNetDriver.cpp
@@ -696,6 +696,7 @@
 	MaxArchiveReadPos = 0;
 	bNeverApplyNetworkEmulationSettings = true;
 	bSkipServerReplicateActors = true;
+	bSkipClearVoicePackets = true;
 	bSkipStartupActorRollback = false;
 
 	if (!HasAnyFlags(RF_ClassDefaultObject))
diff --git a/Engine/Source/Runtime/Engine/Private/NetDriver.cpp b/Engine/Source/Runtime/Engine/Private/NetDriver.cpp
--- a/Engine/Source/Runtime/Engine/Private/NetDriver.cpp
+++ b/Engine/Source/Runtime/Engine/Private/NetDriver.cpp
@@ -381,6 +381,7 @@
 ,   bHasReplayConnection(false)
 ,   bMaySendProperties(false)
 ,   bSkipServerReplicateActors(false)
+,   bSkipClearVoicePackets(false)
 ,   bNoTimeouts(false)
 ,   bNeverApplyNetworkEmulationSettings(false)
 ,   bIsPeer(false)
@@ -1551,7 +1552,7 @@
 
 void UNetDriver::PostTickFlush()
 {
-	if (World)
+	if (World && !bSkipClearVoicePackets)
 	{
 		UOnlineEngineInterface::Get()->ClearVoicePackets(World);
 	}
