LevelInstance: extract ILevelInstanceEditorPivotInterface to allow implementation in different class hierarchy

diff --git a/Engine/Source/Editor/LevelInstanceEditor/Private/LevelInstancePivotDetails.cpp b/Engine/Source/Editor/LevelInstanceEditor/Private/LevelInstancePivotDetails.cpp
--- a/Engine/Source/Editor/LevelInstanceEditor/Private/LevelInstancePivotDetails.cpp
+++ b/Engine/Source/Editor/LevelInstanceEditor/Private/LevelInstancePivotDetails.cpp
@@ -11,7 +11,8 @@
 #include "Engine/Brush.h"
 #include "Engine/World.h"
 #include "LevelInstance/LevelInstanceInterface.h"
-#include "LevelInstance/LevelInstanceEditorPivotActor.h"
+#include "LevelInstance/LevelInstanceEditorPivotInterface.h"
+#include "LevelInstance/LevelInstanceEditorPivot.h"
 
 #define LOCTEXT_NAMESPACE "FLevelInstancePivotDetails"
 
@@ -34,7 +35,7 @@
 		return;
 	}
 
-	TWeakObjectPtr<ALevelInstancePivot> EditingObject = Cast<ALevelInstancePivot>(EditingObjects[0].Get());
+	TWeakObjectPtr<AActor> EditingObject = Cast<AActor>(EditingObjects[0].Get());
 
 	UWorld* World = EditingObject->GetWorld();
 
@@ -105,17 +106,17 @@
 		];
 }
 
-bool FLevelInstancePivotDetails::IsApplyButtonEnabled(TWeakObjectPtr<ALevelInstancePivot> LevelInstancePivot) const
+bool FLevelInstancePivotDetails::IsApplyButtonEnabled(TWeakObjectPtr<AActor> LevelInstancePivot) const
 {
 	return PivotType != ELevelInstancePivotType::Actor || PivotActor.IsValid();
 }
 
-FReply FLevelInstancePivotDetails::OnApplyButtonClicked(TWeakObjectPtr<ALevelInstancePivot> LevelInstancePivot)
+FReply FLevelInstancePivotDetails::OnApplyButtonClicked(TWeakObjectPtr<AActor> LevelInstancePivot)
 {
 	if (LevelInstancePivot.IsValid())
 	{
 		FScopedTransaction Transaction(LOCTEXT("ChangePivot", "Change Level Instance Pivot"));
-		LevelInstancePivot->SetPivot(PivotType, PivotActor.Get());
+		FLevelInstanceEditorPivotHelper::SetPivot(CastChecked<ILevelInstanceEditorPivotInterface>(LevelInstancePivot.Get()), PivotType, PivotActor.Get());
 	}
 	return FReply::Handled();
 }
diff --git a/Engine/Source/Editor/LevelInstanceEditor/Private/LevelInstancePivotDetails.h b/Engine/Source/Editor/LevelInstanceEditor/Private/LevelInstancePivotDetails.h
--- a/Engine/Source/Editor/LevelInstanceEditor/Private/LevelInstancePivotDetails.h
+++ b/Engine/Source/Editor/LevelInstanceEditor/Private/LevelInstancePivotDetails.h
@@ -8,7 +8,7 @@
 #include "LevelInstance/LevelInstanceTypes.h"
 #include "Input/Reply.h"
 
-class ALevelInstancePivot;
+class AActor;
 
 class FLevelInstancePivotDetails : public IDetailCustomization
 {
@@ -24,9 +24,9 @@
 	/** Use MakeInstance to create an instance of this class */
 	FLevelInstancePivotDetails();
 
-	bool IsApplyButtonEnabled(TWeakObjectPtr<ALevelInstancePivot> LevelInstancePivot) const;
-	FReply OnApplyButtonClicked(TWeakObjectPtr<ALevelInstancePivot> LevelInstancePivot);
-	
+	bool IsApplyButtonEnabled(TWeakObjectPtr<AActor> LevelInstancePivot) const;
+	FReply OnApplyButtonClicked(TWeakObjectPtr<AActor> LevelInstancePivot);
+
 	int32 GetSelectedPivotType() const;
 	void OnSelectedPivotTypeChanged(int32 NewValue, ESelectInfo::Type SelectionType);
 
@@ -40,7 +40,7 @@
 	FText GetSelectedPivotActorText() const;
 	void OnSelectedPivotActorChanged(TSharedPtr<FActorInfo> NewValue, ESelectInfo::Type SelectionType);
 	bool IsPivotActorSelectionEnabled() const;
-	
+
 	// Settings for Apply button
 	ELevelInstancePivotType PivotType;
 	TWeakObjectPtr<AActor> PivotActor;
diff --git a/Engine/Source/Runtime/Engine/Private/LevelInstance/LevelInstanceActor.cpp b/Engine/Source/Runtime/Engine/Private/LevelInstance/LevelInstanceActor.cpp
--- a/Engine/Source/Runtime/Engine/Private/LevelInstance/LevelInstanceActor.cpp
+++ b/Engine/Source/Runtime/Engine/Private/LevelInstance/LevelInstanceActor.cpp
@@ -12,6 +12,7 @@
 #if WITH_EDITOR
 #include "WorldPartition/WorldPartitionActorDesc.h"
 #include "WorldPartition/LevelInstance/LevelInstanceActorDesc.h"
+#include "LevelInstance/LevelInstanceEditorPivotActor.h"
 #endif
 
 #define LOCTEXT_NAMESPACE "LevelInstanceActor"
@@ -104,6 +105,11 @@
 	return Cast<ULevelInstanceComponent>(RootComponent);
 }
 
+TSubclassOf<AActor> ALevelInstance::GetEditorPivotClass() const
+{
+	return ALevelInstancePivot::StaticClass();
+}
+
 TUniquePtr<FWorldPartitionActorDesc> ALevelInstance::CreateClassActorDesc() const
 {
 	return TUniquePtr<FWorldPartitionActorDesc>(new FLevelInstanceActorDesc());
diff --git a/Engine/Source/Runtime/Engine/Private/LevelInstance/LevelInstanceEditorLevelStreaming.cpp b/Engine/Source/Runtime/Engine/Private/LevelInstance/LevelInstanceEditorLevelStreaming.cpp
--- a/Engine/Source/Runtime/Engine/Private/LevelInstance/LevelInstanceEditorLevelStreaming.cpp
+++ b/Engine/Source/Runtime/Engine/Private/LevelInstance/LevelInstanceEditorLevelStreaming.cpp
@@ -5,7 +5,7 @@
 #if WITH_EDITOR
 #include "LevelInstance/LevelInstanceInterface.h"
 #include "LevelInstance/LevelInstanceSubsystem.h"
-#include "LevelInstance/LevelInstanceEditorPivotActor.h"
+#include "LevelInstance/LevelInstanceEditorPivot.h"
 #include "EditorLevelUtils.h"
 #include "Editor.h"
 #include "Engine/LevelBounds.h"
@@ -109,7 +109,7 @@
 		}
 
 		// Create special actor that will handle changing the pivot of this level
-		ALevelInstancePivot::Create(LevelInstance, LevelStreaming);
+		FLevelInstanceEditorPivotHelper::Create(LevelInstance, LevelStreaming);
 
 		return LevelStreaming;
 	}
diff --git a/Engine/Source/Runtime/Engine/Private/LevelInstance/LevelInstanceEditorPivot.cpp b/Engine/Source/Runtime/Engine/Private/LevelInstance/LevelInstanceEditorPivot.cpp
new file mode 100644
--- /dev/null
+++ b/Engine/Source/Runtime/Engine/Private/LevelInstance/LevelInstanceEditorPivot.cpp
@@ -0,0 +1,116 @@
+// Copyright Epic Games, Inc. All Rights Reserved.
+
+#include "LevelInstance/LevelInstanceEditorPivot.h"
+#include "LevelInstance/LevelInstanceEditorPivotInterface.h"
+
+#if WITH_EDITOR
+
+#include "LevelInstance/LevelInstanceInterface.h"
+#include "LevelInstance/LevelInstanceSubsystem.h"
+#include "Engine/World.h"
+#include "Engine/LevelStreaming.h"
+#include "GameFramework/WorldSettings.h"
+#include "LevelUtils.h"
+#include "Editor.h"
+
+#endif
+
+ULevelInstanceEditorPivotInterface::ULevelInstanceEditorPivotInterface(const FObjectInitializer& ObjectInitializer)
+	: Super(ObjectInitializer)
+{}
+
+#if WITH_EDITOR
+
+void ILevelInstanceEditorPivotInterface::UpdateOffset()
+{
+	if (FLevelUtils::IsApplyingLevelTransform())
+	{
+		return;
+	}
+
+	AActor* Actor = CastChecked<AActor>(this);
+	// Offset Change is the relative transform of the pivot actor to its spawn transform (we only care about relative translation as we don't support rotation on pivot)
+	FVector LocalToPivot = Actor->GetActorTransform().GetRelativeTransform(InitState.ActorTransform).GetTranslation();
+
+	// We then apply that offset to the original pivot
+	FVector NewPivotOffset = InitState.LevelOffset - LocalToPivot;
+
+	AWorldSettings* WorldSettings = Actor->GetLevel()->GetWorldSettings();
+	if (!NewPivotOffset.Equals(WorldSettings->LevelInstancePivotOffset))
+	{
+		WorldSettings->Modify();
+		WorldSettings->LevelInstancePivotOffset = NewPivotOffset;
+	}
+}
+
+ILevelInstanceEditorPivotInterface* FLevelInstanceEditorPivotHelper::Create(ILevelInstanceInterface* LevelInstance, ULevelStreaming* LevelStreaming)
+{
+	TSubclassOf<AActor> EditorPivotClass = LevelInstance->GetEditorPivotClass();
+	if (!EditorPivotClass || !EditorPivotClass->ImplementsInterface(ULevelInstanceEditorPivotInterface::StaticClass()))
+	{
+		return nullptr;
+	}
+
+	FActorSpawnParameters SpawnParams;
+	SpawnParams.OverrideLevel = LevelStreaming->GetLoadedLevel();
+	SpawnParams.bCreateActorPackage = false;
+	SpawnParams.ObjectFlags |= RF_Transient;
+	SpawnParams.bNoFail = true;
+
+	// We place the pivot actor at the LevelInstance Transform so that it makes sense to the user (the pivot being the zero)
+	AActor* LevelInstanceActor = CastChecked<AActor>(LevelInstance);
+	AActor* PivotActor = LevelInstanceActor->GetWorld()->SpawnActor<AActor>(EditorPivotClass, LevelInstanceActor->GetActorLocation(), LevelInstanceActor->GetActorRotation(), SpawnParams);
+	ILevelInstanceEditorPivotInterface* PivotInterface = CastChecked<ILevelInstanceEditorPivotInterface>(PivotActor);
+
+	AWorldSettings* WorldSettings = LevelStreaming->GetLoadedLevel()->GetWorldSettings();
+
+	PivotInterface->SetInitState(ILevelInstanceEditorPivotInterface::FInitState { LevelInstance->GetLevelInstanceID(), PivotActor->GetActorTransform(), WorldSettings->LevelInstancePivotOffset });
+
+	// Set Label last as this will call PostEditChangeProperty and update the offset so other fields need to be setup first.
+	PivotActor->SetActorLabel(TEXT("Pivot"));
+
+	return PivotInterface;
+}
+
+void FLevelInstanceEditorPivotHelper::SetPivot(ILevelInstanceEditorPivotInterface* PivotInterface, ELevelInstancePivotType PivotType, AActor* PivotToActor)
+{
+	check(PivotType != ELevelInstancePivotType::Actor || PivotToActor != nullptr);
+	AActor* PivotActor = CastChecked<AActor>(PivotInterface);
+
+	PivotActor->Modify();
+	if (PivotType == ELevelInstancePivotType::Actor)
+	{
+		PivotActor->SetActorLocation(PivotToActor->GetActorLocation());
+	}
+	else if(PivotType == ELevelInstancePivotType::Center || PivotType == ELevelInstancePivotType::CenterMinZ)
+	{
+		ULevelInstanceSubsystem* LevelInstanceSubsystem = UWorld::GetSubsystem<ULevelInstanceSubsystem>(PivotActor->GetWorld());
+		ILevelInstanceInterface* LevelInstance = LevelInstanceSubsystem->GetLevelInstance(PivotInterface->GetLevelInstanceID());
+		FBox OutBounds;
+		LevelInstanceSubsystem->GetLevelInstanceBounds(LevelInstance, OutBounds);
+		FVector Location = OutBounds.GetCenter();
+		if (PivotType == ELevelInstancePivotType::CenterMinZ)
+		{
+			Location.Z = OutBounds.Min.Z;
+		}
+		PivotActor->SetActorLocation(Location);
+	}
+	else if (PivotType == ELevelInstancePivotType::WorldOrigin)
+	{
+		PivotActor->SetActorLocation(FVector(0.f, 0.f, 0.f));
+	}
+	else // unsupported
+	{
+		check(0);
+	}
+
+	// Update gizmo
+	if (GEditor)
+	{
+		GEditor->NoteSelectionChange();
+	}
+
+	PivotInterface->UpdateOffset();
+}
+
+#endif
\ No newline at end of file
diff --git a/Engine/Source/Runtime/Engine/Private/LevelInstance/LevelInstanceEditorPivotActor.cpp b/Engine/Source/Runtime/Engine/Private/LevelInstance/LevelInstanceEditorPivotActor.cpp
--- a/Engine/Source/Runtime/Engine/Private/LevelInstance/LevelInstanceEditorPivotActor.cpp
+++ b/Engine/Source/Runtime/Engine/Private/LevelInstance/LevelInstanceEditorPivotActor.cpp
@@ -1,131 +1,41 @@
-// Copyright Epic Games, Inc. All Rights Reserved.
-
-#include "LevelInstance/LevelInstanceEditorPivotActor.h"
-#include "LevelInstance/LevelInstanceInterface.h"
-#include "LevelInstance/LevelInstanceSubsystem.h"
-#include "GameFramework/WorldSettings.h"
-#include "LevelUtils.h"
-#include "Engine/LevelStreaming.h"
-
-#if WITH_EDITOR
-#include "Editor.h"
-#endif
-
-ALevelInstancePivot::ALevelInstancePivot(const FObjectInitializer& ObjectInitializer)
-	: Super(ObjectInitializer)
-{
-	RootComponent = CreateDefaultSubobject<USceneComponent>(TEXT("RootComponent"));
-	RootComponent->Mobility = EComponentMobility::Static;
-#if WITH_EDITORONLY_DATA
-	bActorLabelEditable = false;
-#endif // WITH_EDITORONLY_DATA
-}
-
-#if WITH_EDITOR
-
-ALevelInstancePivot* ALevelInstancePivot::Create(ILevelInstanceInterface* LevelInstance, ULevelStreaming* LevelStreaming)
-{
-	FActorSpawnParameters SpawnParams;
-	SpawnParams.OverrideLevel = LevelStreaming->GetLoadedLevel();
-	SpawnParams.bCreateActorPackage = false;
-	SpawnParams.ObjectFlags |= RF_Transient;
-	SpawnParams.bNoFail = true;
-		
-	// We place the pivot actor at the LevelInstance Transform so that it makes sense to the user (the pivot being the zero)
-	AActor* LevelInstanceActor = CastChecked<AActor>(LevelInstance);
-	ALevelInstancePivot* PivotActor = LevelInstanceActor->GetWorld()->SpawnActor<ALevelInstancePivot>(LevelInstanceActor->GetActorLocation(), LevelInstanceActor->GetActorRotation(), SpawnParams);
-	
-	AWorldSettings* WorldSettings = LevelStreaming->GetLoadedLevel()->GetWorldSettings();
-
-	
-	// Keep Spawn World Transform in case Level Instance transform changes.
-	PivotActor->SpawnTransform = PivotActor->GetActorTransform();
-	PivotActor->OriginalPivotOffset = WorldSettings->LevelInstancePivotOffset;
-	PivotActor->SetLevelInstanceID(LevelInstance->GetLevelInstanceID());
-	// Set Label last as this will call PostEditChangeProperty and update the offset so other fields need to be setup first.
-	PivotActor->SetActorLabel(TEXT("Pivot"));
-
-	return PivotActor;
-}
-
-void ALevelInstancePivot::PostEditMove(bool bFinished)
-{
-	Super::PostEditMove(bFinished);
-
-	if (bFinished && !FLevelUtils::IsApplyingLevelTransform())
-	{
-		UpdateOffset();
-	}
-}
-
-void ALevelInstancePivot::PostEditChangeProperty(FPropertyChangedEvent& PropertyChangedEvent)
-{
-	Super::PostEditChangeProperty(PropertyChangedEvent);
-
-	UpdateOffset();
-}
-
-void ALevelInstancePivot::PostEditUndo()
-{
-	Super::PostEditUndo();
-
-	UpdateOffset();
-}
-
-void ALevelInstancePivot::SetPivot(ELevelInstancePivotType PivotType, AActor* PivotActor)
-{
-	check(PivotType != ELevelInstancePivotType::Actor || PivotActor != nullptr);
-
-	Modify();
-	if (PivotType == ELevelInstancePivotType::Actor)
-	{
-		SetActorLocation(PivotActor->GetActorLocation());
-	}
-	else if(PivotType == ELevelInstancePivotType::Center || PivotType == ELevelInstancePivotType::CenterMinZ)
-	{
-		ULevelInstanceSubsystem* LevelInstanceSubsystem = UWorld::GetSubsystem<ULevelInstanceSubsystem>(GetWorld());
-		ILevelInstanceInterface* LevelInstance = LevelInstanceSubsystem->GetLevelInstance(GetLevelInstanceID());
-		FBox OutBounds;
-		LevelInstanceSubsystem->GetLevelInstanceBounds(LevelInstance, OutBounds);
-		FVector Location = OutBounds.GetCenter();
-		if (PivotType == ELevelInstancePivotType::CenterMinZ)
-		{
-			Location.Z = OutBounds.Min.Z;
-		}
-		SetActorLocation(Location);
-	}
-	else if (PivotType == ELevelInstancePivotType::WorldOrigin)
-	{
-		SetActorLocation(FVector(0.f, 0.f, 0.f));
-	}
-	else // unsupported
-	{
-		check(0);
-	}
-
-	// Update gizmo
-	if (GEditor)
-	{
-		GEditor->NoteSelectionChange();
-	}
-
-	UpdateOffset();
-}
-
-void ALevelInstancePivot::UpdateOffset()
-{
-	// Offset Change is the relative transform of the pivot actor to its spawn transform (we only care about relative translation as we don't support rotation on pivot)
-	FVector LocalToPivot = GetActorTransform().GetRelativeTransform(SpawnTransform).GetTranslation();
-
-	// We then apply that offset to the original pivot
-	FVector NewPivotOffset = OriginalPivotOffset - LocalToPivot;
-
-	AWorldSettings* WorldSettings = GetLevel()->GetWorldSettings();
-	if (!NewPivotOffset.Equals(WorldSettings->LevelInstancePivotOffset))
-	{
-		WorldSettings->Modify();
-		WorldSettings->LevelInstancePivotOffset = NewPivotOffset;
-	}
-}
-
-#endif
+// Copyright Epic Games, Inc. All Rights Reserved.
+
+#include "LevelInstance/LevelInstanceEditorPivotActor.h"
+
+ALevelInstancePivot::ALevelInstancePivot(const FObjectInitializer& ObjectInitializer)
+	: Super(ObjectInitializer)
+{
+	RootComponent = CreateDefaultSubobject<USceneComponent>(TEXT("RootComponent"));
+	RootComponent->Mobility = EComponentMobility::Static;
+#if WITH_EDITORONLY_DATA
+	bActorLabelEditable = false;
+#endif // WITH_EDITORONLY_DATA
+}
+
+#if WITH_EDITOR
+
+void ALevelInstancePivot::PostEditMove(bool bFinished)
+{
+	Super::PostEditMove(bFinished);
+
+	if (bFinished)
+	{
+		UpdateOffset();
+	}
+}
+
+void ALevelInstancePivot::PostEditChangeProperty(FPropertyChangedEvent& PropertyChangedEvent)
+{
+	Super::PostEditChangeProperty(PropertyChangedEvent);
+
+	UpdateOffset();
+}
+
+void ALevelInstancePivot::PostEditUndo()
+{
+	Super::PostEditUndo();
+
+	UpdateOffset();
+}
+
+#endif
diff --git a/Engine/Source/Runtime/Engine/Public/LevelInstance/LevelInstanceActor.h b/Engine/Source/Runtime/Engine/Public/LevelInstance/LevelInstanceActor.h
--- a/Engine/Source/Runtime/Engine/Public/LevelInstance/LevelInstanceActor.h
+++ b/Engine/Source/Runtime/Engine/Public/LevelInstance/LevelInstanceActor.h
@@ -61,6 +61,7 @@
 	virtual bool SetWorldAsset(TSoftObjectPtr<UWorld> WorldAsset) override;
 	virtual ELevelInstanceRuntimeBehavior GetDesiredRuntimeBehavior() const override { return DesiredRuntimeBehavior; }
 	virtual ELevelInstanceRuntimeBehavior GetDefaultRuntimeBehavior() const override { return ELevelInstanceRuntimeBehavior::Partitioned; }
+	virtual TSubclassOf<AActor> GetEditorPivotClass() const override;
 	// End ILevelInstanceInterface
 			
 	// UObject overrides
diff --git a/Engine/Source/Runtime/Engine/Public/LevelInstance/LevelInstanceEditorPivot.h b/Engine/Source/Runtime/Engine/Public/LevelInstance/LevelInstanceEditorPivot.h
new file mode 100644
--- /dev/null
+++ b/Engine/Source/Runtime/Engine/Public/LevelInstance/LevelInstanceEditorPivot.h
@@ -0,0 +1,22 @@
+// Copyright Epic Games, Inc. All Rights Reserved.
+
+#pragma once
+
+#if WITH_EDITOR
+
+#include "CoreMinimal.h"
+#include "LevelInstance/LevelInstanceTypes.h"
+
+class ILevelInstanceEditorPivotInterface;
+class ILevelInstanceInterface;
+class ULevelStreaming;
+class AActor;
+
+class ENGINE_API FLevelInstanceEditorPivotHelper
+{
+public:
+	static ILevelInstanceEditorPivotInterface* Create(ILevelInstanceInterface* LevelInstance, ULevelStreaming* LevelStreaming);
+	static void SetPivot(ILevelInstanceEditorPivotInterface* Pivot, ELevelInstancePivotType PivotType, AActor* PivotToActor);
+};
+
+#endif
\ No newline at end of file
diff --git a/Engine/Source/Runtime/Engine/Public/LevelInstance/LevelInstanceEditorPivotActor.h b/Engine/Source/Runtime/Engine/Public/LevelInstance/LevelInstanceEditorPivotActor.h
--- a/Engine/Source/Runtime/Engine/Public/LevelInstance/LevelInstanceEditorPivotActor.h
+++ b/Engine/Source/Runtime/Engine/Public/LevelInstance/LevelInstanceEditorPivotActor.h
@@ -1,44 +1,25 @@
-// Copyright Epic Games, Inc. All Rights Reserved.
-
-#pragma once
-
-#include "CoreMinimal.h"
-#include "UObject/ObjectMacros.h"
-#include "GameFramework/Actor.h"
-#include "LevelInstance/LevelInstanceTypes.h"
-#include "LevelInstanceEditorPivotActor.generated.h"
-
-class ILevelInstanceInterface;
-class ULevelStreaming;
-
-/**
- * 
- */
-UCLASS(transient, notplaceable, hidecategories=(Rendering,Replication,Collision,Partition,Input,HLOD,Actor,Cooking))
-class ENGINE_API ALevelInstancePivot : public AActor
-{
-	GENERATED_UCLASS_BODY()
-
-#if WITH_EDITOR
-public:
-	static ALevelInstancePivot* Create(ILevelInstanceInterface* LevelInstance, ULevelStreaming* LevelStreaming);
-	
-	virtual bool CanDeleteSelectedActor(FText& OutReason) const override { return false; }
-	virtual void PostEditMove(bool bFinished) override;
-	virtual void PostEditChangeProperty(FPropertyChangedEvent& PropertyChangedEvent) override;
-	virtual void PostEditUndo() override;
-	virtual bool ShouldExport() override { return false; }
-	void SetPivot(ELevelInstancePivotType PivotType, AActor* PivotActor = nullptr);
-private:
-	void SetLevelInstanceID(const FLevelInstanceID& InLevelInstanceID) { LevelInstanceID = InLevelInstanceID; }
-	const FLevelInstanceID& GetLevelInstanceID() const { return LevelInstanceID; }
-	void UpdateOffset();
-#endif
-
-private:
-#if WITH_EDITORONLY_DATA
-	FLevelInstanceID LevelInstanceID;
-	FTransform SpawnTransform;
-	FVector OriginalPivotOffset;
-#endif
+// Copyright Epic Games, Inc. All Rights Reserved.
+
+#pragma once
+
+#include "CoreMinimal.h"
+#include "UObject/ObjectMacros.h"
+#include "GameFramework/Actor.h"
+#include "LevelInstance/LevelInstanceTypes.h"
+#include "LevelInstance/LevelInstanceEditorPivotInterface.h"
+#include "LevelInstanceEditorPivotActor.generated.h"
+
+UCLASS(transient, notplaceable, hidecategories=(Rendering,Replication,Collision,Partition,Input,HLOD,Actor,Cooking))
+class ENGINE_API ALevelInstancePivot : public AActor, public ILevelInstanceEditorPivotInterface
+{
+	GENERATED_UCLASS_BODY()
+
+#if WITH_EDITOR
+public:
+	virtual bool CanDeleteSelectedActor(FText& OutReason) const override { return false; }
+	virtual void PostEditMove(bool bFinished) override;
+	virtual void PostEditChangeProperty(FPropertyChangedEvent& PropertyChangedEvent) override;
+	virtual void PostEditUndo() override;
+	virtual bool ShouldExport() override { return false; }
+#endif
 };
\ No newline at end of file
diff --git a/Engine/Source/Runtime/Engine/Public/LevelInstance/LevelInstanceEditorPivotInterface.h b/Engine/Source/Runtime/Engine/Public/LevelInstance/LevelInstanceEditorPivotInterface.h
new file mode 100644
--- /dev/null
+++ b/Engine/Source/Runtime/Engine/Public/LevelInstance/LevelInstanceEditorPivotInterface.h
@@ -0,0 +1,38 @@
+// Copyright Epic Games, Inc. All Rights Reserved.
+
+#pragma once
+
+#include "CoreMinimal.h"
+#include "UObject/Interface.h"
+#include "LevelInstance/LevelInstanceTypes.h"
+#include "LevelInstanceEditorPivotInterface.generated.h"
+
+UINTERFACE()
+class ENGINE_API ULevelInstanceEditorPivotInterface : public UInterface
+{
+	GENERATED_UINTERFACE_BODY()
+};
+
+/**
+ * Interface to be implemented by Actor classes to implement support for LevelInstance Editor Pivot
+ */
+class ENGINE_API ILevelInstanceEditorPivotInterface
+{
+	GENERATED_IINTERFACE_BODY()
+#if WITH_EDITOR
+public:
+	struct FInitState
+	{
+		FLevelInstanceID LevelInstanceID;
+		FTransform ActorTransform;
+		FVector LevelOffset;
+	};
+
+	void SetInitState(const FInitState& InInitState) { InitState = InInitState; }
+	const FLevelInstanceID& GetLevelInstanceID() const { return InitState.LevelInstanceID; }
+
+	virtual void UpdateOffset();
+protected:
+	FInitState InitState;
+#endif
+};
\ No newline at end of file
diff --git a/Engine/Source/Runtime/Engine/Public/LevelInstance/LevelInstanceInterface.h b/Engine/Source/Runtime/Engine/Public/LevelInstance/LevelInstanceInterface.h
--- a/Engine/Source/Runtime/Engine/Public/LevelInstance/LevelInstanceInterface.h
+++ b/Engine/Source/Runtime/Engine/Public/LevelInstance/LevelInstanceInterface.h
@@ -46,6 +46,8 @@
 	virtual ELevelInstanceRuntimeBehavior GetDesiredRuntimeBehavior() const = 0;
 
 	virtual ELevelInstanceRuntimeBehavior GetDefaultRuntimeBehavior() const = 0;
+
+	virtual TSubclassOf<AActor> GetEditorPivotClass() const { return nullptr; }
 #endif
 	// Pure Interface End
 		
