World Partition: actors pinning don't honor actor descriptors filtering so it's possible to pin HLOD actors.

#rnx

diff --git a/Engine/Source/Runtime/Engine/Private/WorldPartition/WorldPartition.cpp b/Engine/Source/Runtime/Engine/Private/WorldPartition/WorldPartition.cpp
--- a/Engine/Source/Runtime/Engine/Private/WorldPartition/WorldPartition.cpp
+++ b/Engine/Source/Runtime/Engine/Private/WorldPartition/WorldPartition.cpp
@@ -253,7 +253,7 @@ protected:
 	//~ Begin IWorldPartitionActorLoaderInterface::ILoaderAdapterList interface
 	virtual bool ShouldActorBeLoaded(const FWorldPartitionHandle& ActorHandle) const override
 	{
-		return (ActorHandle != ActorToRemove) && FLoaderAdapterList::ShouldActorBeLoaded(ActorHandle);
+		return (ActorHandle != ActorToRemove) && PassDataLayersFilter(ActorHandle);
 	}
 	//~ End IWorldPartitionActorLoaderInterface::ILoaderAdapterList interface
 
diff --git a/Engine/Source/Runtime/Engine/Private/WorldPartition/WorldPartitionActorLoaderInterface.cpp b/Engine/Source/Runtime/Engine/Private/WorldPartition/WorldPartitionActorLoaderInterface.cpp
--- a/Engine/Source/Runtime/Engine/Private/WorldPartition/WorldPartitionActorLoaderInterface.cpp
+++ b/Engine/Source/Runtime/Engine/Private/WorldPartition/WorldPartitionActorLoaderInterface.cpp
@@ -100,6 +100,50 @@
 	FDataLayersEditorBroadcast::Get().OnActorDataLayersEditorLoadingStateChanged().RemoveAll(this);
 }
 
+bool IWorldPartitionActorLoaderInterface::ILoaderAdapter::PassActorDescFilter(const FWorldPartitionHandle& Actor) const
+{
+	return Actor->ShouldBeLoadedByEditor();
+}
+
+bool IWorldPartitionActorLoaderInterface::ILoaderAdapter::PassDataLayersFilter(const FWorldPartitionHandle& Actor) const
+{
+	if (UDataLayerSubsystem* DataLayerSubsystem = UWorld::GetSubsystem<UDataLayerSubsystem>(World))
+	{
+		FWorldPartitionActorViewProxy ActorDescProxy(*Actor);
+
+		if (IsRunningCookCommandlet())
+		{
+			// When running cook commandlet, dont allow loading of actors with runtime loaded data layers
+			for (const FName& DataLayerInstanceName : ActorDescProxy.GetDataLayers())
+			{
+				const UDataLayerInstance* DataLayerInstance = DataLayerSubsystem->GetDataLayerInstance(DataLayerInstanceName);
+				if (DataLayerInstance && DataLayerInstance->IsRuntime())
+				{
+					return false;
+				}
+			}
+		}
+		else
+		{
+			uint32 NumValidLayers = 0;
+			for (const FName& DataLayerInstanceName : ActorDescProxy.GetDataLayers())
+			{
+				if (const UDataLayerInstance* DataLayerInstance = DataLayerSubsystem->GetDataLayerInstance(DataLayerInstanceName))
+				{
+					if (DataLayerInstance->IsEffectiveLoadedInEditor())
+					{
+						return true;
+					}
+					NumValidLayers++;
+				}
+			}
+			return !NumValidLayers;
+		}
+	}
+
+	return true;
+}
+
 bool IWorldPartitionActorLoaderInterface::ILoaderAdapter::RefreshLoadedState()
 {
 	if (bLoaded)
@@ -213,47 +257,7 @@
 bool IWorldPartitionActorLoaderInterface::ILoaderAdapter::ShouldActorBeLoaded(const FWorldPartitionHandle& Actor) const
 {
 	check(Actor.IsValid());
-
-	if (!Actor->ShouldBeLoadedByEditor())
-	{
-		return false;
-	}
-
-	if (UDataLayerSubsystem* DataLayerSubsystem = UWorld::GetSubsystem<UDataLayerSubsystem>(World))
-	{
-		FWorldPartitionActorViewProxy ActorDescProxy(*Actor);
-
-		if (IsRunningCookCommandlet())
-		{
-			// When running cook commandlet, dont allow loading of actors with runtime loaded data layers
-			for (const FName& DataLayerInstanceName : ActorDescProxy.GetDataLayers())
-			{
-				const UDataLayerInstance* DataLayerInstance = DataLayerSubsystem->GetDataLayerInstance(DataLayerInstanceName);
-				if (DataLayerInstance && DataLayerInstance->IsRuntime())
-				{
-					return false;
-				}
-			}
-		}
-		else
-		{
-			uint32 NumValidLayers = 0;
-			for (const FName& DataLayerInstanceName : ActorDescProxy.GetDataLayers())
-			{
-				if (const UDataLayerInstance* DataLayerInstance = DataLayerSubsystem->GetDataLayerInstance(DataLayerInstanceName))
-				{
-					if (DataLayerInstance->IsEffectiveLoadedInEditor())
-					{
-						return true;
-					}
-					NumValidLayers++;
-				}
-			}
-			return !NumValidLayers;
-		}
-	}
-
-	return true;
+	return PassActorDescFilter(Actor) && PassDataLayersFilter(Actor);
 };
 
 void IWorldPartitionActorLoaderInterface::ILoaderAdapter::PostLoadedStateChanged(bool bUnloadedActors)
diff --git a/Engine/Source/Runtime/Engine/Public/WorldPartition/WorldPartitionActorLoaderInterface.h b/Engine/Source/Runtime/Engine/Public/WorldPartition/WorldPartitionActorLoaderInterface.h
--- a/Engine/Source/Runtime/Engine/Public/WorldPartition/WorldPartitionActorLoaderInterface.h
+++ b/Engine/Source/Runtime/Engine/Public/WorldPartition/WorldPartitionActorLoaderInterface.h
@@ -49,6 +49,10 @@
 		void RegisterDelegates();
 		void UnregisterDelegates();
 
+		// Actors filtering
+		bool PassActorDescFilter(const FWorldPartitionHandle& Actor) const;
+		bool PassDataLayersFilter(const FWorldPartitionHandle& Actor) const;
+
 		/**
 		 * Refresh the loaded state of this loader adatper.
 		 *
