World Partition: made all slow task dialogs delayed and added a top scope when pinning an actor selection.

#rnx

diff --git a/Engine/Source/Editor/SceneOutliner/Private/ActorBrowsingMode.cpp b/Engine/Source/Editor/SceneOutliner/Private/ActorBrowsingMode.cpp
--- a/Engine/Source/Editor/SceneOutliner/Private/ActorBrowsingMode.cpp
+++ b/Engine/Source/Editor/SceneOutliner/Private/ActorBrowsingMode.cpp
@@ -34,6 +34,7 @@
 #include "LevelInstance/LevelInstanceInterface.h"
 #include "LevelInstance/LevelInstanceSubsystem.h"
 #include "LevelInstance/LevelInstanceEditorInstanceActor.h"
+#include "Misc/ScopedSlowTask.h"
 #include "EditorLevelUtils.h"
 
 DEFINE_LOG_CATEGORY_STATIC(LogActorBrowser, Log, All);
@@ -1873,9 +1874,13 @@
 		GEditor->GetSelectedActors()->BeginBatchSelectOperation();
 		GEditor->SelectNone(/*bNoteSelectionChange=*/false, /*bDeselectBSPSurfs=*/true);
 
-		Selection.ForEachItem([this](const FSceneOutlinerTreeItemPtr& TreeItem)
+		FScopedSlowTask SlowTask(Selection.Num(), LOCTEXT("LoadingActors", "Loading actors..."));
+		SlowTask.MakeDialogDelayed(1.0f);
+
+		Selection.ForEachItem([this, &SlowTask](const FSceneOutlinerTreeItemPtr& TreeItem)
 		{
 			PinItem(TreeItem);
+			SlowTask.EnterProgressFrame(1);
 			return true;
 		});
 
@@ -1895,9 +1900,13 @@
 	{
 		GEditor->GetSelectedActors()->BeginBatchSelectOperation();
 
-		Selection.ForEachItem([this](const FSceneOutlinerTreeItemPtr& TreeItem)
+		FScopedSlowTask SlowTask(Selection.Num(), LOCTEXT("UnloadingActors", "Unloading actors..."));
+		SlowTask.MakeDialogDelayed(1.f);
+
+		Selection.ForEachItem([this, &SlowTask](const FSceneOutlinerTreeItemPtr& TreeItem)
 		{
 			UnpinItem(TreeItem);
+			SlowTask.EnterProgressFrame(1);
 			return true;
 		});
 
diff --git a/Engine/Source/Runtime/Engine/Private/WorldPartition/WorldPartitionActorLoaderInterface.cpp b/Engine/Source/Runtime/Engine/Private/WorldPartition/WorldPartitionActorLoaderInterface.cpp
--- a/Engine/Source/Runtime/Engine/Private/WorldPartition/WorldPartitionActorLoaderInterface.cpp
+++ b/Engine/Source/Runtime/Engine/Private/WorldPartition/WorldPartitionActorLoaderInterface.cpp
@@ -63,22 +63,21 @@ bool IWorldPartitionActorLoaderInterface::ILoaderAdapter::Unload()
 			{
 				return false;
 			}
+			
+			UnregisterDelegates();
 
 			FScopedSlowTask SlowTask(1, LOCTEXT("UpdatingLoading", "Updating loading..."));
-			SlowTask.MakeDialog();
-
-			UnregisterDelegates();
+			SlowTask.MakeDialogDelayed(1.0f);
 
 			{
 				FWorldPartitionLoadingContext::FDeferred LoadingContext;
 				ActorReferences.Empty();
+				SlowTask.EnterProgressFrame(1);
 			}
 
 			bLoaded = false;
 
 			PostLoadedStateChanged(true);
-
-			SlowTask.EnterProgressFrame(1);
 		}
 	}
 
@@ -173,7 +172,7 @@ bool IWorldPartitionActorLoaderInterface::ILoaderAdapter::RefreshLoadedState()
 			}
 
 			FScopedSlowTask SlowTask(ActorsToLoad.Num() + ActorsToUnload.Num(), LOCTEXT("UpdatingLoading", "Updating loading..."));
-			SlowTask.MakeDialog();
+			SlowTask.MakeDialogDelayed(1.0f);
 
 			{
 				FWorldPartitionLoadingContext::FDeferred LoadingContext;
@@ -184,12 +183,12 @@ bool IWorldPartitionActorLoaderInterface::ILoaderAdapter::RefreshLoadedState()
 					FWorldPartitionHandle ActorHandle(WorldPartition, ActorToLoad->GetGuid());
 					AddReferenceToActor(ActorHandle);
 				}
-			}
 
-			for (FWorldPartitionHandle& ActorToUnload : ActorsToUnload)
-			{
-				SlowTask.EnterProgressFrame(1);
-				RemoveReferenceToActor(ActorToUnload);
+				for (FWorldPartitionHandle& ActorToUnload : ActorsToUnload)
+				{
+					SlowTask.EnterProgressFrame(1);
+					RemoveReferenceToActor(ActorToUnload);
+				}
 			}
 
 			PostLoadedStateChanged(ActorsToUnload.Num() > 0);
