World Partition: greatly optimized pinning multiple actors by batching loading + tweaked slow task dialogs when loading actors in the editor.

#rnx

diff --git a/Engine/Source/Editor/SceneOutliner/Private/ActorBrowsingMode.cpp b/Engine/Source/Editor/SceneOutliner/Private/ActorBrowsingMode.cpp
--- a/Engine/Source/Editor/SceneOutliner/Private/ActorBrowsingMode.cpp
+++ b/Engine/Source/Editor/SceneOutliner/Private/ActorBrowsingMode.cpp
@@ -34,7 +34,6 @@
 #include "LevelInstance/LevelInstanceInterface.h"
 #include "LevelInstance/LevelInstanceSubsystem.h"
 #include "LevelInstance/LevelInstanceEditorInstanceActor.h"
-#include "Misc/ScopedSlowTask.h"
 #include "EditorLevelUtils.h"
 
 DEFINE_LOG_CATEGORY_STATIC(LogActorBrowser, Log, All);
@@ -1874,13 +1873,9 @@ void FActorBrowsingMode::PinSelectedItems()
 		GEditor->GetSelectedActors()->BeginBatchSelectOperation();
 		GEditor->SelectNone(/*bNoteSelectionChange=*/false, /*bDeselectBSPSurfs=*/true);
 
-		FScopedSlowTask SlowTask(Selection.Num(), LOCTEXT("LoadingActors", "Loading actors..."));
-		SlowTask.MakeDialogDelayed(1.0f);
-
-		Selection.ForEachItem([this, &SlowTask](const FSceneOutlinerTreeItemPtr& TreeItem)
+		Selection.ForEachItem([this](const FSceneOutlinerTreeItemPtr& TreeItem)
 		{
 			PinItem(TreeItem);
-			SlowTask.EnterProgressFrame(1);
 			return true;
 		});
 
@@ -1900,13 +1895,9 @@ void FActorBrowsingMode::UnpinSelectedItems()
 	{
 		GEditor->GetSelectedActors()->BeginBatchSelectOperation();
 
-		FScopedSlowTask SlowTask(Selection.Num(), LOCTEXT("UnloadingActors", "Unloading actors..."));
-		SlowTask.MakeDialogDelayed(1.f);
-
-		Selection.ForEachItem([this, &SlowTask](const FSceneOutlinerTreeItemPtr& TreeItem)
+		Selection.ForEachItem([this](const FSceneOutlinerTreeItemPtr& TreeItem)
 		{
 			UnpinItem(TreeItem);
-			SlowTask.EnterProgressFrame(1);
 			return true;
 		});
 
diff --git a/Engine/Source/Runtime/Engine/Private/Level.cpp b/Engine/Source/Runtime/Engine/Private/Level.cpp
--- a/Engine/Source/Runtime/Engine/Private/Level.cpp
+++ b/Engine/Source/Runtime/Engine/Private/Level.cpp
@@ -854,6 +854,9 @@
 		QueueActor(Actor);
 	}
 
+	FScopedSlowTask SlowTask(ActorsQueue.Num() * 3, LOCTEXT("RegisteringActors", "Registering actors..."));
+	SlowTask.MakeDialogDelayed(1.0f);
+
 	// Register all components
 	for (AActor* Actor : ActorsQueue)
 	{
@@ -869,6 +872,8 @@
 		{
 			Actor->RegisterAllComponents();
 		}
+
+		SlowTask.EnterProgressFrame(1);
 	}
 
 	// Rerun construction scripts
@@ -878,6 +883,8 @@
 		{
 			Actor->RerunConstructionScripts();
 		}
+
+		SlowTask.EnterProgressFrame(1);
 	}
 
 	// Finalize actors
@@ -895,6 +902,8 @@
 		}
 
 		OnLoadedActorAddedToLevelEvent.Broadcast(*Actor);
+
+		SlowTask.EnterProgressFrame(1);
 	}
 }
 
@@ -937,6 +946,9 @@
 		QueueActor(Actor);
 	}
 
+	FScopedSlowTask SlowTask(ActorsQueue.Num(), LOCTEXT("UnregisteringActors", "Unregistering actors..."));
+	SlowTask.MakeDialogDelayed(1.0f);
+
 	for (AActor* Actor : ActorsQueue)
 	{
 		Actor->UnregisterAllComponents();
@@ -951,6 +963,8 @@
 		}
 
 		OnLoadedActorRemovedFromLevelEvent.Broadcast(*Actor);
+
+		SlowTask.EnterProgressFrame(1);
 	}
 }
 #endif
diff --git a/Engine/Source/Runtime/Engine/Private/WorldPartition/WorldPartition.cpp b/Engine/Source/Runtime/Engine/Private/WorldPartition/WorldPartition.cpp
--- a/Engine/Source/Runtime/Engine/Private/WorldPartition/WorldPartition.cpp
+++ b/Engine/Source/Runtime/Engine/Private/WorldPartition/WorldPartition.cpp
@@ -216,6 +216,7 @@ class FLoaderAdapterPinnedActors: public FLoaderAdapterList
 public:
 	FLoaderAdapterPinnedActors(UWorld* InWorld)
 		: FLoaderAdapterList(InWorld)
+		, bStateChanged(false)
 	{
 		Load();
 	}
@@ -225,7 +226,7 @@ public:
 		if (!ContainsActor(ActorHandle))
 		{
 			Actors.Add(ActorHandle);
-			RefreshLoadedState();
+			bStateChanged = true;
 		}
 	}
 
@@ -249,6 +250,15 @@ public:
 		return Actors.Contains(ActorHandle);
 	}
 
+	void Tick(float DeltaSeconds)
+	{
+		if (bStateChanged)
+		{
+			RefreshLoadedState();
+			bStateChanged = false;
+		}
+	}
+
 protected:
 	//~ Begin IWorldPartitionActorLoaderInterface::ILoaderAdapterList interface
 	virtual bool ShouldActorBeLoaded(const FWorldPartitionHandle& ActorHandle) const override
@@ -258,7 +268,8 @@ protected:
 	//~ End IWorldPartitionActorLoaderInterface::ILoaderAdapterList interface
 
 private:
-	FWorldPartitionHandle ActorToRemove;
+	bool bStateChanged;
+	FWorldPartitionHandle ActorToRemove;	
 };
 #endif
 
@@ -1084,6 +1095,11 @@ void UWorldPartition::Tick(float DeltaSeconds)
 		EditorHash->Tick(DeltaSeconds);
 	}
 
+	if (PinnedActors)
+	{
+		PinnedActors->Tick(DeltaSeconds);
+	}
+
 	if (bForceGarbageCollection)
 	{
 		GEngine->ForceGarbageCollection(bForceGarbageCollectionPurge);
diff --git a/Engine/Source/Runtime/Engine/Private/WorldPartition/WorldPartitionActorLoaderInterface.cpp b/Engine/Source/Runtime/Engine/Private/WorldPartition/WorldPartitionActorLoaderInterface.cpp
--- a/Engine/Source/Runtime/Engine/Private/WorldPartition/WorldPartitionActorLoaderInterface.cpp
+++ b/Engine/Source/Runtime/Engine/Private/WorldPartition/WorldPartitionActorLoaderInterface.cpp
@@ -171,23 +171,32 @@ bool IWorldPartitionActorLoaderInterface::ILoaderAdapter::RefreshLoadedState()
 				return false;
 			}
 
-			FScopedSlowTask SlowTask(ActorsToLoad.Num() + ActorsToUnload.Num(), LOCTEXT("UpdatingLoading", "Updating loading..."));
-			SlowTask.MakeDialogDelayed(1.0f);
-
 			{
 				FWorldPartitionLoadingContext::FDeferred LoadingContext;
 
-				for (FWorldPartitionHandle& ActorToLoad : ActorsToLoad)
+				if (ActorsToLoad.Num())
 				{
-					SlowTask.EnterProgressFrame(1);
-					FWorldPartitionHandle ActorHandle(WorldPartition, ActorToLoad->GetGuid());
-					AddReferenceToActor(ActorHandle);
+					FScopedSlowTask SlowTask(ActorsToLoad.Num(), LOCTEXT("ActorsLoading", "Loading actors..."));
+					SlowTask.MakeDialogDelayed(1.0f);
+
+					for (FWorldPartitionHandle& ActorToLoad : ActorsToLoad)
+					{
+						FWorldPartitionHandle ActorHandle(WorldPartition, ActorToLoad->GetGuid());
+						AddReferenceToActor(ActorHandle);
+						SlowTask.EnterProgressFrame(1);
+					}
 				}
 
-				for (FWorldPartitionHandle& ActorToUnload : ActorsToUnload)
+				if (ActorsToUnload.Num())
 				{
-					SlowTask.EnterProgressFrame(1);
-					RemoveReferenceToActor(ActorToUnload);
+					FScopedSlowTask SlowTask(ActorsToUnload.Num(), LOCTEXT("ActorsUnoading", "Unloading actors..."));
+					SlowTask.MakeDialogDelayed(1.0f);
+
+					for (FWorldPartitionHandle& ActorToUnload : ActorsToUnload)
+					{
+						RemoveReferenceToActor(ActorToUnload);
+						SlowTask.EnterProgressFrame(1);
+					}
 				}
 			}
 
