World Partition: added getter/setter for default HLOD layer and made member public.

#rnx

diff --git a/Engine/Source/Editor/UnrealEd/Private/FileHelpers.cpp b/Engine/Source/Editor/UnrealEd/Private/FileHelpers.cpp
--- a/Engine/Source/Editor/UnrealEd/Private/FileHelpers.cpp
+++ b/Engine/Source/Editor/UnrealEd/Private/FileHelpers.cpp
@@ -759,11 +759,11 @@
 
 					if (bIsTempPackage)
 					{
-						if (UHLODLayer* CurrentHLODLayer = RenamedWorldPartition->DefaultHLODLayer)
+						if (UHLODLayer* CurrentHLODLayer = RenamedWorldPartition->GetDefaultHLODLayer())
 						{
 							CurrentHLODLayer = UHLODLayer::DuplicateHLODLayersSetup(CurrentHLODLayer, NewPackageName, NewWorldAssetName);
 							
-							RenamedWorldPartition->DefaultHLODLayer = CurrentHLODLayer;
+							RenamedWorldPartition->SetDefaultHLODLayer(CurrentHLODLayer);
 
 							while (CurrentHLODLayer)
 							{
diff --git a/Engine/Source/Runtime/Engine/Private/WorldPartition/HLOD/HLODLayer.cpp b/Engine/Source/Runtime/Engine/Private/WorldPartition/HLOD/HLODLayer.cpp
--- a/Engine/Source/Runtime/Engine/Private/WorldPartition/HLOD/HLODLayer.cpp
+++ b/Engine/Source/Runtime/Engine/Private/WorldPartition/HLOD/HLODLayer.cpp
@@ -49,7 +49,7 @@
 		// Fallback to the world partition default HLOD layer
 		if (UWorldPartition* WorldPartition = InActor->GetWorld()->GetWorldPartition())
 		{
-			return WorldPartition->DefaultHLODLayer;
+			return WorldPartition->GetDefaultHLODLayer();
 		}
 	}
 
@@ -71,7 +71,7 @@
 	if (bIsHLOD0)
 	{
 		// Fallback to the world partition default HLOD layer
-		return InWorldPartition->DefaultHLODLayer;
+		return InWorldPartition->GetDefaultHLODLayer();
 	}
 
 	return nullptr;
diff --git a/Engine/Source/Runtime/Engine/Public/WorldPartition/WorldPartition.h b/Engine/Source/Runtime/Engine/Public/WorldPartition/WorldPartition.h
--- a/Engine/Source/Runtime/Engine/Public/WorldPartition/WorldPartition.h
+++ b/Engine/Source/Runtime/Engine/Public/WorldPartition/WorldPartition.h
@@ -31,6 +31,7 @@
 class FLoaderAdapterAlwaysLoadedActors;
 class FLoaderAdapterPinnedActors;
 class FHLODActorDesc;
+class UHLODLayer;
 class UCanvas;
 class ULevel;
 
@@ -152,6 +153,9 @@
 
 	FBox GetWorldBounds() const;
 	FBox GetEditorWorldBounds() const;
+	
+	UHLODLayer* GetDefaultHLODLayer() const { return DefaultHLODLayer; }
+	void SetDefaultHLODLayer(UHLODLayer* InDefaultHLODLayer) { DefaultHLODLayer = InDefaultHLODLayer; }
 	void GenerateHLOD(ISourceControlHelper* SourceControlHelper, bool bCreateActorsOnly);
 
 	// Debugging Methods
@@ -237,15 +241,14 @@
 	UPROPERTY()
 	TObjectPtr<UWorldPartitionRuntimeHash> RuntimeHash;
 
+private:
 #if WITH_EDITOR
-private:
 	bool bForceGarbageCollection;
 	bool bForceGarbageCollectionPurge;
 	bool bIsPIE;
 #endif
 
 #if WITH_EDITORONLY_DATA
-public:
 	// Default HLOD layer
 	UPROPERTY(EditAnywhere, Category=WorldPartition, meta = (DisplayName = "Default HLOD Layer", EditCondition="bEnableStreaming", EditConditionHides, HideEditConditionToggle))
 	TObjectPtr<class UHLODLayer> DefaultHLODLayer;
@@ -253,7 +256,6 @@
 	TArray<FWorldPartitionReference> LoadedSubobjects;
 #endif
 
-private:
 	EWorldPartitionInitState InitState;
 	TOptional<FTransform> InstanceTransform;
 
