World Partition: remove unused case of refpin scope when updating actor descriptors.

#rnx

diff --git a/Engine/Source/Runtime/Engine/Private/WorldPartition/ActorDescContainer.cpp b/Engine/Source/Runtime/Engine/Private/WorldPartition/ActorDescContainer.cpp
--- a/Engine/Source/Runtime/Engine/Private/WorldPartition/ActorDescContainer.cpp
+++ b/Engine/Source/Runtime/Engine/Private/WorldPartition/ActorDescContainer.cpp
@@ -201,14 +201,10 @@
 				check(IsValidChecked(Actor));
 				if (TUniquePtr<FWorldPartitionActorDesc>* ExistingActorDesc = GetActorDescriptor(Actor->GetActorGuid()))
 				{
-					// Pin the actor handle on the actor to prevent unloading it when unhashing
-					FWorldPartitionHandle ExistingActorHandle(ExistingActorDesc);
-					FWorldPartitionHandlePinRefScope ExistingActorHandlePin(ExistingActorHandle);
-
-					TUniquePtr<FWorldPartitionActorDesc> NewActorDesc(Actor->CreateActorDesc());
-
 					OnActorDescUpdating(ExistingActorDesc->Get());
 
+					TUniquePtr<FWorldPartitionActorDesc> NewActorDesc(Actor->CreateActorDesc());
+
 					// Transfer any internal values not coming from the actor
 					NewActorDesc->TransferFrom(ExistingActorDesc->Get());
 
