World Partition: fix always loaded actors loading/unloading when changing the enable streaming setting.

#rnx

diff --git a/Engine/Source/Runtime/Engine/Private/WorldPartition/WorldPartition.cpp b/Engine/Source/Runtime/Engine/Private/WorldPartition/WorldPartition.cpp
--- a/Engine/Source/Runtime/Engine/Private/WorldPartition/WorldPartition.cpp
+++ b/Engine/Source/Runtime/Engine/Private/WorldPartition/WorldPartition.cpp
@@ -990,28 +990,20 @@ void UWorldPartition::SetCanBeUsedByLevelInstance(bool bInCanBeUsedByLevelInstan
 
 void UWorldPartition::OnEnableStreamingChanged()
 {
-	// Pin the actor handles on the actor to prevent unloading it when unhashing
-	TArray<FWorldPartitionHandlePinRefScope> PinRefScopes;
-	PinRefScopes.Reserve(ActorDescList.Num());
-
 	for (UActorDescContainer::TIterator<> ActorDescIterator(this); ActorDescIterator; ++ActorDescIterator)
 	{
-		FWorldPartitionHandle ExistingActorHandle(this, ActorDescIterator->GetGuid());
-		PinRefScopes.Emplace(ExistingActorHandle);
-
 		UnhashActorDesc(*ActorDescIterator);
-	}
-
-	for (UActorDescContainer::TIterator<> ActorDescIterator(this); ActorDescIterator; ++ActorDescIterator)
-	{
 		ActorDescIterator->bIsForcedNonSpatiallyLoaded = !IsStreamingEnabled();
 		HashActorDesc(*ActorDescIterator);
 	}
 
-	if (!IsStreamingEnabled())
-	{
-		AlwaysLoadedActors->Load();
-	}
+	FLoaderAdapterAlwaysLoadedActors* OldAlwaysLoadedActors = AlwaysLoadedActors;
+
+	AlwaysLoadedActors = new FLoaderAdapterAlwaysLoadedActors(GetTypedOuter<UWorld>());
+	AlwaysLoadedActors->Load();
+
+	OldAlwaysLoadedActors->Unload();
+	delete OldAlwaysLoadedActors;
 
 	if (WorldPartitionEditor)
 	{
