Refactor async physics network sync so it doesn't need a callback.
Made it easy to trigger RPCs on same frame for both client and server.

diff --git a/Engine/Source/Runtime/Engine/Classes/GameFramework/Actor.h b/Engine/Source/Runtime/Engine/Classes/GameFramework/Actor.h
--- a/Engine/Source/Runtime/Engine/Classes/GameFramework/Actor.h
+++ b/Engine/Source/Runtime/Engine/Classes/GameFramework/Actor.h
@@ -467,10 +467,12 @@
 	/** Set if an Actor tries to be destroyed while it is beginning play so that once BeginPlay ends we can issue the destroy call. */
 	uint8 bActorWantsDestroyDuringBeginPlay : 1;
 	
+protected:
 	/** Whether to use use the async physics tick with this actor. */
 	UPROPERTY(EditAnywhere, Category=Physics)
 	uint8 bAsyncPhysicsTickEnabled : 1;
 
+private:
 	/** Enum defining if BeginPlay has started or finished */
 	enum class EActorBeginPlayState : uint8
 	{
diff --git a/Engine/Source/Runtime/Engine/Classes/GameFramework/PlayerController.h b/Engine/Source/Runtime/Engine/Classes/GameFramework/PlayerController.h
--- a/Engine/Source/Runtime/Engine/Classes/GameFramework/PlayerController.h
+++ b/Engine/Source/Runtime/Engine/Classes/GameFramework/PlayerController.h
@@ -209,7 +209,17 @@
 	virtual void ApplyInputMode(FReply& SlateOperations, class UGameViewportClient& GameViewportClient) const override;
 };
 
+USTRUCT()
+struct FAsyncPhysicsTimestamp
+{
+	GENERATED_BODY()
+	UPROPERTY()
+	int32 ServerFrame = INDEX_NONE;
 
+	UPROPERTY()
+	int32 LocalFrame = INDEX_NONE;
+};
+
 /**
  * PlayerControllers are used by human players to control Pawns.
  *
@@ -354,6 +364,21 @@
 	UFUNCTION(BlueprintPure, Category = PlayerController)
 	const UAsyncPhysicsData* GetAsyncPhysicsDataToConsume() const;
 
+	/** Enqueues a command to run at the time specified by AsyncPhysicsTimestamp. Note that if the time specified was missed the command is triggered as soon as possible as part of the async tick.
+		These commands are all run on the game thread. If you want to run on the physics thread see FPhysicsSolverBase::RegisterSimOneShotCallback
+
+		If OwningObject is not null this command will only fire as long as the object is still alive. This allows a lambda to still use the owning object's data for read/write (assuming data is designed for async physics tick)
+	*/
+	void ExecuteAsyncPhysicsCommand(const FAsyncPhysicsTimestamp& AsyncPhysicsTimestamp, UObject* OwningObject, const TFunction<void()>& Command);
+
+	/** Generates a timestamp for the upcoming physics step (plus any pending time). Useful for synchronizing client and server events on a specific physics step */
+	FAsyncPhysicsTimestamp GetAsyncPhysicsTimestamp(float DeltaSeconds = 0.f);
+
+	/** Returns the current estimated offset between the local async physics step and the server. This is useful for dealing with low level synchronization.
+		In general it's recommended to use GetAsyncPhysicsTimestamp which accounts for the offset automatically*/
+	int32 GetLocalToServerAsyncPhysicsTickOffset() const { return LocalToServerAsyncPhysicsTickOffset; }
+	
+
 private:
 
 	UPROPERTY(ReplicatedUsing=OnRep_AsyncPhysicsDataComponent)
@@ -2158,85 +2183,38 @@
 	/** If true, prevent any haptic effects from playing */
 	bool bDisableHaptics : 1;
 
-public: 
+	/** The estimated offset between the local async physics tick frame number and the server's
+	*	This is used to synchronize events that happen in the async physics tick */
+	int32 LocalToServerAsyncPhysicsTickOffset;
 
-	/**
-	 * Frame number exchange. This doesn't inherently do anything but is used by the network prediction physics system.
-	 * This may be moved out at some point.
-	 * 
-	 * This is meant ot provide a mechanism for client side prediction to correlate client input and server frame numbers.
-	 * Frame is a loose concept here. It doesn't necessary mean GFrameNumber. Its just an arbitrary increasing sequence of numbers that is used
-	 * to label disrete units of client->Server input. For example the main thread may tick at a high variable rate but input is generated at a fixed
-	 * step interval.
-	 */
+	/** The latest server step we've received an offset correction for. This allows us to ignore out of order corrections that arrive late */
+	int32 ClientLatestCorrectedOffsetServerStep = INDEX_NONE;
 
-	struct FInputCmdBuffer
-	{
-		int32 HeadFrame() const { return LastWritten; }
-		int32 TailFrame() const { return FMath::Max(0, LastWritten - Buffer.Num() + 1); }
-		TArray<uint8>& Write(int32 Frame) { LastWritten = FMath::Max(Frame, LastWritten); return Buffer[Frame % Buffer.Num()]; }
-		const TArray<uint8>& Get(int32 Frame) const { return Buffer[Frame % Buffer.Num()]; }
+	/** The latest physics step we've sent to the server. Due to async we need to avoid duplicate sends */
+	int32 ClientLatestAsyncPhysicsStepSent = INDEX_NONE;
 
-	private:
-		int32 LastWritten = INDEX_NONE;
-		TStaticArray<TArray<uint8>, 16> Buffer;
-	};
+	/** The latest server step we've received a time dilation for. Needed for out of order updates */
+	int32 ClientLatestTimeDilationServerStep = INDEX_NONE;
 
-	FInputCmdBuffer& GetInputBuffer() { return InputBuffer; }
+	/** The server tells the client to speed up or slow down in order to keep its buffer full */
+	float ServerAsyncPhysicsTimeDilationToSend = 1.f;
 
-	// -------------------------------------------------------------------------
-	// Client
-	// -------------------------------------------------------------------------
+	/** The server records the latest timestamp it has to correct. This is used to update client (which may not happen on every physics step) */
+	FAsyncPhysicsTimestamp ServerLatestTimestampToCorrect;
 
-	struct FClientFrameInfo
-	{
-		int32 LastRecvInputFrame = INDEX_NONE;	// The latest inputcmd that the server acknowledged receiving, but not yet processed (this is our frame number that we gave them)
-		int32 LastProcessedInputFrame = INDEX_NONE; // The latest InputCmd that the server actually processed (this is our frame number that we gave them)
-		int32 LastRecvServerFrame = INDEX_NONE; // the latest ServerFrame number that the processing of LastRecvInputFrame happened on (Server's local frame number)
+	/** The set of timestamps the client has sent to the server. Sorted by local frame number */
+	TArray<FAsyncPhysicsTimestamp> ServerPendingTimestamps;
 
-		int8 QuantizedTimeDilation = 1; // Server sent this to this client, telling them to dilate local time either catch up or slow down
-		float TargetNumBufferedCmds = 0.f;
+	void UpdateServerAsyncPhysicsTickOffset();
 
-		int32 GetLocalFrameOffset() const { return LastProcessedInputFrame - LastRecvServerFrame; }
-	};	
+	UFUNCTION(Server, Unreliable)
+	void ServerSendLatestAsyncPhysicsTimestamp(FAsyncPhysicsTimestamp Timestamp);
 
-	// Client pushes input data locally. RPC is sent here but also includes redundant data
-	void PushClientInput(int32 ClientInputFrame, TArray<uint8>& Data);
+	UFUNCTION(Client, Unreliable)
+	void ClientCorrectionAsyncPhysicsTimestamp(FAsyncPhysicsTimestamp Timestamp);
 
-	// Client says "Here is input frame number X" (and then calls other RPCs to deliver InputCmd payload)
-	UFUNCTION(Server, unreliable)
-	void ServerRecvClientInputFrame(int32 RecvClientInputFrame, const TArray<uint8>& Data);
+	UFUNCTION(Client, Unreliable)
+	void ClientAckTimeDilation(float TimeDilation, int32 ServerStep);
 
-	const FClientFrameInfo& GetClientFrameInfo() const { return ClientFrameInfo; }
-	FClientFrameInfo& GetClientFrameInfo() { return ClientFrameInfo; }
-
-	// -------------------------------------------------------------------------
-	// Server
-	// -------------------------------------------------------------------------
-
-	struct FServerFrameInfo
-	{
-		int32 LastProcessedInputFrame = INDEX_NONE;	// The last client frame number we processed. "processed" is arbitrary and we are informed about when commands are processed via SetServerProessedInputFrame
-		int32 LastLocalFrame = INDEX_NONE; // The local frame number that we processed the latest client input frame on. Again, processed is arbitrary and set via SetServerProessedInputFrame
-		int32 LastSentLocalFrame = INDEX_NONE;	// Tracks the latest LastLocalFrame that we sent to the client. Just to prevent redundantly sending info via RPC
-
-		float TargetTimeDilation = 1.f;
-		int8 QuantizedTimeDilation = 1; // Server sets this to tell client to slowdown or speed up
-		float TargetNumBufferedCmds = 1.f; // How many buffered cmds the server thinks this client should ideally have to absorb PL and latency variance
-		bool bFault = true;
-	};
-
-	// We call this in ::SendClientAdjustment to tell the client what the last processed input frame was for him and on what local frame number it was processed
-	UFUNCTION(Client, unreliable)
-	void ClientRecvServerAckFrame(int32 LastProcessedInputFrame, int32 RecvServerFrameNumber, int8 TimeDilation);
-
-	UFUNCTION(Client, unreliable)
-	void ClientRecvServerAckFrameDebug(uint8 NumBuffered, float TargetNumBufferedCmds);
-
-	FServerFrameInfo& GetServerFrameInfo() { return ServerFrameInfo; };
-private:
-
-	FInputCmdBuffer InputBuffer;
-	FClientFrameInfo ClientFrameInfo;
-	FServerFrameInfo ServerFrameInfo;
+	virtual void AsyncPhysicsTickActor(float DeltaTime, float SimTime) override;
 };
diff --git a/Engine/Source/Runtime/Engine/Private/PhysicsEngine/AsyncPhysicsInputComponent.cpp b/Engine/Source/Runtime/Engine/Private/PhysicsEngine/AsyncPhysicsInputComponent.cpp
--- a/Engine/Source/Runtime/Engine/Private/PhysicsEngine/AsyncPhysicsInputComponent.cpp
+++ b/Engine/Source/Runtime/Engine/Private/PhysicsEngine/AsyncPhysicsInputComponent.cpp
@@ -56,38 +56,8 @@
 	static FAutoConsoleVariableRef CVarLerpTargetNumBufferedCmdsAggresively(TEXT("p.net.LerpTargetNumBufferedCmdsAggresively"), LerpTargetNumBufferedCmdsAggresively, TEXT("Aggresively lerp towards TargetNumBufferedCmds. Reduces server side buffering but can cause more artifacts."));
 }
 
+DECLARE_MULTICAST_DELEGATE_TwoParams(FOnDispatchPhysicsTick, int32 PhysicsStep, int32 NumSteps);
 
-// --------------------------------------------------------------------------------------------------------------------------------------------------
-//	Client InputCmd Stream stuff
-// --------------------------------------------------------------------------------------------------------------------------------------------------
-
-
-int8 QuantizeTimeDilation(float F)
-{
-	if (F == 1.f)
-	{
-		return 0;
-	}
-
-	float Normalized = FMath::Clamp<float>((F - 1.f) / InputCmdCVars::MaxTimeDilationMag, -1.f, 1.f);
-	return (int8)(Normalized * 128.f);
-}
-
-float DeQuantizeTimeDilation(int8 i)
-{
-	if (i == 0)
-	{
-		return 1.f;
-	}
-
-	float Normalized = (float)i / 128.f;
-	float Uncompressed = 1.f + (Normalized * InputCmdCVars::MaxTimeDilationMag);
-	return Uncompressed;
-}
-
-
-DECLARE_MULTICAST_DELEGATE_ThreeParams(FOnDispatchPhysicsTick, int32 PhysicsStep, int32 NumSteps, int32 ServerFrame);
-
 struct FAsyncPhysicsInputRewindCallback : public Chaos::IRewindCallback
 {
 	FAsyncPhysicsInputRewindCallback(UWorld* InWorld) : World(InWorld) { }
@@ -99,153 +69,9 @@
 
 	virtual void InjectInputs_External(int32 PhysicsStep, int32 NumSteps) override
 	{
-		int32 LocalOffset = 0;
-		if (World->GetNetMode() == NM_Client)
-		{
-			if (APlayerController* PC = World->GetFirstPlayerController())
-			{
-				APlayerController::FClientFrameInfo& ClientFrameInfo = PC->GetClientFrameInfo();
-
-				if (ClientFrameInfo.LastProcessedInputFrame != INDEX_NONE)
-				{
-					LocalOffset = ClientFrameInfo.GetLocalFrameOffset();
-				}
-			}
-		}
-		DispatchPhysicsTick.Broadcast(PhysicsStep, NumSteps, PhysicsStep - LocalOffset);
+		DispatchPhysicsTick.Broadcast(PhysicsStep, NumSteps);
 	}
 
-	virtual void ProcessInputs_External(int32 PhysicsStep, const TArray<Chaos::FSimCallbackInputAndObject>& SimCallbackInputs)
-	{
-		CachedServerFrame = PhysicsStep;
-
-		if (World->GetNetMode() == NM_Client)
-		{
-			int32 LocalOffset = 0;
-			if (APlayerController* PC = World->GetFirstPlayerController())
-			{
-				// ------------------------------------------------
-				// Send RPC to server telling them what (client/local) physics step we are running
-				//	* Note that SendData is empty because of the existing API, should change this
-				// ------------------------------------------------	
-
-				TArray<uint8> SendData;
-				PC->PushClientInput(PhysicsStep, SendData);
-
-				// -----------------------------------------------------------------
-				// Calculate latest frame offset to map server frame to local frame
-				// -----------------------------------------------------------------
-				APlayerController::FClientFrameInfo& ClientFrameInfo = PC->GetClientFrameInfo();
-
-				if (ClientFrameInfo.LastProcessedInputFrame != INDEX_NONE)
-				{
-					const int32 FrameOffset = ClientFrameInfo.GetLocalFrameOffset();
-					
-					CachedServerFrame = PhysicsStep - FrameOffset; // Local = Server + Offset
-				}
-
-				// -----------------------------------------------------------------
-				// Apply local TIme Dilation based on server's recommendation.
-				// This speeds up or slows down our consumption of real time (by like < 1%)
-				// Ultimately this causes us to send InputCmds at a lower or higher rate in
-				// order to keep server side buffer at optimal capacity. 
-				// Optimal capacity = as small as possible without ever "missing" a frame (e.g, minimal buffer yet always a new fresh cmd to consume server side)
-				// -----------------------------------------------------------------
-				const float RealTimeDilation = DeQuantizeTimeDilation(ClientFrameInfo.QuantizedTimeDilation);
-
-				if (InputCmdCVars::TimeDilationEnabled > 0)
-				{
-					FPhysScene* PhysScene = World->GetPhysicsScene();
-					PhysScene->SetNetworkDeltaTimeScale(RealTimeDilation);
-				}
-			}
-		}
-		else
-		{
-			// -----------------------------------------------
-			// Server: "consume" an InputCmd from each Player Controller
-			// All this means in this context is updating FServerFrameInfo:: LastProcessedInputFrame, LastLocalFrame
-			// (E.g, telling each client what "Input" of theirs we were processing and our local physics frame number.
-			// In cases where the buffer has a fault, we calculate a suggested time dilation to temporarily make client speed up 
-			// or slow down their input cmd production.
-			// -----------------------------------------------
-
-			const bool bForceFault = InputCmdCVars::ForceFault > 0;
-			InputCmdCVars::ForceFault = FMath::Max(0, InputCmdCVars::ForceFault - 1);
-
-			for (FConstPlayerControllerIterator Iterator = World->GetPlayerControllerIterator(); Iterator; ++Iterator)
-			{
-				if (APlayerController* PC = Iterator->Get())
-				{
-					APlayerController::FServerFrameInfo& FrameInfo = PC->GetServerFrameInfo();
-					APlayerController::FInputCmdBuffer& InputBuffer = PC->GetInputBuffer();
-
-					{
-						const int32 NumBufferedInputCmds = bForceFault ? 0 : (InputBuffer.HeadFrame() - FrameInfo.LastProcessedInputFrame);
-
-						// Check Overflow
-						if (NumBufferedInputCmds > InputCmdCVars::MaxBufferedCmds)
-						{
-							UE_LOG(LogPhysics, Warning, TEXT("[Remote.Input] overflow %d %d -> %d"), InputBuffer.HeadFrame(), FrameInfo.LastProcessedInputFrame, NumBufferedInputCmds);
-							FrameInfo.LastProcessedInputFrame = InputBuffer.HeadFrame() - InputCmdCVars::MaxBufferedCmds + 1;
-						}
-
-						// Check fault - we are waiting for Cmds to reach TargetNumBufferedCmds before continuing
-						if (FrameInfo.bFault)
-						{
-							if (NumBufferedInputCmds < (int32)FrameInfo.TargetNumBufferedCmds)
-							{
-								// Skip this because it is in fault. We will use the prev input for this frame.
-								UE_CLOG(FrameInfo.LastProcessedInputFrame != INDEX_NONE, LogPhysics, Warning, TEXT("[Remote.Input] in fault. Reusing Inputcmd. (Client) Input: %d. (Server) Local Frame: %d"), FrameInfo.LastProcessedInputFrame, FrameInfo.LastLocalFrame);
-								return;
-							}
-							FrameInfo.bFault = false;
-						}
-						else if (NumBufferedInputCmds <= 0)
-						{
-							// No Cmds to process, enter fault state. Increment TargetNumBufferedCmds each time this happens.
-							// TODO: We should have something to bring this back down (which means skipping frames) we don't want temporary poor conditions to cause permanent high input buffering
-							FrameInfo.bFault = true;
-							FrameInfo.TargetNumBufferedCmds = FMath::Min(FrameInfo.TargetNumBufferedCmds + InputCmdCVars::TargetNumBufferedCmdsDeltaOnFault, InputCmdCVars::MaxTargetNumBufferedCmds);
-
-							UE_CLOG(FrameInfo.LastProcessedInputFrame != INDEX_NONE, LogPhysics, Warning, TEXT("[Remote.Input] ENTERING fault. New Target: %.2f. (Client) Input: %d. (Server) Local Frame: %d"), FrameInfo.TargetNumBufferedCmds, FrameInfo.LastProcessedInputFrame, FrameInfo.LastLocalFrame);
-							return;
-						}
-
-						float TargetTimeDilation = 1.f;
-						if (NumBufferedInputCmds < (int32)FrameInfo.TargetNumBufferedCmds)
-						{
-							TargetTimeDilation += InputCmdCVars::MaxTimeDilationMag; // Tell client to speed up, we are starved on cmds
-						}
-						/*else if (NumBufferedInputCmds > (int32)FrameInfo.TargetNumBufferedCmds)
-						{
-							TargetTimeDilation -= InputCmdCVars::MaxTimeDilationMag; // Tell client to slow down, we have too many buffered cmds
-
-							if (InputCmdCVars::LerpTargetNumBufferedCmdsAggresively == 0)
-							{
-								// When non aggressive, only lerp when we are above our limit
-								FrameInfo.TargetNumBufferedCmds = FMath::Lerp(FrameInfo.TargetNumBufferedCmds, InputCmdCVars::TargetNumBufferedCmds, InputCmdCVars::TargetNumBufferedCmdsAlpha);
-							}
-						}*/
-
-						FrameInfo.TargetTimeDilation = FMath::Lerp(FrameInfo.TargetTimeDilation, TargetTimeDilation, InputCmdCVars::TimeDilationAlpha);
-						FrameInfo.QuantizedTimeDilation = QuantizeTimeDilation(TargetTimeDilation);
-
-						if (InputCmdCVars::LerpTargetNumBufferedCmdsAggresively != 0)
-						{
-							// When aggressive, always lerp towards target
-							FrameInfo.TargetNumBufferedCmds = FMath::Lerp(FrameInfo.TargetNumBufferedCmds, InputCmdCVars::TargetNumBufferedCmds, InputCmdCVars::TargetNumBufferedCmdsAlpha);
-						}
-
-						FrameInfo.LastProcessedInputFrame++;
-						FrameInfo.LastLocalFrame = PhysicsStep;
-					}
-				}
-			}
-		}
-	}
-
-
 	// Updates the TMap on PhysScene that stores (non interpolated) physics data for replication.
 	// 
 	// Needs to be called from PT context to access fixed tick handle
@@ -401,7 +227,7 @@
 	//BufferedData.Add(NewData);
 }
 
-void UAsyncPhysicsInputComponent::OnDispatchPhysicsTick(int32 PhysicsStep, int32 NumSteps, int32 ServerFrame)
+void UAsyncPhysicsInputComponent::OnDispatchPhysicsTick(int32 PhysicsStep, int32 NumSteps)
 {
 	ensureMsgf(DataClass != nullptr, TEXT("You must call SetDataClass after creating the component"));
 	if (!ensure(DataToWrite != nullptr))
@@ -416,6 +242,8 @@
 	{
 		if (PC->IsLocalController())
 		{
+			const int32 LocalToServerOffset = PC->GetLocalToServerAsyncPhysicsTickOffset();
+			const int32 ServerFrame = PhysicsStep + LocalToServerOffset;
 			UAsyncPhysicsData* DataToSend = DataToWrite;
 			DataToWrite = DuplicateObject((UAsyncPhysicsData*)DataClass->GetDefaultObject(), nullptr);
 			for (int32 Step = 0; Step < NumSteps; ++Step)
@@ -426,7 +254,7 @@
 					//TODO: should probably given user opportunity to modify this per sub-step. For example a jump instruction should only happen on first sub-step
 					DataToSend = DuplicateObject(DataToSend, nullptr);
 				}
-				DataToSend->ServerFrame = ServerFrame + 1 + Step;
+				DataToSend->ServerFrame = ServerFrame + 1 + Step;	//why do we need + 1?
 				BufferedData.Add(DataToSend);
 			}
 		}
diff --git a/Engine/Source/Runtime/Engine/Private/PhysicsEngine/Experimental/PhysScene_Chaos.cpp b/Engine/Source/Runtime/Engine/Private/PhysicsEngine/Experimental/PhysScene_Chaos.cpp
--- a/Engine/Source/Runtime/Engine/Private/PhysicsEngine/Experimental/PhysScene_Chaos.cpp
+++ b/Engine/Source/Runtime/Engine/Private/PhysicsEngine/Experimental/PhysScene_Chaos.cpp
@@ -71,6 +71,13 @@
 #include "Editor.h"
 #endif
 
+struct FPendingAsyncPhysicsCommand
+{
+	int32 PhysicsStep;
+	TWeakObjectPtr<UObject> OwningObject;
+	TFunction<void()> Command;
+};
+
 class FAsyncPhysicsTickCallback : public Chaos::TSimCallbackObject<>
 {
 public:
@@ -82,9 +89,30 @@
 
 	TSet<UActorComponent*> AsyncPhysicsTickComponents;
 	TSet<AActor*> AsyncPhysicsTickActors;
+	TArray<FPendingAsyncPhysicsCommand> PendingCommands;
 	virtual void OnPreSimulate_Internal() override
 	{
 		using namespace Chaos;
+
+		for(int32 Idx = 0; Idx < PendingCommands.Num(); ++Idx)
+		{
+			const int32 CurrentFrame = static_cast<FPBDRigidsSolver*>(GetSolver())->GetCurrentFrame();
+			
+			FPendingAsyncPhysicsCommand* PendingCommand = &PendingCommands[Idx];
+			bool bRemove = PendingCommand->OwningObject.IsStale();
+			if(!bRemove && PendingCommands[Idx].PhysicsStep <= CurrentFrame)
+			{
+				PendingCommands[Idx].Command();
+				bRemove = true;
+			}
+
+			if(bRemove)
+			{
+				PendingCommands.RemoveAt(Idx);	//Need to keep functions in order. If this is slow we could try going in reverse order, but expecting number of commands to be low per frame
+				--Idx;
+			}
+		}
+
 		const FReal DeltaTime = GetDeltaTime_Internal();
 		const FReal SimTime = GetSimTime_Internal();
 		//TODO: handle case where callbacks modify AsyncPhysicsTickComponents or AsyncPhysicsTickActors
@@ -1888,6 +1916,12 @@
 	}
 }
 
+void FPhysScene_Chaos::EnqueueAsyncPhysicsCommand(int32 PhysicsStep, UObject* OwningObject, const TFunction<void()>& Command)
+{
+	EnableAsyncPhysicsTickCallback();
+	AsyncPhysicsTickCallback->PendingCommands.Add({ PhysicsStep, TWeakObjectPtr<UObject>(OwningObject), Command });
+}
+
 TSharedPtr<IPhysicsReplicationFactory> FPhysScene_Chaos::PhysicsReplicationFactory;
 
 #endif // WITH_CHAOS
diff --git a/Engine/Source/Runtime/Engine/Private/PhysicsEngine/PhysicsReplication.cpp b/Engine/Source/Runtime/Engine/Private/PhysicsEngine/PhysicsReplication.cpp
--- a/Engine/Source/Runtime/Engine/Private/PhysicsEngine/PhysicsReplication.cpp
+++ b/Engine/Source/Runtime/Engine/Private/PhysicsEngine/PhysicsReplication.cpp
@@ -536,9 +536,9 @@
 				Chaos::FPhysicsSolver* Solver = PhysScene->GetSolver();
 				check(Solver);
 
-				const APlayerController::FClientFrameInfo& ClientFrameInfo = PlayerController->GetClientFrameInfo();
-				LocalFrameOffset = ClientFrameInfo.GetLocalFrameOffset();
-				NumPredictedFrames = Solver->GetCurrentFrame() - ClientFrameInfo.LastProcessedInputFrame;
+				LocalFrameOffset = PlayerController->GetLocalToServerAsyncPhysicsTickOffset();
+				//TODO: as we send physics updates down we need to record latest seen
+				//NumPredictedFrames = Solver->GetCurrentFrame() - ClientFrameInfo.LastProcessedInputFrame;
 			}
 		}
 	}
diff --git a/Engine/Source/Runtime/Engine/Private/PlayerController.cpp b/Engine/Source/Runtime/Engine/Private/PlayerController.cpp
--- a/Engine/Source/Runtime/Engine/Private/PlayerController.cpp
+++ b/Engine/Source/Runtime/Engine/Private/PlayerController.cpp
@@ -70,7 +70,9 @@
 #include "LevelUtils.h"
 #include "WorldPartition/WorldPartitionStreamingSource.h"
 #include "Physics/AsyncPhysicsInputComponent.h"
+#include "PBDRigidsSolver.h"
 
+
 DEFINE_LOG_CATEGORY(LogPlayerController);
 
 #define LOCTEXT_NAMESPACE "PlayerController"
@@ -93,6 +95,22 @@
 		ECVF_Default);
 }
 
+namespace UE_NETWORK_PHYSICS
+{
+	int32 NumRedundantCmds = 3;
+	FAutoConsoleVariableRef CVarNumRedundantCmds(TEXT("np2.NumRedundantCmds"), NumRedundantCmds, TEXT("Number of redundant user cmds to send per frame"));
+
+#if (UE_BUILD_SHIPPING || UE_BUILD_TEST)
+	int32 EnableDebugRPC = 0;
+#else
+	int32 EnableDebugRPC = 1;
+#endif
+	FAutoConsoleVariableRef CVarEnableDebugRPC(TEXT("np2.EnableDebugRPC"), EnableDebugRPC, TEXT("Sends extra debug information to clients about server side input buffering"));
+
+	int32 EnableNetworkPhysicsPrediction = 0;
+	FAutoConsoleVariableRef CVarEnableNetworkPhysicsPrediction(TEXT("np2.EnableNetworkPhysicsPreidction"), EnableNetworkPhysicsPrediction, TEXT("Enables network physics prediction"));
+}
+
 const float RetryClientRestartThrottleTime = 0.5f;
 const float RetryServerAcknowledgeThrottleTime = 0.25f;
 const float RetryServerCheckSpectatorThrottleTime = 0.25f;
@@ -156,6 +174,8 @@
 		// We want to drive rotation with ControlRotation regardless of attachment state.
 		RootComponent->SetUsingAbsoluteRotation(true);
 	}
+
+	bAsyncPhysicsTickEnabled = !!UE_NETWORK_PHYSICS::EnableNetworkPhysicsPrediction;
 }
 
 float APlayerController::GetNetPriority(const FVector& ViewPos, const FVector& ViewDir, AActor* Viewer, AActor* ViewTarget, UActorChannel* InChannel, float Time, bool bLowBandwidth)
@@ -1634,32 +1654,17 @@
 
 /// @endcond
 
-namespace UE_NETWORK_PHYSICS
-{
-	int32 NumRedundantCmds=3;
-	FAutoConsoleVariableRef CVarNumRedundantCmds(TEXT("np2.NumRedundantCmds"), NumRedundantCmds, TEXT("Number of redundant user cmds to send per frame"));
-
-#if (UE_BUILD_SHIPPING || UE_BUILD_TEST)
-	int32 EnableDebugRPC=0;
-#else
-	int32 EnableDebugRPC=1;
-#endif
-	FAutoConsoleVariableRef CVarEnableDebugRPC(TEXT("np2.EnableDebugRPC"), EnableDebugRPC, TEXT("Sends extra debug information to clients about server side input buffering"));
-}
-
 void APlayerController::SendClientAdjustment()
 {
-	if (ServerFrameInfo.LastProcessedInputFrame != INDEX_NONE && ServerFrameInfo.LastProcessedInputFrame != ServerFrameInfo.LastSentLocalFrame)
+	if (ServerLatestTimestampToCorrect.ServerFrame == INDEX_NONE)
 	{
-		ServerFrameInfo.LastSentLocalFrame = ServerFrameInfo.LastProcessedInputFrame;		
-		ClientRecvServerAckFrame(ServerFrameInfo.LastProcessedInputFrame, ServerFrameInfo.LastLocalFrame, ServerFrameInfo.QuantizedTimeDilation);
-
-		if (UE_NETWORK_PHYSICS::EnableDebugRPC)
-		{
-			ClientRecvServerAckFrameDebug(InputBuffer.HeadFrame() - ServerFrameInfo.LastProcessedInputFrame, ServerFrameInfo.TargetNumBufferedCmds);
-		}
+		//Nothing to correct so do not update client
+		return;
 	}
 
+	ClientCorrectionAsyncPhysicsTimestamp(ServerLatestTimestampToCorrect);
+	ServerLatestTimestampToCorrect.ServerFrame = INDEX_NONE;
+
 	if (AcknowledgedPawn != GetPawn() && !GetSpectatorPawn())
 	{
 		return;
@@ -1677,56 +1682,6 @@
 		}
 	}
 }
-
-void APlayerController::PushClientInput(int32 InRecvClientInputFrame, TArray<uint8>& Data)
-{
-	InputBuffer.Write(InRecvClientInputFrame) = MoveTemp(Data);
-
-	// Do the RPC right here, including the redundant send. This should probably be time based and managed somewhere else like in Tick eventually
-	for (int32 Frame = FMath::Max(1, InRecvClientInputFrame-UE_NETWORK_PHYSICS::NumRedundantCmds+1); Frame <= InRecvClientInputFrame; ++Frame)
-	{
-		ServerRecvClientInputFrame(Frame, InputBuffer.Get(Frame));
-	}
-}
-
-void APlayerController::ServerRecvClientInputFrame_Implementation(int32 InRecvClientInputFrame, const TArray<uint8>& Data)
-{
-	if (InRecvClientInputFrame < 0)
-	{
-		return;
-	}
-
-	for (int32 DroppedFrame = InputBuffer.HeadFrame()+1; DroppedFrame < InRecvClientInputFrame && DroppedFrame > 0; ++DroppedFrame)
-	{
-		UE_LOG(LogPlayerController, Warning, TEXT( "ClientInput Gap in frames (Dropped). %s [%s]. DroppedFrame: %d. RecvFrame: %d. LastProcessInputFrame: %d" ), *GetName(), PlayerState ? *PlayerState->GetPlayerName() : TEXT("???"), DroppedFrame, InRecvClientInputFrame, ServerFrameInfo.LastProcessedInputFrame);
-		//ServerFrameInfo.LastProcessedInputFrame++; // Ehh lets try this for now
-		InputBuffer.Write(DroppedFrame) = InputBuffer.Get(DroppedFrame-1);
-	}
-
-	InputBuffer.Write(InRecvClientInputFrame) = MoveTemp(const_cast<TArray<uint8>&>(Data));
-
-	if (ServerFrameInfo.LastProcessedInputFrame < InputBuffer.TailFrame())
-	{
-		// At this point, things are pretty bad and we are going to drop commands. 
-		// We still need guards client side to not send too many commands
-		UE_LOG(LogPlayerController, Warning, TEXT( "ClientInput buffer overflow. %s [%s]. RecvFrame: %d. LastProcessInputFrame: %d." ), *GetName(), PlayerState ? *PlayerState->GetPlayerName() : TEXT("???"), InRecvClientInputFrame, ServerFrameInfo.LastProcessedInputFrame);
-		ServerFrameInfo.LastProcessedInputFrame = (InputBuffer.TailFrame() + InputBuffer.HeadFrame()) / 2;
-	}
-}
-
-void APlayerController::ClientRecvServerAckFrame_Implementation(int32 LastProcessedInputFrame, int32 RecvServerFrameNumber, int8 TimeDilation)
-{
-	ClientFrameInfo.LastRecvServerFrame = RecvServerFrameNumber;
-	ClientFrameInfo.LastProcessedInputFrame = LastProcessedInputFrame;
-	ClientFrameInfo.QuantizedTimeDilation = TimeDilation;
-}
-
-void APlayerController::ClientRecvServerAckFrameDebug_Implementation(uint8 NumBuffered, float TargetNumBufferedCmds)
-{
-	ClientFrameInfo.LastRecvInputFrame = ClientFrameInfo.LastProcessedInputFrame + NumBuffered;
-	ClientFrameInfo.TargetNumBufferedCmds = TargetNumBufferedCmds;
-}
-
 /// @cond DOXYGEN_WARNINGS
 
 void APlayerController::ClientCapBandwidth_Implementation(int32 Cap)
@@ -5034,6 +4989,11 @@
 	// Clear old axis inputs since we are done with them. 
 	RotationInput = FRotator::ZeroRotator;
 
+	if(!!UE_NETWORK_PHYSICS::EnableNetworkPhysicsPrediction && GetLocalRole() == ROLE_AutonomousProxy && bIsClient)
+	{
+		UpdateServerAsyncPhysicsTickOffset();
+	}
+
 #if !(UE_BUILD_SHIPPING || UE_BUILD_TEST)
 	if (CheatManager != nullptr)
 	{
@@ -5849,5 +5809,153 @@
 	return AsyncPhysicsDataComponent ? AsyncPhysicsDataComponent->GetDataToConsume() : nullptr;
 }
 
+void APlayerController::ExecuteAsyncPhysicsCommand(const FAsyncPhysicsTimestamp& AsyncPhysicsTimestamp, UObject* OwningObject, const TFunction<void()>& Command)
+{
+	if(UWorld* World = GetWorld())
+	{
+		if(FPhysScene* PhysScene = World->GetPhysicsScene())
+		{
+			const int32 PhysicsStep = IsLocalController() ? AsyncPhysicsTimestamp.LocalFrame : AsyncPhysicsTimestamp.ServerFrame;
+			PhysScene->EnqueueAsyncPhysicsCommand(PhysicsStep, OwningObject, Command);
+		}
+	}
+}
 
+FAsyncPhysicsTimestamp APlayerController::GetAsyncPhysicsTimestamp(float DeltaSeconds)
+{
+	using namespace Chaos;
+
+	FAsyncPhysicsTimestamp Timestamp;
+
+	if(UWorld* World = GetWorld())
+	{
+		if (FPhysScene* PhysScene = World->GetPhysicsScene())
+		{
+			if (FPBDRigidsSolver* Solver = static_cast<FPBDRigidsSolver*>(PhysScene->GetSolver()))
+			{
+				const int32 PendingSteps = DeltaSeconds / Solver->GetAsyncDeltaTime();
+
+				int32 LocalPhysicsStep;
+				if(Solver->IsGameThreadFrozen())
+				{
+					//We are calling inside the async tick, so use the current frame
+					LocalPhysicsStep = Solver->GetCurrentFrame();
+				}
+				else
+				{
+					//We are on the GT so give the upcoming async tick
+					LocalPhysicsStep = Solver->GetMarshallingManager().GetExternalTimestamp_External();
+				}
+				
+				LocalPhysicsStep += PendingSteps;	//Add any pending steps user wants to wait on
+				Timestamp.ServerFrame = LocalPhysicsStep;
+				Timestamp.LocalFrame = LocalPhysicsStep;
+
+				if (IsLocalController())
+				{
+					//If local controller we update server frame based on our estimate
+					Timestamp.ServerFrame = LocalPhysicsStep + LocalToServerAsyncPhysicsTickOffset;
+				}
+			}
+		}
+	}
+
+	return Timestamp;
+}
+
+void APlayerController::UpdateServerAsyncPhysicsTickOffset()
+{
+	FAsyncPhysicsTimestamp Timestamp = GetAsyncPhysicsTimestamp();
+	if(ClientLatestAsyncPhysicsStepSent == Timestamp.LocalFrame)
+	{
+		//If GT is running faster than physics sim the physics timestep will not have changed, so no need to send another update to server
+		//This ensures monotonic increase
+		return;
+	}
+
+	ClientLatestAsyncPhysicsStepSent = Timestamp.LocalFrame;
+	ServerSendLatestAsyncPhysicsTimestamp(Timestamp);
+}
+
+void APlayerController::ServerSendLatestAsyncPhysicsTimestamp_Implementation(FAsyncPhysicsTimestamp Timestamp)
+{
+	//This tells the server how the client thinks the async physics tick will line up.
+	//Timestamps could be out of order due to networking, so we make sure they are sorted using the client's local frame which is monotonically increasing
+	int32 Idx;
+	for(Idx = ServerPendingTimestamps.Num() - 1; Idx >= 0; --Idx)
+	{
+		ensureMsgf(ServerPendingTimestamps[Idx].LocalFrame != Timestamp.LocalFrame, TEXT("Client should never send duplicate timestamps, something is wrong"));
+		if(ServerPendingTimestamps[Idx].LocalFrame < Timestamp.LocalFrame)
+		{
+			break;
+		}
+	}
+	ServerPendingTimestamps.Insert(Timestamp, Idx + 1);
+}
+
+void APlayerController::ClientCorrectionAsyncPhysicsTimestamp_Implementation(FAsyncPhysicsTimestamp Timestamp)
+{
+	//This tells the client that a timestamp it sent out was wrong (for example we ran locally on step 5 and expected server to run on step 10, but it actually ran on step 11).
+	//The error can only be later. That is, it can never be that we expect to run on server step 10 but actually ran on sever step 9
+	//Once a timestamp has been corrected, any earlier timestamps can be ignored (these can be out of order because of networking)
+
+	ensureMsgf(Timestamp.ServerFrame != ClientLatestCorrectedOffsetServerStep, TEXT("Server only sends at most one correction per timestamp, duplicate corrections means something is wrong"));
+	if(Timestamp.ServerFrame < ClientLatestCorrectedOffsetServerStep)
+	{
+		//already corrected after this so do nothing
+		return;
+	}
+
+	const int32 NewOffset = Timestamp.ServerFrame - Timestamp.LocalFrame; //The new offset as reported by the server
+	ensureMsgf(NewOffset >= LocalToServerAsyncPhysicsTickOffset, TEXT("The offset between client and server can only ever increase"));
+	LocalToServerAsyncPhysicsTickOffset = NewOffset;
+}
+
+void APlayerController::AsyncPhysicsTickActor(float DeltaTime, float SimTime)
+{
+	Super::AsyncPhysicsTickActor(DeltaTime, SimTime);
+
+	//TODO: only kick this off if server and using this feature
+	if (IsLocalController()) { return; }
+
+	if (ServerPendingTimestamps.Num() == 0)
+	{
+		//TODO: starved for input user needs to speed up
+		return;
+	}
+
+	const FAsyncPhysicsTimestamp ActualTimestamp = GetAsyncPhysicsTimestamp();
+
+	//If the pending timestamp is bigger than the server frame we simple wait, otherwise make sure they aren't too early
+	if (ServerPendingTimestamps[0].ServerFrame <= ActualTimestamp.ServerFrame)
+	{
+		if (ServerPendingTimestamps[0].ServerFrame < ActualTimestamp.ServerFrame)
+		{
+			//The earliest pending client timestamp is too early, so their offset must be bigger than they thought
+			ServerLatestTimestampToCorrect.ServerFrame = ActualTimestamp.ServerFrame;
+			ServerLatestTimestampToCorrect.LocalFrame = ServerPendingTimestamps[0].LocalFrame;	//The client's local frame stays the same there's just a mismatch on the server frame
+		}
+
+		ServerPendingTimestamps.RemoveAt(0);	//We've updated the client if needed, so can discard
+		//NOTE: we purposely don't fixup any future pending client timestamps. This is because the RPC is unreliable so it's best to send the error correction redundantly over multiple frames if needed
+	}
+}
+
+void APlayerController::ClientAckTimeDilation_Implementation(float TimeDilation, int32 ServerStep)
+{
+	if(ServerStep < ClientLatestTimeDilationServerStep)
+	{
+		//Stale ack so do nothing
+		return;
+	}
+
+	ensureMsgf(ClientLatestTimeDilationServerStep != ServerStep, TEXT("Server should send at most one time dilation per step, duplicate means something is wrong"));
+	ClientLatestTimeDilationServerStep = ServerStep;
+	if(UWorld* World = GetWorld())
+	{
+		World->GetPhysicsScene()->SetNetworkDeltaTimeScale(TimeDilation);
+	}
+}
+
+
 #undef LOCTEXT_NAMESPACE
diff --git a/Engine/Source/Runtime/Engine/Public/Physics/AsyncPhysicsInputComponent.h b/Engine/Source/Runtime/Engine/Public/Physics/AsyncPhysicsInputComponent.h
--- a/Engine/Source/Runtime/Engine/Public/Physics/AsyncPhysicsInputComponent.h
+++ b/Engine/Source/Runtime/Engine/Public/Physics/AsyncPhysicsInputComponent.h
@@ -21,7 +21,7 @@
 	UFUNCTION(Server, unreliable)
 	void ServerRPCBufferInput(UAsyncPhysicsData* AsyncPhysicsData);
 
-	void OnDispatchPhysicsTick(int32 PhysicsStep, int32 NumSteps, int32 ServerFrame);
+	void OnDispatchPhysicsTick(int32 PhysicsStep, int32 NumSteps);
 
 	APlayerController* GetPlayerController();
 
diff --git a/Engine/Source/Runtime/Engine/Public/Physics/Experimental/PhysScene_Chaos.h b/Engine/Source/Runtime/Engine/Public/Physics/Experimental/PhysScene_Chaos.h
--- a/Engine/Source/Runtime/Engine/Public/Physics/Experimental/PhysScene_Chaos.h
+++ b/Engine/Source/Runtime/Engine/Public/Physics/Experimental/PhysScene_Chaos.h
@@ -128,6 +128,8 @@
 	void RegisterAsyncPhysicsTickActor(AActor* Actor);
 	void UnregisterAsyncPhysicsTickActor(AActor* Actor);
 
+	void EnqueueAsyncPhysicsCommand(int32 PhysicsStep, UObject* OwningObject, const TFunction<void()>& Command);
+
 	/**
 	 * Called during creation of the physics state for gamethread objects to pass off an object to the physics thread
 	 */
diff --git a/Engine/Source/Runtime/Experimental/Chaos/Public/Chaos/Framework/PhysicsSolverBase.h b/Engine/Source/Runtime/Experimental/Chaos/Public/Chaos/Framework/PhysicsSolverBase.h
--- a/Engine/Source/Runtime/Experimental/Chaos/Public/Chaos/Framework/PhysicsSolverBase.h
+++ b/Engine/Source/Runtime/Experimental/Chaos/Public/Chaos/Framework/PhysicsSolverBase.h
@@ -436,6 +436,8 @@
 			return IsUsingAsyncResults() && UseAsyncInterpolation;
 		}
 
+		FReal GetAsyncDeltaTime() const { return AsyncDt; }
+
 		void SetMaxDeltaTime_External(float InMaxDeltaTime)
 		{
 			MMaxDeltaTime = InMaxDeltaTime;
