Cleaning up replay backwards compat checks from prior to the new minimum supported version

diff --git a/Engine/Source/Runtime/Engine/Private/DemoNetDriver.cpp b/Engine/Source/Runtime/Engine/Private/DemoNetDriver.cpp
--- a/Engine/Source/Runtime/Engine/Private/DemoNetDriver.cpp
+++ b/Engine/Source/Runtime/Engine/Private/DemoNetDriver.cpp
@@ -3616,30 +3616,24 @@ bool UDemoNetDriver::FastForwardLevels(const FGotoResult& GotoResult)
 
 				PacketOffset += CheckpointArchive->Tell();
 
-				if (PlaybackVersion >= HISTORY_MULTIPLE_LEVELS)
-				{
-					int32 LevelIndex = INDEX_NONE;
-					*CheckpointArchive << LevelIndex;
-				}
+				int32 LevelIndex = INDEX_NONE;
+				*CheckpointArchive << LevelIndex;
 
-				if (PlaybackVersion >= HISTORY_DELETED_STARTUP_ACTORS)
+				if (bDeltaCheckpoint)
 				{
-					if (bDeltaCheckpoint)
-					{
-						TUniquePtr<FDeltaCheckpointData>& CheckpointData = PlaybackDeltaCheckpointData.Emplace_GetRef(new FDeltaCheckpointData());
+					TUniquePtr<FDeltaCheckpointData>& CheckpointData = PlaybackDeltaCheckpointData.Emplace_GetRef(new FDeltaCheckpointData());
 
-						ReplayHelper.ReadDeletedStartupActors(ServerConnection, *CheckpointArchive, CheckpointData->DestroyedNetStartupActors);
-						ReplayHelper.PlaybackDeletedNetStartupActors.Append(CheckpointData->DestroyedNetStartupActors);
+					ReplayHelper.ReadDeletedStartupActors(ServerConnection, *CheckpointArchive, CheckpointData->DestroyedNetStartupActors);
+					ReplayHelper.PlaybackDeletedNetStartupActors.Append(CheckpointData->DestroyedNetStartupActors);
 
-						*CheckpointArchive << CheckpointData->DestroyedDynamicActors;
-						*CheckpointArchive << CheckpointData->ChannelsToClose;
-					}
-					else
-					{
-						ReplayHelper.PlaybackDeletedNetStartupActors.Empty();
+					*CheckpointArchive << CheckpointData->DestroyedDynamicActors;
+					*CheckpointArchive << CheckpointData->ChannelsToClose;
+				}
+				else
+				{
+					ReplayHelper.PlaybackDeletedNetStartupActors.Empty();
 
-						ReplayHelper.ReadDeletedStartupActors(ServerConnection, *CheckpointArchive, ReplayHelper.PlaybackDeletedNetStartupActors);
-					}
+					ReplayHelper.ReadDeletedStartupActors(ServerConnection, *CheckpointArchive, ReplayHelper.PlaybackDeletedNetStartupActors);
 				}
 
 				CheckpointArchive->Seek(PacketOffset);
@@ -3889,12 +3883,9 @@ bool UDemoNetDriver::LoadCheckpoint(const FGotoResult& GotoResult)
 
 	uint32 PlaybackVersion = GetPlaybackDemoVersion();
 
-	if (PlaybackVersion >= HISTORY_MULTIPLE_LEVELS)
+	if (GotoCheckpointArchive->TotalSize() > 0)
 	{
-		if (GotoCheckpointArchive->TotalSize() > 0)
-		{
-			*GotoCheckpointArchive << LevelForCheckpoint;
-		}
+		*GotoCheckpointArchive << LevelForCheckpoint;
 	}
 
 	check(World);
@@ -4224,32 +4215,26 @@ bool UDemoNetDriver::LoadCheckpoint(const FGotoResult& GotoResult)
 			*GotoCheckpointArchive << PacketOffset;
 		}
 
-		if (PlaybackVersion >= HISTORY_MULTIPLE_LEVELS)
-		{
-			int32 LevelIndex = INDEX_NONE;
-			*GotoCheckpointArchive << LevelIndex;
-		}
+		int32 LevelIndex = INDEX_NONE;
+		*GotoCheckpointArchive << LevelIndex;
 
 		// Load net startup actors that need to be destroyed
-		if (PlaybackVersion >= HISTORY_DELETED_STARTUP_ACTORS)
+		if (bDeltaCheckpoint)
 		{
-			if (bDeltaCheckpoint)
-			{
-				TUniquePtr<FDeltaCheckpointData>& CheckpointData = PlaybackDeltaCheckpointData.Emplace_GetRef(new FDeltaCheckpointData());
+			TUniquePtr<FDeltaCheckpointData>& CheckpointData = PlaybackDeltaCheckpointData.Emplace_GetRef(new FDeltaCheckpointData());
 
-				ReplayHelper.ReadDeletedStartupActors(ServerConnection, *GotoCheckpointArchive, CheckpointData->DestroyedNetStartupActors);
+			ReplayHelper.ReadDeletedStartupActors(ServerConnection, *GotoCheckpointArchive, CheckpointData->DestroyedNetStartupActors);
 
-				ReplayHelper.PlaybackDeletedNetStartupActors.Append(CheckpointData->DestroyedNetStartupActors);
+			ReplayHelper.PlaybackDeletedNetStartupActors.Append(CheckpointData->DestroyedNetStartupActors);
 
-				*GotoCheckpointArchive << CheckpointData->DestroyedDynamicActors;
-				*GotoCheckpointArchive << CheckpointData->ChannelsToClose;
-			}
-			else
-			{
-				ReplayHelper.PlaybackDeletedNetStartupActors.Empty();
+			*GotoCheckpointArchive << CheckpointData->DestroyedDynamicActors;
+			*GotoCheckpointArchive << CheckpointData->ChannelsToClose;
+		}
+		else
+		{
+			ReplayHelper.PlaybackDeletedNetStartupActors.Empty();
 
-				ReplayHelper.ReadDeletedStartupActors(ServerConnection, *GotoCheckpointArchive, ReplayHelper.PlaybackDeletedNetStartupActors);
-			}
+			ReplayHelper.ReadDeletedStartupActors(ServerConnection, *GotoCheckpointArchive, ReplayHelper.PlaybackDeletedNetStartupActors);
 		}
 
 		int32 NumValues = 0;
diff --git a/Engine/Source/Runtime/Engine/Private/ReplayHelper.cpp b/Engine/Source/Runtime/Engine/Private/ReplayHelper.cpp
--- a/Engine/Source/Runtime/Engine/Private/ReplayHelper.cpp
+++ b/Engine/Source/Runtime/Engine/Private/ReplayHelper.cpp
@@ -1897,13 +1897,8 @@
 
 	int32 ReadCurrentLevelIndex = 0;
 
-	uint32 PlaybackVersion = PlaybackDemoHeader.Version;
+	Ar << ReadCurrentLevelIndex;
 
-	if (PlaybackVersion >= HISTORY_MULTIPLE_LEVELS)
-	{
-		Ar << ReadCurrentLevelIndex;
-	}
-
 	float TimeSeconds = 0.0f;
 
 	Ar << TimeSeconds;
@@ -1913,7 +1908,6 @@
 		*OutTime = TimeSeconds;
 	}
 
-	if (PlaybackVersion >= HISTORY_LEVEL_STREAMING_FIXES)
 	{
 		DECLARE_SCOPE_CYCLE_COUNTER(TEXT("Demo_ReceiveExports"), Demo_ReceiveExports, STATGROUP_Net);
 		CastChecked<UPackageMapClient>(Connection->PackageMap)->ReceiveExportData(Ar);
diff --git a/Engine/Source/Runtime/Engine/Public/ReplayTypes.h b/Engine/Source/Runtime/Engine/Public/ReplayTypes.h
--- a/Engine/Source/Runtime/Engine/Public/ReplayTypes.h
+++ b/Engine/Source/Runtime/Engine/Public/ReplayTypes.h
@@ -185,29 +185,9 @@
 		Ar << Header.NetworkChecksum;
 		Ar << Header.EngineNetworkProtocolVersion;
 		Ar << Header.GameNetworkProtocolVersion;
+		Ar << Header.Guid;
+		Ar << Header.EngineVersion;
 
-		if (Header.Version >= HISTORY_HEADER_GUID)
-		{
-			Ar << Header.Guid;
-		}
-
-		if (Header.Version >= HISTORY_SAVE_FULL_ENGINE_VERSION)
-		{
-			Ar << Header.EngineVersion;
-		}
-		else
-		{
-			// Previous versions only stored the changelist
-			uint32 Changelist = 0;
-			Ar << Changelist;
-
-			if (Ar.IsLoading())
-			{
-				// We don't have any valid information except the changelist.
-				Header.EngineVersion.Set(0, 0, 0, Changelist, FString());
-			}
-		}
-
 		if (Header.Version >= HISTORY_SAVE_PACKAGE_VERSION_UE)
 		{
 			Ar << Header.PackageVersionUE;
@@ -240,32 +220,8 @@
 			}
 		}
 
-		if (Header.Version < HISTORY_MULTIPLE_LEVELS)
-		{
-			FString LevelName;
-			Ar << LevelName;
-			Header.LevelNamesAndTimes.Add(FLevelNameAndTime(LevelName, 0));
-		}
-		else if (Header.Version == HISTORY_MULTIPLE_LEVELS)
-		{
-			TArray<FString> LevelNames;
-			Ar << LevelNames;
-
-			for (const FString& LevelName : LevelNames)
-			{
-				Header.LevelNamesAndTimes.Add(FLevelNameAndTime(LevelName, 0));
-			}
-		}
-		else
-		{
-			Ar << Header.LevelNamesAndTimes;
-		}
-
-		if (Header.Version >= HISTORY_HEADER_FLAGS)
-		{
-			Ar << Header.HeaderFlags;
-		}
-
+		Ar << Header.LevelNamesAndTimes;
+		Ar << Header.HeaderFlags;
 		Ar << Header.GameSpecificData;
 
 		if (Header.Version >= HISTORY_RECORDING_METADATA)
