Removing the option to block resize frame

diff --git a/Engine/Source/Runtime/Engine/Private/Slate/SceneViewport.cpp b/Engine/Source/Runtime/Engine/Private/Slate/SceneViewport.cpp
--- a/Engine/Source/Runtime/Engine/Private/Slate/SceneViewport.cpp
+++ b/Engine/Source/Runtime/Engine/Private/Slate/SceneViewport.cpp
@@ -25,13 +25,6 @@
 
 DEFINE_LOG_CATEGORY(LogViewport);
 
-static bool bSkipViewportFrameResize = false;
-FAutoConsoleVariableRef CVarSceneSkipResizeFrame(
-	TEXT("Viewport.SkipResizeFrame"),
-	bSkipViewportFrameResize,
-	TEXT("Block the resize frame method"),
-	ECVF_Default);
-
 extern EWindowMode::Type GetWindowModeType(EWindowMode::Type WindowMode);
 
 const FName NAME_SceneViewport = FName(TEXT("SceneViewport"));
@@ -1306,10 +1299,6 @@
 
 void FSceneViewport::ResizeFrame(uint32 NewWindowSizeX, uint32 NewWindowSizeY, EWindowMode::Type NewWindowMode)
 {
-	if (bSkipViewportFrameResize)
-	{
-		return;
-	}
 	// Resizing the window directly is only supported in the game
 	if( FApp::IsGame() && FApp::CanEverRender() && NewWindowSizeX > 0 && NewWindowSizeY > 0 )
 	{		
