Altering a texture encoding setting that only applies to Final encode now falls back to the Fast encoding instead of ignoring it. This is more "correct" than ignoring it, however it probably still isn't what people expect. Currently punting clarity for the upcoming texture editor ui pass.

diff --git a/Engine/Source/Runtime/Engine/Classes/Engine/Texture.h b/Engine/Source/Runtime/Engine/Classes/Engine/Texture.h
--- a/Engine/Source/Runtime/Engine/Classes/Engine/Texture.h
+++ b/Engine/Source/Runtime/Engine/Classes/Engine/Texture.h
@@ -843,14 +843,15 @@ struct FTexturePlatformData
 		}
 	};
 
-	/** This is the key for the FetchOrBuild variant of our Cache. We assume that uniqueness
+	/** 
+		The keys for both types of fetches. We assume that uniqueness
 	*	for that is equivalent to uniqueness if we use both FetchFirst and FetchOrBuild. This
-	*	is used as the key in to CookedPlatformData, as well as to determine if we are already
-	*	cooking the data the editor needs in CachePlatformData. 
+	*	is used as the key in to CookedPlatformData, as well as to determine if we need to recache.
 	*	Note that since this is read on the game thread constantly in CachePlatformData, it
 	*	must be written to on the game thread to avoid false recaches.
 	*/
 	TVariant<FString, FStructuredDerivedDataKey> FetchOrBuildDerivedDataKey;
+	TVariant<FString, FStructuredDerivedDataKey> FetchFirstDerivedDataKey;
 
 	/** Async cache task if one is outstanding. */
 	struct FTextureAsyncCacheDerivedDataTask* AsyncTask;
diff --git a/Engine/Source/Runtime/Engine/Private/TextureDerivedData.cpp b/Engine/Source/Runtime/Engine/Private/TextureDerivedData.cpp
--- a/Engine/Source/Runtime/Engine/Private/TextureDerivedData.cpp
+++ b/Engine/Source/Runtime/Engine/Private/TextureDerivedData.cpp
@@ -2651,16 +2651,35 @@ void UTexture::CachePlatformData(bool bAsyncCache, bool bAllowAsyncBuild, bool b
 			// know we always have those. If we need the FetchFirst key, we generate it later when
 			// we know we're actually going to Cache()
 			//
-			TArray<FTextureBuildSettings> BuildSettingsFetchOrBuild;
+			TArray<FTextureBuildSettings> BuildSettingsFetchOrBuild;			
 			TArray<FTexturePlatformData::FTextureEncodeResultMetadata> ResultMetadataFetchOrBuild;
-			if (EncodeSpeed == ETextureEncodeSpeed::FinalIfAvailable ||
-				EncodeSpeed == ETextureEncodeSpeed::Fast)
-			{
-				GetBuildSettingsForRunningPlatform(*this, ETextureEncodeSpeed::Fast, BuildSettingsFetchOrBuild, &ResultMetadataFetchOrBuild);
-			}
-			else
+
+			// These might be empty.
+			TArray<FTextureBuildSettings> BuildSettingsFetchFirst;
+			TArray<FTexturePlatformData::FTextureEncodeResultMetadata> ResultMetadataFetchFirst;
+
+			switch (EncodeSpeed)
 			{
-				GetBuildSettingsForRunningPlatform(*this, ETextureEncodeSpeed::Final, BuildSettingsFetchOrBuild, &ResultMetadataFetchOrBuild);
+			case ETextureEncodeSpeed::FinalIfAvailable:
+				{
+					GetBuildSettingsForRunningPlatform(*this, ETextureEncodeSpeed::Final, BuildSettingsFetchFirst, &ResultMetadataFetchFirst);
+					GetBuildSettingsForRunningPlatform(*this, ETextureEncodeSpeed::Fast, BuildSettingsFetchOrBuild, &ResultMetadataFetchOrBuild);
+					break;
+				}
+			case ETextureEncodeSpeed::Fast:
+				{
+					GetBuildSettingsForRunningPlatform(*this, ETextureEncodeSpeed::Fast, BuildSettingsFetchOrBuild, &ResultMetadataFetchOrBuild);
+					break;
+				}
+			case ETextureEncodeSpeed::Final:
+				{
+					GetBuildSettingsForRunningPlatform(*this, ETextureEncodeSpeed::Final, BuildSettingsFetchOrBuild, &ResultMetadataFetchOrBuild);
+					break;
+				}
+			default:
+				{
+					UE_LOG(LogTexture, Fatal, TEXT("Invalid encode speed in CachePlatformData"));
+				}
 			}
 
 			// If we're open in a texture editor, then we might have custom build settings.
@@ -2674,6 +2693,8 @@ void UTexture::CachePlatformData(bool bAsyncCache, bool bAllowAsyncBuild, bool b
 					// speed to whatever we already have staged, then set those settings to the custom
 					// ones.
 					EncodeSpeed = (ETextureEncodeSpeed)BuildSettingsFetchOrBuild[0].RepresentsEncodeSpeedNoSend;
+					BuildSettingsFetchFirst.Empty();
+					ResultMetadataFetchFirst.Empty();
 
 					for (int32 i = 0; i < BuildSettingsFetchOrBuild.Num(); i++)
 					{
@@ -2701,19 +2722,33 @@ void UTexture::CachePlatformData(bool bAsyncCache, bool bAllowAsyncBuild, bool b
 			bPerformCache = true;
 			if (PlatformDataLink != nullptr)
 			{
-				// Check if our keys match.
+				bPerformCache = false;
+
+				// Check if our keys match. If we have two, they both have to match, otherwise a change that only affects one
+				// might not cause a rebuild, leading to confusion in the texture editor.
 				if (FTexturePlatformData::IsUsingNewDerivedData() && (Source.GetNumLayers() == 1) && !BuildSettingsFetchOrBuild[0].bVirtualStreamable)
 				{
 					// DDC2 version
 					using namespace UE::DerivedData;
 					if (const FTexturePlatformData::FStructuredDerivedDataKey* ExistingDerivedDataKey = PlatformDataLink->FetchOrBuildDerivedDataKey.TryGet<FTexturePlatformData::FStructuredDerivedDataKey>())
 					{
-						if (*ExistingDerivedDataKey == CreateTextureDerivedDataKey(*this, CacheFlags, BuildSettingsFetchOrBuild[0]))
+						if (*ExistingDerivedDataKey != CreateTextureDerivedDataKey(*this, CacheFlags, BuildSettingsFetchOrBuild[0]))
 						{
-							bPerformCache = false;
+							bPerformCache = true;
 						}						
 					}
-				}
+
+					if (BuildSettingsFetchFirst.Num())
+					{
+						if (const FTexturePlatformData::FStructuredDerivedDataKey* ExistingDerivedDataKey = PlatformDataLink->FetchFirstDerivedDataKey.TryGet<FTexturePlatformData::FStructuredDerivedDataKey>())
+						{
+							if (*ExistingDerivedDataKey != CreateTextureDerivedDataKey(*this, CacheFlags, BuildSettingsFetchFirst[0]))
+							{
+								bPerformCache = true;
+							}
+						}
+					}
+				} // end if ddc2
 				else
 				{
 					// DDC1 version.
@@ -2721,12 +2756,25 @@ void UTexture::CachePlatformData(bool bAsyncCache, bool bAllowAsyncBuild, bool b
 					{
 						FString DerivedDataKey;
 						GetTextureDerivedDataKey(*this, BuildSettingsFetchOrBuild.GetData(), DerivedDataKey);
-						if (*ExistingDerivedDataKey == DerivedDataKey)
+						if (*ExistingDerivedDataKey != DerivedDataKey)
 						{
-							bPerformCache = false;
+							bPerformCache = true;
 						}
 					}
-				}
+
+					if (BuildSettingsFetchFirst.Num())
+					{
+						if (const FString* ExistingDerivedDataKey = PlatformDataLink->FetchFirstDerivedDataKey.TryGet<FString>())
+						{
+							FString DerivedDataKey;
+							GetTextureDerivedDataKey(*this, BuildSettingsFetchFirst.GetData(), DerivedDataKey);
+							if (*ExistingDerivedDataKey != DerivedDataKey)
+							{
+								bPerformCache = true;
+							}
+						}
+					}
+				} // end if ddc1
 			} // end if checking existing data matches.
 
 			if (bPerformCache)
@@ -2745,15 +2793,6 @@ void UTexture::CachePlatformData(bool bAsyncCache, bool bAllowAsyncBuild, bool b
 					PlatformDataLink = new FTexturePlatformData();
 				}
 
-				// We delayed generating our FetchFirst settings since we assume we'll usually be just testing
-				// keys above.
-				TArray<FTextureBuildSettings> BuildSettingsFetchFirst;
-				TArray<FTexturePlatformData::FTextureEncodeResultMetadata> ResultMetadataFetchFirst;
-				if (EncodeSpeed == ETextureEncodeSpeed::FinalIfAvailable)
-				{
-					GetBuildSettingsForRunningPlatform(*this, ETextureEncodeSpeed::Final, BuildSettingsFetchFirst, &ResultMetadataFetchFirst);
-				}
-
 				PlatformDataLink->Cache(
 					*this, 
 					BuildSettingsFetchFirst.Num() ? BuildSettingsFetchFirst.GetData() : nullptr, 
diff --git a/Engine/Source/Runtime/Engine/Private/TextureDerivedDataTask.cpp b/Engine/Source/Runtime/Engine/Private/TextureDerivedDataTask.cpp
--- a/Engine/Source/Runtime/Engine/Private/TextureDerivedDataTask.cpp
+++ b/Engine/Source/Runtime/Engine/Private/TextureDerivedDataTask.cpp
@@ -636,7 +636,7 @@
 		FetchOrBuildMetadata = *InFetchOrBuildMetadata;
 	}
 
-	// FetchOrBuildDerivedDataKey needs to be assigned on the create thread.
+	// Keys need to be assigned on the create thread.
 	{
 		FString LocalKeySuffix;
 		GetTextureDerivedDataKeySuffix(Texture, BuildSettingsPerLayerFetchOrBuild.GetData(), LocalKeySuffix);
@@ -644,6 +644,14 @@
 		GetTextureDerivedDataKeyFromSuffix(LocalKeySuffix, DDK);
 		InDerivedData->FetchOrBuildDerivedDataKey.Emplace<FString>(DDK);
 	}
+	if (BuildSettingsPerLayerFetchFirst.Num())
+	{
+		FString LocalKeySuffix;
+		GetTextureDerivedDataKeySuffix(Texture, BuildSettingsPerLayerFetchFirst.GetData(), LocalKeySuffix);
+		FString DDK;
+		GetTextureDerivedDataKeyFromSuffix(LocalKeySuffix, DDK);
+		InDerivedData->FetchFirstDerivedDataKey.Emplace<FString>(DDK);
+	}
 
 	// At this point, the texture *MUST* have a valid GUID.
 	if (!Texture.Source.GetId().IsValid())
@@ -1191,6 +1199,7 @@
 		{
 			// If the keys are the same, ignore fetch first.
 			FBuildDefinition FetchDefinition = CreateDefinition(Build, Texture, TexturePath, FunctionName, *InSettingsFetchFirst, bUseCompositeTexture);
+			DerivedData.FetchFirstDerivedDataKey.Emplace<FTexturePlatformData::FStructuredDerivedDataKey>(GetKey(FetchDefinition, Texture, bUseCompositeTexture));
 			if (FetchDefinition.GetKey() != FetchOrBuildDefinition.GetKey())
 			{
 				bBuildKicked = true;
