Fix camera manager's BlendParams keeping outdated information

It wasn't being set when a non-blending view-target change would come through, which means that any gameplay code reading it would incorrectly think that the current view-target was blending.

diff --git a/Engine/Source/Runtime/Engine/Private/PlayerCameraManager.cpp b/Engine/Source/Runtime/Engine/Private/PlayerCameraManager.cpp
--- a/Engine/Source/Runtime/Engine/Private/PlayerCameraManager.cpp
+++ b/Engine/Source/Runtime/Engine/Private/PlayerCameraManager.cpp
@@ -151,7 +151,6 @@
 
 			// use last frame's POV
 			ViewTarget.POV = GetLastFrameCameraCacheView();
-			BlendParams = TransitionParams;
 			BlendTimeToGo = TransitionParams.BlendTime;
 
 			AssignViewTarget(NewTarget, PendingViewTarget, TransitionParams);
@@ -189,6 +188,10 @@
 		}
 		PendingViewTarget.Target = NULL;
 	}
+
+	// update the blend params after all the assignment logic so that sub-classes can compare
+	// the old vs new parameters if needed.
+	BlendParams = TransitionParams;
 }
 
 
