LevelInstanceActorDesc: Make constructor public

diff --git a/Engine/Source/Runtime/Engine/Public/WorldPartition/LevelInstance/LevelInstanceActorDesc.h b/Engine/Source/Runtime/Engine/Public/WorldPartition/LevelInstance/LevelInstanceActorDesc.h
--- a/Engine/Source/Runtime/Engine/Public/WorldPartition/LevelInstance/LevelInstanceActorDesc.h
+++ b/Engine/Source/Runtime/Engine/Public/WorldPartition/LevelInstance/LevelInstanceActorDesc.h
@@ -19,9 +19,9 @@
 class ENGINE_API FLevelInstanceActorDesc : public FWorldPartitionActorDesc
 {
 #if WITH_EDITOR
-	friend class ALevelInstance;
 
 public:
+	FLevelInstanceActorDesc();
 	virtual ~FLevelInstanceActorDesc() override;
 
 	inline FName GetLevelPackage() const { return LevelPackage; }
@@ -29,7 +29,6 @@
 	virtual bool GetContainerInstance(const UActorDescContainer*& OutLevelContainer, FTransform& OutLevelTransform, EContainerClusterMode& OutClusterMode) const override;
 
 protected:
-	FLevelInstanceActorDesc();
 	virtual void Init(const AActor* InActor) override;
 	virtual void Init(const FWorldPartitionActorDescInitData& DescData) override;
 	virtual bool Equals(const FWorldPartitionActorDesc* Other) const override;
