Fix crash when a replay tries to load a map that doesn't exist

diff --git a/Engine/Source/Runtime/Engine/Private/DemoNetDriver.cpp b/Engine/Source/Runtime/Engine/Private/DemoNetDriver.cpp
--- a/Engine/Source/Runtime/Engine/Private/DemoNetDriver.cpp
+++ b/Engine/Source/Runtime/Engine/Private/DemoNetDriver.cpp
@@ -1035,6 +1035,12 @@ bool UDemoNetDriver::InitConnectInternal(FString& Error)
 			FURL LocalDemoURL;
 			LocalDemoURL.Map = ReplayHelper.PlaybackDemoHeader.LevelNamesAndTimes[0].LevelName;
 
+			if (!GEngine->MakeSureMapNameIsValid(LocalDemoURL.Map))
+			{
+				NotifyDemoPlaybackFailure(EDemoPlayFailure::LoadMap);
+				return false;
+			}
+
 			FWorldContext * WorldContext = GEngine->GetWorldContextFromWorld(World);
 
 			if (WorldContext == nullptr)
