Fix for edge case decoder streaming error

diff --git a/Engine/Source/Runtime/Engine/Private/AudioDecompress.cpp b/Engine/Source/Runtime/Engine/Private/AudioDecompress.cpp
--- a/Engine/Source/Runtime/Engine/Private/AudioDecompress.cpp
+++ b/Engine/Source/Runtime/Engine/Private/AudioDecompress.cpp
@@ -128,8 +128,7 @@
 
 	while (RawPCMOffset < QualityInfo->SampleDataSize)
 	{
-		uint16 FrameSize = GetFrameSize();
-		int32 DecodedSamples = DecompressToPCMBuffer(FrameSize);
+		int32 DecodedSamples = DecompressToPCMBuffer( /*Unused*/ 0);
 
 		if (DecodedSamples < 0)
 		{
@@ -185,8 +184,18 @@
 
 	SCOPE_CYCLE_COUNTER(STAT_AudioStreamedDecompressTime);
 
-	UE_LOG(LogAudio, Log, TEXT("Streaming compressed data from SoundWave'%s' - Chunk %d, Offset %d"), *StreamingSoundWave->GetFName().ToString(), CurrentChunkIndex, SrcBufferOffset);
 
+	UE_LOG(LogAudio, Log, TEXT("Streaming compressed data from SoundWave'%s' - Chunk=%d\tNumChunks=%d\tOffset=%d\tChunkSize=%d\tLooping=%s\tLastPCMOffset=%d\tContainsEOF=%s" ), 
+		*StreamingSoundWave->GetFName().ToString(), 	
+		CurrentChunkIndex, 
+		StreamingSoundWave->GetNumChunks(),
+		SrcBufferOffset, 
+		SrcBufferDataSize,
+		bLooping ? TEXT("YES") : TEXT("NO"),
+		LastPCMOffset,		
+		bStoringEndOfFile ? TEXT("YES") : TEXT("NO")
+	);
+
 	// Write out any PCM data that was decoded during the last request
 	uint32 RawPCMOffset = WriteFromDecodedPCM(Destination, BufferSize);
 
@@ -245,11 +254,9 @@
 
 	while (RawPCMOffset < BufferSize)
 	{
-		// Get the platform-dependent size of the current "frame" of encoded audio (note: frame is used here differently than audio frame/sample)
-		uint16 FrameSize = GetFrameSize();
-
 		// Decompress the next compression frame of audio (many samples) into the PCM buffer
-		int32 DecodedSamples = DecompressToPCMBuffer(FrameSize);
+		int32 DecodedSamples = DecompressToPCMBuffer(/*Unused*/ 0);
+
 
 		if (DecodedSamples < 0)
 		{
@@ -264,6 +271,8 @@
 
 			RawPCMOffset += WriteFromDecodedPCM(Destination + RawPCMOffset, BufferSize - RawPCMOffset);
 
+			const int32 PreviousChunkIndex = CurrentChunkIndex;
+
 			// Have we reached the end of buffer?
 			if (SrcBufferOffset >= SrcBufferDataSize - SrcBufferPadding)
 			{
@@ -304,7 +313,7 @@
 				SrcBufferData = GetLoadedChunk(StreamingSoundWave, CurrentChunkIndex, SrcBufferDataSize);
 				if (SrcBufferData)
 				{
-					UE_LOG(LogAudio, Log, TEXT("Incremented current chunk from SoundWave'%s' - Chunk %d, Offset %d"), *StreamingSoundWave->GetFName().ToString(), CurrentChunkIndex, SrcBufferOffset);
+					UE_CLOG(PreviousChunkIndex != CurrentChunkIndex, LogAudio, Log, TEXT("Changed current chunk '%s' from %d to %d, Offset %d"), *StreamingSoundWave->GetFName().ToString(), PreviousChunkIndex, CurrentChunkIndex, SrcBufferOffset);
 				}
 				else
 				{
@@ -319,8 +328,22 @@
 	return bLooped;
 }
 
-int32 IStreamedCompressedInfo::DecompressToPCMBuffer(uint16 FrameSize)
+int32 IStreamedCompressedInfo::DecompressToPCMBuffer(uint16 /*Unused*/ )
 {
+	// At the end of buffer? 
+	if (SrcBufferOffset >= SrcBufferDataSize - SrcBufferPadding)
+	{
+		// Important we say that nothing was decoded and not an error.
+		return 0;
+	}
+		
+	// Ask for the frame size only after checking if there's space in the buffer (as the decoders will flag errors if there nothing to read).
+	int32 FrameSize = GetFrameSize();
+	if (FrameSize <= 0)
+	{
+		// Error.
+		return -1;
+	}
 	if (SrcBufferOffset + FrameSize > SrcBufferDataSize)
 	{
 		// if frame size is too large, something has gone wrong
