Added uobject stat scope under STAT_UpdateShake for easy asset identification.

diff --git a/Engine/Source/Runtime/Engine/Private/Camera/CameraShakeBase.cpp b/Engine/Source/Runtime/Engine/Private/Camera/CameraShakeBase.cpp
--- a/Engine/Source/Runtime/Engine/Private/Camera/CameraShakeBase.cpp
+++ b/Engine/Source/Runtime/Engine/Private/Camera/CameraShakeBase.cpp
@@ -226,6 +226,7 @@
 void UCameraShakeBase::UpdateAndApplyCameraShake(float DeltaTime, float Alpha, FMinimalViewInfo& InOutPOV)
 {
 	SCOPE_CYCLE_COUNTER(STAT_UpdateShake);
+	SCOPE_CYCLE_UOBJECT(This, this);
 
 	checkf(State.IsActive(), TEXT("Updating a camera shake that wasn't started with a call to StartShake!"));
 
@@ -267,6 +268,7 @@
 	// on the state manager and root pattern instead of the update method.
 	
 	SCOPE_CYCLE_COUNTER(STAT_UpdateShake);
+	SCOPE_CYCLE_UOBJECT(This, this);
 
 	checkf(State.IsActive(), TEXT("Updating a camera shake that wasn't started with a call to StartShake!"));
 
