World Partition: keep dirty actors loaded in the editor and pin them when saving to avoid unloading on save.

#rnx

diff --git a/Engine/Source/Runtime/Engine/Private/WorldPartition/WorldPartition.cpp b/Engine/Source/Runtime/Engine/Private/WorldPartition/WorldPartition.cpp
--- a/Engine/Source/Runtime/Engine/Private/WorldPartition/WorldPartition.cpp
+++ b/Engine/Source/Runtime/Engine/Private/WorldPartition/WorldPartition.cpp
@@ -221,12 +221,21 @@ public:
 		Load();
 	}
 
-	void AddActor(FWorldPartitionHandle& ActorHandle)
+	void AddActor(FWorldPartitionHandle& ActorHandle, bool bImmediate)
 	{
 		if (!ContainsActor(ActorHandle))
 		{
 			Actors.Add(ActorHandle);
-			bStateChanged = true;
+
+			if (bImmediate)
+			{
+				RefreshLoadedState();
+				bStateChanged = false;
+			}
+			else
+			{
+				bStateChanged = true;
+			}
 		}
 	}
 
@@ -235,12 +244,8 @@ public:
 		if (ContainsActor(ActorHandle))
 		{
 			ActorToRemove = ActorHandle;
-
-			if (RefreshLoadedState())
-			{
-				Actors.Remove(ActorHandle);
-			}
-
+			RefreshLoadedState();
+			Actors.Remove(ActorHandle);
 			ActorToRemove.Reset();
 		}
 	}
@@ -339,6 +344,25 @@ void UWorldPartition::OnGCPostReachabilityAnalysis()
 	}
 }
 
+void UWorldPartition::OnPackageDirtyStateChanged(UPackage* Package)
+{
+	if (AActor* Actor = AActor::FindActorInPackage(Package); ShouldHandleActorEvent(Actor))
+	{
+		if (FWorldPartitionHandle ActorHandle(this, Actor->GetActorGuid()); ActorHandle.IsValid())
+		{
+			if (Package->IsDirty())
+			{
+				DirtyActors.Add(ActorHandle, Actor);
+			}
+			else if (PinnedActors)
+			{
+				PinnedActors->AddActor(ActorHandle, /*bImmediate*/true);
+				DirtyActors.Remove(ActorHandle);
+			}
+		}
+	}
+}
+
 void UWorldPartition::OnPreBeginPIE(bool bStartSimulate)
 {
 	check(!bIsPIE);
@@ -651,6 +675,8 @@ void UWorldPartition::Uninitialize()
 
 		WorldDataLayersActor = FWorldPartitionReference();
 
+		DirtyActors.Empty();
+
 		EditorHash = nullptr;
 #endif		
 
@@ -713,6 +739,7 @@ void UWorldPartition::RegisterDelegates()
 			FCoreUObjectDelegates::PostReachabilityAnalysis.AddUObject(this, &UWorldPartition::OnGCPostReachabilityAnalysis);
 			GEditor->OnPostBugItGoCalled().AddUObject(this, &UWorldPartition::OnPostBugItGoCalled);
 			GEditor->OnEditorClose().AddUObject(this, &UWorldPartition::SavePerUserSettings);
+			UPackage::PackageDirtyStateChangedEvent.AddUObject(this, &UWorldPartition::OnPackageDirtyStateChanged);
 		}
 	}
 #endif
@@ -755,6 +782,7 @@ void UWorldPartition::UnregisterDelegates()
 
 			GEditor->OnPostBugItGoCalled().RemoveAll(this);
 			GEditor->OnEditorClose().RemoveAll(this);
+			UPackage::PackageDirtyStateChangedEvent.RemoveAll(this);
 		}
 	}
 #endif
@@ -1079,6 +1107,29 @@ bool UWorldPartition::ResolveSubobject(const TCHAR* SubObjectPath, UObject*& Out
 	return false;
 }
 
+void UWorldPartition::AddReferencedObjects(UObject* InThis, FReferenceCollector& Collector)
+{
+#if WITH_EDITOR
+	// We need to keep all dirty actors alive, mainly for deleted actors. Normally, these actors are only referenced
+	// by the transaction buffer, but we clear it when unloading regions, etc. and we don't want these actors to die.
+	// Also, we must avoid reporting these references when not collecting garbage, as code such as package deletion
+	// will skip packages with actors still referenced (via GatherObjectReferencersForDeletion).
+	if (IsGarbageCollecting())
+	{
+		UWorldPartition* This = CastChecked<UWorldPartition>(InThis);
+	
+		Collector.AllowEliminatingReferences(false);
+		for (auto [ActorReference, Actor] : This->DirtyActors)
+		{
+			Collector.AddReferencedObject(Actor);
+		}
+		Collector.AllowEliminatingReferences(true);
+	}
+#endif
+
+	Super::AddReferencedObjects(InThis, Collector);
+}
+
 void UWorldPartition::Tick(float DeltaSeconds)
 {
 #if WITH_EDITOR
@@ -1319,7 +1370,7 @@ AActor* UWorldPartition::PinActor(const FGuid& ActorGuid)
 		FWorldPartitionHandle ActorHandle(this, ActorGuid);
 		if (ActorHandle.IsValid())
 		{
-			PinnedActors->AddActor(ActorHandle);
+			PinnedActors->AddActor(ActorHandle, /*bImmediate*/false);
 			return ActorHandle->GetActor();
 		}		
 	}
diff --git a/Engine/Source/Runtime/Engine/Private/WorldPartition/WorldPartitionActorLoaderInterface.cpp b/Engine/Source/Runtime/Engine/Private/WorldPartition/WorldPartitionActorLoaderInterface.cpp
--- a/Engine/Source/Runtime/Engine/Private/WorldPartition/WorldPartitionActorLoaderInterface.cpp
+++ b/Engine/Source/Runtime/Engine/Private/WorldPartition/WorldPartitionActorLoaderInterface.cpp
@@ -59,11 +59,6 @@ bool IWorldPartitionActorLoaderInterface::ILoaderAdapter::Unload()
 				ActorsToUnload.Emplace(WorldPartition, ActorRefToUnload.Key);
 			}
 			
-			if (!AllowUnloadingActors(ActorsToUnload))
-			{
-				return false;
-			}
-			
 			UnregisterDelegates();
 
 			FScopedSlowTask SlowTask(1, LOCTEXT("UpdatingLoading", "Updating loading..."));
@@ -143,7 +138,7 @@ bool IWorldPartitionActorLoaderInterface::ILoaderAdapter::PassDataLayersFilter(c
 	return true;
 }
 
-bool IWorldPartitionActorLoaderInterface::ILoaderAdapter::RefreshLoadedState()
+void IWorldPartitionActorLoaderInterface::ILoaderAdapter::RefreshLoadedState()
 {
 	if (bLoaded)
 	{
@@ -166,11 +161,7 @@ bool IWorldPartitionActorLoaderInterface::ILoaderAdapter::RefreshLoadedState()
 				}
 			});
 
-			if (!AllowUnloadingActors(ActorsToUnload))
-			{
-				return false;
-			}
-
+			if (ActorsToLoad.Num() || ActorsToUnload.Num())
 			{
 				FWorldPartitionLoadingContext::FDeferred LoadingContext;
 
@@ -203,63 +194,6 @@ bool IWorldPartitionActorLoaderInterface::ILoaderAdapter::RefreshLoadedState()
 			PostLoadedStateChanged(ActorsToUnload.Num() > 0);
 		}
 	}
-
-	return true;
-}
-
-bool IWorldPartitionActorLoaderInterface::ILoaderAdapter::AllowUnloadingActors(const TArray<FWorldPartitionHandle>& ActorsToUnload) const
-{
-	if (IsRunningCommandlet())
-	{
-		return true;
-	}
-
-	TSet<UPackage*> ModifiedPackages;
-	TMap<FWorldPartitionHandle, int32> UnloadCount;
-
-	// Count the number of times we are referencing each actor to unload
-	for (const FWorldPartitionHandle& ActorToUnload : ActorsToUnload)
-	{
-		const TMap<FGuid, FWorldPartitionReference>& ActorToUnloadRefs = ActorReferences.FindChecked(ActorToUnload->GetGuid());
-
-		for (const TPair<FGuid, FWorldPartitionReference>& ActorToUnloadRef : ActorToUnloadRefs)
-		{
-			UnloadCount.FindOrAdd(ActorToUnloadRef.Value, 0)++;
-		}
-	}
-
-	for (const TPair<FWorldPartitionHandle, int32>& Pair : UnloadCount)
-	{
-		const FWorldPartitionHandle& ActorHandle = Pair.Key;
-
-		// Only prompt if the actor will get unloaded by the unloading cells
-		if (ActorHandle->GetHardRefCount() == Pair.Value)
-		{
-			if (AActor* LoadedActor = ActorHandle->GetActor())
-			{
-				if (UPackage* ActorPackage = LoadedActor->GetExternalPackage(); ActorPackage && ActorPackage->IsDirty())
-				{
-					ModifiedPackages.Add(ActorPackage);
-				}
-			}
-		}
-	}
-
-	if (ModifiedPackages.Num())
-	{
-		const FText Title = LOCTEXT("SaveActorsTitle", "Save Actor(s)");
-		const FText Message = LOCTEXT("SaveActorsMessage", "Save Actor(s) before unloading them.");
-
-		FEditorFileUtils::EPromptReturnCode RetCode = FEditorFileUtils::PromptForCheckoutAndSave(ModifiedPackages.Array(), false, true, Title, Message, nullptr, false, true);
-		check(RetCode != FEditorFileUtils::PR_Failure);
-
-		if (RetCode == FEditorFileUtils::PR_Cancelled)
-		{
-			return false;
-		}
-	}
-
-	return true;
 }
 
 bool IWorldPartitionActorLoaderInterface::ILoaderAdapter::ShouldActorBeLoaded(const FWorldPartitionHandle& Actor) const
@@ -319,13 +253,7 @@ void IWorldPartitionActorLoaderInterface::ILoaderAdapter::RemoveReferenceToActor
 
 void IWorldPartitionActorLoaderInterface::ILoaderAdapter::OnActorDataLayersEditorLoadingStateChanged(bool bFromUserOperation)
 {
-	if (!RefreshLoadedState())
-	{
-		// Temporary measure: reset transactions if the user cancels saving to avoid inconsistencies. The next step is to change that
-		// behavior in order to keep dirty actors in memory instead of asking to save them when unloading, in preparation of in editor
-		// camera loading.
-		GEditor->ResetTransaction(LOCTEXT("LoadingEditorActorResetTrans", "Editor Actors Loading State Changed"));
-	}
+	RefreshLoadedState();
 }
 #endif
 
diff --git a/Engine/Source/Runtime/Engine/Public/WorldPartition/WorldPartition.h b/Engine/Source/Runtime/Engine/Public/WorldPartition/WorldPartition.h
--- a/Engine/Source/Runtime/Engine/Public/WorldPartition/WorldPartition.h
+++ b/Engine/Source/Runtime/Engine/Public/WorldPartition/WorldPartition.h
@@ -113,6 +113,7 @@ private:
 	void SavePerUserSettings();
 		
 	void OnGCPostReachabilityAnalysis();
+	void OnPackageDirtyStateChanged(UPackage* Package);
 
 	// PIE/Game Methods
 	void OnPreBeginPIE(bool bStartSimulate);
@@ -138,6 +139,7 @@ public:
 
 	//~ Begin UObject Interface
 	virtual bool ResolveSubobject(const TCHAR* SubObjectPath, UObject*& OutObject, bool bLoadIfExists) override;
+	static void AddReferencedObjects(UObject* InThis, FReferenceCollector& Collector);
 	//~ End UObject Interface
 
 #if WITH_EDITOR
@@ -254,6 +256,8 @@ private:
 	TObjectPtr<class UHLODLayer> DefaultHLODLayer;
 
 	TArray<FWorldPartitionReference> LoadedSubobjects;
+
+	TMap<FWorldPartitionReference, AActor*> DirtyActors;
 #endif
 
 	EWorldPartitionInitState InitState;
diff --git a/Engine/Source/Runtime/Engine/Public/WorldPartition/WorldPartitionActorLoaderInterface.h b/Engine/Source/Runtime/Engine/Public/WorldPartition/WorldPartitionActorLoaderInterface.h
--- a/Engine/Source/Runtime/Engine/Public/WorldPartition/WorldPartitionActorLoaderInterface.h
+++ b/Engine/Source/Runtime/Engine/Public/WorldPartition/WorldPartitionActorLoaderInterface.h
@@ -52,17 +52,9 @@ public:
 		// Actors filtering
 		bool PassActorDescFilter(const FWorldPartitionHandle& Actor) const;
 		bool PassDataLayersFilter(const FWorldPartitionHandle& Actor) const;
-
-		/**
-		 * Refresh the loaded state of this loader adatper.
-		 *
-		 * @return true if the operation was successful, false otherwise (operation was cancelled by the user
-		 *         if actors needed resaving, etc.
-		 */
-		bool RefreshLoadedState();
+		void RefreshLoadedState();
 
 		void PostLoadedStateChanged(bool bUnloadedActors);
-		bool AllowUnloadingActors(const TArray<FWorldPartitionHandle>& ActorsToUnload) const;
 		void AddReferenceToActor(FWorldPartitionHandle& Actor);
 		void RemoveReferenceToActor(FWorldPartitionHandle& Actor);
 		void OnActorDataLayersEditorLoadingStateChanged(bool bFromUserOperation);		
