World Partition: properly handle the case when we're cleaning a level, which means that actors are already marked as garbage, in which case we don't even try to remove them from the level.

#rnx

diff --git a/Engine/Source/Runtime/Engine/Private/WorldPartition/WorldPartitionHandle.cpp b/Engine/Source/Runtime/Engine/Private/WorldPartition/WorldPartitionHandle.cpp
--- a/Engine/Source/Runtime/Engine/Private/WorldPartition/WorldPartitionHandle.cpp
+++ b/Engine/Source/Runtime/Engine/Private/WorldPartition/WorldPartitionHandle.cpp
@@ -80,7 +80,8 @@
 
 void FWorldPartitionLoadingContext::FImmediate::UnregisterActor(FWorldPartitionActorDesc* ActorDesc)
 {
-	if (AActor* Actor = ActorDesc->GetActor())
+	// When cleaning up worlds, actors are already marked as garbage at this point, so no need to remove them from the world
+	if (AActor* Actor = ActorDesc->GetActor(); IsValid(Actor))
 	{
 		UActorDescContainer* Container = ActorDesc->GetContainer();
 		check(Container);
@@ -124,9 +125,8 @@
 				}
 			}
 			
-			if (ActorList.Num())
+			if (Level)
 			{
-				check(Level);
 				Level->AddLoadedActors(ActorList, ContainerTransformPtr);
 			}
 		}
@@ -139,7 +139,8 @@
 			ULevel* Level = nullptr;
 			for (FWorldPartitionActorDesc* ActorDesc : ContainerOp.Unregistrations)
 			{
-				if (AActor* Actor = ActorDesc->GetActor())
+				// When cleaning up worlds, actors are already marked as garbage at this point, so no need to remove them from the world
+				if (AActor* Actor = ActorDesc->GetActor(); IsValid(Actor))
 				{
 					ActorList.Add(Actor);
 
@@ -149,9 +150,8 @@
 				}
 			}
 			
-			if (ActorList.Num())
+			if (Level)
 			{
-				check(Level);
 				Level->RemoveLoadedActors(ActorList, ContainerTransformPtr);
 			}
 
