Rewrites dynamic resolution on top of the DynamicRenderScaling API

diff --git a/Engine/Plugins/Experimental/ColorCorrectRegions/Source/ColorCorrectRegions/Private/ColorCorrectRegionsSceneViewExtension.cpp b/Engine/Plugins/Experimental/ColorCorrectRegions/Source/ColorCorrectRegions/Private/ColorCorrectRegionsSceneViewExtension.cpp
--- a/Engine/Plugins/Experimental/ColorCorrectRegions/Source/ColorCorrectRegions/Private/ColorCorrectRegionsSceneViewExtension.cpp
+++ b/Engine/Plugins/Experimental/ColorCorrectRegions/Source/ColorCorrectRegions/Private/ColorCorrectRegionsSceneViewExtension.cpp
@@ -13,6 +13,7 @@
 #include "Engine/World.h"
 #include "EngineUtils.h"
 #include "SceneRendering.h"
+#include "DynamicResolutionState.h"
 // Set this to 1 to clip pixels outside of bounding box.
 #define CLIP_PIXELS_OUTSIDE_AABB 1
 
@@ -306,8 +307,10 @@
 
 	const FSceneViewFamily& ViewFamily = *View.Family;
 
+	DynamicRenderScaling::TMap<float> UpperBounds = ViewFamily.GetScreenPercentageInterface()->GetResolutionFractionsUpperBound();
+
 	const auto FeatureLevel = View.GetFeatureLevel();
-	const float ScreenPercentage = ViewFamily.GetPrimaryResolutionFractionUpperBound() * ViewFamily.SecondaryViewFraction;
+	const float ScreenPercentage = UpperBounds[GDynamicPrimaryResolutionFraction] * ViewFamily.SecondaryViewFraction;
 	
 	// We need to make sure to take Windows and Scene scale into account.
 
diff --git a/Engine/Plugins/Media/ImgMedia/Source/ImgMedia/Private/ImgMediaSceneViewExtension.cpp b/Engine/Plugins/Media/ImgMedia/Source/ImgMedia/Private/ImgMediaSceneViewExtension.cpp
--- a/Engine/Plugins/Media/ImgMedia/Source/ImgMedia/Private/ImgMediaSceneViewExtension.cpp
+++ b/Engine/Plugins/Media/ImgMedia/Source/ImgMedia/Private/ImgMediaSceneViewExtension.cpp
@@ -1,6 +1,7 @@
 // Copyright Epic Games, Inc. All Rights Reserved.
 
 #include "ImgMediaSceneViewExtension.h"
+#include "DynamicResolutionState.h"
 
 #include "SceneView.h"
 
@@ -32,7 +33,8 @@
 
 	if (InViewFamily.GetScreenPercentageInterface())
 	{
-		ResolutionFraction *= InViewFamily.GetPrimaryResolutionFractionUpperBound();
+		DynamicRenderScaling::TMap<float> UpperBounds = InViewFamily.GetScreenPercentageInterface()->GetResolutionFractionsUpperBound();
+		ResolutionFraction *= UpperBounds[GDynamicPrimaryResolutionFraction];
 	}
 
 	static const auto CVarMinAutomaticViewMipBiasOffset = IConsoleManager::Get().FindTConsoleVariableDataFloat(TEXT("r.ViewTextureMipBias.Offset"));
diff --git a/Engine/Plugins/Runtime/Oculus/OculusVR/Source/OculusHMD/Private/OculusHMD_DynamicResolutionState.cpp b/Engine/Plugins/Runtime/Oculus/OculusVR/Source/OculusHMD/Private/OculusHMD_DynamicResolutionState.cpp
--- a/Engine/Plugins/Runtime/Oculus/OculusVR/Source/OculusHMD/Private/OculusHMD_DynamicResolutionState.cpp
+++ b/Engine/Plugins/Runtime/Oculus/OculusVR/Source/OculusHMD/Private/OculusHMD_DynamicResolutionState.cpp
@@ -68,14 +68,20 @@
 	}
 }
 
-float FDynamicResolutionState::GetResolutionFractionApproximation() const
+DynamicRenderScaling::TMap<float> FDynamicResolutionState::GetResolutionFractionsApproximation() const
 {
-	return ResolutionFraction;
+	DynamicRenderScaling::TMap<float> ResolutionFractions;
+	ResolutionFractions.SetAll(1.0f);
+	ResolutionFractions[GDynamicPrimaryResolutionFraction] = ResolutionFraction;
+	return ResolutionFractions;
 }
 
-float FDynamicResolutionState::GetResolutionFractionUpperBound() const
+DynamicRenderScaling::TMap<float> FDynamicResolutionState::GetResolutionFractionsUpperBound() const
 {
-	return ResolutionFractionUpperBound;
+	DynamicRenderScaling::TMap<float> ResolutionFractions;
+	ResolutionFractions.SetAll(1.0f);
+	ResolutionFractions[GDynamicPrimaryResolutionFraction] = ResolutionFractionUpperBound;
+	return ResolutionFractions;
 }
 
 void FDynamicResolutionState::SetEnabled(bool bEnable)
diff --git a/Engine/Plugins/Runtime/Oculus/OculusVR/Source/OculusHMD/Private/OculusHMD_DynamicResolutionState.h b/Engine/Plugins/Runtime/Oculus/OculusVR/Source/OculusHMD/Private/OculusHMD_DynamicResolutionState.h
--- a/Engine/Plugins/Runtime/Oculus/OculusVR/Source/OculusHMD/Private/OculusHMD_DynamicResolutionState.h
+++ b/Engine/Plugins/Runtime/Oculus/OculusVR/Source/OculusHMD/Private/OculusHMD_DynamicResolutionState.h
@@ -25,8 +25,8 @@
 	virtual void SetupMainViewFamily(class FSceneViewFamily& ViewFamily) override;
 
 protected:
-	virtual float GetResolutionFractionApproximation() const override;
-	virtual float GetResolutionFractionUpperBound() const override;
+	virtual DynamicRenderScaling::TMap<float> GetResolutionFractionsApproximation() const override;
+	virtual DynamicRenderScaling::TMap<float> GetResolutionFractionsUpperBound() const override;
 	virtual void SetEnabled(bool bEnable) override;
 	virtual bool IsEnabled() const override;
 	virtual void ProcessEvent(EDynamicResolutionStateEvent Event) override;
diff --git a/Engine/Plugins/Runtime/nDisplay/Source/DisplayCluster/Private/Game/EngineClasses/Basics/DisplayClusterViewportClient.cpp b/Engine/Plugins/Runtime/nDisplay/Source/DisplayCluster/Private/Game/EngineClasses/Basics/DisplayClusterViewportClient.cpp
--- a/Engine/Plugins/Runtime/nDisplay/Source/DisplayCluster/Private/Game/EngineClasses/Basics/DisplayClusterViewportClient.cpp
+++ b/Engine/Plugins/Runtime/nDisplay/Source/DisplayCluster/Private/Game/EngineClasses/Basics/DisplayClusterViewportClient.cpp
@@ -720,14 +720,14 @@
 						GEngine->EmitDynamicResolutionEvent(EDynamicResolutionStateEvent::BeginDynamicResolutionRendering);
 						ViewFamily.SetScreenPercentageInterface(new FLegacyScreenPercentageDriver(
 							ViewFamily,
-							DynamicResolutionStateInfos.ResolutionFractionApproximation,
-							DynamicResolutionStateInfos.ResolutionFractionUpperBound));
+							DynamicResolutionStateInfos.ResolutionFractionApproximations[GDynamicPrimaryResolutionFraction],
+							DynamicResolutionStateInfos.ResolutionFractionUpperBounds[GDynamicPrimaryResolutionFraction]));
 					}
 
 	#if CSV_PROFILER
-					if (DynamicResolutionStateInfos.ResolutionFractionApproximation >= 0.0f)
+					if (DynamicResolutionStateInfos.ResolutionFractionApproximations[GDynamicPrimaryResolutionFraction] >= 0.0f)
 					{
-						CSV_CUSTOM_STAT_GLOBAL(DynamicResolutionPercentage, DynamicResolutionStateInfos.ResolutionFractionApproximation * 100.0f, ECsvCustomStatOp::Set);
+						CSV_CUSTOM_STAT_GLOBAL(DynamicResolutionPercentage, DynamicResolutionStateInfos.ResolutionFractionApproximations[GDynamicPrimaryResolutionFraction] * 100.0f, ECsvCustomStatOp::Set);
 					}
 	#endif
 				}
diff --git a/Engine/Source/Runtime/Engine/Classes/Engine/Engine.h b/Engine/Source/Runtime/Engine/Classes/Engine/Engine.h
--- a/Engine/Source/Runtime/Engine/Classes/Engine/Engine.h
+++ b/Engine/Source/Runtime/Engine/Classes/Engine/Engine.h
@@ -20,6 +20,7 @@
 #include "AudioDeviceManager.h"
 #include "Templates/UniqueObj.h"
 #include "Containers/Ticker.h"
+#include "DynamicRenderScaling.h"
 #include "Engine.generated.h"
 
 #define WITH_DYNAMIC_RESOLUTION (!UE_SERVER)
@@ -132,10 +133,10 @@
 
 	// Approximation of the resolution fraction being applied. This is only an approximation because
 	// of non (and unecessary) thread safety of this value between game thread, and render thread.
-	float ResolutionFractionApproximation;
+	DynamicRenderScaling::TMap<float> ResolutionFractionApproximations;
 
 	// Maximum resolution fraction set, always MaxResolutionFraction >= ResolutionFractionApproximation.
-	float ResolutionFractionUpperBound;
+	DynamicRenderScaling::TMap<float> ResolutionFractionUpperBounds;
 };
 
 
diff --git a/Engine/Source/Runtime/Engine/Private/ChartCreation.cpp b/Engine/Source/Runtime/Engine/Private/ChartCreation.cpp
--- a/Engine/Source/Runtime/Engine/Private/ChartCreation.cpp
+++ b/Engine/Source/Runtime/Engine/Private/ChartCreation.cpp
@@ -21,6 +21,7 @@
 #include "HAL/LowLevelMemTracker.h"
 #include "ProfilingDebugging/CsvProfiler.h"
 #include "DeviceProfiles/DeviceProfileManager.h"
+#include "DynamicResolutionState.h"
 
 #ifndef FPS_CHART_SUPPORT_CSV_PROFILE
 #define FPS_CHART_SUPPORT_CSV_PROFILE (CSV_PROFILER && !UE_BUILD_SHIPPING)
@@ -1361,7 +1362,7 @@
 
 		if (DynamicResolutionStateInfos.Status == EDynamicResolutionStatus::Enabled || DynamicResolutionStateInfos.Status == EDynamicResolutionStatus::DebugForceEnabled)
 		{
-			FrameData.DynamicResolutionScreenPercentage = DynamicResolutionStateInfos.ResolutionFractionApproximation * 100.0f;
+			FrameData.DynamicResolutionScreenPercentage = DynamicResolutionStateInfos.ResolutionFractionApproximations[GDynamicPrimaryResolutionFraction] * 100.0f;
 		}
 	}
 
diff --git a/Engine/Source/Runtime/Engine/Private/DynamicResolution.cpp b/Engine/Source/Runtime/Engine/Private/DynamicResolution.cpp
--- a/Engine/Source/Runtime/Engine/Private/DynamicResolution.cpp
+++ b/Engine/Source/Runtime/Engine/Private/DynamicResolution.cpp
@@ -1,9 +1,5 @@
 // Copyright Epic Games, Inc. All Rights Reserved.
 
-/*=============================================================================
-Renderer.cpp: Renderer module implementation.
-=============================================================================*/
-
 #include "DynamicResolutionProxy.h"
 #include "DynamicResolutionState.h"
 
@@ -15,18 +11,19 @@
 #include "RHICommandList.h"
 #include "DynamicRHI.h"
 #include "UnrealEngine.h"
+#include "RenderGraphEvent.h"
 
 
 static TAutoConsoleVariable<float> CVarDynamicResMinSP(
 	TEXT("r.DynamicRes.MinScreenPercentage"),
-	50,
-	TEXT("Minimal screen percentage."),
+	DynamicRenderScaling::FractionToPercentage(DynamicRenderScaling::FHeuristicSettings::kDefaultMinResolutionFraction),
+	TEXT("Minimal primary screen percentage."),
 	ECVF_RenderThreadSafe | ECVF_Default);
 
 static TAutoConsoleVariable<float> CVarDynamicResMaxSP(
 	TEXT("r.DynamicRes.MaxScreenPercentage"),
-	100,
-	TEXT("Maximal screen percentage."),
+	DynamicRenderScaling::FractionToPercentage(DynamicRenderScaling::FHeuristicSettings::kDefaultMaxResolutionFraction),
+	TEXT("Maximal primary screen percentage."),
 	ECVF_Default);
 
 // TODO: Seriously need a centralized engine perf manager.
@@ -44,25 +41,12 @@
 	TEXT("Targeted GPU headroom (in percent from r.DynamicRes.FrameTimeBudget)."),
 	ECVF_RenderThreadSafe | ECVF_Default);
 
-static TAutoConsoleVariable<float> CVarCPUTimeHeadRoom(
-	TEXT("r.DynamicRes.CPUTimeHeadRoom"),
-	1,
-	TEXT("Head room for the threads compared GPU time to avoid keep getting resolution fraction ")
-	TEXT("shrinking down when CPU bound (in milliseconds)."),
-	ECVF_RenderThreadSafe | ECVF_Default);
-
 static TAutoConsoleVariable<int32> CVarHistorySize(
 	TEXT("r.DynamicRes.HistorySize"),
 	16,
 	TEXT("Number of frames keept in the history."),
 	ECVF_RenderThreadSafe | ECVF_Default);
 
-static TAutoConsoleVariable<float> CVarOutlierThreshold(
-	TEXT("r.DynamicRes.OutlierThreshold"),
-	0.0f,
-	TEXT("Ignore frame timing that have Game thread or render thread X time more than frame budget."),
-	ECVF_RenderThreadSafe | ECVF_Default);
-
 static TAutoConsoleVariable<float> CVarFrameWeightExponent(
 	TEXT("r.DynamicRes.FrameWeightExponent"),
 	0.9f,
@@ -78,13 +62,13 @@
 
 static TAutoConsoleVariable<float> CVarIncreaseAmortizationFactor(
 	TEXT("r.DynamicRes.IncreaseAmortizationBlendFactor"),
-	0.9f,
+	DynamicRenderScaling::FHeuristicSettings::kDefaultIncreaseAmortizationFactor,
 	TEXT("Amortization blend factor when scale resolution back up to reduce resolution fraction oscillations."),
 	ECVF_RenderThreadSafe | ECVF_Default);
 
 static TAutoConsoleVariable<float> CVarChangeThreshold(
 	TEXT("r.DynamicRes.ChangePercentageThreshold"),
-	2.0f,
+	DynamicRenderScaling::FractionToPercentage(DynamicRenderScaling::FHeuristicSettings::kDefaultChangeThreshold),
 	TEXT("Minimal increase percentage threshold to alow when changing resolution."),
 	ECVF_RenderThreadSafe | ECVF_Default);
 
@@ -101,13 +85,31 @@
 	TEXT(" 0: Same as stat unit (default);\n 1: Timestamp queries."),
 	ECVF_RenderThreadSafe | ECVF_Default);
 
-static TAutoConsoleVariable<float> CVarCPUBoundScreenPercentage(
-	TEXT("r.DynamicRes.CPUBoundScreenPercentage"),
-	100,
-	TEXT("Screen percentage to converge to when CPU bound. This can be used when GPU and CPU share same memory."),
-	ECVF_RenderThreadSafe | ECVF_Default);
+
+DynamicRenderScaling::FHeuristicSettings GetPrimaryDynamicResolutionSettings()
+{
+	DynamicRenderScaling::FHeuristicSettings BudgetSetting;
+	BudgetSetting.Model = DynamicRenderScaling::EHeuristicModel::Quadratic;
+	BudgetSetting.MinResolutionFraction      = DynamicRenderScaling::GetPercentageCVarToFraction(CVarDynamicResMinSP);
+	BudgetSetting.MaxResolutionFraction      = DynamicRenderScaling::GetPercentageCVarToFraction(CVarDynamicResMaxSP);
+	// BudgetSetting.BudgetMs depends on the cost of other buckets.
+	BudgetSetting.ChangeThreshold            = DynamicRenderScaling::GetPercentageCVarToFraction(CVarChangeThreshold);
+	BudgetSetting.IncreaseAmortizationFactor = CVarIncreaseAmortizationFactor.GetValueOnRenderThread();
+
+	// CVarTargetedGPUHeadRoomPercentage is taken into account on the entire frame instead.
+	BudgetSetting.TargetedHeadRoom = 0.0f;
+	return BudgetSetting;
+}
+
+DynamicRenderScaling::FBudget GDynamicPrimaryResolutionFraction(TEXT("DynamicPrimaryResolution"), &GetPrimaryDynamicResolutionSettings);
 
 
+static float TimeStampQueryResultToMiliSeconds(uint64 TimestampResult)
+{
+	return float(TimestampResult) / 1000.0f;
+}
+
+
 FDynamicResolutionHeuristicProxy::FDynamicResolutionHeuristicProxy()
 {
 	check(IsInGameThread());
@@ -133,10 +135,11 @@
 {
 	PreviousFrameIndex = -1;
 	HistorySize = 0;
+	BudgetHistorySizes.SetAll(0);
 	History.Reset();
 
 	NumberOfFramesSinceScreenPercentageChange = 0;
-	CurrentFrameResolutionFraction = 1.0f;
+	CurrentFrameResolutionFractions.SetAll(1.0f);
 
 	// Ignore previous frame timings.
 	IgnoreFrameRemainingCount = 1;
@@ -160,18 +163,27 @@
 	{
 		int32 NewHistoryEntryIndex = (PreviousFrameIndex + 1) % History.Num();
 		History[NewHistoryEntryIndex] = FrameHistoryEntry();
-		History[NewHistoryEntryIndex].ResolutionFraction = CurrentFrameResolutionFraction;
+		History[NewHistoryEntryIndex].ResolutionFractions = CurrentFrameResolutionFractions;
 		History[NewHistoryEntryIndex].GameThreadTimeMs = GameThreadTimeMs;
 		History[NewHistoryEntryIndex].RenderThreadTimeMs = RenderThreadTimeMs;
 		PreviousFrameIndex = NewHistoryEntryIndex;
 		HistorySize = FMath::Min(HistorySize + 1, History.Num());
+
+		for (TLinkedList<DynamicRenderScaling::FBudget*>::TIterator BudgetIt(DynamicRenderScaling::FBudget::GetGlobalList()); BudgetIt; BudgetIt.Next())
+		{
+			const DynamicRenderScaling::FBudget& Budget = **BudgetIt;
+			BudgetHistorySizes[Budget] = FMath::Min(BudgetHistorySizes[Budget] + 1, History.Num());
+		}
 	}
 
 	return ++FrameCounter;
 }
 
 void FDynamicResolutionHeuristicProxy::CommitPreviousFrameGPUTimings_RenderThread(
-	uint64 HistoryFrameId, float TotalFrameGPUBusyTimeMs, float DynamicResolutionGPUBusyTimeMs, bool bGPUTimingsHaveCPUBubbles)
+	uint64 HistoryFrameId,
+	float TotalFrameGPUBusyTimeMs,
+	float DynamicResolutionGPUBusyTimeMs,
+	const DynamicRenderScaling::TMap<float>& BudgetTimingMs)
 {
 	check(TotalFrameGPUBusyTimeMs >= 0.0f);
 	check(DynamicResolutionGPUBusyTimeMs >= 0.0f);
@@ -190,237 +202,336 @@
 
 	History[EntryId].TotalFrameGPUBusyTimeMs = TotalFrameGPUBusyTimeMs;
 	History[EntryId].GlobalDynamicResolutionTimeMs = DynamicResolutionGPUBusyTimeMs;
-	History[EntryId].bGPUTimingsHaveCPUBubbles = bGPUTimingsHaveCPUBubbles;
+	History[EntryId].BudgetTimingMs = BudgetTimingMs;
 }
 
 void FDynamicResolutionHeuristicProxy::RefreshCurentFrameResolutionFraction_RenderThread()
 {
+	// Compute new frame resolution fractions only if have an history to work with.
+	if (HistorySize == 0)
+	{
+		return;
+	}
+
 	// GPU time budget per frame=.
 	const float FrameTimeBudgetMs = CVarFrameTimeBudget.GetValueOnRenderThread();
 
 	// Targeted GPU time, lower than budget to diggest noise.
 	const float TargetedGPUBusyTimeMs = FrameTimeBudgetMs * (1.0f - CVarTargetedGPUHeadRoomPercentage.GetValueOnRenderThread() / 100.0f);
 
-	// Resolution fraction to downscale to when CPU bound.
-	const float CPUBoundResolutionFraction = CVarCPUBoundScreenPercentage.GetValueOnRenderThread() / 100.0f;
+	const float FrameWeightExponent = CVarFrameWeightExponent.GetValueOnRenderThread();
+	const int32 MaxConsecutiveOverbudgetGPUFrameCount = FMath::Max(CVarMaxConsecutiveOverbudgetGPUFrameCount.GetValueOnRenderThread(), 2);
+	
+	// Find the first frame that have data.
+	int32 StartFrameId = 0;
+	{
+		for (; StartFrameId < HistorySize; StartFrameId++)
+		{
+			const FrameHistoryEntry& FrameEntry = GetPreviousFrameEntry(StartFrameId);
 
-	// New ResolutionFraction to use for this frame.
-	float NewFrameResolutionFraction = CurrentFrameResolutionFraction;
-
-	// Whether there is a GPU over budget panic.
-	bool bGPUOverbugetPanic = false;
-
-	// Compute NewFrameResolutionFraction if have an history to work with.
-	if (HistorySize > 0)
-	{
-		const float FrameWeightExponent = CVarFrameWeightExponent.GetValueOnRenderThread();
-		const float RenderThreadHeadRoomFromGPUMs = CVarCPUTimeHeadRoom.GetValueOnRenderThread();
-		const int32 MaxConsecutiveOverbudgetGPUFrameCount = FMath::Max(CVarMaxConsecutiveOverbudgetGPUFrameCount.GetValueOnRenderThread(), 2);
-		
-		const float OutLierTimeThreshold = FrameTimeBudgetMs * CVarOutlierThreshold.GetValueOnRenderThread();
-
-		NewFrameResolutionFraction = 0.0f;
-
-		// Total weight of NewFrameResolutionFraction.
-		float TotalWeight = 0.0f;
-		
-		// Frame weight.
-		float Weight = 1.0f;
-
-		// Number of consecutive frames that have over budget GPU.
-		int32 ConsecutiveOverbudgetGPUFramCount = 0;
-
-		// Whether should continue browsing the frame history.
-		bool bCaryOnBrowsingFrameHistory = true;
-
-		// Number of frame browsed.
-		int32 FrameCount = 0;
-		int32 BrowsingFrameId = 0;
-		for (; BrowsingFrameId < HistorySize && bCaryOnBrowsingFrameHistory; BrowsingFrameId++)
-		{
-			const FrameHistoryEntry& NextFrameEntry = GetPreviousFrameEntry(BrowsingFrameId - 1);
-			const FrameHistoryEntry& FrameEntry = GetPreviousFrameEntry(BrowsingFrameId + 0);
-			const FrameHistoryEntry& PreviousFrameEntry = GetPreviousFrameEntry(BrowsingFrameId + 1);
-
 			// Ignores frames that does not have any GPU timing yet. 
 			if (FrameEntry.TotalFrameGPUBusyTimeMs < 0)
 			{
 				continue;
 			}
+			else
+			{
+				break;
+			}
+		}
+	}
 
-			// Compute the maximum thread time between game and render thread for this <i>th frame.
-			// We add a head room for render thread to detect when the GPU could be starving.
-			const float MaxThreadTimeMs = FMath::Max(
-				FrameEntry.RenderThreadTimeMs + (FrameEntry.bGPUTimingsHaveCPUBubbles ? RenderThreadHeadRoomFromGPUMs : 0),
-				FrameEntry.GameThreadTimeMs);
+	// Not enough data to work with.
+	if (HistorySize - StartFrameId <= MaxConsecutiveOverbudgetGPUFrameCount)
+	{
+		return;
+	}
 
-			// Whether bound by game thread.
-			const bool bIsGameThreadBound = FrameEntry.GameThreadTimeMs > FrameTimeBudgetMs;
+	check(StartFrameId + MaxConsecutiveOverbudgetGPUFrameCount < HistorySize);
 
-			// Whether bound by render thread.
-			const bool bIsRenderThreadBound = (FrameEntry.RenderThreadTimeMs + (FrameEntry.bGPUTimingsHaveCPUBubbles ? RenderThreadHeadRoomFromGPUMs : 0)) > TargetedGPUBusyTimeMs;
+	// Look whether all the last MaxConsecutiveOverbudgetGPUFrameCount frames are over budget.
+	bool bGPUFrameOverbudgetPanic = true;
+	for (int32 BrowsingFrameId = StartFrameId; BrowsingFrameId < StartFrameId + MaxConsecutiveOverbudgetGPUFrameCount; BrowsingFrameId++)
+	{
+		const FrameHistoryEntry& FrameEntry = GetPreviousFrameEntry(BrowsingFrameId);
+		bGPUFrameOverbudgetPanic &= (FrameEntry.TotalFrameGPUBusyTimeMs > FrameTimeBudgetMs);
+	}
 
-			// Whether the frame is CPU bound.
-			const bool bIsCPUBound = bIsGameThreadBound || bIsRenderThreadBound;
-			
-			// Whether the render thread might create bubbles in GPU timings.
-			const bool bMayHaveGPUBubbles = FrameEntry.bGPUTimingsHaveCPUBubbles && bIsRenderThreadBound;
+	// Whether a cruising resolution change can happen.
+	bool bCanChangeResolution = NumberOfFramesSinceScreenPercentageChange >= CVarFrameChangePeriod.GetValueOnRenderThread();
 
-			// Whether GPU is over budget, when not CPU bound.
-			const bool bHasOverbudgetGPU = !bIsCPUBound && FrameEntry.TotalFrameGPUBusyTimeMs > FrameTimeBudgetMs;
+	// Array to know how much a frame is overbudget
+	//TArray<float, TInlineAllocator<256>> TotalBudgetOverBudgetSavings;
+	//TotalBudgetOverBudgetSavings.SetNumZeroed(MaxFrameId);
 
-			// Whether this is possibly an outlier compared to previous and next frame.
-			// Normally should look for out lier compared to previous and next frames timings, however
-			// only interested in detected them if they are over budget, when cruising at targeted
-			// GPU busy time.
-			const bool bIsGPUOutlier =
-				bHasOverbudgetGPU &&
-				ConsecutiveOverbudgetGPUFramCount == 0 &&
-				OutLierTimeThreshold > 0.0f &&
-				FrameEntry.TotalFrameGPUBusyTimeMs > OutLierTimeThreshold &&
-				(!NextFrameEntry.HasGPUTimings() || (
-					NextFrameEntry.TotalFrameGPUBusyTimeMs < OutLierTimeThreshold &&
-					NextFrameEntry.GameThreadTimeMs < OutLierTimeThreshold &&
-					NextFrameEntry.RenderThreadTimeMs < OutLierTimeThreshold)) &&
-				(!PreviousFrameEntry.HasGPUTimings() || (
-					PreviousFrameEntry.TotalFrameGPUBusyTimeMs < OutLierTimeThreshold &&
-					PreviousFrameEntry.GameThreadTimeMs < OutLierTimeThreshold &&
-					PreviousFrameEntry.RenderThreadTimeMs < OutLierTimeThreshold));
+	// Function to detect whether the history size is big enough for detecting an overbudget
+	auto HasNotEnoughDataForOverBudgetDetection = [&](int32 InHistorySize)
+	{
+		return (InHistorySize - StartFrameId) <= MaxConsecutiveOverbudgetGPUFrameCount;
+	};
 
-			// Look if this is multiple consecutive GPU over budget frames.
-			if (bHasOverbudgetGPU)
+	// Function to know whether a bucket is overbudget.
+	auto IsBudgetOverBudget = [&](const DynamicRenderScaling::FBudget& Budget, const DynamicRenderScaling::FHeuristicSettings& BudgetSetting)
+	{
+		check(BudgetSetting.IsEnabled());
+
+		if (HasNotEnoughDataForOverBudgetDetection(BudgetHistorySizes[Budget]))
+		{
+			// Not enough data to work with.
+			return false;
+		}
+
+		bool bBudgetIsOverBudget = true;
+
+		for (int32 BrowsingFrameId = StartFrameId; BrowsingFrameId < StartFrameId + MaxConsecutiveOverbudgetGPUFrameCount; BrowsingFrameId++)
+		{
+			const FrameHistoryEntry& FrameEntry = GetPreviousFrameEntry(BrowsingFrameId);
+
+			float bOverBudget = FrameEntry.BudgetTimingMs[Budget] - BudgetSetting.BudgetMs;
+			bBudgetIsOverBudget &= bOverBudget > 0.0f;
+		}
+
+		return bBudgetIsOverBudget;
+	};
+
+	// Compute
+	auto ComputeMaxFrameId = [&](const DynamicRenderScaling::FBudget& Budget, bool bBudgetIsOverBudget)
+	{
+		int32 MaxFrameId = BudgetHistorySizes[Budget];
+		
+		if (bGPUFrameOverbudgetPanic || bBudgetIsOverBudget)
+		{
+			MaxFrameId = FMath::Min(MaxFrameId, StartFrameId + MaxConsecutiveOverbudgetGPUFrameCount);
+		}
+
+		check(MaxFrameId <= BudgetHistorySizes[Budget]);
+		check(MaxFrameId <= HistorySize);
+		return MaxFrameId;
+	};
+
+	DynamicRenderScaling::TMap<float> NewResolutionFractions;
+	bool bCommitResolutionChange = false;
+
+	// Execute heuristic on the bucket.
+	auto ComputeNewResolutionFraction = [&](const DynamicRenderScaling::FBudget& Budget, const DynamicRenderScaling::FHeuristicSettings& BudgetSetting, bool bBudgetIsOverBudget)
+	{
+		check(BudgetSetting.IsEnabled());
+
+		float NewResolutionFraction = CurrentFrameResolutionFractions[Budget];
+		if (!HasNotEnoughDataForOverBudgetDetection(BudgetHistorySizes[Budget]))
+		{
+			float CurrentResolutionFraction = CurrentFrameResolutionFractions[Budget];
+			if (BudgetSetting.bModelScalesWithPrimaryScreenPercentage)
 			{
-				ConsecutiveOverbudgetGPUFramCount++;
-				check(ConsecutiveOverbudgetGPUFramCount <= MaxConsecutiveOverbudgetGPUFrameCount);
+				CurrentResolutionFraction *= CurrentFrameResolutionFractions[GDynamicPrimaryResolutionFraction];
+			}
 
-				// Max number of over budget frames where reached.
-				if (ConsecutiveOverbudgetGPUFramCount == MaxConsecutiveOverbudgetGPUFrameCount)
+			int32 MaxFrameId = ComputeMaxFrameId(Budget, bBudgetIsOverBudget);
+
+			// Find out a new resolution fraction based on whether this bucket is overbudget.
+			float MostRecentResolutionFractionScale = 1.0f;
+			float SuggestedResolutionFraction = 0.0f;
+			float SuggestedResolutionFractionWeight = 0.0f;
+			float Weight = 1.0f;
+			for (int32 BrowsingFrameId = StartFrameId; BrowsingFrameId < MaxFrameId; BrowsingFrameId++)
+			{
+				const FrameHistoryEntry& FrameEntry = GetPreviousFrameEntry(BrowsingFrameId);
+				float TimingMs = FrameEntry.BudgetTimingMs[Budget];
+
+				float ResolutionFraction = FrameEntry.ResolutionFractions[Budget];
+				if (BudgetSetting.bModelScalesWithPrimaryScreenPercentage)
 				{
-					// We ignore frames in history that happen before consecutive GPU overbudget frames.
-					bCaryOnBrowsingFrameHistory = false;
+					ResolutionFraction *= FrameEntry.ResolutionFractions[GDynamicPrimaryResolutionFraction];
+				}
 
-					ensureMsgf(!bIsGPUOutlier, TEXT("Does not make any sens to have ConsecutiveOverbudgetGPUFramCount >= 2 if this frame is detected as being an outlier."));
+				// Removing the savings from other bucket
+				//if (Budget != EDynamicResolutionTimingBudget::PrimaryScreenPercentage)
+				//{
+				//	TimingMs -= TotalBudgetOverBudgetSavings[BrowsingFrameId];
+				//}
+
+				// Estimate how much a resolution fraction should be changed.
+				float ResolutionFractionFactor = BudgetSetting.EstimateResolutionFactor(TimingMs);
+
+				// Don't increase the resolution if frame overbudget.
+				if (bGPUFrameOverbudgetPanic && ResolutionFractionFactor > 1.0f)
+				{
+					ResolutionFractionFactor = 1.0f;
 				}
+
+				float AdjustedResolutionFraction = ResolutionFractionFactor * ResolutionFraction;
+
+				// Track how much saving can be done in a bucket to max out PrimaryScreenPercentage
+				//if (Budget != EDynamicResolutionTimingBudget::PrimaryScreenPercentage)
+				//{
+				//	float EstimatedSavings = (
+				//		FrameEntry.BudgetTimingMs[Budget] *
+				//		(1.0f - BudgetSetting.EstimateTimeFactor(FrameEntry.ResolutionFractions[Budget], AdjustedResolutionFraction)));
+				//
+				//	TotalBudgetOverBudgetSavings[BrowsingFrameId] += EstimatedSavings;
+				//}
+
+				SuggestedResolutionFraction += AdjustedResolutionFraction * Weight;
+				SuggestedResolutionFractionWeight += Weight;
+				Weight *= FrameWeightExponent;
+
+				if (BrowsingFrameId == StartFrameId && BudgetSetting.bModelScalesWithPrimaryScreenPercentage)
+				{
+					MostRecentResolutionFractionScale = FrameEntry.ResolutionFractions[GDynamicPrimaryResolutionFraction];
+				}
+			}
+
+			NewResolutionFraction = BudgetSetting.CorrectNewResolutionFraction(
+				CurrentResolutionFraction,
+				SuggestedResolutionFraction / SuggestedResolutionFractionWeight,
+				MostRecentResolutionFractionScale);
+
+			NewResolutionFraction /= MostRecentResolutionFractionScale;
+
+			bool bTakeNewResolutionFraction = (
+				BudgetSetting.DoesResolutionChangeEnough(CurrentResolutionFraction, NewResolutionFraction, bCanChangeResolution) ||
+				bBudgetIsOverBudget);
+
+			if (bTakeNewResolutionFraction)
+			{
+				bCommitResolutionChange = true;
 			}
 			else
 			{
-				ConsecutiveOverbudgetGPUFramCount = 0;
+				NewResolutionFraction = CurrentFrameResolutionFractions[Budget];
 			}
+		}
 
-			// Ignore if an outlier.
-			if (bIsGPUOutlier)
-			{
-				continue;
-			}
+		return NewResolutionFraction;
+	};
 
-			float SuggestedResolutionFraction = 1.0f;
+	// Estimate the new cost of the bucket.
+	auto EstimateBudgetCost = [&](const DynamicRenderScaling::FBudget& Budget, bool bBudgetIsOverBudget)
+	{
+		const DynamicRenderScaling::FHeuristicSettings& BudgetSetting = Budget.GetSettings();
+		check(BudgetSetting.IsEnabled());
 
-			// If have reliable GPU times, or guess there is not GPU bubbles -> estimate the suggested resolution fraction that could have been used.
-			if (!FrameEntry.bGPUTimingsHaveCPUBubbles ||
-				(FrameEntry.bGPUTimingsHaveCPUBubbles && !bIsCPUBound && !bMayHaveGPUBubbles))
-			{
-				// This assumes GPU busy time is directly proportional to ResolutionFraction^2, but in practice
-				// this is more A * ResolutionFraction^2 + B with B >= 0 non constant unknown cost such as unscaled
-				// post processing, vertex fetching & processing, occlusion queries, shadow maps rendering...
-				//
-				// This assumption means we may drop ResolutionFraction lower than needed, or be slower to increase
-				// resolution.
-				//
-				// TODO: If we have RHI guarantee of frame timing association, we could make an estimation of B
-				// At resolution change that happen every N frames, amortized over time and scaled down as the
-				// standard variation of the GPU timing over non resolution changing frames increases.
-				SuggestedResolutionFraction = (
-					FMath::Sqrt(TargetedGPUBusyTimeMs / FrameEntry.TotalFrameGPUBusyTimeMs)
-					* FrameEntry.ResolutionFraction);
-			}
-			else if (FrameEntry.ResolutionFraction > CPUBoundResolutionFraction)
-			{
-				// When CPU bound, we no longer want to downscale, or drop resolution if CPU and GPU are sharing same memory.
-				SuggestedResolutionFraction = CPUBoundResolutionFraction;
-			}
-			else
-			{
-				// Unable to know reliably the GPU time so don't try to make any prediction to avoid issues such
-				// as resolution keep shrinking because GPU time is slightly more than render thread time, or
-				// increasing because GPU time slightly lower than frame time.
-				SuggestedResolutionFraction = FrameEntry.ResolutionFraction;
-			}
+		int32 MaxFrameId = ComputeMaxFrameId(Budget, bBudgetIsOverBudget);
 
-			NewFrameResolutionFraction += SuggestedResolutionFraction * Weight;
-			TotalWeight += Weight;
-			FrameCount++;
+		float CurrentResolutionFraction = CurrentFrameResolutionFractions[Budget];
+		float NewResolutionFraction = NewResolutionFractions[Budget];
 
+		ensure(MaxFrameId != StartFrameId);
+
+		float BudgetPredictedCost = 0.0;
+		float BudgetPredictedCostWeight = 0.0;
+		float Weight = 1.0f;
+		for (int32 BrowsingFrameId = StartFrameId; BrowsingFrameId < MaxFrameId; BrowsingFrameId++)
+		{
+			const FrameHistoryEntry& FrameEntry = GetPreviousFrameEntry(BrowsingFrameId);
+			float EstimatedCost = (
+				FrameEntry.BudgetTimingMs[Budget] *
+				BudgetSetting.EstimateTimeFactor(CurrentResolutionFraction, NewResolutionFraction));
+
+			BudgetPredictedCost += EstimatedCost * Weight;
+			BudgetPredictedCostWeight += Weight;
 			Weight *= FrameWeightExponent;
 		}
 
-		NewFrameResolutionFraction /= TotalWeight;
+		return BudgetPredictedCost / BudgetPredictedCostWeight;
+	};
 
-		// If immediate previous frames where over budget, react immediately.
-		bGPUOverbugetPanic = FrameCount > 0 && ConsecutiveOverbudgetGPUFramCount == FrameCount;
+	// Iterate through bucket and adjust them.
+	float TotalStaticBudgetPredictedCost = 0.0f;
+	float TotalPrimaryScreenPercentageDynamicBudgetPredictedCost = 0.0f;
+	for (TLinkedList<DynamicRenderScaling::FBudget*>::TIterator BudgetIt(DynamicRenderScaling::FBudget::GetGlobalList()); BudgetIt; BudgetIt.Next())
+	{
+		const DynamicRenderScaling::FBudget& Budget = **BudgetIt;
+		const DynamicRenderScaling::FHeuristicSettings& HeuristicSettings = Budget.GetSettings();
 
-		// If over budget, reset history size to 0 so that this frame really behave as a first frame after an history reset.
-		if (bGPUOverbugetPanic)
+		// Primary screen percentage is done last.
+		if (Budget == GDynamicPrimaryResolutionFraction)
 		{
-			HistorySize = 0;
+			continue;
 		}
-		// If not immediately over budget, refine the new resolution fraction.
-		else
+
+		// Skip if the bucket is not enabled.
+		if (HeuristicSettings.IsEnabled())
 		{
-			// If scaling the resolution, look if this is above a threshold compared to current res.
-			if (CVarChangeThreshold.GetValueOnRenderThread())
+			bool bBudgetIsOverBudget = IsBudgetOverBudget(Budget, HeuristicSettings);
+
+			NewResolutionFractions[Budget] = ComputeNewResolutionFraction(Budget, HeuristicSettings, bBudgetIsOverBudget);
+
+			float BudgetCost = EstimateBudgetCost(Budget, bBudgetIsOverBudget);
+			if (HeuristicSettings.bModelScalesWithPrimaryScreenPercentage)
 			{
-				float IncreaseResolutionFractionThreshold = CurrentFrameResolutionFraction * (1.0f + CVarChangeThreshold.GetValueOnRenderThread() / 100.0f);
-
-				float DecreaseResolutionFractionThreshold = CurrentFrameResolutionFraction * (1.0f - CVarChangeThreshold.GetValueOnRenderThread() / 100.0f);
-
-				if (NewFrameResolutionFraction > CurrentFrameResolutionFraction && NewFrameResolutionFraction < IncreaseResolutionFractionThreshold)
-				{
-					NewFrameResolutionFraction = CurrentFrameResolutionFraction;
-				}
+				TotalPrimaryScreenPercentageDynamicBudgetPredictedCost += BudgetCost;
+			}
+			else
+			{
+				TotalStaticBudgetPredictedCost += BudgetCost;
 			}
 
-			// If scaling the resolution up, amortize to avoid oscillations.
-			if (NewFrameResolutionFraction > CurrentFrameResolutionFraction)
+			// Reset the size of the history for the bucket when it is over budget.
+			if (bBudgetIsOverBudget)
 			{
-				NewFrameResolutionFraction = FMath::Lerp(
-					CurrentFrameResolutionFraction,
-					NewFrameResolutionFraction,
-					CVarIncreaseAmortizationFactor.GetValueOnRenderThread());
+				BudgetHistorySizes[Budget] = 0;
 			}
-		}
-
-		// Clamp resolution fraction.
-		NewFrameResolutionFraction = FMath::Clamp(
-			NewFrameResolutionFraction,
-			CVarDynamicResMinSP.GetValueOnRenderThread() / 100.0f,
-			GetResolutionFractionUpperBound());
-	}
-
-	// Update the current frame's resolution fraction.
-	{
-		// CVarChangeThreshold avoids very small changes.
-		bool bWouldBeWorthChangingRes = CurrentFrameResolutionFraction != NewFrameResolutionFraction;
-
-		// We do not change resolution too often to avoid interferences with temporal sub pixel in TAA upsample.
-		if ((bWouldBeWorthChangingRes && 
-			NumberOfFramesSinceScreenPercentageChange >= CVarFrameChangePeriod.GetValueOnRenderThread()) ||
-			bGPUOverbugetPanic)
-		{
-			NumberOfFramesSinceScreenPercentageChange = 0;
-			CurrentFrameResolutionFraction = NewFrameResolutionFraction;
 		}
 		else
 		{
-			NumberOfFramesSinceScreenPercentageChange++;
+			NewResolutionFractions[Budget] = 1.0f;
+			continue;
 		}
+	}
+
+	// Takes care of primary screen percentage to finally fit the frame within budget
+	{
+		const DynamicRenderScaling::FBudget& Budget = GDynamicPrimaryResolutionFraction;
+
+		// Set the budget of the budget.
+		DynamicRenderScaling::FHeuristicSettings HeuristicSettings = Budget.GetSettings();
+		HeuristicSettings.BudgetMs = FMath::Max(TargetedGPUBusyTimeMs - TotalStaticBudgetPredictedCost, 0.5f);
+		check(HeuristicSettings.IsEnabled());
+
+		bool bBudgetIsOverBudget = IsBudgetOverBudget(Budget, HeuristicSettings);
+
+		NewResolutionFractions[Budget] = ComputeNewResolutionFraction(Budget, HeuristicSettings, bBudgetIsOverBudget);
+	}
+
+	// Commit the new resolution fractions.
+	if (bCommitResolutionChange)
+	{
+		NumberOfFramesSinceScreenPercentageChange = 0;
+		CurrentFrameResolutionFractions = NewResolutionFractions;
+	}
+	else
+	{
+		NumberOfFramesSinceScreenPercentageChange++;
 	}
 }
 
-float FDynamicResolutionHeuristicProxy::GetResolutionFractionUpperBound() const
+// static
+DynamicRenderScaling::TMap<float> FDynamicResolutionHeuristicProxy::GetResolutionFractionUpperBounds()
 {
-	return CVarDynamicResMaxSP.GetValueOnAnyThread() / 100.0f;
+	DynamicRenderScaling::TMap<float> MaxResolutionFractions;
+	for (TLinkedList<DynamicRenderScaling::FBudget*>::TIterator BudgetIt(DynamicRenderScaling::FBudget::GetGlobalList()); BudgetIt; BudgetIt.Next())
+	{
+		const DynamicRenderScaling::FBudget& Budget = **BudgetIt;
+		MaxResolutionFractions[Budget] = Budget.GetSettings().MaxResolutionFraction;
+	}
+	return MaxResolutionFractions;
 }
 
+
+/** Returns the view fraction that should be used for current frame. */
+DynamicRenderScaling::TMap<float> FDynamicResolutionHeuristicProxy::QueryCurentFrameResolutionFractions_Internal() const
+{
+	DynamicRenderScaling::TMap<float> MaxResolutionFractions = GetResolutionFractionUpperBounds();
+	DynamicRenderScaling::TMap<float> ResolutionFractions = CurrentFrameResolutionFractions;
+	for (TLinkedList<DynamicRenderScaling::FBudget*>::TIterator BudgetIt(DynamicRenderScaling::FBudget::GetGlobalList()); BudgetIt; BudgetIt.Next())
+	{
+		const DynamicRenderScaling::FBudget& Budget = **BudgetIt;
+		ResolutionFractions[Budget] = FMath::Min(ResolutionFractions[Budget], MaxResolutionFractions[Budget]);
+	}
+	return ResolutionFractions;
+}
+
 void FDynamicResolutionHeuristicProxy::ResizeHistoryIfNeeded()
 {
 	uint32 DesiredHistorySize = FMath::Max(1, CVarHistorySize.GetValueOnRenderThread());
@@ -448,59 +559,6 @@
 
 
 /**
- * Engine's default dynamic resolution driver for view families.
- */
-class FDefaultDynamicResolutionDriver : public ISceneViewFamilyScreenPercentage
-{
-public:
-
-	FDefaultDynamicResolutionDriver(const FDynamicResolutionHeuristicProxy* InProxy, const FSceneViewFamily& InViewFamily)
-		: Proxy(InProxy)
-		, ViewFamily(InViewFamily)
-	{
-		check(IsInGameThread());
-	}
-
-	virtual float GetPrimaryResolutionFractionUpperBound() const override
-	{
-		if (!ViewFamily.EngineShowFlags.ScreenPercentage)
-		{
-			return 1.0f;
-		}
-
-		return Proxy->GetResolutionFractionUpperBound();
-	}
-
-	virtual float GetPrimaryResolutionFraction_RenderThread() const override
-	{
-		check(IsInRenderingThread());
-
-		if (!ViewFamily.EngineShowFlags.ScreenPercentage)
-		{
-			return 1.0f;
-		}
-
-		return Proxy->QueryCurentFrameResolutionFraction_RenderThread();
-	}
-
-	virtual ISceneViewFamilyScreenPercentage* Fork_GameThread(const class FSceneViewFamily& ForkedViewFamily) const override
-	{
-		check(IsInGameThread());
-
-		return new FDefaultDynamicResolutionDriver(Proxy, ForkedViewFamily);
-	}
-
-private:
-	// Dynamic resolution proxy to use.
-	const FDynamicResolutionHeuristicProxy* Proxy;
-
-	// View family to take care of.
-	const FSceneViewFamily& ViewFamily;
-
-};
-
-
-/**
  * Render thread proxy for engine's dynamic resolution state.
  */
 class FDefaultDynamicResolutionStateProxy
@@ -538,6 +596,22 @@
 	{
 		check(IsInRenderingThread());
 
+		if (DynamicRenderScaling::IsSupported())
+		{
+			DynamicRenderScaling::UpdateHeuristicsSettings();
+
+			DynamicRenderScaling::TMap<bool> bIsBudgetEnabled;
+			bIsBudgetEnabled.SetAll(false);
+
+			for (TLinkedList<DynamicRenderScaling::FBudget*>::TIterator BudgetIt(DynamicRenderScaling::FBudget::GetGlobalList()); BudgetIt; BudgetIt.Next())
+			{
+				const DynamicRenderScaling::FBudget& Budget = **BudgetIt;
+				bIsBudgetEnabled[Budget] = Budget.GetSettings().IsEnabled();
+			}
+
+			DynamicRenderScaling::BeginFrame(bIsBudgetEnabled);
+		}
+
 		// Query render thread time Ms.
 		float PrevRenderThreadTimeMs = FPlatformTime::ToMilliseconds(GRenderThreadTime);
 
@@ -568,16 +642,26 @@
 		else
 		{
 			// If RHI does not support GPU busy time queries, fall back to what stat unit does.
-			// Multi-GPU support : should be updated to use GPUIndex for AFR.
+			ensure(GRHISupportsFrameCyclesBubblesRemoval);
 			float PrevFrameGPUTimeMs = FPlatformTime::ToMilliseconds(RHIGetGPUFrameCycles());
 
 			uint64 HistoryEntryId = Heuristic.CreateNewPreviousFrameTimings_RenderThread(
 				PrevGameThreadTimeMs, PrevRenderThreadTimeMs);
 
+			const DynamicRenderScaling::TMap<uint64>& LattestTimings = DynamicRenderScaling::GetLastestTimings();
+			DynamicRenderScaling::TMap<float> BudgetTimingMs;
+			for (TLinkedList<DynamicRenderScaling::FBudget*>::TIterator BudgetIt(DynamicRenderScaling::FBudget::GetGlobalList()); BudgetIt; BudgetIt.Next())
+			{
+				const DynamicRenderScaling::FBudget& Budget = **BudgetIt;
+				BudgetTimingMs[Budget] = TimeStampQueryResultToMiliSeconds(LattestTimings[Budget]);
+			}
+			BudgetTimingMs[GDynamicPrimaryResolutionFraction] = PrevFrameGPUTimeMs;
+
+
 			Heuristic.CommitPreviousFrameGPUTimings_RenderThread(HistoryEntryId,
 				/* TotalFrameGPUBusyTimeMs = */ PrevFrameGPUTimeMs,
 				/* DynamicResolutionGPUBusyTimeMs = */ PrevFrameGPUTimeMs,
-				/* bGPUTimingsHaveCPUBubbles = */ !GRHISupportsFrameCyclesBubblesRemoval);
+				/* BudgetTimingMs = */ BudgetTimingMs);
 
 			Heuristic.RefreshCurentFrameResolutionFraction_RenderThread();
 
@@ -617,6 +701,8 @@
 		{
 			CurrentFrameInFlightIndex = -1;
 			bUseTimeQueriesThisFrame = false;
+
+			DynamicRenderScaling::EndFrame();
 		}
 	}
 
@@ -682,6 +768,7 @@
 	void HandLandedQueriesToHeuristic(bool bWait)
 	{
 		check(IsInRenderingThread());
+		check(GRHISupportsGPUTimestampBubblesRemoval || bWait);
 
 		bool ShouldRefreshHeuristic = false;
 
@@ -738,8 +825,8 @@
 			{
 				Heuristic.CommitPreviousFrameGPUTimings_RenderThread(
 					InFlightFrame.HeuristicHistoryEntry,
-					/* TotalFrameGPUBusyTimeMs = */ float(EndFrameResult - BeginFrameResult) / 1000.0f,
-					/* DynamicResolutionGPUBusyTimeMs = */ float(EndDynamicResolutionResult - BeginDynamicResolutionResult) / 1000.0f,
+					/* TotalFrameGPUBusyTimeMs = */ TimeStampQueryResultToMiliSeconds(EndFrameResult - BeginFrameResult),
+					/* DynamicResolutionGPUBusyTimeMs = */ TimeStampQueryResultToMiliSeconds(EndDynamicResolutionResult - BeginDynamicResolutionResult),
 					/* bGPUTimingsHaveCPUBubbles = */ !GRHISupportsGPUTimestampBubblesRemoval);
 
 				// Reset this in-flight frame queries to be reused.
@@ -781,6 +868,61 @@
 
 
 /**
+ * Engine's default dynamic resolution driver for view families.
+ */
+class FDefaultDynamicResolutionDriver : public ISceneViewFamilyScreenPercentage
+{
+public:
+
+	FDefaultDynamicResolutionDriver(FDefaultDynamicResolutionStateProxy* InProxy, const FSceneViewFamily& InViewFamily)
+		: Proxy(InProxy)
+		, ViewFamily(InViewFamily)
+	{
+		check(IsInGameThread());
+	}
+
+	virtual DynamicRenderScaling::TMap<float> GetResolutionFractionsUpperBound() const override
+	{
+		DynamicRenderScaling::TMap<float> UpperBounds = Proxy->Heuristic.GetResolutionFractionUpperBounds();
+		if (!ViewFamily.EngineShowFlags.ScreenPercentage)
+		{
+			UpperBounds[GDynamicPrimaryResolutionFraction] = 1.0f;
+		}
+
+		return UpperBounds;
+	}
+
+	virtual ISceneViewFamilyScreenPercentage* Fork_GameThread(const class FSceneViewFamily& ForkedViewFamily) const override
+	{
+		check(IsInGameThread());
+
+		return new FDefaultDynamicResolutionDriver(Proxy, ForkedViewFamily);
+	}
+
+	virtual DynamicRenderScaling::TMap<float> GetResolutionFractions_RenderThread() const override
+	{
+		check(IsInRenderingThread());
+
+		DynamicRenderScaling::TMap<float> ResolutionFractions = Proxy->Heuristic.QueryCurentFrameResolutionFractions();
+		if (!ViewFamily.EngineShowFlags.ScreenPercentage)
+		{
+			ResolutionFractions[GDynamicPrimaryResolutionFraction] = 1.0f;
+		}
+
+		return ResolutionFractions;
+	}
+
+private:
+	// Dynamic resolution proxy to use.
+	FDefaultDynamicResolutionStateProxy* Proxy;
+
+	// View family to take care of.
+	const FSceneViewFamily& ViewFamily;
+
+};
+
+
+/**
  * Engine's default dynamic resolution state.
  */
 class FDefaultDynamicResolutionState : public IDynamicResolutionState
@@ -846,16 +988,16 @@
 		return bIsEnabled;
 	}
 
-	virtual float GetResolutionFractionApproximation() const override
+	virtual DynamicRenderScaling::TMap<float> GetResolutionFractionsApproximation() const override
 	{
 		check(IsInGameThread());
-		return Proxy->Heuristic.GetResolutionFractionApproximation_GameThread();
+		return Proxy->Heuristic.GetResolutionFractionsApproximation_GameThread();
 	}
 
-	virtual float GetResolutionFractionUpperBound() const override
+	virtual DynamicRenderScaling::TMap<float> GetResolutionFractionsUpperBound() const override
 	{
 		check(IsInGameThread());
-		return Proxy->Heuristic.GetResolutionFractionUpperBound();
+		return Proxy->Heuristic.GetResolutionFractionUpperBounds();
 	}
 
 	virtual void ProcessEvent(EDynamicResolutionStateEvent Event) override
@@ -910,7 +1052,7 @@
 
 		if (bIsEnabled)
 		{
-			ViewFamily.SetScreenPercentageInterface(new FDefaultDynamicResolutionDriver(&Proxy->Heuristic, ViewFamily));
+			ViewFamily.SetScreenPercentageInterface(new FDefaultDynamicResolutionDriver(Proxy, ViewFamily));
 		}
 	}
 
diff --git a/Engine/Source/Runtime/Engine/Private/GameViewportClient.cpp b/Engine/Source/Runtime/Engine/Private/GameViewportClient.cpp
--- a/Engine/Source/Runtime/Engine/Private/GameViewportClient.cpp
+++ b/Engine/Source/Runtime/Engine/Private/GameViewportClient.cpp
@@ -68,6 +68,7 @@
 #include "HAL/PlatformApplicationMisc.h"
 #include "CustomStaticScreenPercentage.h"
 #include "ObjectTrace.h"
+#include "DynamicResolutionState.h"
 
 #if WITH_EDITOR
 #include "Settings/LevelEditorPlaySettings.h"
@@ -1412,17 +1413,33 @@
 				GEngine->EmitDynamicResolutionEvent(EDynamicResolutionStateEvent::BeginDynamicResolutionRendering);
 				ViewFamily.SetScreenPercentageInterface(new FLegacyScreenPercentageDriver(
 					ViewFamily,
-					DynamicResolutionStateInfos.ResolutionFractionApproximation,
-					DynamicResolutionStateInfos.ResolutionFractionUpperBound));
+					DynamicResolutionStateInfos.ResolutionFractionApproximations[GDynamicPrimaryResolutionFraction],
+					DynamicResolutionStateInfos.ResolutionFractionUpperBounds[GDynamicPrimaryResolutionFraction]));
 
 				bUsesDynamicResolution = true;
 			}
 
+			// Feed approximated resolution fraction to CSV
 			#if CSV_PROFILER
-			if (DynamicResolutionStateInfos.ResolutionFractionApproximation >= 0.0f)
+			if (DynamicResolutionStateInfos.ResolutionFractionApproximations[GDynamicPrimaryResolutionFraction] >= 0.0f)
 			{
-				CSV_CUSTOM_STAT_GLOBAL(DynamicResolutionPercentage, DynamicResolutionStateInfos.ResolutionFractionApproximation * 100.0f, ECsvCustomStatOp::Set);
+				// Keep same name as before for primary screen percentage
+				CSV_CUSTOM_STAT_GLOBAL(DynamicResolutionPercentage, DynamicResolutionStateInfos.ResolutionFractionApproximations[GDynamicPrimaryResolutionFraction] * 100.0f, ECsvCustomStatOp::Set);
 			}
+			for (TLinkedList<DynamicRenderScaling::FBudget*>::TIterator BudgetIt(DynamicRenderScaling::FBudget::GetGlobalList()); BudgetIt; BudgetIt.Next())
+			{
+				const DynamicRenderScaling::FBudget& Budget = **BudgetIt;
+				if (Budget == GDynamicPrimaryResolutionFraction)
+				{
+					continue;
+				}
+
+				float Value = DynamicResolutionStateInfos.ResolutionFractionApproximations[Budget] * 100.0f;
+				const char* NameChar = Budget.GetAnsiName();
+
+				TRACE_CSV_PROFILER_INLINE_STAT(NameChar, CSV_CATEGORY_INDEX_GLOBAL);
+				FCsvProfiler::RecordCustomStat(NameChar, CSV_CATEGORY_INDEX_GLOBAL, Value, ECsvCustomStatOp::Set);
+			}
 			#endif
 		}
 		#endif
diff --git a/Engine/Source/Runtime/Engine/Private/LegacyScreenPercentageDriver.cpp b/Engine/Source/Runtime/Engine/Private/LegacyScreenPercentageDriver.cpp
--- a/Engine/Source/Runtime/Engine/Private/LegacyScreenPercentageDriver.cpp
+++ b/Engine/Source/Runtime/Engine/Private/LegacyScreenPercentageDriver.cpp
@@ -3,6 +3,7 @@
 #include "LegacyScreenPercentageDriver.h"
 #include "UnrealEngine.h"
 #include "Misc/ConfigCacheIni.h"
+#include "DynamicResolutionState.h"
 
 
 static TAutoConsoleVariable<int32> CVarScreenPercentageMode(
@@ -118,33 +119,38 @@
 	}
 }
 
-float FLegacyScreenPercentageDriver::GetPrimaryResolutionFractionUpperBound() const
+DynamicRenderScaling::TMap<float> FLegacyScreenPercentageDriver::GetResolutionFractionsUpperBound() const
 {
-	if (!ViewFamily.EngineShowFlags.ScreenPercentage)
+	DynamicRenderScaling::TMap<float> UpperBounds;
+	UpperBounds.SetAll(1.0f);
+
+	if (ViewFamily.EngineShowFlags.ScreenPercentage)
 	{
-		return 1.0f;
+		UpperBounds[GDynamicPrimaryResolutionFraction] = FMath::Clamp(
+			GlobalResolutionFractionUpperBound,
+			ISceneViewFamilyScreenPercentage::kMinResolutionFraction,
+			ISceneViewFamilyScreenPercentage::kMaxResolutionFraction);
 	}
 
-	return FMath::Clamp(
-		GlobalResolutionFractionUpperBound,
-		ISceneViewFamilyScreenPercentage::kMinResolutionFraction,
-		ISceneViewFamilyScreenPercentage::kMaxResolutionFraction);
+	return UpperBounds;
 }
 
-float FLegacyScreenPercentageDriver::GetPrimaryResolutionFraction_RenderThread() const
+DynamicRenderScaling::TMap<float> FLegacyScreenPercentageDriver::GetResolutionFractions_RenderThread() const
 {
 	check(IsInRenderingThread());
 
-	// Early return if no screen percentage should be done.
-	if (!ViewFamily.EngineShowFlags.ScreenPercentage)
+	DynamicRenderScaling::TMap<float> ResolutionFractions;
+	ResolutionFractions.SetAll(1.0f);
+
+	if (ViewFamily.EngineShowFlags.ScreenPercentage)
 	{
-		return 1.0f;
+		ResolutionFractions[GDynamicPrimaryResolutionFraction] = FMath::Clamp(
+			GlobalResolutionFraction,
+			ISceneViewFamilyScreenPercentage::kMinResolutionFraction,
+			ISceneViewFamilyScreenPercentage::kMaxResolutionFraction);
 	}
 
-	return FMath::Clamp(
-		GlobalResolutionFraction,
-		ISceneViewFamilyScreenPercentage::kMinResolutionFraction,
-		ISceneViewFamilyScreenPercentage::kMaxResolutionFraction);
+	return ResolutionFractions;
 }
 
 ISceneViewFamilyScreenPercentage* FLegacyScreenPercentageDriver::Fork_GameThread(
diff --git a/Engine/Source/Runtime/Engine/Private/UnrealClient.cpp b/Engine/Source/Runtime/Engine/Private/UnrealClient.cpp
--- a/Engine/Source/Runtime/Engine/Private/UnrealClient.cpp
+++ b/Engine/Source/Runtime/Engine/Private/UnrealClient.cpp
@@ -29,6 +29,7 @@
 #include "Elements/Framework/TypedElementList.h"
 #include "EngineUtils.h"
 #include "RenderGraphUtils.h"
+#include "DynamicResolutionState.h"
 
 DEFINE_LOG_CATEGORY_STATIC(LogClient, Log, All);
 
@@ -497,7 +498,11 @@
 	FrameTimes[CurrentIndex] = bShowRawUnitTimes ? RawFrameTime : FrameTime;
 	RHITTimes[CurrentIndex] = bShowRawUnitTimes ? RawRHITTime : RHITTime;
 	InputLatencyTimes[CurrentIndex] = bShowRawUnitTimes ? RawInputLatencyTime : InputLatencyTime;
-	ResolutionFractions[CurrentIndex] = DynamicResolutionStateInfos.ResolutionFractionApproximation;
+	for (TLinkedList<DynamicRenderScaling::FBudget*>::TIterator BudgetIt(DynamicRenderScaling::FBudget::GetGlobalList()); BudgetIt; BudgetIt.Next())
+	{
+		const DynamicRenderScaling::FBudget& Budget = **BudgetIt;
+		ResolutionFractions[Budget][CurrentIndex] = DynamicResolutionStateInfos.ResolutionFractionApproximations[Budget];
+	}
 	CurrentIndex++;
 	if (CurrentIndex == NumberOfSamples)
 	{
@@ -569,12 +574,18 @@
 		}
 
 		int32 X2 = bShowUnitMaxTimes ? X3 - (int32)((float)Font->GetStringSize(TEXT(" 000.00 ms "))) : X3;
-		int32 X1 = X2 - (int32)((float)Font->GetStringSize(TEXT("DynRes: ")));
 		const int32 RowHeight = FMath::TruncToInt(Font->GetMaxCharHeight() * 1.1f);
 
+		auto DrawTitleString = [&](const TCHAR* Title, const FColor& UnitGraphColor)
 		{
+			FString FullTitle = FString::Printf(TEXT("%s:   "), Title);
+			int32 TitleSize = Font->GetStringSize(*FullTitle);
+			InCanvas->DrawShadowedString(X2 - TitleSize, InY, *FullTitle, Font, bShowUnitTimeGraph ? UnitGraphColor : FColor::White);
+		};
+
+		{
 			const FColor FrameTimeAverageColor = GEngine->GetFrameTimeDisplayColor(FrameTime);
-			InCanvas->DrawShadowedString(X1, InY, TEXT("Frame:"), Font, bShowUnitTimeGraph ? FColor(100, 255, 100) : FColor::White);
+			DrawTitleString(TEXT("Frame"), /* UnitGraphColor = */ FColor(100, 255, 100));
 			InCanvas->DrawShadowedString(X2, InY, *FString::Printf(TEXT("%3.2f ms"), FrameTime), Font, FrameTimeAverageColor);
 			if (bShowUnitMaxTimes)
 			{
@@ -586,7 +597,7 @@
 
 		{
 			const FColor GameThreadAverageColor = GEngine->GetFrameTimeDisplayColor(GameThreadTime);
-			InCanvas->DrawShadowedString(X1, InY, TEXT("Game:"), Font, bShowUnitTimeGraph ? FColor(255, 100, 100) : FColor::White);
+			DrawTitleString(TEXT("Game"), /* UnitGraphColor = */ FColor(255, 100, 100));
 			InCanvas->DrawShadowedString(X2, InY, *FString::Printf(TEXT("%3.2f ms"), GameThreadTime), Font, GameThreadAverageColor);
 			if (bShowUnitMaxTimes)
 			{
@@ -598,7 +609,7 @@
 
 		{
 			const FColor RenderThreadAverageColor = GEngine->GetFrameTimeDisplayColor(RenderThreadTime);
-			InCanvas->DrawShadowedString(X1, InY, TEXT("Draw:"), Font, bShowUnitTimeGraph ? FColor(100, 100, 255) : FColor::White);
+			DrawTitleString(TEXT("Draw"), /* UnitGraphColor = */ FColor(100, 100, 255));
 			InCanvas->DrawShadowedString(X2, InY, *FString::Printf(TEXT("%3.2f ms"), RenderThreadTime), Font, RenderThreadAverageColor);
 			if (bShowUnitMaxTimes)
 			{
@@ -613,8 +624,8 @@
 			if (bHaveGPUData[GPUIndex])
 			{
 				const FColor GPUAverageColor = GEngine->GetFrameTimeDisplayColor(GPUFrameTime[GPUIndex]);
-				FString GPUString = GNumExplicitGPUsForRendering > 1 ? FString::Printf(TEXT("GPU%u:"), GPUIndex) : TEXT("GPU:");
-				InCanvas->DrawShadowedString(X1, InY, *GPUString, Font, bShowUnitTimeGraph ? FColor(255, 255, 100) : FColor::White);
+				FString GPUString = GNumExplicitGPUsForRendering > 1 ? FString::Printf(TEXT("GPU%u"), GPUIndex) : TEXT("GPU");
+				DrawTitleString(*GPUString, /* UnitGraphColor = */ FColor(255, 255, 100));
 				InCanvas->DrawShadowedString(X2, InY, *FString::Printf(TEXT("%3.2f ms"), GPUFrameTime[GPUIndex]), Font, GPUAverageColor);
 				if (bShowUnitMaxTimes)
 				{
@@ -627,7 +638,7 @@
 		if (IsRunningRHIInSeparateThread())
 		{
 			const FColor RenderThreadAverageColor = GEngine->GetFrameTimeDisplayColor(RHITTime);
-			InCanvas->DrawShadowedString(X1, InY, TEXT("RHIT:"), Font, bShowUnitTimeGraph ? FColor(255, 100, 255) : FColor::White);
+			DrawTitleString(TEXT("RHIT"), /* UnitGraphColor = */ FColor(255, 100, 255));
 			InCanvas->DrawShadowedString(X2, InY, *FString::Printf(TEXT("%3.2f ms"), RHITTime), Font, RenderThreadAverageColor);
 			if (bShowUnitMaxTimes)
 			{
@@ -640,7 +651,7 @@
 		{
 			const float ReasonableInputLatencyFactor = 2.5f;
 			const FColor InputLatencyAverageColor = GEngine->GetFrameTimeDisplayColor(InputLatencyTime / ReasonableInputLatencyFactor);
-			InCanvas->DrawShadowedString(X1, InY, TEXT("Input:"), Font, bShowUnitTimeGraph ? FColor(255, 255, 100) : FColor::White);
+			DrawTitleString(TEXT("Input"), /* UnitGraphColor = */ FColor(255, 255, 100));
 			InCanvas->DrawShadowedString(X2, InY, *FString::Printf(TEXT("%3.2f ms"), InputLatencyTime), Font, InputLatencyAverageColor);
 			if (bShowUnitMaxTimes)
 			{
@@ -654,12 +665,12 @@
 			{
 				FPlatformMemoryStats Stats = FPlatformMemory::GetStats();
 
-				InCanvas->DrawShadowedString(X1, InY, TEXT("Mem:"), Font, bShowUnitTimeGraph ? FColor(100, 100, 255) : FColor::White);
+				DrawTitleString(TEXT("Mem"), /* UnitGraphColor = */ FColor(100, 100, 255));
 				InCanvas->DrawShadowedString(X2, InY, *GetMemoryString(Stats.UsedPhysical), Font, StatGreen);
 				InCanvas->DrawShadowedString(X3, InY, *GetMemoryString(Stats.PeakUsedPhysical), Font, StatGreen);
 				InY += RowHeight;
 				
-				InCanvas->DrawShadowedString(X1, InY, TEXT("VMem:"), Font, bShowUnitTimeGraph ? FColor(100, 100, 255) : FColor::White);
+				DrawTitleString(TEXT("VMem"), /* UnitGraphColor = */ FColor(100, 100, 255));
 				InCanvas->DrawShadowedString(X2, InY, *GetMemoryString(Stats.UsedVirtual), Font, StatGreen);
 				InCanvas->DrawShadowedString(X3, InY, *GetMemoryString(Stats.PeakUsedVirtual), Font, StatGreen);
 				InY += RowHeight;
@@ -670,18 +681,19 @@
 				if (MemoryUsed > 0)
 				{
 					// print out currently used memory
-					InCanvas->DrawShadowedString(X1, InY, TEXT("Mem:"), Font, bShowUnitTimeGraph ? FColor(100, 100, 255) : FColor::White);
+					DrawTitleString(TEXT("Mem"), /* UnitGraphColor = */ FColor(100, 100, 255));
 					InCanvas->DrawShadowedString(X2, InY, *GetMemoryString(MemoryUsed), Font, StatGreen);
 					InY += RowHeight;
 				}
 			}
 		}
 
+		// Dynamic resolution
 		{
-			float ResolutionFraction = DynamicResolutionStateInfos.ResolutionFractionApproximation;
+			float ResolutionFraction = DynamicResolutionStateInfos.ResolutionFractionApproximations[GDynamicPrimaryResolutionFraction];
 			float ScreenPercentage = ResolutionFraction * 100.0f;
 
-			InCanvas->DrawShadowedString(X1, InY, TEXT("DynRes:"), Font, bShowUnitTimeGraph ? FColor(255, 160, 100) : FColor::White);
+			DrawTitleString(TEXT("DynRes"), /* UnitGraphColor = */ FColor(255, 160, 100));
 			if (DynamicResolutionStateInfos.Status == EDynamicResolutionStatus::Enabled)
 			{
 				FColor Color = (ResolutionFraction < AlertResolutionFraction) ? StatRed : ((ResolutionFraction < FMath::Min(ResolutionFraction * 0.97f, 1.0f)) ? StatOrange : StatGreen);
@@ -710,20 +722,58 @@
 			InY += RowHeight;
 		}
 
+		// Other dynamic render scalings
+		if (DynamicResolutionStateInfos.Status == EDynamicResolutionStatus::Enabled && DynamicRenderScaling::IsSupported())
+		{
+			for (TLinkedList<DynamicRenderScaling::FBudget*>::TIterator BudgetIt(DynamicRenderScaling::FBudget::GetGlobalList()); BudgetIt; BudgetIt.Next())
+			{
+				const DynamicRenderScaling::FBudget& Budget = **BudgetIt;
+				const DynamicRenderScaling::FHeuristicSettings& HeuristicSettings = Budget.GetSettings();
+				if (Budget == GDynamicPrimaryResolutionFraction)
+				{
+					continue;
+				}
+
+				float ResolutionFraction = DynamicResolutionStateInfos.ResolutionFractionApproximations[Budget];
+				float ScreenPercentage = ResolutionFraction * 100.0f;
+
+				FString DisplayName = Budget.GetName();
+				DisplayName.ReplaceInline(TEXT("Dynamic"), TEXT("Dyn"));
+				DisplayName.ReplaceInline(TEXT("Resolution"), TEXT("Res"));
+
+				FColor Color = (ResolutionFraction < AlertResolutionFraction) ? StatRed : ((ResolutionFraction < FMath::Min(ResolutionFraction * 0.97f, 1.0f)) ? StatOrange : StatGreen);
+
+				DrawTitleString(*DisplayName, /* UnitGraphColor = */ FColor::White);
+				if (!HeuristicSettings.IsEnabled())
+				{
+					InCanvas->DrawShadowedString(X2, InY, TEXT("OFF"), Font, FColor(160, 160, 160));
+				}
+				else if (HeuristicSettings.Model == DynamicRenderScaling::EHeuristicModel::Quadratic)
+				{
+					InCanvas->DrawShadowedString(X2, InY, *FString::Printf(TEXT("%3.1f%% x %3.1f%%"), ScreenPercentage, ScreenPercentage), Font, Color);
+				}
+				else
+				{
+					InCanvas->DrawShadowedString(X2, InY, *FString::Printf(TEXT("%3.1f%%"), ScreenPercentage), Font, Color);
+				}
+				InY += RowHeight;
+			}
+		}
+
 		// Draw calls
 		{
-				// Assume we don't have more than 1 GPU in mobile.
-				int32 NumDrawCalls = GNumDrawCallsRHI[0];
-			InCanvas->DrawShadowedString(X1, InY, TEXT("Draws:"), Font, bShowUnitTimeGraph ? FColor(100, 100, 255) : FColor::White);
+			// Assume we don't have more than 1 GPU in mobile.
+			int32 NumDrawCalls = GNumDrawCallsRHI[0];
+			DrawTitleString(TEXT("Draws"), /* UnitGraphColor = */ FColor(100, 100, 255));
 			InCanvas->DrawShadowedString(X2, InY, *FString::Printf(TEXT("%d"), NumDrawCalls), Font, StatGreen);
 			InY += RowHeight;
 		}
 			
 		// Primitives
 		{
-				// Assume we don't have more than 1 GPU in mobile.
-				int32 NumPrimitives = GNumPrimitivesDrawnRHI[0];
-			InCanvas->DrawShadowedString(X1, InY, TEXT("Prims:"), Font, bShowUnitTimeGraph ? FColor(100, 100, 255) : FColor::White);
+			// Assume we don't have more than 1 GPU in mobile.
+			int32 NumPrimitives = GNumPrimitivesDrawnRHI[0];
+			DrawTitleString(TEXT("Prims"), /* UnitGraphColor = */ FColor(100, 100, 255));
 			if (NumPrimitives < 10000)
 			{
 				InCanvas->DrawShadowedString(X2, InY, *FString::Printf(TEXT("%d"), NumPrimitives), Font, StatGreen);
@@ -797,7 +847,7 @@
 		const float GraphVerticalPixelsPerMS = TargetTimeMSHeight / TargetTimeMS;
 
 		// Scale dyn res so that RawMaxResolutionFraction is at MaxDynresTargetTimeMSHeight or below.
-		const float GraphVerticalPixelsPerResolutionFraction = FMath::Min(100.0f, MaxDynresTargetTimeMSHeight / (DynamicResolutionStateInfos.ResolutionFractionUpperBound * DynamicResolutionStateInfos.ResolutionFractionUpperBound));
+		const float GraphVerticalPixelsPerResolutionFraction = FMath::Min(100.0f, MaxDynresTargetTimeMSHeight / GDynamicPrimaryResolutionFraction.GetSettings().EstimateCostScale(DynamicResolutionStateInfos.ResolutionFractionUpperBounds[GDynamicPrimaryResolutionFraction]));
 
 		// Compute pulse effect for lines above alert threshold
 		const float AlertPulseFreq = 8.0f;
@@ -873,7 +923,7 @@
 			const FLinearColor LineColor(0.2f, 0.1f, 0.02f);
 			FVector StartPos(
 				GraphLeftXPos - 1.0f,
-				GraphBottomYPos - GraphVerticalPixelsPerResolutionFraction * DynamicResolutionStateInfos.ResolutionFractionUpperBound  * DynamicResolutionStateInfos.ResolutionFractionUpperBound ,
+				GraphBottomYPos - GraphVerticalPixelsPerResolutionFraction * GDynamicPrimaryResolutionFraction.GetSettings().EstimateCostScale(DynamicResolutionStateInfos.ResolutionFractionUpperBounds[GDynamicPrimaryResolutionFraction]),
 				0.0f);
 			FVector EndPos(
 				GraphLeftXPos + GraphHorizPixelsPerFrame * NumberOfSamples + GraphBackgroundMarginSize,
@@ -886,7 +936,7 @@
 				LineColor,
 				HitProxyId);
 
-			float MaxScreenPercentage = DynamicResolutionStateInfos.ResolutionFractionUpperBound  * 100.0f;
+			float MaxScreenPercentage = DynamicResolutionStateInfos.ResolutionFractionUpperBounds[GDynamicPrimaryResolutionFraction] * 100.0f;
 			InCanvas->DrawShadowedString(
 				EndPos.X + 4.0f,
 				EndPos.Y - AlertPrintHeight / 2,
@@ -894,7 +944,7 @@
 		}
 
 		// Screen percentage = 100% native line
-		if (DynamicResolutionStateInfos.ResolutionFractionUpperBound > 1.0f)
+		if (DynamicResolutionStateInfos.ResolutionFractionUpperBounds[GDynamicPrimaryResolutionFraction] > 1.0f)
 		{
 			const FLinearColor LineColor(0.2f, 0.1f, 0.02f);
 			FVector StartPos(
@@ -912,7 +962,7 @@
 				LineColor,
 				HitProxyId);
 
-			if (GraphVerticalPixelsPerResolutionFraction * (DynamicResolutionStateInfos.ResolutionFractionUpperBound  * DynamicResolutionStateInfos.ResolutionFractionUpperBound  - 1.0f) >= AlertPrintHeight)
+			if (GraphVerticalPixelsPerResolutionFraction * (GDynamicPrimaryResolutionFraction.GetSettings().EstimateCostScale(DynamicResolutionStateInfos.ResolutionFractionUpperBounds[GDynamicPrimaryResolutionFraction]) - 1.0f) >= AlertPrintHeight)
 			{
 				InCanvas->DrawShadowedString(EndPos.X + 4.0f, EndPos.Y - AlertPrintHeight / 2, TEXT("100.0% x 100.0% (native)"), SmallFont, LineColor);
 			}
@@ -965,57 +1015,62 @@
 			int32 DisplayPow = 1;
 			float DisplayMultiplier = 1.0f;
 			bool HigherIsBest = false;
-			switch (StatIndex)
+			if (StatIndex == EGS_Render)
 			{
-			case EGS_Render:
 				AbsoluteAlertValueThreshold = AlertTimeMS;
 				Values = RenderThreadTimes.GetData();
 				GraphVerticalPixelPerValue = GraphVerticalPixelsPerMS;
 				StatColor = FLinearColor(0.1f, 0.1f, 1.0f);		// Blue
-				break;
-
-			case EGS_Game:
+			}
+			else if (StatIndex == EGS_Game)
+			{
 				AbsoluteAlertValueThreshold = AlertTimeMS;
 				Values = GameThreadTimes.GetData();
 				GraphVerticalPixelPerValue = GraphVerticalPixelsPerMS;
 				StatColor = FLinearColor(1.0f, 0.1f, 0.1f);		// Red
-				break;
-
-			case EGS_GPU:
+			}
+			else if (StatIndex == EGS_GPU)
+			{
 				AbsoluteAlertValueThreshold = AlertTimeMS;
 				// Multi-GPU support : We don't support more than 1 GPU in stat unitgraph yet.
 				Values = GPUFrameTimes[0].GetData();
 				GraphVerticalPixelPerValue = GraphVerticalPixelsPerMS;
 				StatColor = FLinearColor(1.0f, 1.0f, 0.1f);		// Yellow
-				break;
-
-			case EGS_Frame:
+			}
+			else if (StatIndex == EGS_Frame)
+			{
 				AbsoluteAlertValueThreshold = AlertTimeMS;
 				Values = FrameTimes.GetData();
 				GraphVerticalPixelPerValue = GraphVerticalPixelsPerMS;
 				StatColor = FLinearColor(0.1f, 1.0f, 0.1f);		// Green
-				break;
-
-			case EGS_RHIT:
+			}
+			else if (StatIndex == EGS_RHIT)
+			{
 				AbsoluteAlertValueThreshold = AlertTimeMS;
 				Values = RHITTimes.GetData();
 				GraphVerticalPixelPerValue = GraphVerticalPixelsPerMS;
 				StatColor = FLinearColor(1.0f, 0.1f, 1.0f);		// Green
-				break;
+			}
+			else if (StatIndex == EGS_DynRes)
+			{
+				const DynamicRenderScaling::FBudget& Budget = GDynamicPrimaryResolutionFraction;
 
-			case EGS_DynRes:
 				AbsoluteAlertValueThreshold = AlertResolutionFraction;
 				RelativeAlertValueThreshold = 0.05;
-				Values = ResolutionFractions.GetData();
+				Values = ResolutionFractions[GDynamicPrimaryResolutionFraction].GetData();
 				GraphVerticalPixelPerValue = GraphVerticalPixelsPerResolutionFraction;
-				StatColor = FLinearColor(1.0f, 0.5f, 0.1f);		// Orange
-				DisplayPow = 2;
+				StatColor = FLinearColor(1.0f, 0.5f, 0.1f);
+				DisplayPow = Budget.GetSettings().Model == DynamicRenderScaling::EHeuristicModel::Quadratic ? 2 : 1;
 				DisplayMultiplier = 100.0f;
 				HigherIsBest = true;
 				AlertPrintY = GraphBottomYPos - AlertResolutionFraction * AlertResolutionFraction * GraphVerticalPixelsPerResolutionFraction + AlertPrintHeight;
-				break;
+			}
+			else
+			{
+				unimplemented();
 			}
 
+
 			// For each sample in our data set
 			for (int32 CurFrameIndex = 0; CurFrameIndex < NumberOfSamples; ++CurFrameIndex)
 			{
diff --git a/Engine/Source/Runtime/Engine/Private/UnrealEngine.cpp b/Engine/Source/Runtime/Engine/Private/UnrealEngine.cpp
--- a/Engine/Source/Runtime/Engine/Private/UnrealEngine.cpp
+++ b/Engine/Source/Runtime/Engine/Private/UnrealEngine.cpp
@@ -12274,8 +12274,8 @@
 void UEngine::GetDynamicResolutionCurrentStateInfos(FDynamicResolutionStateInfos& OutInfos) const
 {
 	OutInfos.Status = EDynamicResolutionStatus::Unsupported;
-	OutInfos.ResolutionFractionApproximation = -1.0f;
-	OutInfos.ResolutionFractionUpperBound = -1.0f;
+	OutInfos.ResolutionFractionApproximations.SetAll(-1.0f);
+	OutInfos.ResolutionFractionUpperBounds.SetAll(-1.0f);
 
 	#if WITH_DYNAMIC_RESOLUTION
 	{
@@ -12293,15 +12293,16 @@
 		else if (ForceResolutionFraction > 0)
 		{
 			OutInfos.Status = EDynamicResolutionStatus::DebugForceEnabled;
-			OutInfos.ResolutionFractionApproximation = ForceResolutionFraction;
-			OutInfos.ResolutionFractionUpperBound = DynamicResolutionState->GetResolutionFractionUpperBound();
+			OutInfos.ResolutionFractionApproximations = DynamicResolutionState->GetResolutionFractionsApproximation();
+			OutInfos.ResolutionFractionApproximations[GDynamicPrimaryResolutionFraction] = ForceResolutionFraction;
+			OutInfos.ResolutionFractionUpperBounds = DynamicResolutionState->GetResolutionFractionsUpperBound();
 		}
 		#endif
 		else if (DynamicResolutionState->IsEnabled())
 		{
 			OutInfos.Status = EDynamicResolutionStatus::Enabled;
-			OutInfos.ResolutionFractionApproximation = DynamicResolutionState->GetResolutionFractionApproximation();
-			OutInfos.ResolutionFractionUpperBound = DynamicResolutionState->GetResolutionFractionUpperBound();
+			OutInfos.ResolutionFractionApproximations = DynamicResolutionState->GetResolutionFractionsApproximation();
+			OutInfos.ResolutionFractionUpperBounds = DynamicResolutionState->GetResolutionFractionsUpperBound();
 		}
 		else if (bIsDynamicResolutionPaused)
 		{
@@ -12313,7 +12314,13 @@
 		}
 
 		// Min to ensure consistency when changing ResolutionFractionUpperBound setting.
-		OutInfos.ResolutionFractionApproximation = FMath::Min(OutInfos.ResolutionFractionApproximation, OutInfos.ResolutionFractionUpperBound);
+		for (TLinkedList<DynamicRenderScaling::FBudget*>::TIterator BudgetIt(DynamicRenderScaling::FBudget::GetGlobalList()); BudgetIt; BudgetIt.Next())
+		{
+			const DynamicRenderScaling::FBudget& Budget = **BudgetIt;
+			OutInfos.ResolutionFractionApproximations[Budget] = FMath::Min(
+				OutInfos.ResolutionFractionApproximations[Budget],
+				OutInfos.ResolutionFractionUpperBounds[Budget]);
+		}
 	}
 	#endif // !WITH_DYNAMIC_RESOLUTION
 }
diff --git a/Engine/Source/Runtime/Engine/Public/DynamicResolutionProxy.h b/Engine/Source/Runtime/Engine/Public/DynamicResolutionProxy.h
--- a/Engine/Source/Runtime/Engine/Public/DynamicResolutionProxy.h
+++ b/Engine/Source/Runtime/Engine/Public/DynamicResolutionProxy.h
@@ -1,12 +1,10 @@
 // Copyright Epic Games, Inc. All Rights Reserved.
 
-/*=============================================================================
-	StereoRendering.h: Abstract stereoscopic rendering interface
-=============================================================================*/
-
 #pragma once
 
 #include "CoreMinimal.h"
+#include "SceneView.h"
+#include "Engine/EngineTypes.h"
 
 
 /** Render thread proxy that hold's the heuristic for dynamic resolution. */
@@ -26,34 +24,34 @@
 	/** Create a new previous frame and feeds its timings and returns history's unique id. */
 	uint64 CreateNewPreviousFrameTimings_RenderThread(float GameThreadTimeMs, float RenderThreadTimeMs);
 
-	/** Commit GPU busy times
-	 *
-	 *	@param	HistoryFrameId The history's unique id for the frame returned by CreateNewPreviousFrameTimings_RenderThread()
-	 *	@param	TotalFrameGPUBusyTimeMs Total GPU busy time for the frame.
-	 *	@param	DynamicResolutionGPUBusyTimeMs GPU time for sub parts of the frame that does dynamic resolution.
-	 */
+	/** Commit GPU busy times */
 	void CommitPreviousFrameGPUTimings_RenderThread(
-		uint64 HistoryFrameId, float TotalFrameGPUBusyTimeMs, float DynamicResolutionGPUBusyTimeMs, bool bGPUTimingsHaveCPUBubbles);
+		uint64 HistoryFrameId,
+		float TotalFrameGPUBusyTimeMs,
+		float DynamicResolutionGPUBusyTimeMs,
+		const DynamicRenderScaling::TMap<float>& BudgetTimingMs);
 
 	/** Refresh resolution fraction's from history. */
 	void RefreshCurentFrameResolutionFraction_RenderThread();
 
-	/** Returns the view fraction upper bound. */
-	float GetResolutionFractionUpperBound() const;
-
 	/** Returns the view fraction that should be used for current frame. */
-	FORCEINLINE float QueryCurentFrameResolutionFraction_RenderThread() const
+	FORCEINLINE DynamicRenderScaling::TMap<float> QueryCurentFrameResolutionFractions() const
 	{
-		return FMath::Min(CurrentFrameResolutionFraction, GetResolutionFractionUpperBound());
+		check(IsInRenderingThread());
+		return QueryCurentFrameResolutionFractions_Internal();
 	}
 
 	/** Returns a non thread safe approximation of the current resolution fraction applied on render thread. */
-	FORCEINLINE float GetResolutionFractionApproximation_GameThread() const
+	FORCEINLINE DynamicRenderScaling::TMap<float> GetResolutionFractionsApproximation_GameThread() const
 	{
-		return FMath::Min(CurrentFrameResolutionFraction, GetResolutionFractionUpperBound());
+		check(IsInGameThread());
+		return QueryCurentFrameResolutionFractions_Internal();
 	}
 
 
+	/** Returns the view fraction upper bound. */
+	static DynamicRenderScaling::TMap<float> GetResolutionFractionUpperBounds();
+
 	/** Creates a default dynamic resolution state using this proxy that queries GPU timing from the RHI. */
 	static TSharedPtr< class IDynamicResolutionState > CreateDefaultState();
 
@@ -61,9 +59,6 @@
 private:
 	struct FrameHistoryEntry
 	{
-		// The resolution fraction the frame was rendered with.
-		float ResolutionFraction;
-
 		// Thread timings in milliseconds.
 		float GameThreadTimeMs;
 		float RenderThreadTimeMs;
@@ -74,18 +69,21 @@
 		// Total GPU busy time for the render thread commands that does dynamic resolutions.
 		float GlobalDynamicResolutionTimeMs;
 
-		// Whether GPU timings have GPU bubbles.
-		bool bGPUTimingsHaveCPUBubbles;
+		// Time for each individual timings
+		DynamicRenderScaling::TMap<float> BudgetTimingMs;
 
+		// The resolution fraction the frame was rendered with.
+		DynamicRenderScaling::TMap<float> ResolutionFractions;
 
 		inline FrameHistoryEntry()
-			: ResolutionFraction(1.0f)
-			, GameThreadTimeMs(-1.0f)
+			: GameThreadTimeMs(-1.0f)
 			, RenderThreadTimeMs(-1.0f)
 			, TotalFrameGPUBusyTimeMs(-1.0f)
 			, GlobalDynamicResolutionTimeMs(-1.0f)
-			, bGPUTimingsHaveCPUBubbles(true)
-		{ }
+		{
+			ResolutionFractions.SetAll(1.0f);
+			BudgetTimingMs.SetAll(-1.0f);
+		}
 
 
 		// Returns whether GPU timings have landed.
@@ -108,7 +106,8 @@
 	int32 IgnoreFrameRemainingCount;
 
 	// Current frame's view fraction.
-	float CurrentFrameResolutionFraction;
+	DynamicRenderScaling::TMap<float> CurrentFrameResolutionFractions;
+	DynamicRenderScaling::TMap<int32> BudgetHistorySizes;
 
 	// Frame counter to allocate unique ID for CommitPreviousFrameGPUTimings_RenderThread().
 	uint64 FrameCounter;
@@ -124,6 +123,8 @@
 		return History[(History.Num() + PreviousFrameIndex - BrowsingFrameId) % History.Num()];
 	}
 
+	DynamicRenderScaling::TMap<float> QueryCurentFrameResolutionFractions_Internal() const;
+
 	void ResetInternal();
 
 	void ResizeHistoryIfNeeded();
diff --git a/Engine/Source/Runtime/Engine/Public/DynamicResolutionState.h b/Engine/Source/Runtime/Engine/Public/DynamicResolutionState.h
--- a/Engine/Source/Runtime/Engine/Public/DynamicResolutionState.h
+++ b/Engine/Source/Runtime/Engine/Public/DynamicResolutionState.h
@@ -7,8 +7,14 @@
 #pragma once
 
 #include "CoreMinimal.h"
+#include "Engine/EngineTypes.h"
+#include "DynamicRenderScaling.h"
 
 
+/** Dynamic resolution of the primary ScreenPercentage. */
+extern ENGINE_API DynamicRenderScaling::FBudget GDynamicPrimaryResolutionFraction;
+
+
 /** Game thread events for dynamic resolution state. */
 enum class EDynamicResolutionStateEvent : uint8
 {
@@ -48,10 +54,10 @@
 protected:
 
 	/** Returns a non thread safe aproximation of the current resolution fraction applied on render thread, mostly used for stats and analytic. */
-	virtual float GetResolutionFractionApproximation() const = 0;
+	virtual DynamicRenderScaling::TMap<float> GetResolutionFractionsApproximation() const = 0;
 
 	/** Returns the max resolution resolution fraction. */
-	virtual float GetResolutionFractionUpperBound() const = 0;
+	virtual DynamicRenderScaling::TMap<float> GetResolutionFractionsUpperBound() const = 0;
 
 	/** Enables/Disables dynamic resolution. This is only called by GEngine automatically. */
 	virtual void SetEnabled(bool bEnable) = 0;
diff --git a/Engine/Source/Runtime/Engine/Public/LegacyScreenPercentageDriver.h b/Engine/Source/Runtime/Engine/Public/LegacyScreenPercentageDriver.h
--- a/Engine/Source/Runtime/Engine/Public/LegacyScreenPercentageDriver.h
+++ b/Engine/Source/Runtime/Engine/Public/LegacyScreenPercentageDriver.h
@@ -41,8 +41,8 @@
 
 
 	// Implements ISceneViewFamilyScreenPercentage
-	virtual float GetPrimaryResolutionFractionUpperBound() const override;
-	virtual float GetPrimaryResolutionFraction_RenderThread() const override;
+	virtual DynamicRenderScaling::TMap<float> GetResolutionFractionsUpperBound() const override;
+	virtual DynamicRenderScaling::TMap<float> GetResolutionFractions_RenderThread() const override;
 	virtual ISceneViewFamilyScreenPercentage* Fork_GameThread(const class FSceneViewFamily& ForkedViewFamily) const override;
 };
 
diff --git a/Engine/Source/Runtime/Engine/Public/SceneView.h b/Engine/Source/Runtime/Engine/Public/SceneView.h
--- a/Engine/Source/Runtime/Engine/Public/SceneView.h
+++ b/Engine/Source/Runtime/Engine/Public/SceneView.h
@@ -18,6 +18,7 @@
 #include "DebugViewModeHelpers.h"
 #include "RendererInterface.h"
 #include "Interfaces/Interface_PostProcessVolume.h"
+#include "DynamicRenderScaling.h"
 
 class FSceneView;
 class FSceneViewFamily;
@@ -1552,18 +1553,18 @@
 	}
 #endif
 
-protected:
 	/** 
 	 * Method to know the maximum value that can be returned by GetPrimaryResolutionFraction_RenderThread().
 	 * Can be called on game or rendering thread. This should return >= 1 if screen percentage show flag is disabled.
 	 */
-	virtual float GetPrimaryResolutionFractionUpperBound() const = 0;
+	virtual DynamicRenderScaling::TMap<float> GetResolutionFractionsUpperBound() const = 0;
 
+protected:
 	/**
 	 * Setup view family's view's screen percentage on rendering thread.
 	 * This should leave ResolutionFraction == 1 if screen percentage show flag is disabled.
 	 */
-	virtual float GetPrimaryResolutionFraction_RenderThread() const = 0;
+	virtual DynamicRenderScaling::TMap<float> GetResolutionFractions_RenderThread() const = 0;
 
 	/** Create a new screen percentage interface for a new view family. */
 	virtual ISceneViewFamilyScreenPercentage* Fork_GameThread(const class FSceneViewFamily& ViewFamily) const = 0;
@@ -1855,26 +1856,6 @@
 
 	FORCEINLINE bool AllowTranslucencyAfterDOF() const { return bAllowTranslucencyAfterDOF; }
 
-	/* Returns the maximum PrimaryResolutionFraction. */
-	FORCEINLINE float GetPrimaryResolutionFractionUpperBound() const
-	{
-		check(ScreenPercentageInterface != nullptr);
-		float PrimaryUpperBoundFraction = ScreenPercentageInterface->GetPrimaryResolutionFractionUpperBound();
-
-		checkf(ISceneViewFamilyScreenPercentage::IsValidResolutionFraction(PrimaryUpperBoundFraction),
-			TEXT("ISceneViewFamilyScreenPercentage::GetPrimaryResolutionFractionUpperBound()")
-			TEXT(" should return a valide value."));
-
-		if (!EngineShowFlags.ScreenPercentage)
-		{
-			checkf(PrimaryUpperBoundFraction >= 1.0f,
-				TEXT("ISceneViewFamilyScreenPercentage::GetPrimaryResolutionFractionUpperBound()")
-				TEXT(" should return >= 1 if screen percentage show flag is off."));
-		}
-
-		return PrimaryUpperBoundFraction;
-	}
-
 	FORCEINLINE const ISceneViewFamilyScreenPercentage* GetScreenPercentageInterface() const
 	{
 		return ScreenPercentageInterface;
diff --git a/Engine/Source/Runtime/Engine/Public/UnrealClient.h b/Engine/Source/Runtime/Engine/Public/UnrealClient.h
--- a/Engine/Source/Runtime/Engine/Public/UnrealClient.h
+++ b/Engine/Source/Runtime/Engine/Public/UnrealClient.h
@@ -17,6 +17,7 @@
 #include "HitProxies.h"
 #include "RenderGraphDefinitions.h"
 #include "Elements/Framework/TypedElementListFwd.h"
+#include "DynamicRenderScaling.h"
 
 class FCanvas;
 class FViewport;
@@ -256,7 +257,7 @@
 	TArray<float> FrameTimes;
 	TArray<float> RHITTimes;
 	TArray<float> InputLatencyTimes;
-	TArray<float> ResolutionFractions;
+	DynamicRenderScaling::TMap<TArray<float>> ResolutionFractions;
 #endif //!UE_BUILD_SHIPPING
 
 	FStatUnitData()
@@ -285,10 +286,14 @@
 		FrameTimes.AddZeroed(NumberOfSamples);
 		RHITTimes.AddZeroed(NumberOfSamples);
 		InputLatencyTimes.AddZeroed(NumberOfSamples);
-		ResolutionFractions.Reserve(NumberOfSamples);
-		for (int32 i = 0; i < NumberOfSamples; i++)
+		for (TLinkedList<DynamicRenderScaling::FBudget*>::TIterator BudgetIt(DynamicRenderScaling::FBudget::GetGlobalList()); BudgetIt; BudgetIt.Next())
 		{
-			ResolutionFractions.Add(-1.0f);
+			const DynamicRenderScaling::FBudget& Budget = **BudgetIt;
+			ResolutionFractions[Budget].Reserve(NumberOfSamples);
+			for (int32 i = 0; i < NumberOfSamples; i++)
+			{
+				ResolutionFractions[Budget].Add(-1.0f);
+			}
 		}
 #endif //!UE_BUILD_SHIPPING
 	}
diff --git a/Engine/Source/Runtime/RenderCore/Private/DynamicRenderScaling.cpp b/Engine/Source/Runtime/RenderCore/Private/DynamicRenderScaling.cpp
--- a/Engine/Source/Runtime/RenderCore/Private/DynamicRenderScaling.cpp
+++ b/Engine/Source/Runtime/RenderCore/Private/DynamicRenderScaling.cpp
@@ -123,6 +123,10 @@ FBudget::FBudget(const TCHAR* InName, FHeuristicSettings (*InHeuristicSettingsGe
 	, GlobalListLink(this)
 {
 	check(InName);
+
+	AnsiName.SetNumZeroed(FCString::Strlen(InName) + 1);
+	FCStringAnsi::Strcpy(AnsiName.GetData(), AnsiName.Num(), TCHAR_TO_ANSI(InName));
+
 	GlobalListLink.LinkHead(FBudget::GetGlobalList());
 	BudgetId = GRDSGlobalListLinkSize;
 	GRDSGlobalListLinkSize++;
diff --git a/Engine/Source/Runtime/RenderCore/Private/RenderGraphEvent.cpp b/Engine/Source/Runtime/RenderCore/Private/RenderGraphEvent.cpp
--- a/Engine/Source/Runtime/RenderCore/Private/RenderGraphEvent.cpp
+++ b/Engine/Source/Runtime/RenderCore/Private/RenderGraphEvent.cpp
@@ -290,7 +290,7 @@ public:
 		// Allocate new inflight frame.
 		if (bRecordThisFrame)
 		{
-			check(GRHISupportsGPUTimestampBubblesRemoval);
+			check(DynamicRenderScaling::IsSupported());
 
 			if (!QueryPool.IsValid())
 			{
@@ -464,7 +464,8 @@ namespace DynamicRenderScaling
 
 FRDGScope::FRDGScope(FRDGBuilder& InGraphBuilder, const DynamicRenderScaling::FBudget& InBudget)
 	: GraphBuilder(InGraphBuilder)
-	, bIsEnabled(GRDGTimingPool.IsRecordingThisFrame(InBudget))
+	, Budget(InBudget)
+	, bIsEnabled(GRDGTimingPool.IsRecordingThisFrame(InBudget) && !GraphBuilder.GPUScopeStacks.IsTimingScopeAlreadyEnabled(InBudget))
 {
 	if (bIsEnabled)
 	{
@@ -476,10 +477,15 @@ FRDGScope::~FRDGScope()
 {
 	if (bIsEnabled)
 	{
-		GraphBuilder.GPUScopeStacks.EndTimingScope();
+		GraphBuilder.GPUScopeStacks.EndTimingScope(Budget);
 	}
 }
 
+bool IsSupported()
+{
+	return GRHISupportsGPUTimestampBubblesRemoval;
+}
+
 void BeginFrame(const DynamicRenderScaling::TMap<bool>& bIsBudgetEnabled)
 {
 	GRDGTimingPool.BeginFrame(bIsBudgetEnabled);
diff --git a/Engine/Source/Runtime/RenderCore/Public/DynamicRenderScaling.h b/Engine/Source/Runtime/RenderCore/Public/DynamicRenderScaling.h
--- a/Engine/Source/Runtime/RenderCore/Public/DynamicRenderScaling.h
+++ b/Engine/Source/Runtime/RenderCore/Public/DynamicRenderScaling.h
@@ -87,6 +87,11 @@ public:
 		return Name;
 	}
 
+	inline const char* GetAnsiName() const
+	{
+		return AnsiName.GetData();
+	}
+
 	/* Returns the settings of the heuristic for this budget. */
 	inline const FHeuristicSettings& GetSettings() const
 	{
@@ -105,6 +110,7 @@ public:
 
 private:
 	const TCHAR* Name;
+	TArray<char> AnsiName;
 	FHeuristicSettings(*HeuristicSettingsGetter)(void);
 	TLinkedList<FBudget*> GlobalListLink;
 	FHeuristicSettings CachedSettings;
@@ -128,10 +134,12 @@ public:
 	TMap()
 	{
 		Array.SetNum(FBudget::GetGlobalListSize());
+#if !UE_BUILD_SHIPPING
 		for (TLinkedList<FBudget*>::TIterator BudgetIt(FBudget::GetGlobalList()); BudgetIt; BudgetIt.Next())
 		{
 			Array[BudgetIt->BudgetId].Name = BudgetIt->GetName();
 		}
+#endif
 	}
 
 	TMap(const Type& Value)
@@ -179,7 +187,9 @@ public:
 private:
 	struct FItem
 	{
+#if !UE_BUILD_SHIPPING
 		const TCHAR* Name = nullptr;
+#endif
 		Type Value;
 	};
 	TArray<FItem, TInlineAllocator<kInlineAllocatedBudgets>> Array;
@@ -201,6 +211,9 @@ inline float GetPercentageCVarToFraction(const TAutoConsoleVariable<float>& Perc
 	return PercentageToFraction(Percentage.GetValueOnAnyThread());
 }
 
+/** Returns whether the DynamicRenderScaling API is supported */
+RENDERCORE_API bool IsSupported();
+
 /** Updates all FBudget::GetSettings() with their new FHeuristicSettings. */
 RENDERCORE_API void UpdateHeuristicsSettings();
 
diff --git a/Engine/Source/Runtime/RenderCore/Public/RenderGraphEvent.h b/Engine/Source/Runtime/RenderCore/Public/RenderGraphEvent.h
--- a/Engine/Source/Runtime/RenderCore/Public/RenderGraphEvent.h
+++ b/Engine/Source/Runtime/RenderCore/Public/RenderGraphEvent.h
@@ -348,6 +348,7 @@ public:
 
 private:
 	FRDGBuilder& GraphBuilder;
+	const DynamicRenderScaling::FBudget& Budget;
 	const bool bIsEnabled;
 };
 
@@ -770,16 +771,27 @@ struct RENDERCORE_API FRDGGPUScopeStacksByPipeline
 	FRDGGPUScopeStacksByPipeline(FRDGAllocator& Allocator)
 		: Graphics(Allocator)
 		, AsyncCompute(Allocator)
-	{}
+	{
+		IsTimingIsEnabled.SetAll(false);
+	}
+
+	inline bool IsTimingScopeAlreadyEnabled(const DynamicRenderScaling::FBudget& Budget) const
+	{
+		return IsTimingIsEnabled[Budget];
+	}
 
 	inline void BeginTimingScope(const DynamicRenderScaling::FBudget& Budget)
 	{
+		check(!IsTimingIsEnabled[Budget]);
+		IsTimingIsEnabled[Budget] = true;
 		Graphics.Timing.BeginScope(Budget);
 		AsyncCompute.Timing.BeginScope(Budget);
 	}
 
-	inline void EndTimingScope()
+	inline void EndTimingScope(const DynamicRenderScaling::FBudget& Budget)
 	{
+		check(IsTimingIsEnabled[Budget]);
+		IsTimingIsEnabled[Budget] = false;
 		Graphics.Timing.EndScope();
 		AsyncCompute.Timing.EndScope();
 	}
@@ -830,6 +842,9 @@ struct RENDERCORE_API FRDGGPUScopeStacksByPipeline
 
 	FRDGGPUScopeStacks Graphics;
 	FRDGGPUScopeStacks AsyncCompute;
+
+private:
+	DynamicRenderScaling::TMap<bool> IsTimingIsEnabled;
 };
 
 //////////////////////////////////////////////////////////////////////////
diff --git a/Engine/Source/Runtime/Renderer/Private/SceneRendering.cpp b/Engine/Source/Runtime/Renderer/Private/SceneRendering.cpp
--- a/Engine/Source/Runtime/Renderer/Private/SceneRendering.cpp
+++ b/Engine/Source/Runtime/Renderer/Private/SceneRendering.cpp
@@ -77,6 +77,7 @@
 #include "Rendering/NaniteCoarseMeshStreamingManager.h"
 #include "Rendering/NaniteStreamingManager.h"
 #include "RectLightTextureManager.h"
+#include "DynamicResolutionState.h"
 
 /*-----------------------------------------------------------------------------
 	Globals
@@ -532,9 +533,14 @@ public:
 		check(MaxResolutionFraction > 0.0f);
 	}
 
-	virtual float GetPrimaryResolutionFractionUpperBound() const override
+	virtual DynamicRenderScaling::TMap<float> GetResolutionFractionsUpperBound() const override
 	{
-		return MaxResolutionFraction;
+		DynamicRenderScaling::TMap<float> ResolutionFractions;
+		if (ViewFamily.EngineShowFlags.ScreenPercentage)
+		{
+			ResolutionFractions[GDynamicPrimaryResolutionFraction] = MaxResolutionFraction;
+		}
+		return ResolutionFractions;
 	}
 
 	virtual ISceneViewFamilyScreenPercentage* Fork_GameThread(const class FSceneViewFamily& ForkedViewFamily) const override
@@ -550,16 +556,10 @@ public:
 			ForkedViewFamily, /* GlobalResolutionFraction = */ MaxResolutionFraction);
 	}
 
-	virtual float GetPrimaryResolutionFraction_RenderThread() const override
+	virtual DynamicRenderScaling::TMap<float> GetResolutionFractions_RenderThread() const override
 	{
 		check(IsInRenderingThread());
 
-		// Early return if no screen percentage should be done.
-		if (!ViewFamily.EngineShowFlags.ScreenPercentage)
-		{
-			return 1.0f;
-		}
-
 		uint32 FrameId = 0;
 
 		const FSceneViewState* ViewState = static_cast<const FSceneViewState*>(ViewFamily.Views[0]->State);
@@ -567,7 +567,14 @@ public:
 		{
 			FrameId = ViewState->GetFrameIndex(8);
 		}
-		return FrameId == 0 ? MaxResolutionFraction : FMath::Lerp(MinResolutionFraction, MaxResolutionFraction, 0.5f + 0.5f * FMath::Cos((FrameId + 0.25) * PI / 8));
+
+		DynamicRenderScaling::TMap<float> ResolutionFractions;
+		if (ViewFamily.EngineShowFlags.ScreenPercentage)
+		{
+			ResolutionFractions[GDynamicPrimaryResolutionFraction] =
+				FrameId == 0 ? MaxResolutionFraction : FMath::Lerp(MinResolutionFraction, MaxResolutionFraction, 0.5f + 0.5f * FMath::Cos((FrameId + 0.25) * PI / 8));
+		}
+		return ResolutionFractions;
 	}
 
 private:
@@ -2724,7 +2731,8 @@ FIntPoint FSceneRenderer::GetDesiredInternalBufferSize(const FSceneViewFamily& V
 		return DesiredBufferSize;
 	}
 
-	float PrimaryResolutionFractionUpperBound = ViewFamily.GetPrimaryResolutionFractionUpperBound();
+	DynamicRenderScaling::TMap<float> DynamicResolutionUpperBounds = ViewFamily.GetScreenPercentageInterface()->GetResolutionFractionsUpperBound();
+	float PrimaryResolutionFractionUpperBound = DynamicResolutionUpperBounds[GDynamicPrimaryResolutionFraction];
 
 	// Compute final resolution fraction.
 	float ResolutionFractionUpperBound = PrimaryResolutionFractionUpperBound * ViewFamily.SecondaryViewFraction;
@@ -2766,9 +2774,18 @@ void FSceneRenderer::PrepareViewRectsForRendering(FRHICommandListImmediate& RHIC
 {
 	check(IsInRenderingThread());
 
+	// Read the resolution data.
+	{
+		check(ActiveViewFamily->ScreenPercentageInterface);
+		ActiveViewFamily->DynamicResolutionUpperBounds = ActiveViewFamily->ScreenPercentageInterface->GetResolutionFractionsUpperBound();
+		ActiveViewFamily->DynamicResolutionFractions = ActiveViewFamily->ScreenPercentageInterface->GetResolutionFractions_RenderThread();
+	}
+
 	// If not supporting screen percentage, bypass all computation.
 	if (!ActiveViewFamily->SupportsScreenPercentage())
 	{
+		ActiveViewFamily->DynamicResolutionFractions[GDynamicPrimaryResolutionFraction] = 1.0f;
+
 		// The base pass have to respect FSceneView::UnscaledViewRect.
 		for (FViewInfo& View : Views)
 		{
@@ -2812,8 +2829,7 @@ void FSceneRenderer::PrepareViewRectsForRendering(FRHICommandListImmediate& RHIC
 		}
 	}
 
-	check(ActiveViewFamily->ScreenPercentageInterface);
-	float PrimaryResolutionFraction = ActiveViewFamily->ScreenPercentageInterface->GetPrimaryResolutionFraction_RenderThread();
+	float PrimaryResolutionFraction = ActiveViewFamily->DynamicResolutionFractions[GDynamicPrimaryResolutionFraction];
 	{
 		// Ensure screen percentage show flag is respected. Prefer to check() rather rendering at a differen screen percentage
 		// to make sure the renderer does not lie how a frame as been rendering to a dynamic resolution heuristic.
@@ -2823,7 +2839,7 @@ void FSceneRenderer::PrepareViewRectsForRendering(FRHICommandListImmediate& RHIC
 		}
 
 		// Make sure the screen percentage interface has not lied to the renderer about the upper bound.
-		checkf(PrimaryResolutionFraction <= ActiveViewFamily->GetPrimaryResolutionFractionUpperBound(),
+		checkf(PrimaryResolutionFraction <= ActiveViewFamily->DynamicResolutionUpperBounds[GDynamicPrimaryResolutionFraction],
 			TEXT("ISceneViewFamilyScreenPercentage::GetPrimaryResolutionFractionUpperBound() should not lie to the renderer."));
 
 		check(ISceneViewFamilyScreenPercentage::IsValidResolutionFraction(PrimaryResolutionFraction));
diff --git a/Engine/Source/Runtime/Renderer/Private/SceneRendering.h b/Engine/Source/Runtime/Renderer/Private/SceneRendering.h
--- a/Engine/Source/Runtime/Renderer/Private/SceneRendering.h
+++ b/Engine/Source/Runtime/Renderer/Private/SceneRendering.h
@@ -1857,6 +1857,10 @@ public:
 	FViewFamilyInfo(const FSceneViewFamily& InViewFamily);
 	virtual ~FViewFamilyInfo();
 
+	/** All the dynamic scaling informations */
+	DynamicRenderScaling::TMap<float> DynamicResolutionFractions;
+	DynamicRenderScaling::TMap<float> DynamicResolutionUpperBounds;
+
 	FMeshElementCollector MeshCollector;
 
 	FMeshElementCollector RayTracingCollector;
