World Partition: only pin dirty actors on save if we are holding the last hard references to avoid unloading.

#rnx

diff --git a/Engine/Source/Runtime/Engine/Private/WorldPartition/WorldPartition.cpp b/Engine/Source/Runtime/Engine/Private/WorldPartition/WorldPartition.cpp
--- a/Engine/Source/Runtime/Engine/Private/WorldPartition/WorldPartition.cpp
+++ b/Engine/Source/Runtime/Engine/Private/WorldPartition/WorldPartition.cpp
@@ -354,9 +354,14 @@ void UWorldPartition::OnPackageDirtyStateChanged(UPackage* Package)
 			{
 				DirtyActors.Add(ActorHandle, Actor);
 			}
-			else if (PinnedActors)
+			else
 			{
-				PinnedActors->AddActor(ActorHandle, /*bImmediate*/true);
+				// If we hold the last reference to that actor, pin it to avoid unloading
+				if (PinnedActors && (ActorHandle->GetHardRefCount() == 1))
+				{
+					PinnedActors->AddActor(ActorHandle, /*bImmediate*/true);
+				}
+
 				DirtyActors.Remove(ActorHandle);
 			}
 		}
