Fix some asserts when async loading thread is enabled in editor

#rnx

diff --git a/Engine/Source/Runtime/Engine/Private/Rendering/SkeletalMeshLODImporterData.cpp b/Engine/Source/Runtime/Engine/Private/Rendering/SkeletalMeshLODImporterData.cpp
--- a/Engine/Source/Runtime/Engine/Private/Rendering/SkeletalMeshLODImporterData.cpp
+++ b/Engine/Source/Runtime/Engine/Private/Rendering/SkeletalMeshLODImporterData.cpp
@@ -533,7 +533,7 @@
 
 void FReductionBaseSkeletalMeshBulkData::LoadReductionData(FSkeletalMeshLODModel& BaseLODModel, TMap<FString, TArray<FMorphTargetDelta>>& BaseLODMorphTargetData, UObject* Owner)
 {
-	check(IsInGameThread());
+	check(IsInGameThread() || IsInAsyncLoadingThread());
 
 	BaseLODMorphTargetData.Empty();
 	if (BulkData.GetElementCount() > 0)
diff --git a/Engine/Source/Runtime/Engine/Private/StaticMeshSourceData.cpp b/Engine/Source/Runtime/Engine/Private/StaticMeshSourceData.cpp
--- a/Engine/Source/Runtime/Engine/Private/StaticMeshSourceData.cpp
+++ b/Engine/Source/Runtime/Engine/Private/StaticMeshSourceData.cpp
@@ -301,7 +301,7 @@
 		if (StaticMeshDescriptionBulkData == nullptr)
 		{
 			// For now, we assume that serialization will only occur on the game thread
-			ensure(IsInGameThread());
+			ensure(IsInGameThread() || IsInAsyncLoadingThread());
 			StaticMeshDescriptionBulkData = NewObject<UStaticMeshDescriptionBulkData>(Owner, NAME_None, RF_Transactional);
 		}
 	}
diff --git a/Engine/Source/Runtime/Engine/Private/Texture.cpp b/Engine/Source/Runtime/Engine/Private/Texture.cpp
--- a/Engine/Source/Runtime/Engine/Private/Texture.cpp
+++ b/Engine/Source/Runtime/Engine/Private/Texture.cpp
@@ -160,7 +160,7 @@
 
 const FTextureResource* UTexture::GetResource() const
 {
-	if (IsInParallelGameThread() || IsInGameThread() || IsInSlateThread())
+	if (IsInParallelGameThread() || IsInGameThread() || IsInSlateThread() || IsInAsyncLoadingThread())
 	{
 		return PrivateResource;
 	}
@@ -175,7 +175,7 @@
 
 FTextureResource* UTexture::GetResource()
 {
-	if (IsInParallelGameThread() || IsInGameThread() || IsInSlateThread())
+	if (IsInParallelGameThread() || IsInGameThread() || IsInSlateThread() || IsInAsyncLoadingThread())
 	{
 		return PrivateResource;
 	}
