Add unshortened DDC key for static mesh in Insight context to easily compare them when recompiled more than once during cook

#rnx

diff --git a/Engine/Source/Runtime/Engine/Private/StaticMesh.cpp b/Engine/Source/Runtime/Engine/Private/StaticMesh.cpp
--- a/Engine/Source/Runtime/Engine/Private/StaticMesh.cpp
+++ b/Engine/Source/Runtime/Engine/Private/StaticMesh.cpp
@@ -47,6 +47,8 @@
 #include "Interfaces/ITargetPlatformManagerModule.h"
 #include "Engine/Engine.h"
 #include "EngineGlobals.h"
+#include "Trace/Trace.h"
+#include "Trace/Trace.inl"
 #include "HAL/LowLevelMemTracker.h"
 #include "DynamicMeshBuilder.h"
 #include "Math/UnrealMathUtility.h"
@@ -2769,6 +2771,10 @@
 	}
 }
 
+UE_TRACE_EVENT_BEGIN(Cpu, StaticMeshCache, NoSync)
+UE_TRACE_EVENT_FIELD(UE::Trace::WideString, KeySuffix)
+UE_TRACE_EVENT_END()
+
 void FStaticMeshRenderData::Cache(const ITargetPlatform* TargetPlatform, UStaticMesh* Owner, const FStaticMeshLODSettings& LODSettings)
 {
 	if (Owner->GetOutermost()->bIsCookedForEditor)
@@ -2786,6 +2792,10 @@
 		int32 NumLODs = Owner->GetNumSourceModels();
 		const FStaticMeshLODGroup& LODGroup = LODSettings.GetLODGroup(Owner->LODGroup);
 		const FString KeySuffix = BuildStaticMeshDerivedDataKeySuffix(TargetPlatform, Owner, LODGroup);
+#if CPUPROFILERTRACE_ENABLED
+		UE_TRACE_LOG_SCOPED_T(Cpu, StaticMeshCache, CpuChannel)
+			<< StaticMeshCache.KeySuffix(*KeySuffix);
+#endif
 		DerivedDataKey = BuildStaticMeshDerivedDataKey(KeySuffix);
 
 		using namespace UE::DerivedData;
@@ -2802,6 +2812,8 @@
 		FSharedBuffer MeshDataBuffer;
 		FIoHash NaniteStreamingDataHash;
 		{
+			TRACE_CPUPROFILER_EVENT_SCOPE(GetDDC);
+
 			FCacheRecordPolicyBuilder PolicyBuilder(ECachePolicy::Default | ECachePolicy::KeepAlive);
 			PolicyBuilder.AddValuePolicy(NaniteStreamingDataId, ECachePolicy::Default | ECachePolicy::SkipData);
 			
@@ -2962,6 +2974,8 @@
 			bool bSavedToDDC = false;
 			if (bSaveDDC)
 			{
+				TRACE_CPUPROFILER_EVENT_SCOPE(SaveDDC);
+
 				FValue Value = FValue::Compress(FSharedBuffer::MakeView(Ar.GetData(), Ar.TotalSize()));
 				RecordBuilder.AddValue(MeshDataId, Value);
 				TotalPushedBytes += Ar.TotalSize();
@@ -4592,6 +4606,8 @@
 #if WITH_EDITOR
 FStaticMeshRenderData& UStaticMesh::GetPlatformStaticMeshRenderData(UStaticMesh* Mesh, const ITargetPlatform* Platform)
 {
+	TRACE_CPUPROFILER_EVENT_SCOPE(UStaticMesh::GetPlatformStaticMeshRenderData);
+
 	check(Mesh && Mesh->GetRenderData());
 	const FStaticMeshLODSettings& PlatformLODSettings = Platform->GetStaticMeshLODSettings();
 	FString PlatformDerivedDataKey = BuildStaticMeshDerivedDataKey(
@@ -5103,6 +5119,8 @@
 
 void UStaticMesh::PrepareDerivedDataForActiveTargetPlatforms()
 {
+	TRACE_CPUPROFILER_EVENT_SCOPE(UStaticMesh::PrepareDerivedDataForActiveTargetPlatforms);
+
 	ITargetPlatformManagerModule& TargetPlatformManager = GetTargetPlatformManagerRef();
 	ITargetPlatform* RunningPlatform = TargetPlatformManager.GetRunningTargetPlatform();
 	const TArray<ITargetPlatform*>& TargetPlatforms = TargetPlatformManager.GetActiveTargetPlatforms();
