World Partition: go back to FindObject as actor descriptors for newly created actors won't resolve here.

#rnx

diff --git a/Engine/Source/Runtime/Engine/Private/WorldPartition/WorldPartitionRuntimeSpatialHash.cpp b/Engine/Source/Runtime/Engine/Private/WorldPartition/WorldPartitionRuntimeSpatialHash.cpp
--- a/Engine/Source/Runtime/Engine/Private/WorldPartition/WorldPartitionRuntimeSpatialHash.cpp
+++ b/Engine/Source/Runtime/Engine/Private/WorldPartition/WorldPartitionRuntimeSpatialHash.cpp
@@ -942,7 +942,7 @@ bool UWorldPartitionRuntimeSpatialHash::CreateStreamingGrid(const FSpatialHashRu
 								{
 									// This will load the actor if it isn't already loaded
 									FWorldPartitionReference Reference(WorldPartition, ActorInstance.Actor);
-									AActor* AlwaysLoadedActor = Reference->GetActor();
+									AActor* AlwaysLoadedActor = FindObject<AActor>(nullptr, *ActorDescView.GetActorPath().ToString());
 									AlwaysLoadedActorsForPIE.Emplace(Reference, AlwaysLoadedActor);
 
 									// Handle child actors
