World Partition: added immediate and deferred loading contexts support to handles automated tests.

#rnx

diff --git a/Engine/Source/Runtime/Engine/Private/Tests/WorldPartition/WorldPartitionHandleTests.cpp b/Engine/Source/Runtime/Engine/Private/Tests/WorldPartition/WorldPartitionHandleTests.cpp
--- a/Engine/Source/Runtime/Engine/Private/Tests/WorldPartition/WorldPartitionHandleTests.cpp
+++ b/Engine/Source/Runtime/Engine/Private/Tests/WorldPartition/WorldPartitionHandleTests.cpp
@@ -12,16 +12,24 @@
 namespace WorldPartitionTests
 {
 	constexpr const uint32 TestFlags = EAutomationTestFlags::ApplicationContextMask | EAutomationTestFlags::EngineFilter;
-
 	IMPLEMENT_SIMPLE_AUTOMATION_TEST(FWorldPartitionSoftRefTest, TEST_NAME_ROOT ".Handle", TestFlags)
-	bool FWorldPartitionSoftRefTest::RunTest(const FString& Parameters)
+
+	template <class LoadingContextType>
+	void PerformTests(FWorldPartitionSoftRefTest* Test)
 	{
-#if WITH_EDITOR
+		auto TestTrue = [Test](const TCHAR* What, bool Value) -> bool { return Test->TestTrue(What, Value); };
+		auto TestFalse = [Test](const TCHAR* What, bool Value) -> bool { return Test->TestFalse(What, Value); };
+
 		UActorDescContainer* ActorDescContainer = NewObject<UActorDescContainer>(GetTransientPackage());
 		ActorDescContainer->Initialize(nullptr, TEXT("/Engine/WorldPartition/WorldPartitionUnitTest"));
 
 		FWorldPartitionHandle Handle = FWorldPartitionHandle(ActorDescContainer, FGuid(TEXT("5D9F93BA407A811AFDDDAAB4F1CECC6A")));
-		FWorldPartitionReference Reference = FWorldPartitionReference(ActorDescContainer, FGuid(TEXT("0D2B04D240BE5DE58FE437A8D2DBF5C9")));
+
+		FWorldPartitionReference Reference;
+		{
+			LoadingContextType LoadingContext;
+			Reference = FWorldPartitionReference(ActorDescContainer, FGuid(TEXT("0D2B04D240BE5DE58FE437A8D2DBF5C9")));
+		}
 
 		TestTrue(TEXT("Handle container"), Handle->GetContainer() == ActorDescContainer);
 		TestTrue(TEXT("Reference container"), Reference->GetContainer() == ActorDescContainer);
@@ -230,6 +238,13 @@ namespace WorldPartitionTests
 			TestFalse(TEXT("Invalid container test"), Handle.IsValid());
 			TestFalse(TEXT("Invalid container test"), Reference.IsValid());
 		}
+	};
+
+	bool FWorldPartitionSoftRefTest::RunTest(const FString& Parameters)
+	{
+#if WITH_EDITOR
+		WorldPartitionTests::PerformTests<FWorldPartitionLoadingContext::FImmediate>(this);
+		WorldPartitionTests::PerformTests<FWorldPartitionLoadingContext::FDeferred>(this);
 #endif
 		return true;
 	}
