fix non-editor builds

#rnx

diff --git a/Engine/Source/Runtime/Engine/Private/Tests/WorldPartition/WorldPartitionHandleTests.cpp b/Engine/Source/Runtime/Engine/Private/Tests/WorldPartition/WorldPartitionHandleTests.cpp
--- a/Engine/Source/Runtime/Engine/Private/Tests/WorldPartition/WorldPartitionHandleTests.cpp
+++ b/Engine/Source/Runtime/Engine/Private/Tests/WorldPartition/WorldPartitionHandleTests.cpp
@@ -14,6 +14,7 @@ namespace WorldPartitionTests
 	constexpr const uint32 TestFlags = EAutomationTestFlags::ApplicationContextMask | EAutomationTestFlags::EngineFilter;
 	IMPLEMENT_SIMPLE_AUTOMATION_TEST(FWorldPartitionSoftRefTest, TEST_NAME_ROOT ".Handle", TestFlags)
 
+#if WITH_EDITOR
 	template <class LoadingContextType>
 	void PerformTests(FWorldPartitionSoftRefTest* Test)
 	{
@@ -239,6 +240,7 @@ namespace WorldPartitionTests
 			TestFalse(TEXT("Invalid container test"), Reference.IsValid());
 		}
 	};
+#endif
 
 	bool FWorldPartitionSoftRefTest::RunTest(const FString& Parameters)
 	{
