World Partition: fix missing editor loading package scope when loading external actors and merged some identical code.

#rnx

diff --git a/Engine/Source/Runtime/Engine/Private/WorldPartition/WorldPartitionHandle.cpp b/Engine/Source/Runtime/Engine/Private/WorldPartition/WorldPartitionHandle.cpp
--- a/Engine/Source/Runtime/Engine/Private/WorldPartition/WorldPartitionHandle.cpp
+++ b/Engine/Source/Runtime/Engine/Private/WorldPartition/WorldPartitionHandle.cpp
@@ -105,28 +105,34 @@
 		const FTransform& ContainerTransform = Container->GetInstanceTransform();
 		const FTransform* ContainerTransformPtr = ContainerTransform.Equals(FTransform::Identity) ? nullptr : &ContainerTransform;
 
+		auto CreateActorList = [](TArray<AActor*>& ActorList, const TArray<FWorldPartitionActorDesc*>& SourceList) -> ULevel*
+		{
+			ActorList.Empty(SourceList.Num());
+
+			ULevel* Level = nullptr;
+			for (FWorldPartitionActorDesc* ActorDesc : SourceList)
+			{
+				// When cleaning up worlds, actors are already marked as garbage at this point, so no need to remove them from the world
+				if (AActor* Actor = ActorDesc->GetActor(); IsValid(Actor))
+				{
+					ActorList.Add(Actor);
+
+					ULevel* ActorLevel = Actor->GetLevel();
+					check(!Level || (Level == ActorLevel));
+					Level = ActorLevel;
+				}
+			}
+
+			return Level;
+		};
+
 		if (ContainerOp.Registrations.Num())
 		{
 			TGuardValue<bool> IsEditorLoadingPackageGuard(GIsEditorLoadingPackage, true);
 
 			TArray<AActor*> ActorList;
-			ActorList.Reserve(ContainerOp.Registrations.Num());
-
-			ULevel* Level = nullptr;
-			for (FWorldPartitionActorDesc* ActorDesc : ContainerOp.Registrations)
+			if (ULevel* Level = CreateActorList(ActorList, ContainerOp.Registrations))
 			{
-				if (AActor* Actor = ActorDesc->GetActor())
-				{
-					ActorList.Add(Actor);
-
-					ULevel* ActorLevel = Actor->GetLevel();
-					check(!Level || (Level == ActorLevel));
-					Level = ActorLevel;
-				}
-			}
-			
-			if (Level)
-			{
 				Level->AddLoadedActors(ActorList, ContainerTransformPtr);
 			}
 		}
@@ -134,24 +140,8 @@
 		if (ContainerOp.Unregistrations.Num())
 		{
 			TArray<AActor*> ActorList;
-			ActorList.Reserve(ContainerOp.Registrations.Num());
-
-			ULevel* Level = nullptr;
-			for (FWorldPartitionActorDesc* ActorDesc : ContainerOp.Unregistrations)
+			if (ULevel* Level = CreateActorList(ActorList, ContainerOp.Unregistrations))
 			{
-				// When cleaning up worlds, actors are already marked as garbage at this point, so no need to remove them from the world
-				if (AActor* Actor = ActorDesc->GetActor(); IsValid(Actor))
-				{
-					ActorList.Add(Actor);
-
-					ULevel* ActorLevel = Actor->GetLevel();
-					check(!Level || (Level == ActorLevel));
-					Level = ActorLevel;
-				}
-			}
-			
-			if (Level)
-			{
 				Level->RemoveLoadedActors(ActorList, ContainerTransformPtr);
 			}
 
@@ -165,6 +155,8 @@
 
 void FWorldPartitionLoadingContext::FDeferred::RegisterActor(FWorldPartitionActorDesc* ActorDesc)
 {
+	TGuardValue<bool> IsEditorLoadingPackageGuard(GIsEditorLoadingPackage, true);
+
 	check(ActorDesc);
 
 	UActorDescContainer* Container = ActorDesc->GetContainer();
