Restored FTextureRenderTargetCubeResource to use a staging 2D texture as its render target. Replaced the render target clears with copies from the staging texture which mirrors the old implementation with CopyToResolveTarget. Also fixed another regression in CopyTexture that wasn't part of the unit test.

diff --git a/Engine/Source/Runtime/D3D12RHI/Private/D3D12Texture.cpp b/Engine/Source/Runtime/D3D12RHI/Private/D3D12Texture.cpp
--- a/Engine/Source/Runtime/D3D12RHI/Private/D3D12Texture.cpp
+++ b/Engine/Source/Runtime/D3D12RHI/Private/D3D12Texture.cpp
@@ -2909,13 +2909,16 @@ void FD3D12CommandContext::RHICopyTexture(FRHITexture* SourceTextureRHI, FRHITex
 
 	const FRHITextureDesc& SourceDesc = SourceTextureRHI->GetDesc();
 	const FRHITextureDesc& DestDesc = DestTextureRHI->GetDesc();
+	
+	const uint16 SourceArraySize = SourceDesc.ArraySize * (SourceDesc.IsTextureCube() ? 6 : 1);
+	const uint16 DestArraySize   = SourceDesc.ArraySize * (DestDesc.IsTextureCube()   ? 6 : 1);
 
 	const bool bAllPixels =
 		SourceDesc.GetSize() == DestDesc.GetSize() && (CopyInfo.Size == FIntVector::ZeroValue || CopyInfo.Size == SourceDesc.GetSize());
 
 	const bool bAllSubresources =
 		SourceDesc.NumMips   == DestDesc.NumMips   && SourceDesc.NumMips   == CopyInfo.NumMips    &&
-		SourceDesc.ArraySize == DestDesc.ArraySize && SourceDesc.ArraySize == CopyInfo.NumSlices;
+		SourceArraySize == DestArraySize && SourceArraySize == CopyInfo.NumSlices;
 
 	if (!bAllPixels || !bAllSubresources || bReadback)
 	{
diff --git a/Engine/Source/Runtime/Engine/Private/TextureRenderTargetCube.cpp b/Engine/Source/Runtime/Engine/Private/TextureRenderTargetCube.cpp
--- a/Engine/Source/Runtime/Engine/Private/TextureRenderTargetCube.cpp
+++ b/Engine/Source/Runtime/Engine/Private/TextureRenderTargetCube.cpp
@@ -263,6 +263,20 @@ void FTextureRenderTargetCubeResource::InitDynamicRHI()
 
 		RHIUpdateTextureReference(Owner->TextureReference.TextureReferenceRHI, TextureRHI);
 
+		// Create the RHI target surface used for rendering to
+		{
+			const FRHITextureCreateDesc Desc =
+				FRHITextureCreateDesc::Create2D(TEXT("FTextureRenderTargetCubeResource"))
+				.SetExtent(Owner->SizeX, Owner->SizeX)
+				.SetFormat(Owner->GetFormat())
+				.SetNumMips(Owner->GetNumMips())
+				.SetFlags(TexCreateFlags | ETextureCreateFlags::RenderTargetable | ETextureCreateFlags::ShaderResource)
+				.SetClearValue(FClearValueBinding(Owner->ClearColor))
+				.SetInitialState(ERHIAccess::SRVMask);
+
+			RenderTargetTextureRHI = RHICreateTexture(Desc);
+		}
+
 		AddToDeferredUpdateList(true);
 	}
 
@@ -306,14 +320,27 @@ void FTextureRenderTargetCubeResource::UpdateDeferredResource(FRHICommandListImm
 		return;
 	}
 
-	RHICmdList.Transition(FRHITransitionInfo(TextureRHI, ERHIAccess::Unknown, ERHIAccess::RTV));
+	RHICmdList.Transition(FRHITransitionInfo(RenderTargetTextureRHI, ERHIAccess::Unknown, ERHIAccess::RTV));
+
+	ClearRenderTarget(RHICmdList, RenderTargetTextureRHI);
+
+	RHICmdList.Transition({
+		FRHITransitionInfo(RenderTargetTextureRHI, ERHIAccess::RTV, ERHIAccess::CopySrc),
+		FRHITransitionInfo(TextureRHI, ERHIAccess::Unknown, ERHIAccess::CopyDest)
+	});
+
+	FRHICopyTextureInfo CopyInfo;
 
 	for(int32 FaceIdx = CubeFace_PosX; FaceIdx < CubeFace_MAX; FaceIdx++)
 	{
-		ClearRenderTarget(RHICmdList, TextureRHI, 0, FaceIdx);
+		CopyInfo.DestSliceIndex = FaceIdx;
+		RHICmdList.CopyTexture(RenderTargetTextureRHI, TextureRHI, CopyInfo);
 	}
 
-	RHICmdList.Transition(FRHITransitionInfo(TextureRHI, ERHIAccess::RTV, ERHIAccess::SRVMask));
+	RHICmdList.Transition({
+		FRHITransitionInfo(RenderTargetTextureRHI, ERHIAccess::CopySrc, ERHIAccess::SRVMask),
+		FRHITransitionInfo(TextureRHI, ERHIAccess::CopyDest, ERHIAccess::SRVMask)
+	});
 }
 
 /** 
diff --git a/Engine/Source/Runtime/Renderer/Private/SceneCaptureRendering.cpp b/Engine/Source/Runtime/Renderer/Private/SceneCaptureRendering.cpp
--- a/Engine/Source/Runtime/Renderer/Private/SceneCaptureRendering.cpp
+++ b/Engine/Source/Runtime/Renderer/Private/SceneCaptureRendering.cpp
@@ -1248,8 +1248,7 @@ void FScene::UpdateSceneCaptureContents(USceneCaptureComponentCube* CaptureCompo
 					[SceneRenderer, TextureRenderTarget, EventName, TargetFace](FRHICommandListImmediate& RHICmdList)
 					{
 						FRHICopyTextureInfo CopyInfo;
-						CopyInfo.SourceSliceIndex = TargetFace;
-						CopyInfo.DestSliceIndex   = TargetFace; 
+						CopyInfo.DestSliceIndex = TargetFace;
 						UpdateSceneCaptureContent_RenderThread(RHICmdList, SceneRenderer, TextureRenderTarget, TextureRenderTarget, EventName, CopyInfo, false, FGenerateMipsParams(), false, true, false);
 					}
 				);
diff --git a/Engine/Source/Runtime/Windows/D3D11RHI/Private/D3D11Texture.cpp b/Engine/Source/Runtime/Windows/D3D11RHI/Private/D3D11Texture.cpp
--- a/Engine/Source/Runtime/Windows/D3D11RHI/Private/D3D11Texture.cpp
+++ b/Engine/Source/Runtime/Windows/D3D11RHI/Private/D3D11Texture.cpp
@@ -2076,12 +2076,15 @@ void FD3D11DynamicRHI::RHICopyTexture(FRHITexture* SourceTextureRHI, FRHITexture
 	const FRHITextureDesc& SourceDesc = SourceTextureRHI->GetDesc();
 	const FRHITextureDesc& DestDesc = DestTextureRHI->GetDesc();
 
+	const uint16 SourceArraySize = SourceDesc.ArraySize * (SourceDesc.IsTextureCube() ? 6 : 1);
+	const uint16 DestArraySize   = SourceDesc.ArraySize * (DestDesc.IsTextureCube()   ? 6 : 1);
+
 	const bool bAllPixels =
 		SourceDesc.GetSize() == DestDesc.GetSize() && (CopyInfo.Size == FIntVector::ZeroValue || CopyInfo.Size == SourceDesc.GetSize());
 
 	const bool bAllSubresources =
 		SourceDesc.NumMips == DestDesc.NumMips && SourceDesc.NumMips == CopyInfo.NumMips &&
-		SourceDesc.ArraySize == DestDesc.ArraySize && SourceDesc.ArraySize == CopyInfo.NumSlices;
+		SourceArraySize == DestArraySize && SourceArraySize == CopyInfo.NumSlices;
 
 	if (!bAllPixels || !bAllSubresources)
 	{
