Fix for crash when virtual assets are missing a payload

diff --git a/Engine/Source/Runtime/Engine/Private/AudioDerivedData.cpp b/Engine/Source/Runtime/Engine/Private/AudioDerivedData.cpp
--- a/Engine/Source/Runtime/Engine/Private/AudioDerivedData.cpp
+++ b/Engine/Source/Runtime/Engine/Private/AudioDerivedData.cpp
@@ -1421,6 +1421,12 @@
 	const uint8* RawWaveData = (const uint8*)FutureBuffer.Get().GetData(); // Will block 
 	int32 RawDataSize = FutureBuffer.Get().GetSize();
 
+	if (!RawWaveData || RawDataSize <=0 )
+	{
+		UE_LOG(LogAudioDerivedData, Warning, TEXT("Cooking surround sound failed: %s, Failed to load virtualized bulkdata payload"), *Inputs.SoundFullName);
+		return;
+	}
+
 	// Front left channel is the master
 	static_assert(SPEAKER_FrontLeft == 0, "Front-left speaker must be first.");
 
