Fix a crash with InstancedPlacementPartitionActor serialization

diff --git a/Engine/Source/Runtime/Engine/Private/Instances/InstancedPlacementPartitionActor.cpp b/Engine/Source/Runtime/Engine/Private/Instances/InstancedPlacementPartitionActor.cpp
--- a/Engine/Source/Runtime/Engine/Private/Instances/InstancedPlacementPartitionActor.cpp
+++ b/Engine/Source/Runtime/Engine/Private/Instances/InstancedPlacementPartitionActor.cpp
@@ -54,12 +54,15 @@ void AInstancedPlacementPartitionActor::Serialize(FArchive& Ar)
 	Super::Serialize(Ar);
 
 #if WITH_EDITORONLY_DATA
-	Ar << PlacedClientInfo;
-
-	for (auto& Pair : PlacedClientInfo)
+	if (!Ar.IsFilterEditorOnly())
 	{
-		FClientPlacementInfo& ClientInfo = *Pair.Value;
-		ClientInfo.PostSerialize(Ar, this);
+		Ar << PlacedClientInfo;
+		
+		for (auto& Pair : PlacedClientInfo)
+		{
+			FClientPlacementInfo& ClientInfo = *Pair.Value;
+			ClientInfo.PostSerialize(Ar, this);
+		}
 	}
 #endif	// WITH_EDITORONLY_DATA
 }
