Fix invalid type aliasing in FPrimitiveSceneShaderData::Setup.

#rnx

diff --git a/Engine/Source/Runtime/Engine/Private/PrimitiveUniformShaderParameters.cpp b/Engine/Source/Runtime/Engine/Private/PrimitiveUniformShaderParameters.cpp
--- a/Engine/Source/Runtime/Engine/Private/PrimitiveUniformShaderParameters.cpp
+++ b/Engine/Source/Runtime/Engine/Private/PrimitiveUniformShaderParameters.cpp
@@ -206,14 +206,14 @@ void FPrimitiveSceneShaderData::Setup(const FPrimitiveUniformShaderParameters& P
 
 	// Set W directly in order to bypass NaN check, when passing int through FVector to shader.
 
-	Data[0].X	= *(const float*)&PrimitiveUniformShaderParameters.Flags;
-	Data[0].Y	= *(const float*)&PrimitiveUniformShaderParameters.InstanceSceneDataOffset;
-	Data[0].Z	= *(const float*)&PrimitiveUniformShaderParameters.NumInstanceSceneDataEntries;
-	Data[0].W	= *(const float*)&PrimitiveUniformShaderParameters.SingleCaptureIndex;
-
-	Data[1].X	= *(const float*)&PrimitiveUniformShaderParameters.TilePosition.X;
-	Data[1].Y	= *(const float*)&PrimitiveUniformShaderParameters.TilePosition.Y;
-	Data[1].Z	= *(const float*)&PrimitiveUniformShaderParameters.TilePosition.Z;
+	Data[0].X	= FMath::AsFloat(PrimitiveUniformShaderParameters.Flags);
+	Data[0].Y	= FMath::AsFloat(PrimitiveUniformShaderParameters.InstanceSceneDataOffset);
+	Data[0].Z	= FMath::AsFloat(PrimitiveUniformShaderParameters.NumInstanceSceneDataEntries);
+	Data[0].W	= FMath::AsFloat((uint32)PrimitiveUniformShaderParameters.SingleCaptureIndex);
+
+	Data[1].X	= PrimitiveUniformShaderParameters.TilePosition.X;
+	Data[1].Y	= PrimitiveUniformShaderParameters.TilePosition.Y;
+	Data[1].Z	= PrimitiveUniformShaderParameters.TilePosition.Z;
 	Data[1].W	= FMath::AsFloat(PrimitiveUniformShaderParameters.PrimitiveComponentId);
 
 	Data[2]		= *(const FVector4f*)&PrimitiveUniformShaderParameters.LocalToRelativeWorld.M[0][0];
@@ -240,18 +240,18 @@ void FPrimitiveSceneShaderData::Setup(const FPrimitiveUniformShaderParameters& P
 	Data[19]	= PrimitiveUniformShaderParameters.ObjectRelativeWorldPositionAndRadius;
 
 	Data[20]	= FVector4f(PrimitiveUniformShaderParameters.ActorRelativeWorldPosition, 0.0f);
-	Data[20].W	= *(const float*)&PrimitiveUniformShaderParameters.LightmapUVIndex;
+	Data[20].W	= FMath::AsFloat(PrimitiveUniformShaderParameters.LightmapUVIndex);
 
 	Data[21]	= FVector4f(PrimitiveUniformShaderParameters.ObjectOrientation, 0.0f);
-	Data[21].W	= *(const float*)&PrimitiveUniformShaderParameters.LightmapDataIndex;
+	Data[21].W	= FMath::AsFloat(PrimitiveUniformShaderParameters.LightmapDataIndex);
 
 	Data[22]	= PrimitiveUniformShaderParameters.NonUniformScale;
 
 	Data[23]	= FVector4f(PrimitiveUniformShaderParameters.PreSkinnedLocalBoundsMin, 0.0f);
-	Data[23].W	= *(const float*)&PrimitiveUniformShaderParameters.NaniteResourceID;
+	Data[23].W	= FMath::AsFloat(PrimitiveUniformShaderParameters.NaniteResourceID);
 
 	Data[24]	= FVector4f(PrimitiveUniformShaderParameters.PreSkinnedLocalBoundsMax, 0.0f);
-	Data[24].W	= *(const float*)&PrimitiveUniformShaderParameters.NaniteHierarchyOffset;
+	Data[24].W	= FMath::AsFloat(PrimitiveUniformShaderParameters.NaniteHierarchyOffset);
 
 	Data[25]	= FVector4f(PrimitiveUniformShaderParameters.LocalObjectBoundsMin, PrimitiveUniformShaderParameters.ObjectBoundsY);
 	Data[26]	= FVector4f(PrimitiveUniformShaderParameters.LocalObjectBoundsMax, PrimitiveUniformShaderParameters.ObjectBoundsZ);
@@ -259,17 +259,17 @@ void FPrimitiveSceneShaderData::Setup(const FPrimitiveUniformShaderParameters& P
 	Data[27].X = PrimitiveUniformShaderParameters.InstanceLocalBoundsCenter.X;
 	Data[27].Y = PrimitiveUniformShaderParameters.InstanceLocalBoundsCenter.Y;
 	Data[27].Z = PrimitiveUniformShaderParameters.InstanceLocalBoundsCenter.Z;
-	Data[27].W = *(const float*)&PrimitiveUniformShaderParameters.InstancePayloadDataOffset;
+	Data[27].W = FMath::AsFloat(PrimitiveUniformShaderParameters.InstancePayloadDataOffset);
 
 	Data[28].X = PrimitiveUniformShaderParameters.InstanceLocalBoundsExtent.X;
 	Data[28].Y = PrimitiveUniformShaderParameters.InstanceLocalBoundsExtent.Y;
 	Data[28].Z = PrimitiveUniformShaderParameters.InstanceLocalBoundsExtent.Z;
-	Data[28].W = *(const float*)&PrimitiveUniformShaderParameters.InstancePayloadDataStride;
+	Data[28].W = FMath::AsFloat(PrimitiveUniformShaderParameters.InstancePayloadDataStride);
 	
 	Data[29].X = PrimitiveUniformShaderParameters.WireframeColor.X;
 	Data[29].Y = PrimitiveUniformShaderParameters.WireframeColor.Y;
 	Data[29].Z = PrimitiveUniformShaderParameters.WireframeColor.Z;
-	Data[29].W = *(const float*)&PrimitiveUniformShaderParameters.NaniteImposterIndexAndFilterFlags;
+	Data[29].W = FMath::AsFloat(PrimitiveUniformShaderParameters.NaniteImposterIndexAndFilterFlags);
 
 	Data[30].X = PrimitiveUniformShaderParameters.LevelColor.X;
 	Data[30].Y = PrimitiveUniformShaderParameters.LevelColor.Y;
