Fix incorrect handling of compressed formats in FSlateTexture2DRHIRef::InitDynamicRHI.

#rnx

diff --git a/Engine/Source/Runtime/Engine/Private/Slate/SlateTextures.cpp b/Engine/Source/Runtime/Engine/Private/Slate/SlateTextures.cpp
--- a/Engine/Source/Runtime/Engine/Private/Slate/SlateTextures.cpp
+++ b/Engine/Source/Runtime/Engine/Private/Slate/SlateTextures.cpp
@@ -68,18 +68,19 @@ void FSlateTexture2DRHIRef::InitDynamicRHI()
 			uint32 Stride;
 			uint8* DestTextureData = (uint8*)RHILockTexture2D(ShaderResource, 0, RLM_WriteOnly, Stride, false);
 			const uint8* SourceTextureData = TextureData->GetRawBytes().GetData();
-			const uint32 DataStride = Width * GPixelFormats[PixelFormat].BlockBytes;
+
+			const uint32 BlocksX = CalcTextureMipWidthInBlocks(Width, PixelFormat, 0);
+			const uint32 BlocksY = CalcTextureMipWidthInBlocks(Height, PixelFormat, 0);
+			const uint32 DataStride = BlocksX * GPixelFormats[PixelFormat].BlockBytes;
+
 			if (Stride == DataStride)
 			{
-				FMemory::Memcpy(DestTextureData, SourceTextureData, DataStride * Height);
+				FMemory::Memcpy(DestTextureData, SourceTextureData, Stride * BlocksY);
 			}
 			else
 			{
-				checkf(GPixelFormats[PixelFormat].BlockSizeX == 1 
-					&& GPixelFormats[PixelFormat].BlockSizeY == 1 
-					&& GPixelFormats[PixelFormat].BlockSizeZ == 1,
-					TEXT("Tried to use compressed format?"));
-				for (uint32 i = 0; i < Height; i++)
+				checkf(DataStride < Stride, TEXT("Texture stride of %u is smaller than source data stride of %u, PixelFormat=%s (%d)"), Stride, DataStride, GPixelFormats[PixelFormat].Name, (int32)PixelFormat);
+				for (uint32 i = 0; i < BlocksY; i++)
 				{
 					FMemory::Memcpy(DestTextureData, SourceTextureData, DataStride);
 					DestTextureData += Stride;
