Added thorough Tick/Packet/RPC scope guarding and debugging to RPC DoS code, for tracking/fixing a very rare crash.

diff --git a/Engine/Source/Runtime/Engine/Private/Net/RPCDoSDetection.cpp b/Engine/Source/Runtime/Engine/Private/Net/RPCDoSDetection.cpp
--- a/Engine/Source/Runtime/Engine/Private/Net/RPCDoSDetection.cpp
+++ b/Engine/Source/Runtime/Engine/Private/Net/RPCDoSDetection.cpp
@@ -70,6 +70,18 @@ TAutoConsoleVariable<FString> CVarRPCDoSForcedRPCTracking(
 	TEXT("and a length of time for leaving tracking enabled (disables the next tick, otherwise).")
 	TEXT("Example (50% chance for 10 seconds): net.RPCDoSForcedRPCTracking=ServerAdmin,0.5,10"));
 
+#if RPC_DOS_SCOPE_DEBUG
+namespace UE::Net
+{
+	int32 GRPCDoSScopeDebugging = 0;
+
+	FAutoConsoleVariableRef CVarRPCDoSScopeDebugging(
+		TEXT("net.RPCDoSScopeDebugging"),
+		GRPCDoSScopeDebugging,
+		TEXT("Sets whether or not debugging/ensures for RPC DoS Tick/Packet scopes should be enabled."));
+}
+#endif
+
 
 /**
  * ERPCDoSEscalateReason
@@ -488,7 +500,7 @@ void FRPCDoSDetection::InitConfig(FName NetDriverName)
 	}
 }
 
-void FRPCDoSDetection::UpdateSeverity(ERPCDoSSeverityUpdate Update, ERPCDoSEscalateReason Reason)
+void FRPCDoSDetection::UpdateSeverity_Private(ERPCDoSSeverityUpdate Update, ERPCDoSEscalateReason Reason)
 {
 	bool bEscalate = Update == ERPCDoSSeverityUpdate::Escalate || Update == ERPCDoSSeverityUpdate::AutoEscalate;
 	int32 NewStateIdx = FMath::Clamp(ActiveState + (bEscalate ? 1 : -1), 0, DetectionSeverity.Num()-1);
@@ -509,6 +521,8 @@ void FRPCDoSDetection::UpdateSeverity(ERPCDoSSeverityUpdate Update, ERPCDoSEscal
 		CachePlayerAddress();
 		CachePlayerUID();
 
+		FTickScope& TickScope = GetTickScope();
+
 		if (bEscalate)
 		{
 			LastMetEscalationConditions = CurTime;
@@ -541,8 +555,8 @@ void FRPCDoSDetection::UpdateSeverity(ERPCDoSSeverityUpdate Update, ERPCDoSEscal
 					RecalculatePeriodHistory(CurStateTimePeriods, CurPerPeriodHistory, StartIdx);
 
 					// Determine if any time period quota's for the current SecondsDelta-offset CounterPerSecHistory were breached
-					if (PrevState.HasHitQuota_Count(CurPerPeriodHistory, FrameCounter) ||
-						PrevState.HasHitQuota_Time(CurPerPeriodHistory, FrameCounter))
+					if (PrevState.HasHitQuota_Count(CurPerPeriodHistory, TickScope.FrameCounter) ||
+						PrevState.HasHitQuota_Time(CurPerPeriodHistory, TickScope.FrameCounter))
 					{
 						// The state we're transitioning down into, would have last had its cooloff reset around this time
 						LastMetEscalationConditions = CurTime - (double)SecondsDelta;
@@ -576,8 +590,8 @@ void FRPCDoSDetection::UpdateSeverity(ERPCDoSSeverityUpdate Update, ERPCDoSEscal
 
 
 			// If escalating, keep escalating until the quota checks fail
-			if (bEscalate && (HasHitQuota_Count(CounterPerPeriodHistory, FrameCounter) ||
-								HasHitQuota_Time(CounterPerPeriodHistory, FrameCounter)))
+			if (bEscalate && (HasHitQuota_Count(CounterPerPeriodHistory, TickScope.FrameCounter) ||
+								HasHitQuota_Time(CounterPerPeriodHistory, TickScope.FrameCounter)))
 			{
 				NewStateIdx = FMath::Clamp(ActiveState + 1, 0, DetectionSeverity.Num()-1);
 
@@ -626,7 +640,7 @@ void FRPCDoSDetection::UpdateSeverity(ERPCDoSSeverityUpdate Update, ERPCDoSEscal
 			if (ActiveState > WorstAnalyticsState && !NewState.bEscalationConfirmed)
 			{
 				NewState.bEscalationConfirmed = NewState.EscalationCount >= NewState.EscalationCountTolerance ||
-					(NewState.EscalationTimeToleranceMS != -1 && FrameCounter.AccumRPCTime >= NewState.EscalationTimeToleranceSeconds);
+					(NewState.EscalationTimeToleranceMS != -1 && TickScope.FrameCounter.AccumRPCTime >= NewState.EscalationTimeToleranceSeconds);
 
 				if (NewState.bEscalationConfirmed)
 				{
@@ -660,22 +674,30 @@ void FRPCDoSDetection::UpdateSeverity(ERPCDoSSeverityUpdate Update, ERPCDoSEscal
 						// If RPC tracking was enabled during this escalation, but lightweight RPC tracking has not been reset,
 						// that means the RPC responsible for the escalation could not be identified,
 						// and analytics should list all potential RPC's
-						if (ShouldMonitorReceivedRPC() && LightweightRPCTracking.Count > 0)
+						if (ShouldMonitorReceivedRPC() && PacketScopePrivate.IsActive())
 						{
-							for (int32 LWIdx=0; LWIdx<LightweightRPCTracking.Count; LWIdx++)
-							{
-								UntimedRPCs.AddUnique(LightweightRPCTracking.RPC[LWIdx].Name);
-							}
+							FSequentialRPCScope& SequentialRPCScope = GetSequentialRPCScope();
+							FLightweightRPCTracking& LightweightRPCTracking = SequentialRPCScope.LightweightRPCTracking;
 
-							UntimedRPCs.Sort(
-								[](const FName& A, const FName& B)
+							if (SequentialRPCScope.LightweightRPCTracking.Count > 0)
+							{
+								for (int32 LWIdx=0; LWIdx<LightweightRPCTracking.Count; LWIdx++)
 								{
-									return A.ToString() < B.ToString();
-								});
+									UntimedRPCs.AddUnique(LightweightRPCTracking.RPC[LWIdx].Name);
+								}
+
+								UntimedRPCs.Sort(
+									[](const FName& A, const FName& B)
+									{
+										return A.ToString() < B.ToString();
+									});
+
+								LightweightRPCTracking.Count = 0;
 
-							LightweightRPCTracking.Count = 0;
+								FPacketScope& PacketScope = GetPacketScope();
 
-							RPCGroupTime = FPlatformTime::Seconds() - ReceivedPacketStartTime;
+								RPCGroupTime = FPlatformTime::Seconds() - PacketScope.ReceivedPacketStartTime;
+							}
 						}
 
 
@@ -733,10 +755,10 @@ void FRPCDoSDetection::PreTickDispatch(double TimeSeconds)
 	CSV_SCOPED_TIMING_STAT_EXCLUSIVE(RPCDoS_Checks);
 #endif
 
+	TickScopePrivate.SetActive(true);
+
 	if (bRPCDoSDetection)
 	{
-		FrameCounter.ResetRPCCounters();
-
 		NextTimeQuotaCheck = FMath::Max(TimeSeconds + TimeQuotaCheckInterval, NextTimeQuotaCheck);
 
 		if (HitchTimeQuotaMS > 0 && ReceivedPacketEndTime != 0.0)
@@ -762,15 +784,16 @@ void FRPCDoSDetection::PreTickDispatch(double TimeSeconds)
 
 		if (ActiveState > 0)
 		{
+			FTickScope& TickScope = GetTickScope();
 			double ActiveStateTime = TimeSeconds - LastMetEscalationConditions;
 
 			if (CooloffTime > 0 && ActiveStateTime > CooloffTime)
 			{
-				UpdateSeverity(ERPCDoSSeverityUpdate::Deescalate, ERPCDoSEscalateReason::Deescalate);
+				TickScope.UpdateSeverity(*this, ERPCDoSSeverityUpdate::Deescalate, ERPCDoSEscalateReason::Deescalate);
 			}
 			else if (AutoEscalateTime > 0 && ActiveStateTime > AutoEscalateTime)
 			{
-				UpdateSeverity(ERPCDoSSeverityUpdate::AutoEscalate, ERPCDoSEscalateReason::AutoEscalate);
+				TickScope.UpdateSeverity(*this, ERPCDoSSeverityUpdate::AutoEscalate, ERPCDoSEscalateReason::AutoEscalate);
 			}
 		}
 
@@ -844,7 +867,12 @@ void FRPCDoSDetection::PreTickDispatch(double TimeSeconds)
 void FRPCDoSDetection::PostSequentialRPC(EPostSequentialRPCType SequenceType, double TimeSeconds, FRPCDoSCounters* RPCCounter,
 											FRPCTrackingInfo* RPCTrackingInfo)
 {
-	const double StartTime = LastReceivedRPCTimeCache == 0.0 ? ReceivedPacketStartTime : LastReceivedRPCTimeCache;
+	using namespace UE::Net;
+
+	FPacketScope& PacketScope = GetPacketScope();
+	FSequentialRPCScope& SequentialRPCScope = GetSequentialRPCScope();
+	const double StartTime = SequentialRPCScope.LastReceivedRPCTimeCache == 0.0 ? PacketScope.ReceivedPacketStartTime :
+								SequentialRPCScope.LastReceivedRPCTimeCache;
 
 	RPCCounter->AccumRPCTime += (TimeSeconds - StartTime);
 
@@ -873,7 +901,7 @@ void FRPCDoSDetection::PostSequentialRPC(EPostSequentialRPCType SequenceType, do
 				{
 					const TSharedPtr<FRPCAnalytics>& CurAnalytics = RPCTrackingAnalytics[AnalyticsIdx];
 
-					if (CurAnalytics->RPCName == PostReceivedRPCName)
+					if (CurAnalytics->RPCName == SequentialRPCScope.PostReceivedRPCName)
 					{
 						RPCTrackingAnalyticsEntry = CurAnalytics;
 						break;
@@ -909,7 +937,7 @@ void FRPCDoSDetection::PostSequentialRPC(EPostSequentialRPCType SequenceType, do
 						RPCTrackingAnalyticsEntry = RPCTrackingAnalytics.Insert_GetRef(MoveTemp(NewRPCAnalytics), InsertIdx);
 
 						// Setup initial/persistent analytics values
-						RPCTrackingAnalyticsEntry->RPCName = PostReceivedRPCName;
+						RPCTrackingAnalyticsEntry->RPCName = SequentialRPCScope.PostReceivedRPCName;
 						RPCTrackingAnalyticsEntry->PlayerIP = GetPlayerAddress();
 						RPCTrackingAnalyticsEntry->PlayerUID = GetPlayerUID();
 					}
@@ -932,17 +960,17 @@ void FRPCDoSDetection::PostSequentialRPC(EPostSequentialRPCType SequenceType, do
 
 			if (RPCTrackingInfo->BlockState == ERPCBlockState::Blocked)
 			{
-				RawAnalyticsEntry->BlockedCount += PostReceivedRPCBlockCount;
+				RawAnalyticsEntry->BlockedCount += SequentialRPCScope.PostReceivedRPCBlockCount;
 			}
 
-			if (SequenceType == EPostSequentialRPCType::PostPacket && bReceivedPacketRPCUnique)
+			if (SequenceType == EPostSequentialRPCType::PostPacket && SequentialRPCScope.bReceivedPacketRPCUnique)
 			{
-				const double PacketProcessTime = TimeSeconds - ReceivedPacketStartTime;
+				const double PacketProcessTime = TimeSeconds - PacketScope.ReceivedPacketStartTime;
 
 				if (PacketProcessTime > RawAnalyticsEntry->MaxSinglePacketRPCTime)
 				{
 					RawAnalyticsEntry->MaxSinglePacketRPCTime = PacketProcessTime;
-					RawAnalyticsEntry->SinglePacketRPCCount = ReceivedPacketRPCCount;
+					RawAnalyticsEntry->SinglePacketRPCCount = SequentialRPCScope.ReceivedPacketRPCCount;
 					RawAnalyticsEntry->SinglePacketGameThreadCPU = static_cast<uint8>(FPlatformTime::GetThreadCPUTime().CPUTimePctRelative);
 				}
 			}
@@ -952,26 +980,42 @@ void FRPCDoSDetection::PostSequentialRPC(EPostSequentialRPCType SequenceType, do
 
 void FRPCDoSDetection::PostReceivedRPCPacket(double TimeSeconds)
 {
-	const double PacketProcessTime = TimeSeconds - ReceivedPacketStartTime;
+	using namespace UE::Net;
 
-	FrameCounter.AccumRPCTime += PacketProcessTime;
+	FTickScope& TickScope = GetTickScope();
+	FPacketScope& PacketScope = GetPacketScope();
+	FSequentialRPCScope& SequentialRPCScope = GetSequentialRPCScope();
+	const double PacketProcessTime = TimeSeconds - PacketScope.ReceivedPacketStartTime;
 
-	if (PostReceivedRPCCounter != nullptr)
+	TickScope.FrameCounter.AccumRPCTime += PacketProcessTime;
+
+	if (SequentialRPCScope.PostReceivedRPCCounter != nullptr)
 	{
-		PostSequentialRPC(EPostSequentialRPCType::PostPacket, TimeSeconds, PostReceivedRPCCounter, PostReceivedRPCTracking);
+#if RPC_DOS_SCOPE_DEBUG
+		// If this ensure fails, there is the potential for a crash - seen very rarely, and the Tick/Packet/SequentialRPC scoping should eliminate it
+		ensure(GRPCDoSScopeDebugging == 0 || (RPCTracking.Num() > 0 && ActiveRPCTracking.Num() > 0));
+#endif
+
+		if (RPCTracking.Num() > 0 && ActiveRPCTracking.Num() > 0)
+		{
+			PostSequentialRPC(EPostSequentialRPCType::PostPacket, TimeSeconds, SequentialRPCScope.PostReceivedRPCCounter,
+								SequentialRPCScope.PostReceivedRPCTracking);
 
 #if RPC_QUOTA_DEBUG
-		PostReceivedRPCCounter->DebugAccumRPCTime += (DebugReceivedRPCEndTime - DebugReceivedRPCStartTime);
+			PostReceivedRPCCounter->DebugAccumRPCTime += (DebugReceivedRPCEndTime - DebugReceivedRPCStartTime);
 #endif
-
-		PostReceivedRPCName = NAME_None;
-		PostReceivedRPCTracking = nullptr;
-		PostReceivedRPCCounter = nullptr;
-		PostReceivedRPCBlockCount = 0;
-		LastReceivedRPCTimeCache = 0.0;
+		}
 	}
 
-	CondCheckTimeQuota(TimeSeconds);
+	TickScope.CondCheckTimeQuota(*this, TimeSeconds);
+
+#if RPC_DOS_SCOPE_DEBUG
+	ensure(GRPCDoSScopeDebugging == 0 || SequentialRPCScopePrivate.IsActive());
+	ensure(GRPCDoSScopeDebugging == 0 || PacketScopePrivate.IsActive());
+	ensure(GRPCDoSScopeDebugging == 0 || TickScopePrivate.IsActive());
+#endif
+
+	SequentialRPCScopePrivate.SetActive(false);
 }
 
 void FRPCDoSDetection::PostTickDispatch()
@@ -980,11 +1024,15 @@ void FRPCDoSDetection::PostTickDispatch()
 	CSV_SCOPED_TIMING_STAT_EXCLUSIVE(RPCDoS_Checks);
 #endif
 
+	FTickScope& TickScope = GetTickScope();
+	FRPCDoSCounters& FrameCounter = TickScope.FrameCounter;
+
 	if (FrameCounter.RPCCounter > 0)
 	{
 		SecondCounter.AccumulateCounter(FrameCounter);
-		FrameCounter.ResetRPCCounters();
 	}
+
+	TickScopePrivate.SetActive(false);
 }
 
 void FRPCDoSDetection::NotifyClose()
@@ -1023,7 +1071,8 @@ ERPCNotifyResult FRPCDoSDetection::CheckRPCTracking(UFunction* Function, FName F
 
 	if (!bHitchSuspendDetection)
 	{
-		FRPCTrackingInfo& Tracking = FindOrAddRPCTracking(Function, bNewTracking);
+		FSequentialRPCScope& SequentialRPCScope = GetSequentialRPCScope();
+		FRPCTrackingInfo& Tracking = SequentialRPCScope.FindOrAddRPCTracking(*this, Function, bNewTracking);
 		const uint8 OldTrackedSecondIncrement = Tracking.LastTrackedSecondIncrement;
 
 		Tracking.LastTrackedSecondIncrement = SecondsIncrementer;
@@ -1074,40 +1123,43 @@ ERPCNotifyResult FRPCDoSDetection::CheckRPCTracking(UFunction* Function, FName F
 
 		CurActiveCounter.RPCCounter++;
 
-		if (bReceivedPacketRPCUnique && PostReceivedRPCName != NAME_None && PostReceivedRPCName != FunctionName)
+		if (SequentialRPCScope.bReceivedPacketRPCUnique && SequentialRPCScope.PostReceivedRPCName != NAME_None
+			&& SequentialRPCScope.PostReceivedRPCName != FunctionName)
 		{
-			bReceivedPacketRPCUnique = false;
+			SequentialRPCScope.bReceivedPacketRPCUnique = false;
 		}
 
 
 		// To minimize timestamp generation, use the 'free' pre/post packet timestamps as much as possible,
 		// and group sequential calls to the same RPC together instead of timing for each one
-		if (PostReceivedRPCCounter == nullptr)
+		if (SequentialRPCScope.PostReceivedRPCCounter == nullptr)
 		{
-			PostReceivedRPCName = FunctionName;
-			PostReceivedRPCTracking = &Tracking;
-			PostReceivedRPCCounter = &CurActiveCounter;
+			SequentialRPCScope.PostReceivedRPCName = FunctionName;
+			SequentialRPCScope.PostReceivedRPCTracking = &Tracking;
+			SequentialRPCScope.PostReceivedRPCCounter = &CurActiveCounter;
 		}
-		else if (PostReceivedRPCCounter != &CurActiveCounter)
+		else if (SequentialRPCScope.PostReceivedRPCCounter != &CurActiveCounter)
 		{
 			const double CurTimeSeconds = FPlatformTime::Seconds();
 
-			PostSequentialRPC(EPostSequentialRPCType::MidPacket, CurTimeSeconds, PostReceivedRPCCounter, PostReceivedRPCTracking);
+			PostSequentialRPC(EPostSequentialRPCType::MidPacket, CurTimeSeconds, SequentialRPCScope.PostReceivedRPCCounter,
+								SequentialRPCScope.PostReceivedRPCTracking);
 
-			LastReceivedRPCTimeCache = CurTimeSeconds;
+			SequentialRPCScope.LastReceivedRPCTimeCache = CurTimeSeconds;
 
 #if RPC_QUOTA_DEBUG
-			PostReceivedRPCCounter->DebugAccumRPCTime += (DebugReceivedRPCEndTime - DebugReceivedRPCStartTime);
+			SequentialRPCScope.PostReceivedRPCCounter->DebugAccumRPCTime +=
+				(SequentialRPCScope.DebugReceivedRPCEndTime - SequentialRPCScope.DebugReceivedRPCStartTime);
 #endif
 
-			PostReceivedRPCName = FunctionName;
-			PostReceivedRPCTracking = &Tracking;
-			PostReceivedRPCCounter = &CurActiveCounter;
-			PostReceivedRPCBlockCount = 0;
+			SequentialRPCScope.PostReceivedRPCName = FunctionName;
+			SequentialRPCScope.PostReceivedRPCTracking = &Tracking;
+			SequentialRPCScope.PostReceivedRPCCounter = &CurActiveCounter;
+			SequentialRPCScope.PostReceivedRPCBlockCount = 0;
 		}
 
 		// When tracking is enabled, the quota's are checked every RPC call until blocked (and analytics in the PostSequential call)
-		if (!bNewTracking && PostReceivedRPCBlockCount == 0 && Tracking.BlockState != ERPCBlockState::OnAllowList)
+		if (!bNewTracking && SequentialRPCScope.PostReceivedRPCBlockCount == 0 && Tracking.BlockState != ERPCBlockState::OnAllowList)
 		{
 			auto HasHitQuota = [&Tracking](int32 TimePeriod, int32 CountPerPeriod, double SecsPerPeriod)
 				{
@@ -1127,7 +1179,7 @@ ERPCNotifyResult FRPCDoSDetection::CheckRPCTracking(UFunction* Function, FName F
 
 		if (Tracking.BlockState == ERPCBlockState::Blocked)
 		{
-			PostReceivedRPCBlockCount++;
+			SequentialRPCScope.PostReceivedRPCBlockCount++;
 
 			Result = ERPCNotifyResult::BlockRPC;
 		}
@@ -1138,7 +1190,9 @@ ERPCNotifyResult FRPCDoSDetection::CheckRPCTracking(UFunction* Function, FName F
 
 		if (TrackingPtr != nullptr && TrackingPtr->IsValid() && TrackingPtr->Get()->BlockState == ERPCBlockState::Blocked)
 		{
-			PostReceivedRPCBlockCount++;
+			FSequentialRPCScope& SequentialRPCScope = GetSequentialRPCScope();
+
+			SequentialRPCScope.PostReceivedRPCBlockCount++;
 
 			Result = ERPCNotifyResult::BlockRPC;
 		}
@@ -1183,9 +1237,12 @@ void FRPCDoSDetection::EnableRPCTracking(double TimeSeconds)
 		LastRPCTrackingClean = TimeSeconds;
 	}
 
-	// Do checks that should occur in the middle of receiving a packet (detected based on which received packet timestamps are fresh)
-	if (ReceivedPacketEndTime < ReceivedPacketStartTime)
+	// Do checks that should occur in the middle of receiving a packet
+	if (SequentialRPCScopePrivate.IsActive())
 	{
+		FSequentialRPCScope& SequentialRPCScope = GetSequentialRPCScope();
+		FLightweightRPCTracking& LightweightRPCTracking = SequentialRPCScope.LightweightRPCTracking;
+
 		// If lightweight tracking has entries, see if we can retroactively add this to full RPC tracking (only possible with single RPC)
 		// NOTE: Lightweight tracking UFunction's are valid, here.
 		bool bUniqueLightweightRPC = LightweightRPCTracking.Count > 0;
@@ -1203,17 +1260,17 @@ void FRPCDoSDetection::EnableRPCTracking(double TimeSeconds)
 		if (bUniqueLightweightRPC)
 		{
 			// Add and time the lightweight-tracked RPC, from the start of packet receive
-			CheckRPCTracking(FirstFunc, FirstFunc->GetFName());
+			SequentialRPCScope.CheckRPCTracking(*this, FirstFunc, FirstFunc->GetFName());
 		}
 		else
 		{
 			// Ensure tracking knows the packet contains more than one RPC type
-			bReceivedPacketRPCUnique = false;
+			SequentialRPCScope.bReceivedPacketRPCUnique = false;
 
 			// Very important to update the cached time for last received RPC, to avoid mistiming RPC's when tracking is enabled mid-receive
-			if (LastReceivedRPCTimeCache == 0.0)
+			if (SequentialRPCScope.LastReceivedRPCTimeCache == 0.0)
 			{
-				LastReceivedRPCTimeCache = TimeSeconds;
+				SequentialRPCScope.LastReceivedRPCTimeCache = TimeSeconds;
 			}
 		}
 	}
@@ -1226,7 +1283,9 @@ void FRPCDoSDetection::EnableForcedRPCTracking(UFunction* Function, FName Functi
 	EnableRPCTracking(TimeSeconds);
 
 	// If lightweight RPC tracking was not made up of one unique RPC, this RPC will not have been added to tracking - add it now
-	if (!bReceivedPacketRPCUnique)
+	FSequentialRPCScope& SequentialRPCScope = GetSequentialRPCScope();
+
+	if (!SequentialRPCScope.bReceivedPacketRPCUnique)
 	{
 		CheckRPCTracking(Function, FunctionName);
 	}
diff --git a/Engine/Source/Runtime/Engine/Public/Net/RPCDoSDetection.h b/Engine/Source/Runtime/Engine/Public/Net/RPCDoSDetection.h
--- a/Engine/Source/Runtime/Engine/Public/Net/RPCDoSDetection.h
+++ b/Engine/Source/Runtime/Engine/Public/Net/RPCDoSDetection.h
@@ -18,6 +18,9 @@
 /** Whether or not CSV stats should be enabled for the RPC DoS checks (development only - for stress testing the RPC DoS code) */
 #define RPC_DOS_DEV_STATS 0
 
+/** Tick/Packet scope correctness debugging */
+#define RPC_DOS_SCOPE_DEBUG 1
+
 
 // Typedefs
 
@@ -41,6 +44,17 @@ using FGetRPCDoSPlayerUID = TUniqueFunction<FString()>;
 using FRPCDoSKickPlayer = TUniqueFunction<void()>;
 
 
+// Globals/CVars
+
+#if RPC_DOS_SCOPE_DEBUG
+namespace UE::Net
+{
+	/** Whether or not debugging/ensures for RPC DoS Tick/Packet scopes should be enabled */
+	extern int32 GRPCDoSScopeDebugging;
+}
+#endif
+
+
 // Structs
 
 /**
@@ -457,7 +471,11 @@ public:
 	 * @param Update	Whether or not we are escalating or de-escalating the severity state
 	 * @param Reason	The reason for the escalation change
 	 */
-	void UpdateSeverity(ERPCDoSSeverityUpdate Update, ERPCDoSEscalateReason Reason);
+	UE_DEPRECATED(5.2, "UpdateSeverity will be made private soon.")
+	void UpdateSeverity(ERPCDoSSeverityUpdate Update, ERPCDoSEscalateReason Reason)
+	{
+		UpdateSeverity_Private(Update, Reason);
+	}
 
 
 	/**
@@ -474,15 +492,27 @@ public:
 	 */
 	void PreReceivedPacket(double TimeSeconds)
 	{
+		using namespace UE::Net;
+
 #if RPC_DOS_DEV_STATS
 		CSV_SCOPED_TIMING_STAT_EXCLUSIVE(RPCDoS_Checks);
 #endif
 
-		ReceivedPacketStartTime = TimeSeconds;
-		bPacketContainsRPC = false;
-		bReceivedPacketRPCUnique = true;
-		ReceivedPacketRPCCount = 0;
-		LightweightRPCTracking.Count = 0;
+#if RPC_DOS_SCOPE_DEBUG
+		ensure(GRPCDoSScopeDebugging == 0 || TickScopePrivate.IsActive());
+		ensure(GRPCDoSScopeDebugging == 0 || !SequentialRPCScopePrivate.IsActive());
+#endif
+
+		PacketScopePrivate.SetActive(true);
+
+		FPacketScope& PacketScope = GetPacketScope();
+
+		PacketScope.ReceivedPacketStartTime = TimeSeconds;
+
+		// Redundant call to Reset for sequential RPC scope - remove when scope debugging verifies the sequential RPC scope is never active here
+#if 1 // !RPC_DOS_SCOPE_DEBUG
+		SequentialRPCScopePrivate.Reset();
+#endif
 	}
 
 	/**
@@ -504,14 +534,26 @@ public:
 		CSV_SCOPED_TIMING_STAT_EXCLUSIVE(RPCDoS_Checks);
 #endif
 
-		if (bRPCTrackingEnabled)
+		if (bRPCDoSDetection && !bHitchSuspendDetection)
 		{
-			Result = CheckRPCTracking(Function, FunctionName);
-		}
+			if (!SequentialRPCScopePrivate.IsActive())
+			{
+				PreSequentialRPC();
+			}
+
+			if (bRPCTrackingEnabled)
+			{
+				FSequentialRPCScope& SequentialRPCScope = GetSequentialRPCScope();
+
+				Result = SequentialRPCScope.CheckRPCTracking(*this, Function, FunctionName);
+			}
 
 #if RPC_QUOTA_DEBUG
-		DebugReceivedRPCStartTime = FPlatformTime::Seconds();
+			FSequentialRPCScope& SequentialRPCScope = GetSequentialRPCScope();
+
+			SequentialRPCScope.DebugReceivedRPCStartTime = FPlatformTime::Seconds();
 #endif
+		}
 
 
 		return Result;
@@ -531,16 +573,26 @@ public:
 		CSV_SCOPED_TIMING_STAT_EXCLUSIVE(RPCDoS_Checks);
 #endif
 
-		if (UNLIKELY(ForcedRPCTracking == FunctionName && FMath::FRand() <= ForcedRPCTrackingChance))
-		{
-			EnableForcedRPCTracking(Function, FunctionName, FPlatformTime::Seconds());
-		}
-		else
+		if (bRPCDoSDetection && !bHitchSuspendDetection)
 		{
-			FLightweightRPCTracking::FLightweightRPCEntry& CurEntry = LightweightRPCTracking.RPC[LightweightRPCTracking.Count++];
+			if (!SequentialRPCScopePrivate.IsActive())
+			{
+				PreSequentialRPC();
+			}
+
+			if (UNLIKELY(ForcedRPCTracking == FunctionName && FMath::FRand() <= ForcedRPCTrackingChance))
+			{
+				EnableForcedRPCTracking(Function, FunctionName, FPlatformTime::Seconds());
+			}
+			else
+			{
+				FSequentialRPCScope& SequentialRPCScope = GetSequentialRPCScope();
+				FLightweightRPCTracking& LightweightRPCTracking = SequentialRPCScope.LightweightRPCTracking;
+				FLightweightRPCTracking::FLightweightRPCEntry& CurEntry = LightweightRPCTracking.RPC[LightweightRPCTracking.Count++];
 
-			CurEntry.Function = Function;
-			CurEntry.Name = FunctionName;
+				CurEntry.Function = Function;
+				CurEntry.Name = FunctionName;
+			}
 		}
 	}
 
@@ -554,19 +606,27 @@ public:
 #endif
 
 #if RPC_QUOTA_DEBUG
-		DebugReceivedRPCEndTime = FPlatformTime::Seconds();
-		FrameCounter.DebugAccumRPCTime += (DebugReceivedRPCEndTime - DebugReceivedRPCStartTime);
+		{
+			FPacketScope& PacketScope = GetPacketScope();
+
+			PacketScope.DebugReceivedRPCEndTime = FPlatformTime::Seconds();
+			FrameCounter.DebugAccumRPCTime += (PacketScope.DebugReceivedRPCEndTime - PacketScope.DebugReceivedRPCStartTime);
+		}
 #endif
 
 		if (bRPCDoSDetection && !bHitchSuspendDetection)
 		{
-			FrameCounter.RPCCounter++;
+			FTickScope& TickScope = GetTickScope();
+			FPacketScope& PacketScope = GetPacketScope();
+			FSequentialRPCScope& SequentialRPCScope = GetSequentialRPCScope();
+
+			TickScope.FrameCounter.RPCCounter++;
 
 			RPCIntervalCounter++;
-			ReceivedPacketRPCCount++;
-			bPacketContainsRPC = true;
+			SequentialRPCScope.ReceivedPacketRPCCount++;
+			PacketScope.bPacketContainsRPC = true;
 
-			CondCheckCountQuota();
+			TickScope.CondCheckCountQuota(*this);
 		}
 	}
 
@@ -588,16 +648,25 @@ public:
 	 */
 	void PostReceivedPacket(double TimeSeconds)
 	{
+		using namespace UE::Net;
+
 #if RPC_DOS_DEV_STATS
 		CSV_SCOPED_TIMING_STAT_EXCLUSIVE(RPCDoS_Checks);
 #endif
 
-		if (bPacketContainsRPC)
+		if (GetPacketScope().bPacketContainsRPC)
 		{
 			PostReceivedRPCPacket(TimeSeconds);
 		}
 
 		ReceivedPacketEndTime = TimeSeconds;
+
+		PacketScopePrivate.SetActive(false);
+
+#if RPC_DOS_SCOPE_DEBUG
+		ensure(GRPCDoSScopeDebugging == 0 || !SequentialRPCScopePrivate.IsActive());
+		ensure(GRPCDoSScopeDebugging == 0 || TickScopePrivate.IsActive());
+#endif
 	}
 
 	/**
@@ -664,6 +733,29 @@ public:
 
 
 private:
+	/**
+	 * Called when we begin receiving the same RPC once or multiple times sequentially.
+	 */
+	void PreSequentialRPC()
+	{
+		using namespace UE::Net;
+
+#if RPC_DOS_SCOPE_DEBUG
+		ensure(GRPCDoSScopeDebugging == 0 || TickScopePrivate.IsActive());
+		ensure(GRPCDoSScopeDebugging == 0 || PacketScopePrivate.IsActive());
+#endif
+
+		SequentialRPCScopePrivate.SetActive(true);
+	}
+
+	/**
+	 * Updates the current RPC DoS detection severity state
+	 *
+	 * @param Update	Whether or not we are escalating or de-escalating the severity state
+	 * @param Reason	The reason for the escalation change
+	 */
+	void UpdateSeverity_Private(ERPCDoSSeverityUpdate Update, ERPCDoSEscalateReason Reason);
+
 	/**
 	 * Initializes a newly active escalation state
 	 *
@@ -677,9 +769,14 @@ private:
 	void CondCheckCountQuota()
 	{
 		// 64 RPC interval
-		if ((RPCIntervalCounter & 0x3F) == 0 && HasHitQuota_Count(CounterPerPeriodHistory, FrameCounter))
+		if ((RPCIntervalCounter & 0x3F) == 0)
 		{
-			UpdateSeverity(ERPCDoSSeverityUpdate::Escalate, ERPCDoSEscalateReason::CountLimit);
+			FTickScope& TickScope = GetTickScope();
+
+			if (HasHitQuota_Count(CounterPerPeriodHistory, TickScope.FrameCounter))
+			{
+				TickScope.UpdateSeverity(*this, ERPCDoSSeverityUpdate::Escalate, ERPCDoSEscalateReason::CountLimit);
+			}
 		}
 	}
 
@@ -692,11 +789,13 @@ private:
 	{
 		if (TimeSeconds > NextTimeQuotaCheck)
 		{
+			FTickScope& TickScope = GetTickScope();
+
 			NextTimeQuotaCheck = TimeSeconds + TimeQuotaCheckInterval;
 
-			if (HasHitQuota_Time(CounterPerPeriodHistory, FrameCounter))
+			if (HasHitQuota_Time(CounterPerPeriodHistory, TickScope.FrameCounter))
 			{
-				UpdateSeverity(ERPCDoSSeverityUpdate::Escalate, ERPCDoSEscalateReason::TimeLimit);
+				TickScope.UpdateSeverity(*this, ERPCDoSSeverityUpdate::Escalate, ERPCDoSEscalateReason::TimeLimit);
 			}
 		}
 	}
@@ -807,6 +906,223 @@ private:
 	}
 
 
+private:
+	/** Base class for setting code scope activity */
+	template<class T>
+	class TScopeBase
+	{
+	public:
+		/** Sets whether or not the current code scope is active */
+		inline void SetActive(bool bInVal)
+		{
+			using namespace UE::Net;
+
+#if RPC_DOS_SCOPE_DEBUG
+			ensure(GRPCDoSScopeDebugging == 0 || bScopeActive != bInVal);
+#endif
+
+			bScopeActive = bInVal;
+
+			Reset();
+		}
+
+		/** Whether or not the current code scope is active */
+		inline bool IsActive() const
+		{
+			return bScopeActive;
+		}
+
+		/** Resets the properties covered by this scope, at both the start/end of the scope */
+		inline void Reset()
+		{
+			static_cast<T*>(this)->Reset();
+		}
+
+	private:
+		/** Whether or not the current code scope is active */
+		bool bScopeActive = false;
+	};
+
+	/** Variables and functions that should only be accessible during TickDispatch */
+	class FTickScope : public TScopeBase<FTickScope>
+	{
+	public:
+		/** Wrapper for UpdateSeverity which forces FTickScope acquisition */
+		inline void UpdateSeverity(FRPCDoSDetection& This, ERPCDoSSeverityUpdate Update, ERPCDoSEscalateReason Reason)
+		{
+			This.UpdateSeverity_Private(Update, Reason);
+		}
+
+		/** Wrapper for CondCheckCountQuota which forces FTickScope acquisition */
+		inline void CondCheckCountQuota(FRPCDoSDetection& This)
+		{
+			This.CondCheckCountQuota();
+		}
+
+		/** Wrapper for CondCheckTimeQuota which forces FTickScope acquisition */
+		inline void CondCheckTimeQuota(FRPCDoSDetection& This, double TimeSeconds)
+		{
+			This.CondCheckTimeQuota(TimeSeconds);
+		}
+
+		void Reset()
+		{
+			FrameCounter.ResetRPCCounters();
+		}
+
+	public:
+		/** Per-frame RPC counting (multiple disjointed packets may be processed for same connection, during a frame) */
+		FRPCDoSCounters FrameCounter;
+	};
+
+	/** Scoped variables/functions accessible during TickDispatch */
+	FTickScope TickScopePrivate;
+
+
+	/** Gets a reference to the TickDispatch scoped variable/function accessor */
+	inline FTickScope& GetTickScope()
+	{
+		using namespace UE::Net;
+
+#if RPC_DOS_SCOPE_DEBUG
+		ensure(GRPCDoSScopeDebugging == 0 || TickScopePrivate.IsActive());
+#endif
+
+		return TickScopePrivate;
+	}
+
+
+	/** Variables and functions that should only be accessible while receiving an individual packet */
+	class FPacketScope : public TScopeBase<FPacketScope>
+	{
+	public:
+		/** Whether or not the current packet being received, contains an RPC */
+		bool bPacketContainsRPC											= false;
+
+		/** Cached free/external timestamp, for when the current received packet began processing */
+		double ReceivedPacketStartTime									= 0.0;
+
+
+	public:
+		void Reset()
+		{
+			bPacketContainsRPC = false;
+			ReceivedPacketStartTime = 0.0;
+		}
+	};
+
+	/** Scoped variables/functions accessible while receiving an individual packet */
+	FPacketScope PacketScopePrivate;
+
+
+	/** Gets a reference to the packet receive scoped variable/function accessor */
+	inline FPacketScope& GetPacketScope()
+	{
+		using namespace UE::Net;
+
+#if RPC_DOS_SCOPE_DEBUG
+		ensure(GRPCDoSScopeDebugging == 0 || PacketScopePrivate.IsActive());
+#endif
+
+		return PacketScopePrivate;
+	}
+
+	/** Variables and functions that should only be accessible while receiving the same RPC individually/sequentially */
+	class FSequentialRPCScope : public TScopeBase<FSequentialRPCScope>
+	{
+	public:
+		/** Wrapper for CondCheckCountQuota which forces FSequentialRPCScope acquisition */
+		inline void CondCheckCountQuota(FRPCDoSDetection& This)
+		{
+			This.CondCheckCountQuota();
+		}
+
+		/** Wrapper for CondCheckTimeQuota which forces FSequentialRPCScope acquisition */
+		inline void CondCheckTimeQuota(FRPCDoSDetection& This, double TimeSeconds)
+		{
+			This.CondCheckTimeQuota(TimeSeconds);
+		}
+
+		/** Wrapper for CheckRPCTracking which forces FSequentialRPCScope acquisition */
+		inline ERPCNotifyResult CheckRPCTracking(FRPCDoSDetection& This, UFunction* Function, FName FunctionName)
+		{
+			return This.CheckRPCTracking(Function, FunctionName);
+		}
+
+		/** Wrapper for FindOrAddRPCTracking which forces FSequentialRPCScope acquisition */
+		inline FRPCTrackingInfo& FindOrAddRPCTracking(FRPCDoSDetection& This, UFunction* InFunc, bool& bOutNewTracking)
+		{
+			return This.FindOrAddRPCTracking(InFunc, bOutNewTracking);
+		}
+
+		void Reset()
+		{
+			bReceivedPacketRPCUnique = true;
+			ReceivedPacketRPCCount = 0;
+			LightweightRPCTracking.Count = 0;
+			PostReceivedRPCName = NAME_None;
+			PostReceivedRPCTracking = nullptr;
+			PostReceivedRPCCounter = nullptr;
+			PostReceivedRPCBlockCount = 0;
+			LastReceivedRPCTimeCache = 0.0;
+
+#if RPC_QUOTA_DEBUG
+			DebugReceivedRPCStartTime = 0.0;
+			DebugReceivedRPCEndTime = 0.0;
+#endif
+		}
+
+	public:
+		/** Whether or not the packet only contains references to the same unique RPC (may be multiple RPC calls to same RPC) */
+		bool bReceivedPacketRPCUnique									= true;
+
+		/** Counts the number of RPC's in the current packet, for unique-RPC analytics */
+		int32 ReceivedPacketRPCCount									= 0;
+
+		/** State for lightweight tracking of RPC's */
+		FLightweightRPCTracking LightweightRPCTracking;
+
+
+		/** Temporarily caches the currently active RPC name, while processing a packet */
+		FName PostReceivedRPCName;
+
+		/** Temporarily caches the currently active RPC tracking, while processing a packet */
+		FRPCTrackingInfo* PostReceivedRPCTracking						= nullptr;
+
+		/** Temporarily caches the currently active RPC counter, while processing a packet */
+		FRPCDoSCounters* PostReceivedRPCCounter							= nullptr;
+
+		/** Temporarily caches the number of blocked RPC calls for the current RPC, while processing a packet */
+		int32 PostReceivedRPCBlockCount									= 0;
+
+		/** Temporarily caches timestamps for sequentially executed RPC's, to minimize new timestamp generation. */
+		double LastReceivedRPCTimeCache									= 0.0;
+
+#if RPC_QUOTA_DEBUG
+		/** Direct/accurate timestamp of RPC receive start time, for debugging against approximate RPC timestamps */
+		double DebugReceivedRPCStartTime								= 0.0;
+
+		/** Direct/accurate timestamp of RPC receive end time, for debugging against approximate RPC timestamps */
+		double DebugReceivedRPCEndTime									= 0.0;
+#endif
+	};
+
+	/** Scoped variables/functions accessible while receiving the same RPC individually/sequentially */
+	FSequentialRPCScope SequentialRPCScopePrivate;
+
+	/** Gets a reference to the sequential RPC scoped variable/function accessor */
+	inline FSequentialRPCScope& GetSequentialRPCScope()
+	{
+		using namespace UE::Net;
+
+#if RPC_DOS_SCOPE_DEBUG
+		ensure(GRPCDoSScopeDebugging == 0 || SequentialRPCScopePrivate.IsActive());
+#endif
+
+		return SequentialRPCScopePrivate;
+	}
+
+
 private:
 	/** Whether or not RPC DoS detection is presently enabled */
 	bool bRPCDoSDetection											= false;
@@ -874,9 +1190,6 @@ private:
 	double HitchSuspendDetectionStartTime							= 0.0;
 
 
-	/** Per-frame RPC counting (multiple disjointed packets may be processed for same connection, during a frame) */
-	FRPCDoSCounters FrameCounter;
-
 	/** Per-second RPC counting (approximate - may cover more than one second, in the case of expensive/long-running RPCs) */
 	FRPCDoSCounters SecondCounter;
 
@@ -896,25 +1209,12 @@ private:
 	uint8 SecondsIncrementer										= 0;
 
 
-	/** Whether or not the current packet being received, contains an RPC */
-	bool bPacketContainsRPC											= false;
-
-	/** Cached free/external timestamp, for when the current received packet began processing */
-	double ReceivedPacketStartTime									= 0.0;
-
-	/** Cached free/external timestamp, for when the last received packet finished procesing */
+	/** Cached free/external timestamp, for when the last received packet finished processing */
 	double ReceivedPacketEndTime									= 0.0;
 
-	/** Cached PreTickDispatch timestamp, to reuse for minmizing timestamp retrieval */
+	/** Cached PreTickDispatch timestamp, to reuse for minimizing timestamp retrieval */
 	double LastPreTickDispatchTime									= 0.0;
 
-	/** Whether or not the packet only contains references to the same unique RPC (may be multiple RPC calls to same RPC) */
-	bool bReceivedPacketRPCUnique									= false;
-
-	/** Counts the number of RPC's in the current packet, for unique-RPC analytics */
-	int32 ReceivedPacketRPCCount									= 0;
-
-
 	/** Unbounded RPC counter, for performing RPC checks at every 'x' intervals */
 	uint32 RPCIntervalCounter										= 0;
 
@@ -943,33 +1243,6 @@ private:
 	/** The last time RPCTracking was cleaned out */
 	double LastRPCTrackingClean										= 0.0;
 
-	/** State for lightweight tracking of RPC's */
-	FLightweightRPCTracking LightweightRPCTracking;
-
-
-	/** Temporarily caches the currently active RPC name, while processing a packet */
-	FName PostReceivedRPCName;
-
-	/** Temporarily caches the currently active RPC tracking, while processing a packet */
-	FRPCTrackingInfo* PostReceivedRPCTracking						= nullptr;
-
-	/** Temporarily caches the currently active RPC counter, while processing a packet */
-	FRPCDoSCounters* PostReceivedRPCCounter							= nullptr;
-
-	/** Temporarily caches the number of blocked RPC calls for the current RPC, while processing a packet */
-	int32 PostReceivedRPCBlockCount									= 0;
-
-	/** Temporarily caches timestamps for sequentially executed RPC's, to minimize new timestamp generation. */
-	double LastReceivedRPCTimeCache									= 0.0;
-
-#if RPC_QUOTA_DEBUG
-	/** Direct/accurate timestamp of RPC receive start time, for debugging against approximate RPC timestamps */
-	double DebugReceivedRPCStartTime								= 0.0;
-
-	/** Direct/accurate timestamp of RPC receive end time, for debugging against approximate RPC timestamps */
-	double DebugReceivedRPCEndTime									= 0.0;
-#endif
-
 
 	/** The locally cached/updated analytics variables, for the RPC DoS Detection - aggregated upon connection Close */
 	FRPCDoSAnalyticsVars AnalyticsVars;
