Fix serialization of RelativeScale3D for placed instances. This broke when the interp flag was removed because this field was missing an Edit flag despite being editable from customizations

diff --git a/Engine/Source/Runtime/Engine/Classes/Components/SceneComponent.h b/Engine/Source/Runtime/Engine/Classes/Components/SceneComponent.h
--- a/Engine/Source/Runtime/Engine/Classes/Components/SceneComponent.h
+++ b/Engine/Source/Runtime/Engine/Classes/Components/SceneComponent.h
@@ -140,7 +140,7 @@ private:
 	*	Non-uniform scaling of the component relative to its parent.
 	*	Note that scaling is always applied in local space (no shearing etc)
 	*/
-	UPROPERTY(BlueprintReadOnly, ReplicatedUsing=OnRep_Transform, Category=Transform, meta=(AllowPrivateAccess="true"))
+	UPROPERTY(EditAnywhere, BlueprintReadOnly, ReplicatedUsing=OnRep_Transform, Category=Transform, meta=(AllowPrivateAccess="true"))
 	FVector RelativeScale3D;
 
 public:
