Prevent crashes with PP materials when one of them is not a MaterialInstance and it tries to blend the parameters.

- Easiest fix is just to skip blending the materials together.

The blending functions below only work on instances, so skip blending if we are just a material interface.  This can happen in two scenarios:
    1. A wrong material type is assigned into the post process material array.
    2. The same material is placed in twice.  The MID lookup code above won't create two separate MIDs for the same material resulting in this code trying to blend a MID w/ a MaterialInterface.

diff --git a/Engine/Source/Runtime/Engine/Private/Materials/Material.cpp b/Engine/Source/Runtime/Engine/Private/Materials/Material.cpp
--- a/Engine/Source/Runtime/Engine/Private/Materials/Material.cpp
+++ b/Engine/Source/Runtime/Engine/Private/Materials/Material.cpp
@@ -1858,13 +1858,22 @@ void UMaterialInterface::OverrideBlendableSettings(class FSceneView& View, float
 		UMaterialInstanceDynamic* DestMID = DestNode->GetMID();
 		check(DestMID);
 
-		UMaterialInstance* SrcMID = (UMaterialInstance*)this;
-		check(SrcMID);
+		// The blending functions below only work on instances, so skip blending if we are just a material interface.
+		// This can happen in two scenarios:
+		// 1. A wrong material type is assigned into the post process material array.
+		// 2. The same material is placed in twice.  The MID lookup code above won't create two separate MIDs for the same
+		//    material resulting in this code trying to blend a MID w/ a MaterialInterface.
+		if (this->IsA<UMaterialInstance>())
+		{
+			// Attempt to cast the UMaterialInterface to a UMaterialInstance pointer.
+			UMaterialInstance* SrcMI = (UMaterialInstance*)this;
+			check(SrcMI);
 
-		// Here we could check for Weight=1.0 and use copy instead of interpolate but that case quite likely not intended anyway.
+			// Here we could check for Weight=1.0 and use copy instead of interpolate but that case quite likely not intended anyway.
 
-		// a material already exists, blend (Scalar and Vector parameters) with existing ones
-		DestMID->K2_InterpolateMaterialInstanceParams(DestMID, SrcMID, Weight);
+			// a material already exists, blend (Scalar and Vector parameters) with existing ones
+			DestMID->K2_InterpolateMaterialInstanceParams(DestMID, SrcMI, Weight);
+		}
 	}
 }
 
