[KillPendingKill] Mark LocalPlayer as garbage when removing it from GameInstance to track outstanding references.

diff --git a/Engine/Source/Runtime/Engine/Private/GameInstance.cpp b/Engine/Source/Runtime/Engine/Private/GameInstance.cpp
--- a/Engine/Source/Runtime/Engine/Private/GameInstance.cpp
+++ b/Engine/Source/Runtime/Engine/Private/GameInstance.cpp
@@ -888,6 +888,12 @@ bool UGameInstance::RemoveLocalPlayer(ULocalPlayer* ExistingPlayer)
 
 	OnLocalPlayerRemovedEvent.Broadcast(ExistingPlayer);
 
+	// Marked as garbage here to detect outstanding references
+	if (!UObjectBaseUtility::IsPendingKillEnabled())
+	{
+		ExistingPlayer->MarkAsGarbage();
+	}
+
 	return true;
 }
 
