Chaos
- add GetScaledVolume to all Aggregate Geometry element types
- GetScaledGeom handles non-uniform and negative scales
- Deprecate GetVolume
- Mass now shows correctly for StaticMeshComponents

diff --git a/Engine/Source/Runtime/Engine/Classes/PhysicsEngine/AggregateGeom.h b/Engine/Source/Runtime/Engine/Classes/PhysicsEngine/AggregateGeom.h
--- a/Engine/Source/Runtime/Engine/Classes/PhysicsEngine/AggregateGeom.h
+++ b/Engine/Source/Runtime/Engine/Classes/PhysicsEngine/AggregateGeom.h
@@ -156,7 +156,11 @@ struct ENGINE_API FKAggregateGeom
 	void CalcBoxSphereBounds(FBoxSphereBounds& Output, const FTransform& LocalToWorld) const;
 
 	/** Returns the volume of this element */
-	float GetVolume(const FVector& Scale3D) const;
+	UE_DEPRECATED(5.1, "Changed to GetScaledVolume. Note that Volume calculation now includes non-uniform scale so values may have changed")
+	FVector::FReal GetVolume(const FVector& Scale3D) const;
+
+	/** Returns the volume of this element */
+	FVector::FReal GetScaledVolume(const FVector& Scale3D) const;
 
 	FGuid MakeDDCKey() const;
 
diff --git a/Engine/Source/Runtime/Engine/Classes/PhysicsEngine/BodySetup.h b/Engine/Source/Runtime/Engine/Classes/PhysicsEngine/BodySetup.h
--- a/Engine/Source/Runtime/Engine/Classes/PhysicsEngine/BodySetup.h
+++ b/Engine/Source/Runtime/Engine/Classes/PhysicsEngine/BodySetup.h
@@ -375,8 +375,12 @@ public:
 	void GetGeometryDDCKey(FString& OutString) const;
 
 	/** Returns the volume of this element */
+	UE_DEPRECATED(5.1, "Use GetScaledVolume which uses the same scaling technique as the generated collision geometry")
 	ENGINE_API virtual float GetVolume(const FVector& Scale) const;
 
+	/** Returns the volume of this element givent the scale */
+	ENGINE_API virtual FVector::FReal GetScaledVolume(const FVector& Scale) const;
+
 	/** Release Physics meshes (ConvexMeshes, TriMesh & TriMeshNegX) */
 	ENGINE_API void ClearPhysicsMeshes();
 
diff --git a/Engine/Source/Runtime/Engine/Classes/PhysicsEngine/BoxElem.h b/Engine/Source/Runtime/Engine/Classes/PhysicsEngine/BoxElem.h
--- a/Engine/Source/Runtime/Engine/Classes/PhysicsEngine/BoxElem.h
+++ b/Engine/Source/Runtime/Engine/Classes/PhysicsEngine/BoxElem.h
@@ -109,7 +109,10 @@ struct FKBoxElem : public FKShapeElem
 		Center = InTransform.GetLocation();
 	}
 
-	FORCEINLINE FVector::FReal GetVolume(const FVector& Scale3D) const { FVector::FReal MinScale = Scale3D.GetMin(); return (X * MinScale) * (Y * MinScale) * (Z * MinScale); }
+	UE_DEPRECATED(5.1, "Changed to GetScaledVolume. Note that Volume calculation now includes non-uniform scale so values may have changed")
+	FORCEINLINE FVector::FReal GetVolume(const FVector& Scale3D) const { return GetScaledVolume(Scale3D); }
+
+	FORCEINLINE FVector::FReal GetScaledVolume(const FVector& Scale3D) const { return FMath::Abs(Scale3D.X * Scale3D.Y * Scale3D.Z * X * Y * Z); }
 
 	ENGINE_API void DrawElemWire(class FPrimitiveDrawInterface* PDI, const FTransform& ElemTM, float Scale, const FColor Color) const override;
 	ENGINE_API void DrawElemSolid(class FPrimitiveDrawInterface* PDI, const FTransform& ElemTM, float Scale, const FMaterialRenderProxy* MaterialRenderProxy) const override;
diff --git a/Engine/Source/Runtime/Engine/Classes/PhysicsEngine/ConvexElem.h b/Engine/Source/Runtime/Engine/Classes/PhysicsEngine/ConvexElem.h
--- a/Engine/Source/Runtime/Engine/Classes/PhysicsEngine/ConvexElem.h
+++ b/Engine/Source/Runtime/Engine/Classes/PhysicsEngine/ConvexElem.h
@@ -98,7 +98,11 @@ public:
 	ENGINE_API void BakeTransformToVerts();
 
 	/** Returns the volume of this element */
-	float GetVolume(const FVector& Scale) const;
+	UE_DEPRECATED(5.1, "Changed to GetScaledVolume. Note that Volume calculation now includes non-uniform scale so values may have changed")
+	FVector::FReal GetVolume(const FVector& Scale) const;
+
+	/** Returns the volume of this element */
+	FVector::FReal GetScaledVolume(const FVector& Scale3D) const;
 
 #if PHYSICS_INTERFACE_PHYSX
 	/** Get the PhysX convex mesh (defined in BODY space) for this element */
diff --git a/Engine/Source/Runtime/Engine/Classes/PhysicsEngine/SphereElem.h b/Engine/Source/Runtime/Engine/Classes/PhysicsEngine/SphereElem.h
--- a/Engine/Source/Runtime/Engine/Classes/PhysicsEngine/SphereElem.h
+++ b/Engine/Source/Runtime/Engine/Classes/PhysicsEngine/SphereElem.h
@@ -72,8 +72,11 @@ struct FKSphereElem : public FKShapeElem
 		Center = InTransform.GetLocation();
 	}
 
-	FORCEINLINE FVector::FReal GetVolume(const FVector& Scale) const { return 1.3333f * UE_PI * FMath::Pow(Radius * Scale.GetMin(), 3); }
-	
+	UE_DEPRECATED(5.1, "Changed to GetScaledVolume to support non-uniform scales on other element types")
+	FORCEINLINE FVector::FReal GetVolume(const FVector& Scale) const { return GetScaledVolume(Scale); }
+
+	FORCEINLINE FVector::FReal GetScaledVolume(const FVector& Scale) const { return 1.3333f * UE_PI * FMath::Pow(Radius * Scale.GetAbsMin(), 3); }
+
 	ENGINE_API void DrawElemWire(class FPrimitiveDrawInterface* PDI, const FTransform& ElemTM, float Scale, const FColor Color) const override;
 	ENGINE_API void DrawElemSolid(class FPrimitiveDrawInterface* PDI, const FTransform& ElemTM, float Scale, const FMaterialRenderProxy* MaterialRenderProxy) const override;
 
diff --git a/Engine/Source/Runtime/Engine/Classes/PhysicsEngine/SphylElem.h b/Engine/Source/Runtime/Engine/Classes/PhysicsEngine/SphylElem.h
--- a/Engine/Source/Runtime/Engine/Classes/PhysicsEngine/SphylElem.h
+++ b/Engine/Source/Runtime/Engine/Classes/PhysicsEngine/SphylElem.h
@@ -90,7 +90,15 @@ struct FKSphylElem : public FKShapeElem
 		Center = InTransform.GetLocation();
 	}
 
-	FORCEINLINE FVector::FReal GetVolume(const FVector& Scale) const { FVector::FReal ScaledRadius = Radius * Scale.GetMin(); return UE_PI * FMath::Square(ScaledRadius) * ( 1.3333f * ScaledRadius + (Length * Scale.GetMin())); }
+	UE_DEPRECATED(5.1, "Changed to GetScaledVolume. Note that Volume calculation now includes non-uniform scale so values may have changed")
+	FORCEINLINE FVector::FReal GetVolume(const FVector& Scale) const { return GetScaledVolume(Scale); }
+	
+	FORCEINLINE FVector::FReal GetScaledVolume(const FVector& Scale3D) const 
+	{  
+		FVector::FReal ScaledRadius = GetScaledRadius(Scale3D);
+		FVector::FReal ScaledLength = GetScaledCylinderLength(Scale3D);
+		return UE_PI * FMath::Square(ScaledRadius) * (1.3333f * ScaledRadius + ScaledLength);
+	}
 
 	ENGINE_API void DrawElemWire(class FPrimitiveDrawInterface* PDI, const FTransform& ElemTM, float Scale, const FColor Color) const override;
 	ENGINE_API void DrawElemSolid(class FPrimitiveDrawInterface* PDI, const FTransform& ElemTM, float Scale, const FMaterialRenderProxy* MaterialRenderProxy) const override;
diff --git a/Engine/Source/Runtime/Engine/Private/PhysicsEngine/BodySetup.cpp b/Engine/Source/Runtime/Engine/Private/PhysicsEngine/BodySetup.cpp
--- a/Engine/Source/Runtime/Engine/Private/PhysicsEngine/BodySetup.cpp
+++ b/Engine/Source/Runtime/Engine/Private/PhysicsEngine/BodySetup.cpp
@@ -1932,28 +1932,33 @@ void FKAggregateGeom::FixupDeprecated(FArchive& Ar)
 }
 #endif
 
-float FKAggregateGeom::GetVolume(const FVector& Scale) const
+FVector::FReal FKAggregateGeom::GetVolume(const FVector& Scale) const
+{
+	return GetScaledVolume(Scale);
+}
+
+FVector::FReal FKAggregateGeom::GetScaledVolume(const FVector& Scale) const
 {
 	float Volume = 0.0f;
 
-	for ( auto SphereElemIt = SphereElems.CreateConstIterator(); SphereElemIt; ++SphereElemIt )
+	for (auto SphereElemIt = SphereElems.CreateConstIterator(); SphereElemIt; ++SphereElemIt )
 	{
-		Volume += SphereElemIt->GetVolume(Scale);
+		Volume += SphereElemIt->GetScaledVolume(Scale);
 	}
 
-	for ( auto BoxElemIt = BoxElems.CreateConstIterator(); BoxElemIt; ++BoxElemIt )
+	for (auto BoxElemIt = BoxElems.CreateConstIterator(); BoxElemIt; ++BoxElemIt )
 	{
-		Volume += BoxElemIt->GetVolume(Scale);
+		Volume += BoxElemIt->GetScaledVolume(Scale);
 	}
 
-	for ( auto SphylElemIt = SphylElems.CreateConstIterator(); SphylElemIt; ++SphylElemIt )
+	for (auto SphylElemIt = SphylElems.CreateConstIterator(); SphylElemIt; ++SphylElemIt )
 	{
-		Volume += SphylElemIt->GetVolume(Scale);
+		Volume += SphylElemIt->GetScaledVolume(Scale);
 	}
 
-	for ( auto ConvexElemIt = ConvexElems.CreateConstIterator(); ConvexElemIt; ++ConvexElemIt )
+	for (auto ConvexElemIt = ConvexElems.CreateConstIterator(); ConvexElemIt; ++ConvexElemIt )
 	{
-		Volume += ConvexElemIt->GetVolume(Scale);
+		Volume += ConvexElemIt->GetScaledVolume(Scale);
 	}
 
 	return Volume;
@@ -2086,7 +2091,12 @@ void FKConvexElem::SetMirroredConvexMesh(physx::PxConvexMesh* InMesh)
 }
 #endif
 
-float FKConvexElem::GetVolume(const FVector& Scale) const
+FVector::FReal FKConvexElem::GetVolume(const FVector& Scale) const
+{
+	return GetScaledVolume(Scale);
+}
+
+FVector::FReal FKConvexElem::GetScaledVolume(const FVector& Scale) const
 {
 	float Volume = 0.0f;
 
@@ -2122,8 +2132,10 @@ float FKConvexElem::GetVolume(const FVector& Scale) const
 		}
 	}
 #elif WITH_CHAOS
-	//TODO Support ChaosConvex.
-	CHAOS_ENSURE(false);
+	if (ChaosConvex != nullptr)
+	{
+		Volume = FMath::Abs(Scale.X * Scale.Y * Scale.Z) * ChaosConvex->GetVolume();
+	}
 #endif
 
 	return Volume;
@@ -2615,20 +2627,28 @@ float UBodySetup::CalculateMass(const UPrimitiveComponent* Component) const
 	}
 
 	// Then scale mass to avoid big differences between big and small objects.
-	const float BasicVolume = GetVolume(ComponentScale);
+	const FVector::FReal BasicVolume = GetScaledVolume(ComponentScale);
+
+	// The below TODO is probably fixed now. GetSCaledVolume handles negative scales...
 	//@TODO: Some static meshes are triggering this - disabling until content can be analyzed - ensureMsgf(BasicVolume >= 0.0f, TEXT("UBodySetup::CalculateMass(%s) - The volume of the aggregate geometry is negative"), *Component->GetReadableName());
 
-	const float BasicMass = FMath::Max<float>(BasicVolume, 0.0f) * DensityKGPerCubicUU;
+	const FVector::FReal BasicMass = FMath::Max(BasicVolume, FVector::FReal(0)) * DensityKGPerCubicUU;
 
-	const float UsePow = FMath::Clamp<float>(RaiseMassToPower, UE_KINDA_SMALL_NUMBER, 1.f);
-	const float RealMass = FMath::Pow(BasicMass, UsePow);
+	const FVector::FReal UsePow = FMath::Clamp(RaiseMassToPower, FVector::FReal(UE_KINDA_SMALL_NUMBER), FVector::FReal(1));
+	const FVector::FReal RealMass = FMath::Pow(BasicMass, UsePow);
 
-	return RealMass * MassScale;
+	return float(RealMass * MassScale);
 }
 
 float UBodySetup::GetVolume(const FVector& Scale) const
 {
-	return AggGeom.GetVolume(Scale);
+	return GetScaledVolume(Scale);
+}
+
+
+FVector::FReal UBodySetup::GetScaledVolume(const FVector& Scale) const
+{
+	return AggGeom.GetScaledVolume(Scale);
 }
 
 
