Make sure job cache DDC is disabled if job cache itself is disabled.

- Also added some more debug info dumping code.

diff --git a/Engine/Source/Runtime/Engine/Private/Materials/MaterialShader.cpp b/Engine/Source/Runtime/Engine/Private/Materials/MaterialShader.cpp
--- a/Engine/Source/Runtime/Engine/Private/Materials/MaterialShader.cpp
+++ b/Engine/Source/Runtime/Engine/Private/Materials/MaterialShader.cpp
@@ -3259,19 +3259,41 @@ void FMaterialShaderMap::DumpDebugInfo() const
 	const FString& DebugDescriptionS = GetDebugDescription();
 	UE_LOG(LogConsoleResponse, Display, TEXT("  DebugDescription %s"), *DebugDescriptionS);
 
-	TMap<FShaderId, TShaderRef<FShader>> ShadersL;
-	GetShaderList(ShadersL);
-	UE_LOG(LogConsoleResponse, Display, TEXT("  --- %d shaders"), ShadersL.Num());
-	int32 Index = 0;
-	for (auto& KeyValue : ShadersL)
-	{
-		UE_LOG(LogConsoleResponse, Display, TEXT("    --- shader %d"), Index);
-		FShader* Shader = KeyValue.Value.GetShader();
-		Shader->DumpDebugInfo(GetPointerTable());
-		Index++;
-	}
-	//TArray<FShaderPipeline*> ShaderPipelinesL;
-	//GetShaderPipelineList(ShaderPipelinesL);
+	{
+		TMap<FShaderId, TShaderRef<FShader>> ShadersL;
+		GetShaderList(ShadersL);
+		UE_LOG(LogConsoleResponse, Display, TEXT("  --- %d shaders"), ShadersL.Num());
+		int32 Index = 0;
+		for (auto& KeyValue : ShadersL)
+		{
+			UE_LOG(LogConsoleResponse, Display, TEXT("    --- shader %d"), Index);
+			FShader* Shader = KeyValue.Value.GetShader();
+			Shader->DumpDebugInfo(GetPointerTable());
+			Index++;
+		}
+	}
+
+	{
+		TArray<FShaderPipelineRef> ShaderPipelinesL;
+		GetShaderPipelineList(ShaderPipelinesL);
+		UE_LOG(LogConsoleResponse, Display, TEXT("  --- %d shaderpipelines"), ShaderPipelinesL.Num());
+		int32 Index = 0;
+		for (FShaderPipelineRef Value : ShaderPipelinesL)
+		{
+			FShaderPipeline* ShaderPipeline = Value.GetPipeline();
+			UE_LOG(LogConsoleResponse, Display, TEXT("    --- shaderpipeline %d"), Index);
+			TArray<TShaderRef<FShader>> ShadersL = ShaderPipeline->GetShaders(*this);
+			UE_LOG(LogConsoleResponse, Display, TEXT("    --- %d shaders"), ShadersL.Num());
+			int32 ShaderIdx = 0;
+			for (TShaderRef<FShader>& Shader : ShadersL)
+			{
+				UE_LOG(LogConsoleResponse, Display, TEXT("      --- shader %d"), ShaderIdx);
+				Shader->DumpDebugInfo(GetPointerTable());
+				++ShaderIdx;
+			}
+			Index++;
+		}
+	}
 }
 
 #if WITH_EDITOR
diff --git a/Engine/Source/Runtime/Engine/Private/ShaderCompiler/ShaderCompiler.cpp b/Engine/Source/Runtime/Engine/Private/ShaderCompiler/ShaderCompiler.cpp
--- a/Engine/Source/Runtime/Engine/Private/ShaderCompiler/ShaderCompiler.cpp
+++ b/Engine/Source/Runtime/Engine/Private/ShaderCompiler/ShaderCompiler.cpp
@@ -139,7 +139,8 @@ bool IsShaderJobCacheDDCEnabled()
 	// For now we only support the editor and not commandlets like the cooker.
 	if (GIsEditor && !IsRunningCommandlet())
 	{
-		return CVarJobCacheDDC.GetValueOnAnyThread();
+		// job cache itself must be enabled first
+		return GShaderCompilerJobCache && CVarJobCacheDDC.GetValueOnAnyThread();
 	}
 
 	return false;
