Minor fixes to FNavigationMemoryStat. Issue was that ForElementType was not provided and that's the code path that is used by TNavStatArray. This means we were calling FDefaultAllocator::ForAnyElementType instead of FNavigationMemoryStat::ForAnyElementType and we were not logging memory used via STAT_NavigationMemory.

diff --git a/Engine/Source/Runtime/Engine/Classes/AI/Navigation/NavigationTypes.h b/Engine/Source/Runtime/Engine/Classes/AI/Navigation/NavigationTypes.h
--- a/Engine/Source/Runtime/Engine/Classes/AI/Navigation/NavigationTypes.h
+++ b/Engine/Source/Runtime/Engine/Classes/AI/Navigation/NavigationTypes.h
@@ -693,7 +693,14 @@ namespace NavMeshMemory
 				}
 			}
 
-			void ResizeAllocation(int32 PreviousNumElements, int32 NumElements, int32 NumBytesPerElement, uint32 AlignmentOfElement)
+			FORCEINLINE_DEBUGGABLE void ResizeAllocation(int32 PreviousNumElements, int32 NumElements, int32 NumBytesPerElement, uint32 AlignmentOfElement)
+			{
+				// Maintain existing behavior, call the default aligned version of this function.
+				// We currently rely on this as we are often storing actual structs into uint8's here.
+				ResizeAllocation(PreviousNumElements, NumElements, NumBytesPerElement);
+			}
+
+			FORCEINLINE_DEBUGGABLE void ResizeAllocation(SizeType PreviousNumElements, SizeType NumElements, SIZE_T NumBytesPerElement)
 			{
 				const int32 NewSize = NumElements * NumBytesPerElement;
 				INC_DWORD_STAT_BY(STAT_NavigationMemory, NewSize - AllocatedSize);
@@ -706,6 +713,16 @@ namespace NavMeshMemory
 			ForAnyElementType(const ForAnyElementType&);
 			ForAnyElementType& operator=(const ForAnyElementType&);
 		};
+
+		template<typename ElementType>
+		class ForElementType : public ForAnyElementType
+		{
+		public:
+			ElementType* GetAllocation() const
+			{
+				return (ElementType*)ForAnyElementType::GetAllocation();
+			}
+		};
 	};
 
 	typedef FNavigationMemoryStat FNavAllocator;
