World Partition: don't manage dirty actors when running commandlets.

#rnx

diff --git a/Engine/Source/Runtime/Engine/Private/WorldPartition/WorldPartition.cpp b/Engine/Source/Runtime/Engine/Private/WorldPartition/WorldPartition.cpp
--- a/Engine/Source/Runtime/Engine/Private/WorldPartition/WorldPartition.cpp
+++ b/Engine/Source/Runtime/Engine/Private/WorldPartition/WorldPartition.cpp
@@ -718,7 +718,11 @@ void UWorldPartition::RegisterDelegates()
 			FCoreUObjectDelegates::PostReachabilityAnalysis.AddUObject(this, &UWorldPartition::OnGCPostReachabilityAnalysis);
 			GEditor->OnPostBugItGoCalled().AddUObject(this, &UWorldPartition::OnPostBugItGoCalled);
 			GEditor->OnEditorClose().AddUObject(this, &UWorldPartition::SavePerUserSettings);
-			UPackage::PackageDirtyStateChangedEvent.AddUObject(this, &UWorldPartition::OnPackageDirtyStateChanged);
+
+			if (!IsRunningCommandlet())
+			{
+				UPackage::PackageDirtyStateChangedEvent.AddUObject(this, &UWorldPartition::OnPackageDirtyStateChanged);
+			}
 		}
 	}
 #endif
@@ -761,7 +765,11 @@ void UWorldPartition::UnregisterDelegates()
 
 			GEditor->OnPostBugItGoCalled().RemoveAll(this);
 			GEditor->OnEditorClose().RemoveAll(this);
-			UPackage::PackageDirtyStateChangedEvent.RemoveAll(this);
+
+			if (!IsRunningCommandlet())
+			{
+				UPackage::PackageDirtyStateChangedEvent.RemoveAll(this);
+			}
 		}
 	}
 #endif
