PR #9200: Fix Visual Logger redirections outside the editor. (Contributed by [REDACTED_PERSON])

minor additional changes to apply similar fix to command "LogNavOctree" and reduce map lookups

diff --git a/Engine/Source/Runtime/Engine/Private/VisualLogger/VisualLogger.cpp b/Engine/Source/Runtime/Engine/Private/VisualLogger/VisualLogger.cpp
--- a/Engine/Source/Runtime/Engine/Private/VisualLogger/VisualLogger.cpp
+++ b/Engine/Source/Runtime/Engine/Private/VisualLogger/VisualLogger.cpp
@@ -672,14 +672,14 @@ int32 FVisualLogger::GetUniqueId(const float Timestamp)
 FVisualLogger::FOwnerToChildrenRedirectionMap& FVisualLogger::GetRedirectionMap(const UObject* InObject)
 {
 	const UWorld* World = nullptr;
-	if (FVisualLogger::Get().ObjectToWorldMap.Contains(InObject))
+	if (const TWeakObjectPtr<const UWorld>* Entry = FVisualLogger::Get().ObjectToWorldMap.Find(InObject))
 	{
-		World = FVisualLogger::Get().ObjectToWorldMap[InObject].Get();
+		World = Entry->Get();
 	}
 
 	if (World == nullptr)
 	{
-		World = GetWorldForVisualLogger(nullptr);
+		World = GetWorldForVisualLogger(InObject);
 	}
 
 	return WorldToRedirectionMap.FindOrAdd(World);
@@ -912,7 +912,7 @@ public:
 #if ENABLE_VISUAL_LOG
 		else if (FParse::Command(&Cmd, TEXT("LogNavOctree")))
 		{
-			FVisualLogger::NavigationDataDump(GetWorldForVisualLogger(nullptr), LogNavigation, ELogVerbosity::Log, FBox());
+			FVisualLogger::NavigationDataDump(GetWorldForVisualLogger(InWorld), LogNavigation, ELogVerbosity::Log, FBox());
 		}
 #endif
 		return false;
