Telemetry for Disabled Audio

diff --git a/Engine/Source/Runtime/Engine/Private/AudioDeviceManager.cpp b/Engine/Source/Runtime/Engine/Private/AudioDeviceManager.cpp
--- a/Engine/Source/Runtime/Engine/Private/AudioDeviceManager.cpp
+++ b/Engine/Source/Runtime/Engine/Private/AudioDeviceManager.cpp
@@ -7,6 +7,7 @@
 #include "AudioDevice.h"
 #include "AudioMixerDevice.h"
 #include "Engine/World.h"
+#include "EngineAnalytics.h"
 #include "Sound/AudioSettings.h"
 #include "Sound/SoundWave.h"
 #include "GameFramework/GameUserSettings.h"
@@ -665,6 +666,14 @@ bool FAudioDeviceManager::LoadDefaultAudioDeviceModule()
 		IsUsingAudioMixerCvar->Set(0, ECVF_SetByConstructor);
 	}
 
+	if (FEngineAnalytics::IsAvailable())
+	{ 		
+		if (bForceNoAudioMixer)
+		{
+			FEngineAnalytics::GetProvider().RecordEvent(TEXT("Audio.Usage.ForceNoAudioMixer"));
+		}
+	}
+
 	return AudioDeviceModule != nullptr;
 }
 
@@ -858,6 +867,13 @@ bool FAudioDeviceManager::Initialize()
 				Singleton = nullptr;
 			}
 		}
+		else
+		{
+			if (FEngineAnalytics::IsAvailable())
+			{
+				FEngineAnalytics::GetProvider().RecordEvent(TEXT("Audio.Usage.AllAudioDisabled"));
+			}
+		}
 	}
 
 	return Singleton != nullptr;
