SetSourceCubemapAngle for skylight components is now exposed in blueprints

diff --git a/Engine/Source/Runtime/Engine/Classes/Components/SkyLightComponent.h b/Engine/Source/Runtime/Engine/Classes/Components/SkyLightComponent.h
--- a/Engine/Source/Runtime/Engine/Classes/Components/SkyLightComponent.h
+++ b/Engine/Source/Runtime/Engine/Classes/Components/SkyLightComponent.h
@@ -246,6 +246,9 @@ class ENGINE_API USkyLightComponent : public ULightComponentBase
 	UFUNCTION(BlueprintCallable, Category="SkyLight")
 	void SetCubemap(UTextureCube* NewCubemap);
 
+	UFUNCTION(BlueprintCallable, Category = "Rendering|Components|SkyLight")
+	void SetSourceCubemapAngle(float NewValue);
+
 	/** 
 	 * Creates sky lighting from a blend between two cubemaps, which is only valid when SourceType is set to SpecifiedCubemap. 
 	 * This can be used to seamlessly transition sky lighting between different times of day.
diff --git a/Engine/Source/Runtime/Engine/Private/Components/SkyLightComponent.cpp b/Engine/Source/Runtime/Engine/Private/Components/SkyLightComponent.cpp
--- a/Engine/Source/Runtime/Engine/Private/Components/SkyLightComponent.cpp
+++ b/Engine/Source/Runtime/Engine/Private/Components/SkyLightComponent.cpp
@@ -984,6 +984,20 @@ void USkyLightComponent::SetCubemap(UTextureCube* NewCubemap)
 	}
 }
 
+void USkyLightComponent::SetSourceCubemapAngle(float NewValue)
+{
+	if (AreDynamicDataChangesAllowed()
+		&& SourceCubemapAngle != NewValue)
+	{
+		SourceCubemapAngle = NewValue;
+
+		if (GetWorld())
+		{
+			GetWorld()->UpdateAllSkyCaptures();
+		}
+	}
+}
+
 void USkyLightComponent::SetCubemapBlend(UTextureCube* SourceCubemap, UTextureCube* DestinationCubemap, float InBlendFraction)
 {
 	if (AreDynamicDataChangesAllowed()
