Updated UBlueprintGeneratedClass::GetAuthoritativeClass: this method now checks the following whenever ClassGeneratedBy is null:
 1. bCooked:  indicating whether the class was generated while cooking
 2. RootPackageHasAnyFlags(PKG_FilterEditorOnly) indicating whether we are working with cooked data

diff --git a/Engine/Source/Runtime/Engine/Private/BlueprintGeneratedClass.cpp b/Engine/Source/Runtime/Engine/Private/BlueprintGeneratedClass.cpp
--- a/Engine/Source/Runtime/Engine/Private/BlueprintGeneratedClass.cpp
+++ b/Engine/Source/Runtime/Engine/Private/BlueprintGeneratedClass.cpp
@@ -304,7 +304,7 @@ UClass* UBlueprintGeneratedClass::GetAuthoritativeClass()
  	if (nullptr == ClassGeneratedBy) // to track UE-11597 and UE-11595
  	{
 		// If this is a cooked blueprint, the generatedby class will have been discarded so we'll just have to assume we're authoritative!
-		if (bCooked)
+		if (bCooked || RootPackageHasAnyFlags(PKG_FilterEditorOnly))
 		{ 
 			return this;
 		}
