Take into account LODBias when calculating texture memory size in blueprints.

diff --git a/Engine/Source/Runtime/Engine/Private/Texture.cpp b/Engine/Source/Runtime/Engine/Private/Texture.cpp
--- a/Engine/Source/Runtime/Engine/Private/Texture.cpp
+++ b/Engine/Source/Runtime/Engine/Private/Texture.cpp
@@ -3293,7 +3293,7 @@
 
 int64 UTexture::Blueprint_GetMemorySize() const
 {
-	return CalcTextureMemorySizeEnum(TMC_AllMips);
+	return CalcTextureMemorySizeEnum(TMC_ResidentMips);
 }
 
 #if WITH_EDITOR
