Fix crash and deadlock when reloading a package containing skeletalmesh.
- Crash: There was a missing FScopeLock when reading the skeletal mesh imported bulk data this was crashing if we access the bulkdata in two different thread in same time.
- Deadlock: The reload package function is doing a TThreadSfeObjectIterator which deadlock if the inner code create a UObject. The skeletalmesh build can be running in another thread and can create UObject when building the morphtargets. The fix is to move the work outside of the iterator, the iteration now just store the UObject into a TArray and the work is execute in a for loop after the iteration.

#rnx

diff --git a/Engine/Source/Runtime/CoreUObject/Private/UObject/PackageReload.cpp b/Engine/Source/Runtime/CoreUObject/Private/UObject/PackageReload.cpp
--- a/Engine/Source/Runtime/CoreUObject/Private/UObject/PackageReload.cpp
+++ b/Engine/Source/Runtime/CoreUObject/Private/UObject/PackageReload.cpp
@@ -589,6 +589,9 @@ void ReloadPackages(const TArrayView<FReloadPackageData>& InPackagesToReload, TA
 
 			// Main pass to go through and fix-up any references pointing to data from the old package to point to data from the new package
 			// todo: multi-thread this like FHotReloadModule::ReplaceReferencesToReconstructedCDOs?
+			//The FThreadSafeObjectIterator will lock the global UObject array, to avoid potential deadlock we simply build the list
+			//of the potential referencers and do the reference fix serialize outside of the FThreadSafeObjectIterator.
+			TArray<UObject*> PotentialReferencers;
 			for (FThreadSafeObjectIterator ObjIter(UObject::StaticClass(), false, RF_NoFlags, EInternalObjectFlags::Garbage); ObjIter; ++ObjIter)
 			{
 				UObject* PotentialReferencer = *ObjIter;
@@ -601,15 +604,20 @@ void ReloadPackages(const TArrayView<FReloadPackageData>& InPackagesToReload, TA
 					AsClass = PotentialReferencer->GetTypedOuter<UClass>();
 				}
 
-				if(AsClass)
+				if (AsClass)
 				{
-					if( AsClass->HasAnyClassFlags(CLASS_NewerVersionExists) || 
+					if (AsClass->HasAnyClassFlags(CLASS_NewerVersionExists) ||
 						AsClass->HasAnyFlags(RF_NewerVersionExists))
 					{
 						continue;
 					}
 				}
+				PotentialReferencers.Add(PotentialReferencer);
+			}
+
 
+			for (UObject* PotentialReferencer : PotentialReferencers)
+			{
 				FixingUpReferencesSlowTask.EnterProgressFrame(1.0f);
 
 				PackageReloadInternal::FReplaceObjectReferencesArchive ReplaceRefsArchive(PotentialReferencer, OldObjectToNewData, ExistingPackages.Refs, NewPackages.Refs);
diff --git a/Engine/Source/Runtime/Engine/Private/Rendering/SkeletalMeshLODImporterData.cpp b/Engine/Source/Runtime/Engine/Private/Rendering/SkeletalMeshLODImporterData.cpp
--- a/Engine/Source/Runtime/Engine/Private/Rendering/SkeletalMeshLODImporterData.cpp
+++ b/Engine/Source/Runtime/Engine/Private/Rendering/SkeletalMeshLODImporterData.cpp
@@ -793,7 +793,13 @@
 		GeoImportVersion = ESkeletalMeshGeoImportVersions::Before_Versionning;
 		SkinningImportVersion = ESkeletalMeshSkinningImportVersions::Before_Versionning;
 	}
-	BulkData.Serialize(Ar, Owner);
+
+	// An exclusive lock is required so we can safely load/save the raw data from multiple threads
+	{
+		FWriteScopeLock ScopeLock(BulkDataLock.Get());
+		BulkData.Serialize(Ar, Owner);
+	}
+
 	Ar << Guid;
 	Ar << bGuidIsHash;
 
@@ -846,7 +852,6 @@
 	OutMesh.Empty();
 	if (BulkData.GetElementCount() > 0)
 	{
-#if WITH_EDITOR
 		// An exclusive lock is required so we can safely load the raw data from multiple threads
 		FWriteScopeLock ScopeLock(BulkDataLock.Get());
 
@@ -858,7 +863,6 @@
 			// This can't be called in -game mode because we're not allowed to load bulk data outside of EDL.
 			bHasBeenLoadedFromFileReader = BulkData.LoadBulkDataWithFileReader();
 		}
-#endif
 		// This is in a scope because the FBulkDataReader need to be destroyed in order
 		// to unlock the BulkData and allow UnloadBulkData to actually do its job.
 		{
@@ -871,7 +875,6 @@
 			Ar << OutMesh;
 		}
 
-#if WITH_EDITOR
 		// Throw away the bulk data allocation only in the case we can safely reload it from disk
 		// and if BulkData.LoadBulkDataWithFileReader() is allowed to work from any thread.
 		// This saves a significant amount of memory during map loading of Nanite Meshes.
@@ -879,7 +882,6 @@
 		{
 			verify(BulkData.UnloadBulkData());
 		}
-#endif
 	}
 }
 
@@ -893,7 +895,6 @@
 
 	// An exclusive lock is required so we can safely load the raw data from multiple threads
 	FWriteScopeLock ScopeLock(BulkDataLock.Get());
-#if WITH_EDITOR
 	// This allows any thread to be able to deserialize from the RawMesh directly
 	// from disk so we can unload bulk data from memory.
 	bool bHasBeenLoadedFromFileReader = false;
@@ -902,7 +903,6 @@
 		// This can't be called in -game mode because we're not allowed to load bulk data outside of EDL.
 		bHasBeenLoadedFromFileReader = BulkData.LoadBulkDataWithFileReader();
 	}
-#endif
 
 	bool bModified = false;
 	uint64 NumBytes = 0;
@@ -928,7 +928,6 @@
 			0 != FMemory::Memcmp(OldBytes.GetData(), NewBytes.GetData(), NumBytes);
 	}
 
-#if WITH_EDITOR
 	// Throw away the bulk data allocation only in the case we can safely reload it from disk
 	// and if BulkData.LoadBulkDataWithFileReader() is allowed to work from any thread.
 	// This saves a significant amount of memory during map loading of Nanite Meshes.
@@ -936,7 +935,6 @@
 	{
 		verify(BulkData.UnloadBulkData());
 	}
-#endif
 
 	if (bModified)
 	{
diff --git a/Engine/Source/Runtime/Engine/Private/SkeletalMeshEditorData.cpp b/Engine/Source/Runtime/Engine/Private/SkeletalMeshEditorData.cpp
--- a/Engine/Source/Runtime/Engine/Private/SkeletalMeshEditorData.cpp
+++ b/Engine/Source/Runtime/Engine/Private/SkeletalMeshEditorData.cpp
@@ -68,6 +68,12 @@ bool USkeletalMeshEditorData::RemoveLODImportedData(int32 LODIndex)
 void USkeletalMeshEditorData::Serialize(FArchive& Ar)
 {
 	Super::Serialize(Ar);
+	
+	if (!Ar.IsLoading() && !Ar.IsSaving() && Ar.IsObjectReferenceCollector())
+	{
+		//There is no UObject reference in the skeletal mesh editor data
+		return;
+	}
 
 	//Serialize all LODs Raw imported source data
 	FRawSkeletalMeshBulkData::Serialize(Ar, RawSkeletalMeshBulkDatas, this);
