Add bEvaluateWorldPositionOffset property which will be used to determine whether to evaluate WPO.
Actual usage of the property will get merged in from UE5, but adding the placeholder property here to reduce pain from data merges.
This change also remaps the original bEvaluateWorldPositionOffset on SMC into bEvaluateWorldPositionOffsetInRayTracing, because this var was only ever driven by ray tracing specific methods.
We expect to need to special case robomerging the remap of bEvaluateWorldPositionOffsetInRayTracing since it needs to be done from different branches with different branch specific versions.
#rnx

diff --git a/Engine/Source/Runtime/Core/Public/UObject/FortniteMainBranchObjectVersion.h b/Engine/Source/Runtime/Core/Public/UObject/FortniteMainBranchObjectVersion.h
--- a/Engine/Source/Runtime/Core/Public/UObject/FortniteMainBranchObjectVersion.h
+++ b/Engine/Source/Runtime/Core/Public/UObject/FortniteMainBranchObjectVersion.h
@@ -214,6 +214,9 @@ struct CORE_API FFortniteMainBranchObjectVersion
 		// Single-frame/key AnimDataModel patch-up
 		SingleFrameAndKeyAnimModel,
 
+		// Remapped bEvaluateWorldPositionOffset to bEvaluateWorldPositionOffsetInRayTracing
+		RemappedEvaluateWorldPositionOffsetInRayTracing,
+
 		// -----<new versions can be added above this line>-------------------------------------------------
 		VersionPlusOne,
 		LatestVersion = VersionPlusOne - 1
diff --git a/Engine/Source/Runtime/Engine/Private/Components/StaticMeshComponent.cpp b/Engine/Source/Runtime/Engine/Private/Components/StaticMeshComponent.cpp
--- a/Engine/Source/Runtime/Engine/Private/Components/StaticMeshComponent.cpp
+++ b/Engine/Source/Runtime/Engine/Private/Components/StaticMeshComponent.cpp
@@ -313,6 +313,7 @@ void UStaticMeshComponent::Serialize(FArchive& Ar)
 
 	Ar.UsingCustomVersion(FRenderingObjectVersion::GUID);
 	Ar.UsingCustomVersion(FFortniteNCBranchObjectVersion::GUID);
+	Ar.UsingCustomVersion(FFortniteMainBranchObjectVersion::GUID);
 
 #if WITH_EDITORONLY_DATA
 	if (Ar.IsCooking())
@@ -366,7 +367,8 @@ void UStaticMeshComponent::Serialize(FArchive& Ar)
 	}
 #endif
 
-	if (Ar.CustomVer(FFortniteNCBranchObjectVersion::GUID) < FFortniteNCBranchObjectVersion::RemappedEvaluateWorldPositionOffsetInRayTracing)
+	if (Ar.CustomVer(FFortniteNCBranchObjectVersion::GUID) < FFortniteNCBranchObjectVersion::RemappedEvaluateWorldPositionOffsetInRayTracing &&
+	Ar.CustomVer(FFortniteMainBranchObjectVersion::GUID) < FFortniteMainBranchObjectVersion::RemappedEvaluateWorldPositionOffsetInRayTracing)
 	{
 		bEvaluateWorldPositionOffsetInRayTracing = bEvaluateWorldPositionOffset;
 		bEvaluateWorldPositionOffset = false; // Default WPO evaluation off
