fix FKShapeElem struct top-level comment

diff --git a/Engine/Source/Runtime/Engine/Classes/PhysicsEngine/ShapeElem.h b/Engine/Source/Runtime/Engine/Classes/PhysicsEngine/ShapeElem.h
--- a/Engine/Source/Runtime/Engine/Classes/PhysicsEngine/ShapeElem.h
+++ b/Engine/Source/Runtime/Engine/Classes/PhysicsEngine/ShapeElem.h
@@ -25,7 +25,7 @@ namespace EAggCollisionShape
 	};
 }
 
-/** Sphere shape used for collision */
+/** Base class of shapes used for collision, such as Sphere, Box, Sphyl, Convex, TaperedCapsule or LevelSet */
 USTRUCT()
 struct FKShapeElem
 {
