[Backout] - CL20377958
Original CL Desc
-----------------------------------------------------------------
World Partition: only pin dirty actors on save if we are holding the last hard references to avoid unloading.

#rnx

diff --git a/Engine/Source/Runtime/Engine/Private/WorldPartition/WorldPartition.cpp b/Engine/Source/Runtime/Engine/Private/WorldPartition/WorldPartition.cpp
--- a/Engine/Source/Runtime/Engine/Private/WorldPartition/WorldPartition.cpp
+++ b/Engine/Source/Runtime/Engine/Private/WorldPartition/WorldPartition.cpp
@@ -328,14 +328,9 @@ void UWorldPartition::OnPackageDirtyStateChanged(UPackage* Package)
 			{
 				DirtyActors.Add(ActorHandle, Actor);
 			}
-			else
+			else if (PinnedActors)
 			{
-				// If we hold the last reference to that actor, pin it to avoid unloading
-				if (PinnedActors && (ActorHandle->GetHardRefCount() == 1))
-				{
-					PinnedActors->AddActors({ ActorHandle });
-				}
-
+				PinnedActors->AddActors({ ActorHandle });
 				DirtyActors.Remove(ActorHandle);
 			}
 		}
@@ -1460,4 +1455,4 @@ FBox UWorldPartition::GetEditorWorldBounds() const
 }
 #endif
 
-#undef LOCTEXT_NAMESPACE
+#undef LOCTEXT_NAMESPACE
\ No newline at end of file
