Fixes crash in Dynamic Resolution with -ExecCmds="stat unitgraph"

diff --git a/Engine/Source/Runtime/Engine/Private/DynamicResolution.cpp b/Engine/Source/Runtime/Engine/Private/DynamicResolution.cpp
--- a/Engine/Source/Runtime/Engine/Private/DynamicResolution.cpp
+++ b/Engine/Source/Runtime/Engine/Private/DynamicResolution.cpp
@@ -94,7 +94,7 @@
 	BudgetSetting.MaxResolutionFraction      = DynamicRenderScaling::GetPercentageCVarToFraction(CVarDynamicResMaxSP);
 	// BudgetSetting.BudgetMs depends on the cost of other buckets.
 	BudgetSetting.ChangeThreshold            = DynamicRenderScaling::GetPercentageCVarToFraction(CVarChangeThreshold);
-	BudgetSetting.IncreaseAmortizationFactor = CVarIncreaseAmortizationFactor.GetValueOnRenderThread();
+	BudgetSetting.IncreaseAmortizationFactor = CVarIncreaseAmortizationFactor.GetValueOnAnyThread();
 
 	// CVarTargetedGPUHeadRoomPercentage is taken into account on the entire frame instead.
 	BudgetSetting.TargetedHeadRoom = 0.0f;
diff --git a/Engine/Source/Runtime/RenderCore/Private/DynamicRenderScaling.cpp b/Engine/Source/Runtime/RenderCore/Private/DynamicRenderScaling.cpp
--- a/Engine/Source/Runtime/RenderCore/Private/DynamicRenderScaling.cpp
+++ b/Engine/Source/Runtime/RenderCore/Private/DynamicRenderScaling.cpp
@@ -137,6 +137,8 @@
 		UE_LOG(LogRendererCore, Warning, TEXT("Consider bumping DynamicRenderScaling::TMap::kInlineAlloactedBudgets to %d."), GRDSGlobalListLinkSize);
 	}
 #endif
+
+	CachedSettings = (*HeuristicSettingsGetter)();
 }
 
 FBudget::~FBudget()
