World Partition: only pin dirty actors on save if we are holding the last hard references to avoid unloading.

#rnx

diff --git a/Engine/Source/Runtime/Engine/Private/WorldPartition/WorldPartition.cpp b/Engine/Source/Runtime/Engine/Private/WorldPartition/WorldPartition.cpp
--- a/Engine/Source/Runtime/Engine/Private/WorldPartition/WorldPartition.cpp
+++ b/Engine/Source/Runtime/Engine/Private/WorldPartition/WorldPartition.cpp
@@ -328,9 +328,14 @@ void UWorldPartition::OnPackageDirtyStateChanged(UPackage* Package)
 			{
 				DirtyActors.Add(ActorHandle, Actor);
 			}
-			else if (PinnedActors)
+			else if (ActorHandle->GetHardRefCount())
 			{
-				PinnedActors->AddActors({ ActorHandle });
+				// If we hold the last reference to that actor, pin it to avoid unloading
+				if (PinnedActors && (ActorHandle->GetHardRefCount() == 1))
+				{
+					PinnedActors->AddActors({ ActorHandle });
+				}
+
 				DirtyActors.Remove(ActorHandle);
 			}
 		}
@@ -1101,7 +1106,7 @@ void UWorldPartition::AddReferencedObjects(UObject* InThis, FReferenceCollector&
 		UWorldPartition* This = CastChecked<UWorldPartition>(InThis);
 	
 		Collector.AllowEliminatingReferences(false);
-		for (auto [ActorReference, Actor] : This->DirtyActors)
+		for (auto& [ActorReference, Actor] : This->DirtyActors)
 		{
 			Collector.AddReferencedObject(Actor);
 		}
@@ -1455,4 +1460,4 @@ FBox UWorldPartition::GetEditorWorldBounds() const
 }
 #endif
 
-#undef LOCTEXT_NAMESPACE
\ No newline at end of file
+#undef LOCTEXT_NAMESPACE
