Added World Partition runtime hash streaming generation logging (only enabled when cooking).

diff --git a/Engine/Source/Runtime/Engine/Private/WorldPartition/WorldPartitionFileLogger.cpp b/Engine/Source/Runtime/Engine/Private/WorldPartition/WorldPartitionFileLogger.cpp
new file mode 100644
--- /dev/null
+++ b/Engine/Source/Runtime/Engine/Private/WorldPartition/WorldPartitionFileLogger.cpp
@@ -0,0 +1,60 @@
+// Copyright Epic Games, Inc. All Rights Reserved.
+
+#include "WorldPartition/WorldPartitionFileLogger.h"
+
+#if WITH_EDITOR
+#include "HAL/FileManager.h"
+#endif
+
+#if WITH_EDITOR
+
+FWorldPartitionFileLogger::FWorldPartitionFileLogger(const FString& InLogFilename)
+	: IndentationCount(0)
+{
+	LogFile = IFileManager::Get().CreateFileWriter(*InLogFilename);
+}
+
+FWorldPartitionFileLogger::~FWorldPartitionFileLogger()
+{
+	if (LogFile)
+	{
+		LogFile->Close();
+		delete LogFile;
+	}
+}
+
+void FWorldPartitionFileLogger::WriteLine(const FString& InLine, bool bIndent /* = false*/)
+{
+	if (LogFile)
+	{
+		TStringBuilder<512> Builder;
+		for (int i = 0; i < IndentationCount - 1; ++i)
+		{
+			Builder += TEXT(" |  ");
+		}
+		if (IndentationCount > 0)
+		{
+			Builder += TEXT(" |- ");
+		}
+		if (bIndent)
+		{
+			Builder += TEXT("[+] ");
+			++IndentationCount;
+		}
+		Builder += InLine;
+		Builder += LINE_TERMINATOR;
+		FString LineEntry = Builder.ToString();
+		LogFile->Serialize(TCHAR_TO_ANSI(*LineEntry), LineEntry.Len());
+	}
+}
+
+void FWorldPartitionFileLogger::DecrementIndentation()
+{
+	if (LogFile)
+	{
+		check(IndentationCount);
+		--IndentationCount;
+	}
+}
+
+#endif
\ No newline at end of file
diff --git a/Engine/Source/Runtime/Engine/Private/WorldPartition/WorldPartitionRuntimeCell.cpp b/Engine/Source/Runtime/Engine/Private/WorldPartition/WorldPartitionRuntimeCell.cpp
--- a/Engine/Source/Runtime/Engine/Private/WorldPartition/WorldPartitionRuntimeCell.cpp
+++ b/Engine/Source/Runtime/Engine/Private/WorldPartition/WorldPartitionRuntimeCell.cpp
@@ -9,6 +9,10 @@
 #include "WorldPartition/DataLayer/DataLayersID.h"
 #include "Engine/World.h"
 
+#if WITH_EDITOR
+#include "WorldPartition/WorldPartitionFileLogger.h"
+#endif
+
 int32 UWorldPartitionRuntimeCell::StreamingSourceCacheEpoch = 0;
 
 UWorldPartitionRuntimeCell::UWorldPartitionRuntimeCell(const FObjectInitializer& ObjectInitializer)
@@ -74,6 +78,12 @@ void UWorldPartitionRuntimeCell::UpdateDebugName()
 	DebugName = Builder.ToString();
 }
 
+void UWorldPartitionRuntimeCell::LogStreamingGeneration(FWorldPartitionFileLogger& Logger)
+{
+	Logger.WriteLine(FString::Printf(TEXT("Actor Count: %d"), GetActorCount()));
+	Logger.WriteLine(FString::Printf(TEXT("MinMaxZ: %s"), *GetMinMaxZ().ToString()));
+}
+
 #endif
 
 bool UWorldPartitionRuntimeCell::CacheStreamingSourceInfo(const UWorldPartitionRuntimeCell::FStreamingSourceInfo& Info) const
@@ -148,4 +158,4 @@ bool UWorldPartitionRuntimeCell::ContainsDataLayer(const UDataLayerAsset* DataLa
 bool UWorldPartitionRuntimeCell::ContainsDataLayer(const UDataLayerInstance* DataLayerInstance) const
 {
 	return GetDataLayers().Contains(DataLayerInstance->GetDataLayerFName());
-}
\ No newline at end of file
+}
diff --git a/Engine/Source/Runtime/Engine/Private/WorldPartition/WorldPartitionRuntimeHash.cpp b/Engine/Source/Runtime/Engine/Private/WorldPartition/WorldPartitionRuntimeHash.cpp
--- a/Engine/Source/Runtime/Engine/Private/WorldPartition/WorldPartitionRuntimeHash.cpp
+++ b/Engine/Source/Runtime/Engine/Private/WorldPartition/WorldPartitionRuntimeHash.cpp
@@ -10,6 +10,7 @@
 #if WITH_EDITOR
 #include "WorldPartition/WorldPartitionHandle.h"
 #include "WorldPartition/DataLayer/WorldDataLayers.h"
+#include "WorldPartition/WorldPartitionFileLogger.h"
 #endif
 
 #define LOCTEXT_NAMESPACE "WorldPartition"
@@ -37,6 +38,27 @@ void UWorldPartitionRuntimeHash::OnEndPlay()
 	ModifiedActorDescListForPIE.Empty();
 }
 
+void UWorldPartitionRuntimeHash::LogStreamingGeneration(FWorldPartitionFileLogger& Logger)
+{
+	Logger.WriteLine(TEXT("----------------------------------------------------------------------------------------------------------------"));
+	Logger.WriteLine(FString::Printf(TEXT("%s - Persistent Level"), *GetWorld()->GetName()));
+	Logger.WriteLine(TEXT("----------------------------------------------------------------------------------------------------------------"));
+	Logger.WriteLine(FString::Printf(TEXT("Always loaded Actor Count: %d "), GetWorld()->PersistentLevel->Actors.Num()));
+	Logger.WriteLine(TEXT(""));
+}
+
+void UWorldPartitionRuntimeHash::LogStreamingGeneration()
+{
+	// For now only log for cook
+	if (IsRunningCookCommandlet())
+	{
+		const TCHAR* Suffix = TEXT("Cook");
+		const FString StateLogOutputFilename = FPaths::ProjectSavedDir() / TEXT("WorldPartition") / FString::Printf(TEXT("StreamingGeneration-RuntimeSpatialHash-%s-%s-%s.log"), Suffix, *GetWorld()->GetName(), *FDateTime::Now().ToString());
+		FWorldPartitionFileLogger Logger(StateLogOutputFilename);
+		LogStreamingGeneration(Logger);
+	}
+}
+
 void UWorldPartitionRuntimeHash::ForceExternalActorLevelReference(bool bForceExternalActorLevelReferenceForPIE)
 {
 	// Do this only on non game worlds prior to PIE so that always loaded actors get duplicated with the world
diff --git a/Engine/Source/Runtime/Engine/Private/WorldPartition/WorldPartitionRuntimeLevelStreamingCell.cpp b/Engine/Source/Runtime/Engine/Private/WorldPartition/WorldPartitionRuntimeLevelStreamingCell.cpp
--- a/Engine/Source/Runtime/Engine/Private/WorldPartition/WorldPartitionRuntimeLevelStreamingCell.cpp
+++ b/Engine/Source/Runtime/Engine/Private/WorldPartition/WorldPartitionRuntimeLevelStreamingCell.cpp
@@ -12,6 +12,7 @@
 #if WITH_EDITOR
 #include "WorldPartition/WorldPartitionLevelStreamingPolicy.h"
 #include "WorldPartition/WorldPartitionLevelHelper.h"
+#include "WorldPartition/WorldPartitionFileLogger.h"
 #include "Engine/LevelStreamingAlwaysLoaded.h"
 #endif
 
@@ -219,6 +220,17 @@ FString UWorldPartitionRuntimeLevelStreamingCell::GetPackageNameToCreate() const
 	return UWorldPartitionLevelStreamingPolicy::GetCellPackagePath(GetFName(), OuterWorld);
 }
 
+void UWorldPartitionRuntimeLevelStreamingCell::LogStreamingGeneration(FWorldPartitionFileLogger& Logger)
+{
+	Super::LogStreamingGeneration(Logger);
+
+	for (const FWorldPartitionRuntimeCellObjectMapping& Mapping : Packages)
+	{
+		Logger.WriteLine(FString::Printf(TEXT("Actor Path: %s"), *Mapping.Path.ToString()));
+		Logger.WriteLine(FString::Printf(TEXT("Actor Package: %s"), *Mapping.Package.ToString()));
+	}
+}
+
 #endif
 
 UWorldPartitionLevelStreamingDynamic* UWorldPartitionRuntimeLevelStreamingCell::GetOrCreateLevelStreaming() const
diff --git a/Engine/Source/Runtime/Engine/Private/WorldPartition/WorldPartitionRuntimeSpatialHash.cpp b/Engine/Source/Runtime/Engine/Private/WorldPartition/WorldPartitionRuntimeSpatialHash.cpp
--- a/Engine/Source/Runtime/Engine/Private/WorldPartition/WorldPartitionRuntimeSpatialHash.cpp
+++ b/Engine/Source/Runtime/Engine/Private/WorldPartition/WorldPartitionRuntimeSpatialHash.cpp
@@ -37,6 +37,7 @@
 #include "Components/LineBatchComponent.h"
 
 #if WITH_EDITOR
+#include "WorldPartition/WorldPartitionFileLogger.h"
 #include "Editor/EditorEngine.h"
 #include "CookPackageSplitter.h"
 #include "Misc/Parse.h"
@@ -852,6 +853,89 @@ bool UWorldPartitionRuntimeSpatialHash::GenerateStreaming(UWorldPartitionStreami
 	return true;
 }
 
+void UWorldPartitionRuntimeSpatialHash::LogStreamingGeneration(FWorldPartitionFileLogger& Logger)
+{
+	Super::LogStreamingGeneration(Logger);
+
+	for (FSpatialHashStreamingGrid& StreamingGrid : StreamingGrids)
+	{
+		Logger.WriteLine(TEXT("----------------------------------------------------------------------------------------------------------------"));
+		Logger.WriteLine(FString::Printf(TEXT("%s - Runtime Spatial Hash - Streaming Grid - %s"), *GetWorld()->GetName(), *StreamingGrid.GridName.ToString()));
+		Logger.WriteLine(TEXT("----------------------------------------------------------------------------------------------------------------"));
+		Logger.WriteLine(FString::Printf(TEXT("            Origin: %s"), *StreamingGrid.Origin.ToString()));
+		Logger.WriteLine(FString::Printf(TEXT("         Cell Size: %d"), StreamingGrid.CellSize));
+		Logger.WriteLine(FString::Printf(TEXT("      World Bounds: %s"), *StreamingGrid.WorldBounds.ToString()));
+		Logger.WriteLine(FString::Printf(TEXT("     Loading Range: %3.2f"), StreamingGrid.LoadingRange));
+		Logger.WriteLine(FString::Printf(TEXT("Block Slow Loading: %s"), StreamingGrid.bBlockOnSlowStreaming ? TEXT("Yes") : TEXT("No")));
+		Logger.WriteLine(FString::Printf(TEXT(" ClientOnlyVisible: %s"), StreamingGrid.bClientOnlyVisible ? TEXT("Yes") : TEXT("No")));
+		Logger.WriteLine(TEXT(""));
+		if (const UHLODLayer* HLODLayer = StreamingGrid.HLODLayer)
+		{
+			Logger.WriteLine(FString::Printf(TEXT("    HLOD Layer: %s"), *HLODLayer->GetName()));
+		}
+
+		{
+			struct FGridLevelStats
+			{
+				int32 CellCount;
+				int32 CellSize;
+				int32 ActorCount;
+			};
+
+			TArray<FGridLevelStats> LevelsStats;
+			int32 TotalActorCount = 0;
+			
+			{
+				int32 Level = 0;
+				for (FSpatialHashStreamingGridLevel& GridLevel : StreamingGrid.GridLevels)
+				{
+					int32 LevelCellCount = 0;
+					int32 LevelActorCount = 0;
+					for (FSpatialHashStreamingGridLayerCell& LayerCell : GridLevel.LayerCells)
+					{
+						LevelCellCount += LayerCell.GridCells.Num();
+						for (TObjectPtr<UWorldPartitionRuntimeSpatialHashCell>& Cell : LayerCell.GridCells)
+						{
+							LevelActorCount += Cell->GetActorCount();
+						}
+					}
+					LevelsStats.Add({ LevelCellCount, (StreamingGrid.CellSize << Level), LevelActorCount });
+					TotalActorCount += LevelActorCount;
+					++Level;
+				}
+				TotalActorCount = (TotalActorCount > 0) ? TotalActorCount : 1;
+			}
+
+			Logger.WriteLine(FString::Printf(TEXT("Grid Levels: %d"), StreamingGrid.GridLevels.Num()), true);
+			for (int Level=0; Level< LevelsStats.Num(); ++Level)
+			{
+				Logger.WriteLine(FString::Printf(TEXT("Level %2d: Cell Count %4d | Cell Size %7d | Actor Count %4d (%3.1f%%)"), Level, LevelsStats[Level].CellCount, LevelsStats[Level].CellSize, LevelsStats[Level].ActorCount, (100.f*LevelsStats[Level].ActorCount)/TotalActorCount));
+			}
+			Logger.DecrementIndentation();
+		}
+
+		{
+			Logger.WriteLine(TEXT(""));
+			int32 Level = 0;
+			for (const FSpatialHashStreamingGridLevel& GridLevel : StreamingGrid.GridLevels)
+			{
+				Logger.WriteLine(FString::Printf(TEXT("Content of Grid Level %d"), Level++), true);
+				for (const FSpatialHashStreamingGridLayerCell& LayerCell : GridLevel.LayerCells)
+				{
+					for (const TObjectPtr<UWorldPartitionRuntimeSpatialHashCell>& Cell : LayerCell.GridCells)
+					{
+						Logger.WriteLine(FString::Printf(TEXT("Content of Cell %s"), *Cell->GetDebugName()), true);
+						Cell->LogStreamingGeneration(Logger);
+						Logger.DecrementIndentation();
+					}
+				}
+				Logger.DecrementIndentation();
+			}
+		}
+		Logger.WriteLine(TEXT(""));
+	}
+}
+
 FName UWorldPartitionRuntimeSpatialHash::GetCellName(UWorldPartition* WorldPartition, FName InGridName, const FIntVector& InCellGlobalCoord, const FDataLayersID& InDataLayerID)
 {
 	const FString PackageName = FPackageName::GetShortName(WorldPartition->GetPackage());
diff --git a/Engine/Source/Runtime/Engine/Private/WorldPartition/WorldPartitionStreamingGeneration.cpp b/Engine/Source/Runtime/Engine/Private/WorldPartition/WorldPartitionStreamingGeneration.cpp
--- a/Engine/Source/Runtime/Engine/Private/WorldPartition/WorldPartitionStreamingGeneration.cpp
+++ b/Engine/Source/Runtime/Engine/Private/WorldPartition/WorldPartitionStreamingGeneration.cpp
@@ -627,6 +627,8 @@ bool UWorldPartition::GenerateStreaming(TArray<FString>* OutPackagesToGenerate)
 	check(RuntimeHash);
 	if (RuntimeHash->GenerateStreaming(StreamingPolicy, ActorClusterContext, OutPackagesToGenerate))
 	{
+		RuntimeHash->LogStreamingGeneration();
+
 		StreamingPolicy->PrepareActorToCellRemapping();
 		return true;
 	}
diff --git a/Engine/Source/Runtime/Engine/Public/WorldPartition/WorldPartitionFileLogger.h b/Engine/Source/Runtime/Engine/Public/WorldPartition/WorldPartitionFileLogger.h
new file mode 100644
--- /dev/null
+++ b/Engine/Source/Runtime/Engine/Public/WorldPartition/WorldPartitionFileLogger.h
@@ -0,0 +1,20 @@
+// Copyright Epic Games, Inc. All Rights Reserved.
+
+#pragma once
+
+#include "CoreMinimal.h"
+
+#if WITH_EDITOR
+struct FWorldPartitionFileLogger
+{
+public:
+	FWorldPartitionFileLogger(const FString& InLogFilename);
+	~FWorldPartitionFileLogger();
+	void WriteLine(const FString& InLine, bool bIndent = false);
+	void DecrementIndentation();
+
+private:
+	int32 IndentationCount;
+	class FArchive* LogFile;
+};
+#endif
\ No newline at end of file
diff --git a/Engine/Source/Runtime/Engine/Public/WorldPartition/WorldPartitionRuntimeCell.h b/Engine/Source/Runtime/Engine/Public/WorldPartition/WorldPartitionRuntimeCell.h
--- a/Engine/Source/Runtime/Engine/Public/WorldPartition/WorldPartitionRuntimeCell.h
+++ b/Engine/Source/Runtime/Engine/Public/WorldPartition/WorldPartitionRuntimeCell.h
@@ -14,6 +14,7 @@
 
 class UDataLayerAsset;
 class UDataLayerInstance;
+struct FWorldPartitionFileLogger;
 
 enum class EWorldPartitionRuntimeCellVisualizeMode
 {
@@ -149,7 +150,7 @@ class UWorldPartitionRuntimeCell : public UObject, public IWorldPartitionCell
 
 	static void DirtyStreamingSourceCacheEpoch() { ++UWorldPartitionRuntimeCell::StreamingSourceCacheEpoch; }
 
-	// IRuntimeCell Interface begin
+	// IWorldPartitionCell Interface begin
 	virtual TArray<const UDataLayerInstance*> GetDataLayerInstances() const override;
 	virtual bool ContainsDataLayer(const UDataLayerAsset* DataLayerAsset) const override;
 	virtual bool ContainsDataLayer(const UDataLayerInstance* DataLayerInstance) const override;
@@ -159,7 +160,7 @@ class UWorldPartitionRuntimeCell : public UObject, public IWorldPartitionCell
 	{
 		return Algo::AnyOf(DataLayers, [&InDataLayers](const FName& DataLayer) { return InDataLayers.Contains(DataLayer); });
 	}
-	// IRuntimeCell Interface end
+	// IWorldPartitionCell Interface end
 
 	const FVector2D& GetMinMaxZ() const { return MinMaxZ; }
 
@@ -183,6 +184,8 @@ class UWorldPartitionRuntimeCell : public UObject, public IWorldPartitionCell
 	virtual FString GetPackageNameToCreate() const PURE_VIRTUAL(UWorldPartitionRuntimeCell::GetPackageNameToCreate, return FString(""););
 
 	void SetIsHLOD(bool bInIsHLOD) { bIsHLOD = bInIsHLOD; }
+
+	virtual void LogStreamingGeneration(FWorldPartitionFileLogger& Logger);
 #endif
 	
 	bool GetIsHLOD() const { return bIsHLOD; }
diff --git a/Engine/Source/Runtime/Engine/Public/WorldPartition/WorldPartitionRuntimeHash.h b/Engine/Source/Runtime/Engine/Public/WorldPartition/WorldPartitionRuntimeHash.h
--- a/Engine/Source/Runtime/Engine/Public/WorldPartition/WorldPartitionRuntimeHash.h
+++ b/Engine/Source/Runtime/Engine/Public/WorldPartition/WorldPartitionRuntimeHash.h
@@ -16,6 +16,7 @@
 #include "WorldPartitionRuntimeHash.generated.h"
 
 class FActorClusterContext;
+struct FWorldPartitionFileLogger;
 
 UENUM()
 enum class EWorldPartitionStreamingPerformance : uint8
@@ -42,6 +43,9 @@ class ENGINE_API UWorldPartitionRuntimeHash : public UObject
 	// PIE/Game methods
 	void OnBeginPlay();
 	void OnEndPlay();
+
+	// GenerateStreaming Logging
+	void LogStreamingGeneration();
 #endif
 
 	class FStreamingSourceCells
@@ -76,6 +80,9 @@ class ENGINE_API UWorldPartitionRuntimeHash : public UObject
 
 protected:
 	virtual EWorldPartitionStreamingPerformance GetStreamingPerformanceForCell(const UWorldPartitionRuntimeCell* Cell) const { return EWorldPartitionStreamingPerformance::Good; }
+#if WITH_EDITOR
+	virtual void LogStreamingGeneration(FWorldPartitionFileLogger& Logger);
+#endif
 
 private:
 #if WITH_EDITOR
diff --git a/Engine/Source/Runtime/Engine/Public/WorldPartition/WorldPartitionRuntimeLevelStreamingCell.h b/Engine/Source/Runtime/Engine/Public/WorldPartition/WorldPartitionRuntimeLevelStreamingCell.h
--- a/Engine/Source/Runtime/Engine/Public/WorldPartition/WorldPartitionRuntimeLevelStreamingCell.h
+++ b/Engine/Source/Runtime/Engine/Public/WorldPartition/WorldPartitionRuntimeLevelStreamingCell.h
@@ -34,6 +34,7 @@ class UWorldPartitionRuntimeLevelStreamingCell : public UWorldPartitionRuntimeSp
 	//~Begin UWorldPartitionRuntimeCell Interface
 	virtual void AddActorToCell(const FWorldPartitionActorDescView& ActorDescView, const FActorContainerID& InContainerID, const FTransform& InContainerTransform, const UActorDescContainer* InContainer) override;
 	virtual int32 GetActorCount() const override;
+	virtual void LogStreamingGeneration(FWorldPartitionFileLogger& Logger) override;
 	// Cook methods
 	virtual bool PrepareCellForCook(UPackage* InPackage) override;
 	virtual bool PopulateGeneratedPackageForCook(UPackage* InPackage) override;
diff --git a/Engine/Source/Runtime/Engine/Public/WorldPartition/WorldPartitionRuntimeSpatialHash.h b/Engine/Source/Runtime/Engine/Public/WorldPartition/WorldPartitionRuntimeSpatialHash.h
--- a/Engine/Source/Runtime/Engine/Public/WorldPartition/WorldPartitionRuntimeSpatialHash.h
+++ b/Engine/Source/Runtime/Engine/Public/WorldPartition/WorldPartitionRuntimeSpatialHash.h
@@ -255,6 +255,7 @@ protected:
 
 #if WITH_EDITOR
 	virtual bool GenerateStreaming(class UWorldPartitionStreamingPolicy* StreamingPolicy, const FActorClusterContext& ActorClusterContext, TArray<FString>* OutPackagesToGenerate = nullptr) override;
+	virtual void LogStreamingGeneration(FWorldPartitionFileLogger& Logger) override;
 #endif
 
 private:
