Chaos: Fix bounding box computation for Box based field commands that was not accounting for rotation

diff --git a/Engine/Source/Runtime/Engine/Private/PhysicsField/PhysicsFieldComponent.cpp b/Engine/Source/Runtime/Engine/Private/PhysicsField/PhysicsFieldComponent.cpp
--- a/Engine/Source/Runtime/Engine/Private/PhysicsField/PhysicsFieldComponent.cpp
+++ b/Engine/Source/Runtime/Engine/Private/PhysicsField/PhysicsFieldComponent.cpp
@@ -1141,9 +1141,15 @@ void FPhysicsFieldInstance::BuildNodeBounds(FFieldNodeBase* FieldNode, FVector&
 		else if (FieldNode->SerializationType() == FFieldNodeBase::ESerializationType::FieldNode_FBoxFalloff)
 		{
 			FBoxFalloff* LocalNode = StaticCast<FBoxFalloff*>(FieldNode);
-
-			MinBounds = (LocalNode->Default == 0) ? LocalNode->Transform.GetTranslation() - 50*LocalNode->Transform.GetScale3D() : FVector(-FLT_MAX);
-			MaxBounds = (LocalNode->Default == 0) ? LocalNode->Transform.GetTranslation() + 50*LocalNode->Transform.GetScale3D() : FVector(FLT_MAX);
+			MinBounds = FVector(-FLT_MAX);
+			MaxBounds = FVector(FLT_MAX);
+			if (LocalNode->Default == 0)
+			{
+				const FBox UnitBox(FVector(-50), FVector(50));
+				const FBox BoundingBox = UnitBox.TransformBy(LocalNode->Transform);
+				MinBounds = BoundingBox.Min;
+				MaxBounds = BoundingBox.Max;
+			}
 			MaxMagnitude = LocalNode->Magnitude;
 		}
 		else if (FieldNode->SerializationType() == FFieldNodeBase::ESerializationType::FieldNode_FSumScalar)
