Fixed check in Build Texture Streaming : Texture streaming build now detects and properly handles when 2 different streamable textures have the same guid.

diff --git a/Engine/Source/Runtime/Engine/Private/Level.cpp b/Engine/Source/Runtime/Engine/Private/Level.cpp
--- a/Engine/Source/Runtime/Engine/Private/Level.cpp
+++ b/Engine/Source/Runtime/Engine/Private/Level.cpp
@@ -2067,6 +2067,18 @@
 	if (InTexture->LevelIndex == INDEX_NONE)
 	{
 		// If this is the first time this texture gets processed in the packing process, encode it.
+		int32 ExistingTextureIndex = StreamingTextureGuids.Find(InTexture->GetLightingGuid());
+		if (ExistingTextureIndex != INDEX_NONE)
+		{
+			check(StreamingTextures.IsValidIndex(ExistingTextureIndex));
+			// Detect that another texture with the same guid was registered
+			int32 Index = StreamingTextures.Find(*InTexture->GetPathName());
+			check(Index == INDEX_NONE);
+			UE_LOG(LogLevel, Warning, TEXT("Another streamable texture %s was already registered with the same guid. Texture %s needs to be resaved."), *StreamingTextures[ExistingTextureIndex].ToString(), *InTexture->GetPathName());
+			InTexture->Modify();
+			InTexture->SetLightingGuid();
+		}
+
 		InTexture->LevelIndex = (int32)RegisterStreamableTexture(InTexture->GetPathName(), InTexture->GetLightingGuid());
 	}
 	check(StreamingTextureGuids.IsValidIndex(InTexture->LevelIndex));
