Setting timecode provider to None does nothing

diff --git a/Engine/Source/Runtime/Engine/Private/UnrealEngine.cpp b/Engine/Source/Runtime/Engine/Private/UnrealEngine.cpp
--- a/Engine/Source/Runtime/Engine/Private/UnrealEngine.cpp
+++ b/Engine/Source/Runtime/Engine/Private/UnrealEngine.cpp
@@ -2920,6 +2920,10 @@ static void LoadTimecodeProvider(UEngine* Engine)
 #endif
 		Engine->SetTimecodeProvider(NewTimecodeProvider);
 	}
+	else
+	{
+		Engine->SetTimecodeProvider(nullptr);
+	}
 }
 
 /**
