World Partition:
- Added FHierarchicalLogArchive proxy that can be used to output a structured log archive.
- Converted world partition streaming generation logging to use FHierarchicalLogArchive.

#rnx

diff --git a/Engine/Source/Runtime/Core/Private/Misc/HierarchicalLogArchive.cpp b/Engine/Source/Runtime/Core/Private/Misc/HierarchicalLogArchive.cpp
new file mode 100644
--- /dev/null
+++ b/Engine/Source/Runtime/Core/Private/Misc/HierarchicalLogArchive.cpp
@@ -0,0 +1,34 @@
+// Copyright Epic Games, Inc. All Rights Reserved.
+
+#include "Misc/HierarchicalLogArchive.h"
+
+FHierarchicalLogArchive::FHierarchicalLogArchive(FArchive& InInnerArchive)
+	: FArchiveProxy(InInnerArchive)
+	, Indentation(0)
+{}
+
+void FHierarchicalLogArchive::WriteLine(const FString& InLine, bool bIndent)
+{
+	TStringBuilder<512> Builder;
+
+	if (Indentation > 0)
+	{
+		for (int i = 0; i < Indentation - 1; ++i)
+		{
+			Builder += TEXT(" |  ");
+		}
+
+		Builder += TEXT(" |- ");
+	}
+	
+	if (bIndent)
+	{
+		Builder += TEXT("[+] ");
+	}
+
+	Builder += InLine;
+	Builder += LINE_TERMINATOR;
+	FString LineEntry = Builder.ToString();
+
+	Serialize(TCHAR_TO_ANSI(*LineEntry), LineEntry.Len());
+}
\ No newline at end of file
diff --git a/Engine/Source/Runtime/Core/Public/Misc/HierarchicalLogArchive.h b/Engine/Source/Runtime/Core/Public/Misc/HierarchicalLogArchive.h
new file mode 100644
--- /dev/null
+++ b/Engine/Source/Runtime/Core/Public/Misc/HierarchicalLogArchive.h
@@ -0,0 +1,37 @@
+// Copyright Epic Games, Inc. All Rights Reserved.
+
+#pragma once
+
+#include "CoreMinimal.h"
+#include "Serialization/ArchiveProxy.h"
+
+struct CORE_API FHierarchicalLogArchive : public FArchiveProxy
+{
+public:
+	FHierarchicalLogArchive(FArchive& InInnerArchive);
+
+	struct FIndentScope
+	{
+		FIndentScope(FHierarchicalLogArchive* InAr) : Ar(InAr) { Ar->Indentation++; }
+		~FIndentScope() { check(Ar->Indentation); Ar->Indentation--; }
+		FHierarchicalLogArchive* Ar;
+	};
+
+	template <typename FmtType, typename... Types>
+	void Printf(const FmtType& Fmt, Types... Args)
+	{
+		WriteLine(FString::Printf(Fmt, Args...), false);
+	}
+
+	template <typename FmtType, typename... Types>
+	UE_NODISCARD FIndentScope PrintfIndent(const FmtType& Fmt, Types... Args)
+	{
+		WriteLine(FString::Printf(Fmt, Args...), true);
+		return FIndentScope(this);
+	}
+
+private:
+	void WriteLine(const FString& InLine, bool bIndent = false);
+
+	int32 Indentation;
+};
\ No newline at end of file
diff --git a/Engine/Source/Runtime/Engine/Private/WorldPartition/WorldPartitionFileLogger.cpp b/Engine/Source/Runtime/Engine/Private/WorldPartition/WorldPartitionFileLogger.cpp
deleted file mode 100644
diff --git a/Engine/Source/Runtime/Engine/Private/WorldPartition/WorldPartitionRuntimeCell.cpp b/Engine/Source/Runtime/Engine/Private/WorldPartition/WorldPartitionRuntimeCell.cpp
--- a/Engine/Source/Runtime/Engine/Private/WorldPartition/WorldPartitionRuntimeCell.cpp
+++ b/Engine/Source/Runtime/Engine/Private/WorldPartition/WorldPartitionRuntimeCell.cpp
@@ -9,10 +9,6 @@
 #include "WorldPartition/DataLayer/DataLayersID.h"
 #include "Engine/World.h"
 
-#if WITH_EDITOR
-#include "WorldPartition/WorldPartitionFileLogger.h"
-#endif
-
 int32 UWorldPartitionRuntimeCell::StreamingSourceCacheEpoch = 0;
 
 UWorldPartitionRuntimeCell::UWorldPartitionRuntimeCell(const FObjectInitializer& ObjectInitializer)
@@ -78,10 +74,10 @@ void UWorldPartitionRuntimeCell::UpdateDebugName()
 	DebugName = Builder.ToString();
 }
 
-void UWorldPartitionRuntimeCell::LogStreamingGeneration(FWorldPartitionFileLogger& Logger)
+void UWorldPartitionRuntimeCell::LogStreamingGeneration(FHierarchicalLogArchive& Logger)
 {
-	Logger.WriteLine(FString::Printf(TEXT("Actor Count: %d"), GetActorCount()));
-	Logger.WriteLine(FString::Printf(TEXT("MinMaxZ: %s"), *GetMinMaxZ().ToString()));
+	Logger.Printf(TEXT("Actor Count: %d"), GetActorCount());
+	Logger.Printf(TEXT("MinMaxZ: %s"), *GetMinMaxZ().ToString());
 }
 
 #endif
diff --git a/Engine/Source/Runtime/Engine/Private/WorldPartition/WorldPartitionRuntimeHash.cpp b/Engine/Source/Runtime/Engine/Private/WorldPartition/WorldPartitionRuntimeHash.cpp
--- a/Engine/Source/Runtime/Engine/Private/WorldPartition/WorldPartitionRuntimeHash.cpp
+++ b/Engine/Source/Runtime/Engine/Private/WorldPartition/WorldPartitionRuntimeHash.cpp
@@ -10,7 +10,6 @@
 #if WITH_EDITOR
 #include "WorldPartition/WorldPartitionHandle.h"
 #include "WorldPartition/DataLayer/WorldDataLayers.h"
-#include "WorldPartition/WorldPartitionFileLogger.h"
 #endif
 
 #define LOCTEXT_NAMESPACE "WorldPartition"
@@ -38,25 +37,13 @@ void UWorldPartitionRuntimeHash::OnEndPlay()
 	ModifiedActorDescListForPIE.Empty();
 }
 
-void UWorldPartitionRuntimeHash::LogStreamingGeneration(FWorldPartitionFileLogger& Logger)
+void UWorldPartitionRuntimeHash::LogStreamingGeneration(FHierarchicalLogArchive& Logger)
 {
-	Logger.WriteLine(TEXT("----------------------------------------------------------------------------------------------------------------"));
-	Logger.WriteLine(FString::Printf(TEXT("%s - Persistent Level"), *GetWorld()->GetName()));
-	Logger.WriteLine(TEXT("----------------------------------------------------------------------------------------------------------------"));
-	Logger.WriteLine(FString::Printf(TEXT("Always loaded Actor Count: %d "), GetWorld()->PersistentLevel->Actors.Num()));
-	Logger.WriteLine(TEXT(""));
-}
-
-void UWorldPartitionRuntimeHash::LogStreamingGeneration()
-{
-	// For now only log for cook
-	if (IsRunningCookCommandlet())
-	{
-		const TCHAR* Suffix = TEXT("Cook");
-		const FString StateLogOutputFilename = FPaths::ProjectSavedDir() / TEXT("WorldPartition") / FString::Printf(TEXT("StreamingGeneration-RuntimeSpatialHash-%s-%s-%s.log"), Suffix, *GetWorld()->GetName(), *FDateTime::Now().ToString());
-		FWorldPartitionFileLogger Logger(StateLogOutputFilename);
-		LogStreamingGeneration(Logger);
-	}
+	Logger.Printf(TEXT("----------------------------------------------------------------------------------------------------------------"));
+	Logger.Printf(TEXT("%s - Persistent Level"), *GetWorld()->GetName());
+	Logger.Printf(TEXT("----------------------------------------------------------------------------------------------------------------"));
+	Logger.Printf(TEXT("Always loaded Actor Count: %d "), GetWorld()->PersistentLevel->Actors.Num());
+	Logger.Printf(TEXT(""));
 }
 
 void UWorldPartitionRuntimeHash::ForceExternalActorLevelReference(bool bForceExternalActorLevelReferenceForPIE)
diff --git a/Engine/Source/Runtime/Engine/Private/WorldPartition/WorldPartitionRuntimeLevelStreamingCell.cpp b/Engine/Source/Runtime/Engine/Private/WorldPartition/WorldPartitionRuntimeLevelStreamingCell.cpp
--- a/Engine/Source/Runtime/Engine/Private/WorldPartition/WorldPartitionRuntimeLevelStreamingCell.cpp
+++ b/Engine/Source/Runtime/Engine/Private/WorldPartition/WorldPartitionRuntimeLevelStreamingCell.cpp
@@ -12,7 +12,6 @@
 #if WITH_EDITOR
 #include "WorldPartition/WorldPartitionLevelStreamingPolicy.h"
 #include "WorldPartition/WorldPartitionLevelHelper.h"
-#include "WorldPartition/WorldPartitionFileLogger.h"
 #include "Engine/LevelStreamingAlwaysLoaded.h"
 #endif
 
@@ -220,14 +219,14 @@ FString UWorldPartitionRuntimeLevelStreamingCell::GetPackageNameToCreate() const
 	return UWorldPartitionLevelStreamingPolicy::GetCellPackagePath(GetFName(), OuterWorld);
 }
 
-void UWorldPartitionRuntimeLevelStreamingCell::LogStreamingGeneration(FWorldPartitionFileLogger& Logger)
+void UWorldPartitionRuntimeLevelStreamingCell::LogStreamingGeneration(FHierarchicalLogArchive& Logger)
 {
 	Super::LogStreamingGeneration(Logger);
 
 	for (const FWorldPartitionRuntimeCellObjectMapping& Mapping : Packages)
 	{
-		Logger.WriteLine(FString::Printf(TEXT("Actor Path: %s"), *Mapping.Path.ToString()));
-		Logger.WriteLine(FString::Printf(TEXT("Actor Package: %s"), *Mapping.Package.ToString()));
+		Logger.Printf(TEXT("Actor Path: %s"), *Mapping.Path.ToString());
+		Logger.Printf(TEXT("Actor Package: %s"), *Mapping.Package.ToString());
 	}
 }
 
diff --git a/Engine/Source/Runtime/Engine/Private/WorldPartition/WorldPartitionRuntimeSpatialHash.cpp b/Engine/Source/Runtime/Engine/Private/WorldPartition/WorldPartitionRuntimeSpatialHash.cpp
--- a/Engine/Source/Runtime/Engine/Private/WorldPartition/WorldPartitionRuntimeSpatialHash.cpp
+++ b/Engine/Source/Runtime/Engine/Private/WorldPartition/WorldPartitionRuntimeSpatialHash.cpp
@@ -37,7 +37,6 @@
 #include "Components/LineBatchComponent.h"
 
 #if WITH_EDITOR
-#include "WorldPartition/WorldPartitionFileLogger.h"
 #include "Editor/EditorEngine.h"
 #include "CookPackageSplitter.h"
 #include "Misc/Parse.h"
@@ -853,87 +852,85 @@ bool UWorldPartitionRuntimeSpatialHash::GenerateStreaming(UWorldPartitionStreami
 	return true;
 }
 
-void UWorldPartitionRuntimeSpatialHash::LogStreamingGeneration(FWorldPartitionFileLogger& Logger)
+void UWorldPartitionRuntimeSpatialHash::LogStreamingGeneration(FHierarchicalLogArchive& Logger)
 {
 	Super::LogStreamingGeneration(Logger);
 
 	for (FSpatialHashStreamingGrid& StreamingGrid : StreamingGrids)
 	{
-		Logger.WriteLine(TEXT("----------------------------------------------------------------------------------------------------------------"));
-		Logger.WriteLine(FString::Printf(TEXT("%s - Runtime Spatial Hash - Streaming Grid - %s"), *GetWorld()->GetName(), *StreamingGrid.GridName.ToString()));
-		Logger.WriteLine(TEXT("----------------------------------------------------------------------------------------------------------------"));
-		Logger.WriteLine(FString::Printf(TEXT("            Origin: %s"), *StreamingGrid.Origin.ToString()));
-		Logger.WriteLine(FString::Printf(TEXT("         Cell Size: %d"), StreamingGrid.CellSize));
-		Logger.WriteLine(FString::Printf(TEXT("      World Bounds: %s"), *StreamingGrid.WorldBounds.ToString()));
-		Logger.WriteLine(FString::Printf(TEXT("     Loading Range: %3.2f"), StreamingGrid.LoadingRange));
-		Logger.WriteLine(FString::Printf(TEXT("Block Slow Loading: %s"), StreamingGrid.bBlockOnSlowStreaming ? TEXT("Yes") : TEXT("No")));
-		Logger.WriteLine(FString::Printf(TEXT(" ClientOnlyVisible: %s"), StreamingGrid.bClientOnlyVisible ? TEXT("Yes") : TEXT("No")));
-		Logger.WriteLine(TEXT(""));
+		Logger.Printf(TEXT("----------------------------------------------------------------------------------------------------------------"));
+		Logger.Printf(TEXT("%s - Runtime Spatial Hash - Streaming Grid - %s"), *GetWorld()->GetName(), *StreamingGrid.GridName.ToString());
+		Logger.Printf(TEXT("----------------------------------------------------------------------------------------------------------------"));
+		Logger.Printf(TEXT("            Origin: %s"), *StreamingGrid.Origin.ToString());
+		Logger.Printf(TEXT("         Cell Size: %d"), StreamingGrid.CellSize);
+		Logger.Printf(TEXT("      World Bounds: %s"), *StreamingGrid.WorldBounds.ToString());
+		Logger.Printf(TEXT("     Loading Range: %3.2f"), StreamingGrid.LoadingRange);
+		Logger.Printf(TEXT("Block Slow Loading: %s"), StreamingGrid.bBlockOnSlowStreaming ? TEXT("Yes") : TEXT("No"));
+		Logger.Printf(TEXT(" ClientOnlyVisible: %s"), StreamingGrid.bClientOnlyVisible ? TEXT("Yes") : TEXT("No"));
+		Logger.Printf(TEXT(""));
 		if (const UHLODLayer* HLODLayer = StreamingGrid.HLODLayer)
 		{
-			Logger.WriteLine(FString::Printf(TEXT("    HLOD Layer: %s"), *HLODLayer->GetName()));
+			Logger.Printf(TEXT("    HLOD Layer: %s"), *HLODLayer->GetName());
 		}
 
+		struct FGridLevelStats
 		{
-			struct FGridLevelStats
-			{
-				int32 CellCount;
-				int32 CellSize;
-				int32 ActorCount;
-			};
+			int32 CellCount;
+			int32 CellSize;
+			int32 ActorCount;
+		};
 
-			TArray<FGridLevelStats> LevelsStats;
-			int32 TotalActorCount = 0;
+		TArray<FGridLevelStats> LevelsStats;
+		int32 TotalActorCount = 0;
 			
+		{
+			int32 Level = 0;
+			for (FSpatialHashStreamingGridLevel& GridLevel : StreamingGrid.GridLevels)
 			{
-				int32 Level = 0;
-				for (FSpatialHashStreamingGridLevel& GridLevel : StreamingGrid.GridLevels)
+				int32 LevelCellCount = 0;
+				int32 LevelActorCount = 0;
+				for (FSpatialHashStreamingGridLayerCell& LayerCell : GridLevel.LayerCells)
 				{
-					int32 LevelCellCount = 0;
-					int32 LevelActorCount = 0;
-					for (FSpatialHashStreamingGridLayerCell& LayerCell : GridLevel.LayerCells)
+					LevelCellCount += LayerCell.GridCells.Num();
+					for (TObjectPtr<UWorldPartitionRuntimeSpatialHashCell>& Cell : LayerCell.GridCells)
 					{
-						LevelCellCount += LayerCell.GridCells.Num();
-						for (TObjectPtr<UWorldPartitionRuntimeSpatialHashCell>& Cell : LayerCell.GridCells)
-						{
-							LevelActorCount += Cell->GetActorCount();
-						}
+						LevelActorCount += Cell->GetActorCount();
 					}
-					LevelsStats.Add({ LevelCellCount, (StreamingGrid.CellSize << Level), LevelActorCount });
-					TotalActorCount += LevelActorCount;
-					++Level;
 				}
-				TotalActorCount = (TotalActorCount > 0) ? TotalActorCount : 1;
+				LevelsStats.Add({ LevelCellCount, (StreamingGrid.CellSize << Level), LevelActorCount });
+				TotalActorCount += LevelActorCount;
+				++Level;
 			}
+			TotalActorCount = (TotalActorCount > 0) ? TotalActorCount : 1;
+		}
 
-			Logger.WriteLine(FString::Printf(TEXT("Grid Levels: %d"), StreamingGrid.GridLevels.Num()), true);
-			for (int Level=0; Level< LevelsStats.Num(); ++Level)
+		{
+			FHierarchicalLogArchive::FIndentScope IndentScope = Logger.PrintfIndent(TEXT("Grid Levels: %d"), StreamingGrid.GridLevels.Num());
+			for (int Level=0; Level<LevelsStats.Num(); ++Level)
 			{
-				Logger.WriteLine(FString::Printf(TEXT("Level %2d: Cell Count %4d | Cell Size %7d | Actor Count %4d (%3.1f%%)"), Level, LevelsStats[Level].CellCount, LevelsStats[Level].CellSize, LevelsStats[Level].ActorCount, (100.f*LevelsStats[Level].ActorCount)/TotalActorCount));
+				Logger.Printf(TEXT("Level %2d: Cell Count %4d | Cell Size %7d | Actor Count %4d (%3.1f%%)"), Level, LevelsStats[Level].CellCount, LevelsStats[Level].CellSize, LevelsStats[Level].ActorCount, (100.f*LevelsStats[Level].ActorCount)/TotalActorCount);
 			}
-			Logger.DecrementIndentation();
 		}
 
 		{
-			Logger.WriteLine(TEXT(""));
+			Logger.Printf(TEXT(""));
 			int32 Level = 0;
 			for (const FSpatialHashStreamingGridLevel& GridLevel : StreamingGrid.GridLevels)
 			{
-				Logger.WriteLine(FString::Printf(TEXT("Content of Grid Level %d"), Level++), true);
+				FHierarchicalLogArchive::FIndentScope LevelIndentScope = Logger.PrintfIndent(TEXT("Content of Grid Level %d"), Level++);
+
 				for (const FSpatialHashStreamingGridLayerCell& LayerCell : GridLevel.LayerCells)
 				{
 					for (const TObjectPtr<UWorldPartitionRuntimeSpatialHashCell>& Cell : LayerCell.GridCells)
 					{
-						Logger.WriteLine(FString::Printf(TEXT("Content of Cell %s"), *Cell->GetDebugName()), true);
+						FHierarchicalLogArchive::FIndentScope CellIndentScope = Logger.PrintfIndent(TEXT("Content of Cell %s"), *Cell->GetDebugName());
 						Cell->LogStreamingGeneration(Logger);
-						Logger.DecrementIndentation();
 					}
 				}
-				Logger.DecrementIndentation();
 			}
 		}
-		Logger.WriteLine(TEXT(""));
-	}
+		Logger.Printf(TEXT(""));
+	}	
 }
 
 FName UWorldPartitionRuntimeSpatialHash::GetCellName(UWorldPartition* WorldPartition, FName InGridName, const FIntVector& InCellGlobalCoord, const FDataLayersID& InDataLayerID)
diff --git a/Engine/Source/Runtime/Engine/Private/WorldPartition/WorldPartitionStreamingGeneration.cpp b/Engine/Source/Runtime/Engine/Private/WorldPartition/WorldPartitionStreamingGeneration.cpp
--- a/Engine/Source/Runtime/Engine/Private/WorldPartition/WorldPartitionStreamingGeneration.cpp
+++ b/Engine/Source/Runtime/Engine/Private/WorldPartition/WorldPartitionStreamingGeneration.cpp
@@ -479,82 +479,73 @@ public:
 		return FActorClusterContext(MoveTemp(ContainerInstances), InFilterActorDescViewFunc);
 	}
 
-	void DumpStateLog(const TCHAR* Suffix)
+	static TUniquePtr<FArchive> CreateDumpStateLogArchive(const TCHAR* Suffix)
 	{
 		const FString StateLogOutputFilename = FPaths::ProjectSavedDir() / TEXT("WorldPartition") / FString::Printf(TEXT("StreamingGeneration-%s-%s.log"), Suffix, *FDateTime::Now().ToString());
-		
-		if (FArchive* LogFile = IFileManager::Get().CreateFileWriter(*StateLogOutputFilename))
-		{
-			auto SerializeLine = [LogFile](const FString& Line)
-			{
-				FString LineEntry(Line);
-				LineEntry += LINE_TERMINATOR;
-				LogFile->Serialize(TCHAR_TO_ANSI(*LineEntry), LineEntry.Len());
-			};
+		return TUniquePtr<FArchive>(IFileManager::Get().CreateFileWriter(*StateLogOutputFilename));
+	}
 
-			// Build the containers tree representation
-			TMultiMap<FActorContainerID, FActorContainerID> InvertedContainersHierarchy;
-			for (auto ContainerPairIt = ContainersHierarchy.CreateIterator(); ContainerPairIt; ++ContainerPairIt)
-			{
-				const FActorContainerID& ChildContainerID = ContainerPairIt.Key();
-				const FActorContainerID& ParentContainerID = ContainerPairIt.Value();				
-				InvertedContainersHierarchy.Add(ParentContainerID, ChildContainerID);
-			}
+	void DumpStateLog(FHierarchicalLogArchive& Ar)
+	{
+		// Build the containers tree representation
+		TMultiMap<FActorContainerID, FActorContainerID> InvertedContainersHierarchy;
+		for (auto ContainerPairIt = ContainersHierarchy.CreateIterator(); ContainerPairIt; ++ContainerPairIt)
+		{
+			const FActorContainerID& ChildContainerID = ContainerPairIt.Key();
+			const FActorContainerID& ParentContainerID = ContainerPairIt.Value();				
+			InvertedContainersHierarchy.Add(ParentContainerID, ChildContainerID);
+		}
 
-			SerializeLine(TEXT("Containers:"));
+		Ar.Printf(TEXT("Containers:"));
 
-			auto DumpContainers = [this, &InvertedContainersHierarchy, &SerializeLine](const FActorContainerID& ContainerID)
+		auto DumpContainers = [this, &InvertedContainersHierarchy, &Ar](const FActorContainerID& ContainerID)
+		{
+			auto DumpContainersRecursive = [this, &InvertedContainersHierarchy, &Ar](const FActorContainerID& ContainerID, auto& RecursiveFunc) -> void
 			{
-				auto DumpContainersRecursive = [this, &InvertedContainersHierarchy, &SerializeLine](const FActorContainerID& ContainerID, FString Prefix, auto& RecursiveFunc) -> void
+				const FContainerDescriptor& ContainerDescriptor = ContainerDescriptorsMap.FindChecked(ContainerID);
+				
 				{
-					const FContainerDescriptor& ContainerDescriptor = ContainerDescriptorsMap.FindChecked(ContainerID);
-					SerializeLine(FString::Printf(TEXT("%s[+] %s:"), *Prefix, *ContainerDescriptor.OwnerName));
-
-					Prefix += TEXT(" | ");
+					FHierarchicalLogArchive::FIndentScope IndentScope = Ar.PrintfIndent(TEXT("%s:"), *ContainerDescriptor.OwnerName);
 
-					SerializeLine(FString::Printf(TEXT("%s           ID: 0x%016llx"), *Prefix, ContainerID.ID));
-					SerializeLine(FString::Printf(TEXT("%s       Bounds: %s"), *Prefix, *ContainerDescriptor.Bounds.ToString()));
-					SerializeLine(FString::Printf(TEXT("%s    Transform: %s"), *Prefix, *ContainerDescriptor.Transform.ToString()));
-					SerializeLine(FString::Printf(TEXT("%s    Container: %s"), *Prefix, *ContainerDescriptor.Container->GetContainerPackage().ToString()));
+					Ar.Printf(TEXT("       ID: 0x%016llx"), ContainerID.ID);
+					Ar.Printf(TEXT("   Bounds: %s"), *ContainerDescriptor.Bounds.ToString());
+					Ar.Printf(TEXT("Transform: %s"), *ContainerDescriptor.Transform.ToString());
+					Ar.Printf(TEXT("Container: %s"), *ContainerDescriptor.Container->GetContainerPackage().ToString());
+				}
 
-					if (ContainerDescriptor.ActorDescViewMap.Num())
-					{
-						SerializeLine(FString::Printf(TEXT("%s ActorDescs:"), *Prefix));
+				if (ContainerDescriptor.ActorDescViewMap.Num())
+				{
+					FHierarchicalLogArchive::FIndentScope IndentScope = Ar.PrintfIndent(TEXT("ActorDescs:"));
 
-						TMap<FGuid, FWorldPartitionActorDescView> SortedActorDescViewMap = ContainerDescriptor.ActorDescViewMap;
-						SortedActorDescViewMap.KeySort([](const FGuid& GuidA, const FGuid& GuidB) { return GuidA < GuidB; });
+					TMap<FGuid, FWorldPartitionActorDescView> SortedActorDescViewMap = ContainerDescriptor.ActorDescViewMap;
+					SortedActorDescViewMap.KeySort([](const FGuid& GuidA, const FGuid& GuidB) { return GuidA < GuidB; });
 
-						for (auto ActorDescIt = SortedActorDescViewMap.CreateConstIterator(); ActorDescIt; ++ActorDescIt)
-						{
-							const FWorldPartitionActorDescView& ActorDescView = ActorDescIt.Value();
-							SerializeLine(FString::Printf(TEXT("%s     - %s"), *Prefix, *ActorDescView.ToString()));
-						}
+					for (auto ActorDescIt = SortedActorDescViewMap.CreateConstIterator(); ActorDescIt; ++ActorDescIt)
+					{
+						const FWorldPartitionActorDescView& ActorDescView = ActorDescIt.Value();
+						Ar.Printf(TEXT("%s"), *ActorDescView.ToString());
 					}
+				}
 
-					TArray<FActorContainerID> ChildContainersIDs;
-					InvertedContainersHierarchy.MultiFind(ContainerID, ChildContainersIDs);
-					ChildContainersIDs.Sort([](const FActorContainerID& ActorContainerIDA, const FActorContainerID& ActorContainerIDB) { return ActorContainerIDA.ID < ActorContainerIDB.ID; });
+				TArray<FActorContainerID> ChildContainersIDs;
+				InvertedContainersHierarchy.MultiFind(ContainerID, ChildContainersIDs);
+				ChildContainersIDs.Sort([](const FActorContainerID& ActorContainerIDA, const FActorContainerID& ActorContainerIDB) { return ActorContainerIDA.ID < ActorContainerIDB.ID; });
 
-					if (ChildContainersIDs.Num())
-					{
-						SerializeLine(FString::Printf(TEXT("%sSubContainers:"), *Prefix));
+				if (ChildContainersIDs.Num())
+				{
+					FHierarchicalLogArchive::FIndentScope IndentScope = Ar.PrintfIndent(TEXT("SubContainers:"));
 						
-						Prefix += TEXT("  ");
-						for (const FActorContainerID& ChildContainerID : ChildContainersIDs)
-						{
-							RecursiveFunc(ChildContainerID, Prefix, RecursiveFunc);
-						}
+					for (const FActorContainerID& ChildContainerID : ChildContainersIDs)
+					{
+						RecursiveFunc(ChildContainerID, RecursiveFunc);
 					}
-				};
-
-				DumpContainersRecursive(ContainerID, TEXT("  "), DumpContainersRecursive);
+				}
 			};
 
-			DumpContainers(FActorContainerID());
+			DumpContainersRecursive(ContainerID, DumpContainersRecursive);
+		};
 
-			LogFile->Close();
-			delete LogFile;
-		}
+		DumpContainers(FActorContainerID());
 	}
 
 private:
@@ -605,6 +596,11 @@ bool UWorldPartition::GenerateStreaming(TArray<FString>* OutPackagesToGenerate)
 		ErrorHandler = &MapCheckErrorHandler;
 	}
 
+	// Dump state log
+	const TCHAR* StateLogSuffix = bIsPIE ? TEXT("PIE") : (IsRunningGame() ? TEXT("Game") : (IsRunningCookCommandlet() ? TEXT("Cook") : TEXT("Unknown")));
+	TUniquePtr<FArchive> LogFileAr = FWorldPartitionStreamingGenerator::CreateDumpStateLogArchive(StateLogSuffix);
+	FHierarchicalLogArchive HierarchicalLogAr(*LogFileAr);
+
 	FActorClusterContext ActorClusterContext;
 	{
 		FWorldPartitionStreamingGenerator StreamingGenerator(ModifiedActorsDescList, ErrorHandler, IsStreamingEnabled());
@@ -612,9 +608,7 @@ bool UWorldPartition::GenerateStreaming(TArray<FString>* OutPackagesToGenerate)
 		// Preparation Phase
 		StreamingGenerator.PreparationPhase(this);
 
-		// Dump state log
-		const TCHAR* StateLogSuffix = bIsPIE ? TEXT("PIE") : (IsRunningGame() ? TEXT("Game") : (IsRunningCookCommandlet() ? TEXT("Cook") : TEXT("Unknown")));
-		StreamingGenerator.DumpStateLog(StateLogSuffix);
+		StreamingGenerator.DumpStateLog(HierarchicalLogAr);
 
 		// Preparation Phase :: Actor Clusters Creation
 		ActorClusterContext = StreamingGenerator.CreateActorClusters();
@@ -627,7 +621,10 @@ bool UWorldPartition::GenerateStreaming(TArray<FString>* OutPackagesToGenerate)
 	check(RuntimeHash);
 	if (RuntimeHash->GenerateStreaming(StreamingPolicy, ActorClusterContext, OutPackagesToGenerate))
 	{
-		RuntimeHash->LogStreamingGeneration();
+		if (IsRunningCookCommandlet())
+		{
+			RuntimeHash->LogStreamingGeneration(HierarchicalLogAr);
+		}
 
 		StreamingPolicy->PrepareActorToCellRemapping();
 		return true;
@@ -641,7 +638,10 @@ void UWorldPartition::GenerateHLOD(ISourceControlHelper* SourceControlHelper, bo
 	FStreamingGenerationLogErrorHandler LogErrorHandler;
 	FWorldPartitionStreamingGenerator StreamingGenerator(nullptr, &LogErrorHandler, IsStreamingEnabled());
 	StreamingGenerator.PreparationPhase(this);
-	StreamingGenerator.DumpStateLog(TEXT("HLOD"));
+
+	TUniquePtr<FArchive> LogFileAr = FWorldPartitionStreamingGenerator::CreateDumpStateLogArchive(TEXT("HLOD"));
+	FHierarchicalLogArchive HierarchicalLogAr(*LogFileAr);
+	StreamingGenerator.DumpStateLog(HierarchicalLogAr);
 
 	// Preparation Phase :: Actor Clusters Creation
 	FActorClusterContext ActorClusterContext = StreamingGenerator.CreateActorClusters([](const FWorldPartitionActorDescView& ActorDescView)
diff --git a/Engine/Source/Runtime/Engine/Public/WorldPartition/WorldPartitionFileLogger.h b/Engine/Source/Runtime/Engine/Public/WorldPartition/WorldPartitionFileLogger.h
deleted file mode 100644
diff --git a/Engine/Source/Runtime/Engine/Public/WorldPartition/WorldPartitionRuntimeCell.h b/Engine/Source/Runtime/Engine/Public/WorldPartition/WorldPartitionRuntimeCell.h
--- a/Engine/Source/Runtime/Engine/Public/WorldPartition/WorldPartitionRuntimeCell.h
+++ b/Engine/Source/Runtime/Engine/Public/WorldPartition/WorldPartitionRuntimeCell.h
@@ -9,12 +9,13 @@
 #include "WorldPartition/WorldPartitionActorCluster.h"
 #include "WorldPartition/WorldPartitionRuntimeCellInterface.h"
 #include "ProfilingDebugging/ProfilingHelpers.h"
+#include "Misc/HierarchicalLogArchive.h"
 #include "Algo/AnyOf.h"
 #include "WorldPartitionRuntimeCell.generated.h"
 
 class UDataLayerAsset;
 class UDataLayerInstance;
-struct FWorldPartitionFileLogger;
+struct FHierarchicalLogArchive;
 
 enum class EWorldPartitionRuntimeCellVisualizeMode
 {
@@ -185,7 +186,7 @@ class UWorldPartitionRuntimeCell : public UObject, public IWorldPartitionCell
 
 	void SetIsHLOD(bool bInIsHLOD) { bIsHLOD = bInIsHLOD; }
 
-	virtual void LogStreamingGeneration(FWorldPartitionFileLogger& Logger);
+	virtual void LogStreamingGeneration(FHierarchicalLogArchive& Logger);
 #endif
 	
 	bool GetIsHLOD() const { return bIsHLOD; }
diff --git a/Engine/Source/Runtime/Engine/Public/WorldPartition/WorldPartitionRuntimeHash.h b/Engine/Source/Runtime/Engine/Public/WorldPartition/WorldPartitionRuntimeHash.h
--- a/Engine/Source/Runtime/Engine/Public/WorldPartition/WorldPartitionRuntimeHash.h
+++ b/Engine/Source/Runtime/Engine/Public/WorldPartition/WorldPartitionRuntimeHash.h
@@ -12,11 +12,12 @@
 #include "WorldPartition/WorldPartitionActorDescViewProxy.h"
 #if WITH_EDITOR
 #include "CookPackageSplitter.h"
+#include "Misc/HierarchicalLogArchive.h"
 #endif
 #include "WorldPartitionRuntimeHash.generated.h"
 
 class FActorClusterContext;
-struct FWorldPartitionFileLogger;
+struct FHierarchicalLogArchive;
 
 UENUM()
 enum class EWorldPartitionStreamingPerformance : uint8
@@ -40,12 +41,11 @@ class ENGINE_API UWorldPartitionRuntimeHash : public UObject
 	virtual bool GenerateHLOD(ISourceControlHelper* SourceControlHelper, FActorClusterContext& ActorClusterContext, bool bCreateActorsOnly) { return false; }
 	virtual void DrawPreview() const {}
 
+	virtual void LogStreamingGeneration(FHierarchicalLogArchive& Logger);
+
 	// PIE/Game methods
 	void OnBeginPlay();
 	void OnEndPlay();
-
-	// GenerateStreaming Logging
-	void LogStreamingGeneration();
 #endif
 
 	class FStreamingSourceCells
@@ -80,9 +80,6 @@ class ENGINE_API UWorldPartitionRuntimeHash : public UObject
 
 protected:
 	virtual EWorldPartitionStreamingPerformance GetStreamingPerformanceForCell(const UWorldPartitionRuntimeCell* Cell) const { return EWorldPartitionStreamingPerformance::Good; }
-#if WITH_EDITOR
-	virtual void LogStreamingGeneration(FWorldPartitionFileLogger& Logger);
-#endif
 
 private:
 #if WITH_EDITOR
diff --git a/Engine/Source/Runtime/Engine/Public/WorldPartition/WorldPartitionRuntimeLevelStreamingCell.h b/Engine/Source/Runtime/Engine/Public/WorldPartition/WorldPartitionRuntimeLevelStreamingCell.h
--- a/Engine/Source/Runtime/Engine/Public/WorldPartition/WorldPartitionRuntimeLevelStreamingCell.h
+++ b/Engine/Source/Runtime/Engine/Public/WorldPartition/WorldPartitionRuntimeLevelStreamingCell.h
@@ -34,7 +34,7 @@ class UWorldPartitionRuntimeLevelStreamingCell : public UWorldPartitionRuntimeSp
 	//~Begin UWorldPartitionRuntimeCell Interface
 	virtual void AddActorToCell(const FWorldPartitionActorDescView& ActorDescView, const FActorContainerID& InContainerID, const FTransform& InContainerTransform, const UActorDescContainer* InContainer) override;
 	virtual int32 GetActorCount() const override;
-	virtual void LogStreamingGeneration(FWorldPartitionFileLogger& Logger) override;
+	virtual void LogStreamingGeneration(FHierarchicalLogArchive& Logger) override;
 	// Cook methods
 	virtual bool PrepareCellForCook(UPackage* InPackage) override;
 	virtual bool PopulateGeneratedPackageForCook(UPackage* InPackage) override;
diff --git a/Engine/Source/Runtime/Engine/Public/WorldPartition/WorldPartitionRuntimeSpatialHash.h b/Engine/Source/Runtime/Engine/Public/WorldPartition/WorldPartitionRuntimeSpatialHash.h
--- a/Engine/Source/Runtime/Engine/Public/WorldPartition/WorldPartitionRuntimeSpatialHash.h
+++ b/Engine/Source/Runtime/Engine/Public/WorldPartition/WorldPartitionRuntimeSpatialHash.h
@@ -4,7 +4,7 @@
 #include "CoreMinimal.h"
 #include "UObject/ObjectMacros.h"
 #include "UObject/Object.h"
-
+#include "Misc/HierarchicalLogArchive.h"
 #include "GameFramework/Info.h"
 
 #include "WorldPartition.h"
@@ -255,7 +255,7 @@ protected:
 
 #if WITH_EDITOR
 	virtual bool GenerateStreaming(class UWorldPartitionStreamingPolicy* StreamingPolicy, const FActorClusterContext& ActorClusterContext, TArray<FString>* OutPackagesToGenerate = nullptr) override;
-	virtual void LogStreamingGeneration(FWorldPartitionFileLogger& Logger) override;
+	virtual void LogStreamingGeneration(FHierarchicalLogArchive& Logger) override;
 #endif
 
 private:
