World Partition: minor code changes to keep consistency between streaming generation log functions.

#rnx

diff --git a/Engine/Source/Runtime/Core/Public/Misc/HierarchicalLogArchive.h b/Engine/Source/Runtime/Core/Public/Misc/HierarchicalLogArchive.h
--- a/Engine/Source/Runtime/Core/Public/Misc/HierarchicalLogArchive.h
+++ b/Engine/Source/Runtime/Core/Public/Misc/HierarchicalLogArchive.h
@@ -5,7 +5,7 @@
 #include "CoreMinimal.h"
 #include "Serialization/ArchiveProxy.h"
 
-struct CORE_API FHierarchicalLogArchive : public FArchiveProxy
+struct CORE_API FHierarchicalLogArchive : private FArchiveProxy
 {
 public:
 	FHierarchicalLogArchive(FArchive& InInnerArchive);
diff --git a/Engine/Source/Runtime/Engine/Private/WorldPartition/WorldPartitionRuntimeCell.cpp b/Engine/Source/Runtime/Engine/Private/WorldPartition/WorldPartitionRuntimeCell.cpp
--- a/Engine/Source/Runtime/Engine/Private/WorldPartition/WorldPartitionRuntimeCell.cpp
+++ b/Engine/Source/Runtime/Engine/Private/WorldPartition/WorldPartitionRuntimeCell.cpp
@@ -74,10 +74,10 @@ void UWorldPartitionRuntimeCell::UpdateDebugName()
 	DebugName = Builder.ToString();
 }
 
-void UWorldPartitionRuntimeCell::LogStreamingGeneration(FHierarchicalLogArchive& Logger)
+void UWorldPartitionRuntimeCell::DumpStateLog(FHierarchicalLogArchive& Ar)
 {
-	Logger.Printf(TEXT("Actor Count: %d"), GetActorCount());
-	Logger.Printf(TEXT("MinMaxZ: %s"), *GetMinMaxZ().ToString());
+	Ar.Printf(TEXT("Actor Count: %d"), GetActorCount());
+	Ar.Printf(TEXT("MinMaxZ: %s"), *GetMinMaxZ().ToString());
 }
 
 #endif
diff --git a/Engine/Source/Runtime/Engine/Private/WorldPartition/WorldPartitionRuntimeHash.cpp b/Engine/Source/Runtime/Engine/Private/WorldPartition/WorldPartitionRuntimeHash.cpp
--- a/Engine/Source/Runtime/Engine/Private/WorldPartition/WorldPartitionRuntimeHash.cpp
+++ b/Engine/Source/Runtime/Engine/Private/WorldPartition/WorldPartitionRuntimeHash.cpp
@@ -37,13 +37,13 @@ void UWorldPartitionRuntimeHash::OnEndPlay()
 	ModifiedActorDescListForPIE.Empty();
 }
 
-void UWorldPartitionRuntimeHash::LogStreamingGeneration(FHierarchicalLogArchive& Logger)
+void UWorldPartitionRuntimeHash::DumpStateLog(FHierarchicalLogArchive& Ar)
 {
-	Logger.Printf(TEXT("----------------------------------------------------------------------------------------------------------------"));
-	Logger.Printf(TEXT("%s - Persistent Level"), *GetWorld()->GetName());
-	Logger.Printf(TEXT("----------------------------------------------------------------------------------------------------------------"));
-	Logger.Printf(TEXT("Always loaded Actor Count: %d "), GetWorld()->PersistentLevel->Actors.Num());
-	Logger.Printf(TEXT(""));
+	Ar.Printf(TEXT("----------------------------------------------------------------------------------------------------------------"));
+	Ar.Printf(TEXT("%s - Persistent Level"), *GetWorld()->GetName());
+	Ar.Printf(TEXT("----------------------------------------------------------------------------------------------------------------"));
+	Ar.Printf(TEXT("Always loaded Actor Count: %d "), GetWorld()->PersistentLevel->Actors.Num());
+	Ar.Printf(TEXT(""));
 }
 
 void UWorldPartitionRuntimeHash::ForceExternalActorLevelReference(bool bForceExternalActorLevelReferenceForPIE)
diff --git a/Engine/Source/Runtime/Engine/Private/WorldPartition/WorldPartitionRuntimeLevelStreamingCell.cpp b/Engine/Source/Runtime/Engine/Private/WorldPartition/WorldPartitionRuntimeLevelStreamingCell.cpp
--- a/Engine/Source/Runtime/Engine/Private/WorldPartition/WorldPartitionRuntimeLevelStreamingCell.cpp
+++ b/Engine/Source/Runtime/Engine/Private/WorldPartition/WorldPartitionRuntimeLevelStreamingCell.cpp
@@ -18,8 +18,7 @@
 UWorldPartitionRuntimeLevelStreamingCell::UWorldPartitionRuntimeLevelStreamingCell(const FObjectInitializer& ObjectInitializer)
 	: Super(ObjectInitializer)
 	, LevelStreaming(nullptr)
-{
-}
+{}
 
 EWorldPartitionRuntimeCellState UWorldPartitionRuntimeLevelStreamingCell::GetCurrentState() const
 {
@@ -96,7 +95,6 @@ void UWorldPartitionRuntimeLevelStreamingCell::SetIsAlwaysLoaded(bool bInIsAlway
 }
 
 #if WITH_EDITOR
-
 void UWorldPartitionRuntimeLevelStreamingCell::AddActorToCell(const FWorldPartitionActorDescView& ActorDescView, const FActorContainerID& InContainerID, const FTransform& InContainerTransform, const UActorDescContainer* InContainer)
 {
 	check(!ActorDescView.GetActorIsEditorOnly());
@@ -219,17 +217,16 @@ FString UWorldPartitionRuntimeLevelStreamingCell::GetPackageNameToCreate() const
 	return UWorldPartitionLevelStreamingPolicy::GetCellPackagePath(GetFName(), OuterWorld);
 }
 
-void UWorldPartitionRuntimeLevelStreamingCell::LogStreamingGeneration(FHierarchicalLogArchive& Logger)
+void UWorldPartitionRuntimeLevelStreamingCell::DumpStateLog(FHierarchicalLogArchive& Ar)
 {
-	Super::LogStreamingGeneration(Logger);
+	Super::DumpStateLog(Ar);
 
 	for (const FWorldPartitionRuntimeCellObjectMapping& Mapping : Packages)
 	{
-		Logger.Printf(TEXT("Actor Path: %s"), *Mapping.Path.ToString());
-		Logger.Printf(TEXT("Actor Package: %s"), *Mapping.Package.ToString());
+		Ar.Printf(TEXT("Actor Path: %s"), *Mapping.Path.ToString());
+		Ar.Printf(TEXT("Actor Package: %s"), *Mapping.Package.ToString());
 	}
 }
-
 #endif
 
 UWorldPartitionLevelStreamingDynamic* UWorldPartitionRuntimeLevelStreamingCell::GetOrCreateLevelStreaming() const
diff --git a/Engine/Source/Runtime/Engine/Private/WorldPartition/WorldPartitionRuntimeSpatialHash.cpp b/Engine/Source/Runtime/Engine/Private/WorldPartition/WorldPartitionRuntimeSpatialHash.cpp
--- a/Engine/Source/Runtime/Engine/Private/WorldPartition/WorldPartitionRuntimeSpatialHash.cpp
+++ b/Engine/Source/Runtime/Engine/Private/WorldPartition/WorldPartitionRuntimeSpatialHash.cpp
@@ -852,25 +852,25 @@ bool UWorldPartitionRuntimeSpatialHash::GenerateStreaming(UWorldPartitionStreami
 	return true;
 }
 
-void UWorldPartitionRuntimeSpatialHash::LogStreamingGeneration(FHierarchicalLogArchive& Logger)
+void UWorldPartitionRuntimeSpatialHash::DumpStateLog(FHierarchicalLogArchive& Ar)
 {
-	Super::LogStreamingGeneration(Logger);
+	Super::DumpStateLog(Ar);
 
 	for (FSpatialHashStreamingGrid& StreamingGrid : StreamingGrids)
 	{
-		Logger.Printf(TEXT("----------------------------------------------------------------------------------------------------------------"));
-		Logger.Printf(TEXT("%s - Runtime Spatial Hash - Streaming Grid - %s"), *GetWorld()->GetName(), *StreamingGrid.GridName.ToString());
-		Logger.Printf(TEXT("----------------------------------------------------------------------------------------------------------------"));
-		Logger.Printf(TEXT("            Origin: %s"), *StreamingGrid.Origin.ToString());
-		Logger.Printf(TEXT("         Cell Size: %d"), StreamingGrid.CellSize);
-		Logger.Printf(TEXT("      World Bounds: %s"), *StreamingGrid.WorldBounds.ToString());
-		Logger.Printf(TEXT("     Loading Range: %3.2f"), StreamingGrid.LoadingRange);
-		Logger.Printf(TEXT("Block Slow Loading: %s"), StreamingGrid.bBlockOnSlowStreaming ? TEXT("Yes") : TEXT("No"));
-		Logger.Printf(TEXT(" ClientOnlyVisible: %s"), StreamingGrid.bClientOnlyVisible ? TEXT("Yes") : TEXT("No"));
-		Logger.Printf(TEXT(""));
+		Ar.Printf(TEXT("----------------------------------------------------------------------------------------------------------------"));
+		Ar.Printf(TEXT("%s - Runtime Spatial Hash - Streaming Grid - %s"), *GetWorld()->GetName(), *StreamingGrid.GridName.ToString());
+		Ar.Printf(TEXT("----------------------------------------------------------------------------------------------------------------"));
+		Ar.Printf(TEXT("            Origin: %s"), *StreamingGrid.Origin.ToString());
+		Ar.Printf(TEXT("         Cell Size: %d"), StreamingGrid.CellSize);
+		Ar.Printf(TEXT("      World Bounds: %s"), *StreamingGrid.WorldBounds.ToString());
+		Ar.Printf(TEXT("     Loading Range: %3.2f"), StreamingGrid.LoadingRange);
+		Ar.Printf(TEXT("Block Slow Loading: %s"), StreamingGrid.bBlockOnSlowStreaming ? TEXT("Yes") : TEXT("No"));
+		Ar.Printf(TEXT(" ClientOnlyVisible: %s"), StreamingGrid.bClientOnlyVisible ? TEXT("Yes") : TEXT("No"));
+		Ar.Printf(TEXT(""));
 		if (const UHLODLayer* HLODLayer = StreamingGrid.HLODLayer)
 		{
-			Logger.Printf(TEXT("    HLOD Layer: %s"), *HLODLayer->GetName());
+			Ar.Printf(TEXT("    HLOD Layer: %s"), *HLODLayer->GetName());
 		}
 
 		struct FGridLevelStats
@@ -905,31 +905,31 @@ void UWorldPartitionRuntimeSpatialHash::LogStreamingGeneration(FHierarchicalLogA
 		}
 
 		{
-			FHierarchicalLogArchive::FIndentScope IndentScope = Logger.PrintfIndent(TEXT("Grid Levels: %d"), StreamingGrid.GridLevels.Num());
+			FHierarchicalLogArchive::FIndentScope IndentScope = Ar.PrintfIndent(TEXT("Grid Levels: %d"), StreamingGrid.GridLevels.Num());
 			for (int Level=0; Level<LevelsStats.Num(); ++Level)
 			{
-				Logger.Printf(TEXT("Level %2d: Cell Count %4d | Cell Size %7d | Actor Count %4d (%3.1f%%)"), Level, LevelsStats[Level].CellCount, LevelsStats[Level].CellSize, LevelsStats[Level].ActorCount, (100.f*LevelsStats[Level].ActorCount)/TotalActorCount);
+				Ar.Printf(TEXT("Level %2d: Cell Count %4d | Cell Size %7d | Actor Count %4d (%3.1f%%)"), Level, LevelsStats[Level].CellCount, LevelsStats[Level].CellSize, LevelsStats[Level].ActorCount, (100.f*LevelsStats[Level].ActorCount)/TotalActorCount);
 			}
 		}
 
 		{
-			Logger.Printf(TEXT(""));
+			Ar.Printf(TEXT(""));
 			int32 Level = 0;
 			for (const FSpatialHashStreamingGridLevel& GridLevel : StreamingGrid.GridLevels)
 			{
-				FHierarchicalLogArchive::FIndentScope LevelIndentScope = Logger.PrintfIndent(TEXT("Content of Grid Level %d"), Level++);
+				FHierarchicalLogArchive::FIndentScope LevelIndentScope = Ar.PrintfIndent(TEXT("Content of Grid Level %d"), Level++);
 
 				for (const FSpatialHashStreamingGridLayerCell& LayerCell : GridLevel.LayerCells)
 				{
 					for (const TObjectPtr<UWorldPartitionRuntimeSpatialHashCell>& Cell : LayerCell.GridCells)
 					{
-						FHierarchicalLogArchive::FIndentScope CellIndentScope = Logger.PrintfIndent(TEXT("Content of Cell %s"), *Cell->GetDebugName());
-						Cell->LogStreamingGeneration(Logger);
+						FHierarchicalLogArchive::FIndentScope CellIndentScope = Ar.PrintfIndent(TEXT("Content of Cell %s"), *Cell->GetDebugName());
+						Cell->DumpStateLog(Ar);
 					}
 				}
 			}
 		}
-		Logger.Printf(TEXT(""));
+		Ar.Printf(TEXT(""));
 	}	
 }
 
diff --git a/Engine/Source/Runtime/Engine/Private/WorldPartition/WorldPartitionStreamingGeneration.cpp b/Engine/Source/Runtime/Engine/Private/WorldPartition/WorldPartitionStreamingGeneration.cpp
--- a/Engine/Source/Runtime/Engine/Private/WorldPartition/WorldPartitionStreamingGeneration.cpp
+++ b/Engine/Source/Runtime/Engine/Private/WorldPartition/WorldPartitionStreamingGeneration.cpp
@@ -623,7 +623,7 @@ bool UWorldPartition::GenerateStreaming(TArray<FString>* OutPackagesToGenerate)
 	{
 		if (IsRunningCookCommandlet())
 		{
-			RuntimeHash->LogStreamingGeneration(HierarchicalLogAr);
+			RuntimeHash->DumpStateLog(HierarchicalLogAr);
 		}
 
 		StreamingPolicy->PrepareActorToCellRemapping();
diff --git a/Engine/Source/Runtime/Engine/Public/WorldPartition/WorldPartitionRuntimeCell.h b/Engine/Source/Runtime/Engine/Public/WorldPartition/WorldPartitionRuntimeCell.h
--- a/Engine/Source/Runtime/Engine/Public/WorldPartition/WorldPartitionRuntimeCell.h
+++ b/Engine/Source/Runtime/Engine/Public/WorldPartition/WorldPartitionRuntimeCell.h
@@ -186,7 +186,7 @@ class UWorldPartitionRuntimeCell : public UObject, public IWorldPartitionCell
 
 	void SetIsHLOD(bool bInIsHLOD) { bIsHLOD = bInIsHLOD; }
 
-	virtual void LogStreamingGeneration(FHierarchicalLogArchive& Logger);
+	virtual void DumpStateLog(FHierarchicalLogArchive& Ar);
 #endif
 	
 	bool GetIsHLOD() const { return bIsHLOD; }
diff --git a/Engine/Source/Runtime/Engine/Public/WorldPartition/WorldPartitionRuntimeHash.h b/Engine/Source/Runtime/Engine/Public/WorldPartition/WorldPartitionRuntimeHash.h
--- a/Engine/Source/Runtime/Engine/Public/WorldPartition/WorldPartitionRuntimeHash.h
+++ b/Engine/Source/Runtime/Engine/Public/WorldPartition/WorldPartitionRuntimeHash.h
@@ -41,7 +41,7 @@ class ENGINE_API UWorldPartitionRuntimeHash : public UObject
 	virtual bool GenerateHLOD(ISourceControlHelper* SourceControlHelper, FActorClusterContext& ActorClusterContext, bool bCreateActorsOnly) { return false; }
 	virtual void DrawPreview() const {}
 
-	virtual void LogStreamingGeneration(FHierarchicalLogArchive& Logger);
+	virtual void DumpStateLog(FHierarchicalLogArchive& Ar);
 
 	// PIE/Game methods
 	void OnBeginPlay();
diff --git a/Engine/Source/Runtime/Engine/Public/WorldPartition/WorldPartitionRuntimeLevelStreamingCell.h b/Engine/Source/Runtime/Engine/Public/WorldPartition/WorldPartitionRuntimeLevelStreamingCell.h
--- a/Engine/Source/Runtime/Engine/Public/WorldPartition/WorldPartitionRuntimeLevelStreamingCell.h
+++ b/Engine/Source/Runtime/Engine/Public/WorldPartition/WorldPartitionRuntimeLevelStreamingCell.h
@@ -34,7 +34,7 @@ class UWorldPartitionRuntimeLevelStreamingCell : public UWorldPartitionRuntimeSp
 	//~Begin UWorldPartitionRuntimeCell Interface
 	virtual void AddActorToCell(const FWorldPartitionActorDescView& ActorDescView, const FActorContainerID& InContainerID, const FTransform& InContainerTransform, const UActorDescContainer* InContainer) override;
 	virtual int32 GetActorCount() const override;
-	virtual void LogStreamingGeneration(FHierarchicalLogArchive& Logger) override;
+	virtual void DumpStateLog(FHierarchicalLogArchive& Ar) override;
 	// Cook methods
 	virtual bool PrepareCellForCook(UPackage* InPackage) override;
 	virtual bool PopulateGeneratedPackageForCook(UPackage* InPackage) override;
diff --git a/Engine/Source/Runtime/Engine/Public/WorldPartition/WorldPartitionRuntimeSpatialHash.h b/Engine/Source/Runtime/Engine/Public/WorldPartition/WorldPartitionRuntimeSpatialHash.h
--- a/Engine/Source/Runtime/Engine/Public/WorldPartition/WorldPartitionRuntimeSpatialHash.h
+++ b/Engine/Source/Runtime/Engine/Public/WorldPartition/WorldPartitionRuntimeSpatialHash.h
@@ -255,7 +255,7 @@ protected:
 
 #if WITH_EDITOR
 	virtual bool GenerateStreaming(class UWorldPartitionStreamingPolicy* StreamingPolicy, const FActorClusterContext& ActorClusterContext, TArray<FString>* OutPackagesToGenerate = nullptr) override;
-	virtual void LogStreamingGeneration(FHierarchicalLogArchive& Logger) override;
+	virtual void DumpStateLog(FHierarchicalLogArchive& Ar) override;
 #endif
 
 private:
