FISMComponentDescriptor is now tracking the bEvaluateWorldPositionOffset property

diff --git a/Engine/Source/Runtime/Engine/Private/ISMPartition/ISMComponentDescriptor.cpp b/Engine/Source/Runtime/Engine/Private/ISMPartition/ISMComponentDescriptor.cpp
--- a/Engine/Source/Runtime/Engine/Private/ISMPartition/ISMComponentDescriptor.cpp
+++ b/Engine/Source/Runtime/Engine/Private/ISMPartition/ISMComponentDescriptor.cpp
@@ -67,6 +67,7 @@ void FISMComponentDescriptor::InitFrom(const UStaticMeshComponent* Template, boo
 	bVisible = Template->GetVisibleFlag();
 	bVisibleInRayTracing = Template->bVisibleInRayTracing;
 	bConsiderForActorPlacementWhenHidden = Template->bConsiderForActorPlacementWhenHidden;
+	bEvaluateWorldPositionOffset = Template->bEvaluateWorldPositionOffset;
 
 	if (const UInstancedStaticMeshComponent* ISMTemplate = Cast<UInstancedStaticMeshComponent>(Template))
 	{
@@ -131,6 +132,7 @@ bool FISMComponentDescriptor::operator==(const FISMComponentDescriptor& Other) c
 	bVisible == Other.bVisible &&
 	bVisibleInRayTracing == Other.bVisibleInRayTracing &&
 	bConsiderForActorPlacementWhenHidden == Other.bConsiderForActorPlacementWhenHidden &&
+	bEvaluateWorldPositionOffset == Other.bEvaluateWorldPositionOffset &&
 	BodyInstance.GetCollisionEnabled() == Other.BodyInstance.GetCollisionEnabled() && 
 	BodyInstance.GetCollisionResponse() == Other.BodyInstance.GetCollisionResponse() &&
 	BodyInstance.DoesUseCollisionProfile() == Other.BodyInstance.DoesUseCollisionProfile() &&
@@ -208,6 +210,7 @@ void FISMComponentDescriptor::InitComponent(UInstancedStaticMeshComponent* ISMCo
 	ISMComponent->SetVisibleFlag(bVisible);
 	ISMComponent->bVisibleInRayTracing = bVisibleInRayTracing;
 	ISMComponent->bConsiderForActorPlacementWhenHidden = bConsiderForActorPlacementWhenHidden;
+	ISMComponent->bEvaluateWorldPositionOffset = bEvaluateWorldPositionOffset;
 
 	// HISM Specific
 	if (UHierarchicalInstancedStaticMeshComponent* HISMComponent = Cast<UHierarchicalInstancedStaticMeshComponent>(ISMComponent))
diff --git a/Engine/Source/Runtime/Engine/Public/ISMPartition/ISMComponentDescriptor.h b/Engine/Source/Runtime/Engine/Public/ISMPartition/ISMComponentDescriptor.h
--- a/Engine/Source/Runtime/Engine/Public/ISMPartition/ISMComponentDescriptor.h
+++ b/Engine/Source/Runtime/Engine/Public/ISMPartition/ISMComponentDescriptor.h
@@ -168,5 +168,8 @@ public:
 
 	UPROPERTY(EditAnywhere, Category = "Component Settings")
 	uint8 bConsiderForActorPlacementWhenHidden : 1;
+
+	UPROPERTY(EditAnywhere, Category = "Component Settings")
+	uint8 bEvaluateWorldPositionOffset : 1;
 #endif
 };
