HLOD build crashes when GroupActor is present

Fixed crash caused by null entries in the Level->Actors array

diff --git a/Engine/Source/Runtime/Engine/Private/WorldPartition/HLOD/HLODBuilder.cpp b/Engine/Source/Runtime/Engine/Private/WorldPartition/HLOD/HLODBuilder.cpp
--- a/Engine/Source/Runtime/Engine/Private/WorldPartition/HLOD/HLODBuilder.cpp
+++ b/Engine/Source/Runtime/Engine/Private/WorldPartition/HLOD/HLODBuilder.cpp
@@ -51,7 +51,7 @@ static TArray<UActorComponent*> GatherHLODRelevantComponents(const TArray<AActor
 
 	for (AActor* Actor : InSourceActors)
 	{
-		if (!Actor->IsHLODRelevant())
+		if (!Actor || !Actor->IsHLODRelevant())
 		{
 			continue;
 		}
