World Partition: hide actor handle/references automatic conversions to make code more excplicit and avoid problems creating temporary references that results in destroying actors.

#rnx

diff --git a/Engine/Plugins/Editor/WorldPartitionHLODUtilities/Source/Private/WorldPartitionHLODUtilities.cpp b/Engine/Plugins/Editor/WorldPartitionHLODUtilities/Source/Private/WorldPartitionHLODUtilities.cpp
--- a/Engine/Plugins/Editor/WorldPartitionHLODUtilities/Source/Private/WorldPartitionHLODUtilities.cpp
+++ b/Engine/Plugins/Editor/WorldPartitionHLODUtilities/Source/Private/WorldPartitionHLODUtilities.cpp
@@ -141,7 +141,7 @@ TArray<AWorldPartitionHLOD*> FWorldPartitionHLODUtilities::CreateHLODActors(FHLO
 		FWorldPartitionHandle HLODActorHandle;
 		if (InCreationContext.HLODActorDescs.RemoveAndCopyValue(CellHash, HLODActorHandle))
 		{
-			InCreationContext.ActorReferences.Add(HLODActorHandle);
+			InCreationContext.ActorReferences.Add(HLODActorHandle.ToReference());
 			HLODActor = CastChecked<AWorldPartitionHLOD>(HLODActorHandle->GetActor());
 		}
 
diff --git a/Engine/Source/Runtime/Engine/Private/Tests/WorldPartition/WorldPartitionHandleTests.cpp b/Engine/Source/Runtime/Engine/Private/Tests/WorldPartition/WorldPartitionHandleTests.cpp
--- a/Engine/Source/Runtime/Engine/Private/Tests/WorldPartition/WorldPartitionHandleTests.cpp
+++ b/Engine/Source/Runtime/Engine/Private/Tests/WorldPartition/WorldPartitionHandleTests.cpp
@@ -64,13 +64,13 @@ namespace WorldPartitionTests
 
 		// Conversions
 		{
-			FWorldPartitionHandle HandleToReference = Reference;
+			FWorldPartitionHandle HandleToReference = Reference.ToHandle();
 			TestTrue(TEXT("Handle/Reference equality"), HandleToReference == Reference);
 			TestTrue(TEXT("Reference/Handle equality"), Reference == HandleToReference);
 			TestTrue(TEXT("Reference soft refcount"), Reference->GetSoftRefCount() == 1);
 			TestTrue(TEXT("Reference hard refcount"), Reference->GetHardRefCount() == 1);
 
-			FWorldPartitionReference ReferenceToHandle = Handle;
+			FWorldPartitionReference ReferenceToHandle = Handle.ToReference();
 			TestTrue(TEXT("Handle/Reference equality"), ReferenceToHandle == Handle);
 			TestTrue(TEXT("Handle/Reference equality"), Handle == ReferenceToHandle);
 			TestTrue(TEXT("Reference soft refcount"), Reference->GetSoftRefCount() == 1);
@@ -85,7 +85,7 @@ namespace WorldPartitionTests
 		// inplace new test
 		{
 			uint8 Buffer[sizeof(FWorldPartitionHandle)];
-			FWorldPartitionHandle* HandlePtr = new (Buffer) FWorldPartitionHandle(Reference);
+			FWorldPartitionHandle* HandlePtr = new (Buffer) FWorldPartitionHandle(Reference.ToHandle());
 
 			TestTrue(TEXT("Handle array soft refcount"), Reference->GetSoftRefCount() == 1);
 			TestTrue(TEXT("Handle array hard refcount"), Reference->GetHardRefCount() == 1);
@@ -104,7 +104,7 @@ namespace WorldPartitionTests
 			TestTrue(TEXT("Handle array soft refcount"), Handle->GetSoftRefCount() == 2);
 			TestTrue(TEXT("Handle array hard refcount"), Handle->GetHardRefCount() == 0);
 
-			FWorldPartitionReference ReferenceToHandle = Handle;
+			FWorldPartitionReference ReferenceToHandle = Handle.ToReference();
 			TestTrue(TEXT("Handle/Reference equality"), ReferenceToHandle == Handle);
 			TestTrue(TEXT("Handle/Reference equality"), Handle == ReferenceToHandle);
 			TestTrue(TEXT("Handle soft refcount"), Handle->GetSoftRefCount() == 2);
@@ -113,7 +113,7 @@ namespace WorldPartitionTests
 			TestTrue(TEXT("Handle array contains handle"), HandleList.Contains(Handle));
 			TestTrue(TEXT("Handle array contains reference"), HandleList.Contains(ReferenceToHandle));
 
-			HandleList.Add(Reference);
+			HandleList.Add(Reference.ToHandle());
 			
 			TestTrue(TEXT("Handle array contains reference"), HandleList.Contains(Reference));
 			TestTrue(TEXT("Handle array soft refcount"), Reference->GetSoftRefCount() == 1);
@@ -123,7 +123,7 @@ namespace WorldPartitionTests
 			TestTrue(TEXT("Handle array soft refcount"), Handle->GetSoftRefCount() == 1);
 			TestTrue(TEXT("Handle array hard refcount"), Handle->GetHardRefCount() == 1);
 
-			HandleList.Remove(Reference);
+			HandleList.Remove(Reference.ToHandle());
 			TestTrue(TEXT("Handle array soft refcount"), Reference->GetSoftRefCount() == 0);
 			TestTrue(TEXT("Handle array hard refcount"), Reference->GetHardRefCount() == 1);
 		}
@@ -141,17 +141,17 @@ namespace WorldPartitionTests
 			TestTrue(TEXT("Handle set soft refcount"), Handle->GetSoftRefCount() == 2);
 			TestTrue(TEXT("Handle set hard refcount"), Handle->GetHardRefCount() == 0);
 
-			FWorldPartitionReference ReferenceToHandle = Handle;
+			FWorldPartitionReference ReferenceToHandle = Handle.ToReference();
 			TestTrue(TEXT("Handle/Reference equality"), ReferenceToHandle == Handle);
 			TestTrue(TEXT("Handle/Reference equality"), Handle == ReferenceToHandle);
 			TestTrue(TEXT("Handle soft refcount"), Handle->GetSoftRefCount() == 2);
 			TestTrue(TEXT("Handle hard refcount"), Handle->GetHardRefCount() == 1);
 
 			TestTrue(TEXT("Handle set contains handle"), HandleSet.Contains(Handle));
-			TestTrue(TEXT("Handle set contains reference"), HandleSet.Contains(ReferenceToHandle));
+			TestTrue(TEXT("Handle set contains reference"), HandleSet.Contains(ReferenceToHandle.ToHandle()));
 
-			HandleSet.Add(Reference);
-			TestTrue(TEXT("Handle set contains reference"), HandleSet.Contains(Reference));
+			HandleSet.Add(Reference.ToHandle());
+			TestTrue(TEXT("Handle set contains reference"), HandleSet.Contains(Reference.ToHandle()));
 
 			TestTrue(TEXT("Reference soft refcount"), Reference->GetSoftRefCount() == 1);
 			TestTrue(TEXT("Reference hard refcount"), Reference->GetHardRefCount() == 1);
diff --git a/Engine/Source/Runtime/Engine/Private/WorldPartition/WorldPartition.cpp b/Engine/Source/Runtime/Engine/Private/WorldPartition/WorldPartition.cpp
--- a/Engine/Source/Runtime/Engine/Private/WorldPartition/WorldPartition.cpp
+++ b/Engine/Source/Runtime/Engine/Private/WorldPartition/WorldPartition.cpp
@@ -326,7 +326,7 @@ void UWorldPartition::OnPackageDirtyStateChanged(UPackage* Package)
 		{
 			if (Package->IsDirty())
 			{
-				DirtyActors.Add(ActorHandle, Actor);
+				DirtyActors.Add(ActorHandle.ToReference(), Actor);
 			}
 			else if (ActorHandle->GetHardRefCount())
 			{
@@ -336,7 +336,7 @@ void UWorldPartition::OnPackageDirtyStateChanged(UPackage* Package)
 					PinnedActors->AddActors({ ActorHandle });
 				}
 
-				DirtyActors.Remove(ActorHandle);
+				DirtyActors.Remove(ActorHandle.ToReference());
 			}
 		}
 	}
diff --git a/Engine/Source/Runtime/Engine/Private/WorldPartition/WorldPartitionActorLoaderInterface.cpp b/Engine/Source/Runtime/Engine/Private/WorldPartition/WorldPartitionActorLoaderInterface.cpp
--- a/Engine/Source/Runtime/Engine/Private/WorldPartition/WorldPartitionActorLoaderInterface.cpp
+++ b/Engine/Source/Runtime/Engine/Private/WorldPartition/WorldPartitionActorLoaderInterface.cpp
@@ -229,7 +229,7 @@
 	{
 		if (!ReferenceMap.Contains(Handle->GetGuid()))
 		{
-			ReferenceMap.Emplace(Handle->GetGuid(), Handle);
+			ReferenceMap.Emplace(Handle->GetGuid(), Handle.ToReference());
 			
 			for (const FGuid& ReferencedActorGuid : ActorHandle->GetReferences())
 			{
diff --git a/Engine/Source/Runtime/Engine/Private/WorldPartition/WorldPartitionHandle.cpp b/Engine/Source/Runtime/Engine/Private/WorldPartition/WorldPartitionHandle.cpp
--- a/Engine/Source/Runtime/Engine/Private/WorldPartition/WorldPartitionHandle.cpp
+++ b/Engine/Source/Runtime/Engine/Private/WorldPartition/WorldPartitionHandle.cpp
@@ -6,9 +6,9 @@
 
 #if WITH_EDITOR
 /**
-* FWorldPartitionHandleUtils
+* FWorldPartitionImplBase
 */
-TUniquePtr<FWorldPartitionActorDesc>* FWorldPartitionHandleUtils::GetActorDesc(UActorDescContainer* Container, const FGuid& ActorGuid)
+TUniquePtr<FWorldPartitionActorDesc>* FWorldPartitionImplBase::GetActorDesc(UActorDescContainer* Container, const FGuid& ActorGuid)
 {
 	if (TUniquePtr<FWorldPartitionActorDesc>** ActorDescPtr = Container->ActorsByGuid.Find(ActorGuid))
 	{
@@ -18,12 +18,12 @@ TUniquePtr<FWorldPartitionActorDesc>* FWorldPartitionHandleUtils::GetActorDesc(U
 	return nullptr;
 }
 
-UActorDescContainer* FWorldPartitionHandleUtils::GetActorDescContainer(TUniquePtr<FWorldPartitionActorDesc>* ActorDesc)
+UActorDescContainer* FWorldPartitionImplBase::GetActorDescContainer(TUniquePtr<FWorldPartitionActorDesc>* ActorDesc)
 {
 	return ActorDesc ? ActorDesc->Get()->GetContainer() : nullptr;
 }
 
-bool FWorldPartitionHandleUtils::IsActorDescLoaded(FWorldPartitionActorDesc* ActorDesc)
+bool FWorldPartitionImplBase::IsActorDescLoaded(FWorldPartitionActorDesc* ActorDesc)
 {
 	return ActorDesc->IsLoaded();
 }
@@ -100,7 +100,7 @@ void FWorldPartitionLoadingContext::FImmediate::UnregisterActor(FWorldPartitionA
 */
 FWorldPartitionLoadingContext::FDeferred::~FDeferred()
 {
-	for (auto [Container, ContainerOp] : ContainerOps)
+	for (auto& [Container, ContainerOp] : ContainerOps)
 	{
 		const FTransform& ContainerTransform = Container->GetInstanceTransform();
 		const FTransform* ContainerTransformPtr = ContainerTransform.Equals(FTransform::Identity) ? nullptr : &ContainerTransform;
@@ -192,6 +192,13 @@ void FWorldPartitionHandleImpl::DecRefCount(FWorldPartitionActorDesc* ActorDesc)
 	ActorDesc->DecSoftRefCount();
 }
 
+TWorldPartitionHandle<FWorldPartitionReferenceImpl> FWorldPartitionHandleImpl::ToReference(const TWorldPartitionHandle<FWorldPartitionHandleImpl>& Source)
+{
+	TWorldPartitionHandle<FWorldPartitionReferenceImpl> Result;
+	Result = Source;
+	return Result;
+}
+
 /**
 * FWorldPartitionReferenceImpl
 */
@@ -210,4 +217,11 @@ void FWorldPartitionReferenceImpl::DecRefCount(FWorldPartitionActorDesc* ActorDe
 		FWorldPartitionLoadingContext::UnloadAndUnregisterActor(ActorDesc);
 	}
 }
+
+TWorldPartitionHandle<FWorldPartitionHandleImpl> FWorldPartitionReferenceImpl::ToHandle(const TWorldPartitionHandle<FWorldPartitionReferenceImpl>& Source)
+{
+	TWorldPartitionHandle<FWorldPartitionHandleImpl> Result;
+	Result = Source;
+	return Result;
+}
 #endif
\ No newline at end of file
diff --git a/Engine/Source/Runtime/Engine/Public/WorldPartition/ActorDescList.h b/Engine/Source/Runtime/Engine/Public/WorldPartition/ActorDescList.h
--- a/Engine/Source/Runtime/Engine/Public/WorldPartition/ActorDescList.h
+++ b/Engine/Source/Runtime/Engine/Public/WorldPartition/ActorDescList.h
@@ -9,7 +9,7 @@
 class ENGINE_API FActorDescList
 {
 #if WITH_EDITOR
-	friend struct FWorldPartitionHandleUtils;
+	friend struct FWorldPartitionImplBase;
 
 public:
 	FActorDescList() {}
diff --git a/Engine/Source/Runtime/Engine/Public/WorldPartition/WorldPartitionHandle.h b/Engine/Source/Runtime/Engine/Public/WorldPartition/WorldPartitionHandle.h
--- a/Engine/Source/Runtime/Engine/Public/WorldPartition/WorldPartitionHandle.h
+++ b/Engine/Source/Runtime/Engine/Public/WorldPartition/WorldPartitionHandle.h
@@ -9,13 +9,8 @@
 #if WITH_EDITOR
 class UActorDescContainer;
 class FWorldPartitionActorDesc;
-
-struct ENGINE_API FWorldPartitionHandleUtils
-{
-	static TUniquePtr<FWorldPartitionActorDesc>* GetActorDesc(UActorDescContainer* Container, const FGuid& ActorGuid);
-	static UActorDescContainer* GetActorDescContainer(TUniquePtr<FWorldPartitionActorDesc>* ActorDesc);
-	static bool IsActorDescLoaded(FWorldPartitionActorDesc* ActorDesc);
-};
+struct FWorldPartitionHandleImpl;
+struct FWorldPartitionReferenceImpl;
 
 class ENGINE_API FWorldPartitionLoadingContext
 {
@@ -86,7 +81,7 @@ public:
 	{}
 
 	FORCEINLINE TWorldPartitionHandle(TUniquePtr<FWorldPartitionActorDesc>* InActorDesc)
-		: Container(FWorldPartitionHandleUtils::GetActorDescContainer(InActorDesc))
+		: Container(Impl::GetActorDescContainer(InActorDesc))
 		, ActorDesc(InActorDesc)
 	{
 		if (IsValid())
@@ -97,7 +92,7 @@ public:
 
 	FORCEINLINE TWorldPartitionHandle(UActorDescContainer* Container, const FGuid& ActorGuid)
 		: Container(Container)
-		, ActorDesc(FWorldPartitionHandleUtils::GetActorDesc(Container, ActorGuid))
+		, ActorDesc(Impl::GetActorDesc(Container, ActorGuid))
 	{
 		if (IsValid())
 		{
@@ -119,12 +114,14 @@ public:
 		*this = MoveTemp(Other);
 	}
 
-	// Conversions
-	template <typename T>
-	FORCEINLINE TWorldPartitionHandle<Impl>(const TWorldPartitionHandle<T>& Other)
-		: ActorDesc(nullptr)
+	FORCEINLINE TWorldPartitionHandle<FWorldPartitionHandleImpl> ToHandle() const
 	{
-		*this = Other;
+		return Impl::ToHandle(*this);
+	}
+
+	FORCEINLINE TWorldPartitionHandle<FWorldPartitionReferenceImpl> ToReference() const
+	{
+		return Impl::ToReference(*this);
 	}
 
 	template <typename T>
@@ -242,7 +239,7 @@ public:
 
 	FORCEINLINE bool IsLoaded() const
 	{
-		return IsValid() && FWorldPartitionHandleUtils::IsActorDescLoaded(ActorDesc->Get());
+		return IsValid() && Impl::IsActorDescLoaded(ActorDesc->Get());
 	}
 
 	FORCEINLINE FWorldPartitionActorDesc* Get() const
@@ -299,16 +296,25 @@ public:
 	TUniquePtr<FWorldPartitionActorDesc>* ActorDesc;
 };
 
-struct ENGINE_API FWorldPartitionHandleImpl
+struct ENGINE_API FWorldPartitionImplBase
+{
+	static TUniquePtr<FWorldPartitionActorDesc>* GetActorDesc(UActorDescContainer* Container, const FGuid& ActorGuid);
+	static UActorDescContainer* GetActorDescContainer(TUniquePtr<FWorldPartitionActorDesc>* ActorDesc);
+	static bool IsActorDescLoaded(FWorldPartitionActorDesc* ActorDesc);
+};
+
+struct ENGINE_API FWorldPartitionHandleImpl : FWorldPartitionImplBase
 {
 	static void IncRefCount(FWorldPartitionActorDesc* ActorDesc);
 	static void DecRefCount(FWorldPartitionActorDesc* ActorDesc);
+	static TWorldPartitionHandle<FWorldPartitionReferenceImpl> ToReference(const TWorldPartitionHandle<FWorldPartitionHandleImpl>& Source);
 };
 
-struct ENGINE_API FWorldPartitionReferenceImpl
+struct ENGINE_API FWorldPartitionReferenceImpl : FWorldPartitionImplBase
 {
 	static void IncRefCount(FWorldPartitionActorDesc* ActorDesc);
 	static void DecRefCount(FWorldPartitionActorDesc* ActorDesc);
+	static TWorldPartitionHandle<FWorldPartitionHandleImpl> ToHandle(const TWorldPartitionHandle<FWorldPartitionReferenceImpl>& Source);
 };
 
 /**
@@ -331,13 +337,20 @@ typedef TWorldPartitionHandle<FWorldPartitionReferenceImpl> FWorldPartitionRefer
  */
 struct FWorldPartitionHandlePinRefScope
 {
-	FWorldPartitionHandlePinRefScope(const FWorldPartitionHandle& Handle)
+	FWorldPartitionHandlePinRefScope(const FWorldPartitionHandle& InHandle)
 	{
-		if (Handle.IsLoaded())
+		if (InHandle.IsLoaded())
 		{
-			Reference = Handle;
+			Reference = InHandle;
 		}
 	}
+
+	FWorldPartitionHandlePinRefScope(const FWorldPartitionReference& InReference)
+	{
+		Reference = InReference;
+	}
+
+private:
 	FWorldPartitionReference Reference;
 };
 #endif
\ No newline at end of file
