Added Cache Speed Stats to Analytics
Added Zen Server Stats to Analytics

diff --git a/Engine/Source/Developer/DerivedDataCache/Private/DerivedDataCache.cpp b/Engine/Source/Developer/DerivedDataCache/Private/DerivedDataCache.cpp
--- a/Engine/Source/Developer/DerivedDataCache/Private/DerivedDataCache.cpp
+++ b/Engine/Source/Developer/DerivedDataCache/Private/DerivedDataCache.cpp
@@ -88,39 +88,47 @@ namespace UE::DerivedData::CookStats
 			const TSharedRef<const FDerivedDataCacheStatsNode>* ZenLocalNode = Nodes.FindByPredicate([](TSharedRef<const FDerivedDataCacheStatsNode> Node) { return Node->GetCacheType() == TEXT("Zen") && Node->IsLocal(); });
 			const TSharedRef<const FDerivedDataCacheStatsNode>* ZenRemoteNode = Nodes.FindByPredicate([](TSharedRef<const FDerivedDataCacheStatsNode> Node) { return (Node->GetCacheType() == TEXT("Zen") || Node->GetCacheType() == TEXT("Horde")) && !Node->IsLocal(); });
 
-			const FDerivedDataCacheUsageStats& RootStats = RootNode->Stats.CreateConstIterator().Value();
+			const FDerivedDataCacheUsageStats& RootStats = RootNode->UsageStats.CreateConstIterator().Value();
 			const int64 TotalGetHits = RootStats.GetStats.GetAccumulatedValueAnyThread(FCookStats::CallStats::EHitOrMiss::Hit, FCookStats::CallStats::EStatType::Counter);
 			const int64 TotalGetMisses = RootStats.GetStats.GetAccumulatedValueAnyThread(FCookStats::CallStats::EHitOrMiss::Miss, FCookStats::CallStats::EStatType::Counter);
 			const int64 TotalGets = TotalGetHits + TotalGetMisses;
 
 			int64 LocalHits = 0;
+			FDerivedDataCacheSpeedStats LocalSpeedStats;
 			if (LocalNode)
 			{
-				const FDerivedDataCacheUsageStats& Stats = (*LocalNode)->Stats.CreateConstIterator().Value();
-				LocalHits += Stats.GetStats.GetAccumulatedValueAnyThread(FCookStats::CallStats::EHitOrMiss::Hit, FCookStats::CallStats::EStatType::Counter);
+				const FDerivedDataCacheUsageStats& UsageStats = (*LocalNode)->UsageStats.CreateConstIterator().Value();
+				LocalHits += UsageStats.GetStats.GetAccumulatedValueAnyThread(FCookStats::CallStats::EHitOrMiss::Hit, FCookStats::CallStats::EStatType::Counter);
+				LocalSpeedStats = (*LocalNode)->SpeedStats;
 			}
 			if (ZenLocalNode)
 			{
-				const FDerivedDataCacheUsageStats& Stats = (*ZenLocalNode)->Stats.CreateConstIterator().Value();
-				LocalHits += Stats.GetStats.GetAccumulatedValueAnyThread(FCookStats::CallStats::EHitOrMiss::Hit, FCookStats::CallStats::EStatType::Counter);
+				const FDerivedDataCacheUsageStats& UsageStats = (*ZenLocalNode)->UsageStats.CreateConstIterator().Value();
+				LocalHits += UsageStats.GetStats.GetAccumulatedValueAnyThread(FCookStats::CallStats::EHitOrMiss::Hit, FCookStats::CallStats::EStatType::Counter);
+				LocalSpeedStats = (*ZenLocalNode)->SpeedStats;
 			}
 			int64 SharedHits = 0;
+			FDerivedDataCacheSpeedStats SharedSpeedStats;
 			if (SharedNode)
 			{
 				// The shared DDC is only queried if the local one misses (or there isn't one). So it's hit rate is technically 
-				const FDerivedDataCacheUsageStats& Stats = (*SharedNode)->Stats.CreateConstIterator().Value();
-				SharedHits += Stats.GetStats.GetAccumulatedValueAnyThread(FCookStats::CallStats::EHitOrMiss::Hit, FCookStats::CallStats::EStatType::Counter);
+				const FDerivedDataCacheUsageStats& UsageStats = (*SharedNode)->UsageStats.CreateConstIterator().Value();
+				SharedHits += UsageStats.GetStats.GetAccumulatedValueAnyThread(FCookStats::CallStats::EHitOrMiss::Hit, FCookStats::CallStats::EStatType::Counter);
+				SharedSpeedStats = (*SharedNode)->SpeedStats;
 			}
 			if (ZenRemoteNode)
 			{
-				const FDerivedDataCacheUsageStats& Stats = (*ZenRemoteNode)->Stats.CreateConstIterator().Value();
-				SharedHits += Stats.GetStats.GetAccumulatedValueAnyThread(FCookStats::CallStats::EHitOrMiss::Hit, FCookStats::CallStats::EStatType::Counter);
+				const FDerivedDataCacheUsageStats& UsageStats = (*ZenRemoteNode)->UsageStats.CreateConstIterator().Value();
+				SharedHits += UsageStats.GetStats.GetAccumulatedValueAnyThread(FCookStats::CallStats::EHitOrMiss::Hit, FCookStats::CallStats::EStatType::Counter);
+				SharedSpeedStats = (*ZenRemoteNode)->SpeedStats;
 			}
 			int64 CloudHits = 0;
+			FDerivedDataCacheSpeedStats CloudSpeedStats;
 			if (CloudNode)
 			{
-				const FDerivedDataCacheUsageStats& Stats = (*CloudNode)->Stats.CreateConstIterator().Value();
-				CloudHits += Stats.GetStats.GetAccumulatedValueAnyThread(FCookStats::CallStats::EHitOrMiss::Hit, FCookStats::CallStats::EStatType::Counter);
+				const FDerivedDataCacheUsageStats& UsageStats = (*CloudNode)->UsageStats.CreateConstIterator().Value();
+				CloudHits += UsageStats.GetStats.GetAccumulatedValueAnyThread(FCookStats::CallStats::EHitOrMiss::Hit, FCookStats::CallStats::EStatType::Counter);
+				CloudSpeedStats = (*CloudNode)->SpeedStats;
 			}
 
 			const int64 TotalPutHits = RootStats.PutStats.GetAccumulatedValueAnyThread(FCookStats::CallStats::EHitOrMiss::Hit, FCookStats::CallStats::EStatType::Counter);
@@ -144,8 +152,17 @@ namespace UE::DerivedData::CookStats
 				TEXT("TotalPutHits"), TotalPutHits,
 				TEXT("TotalPuts"), TotalPuts,
 				TEXT("TotalPutHitPct"), SafeDivide(TotalPutHits, TotalPuts),
-				TEXT("PutMissPct"), SafeDivide(TotalPutMisses, TotalPuts)
-				));
+				TEXT("PutMissPct"), SafeDivide(TotalPutMisses, TotalPuts),
+				TEXT("LocalLatency"), LocalSpeedStats.LatencyMS,
+				TEXT("LocalReadSpeed"), LocalSpeedStats.ReadSpeedMBs,
+				TEXT("LocalWriteSpeed"), LocalSpeedStats.WriteSpeedMBs,
+				TEXT("SharedLatency"), SharedSpeedStats.LatencyMS,
+				TEXT("SharedReadSpeed"), SharedSpeedStats.ReadSpeedMBs,
+				TEXT("SharedWriteSpeed"), SharedSpeedStats.WriteSpeedMBs,
+				TEXT("CloudLatency"), CloudSpeedStats.LatencyMS,
+				TEXT("CloudReadSpeed"), CloudSpeedStats.ReadSpeedMBs,
+				TEXT("CloudWriteSpeed"), CloudSpeedStats.WriteSpeedMBs
+			));
 		}
 	}
 
diff --git a/Engine/Source/Developer/DerivedDataCache/Private/DerivedDataCacheStoreAsync.cpp b/Engine/Source/Developer/DerivedDataCache/Private/DerivedDataCacheStoreAsync.cpp
--- a/Engine/Source/Developer/DerivedDataCache/Private/DerivedDataCacheStoreAsync.cpp
+++ b/Engine/Source/Developer/DerivedDataCache/Private/DerivedDataCacheStoreAsync.cpp
@@ -165,7 +165,7 @@
 void FCacheStoreAsync::LegacyStats(FDerivedDataCacheStatsNode& OutNode)
 {
 	OutNode = {TEXT("Async"), TEXT(""), EnumHasAnyFlags(InnerFlags, ECacheStoreFlags::Local)};
-	OutNode.Stats.Add(TEXT(""), UsageStats);
+	OutNode.UsageStats.Add(TEXT(""), UsageStats);
 
 	InnerCache->LegacyStats(OutNode.Children.Add_GetRef(MakeShared<FDerivedDataCacheStatsNode>()).Get());
 	if (MemoryCache)
diff --git a/Engine/Source/Developer/DerivedDataCache/Private/FileSystemCacheStore.cpp b/Engine/Source/Developer/DerivedDataCache/Private/FileSystemCacheStore.cpp
--- a/Engine/Source/Developer/DerivedDataCache/Private/FileSystemCacheStore.cpp
+++ b/Engine/Source/Developer/DerivedDataCache/Private/FileSystemCacheStore.cpp
@@ -783,7 +783,7 @@ public:
 	FString GetName() const final { return CachePath; }
 	bool IsWritable() const final { return !bReadOnly && !bDisabled; }
 	ESpeedClass GetSpeedClass() const final { return SpeedClass; }
-
+	
 	bool CachedDataProbablyExists(const TCHAR* CacheKey) final;
 	bool GetCachedData(const TCHAR* CacheKey, TArray<uint8>& Data) final;
 	EPutStatus PutCachedData(const TCHAR* CacheKey, TConstArrayView<uint8> Data, bool bPutEvenIfExists) final;
@@ -915,6 +915,9 @@ private:
 
 	FDerivedDataCacheUsageStats UsageStats;
 
+	/** Speed stats */
+	FDerivedDataCacheSpeedStats SpeedStats;
+
 	TUniquePtr<FFileSystemCacheStoreMaintainer> Maintainer;
 };
 
@@ -963,10 +966,7 @@ FFileSystemCacheStore::FFileSystemCacheStore(
 	}
 
 	// check latency and speed. Read values should always be valid
-	double ReadSpeedMBs = 0.0;
-	double WriteSpeedMBs = 0.0;
-	double SeekTimeMS = 0.0;
-
+	
 	/* Speeds faster than this are considered local*/
 	const float ConsiderFastAtMS = 10;
 	/* Speeds faster than this are ok. Everything else is slow. This value can be overridden in the ini file */
@@ -977,35 +977,35 @@ FFileSystemCacheStore::FFileSystemCacheStore(
 	bool SkipSpeedTest = !WITH_EDITOR || FParse::Param(FCommandLine::Get(), TEXT("DDCSkipSpeedTest"));
 	if (SkipSpeedTest)
 	{
-		ReadSpeedMBs = 999;
-		WriteSpeedMBs = 999;
-		SeekTimeMS = 0;
+		SpeedStats.ReadSpeedMBs = 999;
+		SpeedStats.WriteSpeedMBs = 999;
+		SpeedStats.LatencyMS = 0;
 		UE_LOG(LogDerivedDataCache, Log, TEXT("Skipping speed test to %s. Assuming local performance"), *CachePath);
 	}
 
-	if (!SkipSpeedTest && !RunSpeedTest(ConsiderSlowAtMS * 2, SeekTimeMS, ReadSpeedMBs, WriteSpeedMBs))
+	if (!SkipSpeedTest && !RunSpeedTest(ConsiderSlowAtMS * 2, SpeedStats.LatencyMS, SpeedStats.ReadSpeedMBs, SpeedStats.WriteSpeedMBs))
 	{
 		bDisabled = true;
 		UE_LOG(LogDerivedDataCache, Warning, TEXT("No read or write access to %s"), *CachePath);
 	}
 	else
 	{
-		bool bReadTestPassed = ReadSpeedMBs > 0.0;
-		bool bWriteTestPassed = WriteSpeedMBs > 0.0;
+		bool bReadTestPassed = SpeedStats.ReadSpeedMBs > 0.0;
+		bool bWriteTestPassed = SpeedStats.WriteSpeedMBs > 0.0;
 
 		// if we failed writes mark this as read only
 		bReadOnly = bReadOnly || !bWriteTestPassed;
 
 		// classify and report on these times
-		if (SeekTimeMS < 1)
+		if (SpeedStats.LatencyMS < 1)
 		{
 			SpeedClass = ESpeedClass::Local;
 		}
-		else if (SeekTimeMS <= ConsiderFastAtMS)
+		else if (SpeedStats.LatencyMS <= ConsiderFastAtMS)
 		{
 			SpeedClass = ESpeedClass::Fast;
 		}
-		else if (SeekTimeMS >= ConsiderSlowAtMS)
+		else if (SpeedStats.LatencyMS >= ConsiderSlowAtMS)
 		{
 			SpeedClass = ESpeedClass::Slow;
 		}
@@ -1017,7 +1017,7 @@ FFileSystemCacheStore::FFileSystemCacheStore(
 		UE_LOG(LogDerivedDataCache, Display,
 			TEXT("Performance to %s: Latency=%.02fms. RandomReadSpeed=%.02fMBs, RandomWriteSpeed=%.02fMBs. ")
 			TEXT("Assigned SpeedClass '%s'"),
-			*CachePath, SeekTimeMS, ReadSpeedMBs, WriteSpeedMBs, LexToString(SpeedClass));
+			*CachePath, SpeedStats.LatencyMS, SpeedStats.ReadSpeedMBs, SpeedStats.WriteSpeedMBs, LexToString(SpeedClass));
 
 		if (SpeedClass <= FDerivedDataBackendInterface::ESpeedClass::Slow && !bReadOnly)
 		{
@@ -1302,11 +1302,6 @@ bool FFileSystemCacheStore::RunSpeedTest(
 		AsyncTask(ENamedThreads::AnyThread, [CustomPath]() {
 			IFileManager::Get().DeleteDirectory(*CustomPath, false, true);
 		});
-
-		// check latency and speed. Read values should always be valid
-		const double ReadSpeedMBs = (bReadTestPassed ? (TotalDataRead / TotalReadTime) : 0) / (1024 * 1024);
-		const double WriteSpeedMBs = (bWriteTestPassed ? (TotalDataWritten / TotalWriteTime) : 0) / (1024 * 1024);
-		const double SeekTimeMS = (TotalSeekTime / TestFileEntries.Num()) * 1000;
 	}
 
 	const double TotalTestTime = FPlatformTime::Seconds() - StatStartTime;
@@ -1577,7 +1572,8 @@ TSharedRef<FDerivedDataCacheStatsNode> FFileSystemCacheStore::GatherUsageStats()
 {
 	TSharedRef<FDerivedDataCacheStatsNode> Usage =
 		MakeShared<FDerivedDataCacheStatsNode>(TEXT("File System"), *CachePath, SpeedClass == ESpeedClass::Local);
-	Usage->Stats.Add(TEXT(""), UsageStats);
+	Usage->UsageStats.Add(TEXT(""), UsageStats);
+	Usage->SpeedStats = SpeedStats;
 	return Usage;
 }
 
diff --git a/Engine/Source/Developer/DerivedDataCache/Private/HttpCacheStore.cpp b/Engine/Source/Developer/DerivedDataCache/Private/HttpCacheStore.cpp
--- a/Engine/Source/Developer/DerivedDataCache/Private/HttpCacheStore.cpp
+++ b/Engine/Source/Developer/DerivedDataCache/Private/HttpCacheStore.cpp
@@ -3272,10 +3272,10 @@ void FHttpCacheStore::RemoveCachedData(const TCHAR* CacheKey, bool bTransient)
 
 TSharedRef<FDerivedDataCacheStatsNode> FHttpCacheStore::GatherUsageStats() const
 {
-	TSharedRef<FDerivedDataCacheStatsNode> Usage =
+	TSharedRef<FDerivedDataCacheStatsNode> StatsNode =
 		MakeShared<FDerivedDataCacheStatsNode>(TEXT("Horde Storage"), FString::Printf(TEXT("%s (%s)"), *Domain, *Namespace), /*bIsLocal*/ false);
-	Usage->Stats.Add(TEXT(""), UsageStats);
-	return Usage;
+	StatsNode->UsageStats.Add(TEXT(""), UsageStats);
+	return StatsNode;
 }
 
 void FHttpCacheStore::Put(
diff --git a/Engine/Source/Developer/DerivedDataCache/Private/MemoryCacheStore.cpp b/Engine/Source/Developer/DerivedDataCache/Private/MemoryCacheStore.cpp
--- a/Engine/Source/Developer/DerivedDataCache/Private/MemoryCacheStore.cpp
+++ b/Engine/Source/Developer/DerivedDataCache/Private/MemoryCacheStore.cpp
@@ -163,7 +163,7 @@ void FMemoryCacheStore::Disable()
 void FMemoryCacheStore::LegacyStats(FDerivedDataCacheStatsNode& OutNode)
 {
 	OutNode = {!bCanBeDisabled ? TEXT("Memory") : TEXT("Boot"), TEXT(""), /*bIsLocal*/ true};
-	OutNode.Stats.Add(TEXT(""), UsageStats);
+	OutNode.UsageStats.Add(TEXT(""), UsageStats);
 }
 
 bool FMemoryCacheStore::LegacyDebugOptions(FBackendDebugOptions& InOptions)
diff --git a/Engine/Source/Developer/DerivedDataCache/Private/PakFileCacheStore.cpp b/Engine/Source/Developer/DerivedDataCache/Private/PakFileCacheStore.cpp
--- a/Engine/Source/Developer/DerivedDataCache/Private/PakFileCacheStore.cpp
+++ b/Engine/Source/Developer/DerivedDataCache/Private/PakFileCacheStore.cpp
@@ -683,10 +683,10 @@ bool IPakFileCacheStore::SortAndCopy(const FString &InputFilename, const FString
 
 TSharedRef<FDerivedDataCacheStatsNode> FPakFileCacheStore::GatherUsageStats() const
 {
-	TSharedRef<FDerivedDataCacheStatsNode> Usage =
+	TSharedRef<FDerivedDataCacheStatsNode> StatsNode =
 		MakeShared<FDerivedDataCacheStatsNode>(TEXT("PakFile"), CachePath, /*bIsLocal*/ true);
-	Usage->Stats.Add(TEXT(""), UsageStats);
-	return Usage;
+	StatsNode->UsageStats.Add(TEXT(""), UsageStats);
+	return StatsNode;
 }
 
 bool FPakFileCacheStore::ApplyDebugOptions(FBackendDebugOptions& InOptions)
diff --git a/Engine/Source/Developer/DerivedDataCache/Private/S3CacheStore.cpp b/Engine/Source/Developer/DerivedDataCache/Private/S3CacheStore.cpp
--- a/Engine/Source/Developer/DerivedDataCache/Private/S3CacheStore.cpp
+++ b/Engine/Source/Developer/DerivedDataCache/Private/S3CacheStore.cpp
@@ -1126,10 +1126,10 @@ void FS3CacheStore::RemoveCachedData(const TCHAR* CacheKey, bool bTransient)
 
 TSharedRef<FDerivedDataCacheStatsNode> FS3CacheStore::GatherUsageStats() const
 {
-	TSharedRef<FDerivedDataCacheStatsNode> Usage =
+	TSharedRef<FDerivedDataCacheStatsNode> StatsNode =
 		MakeShared<FDerivedDataCacheStatsNode>(TEXT("S3"), BaseUrl, /*bIsLocal*/ false);
-	Usage->Stats.Add(TEXT(""), UsageStats);
-	return Usage;
+	StatsNode->UsageStats.Add(TEXT(""), UsageStats);
+	return StatsNode;
 }
 
 FString FS3CacheStore::GetName() const
diff --git a/Engine/Source/Developer/DerivedDataCache/Private/ZenCacheStore.cpp b/Engine/Source/Developer/DerivedDataCache/Private/ZenCacheStore.cpp
--- a/Engine/Source/Developer/DerivedDataCache/Private/ZenCacheStore.cpp
+++ b/Engine/Source/Developer/DerivedDataCache/Private/ZenCacheStore.cpp
@@ -564,7 +564,7 @@ TSharedRef<FDerivedDataCacheStatsNode> FZenCacheStore::GatherUsageStats() const
 
 	TSharedRef<FDerivedDataCacheStatsNode> LocalNode =
 		MakeShared<FDerivedDataCacheStatsNode>(TEXT("Zen"), ZenService.GetInstance().GetURL(), /*bIsLocal*/ true);
-	LocalNode->Stats.Add(TEXT(""), LocalStats);
+	LocalNode->UsageStats.Add(TEXT(""), LocalStats);
 
 	if (ZenStats.UpstreamStats.EndPointStats.IsEmpty())
 	{
@@ -573,7 +573,7 @@ TSharedRef<FDerivedDataCacheStatsNode> FZenCacheStore::GatherUsageStats() const
 
 	TSharedRef<FDerivedDataCacheStatsNode> RemoteNode =
 		MakeShared<FDerivedDataCacheStatsNode>(ZenStats.UpstreamStats.EndPointStats[0].Name, ZenStats.UpstreamStats.EndPointStats[0].Url, /*bIsLocal*/ false);
-	RemoteNode->Stats.Add(TEXT(""), RemoteStats);
+	RemoteNode->UsageStats.Add(TEXT(""), RemoteStats);
 
 	TSharedRef<FDerivedDataCacheStatsNode> GroupNode =
 		MakeShared<FDerivedDataCacheStatsNode>(TEXT("Zen Group"), TEXT(""), /*bIsLocal*/ true);
diff --git a/Engine/Source/Developer/DerivedDataCache/Public/DerivedDataCacheUsageStats.h b/Engine/Source/Developer/DerivedDataCache/Public/DerivedDataCacheUsageStats.h
--- a/Engine/Source/Developer/DerivedDataCache/Public/DerivedDataCacheUsageStats.h
+++ b/Engine/Source/Developer/DerivedDataCache/Public/DerivedDataCacheUsageStats.h
@@ -103,6 +103,14 @@ public:
 #endif
 };
 
+/** Performance stats for this backend */
+struct FDerivedDataCacheSpeedStats
+{
+	double		ReadSpeedMBs = 0.0;
+	double		WriteSpeedMBs = 0.0;
+	double		LatencyMS = 0.0;
+};
+
 /**
  *  Hierarchical usage stats for the DDC nodes.
  */
@@ -126,9 +134,9 @@ public:
 
 	TMap<FString, FDerivedDataCacheUsageStats> ToLegacyUsageMap() const
 	{
-		TMap<FString, FDerivedDataCacheUsageStats> UsageStats;
-		GatherLegacyUsageStats(UsageStats, TEXT(" 0"));
-		return UsageStats;
+		TMap<FString, FDerivedDataCacheUsageStats> Stats;
+		GatherLegacyUsageStats(Stats, TEXT(" 0"));
+		return Stats;
 	}
 
 	void ForEachDescendant(TFunctionRef<void(TSharedRef<const FDerivedDataCacheStatsNode>)> Predicate) const
@@ -144,16 +152,16 @@ public:
 public:
 	void GatherLegacyUsageStats(TMap<FString, FDerivedDataCacheUsageStats>& UsageStatsMap, FString&& GraphPath) const
 	{
-		if (Stats.Num() == 1)
+		if (UsageStats.Num() == 1)
 		{
-			for (const auto& KVP : Stats)
+			for (const auto& KVP : UsageStats)
 			{
 				COOK_STAT(UsageStatsMap.Add(FString::Printf(TEXT("%s: %s"), *GraphPath, *GetCacheName()), KVP.Value));
 			}
 		}
 		else
 		{ //-V523
-			for (const auto& KVP : Stats)
+			for (const auto& KVP : UsageStats)
 			{
 				COOK_STAT(UsageStatsMap.Add(FString::Printf(TEXT("%s: %s.%s"), *GraphPath, *GetCacheName(), *KVP.Key), KVP.Value));
 			}
@@ -166,7 +174,8 @@ public:
 		}
 	}
 
-	TMap<FString, FDerivedDataCacheUsageStats> Stats;
+	TMap<FString, FDerivedDataCacheUsageStats> UsageStats;
+	FDerivedDataCacheSpeedStats SpeedStats;
 
 	TArray<TSharedRef<FDerivedDataCacheStatsNode>> Children;
 
diff --git a/Engine/Source/Editor/DerivedDataEditor/Private/DerivedDataInformation.cpp b/Engine/Source/Editor/DerivedDataEditor/Private/DerivedDataInformation.cpp
--- a/Engine/Source/Editor/DerivedDataEditor/Private/DerivedDataInformation.cpp
+++ b/Engine/Source/Editor/DerivedDataEditor/Private/DerivedDataInformation.cpp
@@ -45,7 +45,7 @@ double FDerivedDataInformation::GetCacheActivitySizeBytes(bool bGet, bool bLocal
 			continue;
 		}
 
-		for (const auto& KVP : Usage->Stats)
+		for (const auto& KVP : Usage->UsageStats)
 		{
 			const FDerivedDataCacheUsageStats& Stats = KVP.Value;
 
@@ -77,7 +77,7 @@ double FDerivedDataInformation::GetCacheActivityTimeSeconds(bool bGet, bool bLoc
 			continue;
 		}
 
-		for (const auto& KVP : Usage->Stats)
+		for (const auto& KVP : Usage->UsageStats)
 		{
 			const FDerivedDataCacheUsageStats& Stats = KVP.Value;
 
diff --git a/Engine/Source/Editor/DerivedDataEditor/Private/SDerivedDataDialogs.cpp b/Engine/Source/Editor/DerivedDataEditor/Private/SDerivedDataDialogs.cpp
--- a/Engine/Source/Editor/DerivedDataEditor/Private/SDerivedDataDialogs.cpp
+++ b/Engine/Source/Editor/DerivedDataEditor/Private/SDerivedDataDialogs.cpp
@@ -609,7 +609,7 @@ TSharedRef<SWidget> SDerivedDataCacheStatisticsDialog::GetGridPanel()
 
 		FDerivedDataCacheUsageStats Stats;
 
-		for (const auto& KVP : Node->Stats)
+		for (const auto& KVP : Node->UsageStats)
 		{
 			Stats.Combine(KVP.Value);
 		}
diff --git a/Engine/Source/Runtime/Engine/Engine.Build.cs b/Engine/Source/Runtime/Engine/Engine.Build.cs
--- a/Engine/Source/Runtime/Engine/Engine.Build.cs
+++ b/Engine/Source/Runtime/Engine/Engine.Build.cs
@@ -108,6 +108,7 @@ public class Engine : ModuleRules
 				"DeveloperSettings",
 				"AudioLinkCore",
 				"CookOnTheFly",
+				"Zen"
 			}
 		);
 
diff --git a/Engine/Source/Runtime/Engine/Private/Analytics/StudioAnalytics.cpp b/Engine/Source/Runtime/Engine/Private/Analytics/StudioAnalytics.cpp
--- a/Engine/Source/Runtime/Engine/Private/Analytics/StudioAnalytics.cpp
+++ b/Engine/Source/Runtime/Engine/Private/Analytics/StudioAnalytics.cpp
@@ -20,6 +20,7 @@
 #include "HAL/PlatformProcess.h"
 #include "DerivedDataCacheInterface.h"
 #include "DerivedDataCacheUsageStats.h"
+#include "ZenServerInterface.h"
 #include "Virtualization/VirtualizationSystem.h"
 
 bool FStudioAnalytics::bInitialized = false;
@@ -222,71 +223,158 @@ void FStudioAnalytics::FireEvent_Loading(const FString& LoadingName, double Seco
 		}
 
 		// Grab the DDC summary stats
-		FDerivedDataCacheSummaryStats DDCSummaryStats;
-		GetDerivedDataCacheRef().GatherSummaryStats(DDCSummaryStats);
-		
-		for (const FDerivedDataCacheSummaryStat& Stat : DDCSummaryStats.Stats)
 		{
-			FString FormattedAttrName = "DDC.Summary." + Stat.Key;
-			Attributes.Emplace(FormattedAttrName, Stat.Value);
-		}
+			FDerivedDataCacheSummaryStats DDCSummaryStats;
+			GetDerivedDataCacheRef().GatherSummaryStats(DDCSummaryStats);
 
-		// Grab the Virtualization stats
-		IVirtualizationSystem& System = IVirtualizationSystem::Get();
+			for (const FDerivedDataCacheSummaryStat& Stat : DDCSummaryStats.Stats)
+			{
+				FString FormattedAttrName = "DDC.Summary." + Stat.Key;
+				Attributes.Emplace(FormattedAttrName, Stat.Value);
+			}
 
-		FPayloadActivityInfo PayloadActivityInfo = System.GetAccumualtedPayloadActivityInfo();
+			// Grab the Virtualization stats
+			IVirtualizationSystem& System = IVirtualizationSystem::Get();
 
-		const FString BaseName = TEXT("Virtualization");
+			FPayloadActivityInfo PayloadActivityInfo = System.GetAccumualtedPayloadActivityInfo();
 
-		{
-			FString AttrName = BaseName + TEXT(".Enabled");
-			Attributes.Emplace(MoveTemp(AttrName), System.IsEnabled());
-		}
+			const FString BaseName = TEXT("Virtualization");
 
-		{
-			FString AttrName = BaseName + TEXT(".Cache.TimeSpent");
-			Attributes.Emplace(MoveTemp(AttrName), (double)PayloadActivityInfo.Cache.CyclesSpent * FPlatformTime::GetSecondsPerCycle());
-		}
+			{
+				FString AttrName = BaseName + TEXT(".Enabled");
+				Attributes.Emplace(MoveTemp(AttrName), System.IsEnabled());
+			}
 
-		{
-			FString AttrName = BaseName + TEXT(".Cache.PayloadCount");
-			Attributes.Emplace(MoveTemp(AttrName), (double)PayloadActivityInfo.Cache.PayloadCount);
-		}
+			{
+				FString AttrName = BaseName + TEXT(".Cache.TimeSpent");
+				Attributes.Emplace(MoveTemp(AttrName), (double)PayloadActivityInfo.Cache.CyclesSpent * FPlatformTime::GetSecondsPerCycle());
+			}
 
-		{
-			FString AttrName = BaseName + TEXT(".Cache.TotalBytes");
-			Attributes.Emplace(MoveTemp(AttrName), (double)PayloadActivityInfo.Cache.TotalBytes);
-		}
+			{
+				FString AttrName = BaseName + TEXT(".Cache.PayloadCount");
+				Attributes.Emplace(MoveTemp(AttrName), (double)PayloadActivityInfo.Cache.PayloadCount);
+			}
 
-		{
-			FString AttrName = BaseName + TEXT(".Push.TimeSpent");
-			Attributes.Emplace(MoveTemp(AttrName), (double)PayloadActivityInfo.Push.CyclesSpent * FPlatformTime::GetSecondsPerCycle());
-		}
+			{
+				FString AttrName = BaseName + TEXT(".Cache.TotalBytes");
+				Attributes.Emplace(MoveTemp(AttrName), (double)PayloadActivityInfo.Cache.TotalBytes);
+			}
 
-		{
-			FString AttrName = BaseName + TEXT(".Push.PayloadCount");
-			Attributes.Emplace(MoveTemp(AttrName), (double)PayloadActivityInfo.Push.PayloadCount);
-		}
+			{
+				FString AttrName = BaseName + TEXT(".Push.TimeSpent");
+				Attributes.Emplace(MoveTemp(AttrName), (double)PayloadActivityInfo.Push.CyclesSpent * FPlatformTime::GetSecondsPerCycle());
+			}
 
-		{
-			FString AttrName = BaseName + TEXT(".Push.TotalBytes");
-			Attributes.Emplace(MoveTemp(AttrName), (double)PayloadActivityInfo.Push.TotalBytes);
-		}
+			{
+				FString AttrName = BaseName + TEXT(".Push.PayloadCount");
+				Attributes.Emplace(MoveTemp(AttrName), (double)PayloadActivityInfo.Push.PayloadCount);
+			}
 
-		{
-			FString AttrName = BaseName + TEXT(".Pull.TimeSpent");
-			Attributes.Emplace(MoveTemp(AttrName), (double)PayloadActivityInfo.Pull.CyclesSpent * FPlatformTime::GetSecondsPerCycle());
-		}
+			{
+				FString AttrName = BaseName + TEXT(".Push.TotalBytes");
+				Attributes.Emplace(MoveTemp(AttrName), (double)PayloadActivityInfo.Push.TotalBytes);
+			}
 
-		{
-			FString AttrName = BaseName + TEXT(".Pull.PayloadCount");
-			Attributes.Emplace(MoveTemp(AttrName), (double)PayloadActivityInfo.Pull.PayloadCount);
+			{
+				FString AttrName = BaseName + TEXT(".Pull.TimeSpent");
+				Attributes.Emplace(MoveTemp(AttrName), (double)PayloadActivityInfo.Pull.CyclesSpent * FPlatformTime::GetSecondsPerCycle());
+			}
+
+			{
+				FString AttrName = BaseName + TEXT(".Pull.PayloadCount");
+				Attributes.Emplace(MoveTemp(AttrName), (double)PayloadActivityInfo.Pull.PayloadCount);
+			}
+
+			{
+				FString AttrName = BaseName + TEXT(".Pull.TotalBytes");
+				Attributes.Emplace(MoveTemp(AttrName), (double)PayloadActivityInfo.Pull.TotalBytes);
+			}
 		}
 
+#if UE_WITH_ZEN
 		{
-			FString AttrName = BaseName + TEXT(".Pull.TotalBytes");
-			Attributes.Emplace(MoveTemp(AttrName), (double)PayloadActivityInfo.Pull.TotalBytes);
+			// Grab the Zen summary stats
+			UE::Zen::FZenStats ZenStats;
+			UE::Zen::GetDefaultServiceInstance().GetStats(ZenStats);
+
+			const FString BaseName = TEXT("Zen");
+
+			{
+				FString AttrName = BaseName + TEXT(".Enabled");
+				Attributes.Emplace(MoveTemp(AttrName), ZenStats.IsValid);
+			}
+
+			{
+				FString AttrName = BaseName + TEXT(".Cache.HitRatio");
+				Attributes.Emplace(MoveTemp(AttrName), ZenStats.CacheStats.HitRatio);
+			}
+
+			{
+				FString AttrName = BaseName + TEXT(".Cache.Hits");
+				Attributes.Emplace(MoveTemp(AttrName), ZenStats.CacheStats.Hits);
+			}
+
+			{
+				FString AttrName = BaseName + TEXT(".Cache.Misses");
+				Attributes.Emplace(MoveTemp(AttrName), ZenStats.CacheStats.Misses);
+			}
+
+			{
+				FString AttrName = BaseName + TEXT(".Cache.Size.Disk");
+				Attributes.Emplace(MoveTemp(AttrName), ZenStats.CacheStats.Size.Disk);
+			}
+
+			{
+				FString AttrName = BaseName + TEXT(".Cache.Size.Memory");
+				Attributes.Emplace(MoveTemp(AttrName), ZenStats.CacheStats.Size.Memory);
+			}
+
+			{
+				FString AttrName = BaseName + TEXT(".Cache.UpstreamHits");
+				Attributes.Emplace(MoveTemp(AttrName), ZenStats.CacheStats.UpstreamHits);
+			}
+
+			{
+				FString AttrName = BaseName + TEXT(".Cache.UpstreamRatio");
+				Attributes.Emplace(MoveTemp(AttrName), ZenStats.CacheStats.UpstreamRatio);
+			}
+
+			{
+				FString AttrName = BaseName + TEXT(".Cache.TotalUploadedMB");
+				Attributes.Emplace(MoveTemp(AttrName), ZenStats.UpstreamStats.TotalUploadedMB);
+			}
+
+			{
+				FString AttrName = BaseName + TEXT(".Upstream.TotalDownloadedMB");
+				Attributes.Emplace(MoveTemp(AttrName), ZenStats.UpstreamStats.TotalDownloadedMB);
+			}
+
+			{
+				FString AttrName = BaseName + TEXT(".Upstream.TotalUploadedMB");
+				Attributes.Emplace(MoveTemp(AttrName), ZenStats.UpstreamStats.TotalUploadedMB);
+			}
+
+			{
+				FString AttrName = BaseName + TEXT(".Cas.Size.Large");
+				Attributes.Emplace(MoveTemp(AttrName), ZenStats.CASStats.Size.Large);
+			}
+
+			{
+				FString AttrName = BaseName + TEXT(".Cas.Size.Small");
+				Attributes.Emplace(MoveTemp(AttrName), ZenStats.CASStats.Size.Small);
+			}
+
+			{
+				FString AttrName = BaseName + TEXT(".Cas.Size.Tiny");
+				Attributes.Emplace(MoveTemp(AttrName), ZenStats.CASStats.Size.Tiny);
+			}
+
+			{
+				FString AttrName = BaseName + TEXT(".Cas.Size.Total");
+				Attributes.Emplace(MoveTemp(AttrName), ZenStats.CASStats.Size.Total);
+			}
 		}
+#endif
 
 		// Store it all in the loading event
 		FStudioAnalytics::GetProvider().RecordEvent(TEXT("Core.Loading"), Attributes);
