Level Instance:
- Actor descriptor containers are now managed by the level instance subsystem (fix potential bug with two worlds using the same sublevel).
- Update level instance actor descriptor containers bounds when committing a level instance.
- Query the level instance actor descriptor containers bounds when possible.

This is needed since we don't compute bounds for partitioned worlds and level instance worlds are now partitioned by default.

#rnx

diff --git a/Engine/Source/Runtime/Engine/Private/LevelInstance/LevelInstanceSubsystem.cpp b/Engine/Source/Runtime/Engine/Private/LevelInstance/LevelInstanceSubsystem.cpp
--- a/Engine/Source/Runtime/Engine/Private/LevelInstance/LevelInstanceSubsystem.cpp
+++ b/Engine/Source/Runtime/Engine/Private/LevelInstance/LevelInstanceSubsystem.cpp
@@ -61,13 +61,23 @@ ULevelInstanceSubsystem::ULevelInstanceSubsystem()
 	, bIsCreatingLevelInstance(false)
 	, bIsCommittingLevelInstance(false)
 #endif
-{
-
-}
+{}
 
 ULevelInstanceSubsystem::~ULevelInstanceSubsystem()
+{}
+
+void ULevelInstanceSubsystem::AddReferencedObjects(UObject* InThis, FReferenceCollector& Collector)
 {
+	ULevelInstanceSubsystem* This = CastChecked<ULevelInstanceSubsystem>(InThis);
 
+#if WITH_EDITORONLY_DATA
+	if (This->LevelInstanceEdit)
+	{
+		This->LevelInstanceEdit->AddReferencedObjects(Collector);
+	}
+
+	This->ActorDescContainerInstanceManager.AddReferencedObjects(Collector);
+#endif
 }
 
 void ULevelInstanceSubsystem::Initialize(FSubsystemCollectionBase& Collection)
@@ -618,7 +628,18 @@ bool ULevelInstanceSubsystem::GetLevelInstanceBounds(const ILevelInstanceInterfa
 	}
 	else if(LevelInstance->IsWorldAssetValid())
 	{
-		return GetLevelInstanceBoundsFromPackage(CastChecked<AActor>(LevelInstance)->GetActorTransform(), FName(*LevelInstance->GetWorldAssetPackage()), OutBounds);
+		//@todo_ow: can't find a way to trigger this codepath, maybe deadcode?
+		FString LevelPackage = LevelInstance->GetWorldAssetPackage();
+
+		if (FBox ContainerBounds = ActorDescContainerInstanceManager.GetContainerBounds(*LevelPackage); ContainerBounds.IsValid)
+		{
+			FTransform LevelInstancePivotOffsetTransform = FTransform(ULevel::GetLevelInstancePivotOffsetFromPackage(*LevelPackage));
+			FTransform LevelTransform = LevelInstancePivotOffsetTransform * CastChecked<AActor>(LevelInstance)->GetActorTransform();
+			OutBounds = ContainerBounds.TransformBy(LevelTransform);
+			return true;
+		}
+
+		return GetLevelInstanceBoundsFromPackage(CastChecked<AActor>(LevelInstance)->GetActorTransform(), *LevelInstance->GetWorldAssetPackage(), OutBounds);
 	}
 
 	return false;
@@ -1314,11 +1335,83 @@ void ULevelInstanceSubsystem::RemoveLevelsFromWorld(const TArray<ULevel*>& InLev
 	}
 }
 
+void ULevelInstanceSubsystem::FActorDescContainerInstanceManager::FActorDescContainerInstance::AddReferencedObjects(FReferenceCollector& Collector)
+{
+	Collector.AddReferencedObject(Container);
+}
+
+void ULevelInstanceSubsystem::FActorDescContainerInstanceManager::FActorDescContainerInstance::UpdateBounds()
+{
+	Bounds.Init();
+	for (FActorDescList::TIterator<> ActorDescIt(Container); ActorDescIt; ++ActorDescIt)
+	{
+		Bounds += ActorDescIt->GetBounds();
+	}
+}
+
+void ULevelInstanceSubsystem::FActorDescContainerInstanceManager::AddReferencedObjects(FReferenceCollector& Collector)
+{
+	for (auto& [Name, ContainerInstance] : ActorDescContainers)
+	{
+		ContainerInstance.AddReferencedObjects(Collector);
+	}
+}
+
+UActorDescContainer* ULevelInstanceSubsystem::FActorDescContainerInstanceManager::RegisterContainer(FName PackageName, UWorld* InWorld)
+{
+	FActorDescContainerInstance* ExistingContainerInstance = &ActorDescContainers.FindOrAdd(PackageName);
+	UActorDescContainer* ActorDescContainer = ExistingContainerInstance->Container;
+	
+	if (ExistingContainerInstance->RefCount++ == 0)
+	{
+		ActorDescContainer = NewObject<UActorDescContainer>(GetTransientPackage());
+		ExistingContainerInstance->Container = ActorDescContainer;
+		
+		// This will potentially invalidate ExistingContainerInstance due to ActorDescContainers reallocation
+		ActorDescContainer->Initialize(InWorld, PackageName);
+			
+		ExistingContainerInstance = &ActorDescContainers.FindChecked(PackageName);
+		ExistingContainerInstance->UpdateBounds();
+	}
+
+	check(ActorDescContainer->GetWorld() == InWorld);
+	return ActorDescContainer;
+}
+
+void ULevelInstanceSubsystem::FActorDescContainerInstanceManager::UnregisterContainer(UActorDescContainer* Container)
+{
+	FName PackageName = Container->GetContainerPackage();
+	FActorDescContainerInstance& ExistingContainerInstance = ActorDescContainers.FindChecked(PackageName);
+
+	if (--ExistingContainerInstance.RefCount == 0)
+	{
+		ExistingContainerInstance.Container->Uninitialize();
+		ActorDescContainers.FindAndRemoveChecked(PackageName);
+	}
+}
+
+FBox ULevelInstanceSubsystem::FActorDescContainerInstanceManager::GetContainerBounds(FName PackageName) const
+{
+	if (const FActorDescContainerInstance* ActorDescContainerInstance = ActorDescContainers.Find(PackageName))
+	{
+		return ActorDescContainerInstance->Bounds;
+	}
+	return FBox(ForceInit);
+}
+
+void ULevelInstanceSubsystem::FActorDescContainerInstanceManager::OnLevelInstanceActorCommitted(ILevelInstanceInterface* LevelInstance)
+{
+	const FName PackageName = *LevelInstance->GetWorldAssetPackage();
+	if (FActorDescContainerInstance* ActorDescContainerInstance = ActorDescContainers.Find(PackageName))
+	{
+		ActorDescContainerInstance->UpdateBounds();
+	}
+}
+
 ULevelInstanceSubsystem::FLevelsToRemoveScope::FLevelsToRemoveScope(ULevelInstanceSubsystem* InOwner)
 	: Owner(InOwner)
 	, bIsBeingDestroyed(false)
-{
-}
+{}
 
 ULevelInstanceSubsystem::FLevelsToRemoveScope::~FLevelsToRemoveScope()
 {
@@ -1495,11 +1588,6 @@ void ULevelInstanceSubsystem::FLevelInstanceEdit::AddReferencedObjects(FReferenc
 	Collector.AddReferencedObject(LevelStreaming);
 }
 
-FString ULevelInstanceSubsystem::FLevelInstanceEdit::GetReferencerName() const
-{
-	return TEXT("FLevelInstanceEdit");
-}
-
 bool ULevelInstanceSubsystem::FLevelInstanceEdit::CanDiscard(FText* OutReason) const
 {
 	return EditorObject->CanDiscard(OutReason);
@@ -1890,6 +1978,8 @@ bool ULevelInstanceSubsystem::CommitLevelInstanceInternal(TUniquePtr<FLevelInsta
 	// Notify (Actor might get destroyed by this call if its a packed bp)
 	LevelInstance->OnCommit(bChangesCommitted);
 
+	ActorDescContainerInstanceManager.OnLevelInstanceActorCommitted(LevelInstance);
+
 	// Update pointer since BP Compilation might have invalidated LevelInstance
 	LevelInstance = GetLevelInstance(LevelInstanceID);
 
diff --git a/Engine/Source/Runtime/Engine/Private/WorldPartition/LevelInstance/LevelInstanceActorDesc.cpp b/Engine/Source/Runtime/Engine/Private/WorldPartition/LevelInstance/LevelInstanceActorDesc.cpp
--- a/Engine/Source/Runtime/Engine/Private/WorldPartition/LevelInstance/LevelInstanceActorDesc.cpp
+++ b/Engine/Source/Runtime/Engine/Private/WorldPartition/LevelInstance/LevelInstanceActorDesc.cpp
@@ -18,83 +18,9 @@ static FAutoConsoleVariableRef CVarForceLevelStreaming(
 	GLevelInstanceDebugForceLevelStreaming,
 	TEXT("Set to 1 to force Level Instance to be streamed instead of embedded in World Partition grid."));
 
-class FActorDescContainerInstanceManager : public FGCObject
-{
-	struct FActorDescContainerInstance
-	{
-		FActorDescContainerInstance()
-		: Container(nullptr)
-		, RefCount(0)
-		{}
-
-		void AddReferencedObjects(FReferenceCollector& Collector)
-		{
-			Collector.AddReferencedObject(Container);
-		}
-
-		UActorDescContainer* Container;
-		uint32 RefCount;
-	};
-
-	//~ Begin FGCObject Interface
-	virtual void AddReferencedObjects(FReferenceCollector& Collector) override
-	{
-		for (auto& It : ActorDescContainers)
-		{
-			It.Value.AddReferencedObjects(Collector);
-		}
-	}
-
-	virtual FString GetReferencerName() const override
-	{
-		return TEXT("FActorDescContainerInstanceManager");
-	}
-	//~ End FGCObject Interface
-
-public:
-	static FActorDescContainerInstanceManager& Get()
-	{
-		static FActorDescContainerInstanceManager Instance;
-		return Instance;
-	}
-
-	UActorDescContainer* RegisterContainer(FName PackageName, UWorld* InWorld)
-	{
-		FActorDescContainerInstance& ExistingContainerInstance = ActorDescContainers.FindOrAdd(PackageName);
-		UActorDescContainer* ActorDescContainer = ExistingContainerInstance.Container;
-	
-		if (ExistingContainerInstance.RefCount++ == 0)
-		{
-			ActorDescContainer = NewObject<UActorDescContainer>(GetTransientPackage());
-			ExistingContainerInstance.Container = ActorDescContainer;
-		
-			// This will potentially invalidate ExistingContainerInstance due to ActorDescContainers reallocation
-			ActorDescContainer->Initialize(InWorld, PackageName);
-		}
-
-		check(ActorDescContainer->GetWorld() == InWorld);
-		return ActorDescContainer;
-	}
-
-	void UnregisterContainer(UActorDescContainer* Container)
-	{
-		FName PackageName = Container->GetContainerPackage();
-		FActorDescContainerInstance& ExistingContainerInstance = ActorDescContainers.FindChecked(PackageName);
-
-		if (ExistingContainerInstance.RefCount-- == 1)
-		{
-			ExistingContainerInstance.Container->Uninitialize();
-			ActorDescContainers.FindAndRemoveChecked(PackageName);
-		}
-	}
-
-private:
-	TMap<FName, FActorDescContainerInstance> ActorDescContainers;
-};
-
 FLevelInstanceActorDesc::FLevelInstanceActorDesc()
 	: DesiredRuntimeBehavior(ELevelInstanceRuntimeBehavior::Partitioned)
-	, LevelInstanceContainer(nullptr)
+	, LevelInstanceContainer(nullptr)	
 {}
 
 FLevelInstanceActorDesc::~FLevelInstanceActorDesc()
@@ -145,8 +71,17 @@ void FLevelInstanceActorDesc::RegisterContainerInstance(UWorld* InWorld)
 	{
 		if (!LevelPackage.IsNone() && ULevel::GetIsLevelUsingExternalActorsFromPackage(LevelPackage) && ULevelInstanceSubsystem::CanUsePackage(LevelPackage))
 		{
-			LevelInstanceContainer = FActorDescContainerInstanceManager::Get().RegisterContainer(LevelPackage, InWorld);
+			ULevelInstanceSubsystem* LevelInstanceSubsystem = UWorld::GetSubsystem<ULevelInstanceSubsystem>(InWorld);
+			check(LevelInstanceSubsystem);
+
+			LevelInstanceContainer = LevelInstanceSubsystem->RegisterContainer(LevelPackage);
 			check(LevelInstanceContainer.IsValid());
+
+			FTransform LevelInstancePivotOffsetTransform = FTransform(ULevel::GetLevelInstancePivotOffsetFromPackage(LevelPackage));
+			FTransform FinalLevelTransform = LevelInstancePivotOffsetTransform * LevelInstanceTransform;
+			FBox ContainerBounds = LevelInstanceSubsystem->GetContainerBounds(LevelPackage).TransformBy(FinalLevelTransform);
+
+			ContainerBounds.GetCenterAndExtents(BoundsLocation, BoundsExtent);
 		}
 	}
 }
@@ -155,7 +90,10 @@ void FLevelInstanceActorDesc::UnregisterContainerInstance()
 {
 	if (LevelInstanceContainer.IsValid())
 	{
-		FActorDescContainerInstanceManager::Get().UnregisterContainer(LevelInstanceContainer.Get());
+		ULevelInstanceSubsystem* LevelInstanceSubsystem = UWorld::GetSubsystem<ULevelInstanceSubsystem>(LevelInstanceContainer->GetWorld());
+		check(LevelInstanceSubsystem);
+
+		LevelInstanceSubsystem->UnregisterContainer(LevelInstanceContainer.Get());
 		LevelInstanceContainer.Reset();
 	}
 }
@@ -170,7 +108,7 @@ void FLevelInstanceActorDesc::SetContainer(UActorDescContainer* InContainer)
 	}
 	else
 	{
-		UnregisterContainerInstance();
+		UnregisterContainerInstance();		
 	}
 }
 
diff --git a/Engine/Source/Runtime/Engine/Public/LevelInstance/LevelInstanceSubsystem.h b/Engine/Source/Runtime/Engine/Public/LevelInstance/LevelInstanceSubsystem.h
--- a/Engine/Source/Runtime/Engine/Public/LevelInstance/LevelInstanceSubsystem.h
+++ b/Engine/Source/Runtime/Engine/Public/LevelInstance/LevelInstanceSubsystem.h
@@ -16,6 +16,7 @@ class ULevelInstanceEditorObject;
 class ULevelStreamingLevelInstance;
 class ULevelStreamingLevelInstanceEditor;
 class UWorldPartitionSubsystem;
+class UActorDescContainer;
 class UBlueprint;
 
 /**
@@ -29,6 +30,10 @@ class ENGINE_API ULevelInstanceSubsystem : public UWorldSubsystem
 public:
 	ULevelInstanceSubsystem();
 	virtual ~ULevelInstanceSubsystem();
+
+	//~ Begin UObject Interface
+	static void AddReferencedObjects(UObject* InThis, FReferenceCollector& Collector);
+	//~ End UObject Interface
 		
 	//~ Begin USubsystem Interface.
 	virtual void Initialize(FSubsystemCollectionBase& Collection) override;
@@ -103,6 +108,10 @@ public:
 	TArray<AActor*> GetParentLevelInstanceActors(const ULevel* Level) const;
 	FString PrefixWithParentLevelInstanceActorLabels(const FString& ActorLabel, const ULevel* Level) const;
 
+	UActorDescContainer* RegisterContainer(FName PackageName) { return ActorDescContainerInstanceManager.RegisterContainer(PackageName, GetWorld()); }
+	void UnregisterContainer(UActorDescContainer* Container) { ActorDescContainerInstanceManager.UnregisterContainer(Container); }
+	FBox GetContainerBounds(FName PackageName) const { return ActorDescContainerInstanceManager.GetContainerBounds(PackageName); }
+
 	static bool CheckForLoop(const ILevelInstanceInterface* LevelInstance, TArray<TPair<FText, TSoftObjectPtr<UWorld>>>* LoopInfo = nullptr, const ILevelInstanceInterface** LoopStart = nullptr);
 	static bool CheckForLoop(const ILevelInstanceInterface* LevelInstance, TSoftObjectPtr<UWorld> WorldAsset, TArray<TPair<FText, TSoftObjectPtr<UWorld>>>* LoopInfo = nullptr, const ILevelInstanceInterface** LoopStart = nullptr);
 	static bool CanUseWorldAsset(const ILevelInstanceInterface* LevelInstance, TSoftObjectPtr<UWorld> WorldAsset, FString* OutReason);
@@ -131,7 +140,7 @@ private:
 
 	static bool ShouldIgnoreDirtyPackage(UPackage* DirtyPackage, const UWorld* EditingWorld);
 
-	class FLevelInstanceEdit : public FGCObject
+	class FLevelInstanceEdit
 	{
 	public:
 		TObjectPtr<ULevelStreamingLevelInstanceEditor> LevelStreaming;
@@ -143,8 +152,7 @@ private:
 		UWorld* GetEditWorld() const;
 		FLevelInstanceID GetLevelInstanceID() const;
 
-		virtual void AddReferencedObjects(FReferenceCollector& Collector) override;
-		virtual FString GetReferencerName() const override;
+		void AddReferencedObjects(FReferenceCollector& Collector);
 
 		void GetPackagesToSave(TArray<UPackage*>& OutPackagesToSave) const;
 		bool CanDiscard(FText* OutReason = nullptr) const;
@@ -172,6 +180,41 @@ private:
 	};
 	
 	void RemoveLevelsFromWorld(const TArray<ULevel*>& Levels, bool bResetTrans = true);
+
+	class FActorDescContainerInstanceManager
+	{
+		friend class ULevelInstanceSubsystem;
+
+		struct FActorDescContainerInstance
+		{
+			FActorDescContainerInstance()
+			: Container(nullptr)
+			, RefCount(0)
+			, Bounds(ForceInit)
+			{}
+
+			void AddReferencedObjects(FReferenceCollector& Collector);
+			void UpdateBounds();
+
+			UActorDescContainer* Container;
+			uint32 RefCount;
+			FBox Bounds;
+		};
+
+		void AddReferencedObjects(FReferenceCollector& Collector);
+
+	public:
+		UActorDescContainer* RegisterContainer(FName PackageName, UWorld* InWorld);
+		void UnregisterContainer(UActorDescContainer* Container);
+
+		FBox GetContainerBounds(FName PackageName) const;
+
+	private:
+		void OnLevelInstanceActorCommitted(ILevelInstanceInterface* LevelInstance);
+
+		TMap<FName, FActorDescContainerInstance> ActorDescContainers;
+	};
+private:
 #endif
 	friend ULevelStreamingLevelInstance;
 	friend ULevelStreamingLevelInstanceEditor;
@@ -199,5 +242,7 @@ private:
 	TUniquePtr<FLevelInstanceEdit> LevelInstanceEdit;
 
 	TMap<FLevelInstanceID, int32> ChildEdits;
+
+	FActorDescContainerInstanceManager ActorDescContainerInstanceManager;
 #endif
 };
\ No newline at end of file
