Added a method to retrieve sockets by tag

diff --git a/Engine/Source/Runtime/Engine/Classes/Engine/StaticMesh.h b/Engine/Source/Runtime/Engine/Classes/Engine/StaticMesh.h
--- a/Engine/Source/Runtime/Engine/Classes/Engine/StaticMesh.h
+++ b/Engine/Source/Runtime/Engine/Classes/Engine/StaticMesh.h
@@ -1859,6 +1859,12 @@ public:
 	UFUNCTION(BlueprintCallable, Category = "StaticMesh")
 	ENGINE_API void RemoveSocket(UStaticMeshSocket* Socket);
 
+	/**
+	 * Returns a list of sockets with the provided tag.
+	 */
+	UFUNCTION(BlueprintPure, Category = "StaticMesh")
+	ENGINE_API TArray<UStaticMeshSocket*> GetSocketsByTag(const FString& InSocketTag) const;
+
 	/**
 	 * Returns vertex color data by position.
 	 * For matching to reimported meshes that may have changed or copying vertex paint data from mesh to mesh.
diff --git a/Engine/Source/Runtime/Engine/Private/StaticMesh.cpp b/Engine/Source/Runtime/Engine/Private/StaticMesh.cpp
--- a/Engine/Source/Runtime/Engine/Private/StaticMesh.cpp
+++ b/Engine/Source/Runtime/Engine/Private/StaticMesh.cpp
@@ -7816,6 +7816,22 @@
 	Sockets.Remove(Socket);
 }
 
+TArray<UStaticMeshSocket*> UStaticMesh::GetSocketsByTag(const FString& InSocketTag) const
+{
+	TArray<UStaticMeshSocket*> FoundSockets;
+
+	for (int32 i = 0; i < Sockets.Num(); ++i)
+	{
+		UStaticMeshSocket* Socket = Sockets[i];
+		if (Socket && Socket->Tag == InSocketTag)
+		{
+			FoundSockets.Add(Socket);
+		}
+	}
+
+	return FoundSockets;
+}
+
 ENGINE_API int32 UStaticMesh::GetDefaultMinLOD() const
 {
 	if (IsMinLodQualityLevelEnable())
