Fix for Zen reference in non-editor builds

diff --git a/Engine/Source/Runtime/Engine/Engine.Build.cs b/Engine/Source/Runtime/Engine/Engine.Build.cs
--- a/Engine/Source/Runtime/Engine/Engine.Build.cs
+++ b/Engine/Source/Runtime/Engine/Engine.Build.cs
@@ -107,8 +107,7 @@ public class Engine : ModuleRules
 				"AudioExtensions",
 				"DeveloperSettings",
 				"AudioLinkCore",
-				"CookOnTheFly",
-				"Zen"
+				"CookOnTheFly"
 			}
 		);
 
@@ -167,6 +166,7 @@ public class Engine : ModuleRules
 			PrivateDependencyModuleNames.AddRange(
 				new string[] {
 					"RawMesh",
+					"Zen"
 				}
 			);
 		}
diff --git a/Engine/Source/Runtime/Engine/Private/Analytics/StudioAnalytics.cpp b/Engine/Source/Runtime/Engine/Private/Analytics/StudioAnalytics.cpp
--- a/Engine/Source/Runtime/Engine/Private/Analytics/StudioAnalytics.cpp
+++ b/Engine/Source/Runtime/Engine/Private/Analytics/StudioAnalytics.cpp
@@ -20,9 +20,12 @@
 #include "HAL/PlatformProcess.h"
 #include "DerivedDataCacheInterface.h"
 #include "DerivedDataCacheUsageStats.h"
-#include "ZenServerInterface.h"
 #include "Virtualization/VirtualizationSystem.h"
 
+#if WITH_EDITOR
+#include "ZenServerInterface.h"
+#endif
+
 bool FStudioAnalytics::bInitialized = false;
 volatile double FStudioAnalytics::TimeEstimation = 0;
 FThread FStudioAnalytics::TimerThread;
