Store off the initial value of bEvaluateWorldPositionOffset at BeginPlay() so that it's known in case we want to reset it during play.

diff --git a/Engine/Source/Runtime/Engine/Classes/Components/StaticMeshComponent.h b/Engine/Source/Runtime/Engine/Classes/Components/StaticMeshComponent.h
--- a/Engine/Source/Runtime/Engine/Classes/Components/StaticMeshComponent.h
+++ b/Engine/Source/Runtime/Engine/Classes/Components/StaticMeshComponent.h
@@ -244,6 +244,12 @@ public:
 	UPROPERTY(EditAnywhere, AdvancedDisplay, BlueprintReadWrite, Category = RayTracing)
 	uint8 bEvaluateWorldPositionOffsetInRayTracing : 1;
 
+protected:
+	/** Initial value of bEvaluateWorldPositionOffset when BeginPlay() was called. Can be useful if we want to reset to initial state. */
+	uint8 bInitialEvaluateWorldPositionOffset : 1;
+
+public:
+
 #if WITH_EDITORONLY_DATA
 	/** The section currently selected in the Editor. Used for highlighting */
 	UPROPERTY(transient)
@@ -441,6 +447,10 @@ public:
 	UFUNCTION(BlueprintCallable, Category = "Rendering|LOD")
 	void SetEvaluateWorldPositionOffset(bool NewValue);
 
+	/** Get the initial value of bEvaluateWorldPositionOffset. This is the value when BeginPlay() was last called. */
+	UFUNCTION(BlueprintCallable, Category = "Rendering|LOD")
+	bool GetInitialEvaluateWorldPositionOffset() { return bInitialEvaluateWorldPositionOffset; }
+
 	/** 
 	 * Get Local bounds
 	 */
@@ -517,6 +527,7 @@ public:
 protected: 
 	virtual void OnRegister() override;
 	virtual void OnUnregister() override;
+	virtual void BeginPlay() override;
 	virtual bool RequiresGameThreadEndOfFrameRecreate() const override;
 	virtual void CreateRenderState_Concurrent(FRegisterComponentContext* Context) override;
 	virtual void OnCreatePhysicsState() override;
diff --git a/Engine/Source/Runtime/Engine/Private/Components/StaticMeshComponent.cpp b/Engine/Source/Runtime/Engine/Private/Components/StaticMeshComponent.cpp
--- a/Engine/Source/Runtime/Engine/Private/Components/StaticMeshComponent.cpp
+++ b/Engine/Source/Runtime/Engine/Private/Components/StaticMeshComponent.cpp
@@ -203,6 +203,7 @@ UStaticMeshComponent::UStaticMeshComponent(const FObjectInitializer& ObjectIniti
 	bDisallowNanite = false;
 	bEvaluateWorldPositionOffset = false;
 	bEvaluateWorldPositionOffsetInRayTracing = false;
+	bInitialEvaluateWorldPositionOffset = false;
 	DistanceFieldIndirectShadowMinVisibility = .1f;
 	GetBodyInstance()->bAutoWeld = true;	//static mesh by default has auto welding
 
@@ -737,6 +738,13 @@ void UStaticMeshComponent::OnUnregister()
 	Super::OnUnregister();
 }
 
+void UStaticMeshComponent::BeginPlay()
+{
+	bInitialEvaluateWorldPositionOffset = bEvaluateWorldPositionOffset;
+
+	Super::BeginPlay();
+}
+
 bool UStaticMeshComponent::RequiresGameThreadEndOfFrameRecreate() const
 {
 #if STATICMESH_ENABLE_DEBUG_RENDERING
