Fix crash when building light in a new level with a geometry asset and there is no light at all.

diff --git a/Engine/Source/Runtime/Engine/Private/ModelLight.cpp b/Engine/Source/Runtime/Engine/Private/ModelLight.cpp
--- a/Engine/Source/Runtime/Engine/Private/ModelLight.cpp
+++ b/Engine/Source/Runtime/Engine/Private/ModelLight.cpp
@@ -941,8 +941,9 @@ void UModel::GroupAllNodes(ULevel* Level, const TArray<ULightComponentBase*>& Li
 
 									// add the relevant lights to the nodegroup
 									TArray<ULightComponent*>* RelevantLights = ComponentRelevantLights.Find(Nodes[NodeIndex].ComponentIndex);
-									check(RelevantLights);
-									for (int32 LightIndex = 0; LightIndex < RelevantLights->Num(); LightIndex++)
+									
+									// add null check in case there is no light at all and the user presses build lighting.
+									for (int32 LightIndex = 0; RelevantLights && LightIndex < RelevantLights->Num(); LightIndex++)
 									{
 										NodeGroup->RelevantLights.AddUnique((*RelevantLights)[LightIndex]);
 									}
@@ -973,7 +974,8 @@ void UModel::GroupAllNodes(ULevel* Level, const TArray<ULightComponentBase*>& Li
 
 			// add the relevant lights to the nodegroup
 			TArray<ULightComponent*>* RelevantLights = ComponentRelevantLights.Find(Nodes[NodeIndex].ComponentIndex);
-			check(RelevantLights);
+
+			// add null check in case there is no light at all and the user presses build lighting.
 			for (int32 LightIndex = 0; RelevantLights && LightIndex < RelevantLights->Num(); LightIndex++)
 			{
 				NodeGroup->RelevantLights.AddUnique((*RelevantLights)[LightIndex]);
