Add helper methods in the engine to get the local player from a FInputDeviceId.

diff --git a/Engine/Source/Runtime/Engine/Classes/Engine/Engine.h b/Engine/Source/Runtime/Engine/Classes/Engine/Engine.h
--- a/Engine/Source/Runtime/Engine/Classes/Engine/Engine.h
+++ b/Engine/Source/Runtime/Engine/Classes/Engine/Engine.h
@@ -2659,6 +2659,9 @@ public:
 	ULocalPlayer* GetLocalPlayerFromControllerId( const UGameViewportClient* InViewport, const int32 ControllerId ) const;
 	ULocalPlayer* GetLocalPlayerFromControllerId( UWorld * InWorld, const int32 ControllerId ) const;
 
+	ULocalPlayer* GetLocalPlayerFromInputDevice(const UGameViewportClient* InViewport, const FInputDeviceId InputDevice) const;
+	ULocalPlayer* GetLocalPlayerFromInputDevice(UWorld * InWorld, const FInputDeviceId InputDevice) const;
+
 	void SwapControllerId(ULocalPlayer *NewPlayer, const int32 CurrentControllerId, const int32 NewControllerID) const;
 
 	/** 
diff --git a/Engine/Source/Runtime/Engine/Private/UnrealEngine.cpp b/Engine/Source/Runtime/Engine/Private/UnrealEngine.cpp
--- a/Engine/Source/Runtime/Engine/Private/UnrealEngine.cpp
+++ b/Engine/Source/Runtime/Engine/Private/UnrealEngine.cpp
@@ -3885,6 +3885,41 @@ ULocalPlayer* UEngine::GetLocalPlayerFromControllerId( UWorld * InWorld, const i
 	return GetLocalPlayerFromControllerId_local(GamePlayers, ControllerId);
 }
 
+namespace UE::Private
+{
+	static ULocalPlayer* GetLocalPlayerFromInputDevice(const TArray<ULocalPlayer*>& GamePlayers, const FInputDeviceId InputDevice)
+	{
+		IPlatformInputDeviceMapper& DeviceMapper = IPlatformInputDeviceMapper::Get();
+		const FPlatformUserId PlatUser = DeviceMapper.GetUserForInputDevice(InputDevice);
+
+		for (ULocalPlayer* const Player : GamePlayers)
+		{
+			if (Player && Player->GetPlatformUserId() == PlatUser)
+			{
+				return Player;
+			}
+		}
+	
+		return nullptr;
+	}
+}	// namespace UE::Private
+
+ULocalPlayer* UEngine::GetLocalPlayerFromInputDevice(const UGameViewportClient* InViewport, const FInputDeviceId InputDevice) const
+{
+	if (GetWorldContextFromGameViewport(InViewport) != nullptr)
+	{
+		const TArray<class ULocalPlayer*>& GamePlayers = GetGamePlayers(InViewport);
+		return UE::Private::GetLocalPlayerFromInputDevice(GamePlayers, InputDevice);
+	}
+	return nullptr;
+}
+
+ULocalPlayer* UEngine::GetLocalPlayerFromInputDevice(UWorld * InWorld, const FInputDeviceId InputDevice) const
+{
+	const TArray<class ULocalPlayer*>& GamePlayers = GetGamePlayers(InWorld);
+	return UE::Private::GetLocalPlayerFromInputDevice(GamePlayers, InputDevice);
+}
+
 void UEngine::SwapControllerId(ULocalPlayer *NewPlayer, const int32 CurrentControllerId, const int32 NewControllerID) const
 {
 	for (auto It = WorldList.CreateConstIterator(); It; ++It)
