Fix crashes when undoing materials in the material editor.

The transaction system was not aware of the EditorOnly data owned by the UMaterial.  This meant that any undo/redo operation left the editor only data on the UMaterial in an incorrect state, for instance, UMaterialExpressions with properties missing.

- Override virtual function Modify on UMaterial.
- Call Modify on the EditorOnly data owned by the UMaterial.

diff --git a/Engine/Source/Runtime/Engine/Classes/Materials/Material.h b/Engine/Source/Runtime/Engine/Classes/Materials/Material.h
--- a/Engine/Source/Runtime/Engine/Classes/Materials/Material.h
+++ b/Engine/Source/Runtime/Engine/Classes/Materials/Material.h
@@ -1100,6 +1100,7 @@ public:
 	ENGINE_API virtual void PreEditChange(FProperty* PropertyAboutToChange) override;
 	ENGINE_API virtual void PostEditChangeProperty(FPropertyChangedEvent& PropertyChangedEvent) override;
 	ENGINE_API virtual bool CanEditChange(const FProperty* InProperty) const override;
+	ENGINE_API virtual bool Modify(bool bAlwaysMarkDirty = true) override;
 #endif // WITH_EDITOR
 	ENGINE_API virtual void BeginDestroy() override;
 	ENGINE_API virtual bool IsReadyForFinishDestroy() override;
diff --git a/Engine/Source/Runtime/Engine/Private/Materials/Material.cpp b/Engine/Source/Runtime/Engine/Private/Materials/Material.cpp
--- a/Engine/Source/Runtime/Engine/Private/Materials/Material.cpp
+++ b/Engine/Source/Runtime/Engine/Private/Materials/Material.cpp
@@ -4007,6 +4007,16 @@ bool UMaterial::CanEditChange(const FProperty* InProperty) const
 	return true;
 }
 
+bool UMaterial::Modify(bool bAlwaysMarkDirty)
+{
+	UMaterialEditorOnlyData* EditorOnly = GetEditorOnlyData();
+	if (EditorOnly)
+	{
+		EditorOnly->Modify(bAlwaysMarkDirty);
+	}
+	return Super::Modify(bAlwaysMarkDirty);
+}
+
 void UMaterial::CreateExecutionFlowExpressions()
 {
 	if (IsUsingControlFlow())
