Change r.UseLegacyMaintainYFOVViewMatrix so that it can be set in config files

diff --git a/Engine/Source/Runtime/Engine/Private/Camera/CameraStackTypes.cpp b/Engine/Source/Runtime/Engine/Private/Camera/CameraStackTypes.cpp
--- a/Engine/Source/Runtime/Engine/Private/Camera/CameraStackTypes.cpp
+++ b/Engine/Source/Runtime/Engine/Private/Camera/CameraStackTypes.cpp
@@ -8,7 +8,7 @@ static TAutoConsoleVariable<bool> CVarUseLegacyMaintainYFOV(
 	TEXT("r.UseLegacyMaintainYFOVViewMatrix"),
 	false,
 	TEXT("Whether to use the old way to compute perspective view matrices when the aspect ratio constraint is vertical"),
-	ECVF_Cheat
+	ECVF_Default
 );
 
 //////////////////////////////////////////////////////////////////////////
