Fixing non-unity build

diff --git a/Engine/Source/Editor/Kismet/Public/BlueprintAssetHandler.h b/Engine/Source/Editor/Kismet/Public/BlueprintAssetHandler.h
--- a/Engine/Source/Editor/Kismet/Public/BlueprintAssetHandler.h
+++ b/Engine/Source/Editor/Kismet/Public/BlueprintAssetHandler.h
@@ -7,6 +7,7 @@
 #include "Containers/ArrayView.h"
 #include "Containers/Array.h"
 #include "UObject/NameTypes.h"
+#include "UObject/TopLevelAssetPath.h"
 
 struct FAssetData;
 
diff --git a/Engine/Source/Editor/UnrealEd/Public/LocalizedAssetUtil.h b/Engine/Source/Editor/UnrealEd/Public/LocalizedAssetUtil.h
--- a/Engine/Source/Editor/UnrealEd/Public/LocalizedAssetUtil.h
+++ b/Engine/Source/Editor/UnrealEd/Public/LocalizedAssetUtil.h
@@ -7,6 +7,7 @@
 struct FAssetData;
 class IAssetRegistry;
 class FLocalizationSCC;
+struct FTopLevelAssetPath;
 
 struct FLocalizedAssetSCCUtil
 {
diff --git a/Engine/Source/Runtime/Engine/Private/DataTable.cpp b/Engine/Source/Runtime/Engine/Private/DataTable.cpp
--- a/Engine/Source/Runtime/Engine/Private/DataTable.cpp
+++ b/Engine/Source/Runtime/Engine/Private/DataTable.cpp
@@ -1,6 +1,7 @@
 // Copyright Epic Games, Inc. All Rights Reserved.
 
 #include "Engine/DataTable.h"
+#include "AssetRegistry/AssetData.h"
 #include "Internationalization/TextPackageNamespaceUtil.h"
 #include "Internationalization/StabilizeLocalizationKeys.h"
 #include "Serialization/PropertyLocalizationDataGathering.h"
