[Replication]
Remove some replay specific changes that would prevent using replay to test feature and performance properly

#RNX

diff --git a/Engine/Source/Runtime/Engine/Private/Character.cpp b/Engine/Source/Runtime/Engine/Private/Character.cpp
--- a/Engine/Source/Runtime/Engine/Private/Character.cpp
+++ b/Engine/Source/Runtime/Engine/Private/Character.cpp
@@ -1164,12 +1164,6 @@ void ACharacter::OnRep_ReplicatedBasedMovement()
 
 void ACharacter::OnRep_ReplicatedMovement()
 {
-	if (GetWorld()->IsPlayingReplay())
-	{
-		Super::OnRep_ReplicatedMovement();
-		return;
-	}
-
 	// Skip standard position correction if we are playing root motion, OnRep_RootMotion will handle it.
 	if (!IsPlayingNetworkedRootMotionMontage()) // animation root motion
 	{
diff --git a/Engine/Source/Runtime/Engine/Private/Components/CharacterMovementComponent.cpp b/Engine/Source/Runtime/Engine/Private/Components/CharacterMovementComponent.cpp
--- a/Engine/Source/Runtime/Engine/Private/Components/CharacterMovementComponent.cpp
+++ b/Engine/Source/Runtime/Engine/Private/Components/CharacterMovementComponent.cpp
@@ -1591,43 +1591,6 @@ void UCharacterMovementComponent::SimulatedTick(float DeltaSeconds)
 	SCOPE_CYCLE_COUNTER(STAT_CharacterMovementSimulated);
 	checkSlow(CharacterOwner != nullptr);
 
-	if (GetWorld()->IsPlayingReplay())
-	{
-		const FVector OldLocation = UpdatedComponent ? UpdatedComponent->GetComponentLocation() : FVector::ZeroVector;
-		const FVector OldVelocity = Velocity;
-
-		// Interpolate between appropriate samples
-		{
-			SCOPE_CYCLE_COUNTER( STAT_CharacterMovementSmoothClientPosition );
-			SmoothClientPosition( DeltaSeconds );
-		}
-
-		// Update replicated movement mode
-		ApplyNetworkMovementMode( GetCharacterOwner()->GetReplicatedMovementMode() );
-
-		UpdateComponentVelocity();
-		bJustTeleported = false;
-
-		if (CharacterOwner)
-		{
-			CharacterOwner->RootMotionRepMoves.Empty();
-			CurrentRootMotion.Clear();
-			CharacterOwner->SavedRootMotion.Clear();
-		}
-
-		// Note: we do not call the Super implementation, that runs prediction.
-		// We do still need to call these though
-		OnMovementUpdated(DeltaSeconds, OldLocation, OldVelocity);
-		CallMovementUpdateDelegate(DeltaSeconds, OldLocation, OldVelocity);
-
-		LastUpdateLocation = UpdatedComponent ? UpdatedComponent->GetComponentLocation() : FVector::ZeroVector;
-		LastUpdateRotation = UpdatedComponent ? UpdatedComponent->GetComponentQuat() : FQuat::Identity;
-		LastUpdateVelocity = Velocity;
-
-		//TickCharacterPose( DeltaSeconds );
-		return;
-	}
-
 	// If we are playing a RootMotion AnimMontage.
 	if (CharacterOwner->IsPlayingNetworkedRootMotionMontage())
 	{
