Fixed FWorldPartitionHelpers::ForEachActorWithLoading which could trigger checks because it could push multiple World Partition Loading Context.

diff --git a/Engine/Source/Runtime/Engine/Private/WorldPartition/WorldPartitionHelpers.cpp b/Engine/Source/Runtime/Engine/Private/WorldPartition/WorldPartitionHelpers.cpp
--- a/Engine/Source/Runtime/Engine/Private/WorldPartition/WorldPartitionHelpers.cpp
+++ b/Engine/Source/Runtime/Engine/Private/WorldPartition/WorldPartitionHelpers.cpp
@@ -62,10 +62,8 @@ void FWorldPartitionHelpers::ForEachActorDesc(UWorldPartition* WorldPartition, T
 
 namespace WorldPartitionHelpers
 {
-	void LoadReferences(UWorldPartition* WorldPartition, const FGuid& ActorGuid, TMap<FGuid, FWorldPartitionReference>& InOutActorReferences)
+	void LoadReferencesInternal(UWorldPartition* WorldPartition, const FGuid& ActorGuid, TMap<FGuid, FWorldPartitionReference>& InOutActorReferences)
 	{
-		FWorldPartitionLoadingContext::FDeferred LoadingContext;
-
 		if (InOutActorReferences.Contains(ActorGuid))
 		{
 			return;
@@ -77,12 +75,18 @@ namespace WorldPartitionHelpers
 
 			for (FGuid ReferenceGuid : ActorDesc->GetReferences())
 			{
-				LoadReferences(WorldPartition, ReferenceGuid, InOutActorReferences);
+				LoadReferencesInternal(WorldPartition, ReferenceGuid, InOutActorReferences);
 			}
 
 			InOutActorReferences[ActorGuid] = FWorldPartitionReference(WorldPartition, ActorGuid);
 		}
 	}
+
+	void LoadReferences(UWorldPartition* WorldPartition, const FGuid& ActorGuid, TMap<FGuid, FWorldPartitionReference>& InOutActorReferences)
+	{
+		FWorldPartitionLoadingContext::FDeferred LoadingContext;
+		LoadReferencesInternal(WorldPartition, ActorGuid, InOutActorReferences);
+	}
 }
 
 FWorldPartitionHelpers::FForEachActorWithLoadingParams::FForEachActorWithLoadingParams()
