Added debug useafterfree command to trigger an Address sanitizer crash

#rnx

diff --git a/Engine/Source/Runtime/Engine/Private/UnrealEngine.cpp b/Engine/Source/Runtime/Engine/Private/UnrealEngine.cpp
--- a/Engine/Source/Runtime/Engine/Private/UnrealEngine.cpp
+++ b/Engine/Source/Runtime/Engine/Private/UnrealEngine.cpp
@@ -10268,6 +10268,18 @@ bool UEngine::PerformError(const TCHAR* Cmd, FOutputDevice& Ar)
 		FGenericCrashContext::SetCrashTrigger(ECrashTrigger::Debug);
 		UE_LOG(LogEngine, Fatal, TEXT("Crashing the worker thread at your request"));
 	}
+#if USING_ADDRESS_SANITISER
+	else if (FParse::Command(&Cmd, TEXT("USEAFTERFREE")))
+	{
+		static constexpr size_t UseAfterFreeArraySize = 32;
+		uint8* UseAfterFreeArray = reinterpret_cast<uint8*>(FMemory::Malloc(UseAfterFreeArraySize));
+		FMemory::Memzero(UseAfterFreeArray, UseAfterFreeArraySize);
+		FMemory::Free(UseAfterFreeArray);
+		
+		// Reuse deleted array
+		FMemory::Memzero(UseAfterFreeArray, UseAfterFreeArraySize);
+	}
+#endif
 #endif // !UE_BUILD_SHIPPING
 	return false;
 }
