Fix for HDR toggling on some platforms.

This does not truly change HDR since these platforms do not support toggling this at runtime. Instead we change the swapchain to be in SDR, and we end up in SDR mode with HDR output mode.

SDR in HDR mode is slightly darker on the TV, more than when comparing the screenshots of the swapchain. This is most likely due to the TV response, which needs to be investigated later on to see if we can somehow match the 2 properly.

In addition to the 2 CLs from main, this CL allows BP to change the CVar r.HDR.EnableHDROutput . On these platforms, the CVar is set with ECVF_SetByCode as soon as HDR was supported and calling EnableHDRDisplayOutput with ECVF_SetByGameSetting would get ignored

Original CLs:
19334984
fix for r.setres on d3d12: the backbuffers were still referenced by the present queue, preventing their proper destruction

19609874
support for HDR change on process resume

diff --git a/Engine/Source/Runtime/Engine/Classes/GameFramework/GameUserSettings.h b/Engine/Source/Runtime/Engine/Classes/GameFramework/GameUserSettings.h
--- a/Engine/Source/Runtime/Engine/Classes/GameFramework/GameUserSettings.h
+++ b/Engine/Source/Runtime/Engine/Classes/GameFramework/GameUserSettings.h
@@ -562,6 +562,8 @@ protected:
 
 	void UpdateResolutionQuality();
 
+	void EnableHDRDisplayOutputInternal(bool bEnable, int32 DisplayNits, bool FromUserSettings);
+
 private:
 
 	UPROPERTY(BlueprintAssignable, meta = (AllowPrivateAccess = "true"))
diff --git a/Engine/Source/Runtime/Engine/Private/GameUserSettings.cpp b/Engine/Source/Runtime/Engine/Private/GameUserSettings.cpp
--- a/Engine/Source/Runtime/Engine/Private/GameUserSettings.cpp
+++ b/Engine/Source/Runtime/Engine/Private/GameUserSettings.cpp
@@ -477,7 +477,7 @@ void UGameUserSettings::ApplyNonResolutionSettings()
 
 	bool bEnableHDR = (IsHDRAllowed() && bUseHDRDisplayOutput && !bWithEditor);
 
-	EnableHDRDisplayOutput(bEnableHDR, HDRDisplayOutputNits);
+	EnableHDRDisplayOutputInternal(bEnableHDR, HDRDisplayOutputNits, true);
 
 }
 
@@ -938,6 +938,11 @@ bool UGameUserSettings::SupportsHDRDisplayOutput() const
 }
 
 void UGameUserSettings::EnableHDRDisplayOutput(bool bEnable, int32 DisplayNits /*= 1000*/)
+{
+	EnableHDRDisplayOutputInternal(bEnable, DisplayNits, false);
+}
+
+void UGameUserSettings::EnableHDRDisplayOutputInternal(bool bEnable, int32 DisplayNits, bool FromUserSettings)
 {
 	static IConsoleVariable* CVarHDROutputEnabled = IConsoleManager::Get().FindConsoleVariable(TEXT("r.HDR.EnableHDROutput"));
 
@@ -965,13 +970,13 @@ void UGameUserSettings::EnableHDRDisplayOutput(bool bEnable, int32 DisplayNits /
 				RequestUIUpdate();
 			}
 #endif
-			CVarHDROutputEnabled->Set(1, ECVF_SetByGameSetting);
+			CVarHDROutputEnabled->Set(1, FromUserSettings ? ECVF_SetByGameSetting : ECVF_SetByCode);
 		}
 
 		// Always test this branch as can be used to flush errors
 		if (!bEnable)
 		{
-			CVarHDROutputEnabled->Set(0, ECVF_SetByGameSetting);
+			CVarHDROutputEnabled->Set(0, FromUserSettings ? ECVF_SetByGameSetting : ECVF_SetByCode);
 		}
 
 		// Update final requested state for saved config
