Added support for delaying network replication until client has verified it has all the correct content.

diff --git a/Engine/Source/Runtime/Engine/Classes/Engine/DemoPendingNetGame.h b/Engine/Source/Runtime/Engine/Classes/Engine/DemoPendingNetGame.h
--- a/Engine/Source/Runtime/Engine/Classes/Engine/DemoPendingNetGame.h
+++ b/Engine/Source/Runtime/Engine/Classes/Engine/DemoPendingNetGame.h
@@ -19,5 +19,5 @@ class UDemoPendingNetGame
 	// UPendingNetGame interface.
 	virtual void Tick( float DeltaTime ) override;
 	virtual void SendJoin() override;
-	virtual void LoadMapCompleted(UEngine* Engine, FWorldContext& Context, bool bLoadedMapSuccessfully, const FString& LoadMapError) override;
+	virtual bool LoadMapCompleted(UEngine* Engine, FWorldContext& Context, bool bLoadedMapSuccessfully, const FString& LoadMapError) override;
 };
diff --git a/Engine/Source/Runtime/Engine/Classes/Engine/GameInstance.h b/Engine/Source/Runtime/Engine/Classes/Engine/GameInstance.h
--- a/Engine/Source/Runtime/Engine/Classes/Engine/GameInstance.h
+++ b/Engine/Source/Runtime/Engine/Classes/Engine/GameInstance.h
@@ -533,6 +533,14 @@ public:
 	/** Return the game mode subclass to use for a given map, options, and portal. By default return passed in one */
 	virtual TSubclassOf<AGameModeBase> OverrideGameModeClass(TSubclassOf<AGameModeBase> GameModeClass, const FString& MapName, const FString& Options, const FString& Portal) const;
 
+
+	/**
+	 * Game instance has an opportunity to modify the level name before the client starts travel
+	 */
+	virtual void ModifyClientTravelLevelURL(FString& LevelName)
+	{
+	}
+
 	/** return true to delay an otherwise ready-to-join PendingNetGame performing LoadMap() and finishing up
 	 * useful to wait for content downloads, etc
 	 */
@@ -541,6 +549,14 @@ public:
 		return false;
 	}
 
+	/**
+	 * return true to delay player controller spawn (sending NMT_Join)
+	 */
+	virtual bool DelayCompletionOfPendingNetGameTravel()
+	{
+		return false;
+	}
+
 	FTimerManager* TimerManager;
 	FLatentActionManager* LatentActionManager;
 
diff --git a/Engine/Source/Runtime/Engine/Classes/Engine/PendingNetGame.h b/Engine/Source/Runtime/Engine/Classes/Engine/PendingNetGame.h
--- a/Engine/Source/Runtime/Engine/Classes/Engine/PendingNetGame.h
+++ b/Engine/Source/Runtime/Engine/Classes/Engine/PendingNetGame.h
@@ -73,6 +73,9 @@ public:
 	/** @todo document */
 	bool					bSentJoinRequest;
 
+	/** set when we call LoadMapCompleted */
+	bool					bLoadedMapSuccessfully;
+
 	/** @todo document */
 	FString					ConnectionError;
 
@@ -123,5 +126,8 @@ public:
 	void InitPeerListen();
 
 	/** Called by the engine after it calls LoadMap for this PendingNetGame. */
-	virtual void LoadMapCompleted(UEngine* Engine, FWorldContext& Context, bool bLoadedMapSuccessfully, const FString& LoadMapError);
+	virtual bool LoadMapCompleted(UEngine* Engine, FWorldContext& Context, bool bLoadedMapSuccessfully, const FString& LoadMapError);
+
+	/** Called by the engine after loadmapCompleted and the GameInstance has finished delaying */
+	virtual void TravelCompleted(UEngine* Engine, FWorldContext& Context);
 };
diff --git a/Engine/Source/Runtime/Engine/Private/DemoNetDriver.cpp b/Engine/Source/Runtime/Engine/Private/DemoNetDriver.cpp
--- a/Engine/Source/Runtime/Engine/Private/DemoNetDriver.cpp
+++ b/Engine/Source/Runtime/Engine/Private/DemoNetDriver.cpp
@@ -5378,8 +5378,9 @@
 	// Don't send a join request to a replay
 }
 
-void UDemoPendingNetGame::LoadMapCompleted(UEngine* Engine, FWorldContext& Context, bool bLoadedMapSuccessfully, const FString& LoadMapError)
+bool UDemoPendingNetGame::LoadMapCompleted(UEngine* Engine, FWorldContext& Context, bool bInLoadedMapSuccessfully, const FString& LoadMapError)
 {
+	bLoadedMapSuccessfully = bInLoadedMapSuccessfully;
 	UDemoNetDriver* TheDriver = GetDemoNetDriver();
 
 	// If we have a demo pending net game we should have a demo net driver
@@ -5406,8 +5407,10 @@
 		{
 			Context.OwningGameInstance->HandleDemoPlaybackFailure(EDemoPlayFailure::LoadMap, FString(TEXT("LoadMap failed")));
 		}
-		return;
+		// we already handled the error so don't need the TickPendingNetGame to handle it 
+		return true;
 	}
+	return true;
 }
 
 void UDemoNetDriver::Serialize(FArchive& Ar)
diff --git a/Engine/Source/Runtime/Engine/Private/PendingNetGame.cpp b/Engine/Source/Runtime/Engine/Private/PendingNetGame.cpp
--- a/Engine/Source/Runtime/Engine/Private/PendingNetGame.cpp
+++ b/Engine/Source/Runtime/Engine/Private/PendingNetGame.cpp
@@ -22,6 +22,7 @@ void UPendingNetGame::Initialize(const FURL& InURL)
 	URL = InURL;
 	bSuccessfullyConnected = false;
 	bSentJoinRequest = false;
+	bLoadedMapSuccessfully = false;
 }
 
 UPendingNetGame::UPendingNetGame(const FObjectInitializer& ObjectInitializer)
@@ -138,26 +139,27 @@ void UPendingNetGame::AddReferencedObjects(UObject* InThis, FReferenceCollector&
 	Super::AddReferencedObjects( This, Collector );
 }
 
-void UPendingNetGame::LoadMapCompleted(UEngine* Engine, FWorldContext& Context, bool bLoadedMapSuccessfully, const FString& LoadMapError)
+bool UPendingNetGame::LoadMapCompleted(UEngine* Engine, FWorldContext& Context, bool bInLoadedMapSuccessfully, const FString& LoadMapError)
 {
+	bLoadedMapSuccessfully = bInLoadedMapSuccessfully;
 	if (!bLoadedMapSuccessfully || LoadMapError != TEXT(""))
 	{
-		// we can't guarantee the current World is in a valid state, so travel to the default map
-		Engine->BrowseToDefaultMap(Context);
-		Engine->BroadcastTravelFailure(Context.World(), ETravelFailure::LoadMapFailure, LoadMapError);
-		check(Context.World() != NULL);
+		// this is handled in the TickWorldTravel
+		return false;
 	}
-	else
-	{
-		// Show connecting message, cause precaching to occur.
-		Engine->TransitionType = ETransitionType::Connecting;
+	return true;
+}
 
-		Engine->RedrawViewports(false);
+void UPendingNetGame::TravelCompleted(UEngine* Engine, FWorldContext& Context)
+{
+	// Show connecting message, cause precaching to occur.
+	Engine->TransitionType = ETransitionType::Connecting;
 
-		// Send join.
-		Context.PendingNetGame->SendJoin();
-		Context.PendingNetGame->NetDriver = NULL;
-	}
+	Engine->RedrawViewports(false);
+
+	// Send join.
+	Context.PendingNetGame->SendJoin();
+	Context.PendingNetGame->NetDriver = NULL;
 }
 
 EAcceptConnection::Type UPendingNetGame::NotifyAcceptingConnection()
diff --git a/Engine/Source/Runtime/Engine/Private/UnrealEngine.cpp b/Engine/Source/Runtime/Engine/Private/UnrealEngine.cpp
--- a/Engine/Source/Runtime/Engine/Private/UnrealEngine.cpp
+++ b/Engine/Source/Runtime/Engine/Private/UnrealEngine.cpp
@@ -13976,14 +13976,14 @@ void UEngine::TickWorldTravel(FWorldContext& Context, float DeltaSeconds)
 			BroadcastNetworkFailure(NULL, Context.PendingNetGame->NetDriver, ENetworkFailure::PendingConnectionFailure, Context.PendingNetGame->ConnectionError);
 			CancelPending(Context);
 		}
-		else if (Context.PendingNetGame && Context.PendingNetGame->bSuccessfullyConnected && !Context.PendingNetGame->bSentJoinRequest && (Context.OwningGameInstance == NULL || !Context.OwningGameInstance->DelayPendingNetGameTravel()))
+		else if (Context.PendingNetGame && Context.PendingNetGame->bSuccessfullyConnected && !Context.PendingNetGame->bSentJoinRequest && !Context.PendingNetGame->bLoadedMapSuccessfully && (Context.OwningGameInstance == NULL || !Context.OwningGameInstance->DelayPendingNetGameTravel()))
 		{
 			if (!MakeSureMapNameIsValid(Context.PendingNetGame->URL.Map))
 			{
 				BrowseToDefaultMap(Context);
 				BroadcastTravelFailure(Context.World(), ETravelFailure::PackageMissing, Context.PendingNetGame->URL.Map);
 			}
-			else
+			else if (!Context.PendingNetGame->bLoadedMapSuccessfully)
 			{
 				// Attempt to load the map.
 				FString Error;
@@ -13992,10 +13992,11 @@ void UEngine::TickWorldTravel(FWorldContext& Context, float DeltaSeconds)
 
 				if (Context.PendingNetGame != nullptr)
 				{
-					Context.PendingNetGame->LoadMapCompleted(this, Context, bLoadedMapSuccessfully, Error);
-
-					// Kill the pending level.
-					Context.PendingNetGame = nullptr;
+					if (!Context.PendingNetGame->LoadMapCompleted(this, Context, bLoadedMapSuccessfully, Error))
+					{
+						BrowseToDefaultMap(Context);
+						BroadcastTravelFailure(Context.World(), ETravelFailure::TravelFailure, Error);
+					}
 				}
 				else
 				{
@@ -14004,6 +14005,12 @@ void UEngine::TickWorldTravel(FWorldContext& Context, float DeltaSeconds)
 				}
 			}
 		}
+		
+		if (Context.PendingNetGame && Context.PendingNetGame->bLoadedMapSuccessfully && (Context.OwningGameInstance == NULL || !Context.OwningGameInstance->DelayCompletionOfPendingNetGameTravel()))
+		{
+			Context.PendingNetGame->TravelCompleted(this, Context);
+			Context.PendingNetGame = nullptr;
+		}
 	}
 	else if (TransitionType == ETransitionType::WaitingToConnect)
 	{
