Fix issues with missing referenced texture data from materials.

- Add function, UpdateTransientExpressionData, to update transient properties on material expressions.  This is called in the cooked editor on cooked packages instead of UpdateCachedExpressionData().  We don't want to destroy the cached expression data in this case (which contains the referenced texture data).
- UpdateTransientExpressionData is responsible for taking updating transient properties on the UMaterialExpressions in a material.  Currently this only fixes up pointers to material function inputs and outputs.
- Make the material code more resilient to not having optional editor only data with null checks.
- Only validate if we have material layers in the editor only data (it can be missing).

diff --git a/Engine/Source/Runtime/Engine/Classes/Materials/Material.h b/Engine/Source/Runtime/Engine/Classes/Materials/Material.h
--- a/Engine/Source/Runtime/Engine/Classes/Materials/Material.h
+++ b/Engine/Source/Runtime/Engine/Classes/Materials/Material.h
@@ -1438,6 +1438,8 @@ public:
 
 		return false;
 	}
+
+	ENGINE_API void UpdateTransientExpressionData();
 #endif // WITH_EDITORONLY_DATA
 
 #if WITH_EDITOR
diff --git a/Engine/Source/Runtime/Engine/Private/Materials/Material.cpp b/Engine/Source/Runtime/Engine/Private/Materials/Material.cpp
--- a/Engine/Source/Runtime/Engine/Private/Materials/Material.cpp
+++ b/Engine/Source/Runtime/Engine/Private/Materials/Material.cpp
@@ -1903,12 +1903,29 @@ void UMaterial::GetMaterialInheritanceChain(FMaterialInheritanceChain& OutChain)
 	}
 }
 
+#if WITH_EDITORONLY_DATA
+
+void UMaterial::UpdateTransientExpressionData()
+{
+	for (UMaterialExpression* Expression : GetExpressions())
+	{
+		if (UMaterialExpressionMaterialFunctionCall* FunctionCall = Cast<UMaterialExpressionMaterialFunctionCall>(Expression))
+		{
+			// Update the function call node, so it can relink inputs and outputs as needed
+			// Update even if MaterialFunctionNode->MaterialFunction is NULL, because we need to remove the invalid inputs in that case
+			FunctionCall->UpdateFromFunctionResource();
+		}
+	}
+}
+
+#endif // WITH_EDITORONLY_DATA
+
 #if WITH_EDITOR
 void UMaterial::UpdateCachedExpressionData()
 {
 	COOK_STAT(FScopedDurationTimer BlockingTimer(MaterialCookStats::UpdateCachedExpressionDataSec));
 
-	if (bLoadedCachedExpressionData && !GetPackage()->bIsCookedForEditor)
+	if (bLoadedCachedExpressionData)
 	{
 		// Don't need to rebuild cached data if it was serialized
 		return;
@@ -1975,9 +1992,17 @@ bool UMaterial::GetMaterialLayers(FMaterialLayersFunctions& OutLayers, TMicRecur
 #if WITH_EDITORONLY_DATA
 		if (CachedExpressionData->EditorOnlyData)
 		{
-			OutLayers.EditorOnly = CachedExpressionData->EditorOnlyData->MaterialLayers;
+			// cooked materials can strip out material layer information
+			if (CachedExpressionData->EditorOnlyData->MaterialLayers.LayerStates.Num() != 0)
+			{
+				OutLayers.EditorOnly = CachedExpressionData->EditorOnlyData->MaterialLayers;
+				OutLayers.Validate();
+			}
+			else
+			{
+				return false;
+			}
 		}
-		OutLayers.Validate();
 #endif // WITH_EDITORONLY_DATA
 		return true;
 	}
@@ -3558,7 +3583,19 @@ void UMaterial::PostLoad()
 	PropagateDataToMaterialProxy();
 
 #if WITH_EDITOR
-	UpdateCachedExpressionData();
+	// cooked materials will not have any expressions in them, so this will obliterate the saved cached expression data
+	if (!GetOutermost()->bIsCookedForEditor)
+	{
+		UpdateCachedExpressionData();
+	}
+	else
+	{
+#if WITH_EDITORONLY_DATA
+		// cooked materials will need to update their expressions, but not their cached expression data.
+		UpdateTransientExpressionData();
+#endif
+	}
+
 #endif // WITH_EDITOR
 
 	// Strata materials conversion needs to be done after expressions are cached, otherwise material function won't have 
diff --git a/Engine/Source/Runtime/Engine/Private/Materials/MaterialInstance.cpp b/Engine/Source/Runtime/Engine/Private/Materials/MaterialInstance.cpp
--- a/Engine/Source/Runtime/Engine/Private/Materials/MaterialInstance.cpp
+++ b/Engine/Source/Runtime/Engine/Private/Materials/MaterialInstance.cpp
@@ -2277,11 +2277,17 @@ bool UMaterialInstance::GetMaterialLayers(FMaterialLayersFunctions& OutLayers, T
 		OutLayers.GetRuntime() = StaticParametersRuntime.MaterialLayers;
 #if WITH_EDITORONLY_DATA
 		const UMaterialInstanceEditorOnlyData* EditorOnly = GetEditorOnlyData();
-		if (EditorOnly)
+
+		// cooked materials can strip out material layer information
+		if (EditorOnly && EditorOnly->StaticParameters.MaterialLayers.LayerStates.Num() != 0)
 		{
 			OutLayers.EditorOnly = EditorOnly->StaticParameters.MaterialLayers;
+			OutLayers.Validate();
+		}
+		else
+		{
+			return false;
 		}
-		OutLayers.Validate();
 #endif // WITH_EDITORONLY_DATA
 		return true;
 	}
