Moving property condition manager into UE::Net
Renaming SET_REPCUSTOMCONDITION_ACTIVE_FAST to DOREPCUSTOMCONDITION_SETACTIVE_FAST

diff --git a/Engine/Source/Runtime/Engine/Private/LevelActor.cpp b/Engine/Source/Runtime/Engine/Private/LevelActor.cpp
--- a/Engine/Source/Runtime/Engine/Private/LevelActor.cpp
+++ b/Engine/Source/Runtime/Engine/Private/LevelActor.cpp
@@ -948,7 +948,7 @@ bool UWorld::DestroyActor( AActor* ThisActor, bool bNetForce, bool bShouldModify
 		}
 	}
 
-	FNetPropertyConditionManager::Get().NotifyObjectDestroyed(ThisActor);
+	UE::Net::Private::FNetPropertyConditionManager::Get().NotifyObjectDestroyed(ThisActor);
 
 	// Notify net drivers that this guy has been destroyed.
 	if (FWorldContext* Context = GEngine->GetWorldContextFromWorld(this))
diff --git a/Engine/Source/Runtime/Engine/Private/NetDriver.cpp b/Engine/Source/Runtime/Engine/Private/NetDriver.cpp
--- a/Engine/Source/Runtime/Engine/Private/NetDriver.cpp
+++ b/Engine/Source/Runtime/Engine/Private/NetDriver.cpp
@@ -3225,7 +3225,7 @@
 
 bool UNetDriver::HandlePropertyConditionsMemCommand(const TCHAR* Cmd, FOutputDevice& Ar)
 {
-	FNetPropertyConditionManager::Get().LogMemory(Ar);
+	UE::Net::Private::FNetPropertyConditionManager::Get().LogMemory(Ar);
 	return true;
 }
 #endif // !UE_BUILD_SHIPPING
@@ -3481,7 +3481,7 @@
 
 void UNetDriver::NotifySubObjectDestroyed(UObject* SubObject)
 {
-	FNetPropertyConditionManager::Get().NotifyObjectDestroyed(SubObject);
+	UE::Net::Private::FNetPropertyConditionManager::Get().NotifyObjectDestroyed(SubObject);
 }
 
 void UNetDriver::RemoveNetworkActor(AActor* Actor)
@@ -6061,12 +6061,12 @@
 {
 	check(IsServer() || MaySendProperties());
 
-	return FNetPropertyConditionManager::Get().FindOrCreatePropertyTracker(Obj);
+	return UE::Net::Private::FNetPropertyConditionManager::Get().FindOrCreatePropertyTracker(Obj);
 }
 
 TSharedPtr<FRepChangedPropertyTracker> UNetDriver::FindRepChangedPropertyTracker(UObject* Obj)
 {
-	return FNetPropertyConditionManager::Get().FindPropertyTracker(Obj);
+	return UE::Net::Private::FNetPropertyConditionManager::Get().FindPropertyTracker(Obj);
 }
 
 TSharedPtr<FRepLayout> UNetDriver::GetObjectClassRepLayout( UClass * Class )
diff --git a/Engine/Source/Runtime/Engine/Public/Net/UnrealNetwork.h b/Engine/Source/Runtime/Engine/Public/Net/UnrealNetwork.h
--- a/Engine/Source/Runtime/Engine/Public/Net/UnrealNetwork.h
+++ b/Engine/Source/Runtime/Engine/Public/Net/UnrealNetwork.h
@@ -302,16 +302,16 @@ static FProperty* GetReplicatedProperty(const UClass* CallingClass, const UClass
 	OutActiveState.SetActiveState((uint16)c::ENetFields_Private::v, active); \
 }
 
-#define SET_REPCUSTOMCONDITION_ACTIVE_FAST(c,v,active) \
+#define DOREPCUSTOMCONDITION_SETACTIVE_FAST(c,v,active) \
 { \
-	FNetPropertyConditionManager::Get().SetPropertyActive(this, (uint16)c::ENetFields_Private::v, active); \
+	UE::Net::Private::FNetPropertyConditionManager::Get().SetPropertyActive(this, (uint16)c::ENetFields_Private::v, active); \
 }
 
-#define SET_REPCUSTOMCONDITION_ACTIVE_FAST_STATIC_ARRAY(c,v,active) \
+#define DOREPCUSTOMCONDITION_SETACTIVE_FAST_STATIC_ARRAY(c,v,active) \
 { \
 	for (int32 i = 0; i < (int32)c::EArrayDims_Private::v; ++i) \
 	{ \
-		FNetPropertyConditionManager::Get().SetPropertyActive(this, (uint16)c::ENetFields_Private::v##_STATIC_ARRAY + i, active); \
+		UE::Net::Private::FNetPropertyConditionManager::Get().SetPropertyActive(this, (uint16)c::ENetFields_Private::v##_STATIC_ARRAY + i, active); \
 	} \
 }
 
diff --git a/Engine/Source/Runtime/Net/Core/Private/Net/Core/PropertyConditions/PropertyConditions.cpp b/Engine/Source/Runtime/Net/Core/Private/Net/Core/PropertyConditions/PropertyConditions.cpp
--- a/Engine/Source/Runtime/Net/Core/Private/Net/Core/PropertyConditions/PropertyConditions.cpp
+++ b/Engine/Source/Runtime/Net/Core/Private/Net/Core/PropertyConditions/PropertyConditions.cpp
@@ -8,7 +8,10 @@
 #include "Stats/Stats2.h"
 
 DECLARE_CYCLE_STAT(TEXT("PropertyConditions PostGarbageCollect"), STAT_PropertyConditions_PostGarbageCollect, STATGROUP_Net);
-	
+
+namespace UE::Net::Private
+{
+
 FNetPropertyConditionManager::FNetPropertyConditionManager()
 {
 	PostGarbageCollectHandle = FCoreUObjectDelegates::GetPostGarbageCollect().AddRaw(this, &FNetPropertyConditionManager::PostGarbageCollect);
@@ -104,3 +107,5 @@ void FNetPropertyConditionManager::LogMemory(FOutputDevice& Ar)
 
 	Ar.Logf(TEXT("  Property Condition Memory: %u"), CountBytes);
 }
+
+}; // UE::Net::Private
\ No newline at end of file
diff --git a/Engine/Source/Runtime/Net/Core/Public/Net/Core/PropertyConditions/PropertyConditions.h b/Engine/Source/Runtime/Net/Core/Public/Net/Core/PropertyConditions/PropertyConditions.h
--- a/Engine/Source/Runtime/Net/Core/Public/Net/Core/PropertyConditions/PropertyConditions.h
+++ b/Engine/Source/Runtime/Net/Core/Public/Net/Core/PropertyConditions/PropertyConditions.h
@@ -8,6 +8,9 @@
 
 class FRepChangedPropertyTracker;
 
+namespace UE::Net::Private
+{
+
 class NETCORE_API FNetPropertyConditionManager
 {
 public:
@@ -31,3 +34,5 @@ private:
 
 	TMap<FObjectKey, TSharedPtr<FRepChangedPropertyTracker>> PropertyTrackerMap;
 };
+
+}; // UE::Net::Private
\ No newline at end of file
