#UE Fix warning about short path usage

diff --git a/Engine/Source/Runtime/Engine/Classes/GameFramework/PlatformInputDeviceMapperLibrary.h b/Engine/Source/Runtime/Engine/Classes/GameFramework/PlatformInputDeviceMapperLibrary.h
--- a/Engine/Source/Runtime/Engine/Classes/GameFramework/PlatformInputDeviceMapperLibrary.h
+++ b/Engine/Source/Runtime/Engine/Classes/GameFramework/PlatformInputDeviceMapperLibrary.h
@@ -112,11 +112,11 @@ public:
 	static bool IsValidPlatformId(FPlatformUserId UserId);
 
 	/** Static invalid platform user */
-	UFUNCTION(BlueprintPure, meta = (ScriptConstant = "None", ScriptConstantHost = "PlatformUserId"), Category = "PlatformInputDevice")
+	UFUNCTION(BlueprintPure, meta = (ScriptConstant = "None", ScriptConstantHost = "/Script/CoreUObject.PlatformUserId"), Category = "PlatformInputDevice")
 	static FPlatformUserId PlatformUserId_None();
 
 	/** Static invalid input device */
-	UFUNCTION(BlueprintPure, meta = (ScriptConstant = "None", ScriptConstantHost = "InputDeviceId"), Category = "PlatformInputDevice")
+	UFUNCTION(BlueprintPure, meta = (ScriptConstant = "None", ScriptConstantHost = "/Script/CoreUObject.InputDeviceId"), Category = "PlatformInputDevice")
 	static FInputDeviceId InputDeviceId_None();
 
 	/** Returns true if PlatformUserId A is equal to PlatformUserId B (A == B) */
