Renamed FSubObjectGetter to UE::Net::FSubObjectRegistryGetter and moved it to separate file

diff --git a/Engine/Source/Runtime/Engine/Classes/GameFramework/Actor.h b/Engine/Source/Runtime/Engine/Classes/GameFramework/Actor.h
--- a/Engine/Source/Runtime/Engine/Classes/GameFramework/Actor.h
+++ b/Engine/Source/Runtime/Engine/Classes/GameFramework/Actor.h
@@ -3670,7 +3670,7 @@
 private:
 	/** Collection of SubObjects that get replicated when this actor gets replicated. */
 	UE::Net::FSubObjectRegistry ReplicatedSubObjects;
-	friend class FSubObjectGetter;
+	friend class UE::Net::FSubObjectRegistryGetter;
 
 	/** Array of replicated components and the list of replicated subobjects they own. Replaces the deprecated ReplicatedCompoments array. */
 	TArray<UE::Net::FReplicatedComponentInfo> ReplicatedComponentsInfo;
diff --git a/Engine/Source/Runtime/Engine/Private/DataChannel.cpp b/Engine/Source/Runtime/Engine/Private/DataChannel.cpp
--- a/Engine/Source/Runtime/Engine/Private/DataChannel.cpp
+++ b/Engine/Source/Runtime/Engine/Private/DataChannel.cpp
@@ -26,6 +26,7 @@
 #include "Net/Subsystems/NetworkSubsystem.h"
 #include "Net/Core/Trace/NetTrace.h"
 #include "Net/Core/Misc/NetSubObjectRegistry.h"
+#include "Net/NetSubObjectRegistryGetter.h"
 #include "Misc/NetworkVersion.h"
 #include "Net/Core/PushModel/PushModel.h"
 #include "HAL/LowLevelMemStats.h"
@@ -138,38 +139,6 @@ static const bool IsBunchTooLarge(UNetConnection* Connection, T* Bunch)
 	return !Connection->IsUnlimitedBunchSizeAllowed() && Bunch != nullptr && Bunch->GetNumBytes() > NetMaxConstructedPartialBunchSizeBytes;
 }
 
-/** Helper class to restrict access to the subobject list of actors and actor components */
-class FSubObjectGetter final
-{
-public:
-	FSubObjectGetter() = delete;
-	~FSubObjectGetter() = delete;
-
-	using FSubObjectRegistry = UE::Net::FSubObjectRegistry;
-
-	static const FSubObjectRegistry& GetSubObjects(AActor* InActor)
-	{
-		return InActor->ReplicatedSubObjects;
-	}
-
-	static const FSubObjectRegistry* GetSubObjectsOfActorCompoment(AActor* InActor, UActorComponent* InActorComp)
-	{
-		UE::Net::FReplicatedComponentInfo* ComponentInfo = InActor->ReplicatedComponentsInfo.FindByKey(InActorComp);
-		return ComponentInfo ? &(ComponentInfo->SubObjects) : nullptr;
-	}
-
-	static const TArray<UE::Net::FReplicatedComponentInfo>& GetReplicatedComponents(AActor* InActor)
-	{
-		return InActor->ReplicatedComponentsInfo;
-	}
-
-	static bool IsSubObjectInRegistry(AActor* InActor, UActorComponent* InActorComp, UObject* InSubObject)
-	{
-		UE::Net::FReplicatedComponentInfo* ComponentInfo = InActor->ReplicatedComponentsInfo.FindByKey(InActorComp);
-		return ComponentInfo ? ComponentInfo->SubObjects.IsSubObjectInRegistry(InSubObject) : false;
-	}
-};
-
 namespace DataChannelInternal
 {
 	/**
@@ -3626,7 +3595,7 @@ bool UActorChannel::ReplicateRegisteredSubObjects(FOutBunch& Bunch, FReplication
 	// Start with the Actor's subobjects
 	{
 		SetCurrentSubObjectOwner(Actor);
-		for (const FSubObjectRegistry::FEntry& SubObjectInfo : FSubObjectGetter::GetSubObjects(Actor).GetRegistryList())
+		for (const FSubObjectRegistry::FEntry& SubObjectInfo : FSubObjectRegistryGetter::GetSubObjects(Actor).GetRegistryList())
 		{
 			if (CanSubObjectReplicateToClient(SubObjectInfo.NetCondition, SubObjectInfo.Key, ConditionMap))
 			{
@@ -3638,7 +3607,7 @@ bool UActorChannel::ReplicateRegisteredSubObjects(FOutBunch& Bunch, FReplication
 	}
 
 	// Now the replicated actor components
-	for( const FReplicatedComponentInfo& RepComponentInfo : FSubObjectGetter::GetReplicatedComponents(Actor) )
+	for( const FReplicatedComponentInfo& RepComponentInfo : FSubObjectRegistryGetter::GetReplicatedComponents(Actor) )
 	{
 		if (CanSubObjectReplicateToClient(RepComponentInfo.NetCondition, RepComponentInfo.Key, ConditionMap))
 		{
@@ -3701,9 +3670,11 @@ bool UActorChannel::ReplicateSubobject(UObject* SubObj, FOutBunch& Bunch, FRepli
 	else
 	{
 #if !UE_BUILD_SHIPPING
+		using namespace UE::Net;
+
 		// Make sure the SubObjectOwner actually has the SubObject in his list and will replicate it later.
 		UActorComponent* Component = CastChecked<UActorComponent>(DataChannelInternal::CurrentSubObjectOwner);
-		const bool bIsInRegistry = FSubObjectGetter::IsSubObjectInRegistry(Actor, Component, SubObj);
+		const bool bIsInRegistry = FSubObjectRegistryGetter::IsSubObjectInRegistry(Actor, Component, SubObj);
 		ensureMsgf(bIsInRegistry, TEXT("ReplicatedSubObject %s was replicated using the legacy method but it is not in %s::%s registered list. It won't be replicated at all."), 
 			*SubObj->GetName(), *Actor->GetName(), *Component->GetName());
 #endif
@@ -3784,7 +3755,8 @@ bool UActorChannel::ReplicateSubobject(UActorComponent* ReplicatedComponent, FOu
 
 bool UActorChannel::WriteSubObjects(UActorComponent* ReplicatedComponent, FOutBunch& Bunch, FReplicationFlags RepFlags, const TStaticBitArray<COND_Max>& ConditionMap)
 {
-	if (const UE::Net::FSubObjectRegistry* SubObjectList = FSubObjectGetter::GetSubObjectsOfActorCompoment(Actor, ReplicatedComponent))
+	using namespace UE::Net;
+	if (const FSubObjectRegistry* SubObjectList = FSubObjectRegistryGetter::GetSubObjectsOfActorCompoment(Actor, ReplicatedComponent))
 	{
 		return WriteSubObjects(ReplicatedComponent, *SubObjectList, Bunch, RepFlags, ConditionMap);
 	}
diff --git a/Engine/Source/Runtime/Engine/Private/Net/NetSubObjectRegistryGetter.h b/Engine/Source/Runtime/Engine/Private/Net/NetSubObjectRegistryGetter.h
new file mode 100644
--- /dev/null
+++ b/Engine/Source/Runtime/Engine/Private/Net/NetSubObjectRegistryGetter.h
@@ -0,0 +1,52 @@
+// Copyright Epic Games, Inc. All Rights Reserved.
+
+#pragma once
+
+#include "CoreTypes.h"
+#include "GameFramework/Actor.h"
+#include "Components/ActorComponent.h"
+#include "Net/Core/Misc/NetSubObjectRegistry.h"
+
+namespace UE::Net
+{
+
+/** Helper class to restrict access to the subobject list of actors and actor components */
+class FSubObjectRegistryGetter final
+{
+public:
+	using FReplicatedComponentInfo = UE::Net::FReplicatedComponentInfo;
+
+	FSubObjectRegistryGetter() = delete;
+	~FSubObjectRegistryGetter() = delete;
+
+	using FSubObjectRegistry = UE::Net::FSubObjectRegistry;
+
+	static const FSubObjectRegistry& GetSubObjects(AActor* InActor)
+	{
+		return InActor->ReplicatedSubObjects;
+	}
+
+	static const FSubObjectRegistry* GetSubObjectsOfActorCompoment(AActor* InActor, UActorComponent* InActorComp)
+	{
+		UE::Net::FReplicatedComponentInfo* ComponentInfo = InActor->ReplicatedComponentsInfo.FindByKey(InActorComp);
+		return ComponentInfo ? &(ComponentInfo->SubObjects) : nullptr;
+	}
+
+	static const TArray<UE::Net::FReplicatedComponentInfo>& GetReplicatedComponents(AActor* InActor)
+	{
+		return InActor->ReplicatedComponentsInfo;
+	}
+
+	static const FReplicatedComponentInfo* GetReplicatedComponentInfoForComponent(AActor* InActor, UActorComponent* InActorComp)
+	{
+		return InActor->ReplicatedComponentsInfo.FindByKey(InActorComp);
+	}
+
+	static bool IsSubObjectInRegistry(AActor* InActor, UActorComponent* InActorComp, UObject* InSubObject)
+	{
+		UE::Net::FReplicatedComponentInfo* ComponentInfo = InActor->ReplicatedComponentsInfo.FindByKey(InActorComp);
+		return ComponentInfo ? ComponentInfo->SubObjects.IsSubObjectInRegistry(InSubObject) : false;
+	}
+};
+
+}
diff --git a/Engine/Source/Runtime/Net/Core/Public/Net/Core/Misc/NetSubObjectRegistry.h b/Engine/Source/Runtime/Net/Core/Public/Net/Core/Misc/NetSubObjectRegistry.h
--- a/Engine/Source/Runtime/Net/Core/Public/Net/Core/Misc/NetSubObjectRegistry.h
+++ b/Engine/Source/Runtime/Net/Core/Public/Net/Core/Misc/NetSubObjectRegistry.h
@@ -7,6 +7,11 @@
 
 class UActorComponent;
 
+namespace UE::Net
+{
+	class FSubObjectRegistryGetter;
+}
+
 namespace UE::Net 
 {
 
