Fix incorrect SDF generation when there's a mismatch between static mesh sections and material array order.
- Iterate over LOD0 sections and use FMeshSectionInfo.MaterialIndex to correctly index into material array.

diff --git a/Engine/Source/Runtime/Engine/Private/DistanceFieldAtlas.cpp b/Engine/Source/Runtime/Engine/Private/DistanceFieldAtlas.cpp
--- a/Engine/Source/Runtime/Engine/Private/DistanceFieldAtlas.cpp
+++ b/Engine/Source/Runtime/Engine/Private/DistanceFieldAtlas.cpp
@@ -127,7 +127,7 @@
 #if WITH_EDITOR
 
 // DDC key for distance field data, must be changed when modifying the generation code or data format
-#define DISTANCEFIELD_DERIVEDDATA_VER TEXT("CB37D8F6-5ED7-476E-BDFB-A0319D4099DC")
+#define DISTANCEFIELD_DERIVEDDATA_VER TEXT("23A632E6-DB31-4A5B-B43D-17D98C6E946C")
 
 FString BuildDistanceFieldDerivedDataKey(const FString& InMeshKey)
 {
@@ -151,27 +151,44 @@
 
 void FDistanceFieldVolumeData::CacheDerivedData(const FString& InStaticMeshDerivedDataKey, const ITargetPlatform* TargetPlatform, UStaticMesh* Mesh, FStaticMeshRenderData& RenderData, UStaticMesh* GenerateSource, float DistanceFieldResolutionScale, bool bGenerateDistanceFieldAsIfTwoSided)
 {
-	FString DistanceFieldKey = BuildDistanceFieldDerivedDataKey(InStaticMeshDerivedDataKey);
+	const TArray<FStaticMaterial>& StaticMaterials = Mesh->GetStaticMaterials();
+
+	TArray<FSignedDistanceFieldBuildMaterialData> BuildMaterialData;
+	BuildMaterialData.SetNum(StaticMaterials.Num());
+
 	const FMeshSectionInfoMap& SectionInfoMap = Mesh->GetSectionInfoMap();
+	const uint32 LODIndex = 0;
 
-	for (int32 MaterialIndex = 0; MaterialIndex < Mesh->GetStaticMaterials().Num(); MaterialIndex++)
+	for (int32 SectionIndex = 0; SectionIndex < SectionInfoMap.GetSectionNumber(LODIndex); SectionIndex++)
 	{
-		FSignedDistanceFieldBuildMaterialData MaterialData;
-		// Default material blend mode
-		MaterialData.BlendMode = BLEND_Opaque;
-		MaterialData.bTwoSided = false;
-		MaterialData.bAffectDistanceFieldLighting = SectionInfoMap.IsValidSection(0, MaterialIndex) ? SectionInfoMap.Get(0, MaterialIndex).bAffectDistanceFieldLighting : true;
+		const FMeshSectionInfo& Section = SectionInfoMap.Get(LODIndex, SectionIndex);
 
-		UMaterialInterface* MaterialInterface = Mesh->GetStaticMaterials()[MaterialIndex].MaterialInterface;
+		if (!BuildMaterialData.IsValidIndex(Section.MaterialIndex))
+		{
+			continue;
+		}
+
+		FSignedDistanceFieldBuildMaterialData& MaterialData = BuildMaterialData[Section.MaterialIndex];
+		MaterialData.bAffectDistanceFieldLighting = Section.bAffectDistanceFieldLighting;
+
+		UMaterialInterface* MaterialInterface = StaticMaterials[Section.MaterialIndex].MaterialInterface;
 		if (MaterialInterface)
 		{
 			MaterialData.BlendMode = MaterialInterface->GetBlendMode();
 			MaterialData.bTwoSided = MaterialInterface->IsTwoSided();
 		}
-
-		DistanceFieldKey += FString::Printf(TEXT("_M%u_%u_%u"), (uint32)MaterialData.BlendMode, MaterialData.bTwoSided ? 1 : 0, MaterialData.bAffectDistanceFieldLighting ? 1 : 0);
 	}
 
+	FString DistanceFieldKey = BuildDistanceFieldDerivedDataKey(InStaticMeshDerivedDataKey);
+
+	for (int32 MaterialIndex = 0; MaterialIndex < Mesh->GetStaticMaterials().Num(); MaterialIndex++)
+	{
+		DistanceFieldKey += FString::Printf(TEXT("_M%u_%u_%u"), 
+			(uint32)BuildMaterialData[MaterialIndex].BlendMode,
+			BuildMaterialData[MaterialIndex].bTwoSided ? 1 : 0,
+			BuildMaterialData[MaterialIndex].bAffectDistanceFieldLighting ? 1 : 0);
+	}
+
 	TArray<uint8> DerivedData;
 
 	COOK_STAT(auto Timer = DistanceFieldCookStats::UsageStats.TimeSyncWork());
@@ -185,11 +202,12 @@
 	}
 	else if (GDistanceFieldAsyncQueue)
 	{
+		check(Mesh && GenerateSource);
+
 		// We don't actually build the resource until later, so only track the cycles used here.
 		COOK_STAT(Timer.TrackCyclesOnly());
 		FAsyncDistanceFieldTask* NewTask = new FAsyncDistanceFieldTask;
 		NewTask->DDCKey = DistanceFieldKey;
-		check(Mesh && GenerateSource);
 		NewTask->TargetPlatform = TargetPlatform;
 		NewTask->StaticMesh = Mesh;
 		NewTask->GenerateSource = GenerateSource;
@@ -198,24 +216,8 @@
 		NewTask->GeneratedVolumeData = new FDistanceFieldVolumeData();
 		NewTask->GeneratedVolumeData->AssetName = Mesh->GetFName();
 		NewTask->GeneratedVolumeData->bAsyncBuilding = true;
+		NewTask->MaterialBlendModes = MoveTemp(BuildMaterialData);
 
-		for (int32 MaterialIndex = 0; MaterialIndex < Mesh->GetStaticMaterials().Num(); MaterialIndex++)
-		{
-			FSignedDistanceFieldBuildMaterialData MaterialData;
-			// Default material blend mode
-			MaterialData.BlendMode = BLEND_Opaque;
-			MaterialData.bTwoSided = false;
-			MaterialData.bAffectDistanceFieldLighting = SectionInfoMap.IsValidSection(0, MaterialIndex) ? SectionInfoMap.Get(0, MaterialIndex).bAffectDistanceFieldLighting : true;
-
-			if (Mesh->GetStaticMaterials()[MaterialIndex].MaterialInterface)
-			{
-				MaterialData.BlendMode = Mesh->GetStaticMaterials()[MaterialIndex].MaterialInterface->GetBlendMode();
-				MaterialData.bTwoSided = Mesh->GetStaticMaterials()[MaterialIndex].MaterialInterface->IsTwoSided();
-			}
-
-			NewTask->MaterialBlendModes.Add(MaterialData);
-		}
-
 		// Nanite overrides source static mesh with a coarse representation. Need to load original data before we build the mesh SDF.
 		if (Mesh->NaniteSettings.bEnabled)
 		{
diff --git a/Engine/Source/Runtime/Engine/Private/MeshCardRepresentation.cpp b/Engine/Source/Runtime/Engine/Private/MeshCardRepresentation.cpp
--- a/Engine/Source/Runtime/Engine/Private/MeshCardRepresentation.cpp
+++ b/Engine/Source/Runtime/Engine/Private/MeshCardRepresentation.cpp
@@ -138,7 +138,7 @@
 #if WITH_EDITOR
 
 // DDC key for card representation data, must be changed when modifying the generation code or data format
-#define CARDREPRESENTATION_DERIVEDDATA_VER TEXT("3F7067D7-0552-4C66-BBC5-F2589BDB798A")
+#define CARDREPRESENTATION_DERIVEDDATA_VER TEXT("3F695875-02C1-459D-8923-222B68165285")
 
 FString BuildCardRepresentationDerivedDataKey(const FString& InMeshKey, int32 MaxLumenMeshCards)
 {
@@ -205,34 +205,41 @@
 	}
 	else
 	{
+		check(Mesh && GenerateSource);
+
 		// We don't actually build the resource until later, so only track the cycles used here.
 		COOK_STAT(Timer.TrackCyclesOnly());
 		FAsyncCardRepresentationTask* NewTask = new FAsyncCardRepresentationTask;
 		NewTask->DDCKey = InDDCKey;
-		check(Mesh && GenerateSource);
 		NewTask->StaticMesh = Mesh;
 		NewTask->GenerateSource = GenerateSource;
 		NewTask->GeneratedCardRepresentation = new FCardRepresentationData();
 		NewTask->MaxLumenMeshCards = MaxLumenMeshCards;
 		NewTask->bGenerateDistanceFieldAsIfTwoSided = bGenerateDistanceFieldAsIfTwoSided;
+		NewTask->MaterialBlendModes.SetNum(Mesh->GetStaticMaterials().Num());
 
+		const TArray<FStaticMaterial>& StaticMaterials = Mesh->GetStaticMaterials();
 		const FMeshSectionInfoMap& SectionInfoMap = Mesh->GetSectionInfoMap();
+		const uint32 LODIndex = 0;
 
-		for (int32 MaterialIndex = 0; MaterialIndex < Mesh->GetStaticMaterials().Num(); MaterialIndex++)
+		for (int32 SectionIndex = 0; SectionIndex < SectionInfoMap.GetSectionNumber(LODIndex); SectionIndex++)
 		{
-			FSignedDistanceFieldBuildMaterialData MaterialData;
-			// Default material blend mode
-			MaterialData.BlendMode = BLEND_Opaque;
-			MaterialData.bTwoSided = false;
-			MaterialData.bAffectDistanceFieldLighting = SectionInfoMap.IsValidSection(0, MaterialIndex) ? SectionInfoMap.Get(0, MaterialIndex).bAffectDistanceFieldLighting : true;
+			const FMeshSectionInfo& Section = SectionInfoMap.Get(LODIndex, SectionIndex);
 
-			if (Mesh->GetStaticMaterials()[MaterialIndex].MaterialInterface)
+			if (!NewTask->MaterialBlendModes.IsValidIndex(Section.MaterialIndex))
 			{
-				MaterialData.BlendMode = Mesh->GetStaticMaterials()[MaterialIndex].MaterialInterface->GetBlendMode();
-				MaterialData.bTwoSided = Mesh->GetStaticMaterials()[MaterialIndex].MaterialInterface->IsTwoSided();
+				continue;
 			}
 
-			NewTask->MaterialBlendModes.Add(MaterialData);
+			FSignedDistanceFieldBuildMaterialData& MaterialData = NewTask->MaterialBlendModes[Section.MaterialIndex];
+			MaterialData.bAffectDistanceFieldLighting = Section.bAffectDistanceFieldLighting;
+
+			UMaterialInterface* MaterialInterface = StaticMaterials[Section.MaterialIndex].MaterialInterface;
+			if (MaterialInterface)
+			{
+				MaterialData.BlendMode = MaterialInterface->GetBlendMode();
+				MaterialData.bTwoSided = MaterialInterface->IsTwoSided();
+			}
 		}
 
 		// Nanite overrides source static mesh with a coarse representation. Need to load original data before we build the mesh SDF.
