Physics interface cleanup.

diff --git a/Engine/Source/Runtime/Engine/Classes/Components/SkeletalMeshComponent.h b/Engine/Source/Runtime/Engine/Classes/Components/SkeletalMeshComponent.h
--- a/Engine/Source/Runtime/Engine/Classes/Components/SkeletalMeshComponent.h
+++ b/Engine/Source/Runtime/Engine/Classes/Components/SkeletalMeshComponent.h
@@ -43,6 +43,10 @@ struct FCompactPose;
 
 enum class EClothingTeleportMode : uint8;
 
+#ifndef WITH_CLOTH_COLLISION_DETECTION
+#define WITH_CLOTH_COLLISION_DETECTION 1
+#endif
+
 DECLARE_MULTICAST_DELEGATE(FOnSkelMeshPhysicsCreatedMultiCast);
 typedef FOnSkelMeshPhysicsCreatedMultiCast::FDelegate FOnSkelMeshPhysicsCreated;
 
@@ -55,13 +59,6 @@ DECLARE_DYNAMIC_MULTICAST_DELEGATE(FOnBoneTransformsFinalized);  // Deprecated,
 
 DECLARE_MULTICAST_DELEGATE(FOnBoneTransformsFinalizedMultiCast);
 
-#if PHYSICS_INTERFACE_PHYSX
-namespace physx
-{
-	class PxAggregate;
-}
-#endif
-
 UENUM()
 enum class EAnimCurveType : uint8 
 {
@@ -632,10 +629,6 @@ public:
 
 	/** Indicates that this SkeletalMeshComponent has deferred kinematic bone updates until next physics sim if not INDEX_NONE. */
 	int32 DeferredKinematicUpdateIndex;
-#if PHYSICS_INTERFACE_PHYSX
-	/** Indicates that this SkeletalMeshComponent has deferred kinematic bone updates until next physics sim.  */
-	uint8 bDeferredKinematicUpdate:1;
-#endif
 
 private:
 	/** Disable rigid body animation nodes and play original animation without simulation */
diff --git a/Engine/Source/Runtime/Engine/Classes/Engine/SkeletalMesh.h b/Engine/Source/Runtime/Engine/Classes/Engine/SkeletalMesh.h
--- a/Engine/Source/Runtime/Engine/Classes/Engine/SkeletalMesh.h
+++ b/Engine/Source/Runtime/Engine/Classes/Engine/SkeletalMesh.h
@@ -228,17 +228,6 @@ struct FSkeletalMeshAsyncBuildTask : public FAsyncTask<FSkeletalMeshAsyncBuildWo
 
 #endif // #if WITH_EDITOR
 
-#if WITH_APEX_CLOTHING
-
-namespace nvidia
-{
-	namespace apex
-	{
-		class ClothingAsset;
-	}
-}
-#endif
-
 UENUM()
 enum class ESkinCacheUsage : uint8
 {
@@ -620,15 +609,8 @@ struct FClothingAssetData_Legacy
 	UPROPERTY()
 	FClothPhysicsProperties_Legacy PhysicsProperties;
 
-#if WITH_APEX_CLOTHING
-	nvidia::apex::ClothingAsset* ApexClothingAsset;
-#endif// #if WITH_APEX_CLOTHING
-
 	FClothingAssetData_Legacy()
 		: bClothPropertiesChanged(false), PhysicsProperties()
-#if WITH_APEX_CLOTHING
-		, ApexClothingAsset(nullptr)
-#endif// #if WITH_APEX_CLOTHING
 	{
 	}
 
@@ -2439,15 +2421,6 @@ public:
 		PRAGMA_ENABLE_DEPRECATION_WARNINGS
 	}
 
-#if WITH_EDITOR && WITH_APEX_CLOTHING
-	/** 
-	 * Take clothing assets that were imported using APEX files before we moved away from the APEX simulation
-	 * framework and upgrade them to UClothingAssets. This will bind the new assets to the mesh so the
-	 * clothing remains working as before.
-	 */
-	void UpgradeOldClothingAssets();
-#endif //WITH_EDITOR && WITH_APEX_CLOTHING
-
 #if WITH_EDITOR
 	/** If the given section of the specified LOD has a clothing asset, unbind it's data and remove it from the asset array */
 	void RemoveClothingAsset(int32 InLodIndex, int32 InSectionIndex);
diff --git a/Engine/Source/Runtime/Engine/Classes/Engine/World.h b/Engine/Source/Runtime/Engine/Classes/Engine/World.h
--- a/Engine/Source/Runtime/Engine/Classes/Engine/World.h
+++ b/Engine/Source/Runtime/Engine/Classes/Engine/World.h
@@ -1542,13 +1542,12 @@ private:
 	FPhysScene*									PhysicsScene;
 	// Note that this should be merged with PhysScene going forward but is needed for now.
 public:
-#if INCLUDE_CHAOS
+
 	/** Current global physics scene. */
 	TSharedPtr<FPhysScene_Chaos> PhysicsScene_Chaos;
 
 	/** Default global physics scene. */
 	TSharedPtr<FPhysScene_Chaos> DefaultPhysicsScene_Chaos;
-#endif
 
 	/** Physics Field component. */
 	UPROPERTY(Transient)
@@ -1604,10 +1603,8 @@ private:
 	/** Utility function to cleanup streaming levels that point to invalid level packages */
 	void RepairStreamingLevels();
 
-#if INCLUDE_CHAOS
 	/** Utility function that is used to ensure that a World has the correct ChaosActor */
 	void RepairChaosActors();
-#endif
 
 #if WITH_EDITOR
 	/** Utility function to make sure there is a valid default builder brush */
diff --git a/Engine/Source/Runtime/Engine/Classes/PhysicsEngine/BodySetup.h b/Engine/Source/Runtime/Engine/Classes/PhysicsEngine/BodySetup.h
--- a/Engine/Source/Runtime/Engine/Classes/PhysicsEngine/BodySetup.h
+++ b/Engine/Source/Runtime/Engine/Classes/PhysicsEngine/BodySetup.h
@@ -26,23 +26,6 @@ enum class EPhysXMeshCookFlags : uint8;
 
 DECLARE_DELEGATE_OneParam(FOnAsyncPhysicsCookFinished, bool);
 
-#if PHYSICS_INTERFACE_PHYSX
-namespace physx
-{
-	class PxTriangleMesh;
-	class PxRigidActor;
-	class PxTransform;
-	class PxSphereGeometry;
-	class PxBoxGeometry;
-	class PxCapsuleGeometry;
-	class PxConvexMeshGeometry;
-	class PxConvexMesh;
-	class PxTriangleMesh;
-	class PxTriangleMeshGeometry;
-}
-#endif
-
-#if WITH_CHAOS
 namespace Chaos
 {
 	class FImplicitObject;
@@ -53,8 +36,6 @@ namespace Chaos
 template<typename T, int d>
 class FChaosDerivedDataReader;
 
-#endif
-
 DECLARE_CYCLE_STAT_EXTERN(TEXT("PhysX Cooking"), STAT_PhysXCooking, STATGROUP_Physics, );
 
 
@@ -91,13 +72,11 @@ struct ENGINE_API FCookBodySetupInfo
 	/** Trimesh data for cooking */
 	FTriMeshCollisionData TriangleMeshDesc;
 
-#if WITH_PHYSX
 	/** Trimesh cook flags */
 	EPhysXMeshCookFlags TriMeshCookFlags;
 
 	/** Convex cook flags */
 	EPhysXMeshCookFlags ConvexCookFlags;
-#endif // WITH_PHYSX
 
 	/** Vertices of NonMirroredConvex hulls */
 	TArray<TArray<FVector>> NonMirroredConvexVertices;
@@ -246,17 +225,10 @@ private:
 	FFormatContainer CookedFormatDataRuntimeOnlyOptimization;
 #endif
 
-#if WITH_PHYSX
-	/** Get cook flags for 'runtime only' cooked physics data */
-	EPhysXMeshCookFlags GetRuntimeOnlyCookOptimizationFlags() const;
-#endif 
-
 public:
 
-#if WITH_CHAOS
 	//FBodySetupTriMeshes* TriMeshWrapper;
 	TArray<TSharedPtr<Chaos::FTriangleMeshImplicitObject, ESPMode::ThreadSafe>> ChaosTriMeshes;
-#endif
 
 	/** Additional UV info, if available. Used for determining UV for a line trace impact. */
 	FBodySetupUVInfo UVInfo;
@@ -275,18 +247,12 @@ public:
 	UPROPERTY()
 	FVector BuildScale3D;
 
-#if PHYSICS_INTERFACE_PHYSX
 	/** References the current async cook helper. Used to be able to abort a cook task */
-	using FAsyncCookHelper = FPhysXCookHelper;
-#elif WITH_CHAOS
 	using FAsyncCookHelper = Chaos::FCookHelper;
-#endif
 	FAsyncCookHelper* CurrentCookHelper;
 
-#if WITH_CHAOS
 	// Will contain deserialized data from the serialization function that can be used at PostLoad time.
 	TUniquePtr<FChaosDerivedDataReader<float, 3>> ChaosDerivedDataReader;
-#endif
 
 public:
 	//~ Begin UObject Interface.
@@ -327,16 +293,7 @@ public:
 private:
 	FByteBulkData* GetCookedFormatData();
 
-#if WITH_PHYSX && PHYSICS_INTERFACE_PHYSX
-	bool ProcessFormatData_PhysX(FByteBulkData* FormatData);
-	bool RuntimeCookPhysics_PhysX();
-
 	// #TODO MRMesh for some reason needs to be able to call this - that case needs fixed to correctly use the create meshes flow
-	friend class UMRMeshComponent;
-	/** Finish creating the physics meshes and update the body setup data with cooked data */
-	ENGINE_API void FinishCreatingPhysicsMeshes_PhysX(const TArray<physx::PxConvexMesh*>& ConvexMeshes, const TArray<physx::PxConvexMesh*>& ConvexMeshesNegX, const TArray<physx::PxTriangleMesh*>& TriMeshes);
-	// TODO: ProcessFormatData_Chaos is calling ProcessFormatData_Chaos directly - it's better if CreatePhysicsMeshes can be used but that code path requires WITH_EDITOR
-#elif WITH_CHAOS
 	friend class UMRMeshComponent;
 	bool ProcessFormatData_Chaos(FByteBulkData* FormatData);
 	bool ProcessFormatData_Chaos(FChaosDerivedDataReader<float, 3>& Reader);
@@ -347,7 +304,6 @@ private:
 										   TArray<TSharedPtr<Chaos::FTriangleMeshImplicitObject, ESPMode::ThreadSafe>>& InTrimeshImplicits,
 										   FBodySetupUVInfo& InUvInfo,
 										   TArray<int32>& InFaceRemap);
-#endif
 
 	/** 
 	 * Finalize game thread data before calling back user's delegate 
diff --git a/Engine/Source/Runtime/Engine/Classes/PhysicsEngine/ConstraintDrives.h b/Engine/Source/Runtime/Engine/Classes/PhysicsEngine/ConstraintDrives.h
--- a/Engine/Source/Runtime/Engine/Classes/PhysicsEngine/ConstraintDrives.h
+++ b/Engine/Source/Runtime/Engine/Classes/PhysicsEngine/ConstraintDrives.h
@@ -6,13 +6,6 @@
 #include "EngineDefines.h"
 #include "ConstraintDrives.generated.h"
 
-#if WITH_PHYSX
-namespace physx
-{
-	class PxD6Joint;
-}
-#endif // WITH_PHYSX
-
 UENUM()
 namespace EAngularDriveMode
 {
@@ -53,10 +46,8 @@ struct ENGINE_API FConstraintDrive
 
 	FConstraintDrive();
 
-#if WITH_PHYSX
 	/** Updates physx drive with properties from unreal */
-	void UpdatePhysXDrive_AssumesLocked(physx::PxD6Joint* Joint, int DriveType, bool bDriveEnabled) const;
-#endif
+	//void UpdatePhysXDrive_AssumesLocked(physx::PxD6Joint* Joint, int DriveType, bool bDriveEnabled) const;
 
 private:
 	friend struct FConstraintInstance;
@@ -176,10 +167,8 @@ struct ENGINE_API FAngularDriveConstraint
 		
 	}
 
-#if WITH_PHYSX
 	/** Updates physx drive with properties from unreal */
-	void UpdatePhysXAngularDrive_AssumesLocked(physx::PxD6Joint* Joint) const;
-#endif	// WITH_PHYSX
+	//void UpdatePhysXAngularDrive_AssumesLocked(physx::PxD6Joint* Joint) const;
 
 private:
 	friend struct FConstraintInstance;
diff --git a/Engine/Source/Runtime/Engine/Classes/PhysicsEngine/ConstraintTypes.h b/Engine/Source/Runtime/Engine/Classes/PhysicsEngine/ConstraintTypes.h
--- a/Engine/Source/Runtime/Engine/Classes/PhysicsEngine/ConstraintTypes.h
+++ b/Engine/Source/Runtime/Engine/Classes/PhysicsEngine/ConstraintTypes.h
@@ -9,13 +9,6 @@
 #include "Chaos/ChaosEngineInterface.h"
 #include "ConstraintTypes.generated.h"
 
-#if WITH_PHYSX
-namespace physx
-{
-	class PxD6Joint;
-}
-#endif // WITH_PHYSX
-
 USTRUCT()
 struct ENGINE_API FConstraintBaseParams
 {
diff --git a/Engine/Source/Runtime/Engine/Classes/PhysicsEngine/ConvexElem.h b/Engine/Source/Runtime/Engine/Classes/PhysicsEngine/ConvexElem.h
--- a/Engine/Source/Runtime/Engine/Classes/PhysicsEngine/ConvexElem.h
+++ b/Engine/Source/Runtime/Engine/Classes/PhysicsEngine/ConvexElem.h
@@ -5,10 +5,7 @@
 #include "CoreMinimal.h"
 #include "UObject/ObjectMacros.h"
 #include "PhysicsEngine/ShapeElem.h"
-
-#if WITH_CHAOS
 #include "Chaos/Serializable.h"
-#endif
 
 #include "ConvexElem.generated.h"
 
@@ -51,17 +48,7 @@ private:
 	UPROPERTY()
 	FTransform Transform;
 
-#if PHYSICS_INTERFACE_PHYSX
-	/** Convex mesh for this body, created from cooked data in CreatePhysicsMeshes */
-	physx::PxConvexMesh*   ConvexMesh;
-
-	/** Convex mesh for this body, flipped across X, created from cooked data in CreatePhysicsMeshes */
-	physx::PxConvexMesh*   ConvexMeshNegX;
-#endif
-
-#if WITH_CHAOS
 	TSharedPtr<Chaos::FConvex, ESPMode::ThreadSafe> ChaosConvex;
-#endif
 
 public:
 
@@ -104,21 +91,6 @@ public:
 	/** Returns the volume of this element */
 	FVector::FReal GetScaledVolume(const FVector& Scale3D) const;
 
-#if PHYSICS_INTERFACE_PHYSX
-	/** Get the PhysX convex mesh (defined in BODY space) for this element */
-	ENGINE_API physx::PxConvexMesh* GetConvexMesh() const;
-
-	/** Set the PhysX convex mesh to use for this element */
-	ENGINE_API void SetConvexMesh(physx::PxConvexMesh* InMesh);
-
-	/** Get the PhysX convex mesh (defined in BODY space) for this element */
-	ENGINE_API physx::PxConvexMesh* GetMirroredConvexMesh() const;
-
-	/** Set the PhysX convex mesh to use for this element */
-	ENGINE_API void SetMirroredConvexMesh(physx::PxConvexMesh* InMesh);
-#endif
-
-#if WITH_CHAOS
 	ENGINE_API const auto& GetChaosConvexMesh() const
 	{
 		return ChaosConvex;
@@ -131,7 +103,6 @@ public:
 	ENGINE_API void ComputeChaosConvexIndices(bool bForceCompute = false);
 
 	ENGINE_API TArray<int32> GetChaosConvexIndices() const;
-#endif
 
 	/** Get current transform applied to convex mesh vertices */
 	FTransform GetTransform() const
diff --git a/Engine/Source/Runtime/Engine/Classes/PhysicsEngine/PhysicalAnimationComponent.h b/Engine/Source/Runtime/Engine/Classes/PhysicsEngine/PhysicalAnimationComponent.h
--- a/Engine/Source/Runtime/Engine/Classes/PhysicsEngine/PhysicalAnimationComponent.h
+++ b/Engine/Source/Runtime/Engine/Classes/PhysicsEngine/PhysicalAnimationComponent.h
@@ -12,13 +12,6 @@
 class FPrimitiveDrawInterface;
 class USkeletalMeshComponent;
 
-#if PHYSICS_INTERFACE_PHYSX
-namespace physx
-{
-	class PxRigidDynamic;
-}
-#endif
-
 /** Stores info on the type of motor that will be used for a given bone */
 USTRUCT(BlueprintType)
 struct ENGINE_API FPhysicalAnimationData
@@ -129,11 +122,7 @@ private:
 	struct FPhysicalAnimationInstanceData
 	{
 		struct FConstraintInstance* ConstraintInstance;
-#if WITH_CHAOS
 		FPhysicsActorHandle TargetActor;
-#else
-		physx::PxRigidDynamic* TargetActor;
-#endif
 	};
 
 	/** constraints used to apply the drive data */
diff --git a/Engine/Source/Runtime/Engine/Classes/PhysicsEngine/PhysicsHandleComponent.h b/Engine/Source/Runtime/Engine/Classes/PhysicsEngine/PhysicsHandleComponent.h
--- a/Engine/Source/Runtime/Engine/Classes/PhysicsEngine/PhysicsHandleComponent.h
+++ b/Engine/Source/Runtime/Engine/Classes/PhysicsEngine/PhysicsHandleComponent.h
@@ -9,14 +9,6 @@
 #include "PhysicsEngine/ConstraintInstance.h"
 #include "PhysicsHandleComponent.generated.h"
 
-#if PHYSICS_INTERFACE_PHYSX
-namespace physx
-{
-	class PxD6Joint;
-	class PxRigidDynamic;
-}
-#endif
-
 /**
  *	Utility object for moving physics objects around.
  */
@@ -72,12 +64,6 @@ class ENGINE_API UPhysicsHandleComponent : public UActorComponent
 
 protected:
 
-#if PHYSICS_INTERFACE_PHYSX
-	/** Pointer to PhysX joint used by the handle*/
-	physx::PxD6Joint* HandleData;
-	/** Pointer to kinematic actor jointed to grabbed object */
-	physx::PxRigidDynamic* KinActorData;
-#elif WITH_CHAOS
 	FTransform PreviousTransform;
 	bool bPendingConstraint;
 
@@ -88,7 +74,6 @@ protected:
 	FPhysicsConstraintHandle ConstraintHandle;
 	FVector ConstraintLocalPosition; // Position of constraint in the grabbed body local space (updated when grabbing)
 	FRotator ConstraintLocalRotation;
-#endif
 
 	//~ Begin UActorComponent Interface.
 	virtual void OnUnregister() override;
diff --git a/Engine/Source/Runtime/Engine/Engine.Build.cs b/Engine/Source/Runtime/Engine/Engine.Build.cs
--- a/Engine/Source/Runtime/Engine/Engine.Build.cs
+++ b/Engine/Source/Runtime/Engine/Engine.Build.cs
@@ -369,11 +369,6 @@ public class Engine : ModuleRules
 
 		SetupModulePhysicsSupport(Target);
 
-		if (Target.bCompilePhysX && (Target.bBuildEditor || Target.bCompileAPEX))
-		{
-			DynamicallyLoadedModuleNames.Add("PhysXCooking");
-		}
-
 		// Engine public headers need to know about some types (enums etc.)
 		PublicIncludePathModuleNames.Add("ClothingSystemRuntimeInterface");
 		PublicDependencyModuleNames.Add("ClothingSystemRuntimeInterface");
diff --git a/Engine/Source/Runtime/Engine/Private/AI/Navigation/NavigationTypes.cpp b/Engine/Source/Runtime/Engine/Private/AI/Navigation/NavigationTypes.cpp
--- a/Engine/Source/Runtime/Engine/Private/AI/Navigation/NavigationTypes.cpp
+++ b/Engine/Source/Runtime/Engine/Private/AI/Navigation/NavigationTypes.cpp
@@ -293,9 +293,6 @@ bool FNavAgentSelector::Serialize(FArchive& Ar)
 //----------------------------------------------------------------------//
 FNavHeightfieldSamples::FNavHeightfieldSamples()
 {
-#if WITH_PHYSX
-	//static_assert(sizeof(physx::PxI16) == sizeof(Heights.GetTypeSize()), "FNavHeightfieldSamples::Heights' type needs to be kept in sync with physx::PxI16");
-#endif // WITH_PHYSX
 }
 
 void FNavHeightfieldSamples::GetResourceSizeEx(FResourceSizeEx& CumulativeResourceSize)
diff --git a/Engine/Source/Runtime/Engine/Private/ActorReplication.cpp b/Engine/Source/Runtime/Engine/Private/ActorReplication.cpp
--- a/Engine/Source/Runtime/Engine/Private/ActorReplication.cpp
+++ b/Engine/Source/Runtime/Engine/Private/ActorReplication.cpp
@@ -20,12 +20,9 @@
 #include "DrawDebugHelpers.h"
 #include "Net/Core/PushModel/PushModel.h"
 #include "Interfaces/Interface_ActorSubobject.h"
-
-#if WITH_CHAOS
 #include "Physics/Experimental/PhysScene_Chaos.h"
 #include "PBDRigidsSolver.h"
 #include "ChaosSolversModule.h"
-#endif
 
 /*-----------------------------------------------------------------------------
 	AActor networking implementation.
@@ -364,7 +361,7 @@ void AActor::GatherCurrentMovement()
 		if (RootPrimComp && RootPrimComp->IsSimulatingPhysics())
 		{
 			bool bFoundInCache = false;
-#if WITH_CHAOS
+
 			UWorld* World = GetWorld();
 			if (FPhysScene_Chaos* Scene = static_cast<FPhysScene_Chaos*>(World->GetPhysicsScene()))
 			{
@@ -374,7 +371,7 @@ void AActor::GatherCurrentMovement()
 					bFoundInCache = true;
 				}
 			}
-#endif
+
 			if (!bFoundInCache)
 			{
 				// fallback to GT data
diff --git a/Engine/Source/Runtime/Engine/Private/CharacterMovementComponentAsync.cpp b/Engine/Source/Runtime/Engine/Private/CharacterMovementComponentAsync.cpp
--- a/Engine/Source/Runtime/Engine/Private/CharacterMovementComponentAsync.cpp
+++ b/Engine/Source/Runtime/Engine/Private/CharacterMovementComponentAsync.cpp
@@ -3015,7 +3015,6 @@ bool FUpdatedComponentAsyncInput::ShouldIgnoreOverlapResult(const UWorld* World,
 
 void FUpdatedComponentAsyncInput::SetPosition(const FVector& InPosition) const
 {
-#if WITH_CHAOS
 	if (PhysicsHandle->GetPhysicsThreadAPI() == nullptr)
 	{
 		return;
@@ -3039,23 +3038,20 @@ void FUpdatedComponentAsyncInput::SetPosition(const FVector& InPosition) const
 	{
 		ensure(false);
 	}
-#endif
 }
 
 FVector FUpdatedComponentAsyncInput::GetPosition() const
 {
-#if WITH_CHAOS
 	if (PhysicsHandle && PhysicsHandle->GetPhysicsThreadAPI())
 	{
 		return PhysicsHandle->GetPhysicsThreadAPI()->X();
 	}
-#endif
+
 	return FVector::ZeroVector;
 }
 
 void FUpdatedComponentAsyncInput::SetRotation(const FQuat& InRotation) const
 {
-#if WITH_CHAOS
 	if (PhysicsHandle->GetPhysicsThreadAPI() == nullptr)
 	{
 		return;
@@ -3079,17 +3075,14 @@ void FUpdatedComponentAsyncInput::SetRotation(const FQuat& InRotation) const
 	{
 		ensure(false);
 	}
-#endif
 }
 
 FQuat FUpdatedComponentAsyncInput::GetRotation() const
 {
-#if WITH_CHAOS
 	if (PhysicsHandle && PhysicsHandle->GetPhysicsThreadAPI())
 	{
 		return PhysicsHandle->GetPhysicsThreadAPI()->R();
 	}
-#endif
 
 	return FQuat::Identity;
 }
diff --git a/Engine/Source/Runtime/Engine/Private/Collision/CollisionConversions.cpp b/Engine/Source/Runtime/Engine/Private/Collision/CollisionConversions.cpp
--- a/Engine/Source/Runtime/Engine/Private/Collision/CollisionConversions.cpp
+++ b/Engine/Source/Runtime/Engine/Private/Collision/CollisionConversions.cpp
@@ -12,14 +12,10 @@
 #include "PhysicsEngine/BodySetup.h"
 #include "Physics/PhysicsInterfaceUtils.h"
 
-#if PHYSICS_INTERFACE_PHYSX
-#include "PhysXInterfaceWrapper.h"
-#include "Collision/CollisionConversionsPhysx.h"
-#elif WITH_CHAOS
 #include "Physics/Experimental/ChaosInterfaceWrapper.h"
 #include "Physics/Experimental/PhysInterface_Chaos.h"
 #include "Chaos/ParticleHandle.h"
-#endif
+
 #include "PhysicsEngine/CollisionQueryFilterCallback.h"
 #include "GameFramework/LightWeightInstanceManager.h"
 #include "PhysicsProxy/SingleParticlePhysicsProxy.h"
@@ -94,17 +90,13 @@ static FVector FindSimpleOpposingNormal(const FHitLocation& Hit, const FVector&
 
 /**
  * Util to find the normal of the face that we hit. Will use faceIndex from the hit if possible.
- * @param PHit - incoming hit from PhysX
+ * @param PHit - incoming hit from Physics
  * @param TraceDirectionDenorm - direction of sweep test (not normalized)
  * @param InNormal - default value in case no new normal is computed.
  * @return New normal we compute for geometry.
  */
-#if WITH_CHAOS
 template <typename THitLocation>
 static FVector FindGeomOpposingNormal(ECollisionShapeType QueryGeomType, const THitLocation& Hit, const FVector& TraceDirectionDenorm, const FVector InNormal)
-#else
-static FVector FindGeomOpposingNormal(ECollisionShapeType QueryGeomType, const FHitLocation& Hit, const FVector& TraceDirectionDenorm, const FVector InNormal)
-#endif
 {
 	// TODO: can we support other shapes here as well?
 	if (QueryGeomType == ECollisionShapeType::Capsule || QueryGeomType == ECollisionShapeType::Sphere)
@@ -112,26 +104,11 @@ static FVector FindGeomOpposingNormal(ECollisionShapeType QueryGeomType, const F
 		const FPhysicsShape* Shape = GetShape(Hit);
 		if (Shape)
 		{
-#if WITH_CHAOS
 			const FTransform ActorTM(Hit.Actor->R(), Hit.Actor->X());
 			const FVector LocalInNormal = ActorTM.InverseTransformVectorNoScale(InNormal);
 			const FVector LocalTraceDirectionDenorm = ActorTM.InverseTransformVectorNoScale(TraceDirectionDenorm);
 			const FVector LocalNormal = Shape->GetGeometry()->FindGeometryOpposingNormal(LocalTraceDirectionDenorm, Hit.FaceIndex, LocalInNormal);
 			return ActorTM.TransformVectorNoScale(LocalNormal);
-#else
-			ECollisionShapeType GeomType = GetGeometryType(*Shape);
-			switch (GeomType)
-			{
-			case ECollisionShapeType::Sphere:
-			case ECollisionShapeType::Capsule:			return FindSimpleOpposingNormal(Hit, TraceDirectionDenorm, InNormal);
-			case ECollisionShapeType::Box:				return FindBoxOpposingNormal(Hit, TraceDirectionDenorm, InNormal);
-			case ECollisionShapeType::Convex:		return FindConvexMeshOpposingNormal(Hit, TraceDirectionDenorm, InNormal);
-			case ECollisionShapeType::Heightfield:		return FindHeightFieldOpposingNormal(Hit, TraceDirectionDenorm, InNormal);
-			case ECollisionShapeType::Trimesh:	return FindTriMeshOpposingNormal(Hit, TraceDirectionDenorm, InNormal);
-			default: break;
-			}
-
-#endif
 		}
 	}
 
@@ -143,11 +120,9 @@ static void SetHitResultFromShapeAndFaceIndex(const FPhysicsShape& Shape,  const
 {
 	SCOPE_CYCLE_COUNTER(STAT_CollisionSetHitResultFromShapeAndFaceIndex);
 
-#if WITH_CHAOS
-		const int32 ShapeIndex = Shape.GetShapeIndex();
-		CHAOS_CHECK(ShapeIndex < (int32)TNumericLimits<uint8>::Max()); // I could just write < 256, but this makes it more clear *why*
-		OutResult.ElementIndex = (uint8)ShapeIndex;
-#endif
+	const int32 ShapeIndex = Shape.GetShapeIndex();
+	CHAOS_CHECK(ShapeIndex < (int32)TNumericLimits<uint8>::Max()); // I could just write < 256, but this makes it more clear *why*
+	OutResult.ElementIndex = (uint8)ShapeIndex;
 	
 	UPrimitiveComponent* OwningComponent = nullptr;
 	if(const FBodyInstance* BodyInst = GetUserData(Actor))
@@ -164,28 +139,8 @@ static void SetHitResultFromShapeAndFaceIndex(const FPhysicsShape& Shape,  const
 
 		OwningComponent = BodyInst->OwnerComponent.Get();
 	}
-#if PHYSICS_INTERFACE_PHYSX
-	else if(const FCustomPhysXPayload* CustomPayload = GetUserData<FCustomPhysXPayload>(Shape))	//todo(ocohen): wrap with PHYSX
-	{
-		//Custom payload case
-		OwningComponent = CustomPayload->GetOwningComponent().Get();
-		if(OwningComponent && OwningComponent->bMultiBodyOverlap)
-		{
-			OutResult.Item = CustomPayload->GetItemIndex();
-			OutResult.BoneName = CustomPayload->GetBoneName();
-		}
-		else
-		{
-			OutResult.Item = INDEX_NONE;
-			OutResult.BoneName = NAME_None;
-		
-		}
-		
-	}
-#endif
 	else
 	{
-#if WITH_CHAOS
 		// Currently geom collections are registered with a primitive component user data, but maybe custom should be adapted
 		// to be more general so we can support leaf identification #BGTODO
 		void* UserData = Actor.UserData();
@@ -198,7 +153,6 @@ static void SetHitResultFromShapeAndFaceIndex(const FPhysicsShape& Shape,  const
 			OutResult.BoneName = NAME_None;
 		}
 		else
-#endif
 		{
 			ensureMsgf(false, TEXT("SetHitResultFromShapeAndFaceIndex hit shape with invalid userData"));
 		}
@@ -222,18 +176,10 @@ static void SetHitResultFromShapeAndFaceIndex(const FPhysicsShape& Shape,  const
 
 		if (bReturnPhysMat)
 		{
-#if WITH_CHAOS
 			if (const FPhysicsMaterial* PhysicsMaterial = GetMaterialFromInternalFaceIndexAndHitLocation(Shape, Actor, FaceIndex, HitLocation))
 			{
 				OutResult.PhysMaterial = GetUserData(*PhysicsMaterial);
 			}
-#else
-			if (const FPhysicsMaterial* PhysicsMaterial = GetMaterialFromInternalFaceIndex(Shape, Actor, FaceIndex))
-			{
-				OutResult.PhysMaterial = GetUserData(*PhysicsMaterial);
-			}
-
-#endif
 		}
 	}
 
@@ -264,7 +210,6 @@ const FPhysicsShape* GetGTShape<false>(const FPhysicsShape* PTShape, const FPhys
 	return nullptr;
 }
 
-#if WITH_CHAOS
 const FPhysicsActor* GetGTActor(const Chaos::FGeometryParticleHandle* PTActor)
 {
 	//TODO: need to pass in context so that in PT we always return null
@@ -273,7 +218,6 @@ const FPhysicsActor* GetGTActor(const Chaos::FGeometryParticleHandle* PTActor)
 	auto Proxy = static_cast<const Chaos::FSingleParticlePhysicsProxy*>(PTActor->PhysicsProxy());
 	return Proxy->GetParticle_LowLevel();
 }
-#endif
 
 template <typename THitLocation>
 EConvertQueryResult ConvertQueryImpactHitImp(const UWorld* World, const THitLocation& Hit, FHitResult& OutResult, float CheckLength, const FCollisionFilterData& QueryFilter, const FVector& StartLoc, const FVector& EndLoc, const FPhysicsGeometry* Geom, const FTransform& QueryTM, bool bReturnFaceIndex, bool bReturnPhysMat)
@@ -644,21 +588,8 @@ void ConvertQueryOverlap(const FPhysicsShape& Shape, const FPhysicsActor& Actor,
 			OutOverlap.ItemIndex = OwnerComponent->bMultiBodyOverlap ? BodyInst->InstanceBodyIndex : INDEX_NONE;
 		}
 	}
-#if PHYSICS_INTERFACE_PHYSX
-	else if(const FCustomPhysXPayload* CustomPayload = GetUserData<FCustomPhysXPayload>(Shape))
-	{
-		TWeakObjectPtr<UPrimitiveComponent> OwnerComponent = CustomPayload->GetOwningComponent();
-		if (UPrimitiveComponent* OwnerComponentRaw = OwnerComponent.Get())
-		{
-			OutOverlap.Component = OwnerComponent; // Copying weak pointer is faster than assigning raw pointer.
-			OutOverlap.OverlapObjectHandle = FActorInstanceHandle(OutOverlap.Component->GetOwner());
-			OutOverlap.ItemIndex = OwnerComponent->bMultiBodyOverlap ? CustomPayload->GetItemIndex() : INDEX_NONE;
-		}
-	}
-#endif
 	else
 	{
-#if WITH_CHAOS
 		// Currently geom collections are registered with a primitive component user data, but maybe custom should be adapted
 		// to be more general so we can support leaf identification #BGTODO
 		void* UserData = Actor.UserData();
@@ -671,7 +602,6 @@ void ConvertQueryOverlap(const FPhysicsShape& Shape, const FPhysicsActor& Actor,
 			OutOverlap.ItemIndex = INDEX_NONE;
 		}
 		else
-#endif
 		{
 			ensureMsgf(false, TEXT("ConvertQueryOverlap called with bad payload type"));
 		}
diff --git a/Engine/Source/Runtime/Engine/Private/Collision/CollisionConversionsPhysx.cpp b/Engine/Source/Runtime/Engine/Private/Collision/CollisionConversionsPhysx.cpp
--- a/Engine/Source/Runtime/Engine/Private/Collision/CollisionConversionsPhysx.cpp
+++ b/Engine/Source/Runtime/Engine/Private/Collision/CollisionConversionsPhysx.cpp
@@ -1,573 +1 @@
 // Copyright Epic Games, Inc. All Rights Reserved.
-
-#if PHYSICS_INTERFACE_PHYSX
-
-#include "Collision/CollisionConversionsPhysx.h"
-#include "Collision/CollisionConversions.h"
-#include "Collision/CollisionDebugDrawing.h"
-#include "Components/LineBatchComponent.h"
-#include "PhysicalMaterials/PhysicalMaterial.h"
-#include "PhysicsEngine/PhysicsSettings.h"
-#include "PhysicsEngine/BodySetup.h"
-#include "Physics/PhysicsInterfaceUtils.h"
-
-#include "PhysXPublic.h"
-#include "CustomPhysXPayload.h"
-#include "PhysXInterfaceWrapper.h"
-
-static FORCEINLINE bool PxQuatIsIdentity(PxQuat const& Q)
-{
-	return
-		Q.x == 0.f &&
-		Q.y == 0.f &&
-		Q.z == 0.f &&
-		Q.w == 1.f;
-}
-
-
-/** Helper to transform a normal when non-uniform scale is present. */
-static PxVec3 TransformNormalToShapeSpace(const PxMeshScale& meshScale, const PxVec3& nIn)
-{
-	// Uniform scale makes this unnecessary
-	if (meshScale.scale.x == meshScale.scale.y &&
-		meshScale.scale.x == meshScale.scale.z)
-	{
-		return nIn;
-	}
-
-	if (PxQuatIsIdentity(meshScale.rotation))
-	{
-		// Inverse transpose: inverse is 1/scale, transpose = original when rotation is identity.
-		const PxVec3 tmp = PxVec3(nIn.x / meshScale.scale.x, nIn.y / meshScale.scale.y, nIn.z / meshScale.scale.z);
-		const PxReal denom = 1.0f / tmp.magnitude();
-		return tmp * denom;
-	}
-	else
-	{
-		const PxMat33 rot(meshScale.rotation);
-		const PxMat33 diagonal = PxMat33::createDiagonal(meshScale.scale);
-		const PxMat33 vertex2Shape = (rot.getTranspose() * diagonal) * rot;
-
-		const PxMat33 shape2Vertex = vertex2Shape.getInverse();
-		const PxVec3 tmp = shape2Vertex.transformTranspose(nIn);
-		const PxReal denom = 1.0f / tmp.magnitude();
-		return tmp * denom;
-	}
-}
-
-FVector FindBoxOpposingNormal(const PxLocationHit& PHit, const FVector& TraceDirectionDenorm, const FVector InNormal)
-{
-	// We require normal info for our algorithm.
-	const bool bNormalData = (PHit.flags & PxHitFlag::eNORMAL);
-	if (!bNormalData)
-	{
-		return InNormal;
-	}
-
-	PxBoxGeometry PxBoxGeom;
-	const bool bReadGeomSuccess = PHit.shape->getBoxGeometry(PxBoxGeom);
-	check(bReadGeomSuccess); // This function should only be used for box geometry
-
-	const PxTransform LocalToWorld = PxShapeExt::getGlobalPose(*PHit.shape, *PHit.actor);
-
-	// Find which faces were included in the contact normal, and for multiple faces, use the one most opposing the sweep direction.
-	const PxVec3 ContactNormalLocal = LocalToWorld.rotateInv(PHit.normal);
-	const float* ContactNormalLocalPtr = &ContactNormalLocal.x;
-	const PxVec3 TraceDirDenormWorld = U2PVector(TraceDirectionDenorm);
-	const float* TraceDirDenormWorldPtr = &TraceDirDenormWorld.x;
-	const PxVec3 TraceDirDenormLocal = LocalToWorld.rotateInv(TraceDirDenormWorld);
-	const float* TraceDirDenormLocalPtr = &TraceDirDenormLocal.x;
-
-	PxVec3 BestLocalNormal(ContactNormalLocal);
-	float* BestLocalNormalPtr = &BestLocalNormal.x;
-	float BestOpposingDot = FLT_MAX;
-
-	for (int32 i = 0; i < 3; i++)
-	{
-		// Select axis of face to compare to, based on normal.
-		if (ContactNormalLocalPtr[i] > KINDA_SMALL_NUMBER)
-		{
-			const float TraceDotFaceNormal = TraceDirDenormLocalPtr[i]; // TraceDirDenormLocal.dot(BoxFaceNormal)
-			if (TraceDotFaceNormal < BestOpposingDot)
-			{
-				BestOpposingDot = TraceDotFaceNormal;
-				BestLocalNormal = PxVec3(0.f);
-				BestLocalNormalPtr[i] = 1.f;
-			}
-		}
-		else if (ContactNormalLocalPtr[i] < -KINDA_SMALL_NUMBER)
-		{
-			const float TraceDotFaceNormal = -TraceDirDenormLocalPtr[i]; // TraceDirDenormLocal.dot(BoxFaceNormal)
-			if (TraceDotFaceNormal < BestOpposingDot)
-			{
-				BestOpposingDot = TraceDotFaceNormal;
-				BestLocalNormal = PxVec3(0.f);
-				BestLocalNormalPtr[i] = -1.f;
-			}
-		}
-	}
-
-	// Fill in result
-	const PxVec3 WorldNormal = LocalToWorld.rotate(BestLocalNormal);
-	return P2UVector(WorldNormal);
-}
-
-FVector FindHeightFieldOpposingNormal(const PxLocationHit& PHit, const FVector& TraceDirectionDenorm, const FVector InNormal)
-{
-	if (IsInvalidFaceIndex(PHit.faceIndex))
-	{
-		return InNormal;
-	}
-
-	PxHeightFieldGeometry PHeightFieldGeom;
-	const bool bReadGeomSuccess = PHit.shape->getHeightFieldGeometry(PHeightFieldGeom);
-	check(bReadGeomSuccess);	//we should only call this function when we have a heightfield
-	if (PHeightFieldGeom.heightField)
-	{
-		const PxU32 TriIndex = PHit.faceIndex;
-		const PxTransform PShapeWorldPose = PxShapeExt::getGlobalPose(*PHit.shape, *PHit.actor);
-
-		PxTriangle Tri;
-		PxMeshQuery::getTriangle(PHeightFieldGeom, PShapeWorldPose, TriIndex, Tri);
-
-		PxVec3 TriNormal;
-		Tri.normal(TriNormal);
-		return P2UVector(TriNormal);
-	}
-
-	return InNormal;
-}
-
-FVector FindConvexMeshOpposingNormal(const PxLocationHit& PHit, const FVector& TraceDirectionDenorm, const FVector InNormal)
-{
-	if (IsInvalidFaceIndex(PHit.faceIndex))
-	{
-		return InNormal;
-	}
-
-	PxConvexMeshGeometry PConvexMeshGeom;
-	bool bSuccess = PHit.shape->getConvexMeshGeometry(PConvexMeshGeom);
-	check(bSuccess);	//should only call this function when we have a convex mesh
-
-	if (PConvexMeshGeom.convexMesh)
-	{
-		check(PHit.faceIndex < PConvexMeshGeom.convexMesh->getNbPolygons());
-
-		const PxU32 PolyIndex = PHit.faceIndex;
-		PxHullPolygon PPoly;
-		bool bSuccessData = PConvexMeshGeom.convexMesh->getPolygonData(PolyIndex, PPoly);
-		if (bSuccessData)
-		{
-			// Account for non-uniform scale in local space normal.
-			const PxVec3 PPlaneNormal(PPoly.mPlane[0], PPoly.mPlane[1], PPoly.mPlane[2]);
-			const PxVec3 PLocalPolyNormal = TransformNormalToShapeSpace(PConvexMeshGeom.scale, PPlaneNormal.getNormalized());
-
-			// Convert to world space
-			const PxTransform PShapeWorldPose = PxShapeExt::getGlobalPose(*PHit.shape, *PHit.actor);
-			const PxVec3 PWorldPolyNormal = PShapeWorldPose.rotate(PLocalPolyNormal);
-			const FVector OutNormal = P2UVector(PWorldPolyNormal);
-
-#if !(UE_BUILD_SHIPPING || UE_BUILD_TEST)
-			if (!OutNormal.IsNormalized())
-			{
-				UE_LOG(LogPhysics, Warning, TEXT("Non-normalized Normal (Hit shape is ConvexMesh): %s (LocalPolyNormal:%s)"), *OutNormal.ToString(), *P2UVector(PLocalPolyNormal).ToString());
-				UE_LOG(LogPhysics, Warning, TEXT("WorldTransform \n: %s"), *P2UTransform(PShapeWorldPose).ToString());
-			}
-#endif
-			return OutNormal;
-		}
-	}
-
-	return InNormal;
-}
-
-FVector FindTriMeshOpposingNormal(const PxLocationHit& PHit, const FVector& TraceDirectionDenorm, const FVector InNormal)
-{
-	if (IsInvalidFaceIndex(PHit.faceIndex))
-	{
-		return InNormal;
-	}
-
-	PxTriangleMeshGeometry PTriMeshGeom;
-	bool bSuccess = PHit.shape->getTriangleMeshGeometry(PTriMeshGeom);
-	check(bSuccess);	//this function should only be called when we have a trimesh
-
-	if (PTriMeshGeom.triangleMesh)
-	{
-		check(PHit.faceIndex < PTriMeshGeom.triangleMesh->getNbTriangles());
-
-		const PxU32 TriIndex = PHit.faceIndex;
-		const void* Triangles = PTriMeshGeom.triangleMesh->getTriangles();
-
-		// Grab triangle indices that we hit
-		int32 I0, I1, I2;
-
-		if (PTriMeshGeom.triangleMesh->getTriangleMeshFlags() & PxTriangleMeshFlag::e16_BIT_INDICES)
-		{
-			PxU16* P16BitIndices = (PxU16*)Triangles;
-			I0 = P16BitIndices[(TriIndex * 3) + 0];
-			I1 = P16BitIndices[(TriIndex * 3) + 1];
-			I2 = P16BitIndices[(TriIndex * 3) + 2];
-		}
-		else
-		{
-			PxU32* P32BitIndices = (PxU32*)Triangles;
-			I0 = P32BitIndices[(TriIndex * 3) + 0];
-			I1 = P32BitIndices[(TriIndex * 3) + 1];
-			I2 = P32BitIndices[(TriIndex * 3) + 2];
-		}
-
-		// Get verts we hit (local space)
-		const PxVec3* PVerts = PTriMeshGeom.triangleMesh->getVertices();
-		const PxVec3 V0 = PVerts[I0];
-		const PxVec3 V1 = PVerts[I1];
-		const PxVec3 V2 = PVerts[I2];
-
-		// Find normal of triangle (local space), and account for non-uniform scale
-		const PxVec3 PTempNormal = ((V1 - V0).cross(V2 - V0)).getNormalized();
-		const PxVec3 PLocalTriNormal = TransformNormalToShapeSpace(PTriMeshGeom.scale, PTempNormal);
-
-		// Convert to world space
-		const PxTransform PShapeWorldPose = PxShapeExt::getGlobalPose(*PHit.shape, *PHit.actor);
-		const PxVec3 PWorldTriNormal = PShapeWorldPose.rotate(PLocalTriNormal);
-		FVector OutNormal = P2UVector(PWorldTriNormal);
-
-		if (PTriMeshGeom.meshFlags & PxMeshGeometryFlag::eDOUBLE_SIDED)
-		{
-			//double sided mesh so we need to consider direction of query
-			const float sign = FVector::DotProduct(OutNormal, TraceDirectionDenorm) > 0.f ? -1.f : 1.f;
-			OutNormal *= sign;
-		}
-
-#if !(UE_BUILD_SHIPPING || UE_BUILD_TEST)
-		if (!OutNormal.IsNormalized())
-		{
-			UE_LOG(LogPhysics, Warning, TEXT("Non-normalized Normal (Hit shape is TriangleMesh): %s (V0:%s, V1:%s, V2:%s)"), *OutNormal.ToString(), *P2UVector(V0).ToString(), *P2UVector(V1).ToString(), *P2UVector(V2).ToString());
-			UE_LOG(LogPhysics, Warning, TEXT("WorldTransform \n: %s"), *P2UTransform(PShapeWorldPose).ToString());
-		}
-#endif
-		return OutNormal;
-	}
-
-	return InNormal;
-}
-
-
-
-static bool ComputeInflatedMTD_Internal(const float MtdInflation, const PxLocationHit& PHit, FHitResult& OutResult, const PxTransform& QueryTM, const PxGeometry& Geom, const PxTransform& PShapeWorldPose)
-{
-	PxGeometry* InflatedGeom = NULL;
-
-	PxVec3 PxMtdNormal(0.f);
-	PxF32 PxMtdDepth = 0.f;
-	PxGeometryHolder Holder = PHit.shape->getGeometry();
-	const PxGeometry& POtherGeom = Holder.any();
-	const bool bMtdResult = PxGeometryQuery::computePenetration(PxMtdNormal, PxMtdDepth, Geom, QueryTM, POtherGeom, PShapeWorldPose);
-	if (bMtdResult)
-	{
-		if (PxMtdNormal.isFinite())
-		{
-			OutResult.ImpactNormal = P2UVector(PxMtdNormal);
-			OutResult.PenetrationDepth = FMath::Max(FMath::Abs(PxMtdDepth) - MtdInflation, 0.f) + KINDA_SMALL_NUMBER;
-			return true;
-		}
-		else
-		{
-			UE_LOG(LogPhysics, Verbose, TEXT("Warning: ComputeInflatedMTD_Internal: MTD returned NaN :( normal: (X:%f, Y:%f, Z:%f)"), PxMtdNormal.x, PxMtdNormal.y, PxMtdNormal.z);
-		}
-	}
-
-	return false;
-}
-
-
-// Compute depenetration vector and distance if possible with a slightly larger geometry
-static bool ComputeInflatedMTD(const float MtdInflation, const PxLocationHit& PHit, FHitResult& OutResult, const PxTransform& QueryTM, const PxGeometry& Geom, const PxTransform& PShapeWorldPose)
-{
-	switch (Geom.getType())
-	{
-	case PxGeometryType::eCAPSULE:
-	{
-		const PxCapsuleGeometry* InCapsule = static_cast<const PxCapsuleGeometry*>(&Geom);
-		PxCapsuleGeometry InflatedCapsule(InCapsule->radius + MtdInflation, InCapsule->halfHeight); // don't inflate halfHeight, radius is added all around.
-		return ComputeInflatedMTD_Internal(MtdInflation, PHit, OutResult, QueryTM, InflatedCapsule, PShapeWorldPose);
-	}
-
-	case PxGeometryType::eBOX:
-	{
-		const PxBoxGeometry* InBox = static_cast<const PxBoxGeometry*>(&Geom);
-		PxBoxGeometry InflatedBox(InBox->halfExtents + PxVec3(MtdInflation));
-		return ComputeInflatedMTD_Internal(MtdInflation, PHit, OutResult, QueryTM, InflatedBox, PShapeWorldPose);
-	}
-
-	case PxGeometryType::eSPHERE:
-	{
-		const PxSphereGeometry* InSphere = static_cast<const PxSphereGeometry*>(&Geom);
-		PxSphereGeometry InflatedSphere(InSphere->radius + MtdInflation);
-		return ComputeInflatedMTD_Internal(MtdInflation, PHit, OutResult, QueryTM, InflatedSphere, PShapeWorldPose);
-	}
-
-	case PxGeometryType::eCONVEXMESH:
-	{
-		// We can't exactly inflate the mesh (not easily), so try jittering it a bit to get an MTD result.
-		PxVec3 TraceDir = U2PVector(OutResult.TraceEnd - OutResult.TraceStart);
-		TraceDir.normalizeSafe();
-
-		// Try forward (in trace direction)
-		PxTransform JitteredTM = PxTransform(QueryTM.p + (TraceDir * MtdInflation), QueryTM.q);
-		if (ComputeInflatedMTD_Internal(MtdInflation, PHit, OutResult, JitteredTM, Geom, PShapeWorldPose))
-		{
-			return true;
-		}
-
-		// Try backward (opposite trace direction)
-		JitteredTM = PxTransform(QueryTM.p - (TraceDir * MtdInflation), QueryTM.q);
-		if (ComputeInflatedMTD_Internal(MtdInflation, PHit, OutResult, JitteredTM, Geom, PShapeWorldPose))
-		{
-			return true;
-		}
-
-		// Try axial directions.
-		// Start with -Z because this is the most common case (objects on the floor).
-		for (int32 i = 2; i >= 0; i--)
-		{
-			PxVec3 Jitter(0.f);
-			float* JitterPtr = &Jitter.x;
-			JitterPtr[i] = MtdInflation;
-
-			JitteredTM = PxTransform(QueryTM.p - Jitter, QueryTM.q);
-			if (ComputeInflatedMTD_Internal(MtdInflation, PHit, OutResult, JitteredTM, Geom, PShapeWorldPose))
-			{
-				return true;
-			}
-
-			JitteredTM = PxTransform(QueryTM.p + Jitter, QueryTM.q);
-			if (ComputeInflatedMTD_Internal(MtdInflation, PHit, OutResult, JitteredTM, Geom, PShapeWorldPose))
-			{
-				return true;
-			}
-		}
-
-		return false;
-	}
-
-	default:
-	{
-		return false;
-	}
-	}
-}
-
-static bool CanFindOverlappedTriangle(const PxShape* PShape)
-{
-	return (PShape && (PShape->getGeometryType() == PxGeometryType::eTRIANGLEMESH || PShape->getGeometryType() == PxGeometryType::eHEIGHTFIELD));
-}
-
-/* Function to find the best normal from the list of triangles that are overlapping our geom. */
-template<typename GeomType>
-FVector FindBestOverlappingNormal(const UWorld* World, const PxGeometry& Geom, const PxTransform& QueryTM, const GeomType& ShapeGeom, const PxTransform& PShapeWorldPose, PxU32* HitTris, int32 NumTrisHit, bool bCanDrawOverlaps = false)
-{
-#if DRAW_OVERLAPPING_TRIS
-	const float Lifetime = 5.f;
-	bCanDrawOverlaps &= World && World->IsGameWorld() && World->PersistentLineBatcher && (World->PersistentLineBatcher->BatchedLines.Num() < 2048);
-	if (bCanDrawOverlaps)
-	{
-		TArray<FOverlapResult> Overlaps;
-		DrawGeomOverlaps(World, Geom, QueryTM, Overlaps, Lifetime);
-	}
-	const FLinearColor LineColor = FLinearColor::Green;
-	const FLinearColor NormalColor = FLinearColor::Red;
-	const FLinearColor PointColor = FLinearColor::Yellow;
-#endif // DRAW_OVERLAPPING_TRIS
-
-	// Track the best triangle plane distance
-	float BestPlaneDist = -BIG_NUMBER;
-	FVector BestPlaneNormal(0, 0, 1);
-	// Iterate over triangles
-	for (int32 TriIdx = 0; TriIdx < NumTrisHit; TriIdx++)
-	{
-		PxTriangle Tri;
-		PxMeshQuery::getTriangle(ShapeGeom, PShapeWorldPose, HitTris[TriIdx], Tri);
-
-		const FVector A = P2UVector(Tri.verts[0]);
-		const FVector B = P2UVector(Tri.verts[1]);
-		const FVector C = P2UVector(Tri.verts[2]);
-
-		FVector TriNormal = ((B - A) ^ (C - A));
-		TriNormal = TriNormal.GetSafeNormal();
-
-		const FPlane TriPlane(A, TriNormal);
-
-		const FVector QueryCenter = P2UVector(QueryTM.p);
-		const float DistToPlane = TriPlane.PlaneDot(QueryCenter);
-
-		if (DistToPlane > BestPlaneDist)
-		{
-			BestPlaneDist = DistToPlane;
-			BestPlaneNormal = TriNormal;
-		}
-
-#if DRAW_OVERLAPPING_TRIS
-		if (bCanDrawOverlaps && World && World->PersistentLineBatcher && World->PersistentLineBatcher->BatchedLines.Num() < 2048)
-		{
-			static const float LineThickness = 0.9f;
-			static const float NormalThickness = 0.75f;
-			static const float PointThickness = 5.0f;
-			World->PersistentLineBatcher->DrawLine(A, B, LineColor, SDPG_Foreground, LineThickness, Lifetime);
-			World->PersistentLineBatcher->DrawLine(B, C, LineColor, SDPG_Foreground, LineThickness, Lifetime);
-			World->PersistentLineBatcher->DrawLine(C, A, LineColor, SDPG_Foreground, LineThickness, Lifetime);
-			const FVector Centroid((A + B + C) / 3.f);
-			World->PersistentLineBatcher->DrawLine(Centroid, Centroid + (35.0f*TriNormal), NormalColor, SDPG_Foreground, NormalThickness, Lifetime);
-			World->PersistentLineBatcher->DrawPoint(Centroid + (35.0f*TriNormal), NormalColor, PointThickness, SDPG_Foreground, Lifetime);
-			World->PersistentLineBatcher->DrawPoint(A, PointColor, PointThickness, SDPG_Foreground, Lifetime);
-			World->PersistentLineBatcher->DrawPoint(B, PointColor, PointThickness, SDPG_Foreground, Lifetime);
-			World->PersistentLineBatcher->DrawPoint(C, PointColor, PointThickness, SDPG_Foreground, Lifetime);
-		}
-#endif // DRAW_OVERLAPPING_TRIS
-	}
-
-	return BestPlaneNormal;
-}
-
-
-static bool FindOverlappedTriangleNormal_Internal(const UWorld* World, const PxGeometry& Geom, const PxTransform& QueryTM, const PxShape* PShape, const PxTransform& PShapeWorldPose, FVector& OutNormal, bool bCanDrawOverlaps = false)
-{
-	if (CanFindOverlappedTriangle(PShape))
-	{
-		PxTriangleMeshGeometry PTriMeshGeom;
-		PxHeightFieldGeometry PHeightfieldGeom;
-
-		if (PShape->getTriangleMeshGeometry(PTriMeshGeom) || PShape->getHeightFieldGeometry(PHeightfieldGeom))
-		{
-			PxGeometryType::Enum GeometryType = PShape->getGeometryType();
-			const bool bIsTriMesh = (GeometryType == PxGeometryType::eTRIANGLEMESH);
-			PxU32 HitTris[64];
-			bool bOverflow = false;
-
-			const int32 NumTrisHit = bIsTriMesh ?
-				PxMeshQuery::findOverlapTriangleMesh(Geom, QueryTM, PTriMeshGeom, PShapeWorldPose, HitTris, UE_ARRAY_COUNT(HitTris), 0, bOverflow) :
-				PxMeshQuery::findOverlapHeightField(Geom, QueryTM, PHeightfieldGeom, PShapeWorldPose, HitTris, UE_ARRAY_COUNT(HitTris), 0, bOverflow);
-
-			if (NumTrisHit > 0)
-			{
-				if (bIsTriMesh)
-				{
-					OutNormal = FindBestOverlappingNormal(World, Geom, QueryTM, PTriMeshGeom, PShapeWorldPose, HitTris, NumTrisHit, bCanDrawOverlaps);
-				}
-				else
-				{
-					OutNormal = FindBestOverlappingNormal(World, Geom, QueryTM, PHeightfieldGeom, PShapeWorldPose, HitTris, NumTrisHit, bCanDrawOverlaps);
-				}
-
-				return true;
-			}
-		}
-	}
-
-	return false;
-}
-
-static bool FindOverlappedTriangleNormal(const UWorld* World, const PxGeometry& Geom, const PxTransform& QueryTM, const PxShape* PShape, const PxTransform& PShapeWorldPose, FVector& OutNormal, float Inflation, bool bCanDrawOverlaps = false)
-{
-	bool bSuccess = false;
-
-	if (CanFindOverlappedTriangle(PShape))
-	{
-		if (Inflation <= 0.f)
-		{
-			bSuccess = FindOverlappedTriangleNormal_Internal(World, Geom, QueryTM, PShape, PShapeWorldPose, OutNormal, bCanDrawOverlaps);
-		}
-		else
-		{
-			// Try a slightly inflated test if possible.
-			switch (Geom.getType())
-			{
-			case PxGeometryType::eCAPSULE:
-			{
-				const PxCapsuleGeometry* InCapsule = static_cast<const PxCapsuleGeometry*>(&Geom);
-				PxCapsuleGeometry InflatedCapsule(InCapsule->radius + Inflation, InCapsule->halfHeight); // don't inflate halfHeight, radius is added all around.
-				bSuccess = FindOverlappedTriangleNormal_Internal(World, InflatedCapsule, QueryTM, PShape, PShapeWorldPose, OutNormal, bCanDrawOverlaps);
-				break;
-			}
-
-			case PxGeometryType::eBOX:
-			{
-				const PxBoxGeometry* InBox = static_cast<const PxBoxGeometry*>(&Geom);
-				PxBoxGeometry InflatedBox(InBox->halfExtents + PxVec3(Inflation));
-				bSuccess = FindOverlappedTriangleNormal_Internal(World, InflatedBox, QueryTM, PShape, PShapeWorldPose, OutNormal, bCanDrawOverlaps);
-				break;
-			}
-
-			case PxGeometryType::eSPHERE:
-			{
-				const PxSphereGeometry* InSphere = static_cast<const PxSphereGeometry*>(&Geom);
-				PxSphereGeometry InflatedSphere(InSphere->radius + Inflation);
-				bSuccess = FindOverlappedTriangleNormal_Internal(World, InflatedSphere, QueryTM, PShape, PShapeWorldPose, OutNormal, bCanDrawOverlaps);
-				break;
-			}
-
-			default:
-			{
-				// No inflation possible
-				break;
-			}
-			}
-		}
-	}
-
-	return bSuccess;
-}
-
-void DrawOverlappingTris(const UWorld* World, const PxLocationHit& Hit, const PxGeometry& Geom, const FTransform& QueryTM)
-{
-	FVector DummyNormal(0.f);
-	const PxTransform PShapeWorldPose = PxShapeExt::getGlobalPose(*Hit.shape, *Hit.actor);
-	FindOverlappedTriangleNormal(World, Geom, U2PTransform(QueryTM), Hit.shape, PShapeWorldPose, DummyNormal, 0.f, true);
-}
-
-void ComputeZeroDistanceImpactNormalAndPenetration(const UWorld* World, const FHitLocation& Hit, const PxGeometry& Geom, const FTransform& QueryTM, FHitResult& OutResult)
-{
-	const PxTransform PQueryTM = U2PTransform(QueryTM);
-	const PxShape* PShape = GetShape(Hit);
-	const PxRigidActor* PActor = GetActor(Hit);
-	const PxTransform PShapeWorldPose = PxShapeExt::getGlobalPose(*PShape, *PActor);
-
-	// Try MTD with a small inflation for better accuracy, then a larger one in case the first one fails due to precision issues.
-	static const float SmallMtdInflation = 0.250f;
-	static const float LargeMtdInflation = 1.750f;
-
-	if (ComputeInflatedMTD(SmallMtdInflation, Hit, OutResult, PQueryTM, Geom, PShapeWorldPose) ||
-		ComputeInflatedMTD(LargeMtdInflation, Hit, OutResult, PQueryTM, Geom, PShapeWorldPose))
-	{
-		// Success
-	}
-	else
-	{
-		static const float SmallOverlapInflation = 0.250f;
-		if (FindOverlappedTriangleNormal(World, Geom, PQueryTM, PShape, PShapeWorldPose, OutResult.ImpactNormal, 0.f, false) ||
-			FindOverlappedTriangleNormal(World, Geom, PQueryTM, PShape, PShapeWorldPose, OutResult.ImpactNormal, SmallOverlapInflation, false))
-		{
-			// Success
-		}
-		else
-		{
-			// MTD failed, use point distance. This is not ideal.
-			// Note: faceIndex seems to be unreliable for convex meshes in these cases, so not using FindGeomOpposingNormal() for them here.
-			PxGeometryHolder Holder = PShape->getGeometry();
-			PxGeometry& PGeom = Holder.any();
-			PxVec3 PClosestPoint;
-			const float Distance = PxGeometryQuery::pointDistance(PQueryTM.p, PGeom, PShapeWorldPose, &PClosestPoint);
-
-			if (Distance < KINDA_SMALL_NUMBER)
-			{
-				UE_LOG(LogCollision, Verbose, TEXT("Warning: ConvertOverlappedShapeToImpactHit: Query origin inside shape, giving poor MTD."));
-				PClosestPoint = PxShapeExt::getWorldBounds(*PShape, *PActor).getCenter();
-			}
-
-			OutResult.ImpactNormal = (OutResult.Location - P2UVector(PClosestPoint)).GetSafeNormal();
-		}
-	}
-}
-
-
-#endif // WITH_PHYSX
diff --git a/Engine/Source/Runtime/Engine/Private/Collision/CollisionConversionsPhysx.h b/Engine/Source/Runtime/Engine/Private/Collision/CollisionConversionsPhysx.h
--- a/Engine/Source/Runtime/Engine/Private/Collision/CollisionConversionsPhysx.h
+++ b/Engine/Source/Runtime/Engine/Private/Collision/CollisionConversionsPhysx.h
@@ -6,19 +6,3 @@
 #include "CoreMinimal.h"
 #include "Engine/EngineTypes.h"
 #include "EngineDefines.h"
-
-#if PHYSICS_INTERFACE_PHYSX
-#include "PhysicsPublic.h"
-#include "PhysXIncludes.h"
-#include "PhysXInterfaceWrapper.h"
-
-FVector FindBoxOpposingNormal(const PxLocationHit& PHit, const FVector& TraceDirectionDenorm, const FVector InNormal);
-FVector FindHeightFieldOpposingNormal(const PxLocationHit& PHit, const FVector& TraceDirectionDenorm, const FVector InNormal);
-FVector FindConvexMeshOpposingNormal(const PxLocationHit& PHit, const FVector& TraceDirectionDenorm, const FVector InNormal);
-FVector FindTriMeshOpposingNormal(const PxLocationHit& PHit, const FVector& TraceDirectionDenorm, const FVector InNormal);
-
-void DrawOverlappingTris(const UWorld* World, const PxLocationHit& Hit, const PxGeometry& Geom, const FTransform& QueryTM);
-void ComputeZeroDistanceImpactNormalAndPenetration(const UWorld* World, const FHitLocation& Hit, const PxGeometry& Geom, const FTransform& QueryTM, FHitResult& OutResult);
-
-
-#endif //PHYSICS_INTERFACE_PHYSX
\ No newline at end of file
diff --git a/Engine/Source/Runtime/Engine/Private/Collision/CollisionDebugDrawing.cpp b/Engine/Source/Runtime/Engine/Private/Collision/CollisionDebugDrawing.cpp
--- a/Engine/Source/Runtime/Engine/Private/Collision/CollisionDebugDrawing.cpp
+++ b/Engine/Source/Runtime/Engine/Private/Collision/CollisionDebugDrawing.cpp
@@ -202,55 +202,6 @@ ENGINE_API void DrawCapsuleOverlap(const UWorld* InWorld,const FVector& Pos, con
 	}
 }
 
-#if PHYSICS_INTERFACE_PHYSX
-
-
-void DrawGeomOverlaps(const UWorld* InWorld, const PxGeometry& PGeom, const PxTransform& PGeomPose, TArray<struct FOverlapResult>& Overlaps, float Lifetime)
-{
-	FVector Pos = P2UVector(PGeomPose.p);
-	FQuat Rot = P2UQuat(PGeomPose.q);
-
-	if (PGeom.getType() == PxGeometryType::eBOX)
-	{
-		const PxBoxGeometry * Box = (PxBoxGeometry*)&PGeom;
-		DrawBoxOverlap(InWorld, Pos, P2UVector(Box->halfExtents), Rot, Overlaps, Lifetime);		
-	}
-	else if (PGeom.getType() == PxGeometryType::eSPHERE)
-	{
-		const PxSphereGeometry * Sphere = (PxSphereGeometry*)&PGeom;
-		DrawSphereOverlap(InWorld, Pos, Sphere->radius, Overlaps, Lifetime);		
-	}
-	else if (PGeom.getType() == PxGeometryType::eCAPSULE)
-	{
-		const PxCapsuleGeometry * Capsule = (PxCapsuleGeometry*)&PGeom;
-		// Convert here from PhysX to unreal definition of capsule height
-		DrawCapsuleOverlap(InWorld, Pos, Capsule->halfHeight + Capsule->radius, Capsule->radius, ConvertToUECapsuleRot(PGeomPose.q), Overlaps, Lifetime);		
-	}
-}
-
-void DrawGeomSweeps(const UWorld* InWorld, const FVector& Start, const FVector& End, const PxGeometry& PGeom, const PxQuat& PGeomRot, const TArray<FHitResult>& Hits, float Lifetime)
-{
-	if (PGeom.getType() == PxGeometryType::eBOX)
-	{
-		const PxBoxGeometry * Box = (PxBoxGeometry*)&PGeom;
-		DrawBoxSweeps(InWorld, Start, End, P2UVector(Box->halfExtents), P2UQuat(PGeomRot), Hits, Lifetime);		
-	}
-	else if (PGeom.getType() == PxGeometryType::eSPHERE)
-	{
-		const PxSphereGeometry * Sphere = (PxSphereGeometry *)&PGeom;
-		DrawSphereSweeps(InWorld, Start, End, Sphere->radius, Hits, Lifetime);
-	}
-	else if (PGeom.getType() == PxGeometryType::eCAPSULE)
-	{
-		const PxCapsuleGeometry * Capsule = (PxCapsuleGeometry*)&PGeom;
-		// Convert here from PhysX to unreal definition of capsule height
-		DrawCapsuleSweeps(InWorld, Start, End, Capsule->halfHeight + Capsule->radius, Capsule->radius, ConvertToUECapsuleRot(PGeomRot), Hits, Lifetime);		
-	}
-}
-
-
-#endif // WITH_PHYSX
-
 void DrawGeomOverlaps(const UWorld* InWorld, const Chaos::FImplicitObject& Geom, const FTransform& GeomPose, TArray<struct FOverlapResult>& Overlaps, float Lifetime)
 {
 	using namespace Chaos;
diff --git a/Engine/Source/Runtime/Engine/Private/Collision/CollisionDebugDrawing.h b/Engine/Source/Runtime/Engine/Private/Collision/CollisionDebugDrawing.h
--- a/Engine/Source/Runtime/Engine/Private/Collision/CollisionDebugDrawing.h
+++ b/Engine/Source/Runtime/Engine/Private/Collision/CollisionDebugDrawing.h
@@ -19,12 +19,5 @@ namespace Chaos
 	class FImplicitObject;
 }
 
-#if PHYSICS_INTERFACE_PHYSX
-/** Draw PhysX geom with overlaps */
-void DrawGeomOverlaps(const UWorld* InWorld, const PxGeometry& PGeom, const PxTransform& PGeomPose, TArray<struct FOverlapResult>& Overlaps, float Lifetime);
-/** Draw PhysX geom being swept with hits */
-void DrawGeomSweeps(const UWorld* InWorld, const FVector& Start, const FVector& End, const PxGeometry& PGeom, const PxQuat& Rotation, const TArray<FHitResult>& Hits, float Lifetime);
-#endif
-
 void DrawGeomOverlaps(const UWorld* InWorld, const Chaos::FImplicitObject& Geom, const FTransform& GeomPose, TArray<struct FOverlapResult>& Overlaps, float Lifetime);
 void DrawGeomSweeps(const UWorld* InWorld, const FVector& Start, const FVector& End, const Chaos::FImplicitObject& Geom, const FQuat& Rotation, const TArray<FHitResult>& Hits, float Lifetime);
diff --git a/Engine/Source/Runtime/Engine/Private/Collision/KAggregateGeom.cpp b/Engine/Source/Runtime/Engine/Private/Collision/KAggregateGeom.cpp
--- a/Engine/Source/Runtime/Engine/Private/Collision/KAggregateGeom.cpp
+++ b/Engine/Source/Runtime/Engine/Private/Collision/KAggregateGeom.cpp
@@ -21,9 +21,7 @@
 #include "KAggregateGeom.ispc.generated.h"
 #endif
 
-#if WITH_CHAOS
 #include "Chaos/ImplicitObject.h"
-#endif
 
 
 #define MIN_HULL_VERT_DISTANCE		(0.1f)
@@ -440,26 +438,6 @@ FBox FKConvexElem::CalcAABB(const FTransform& BoneTM, const FVector& Scale3D) co
 
 void FKConvexElem::GetPlanes(TArray<FPlane>& Planes) const
 {
-#if PHYSICS_INTERFACE_PHYSX
-	if (ConvexMesh != nullptr)
-	{
-		Planes.Empty();
-
-		PxU32 NumPolys = ConvexMesh->getNbPolygons();
-		for (PxU32 PolyIndex = 0; PolyIndex < NumPolys; PolyIndex++)
-		{
-			PxHullPolygon Data;
-			bool bStatus = ConvexMesh->getPolygonData(PolyIndex, Data);
-			check(bStatus);
-
-			// Convert to UE type
-			FPlane Plane(Data.mPlane[0], Data.mPlane[1], Data.mPlane[2], -Data.mPlane[3]);
-
-			// Add to output array
-			Planes.Add(Plane);
-		}
-	}
-#elif WITH_CHAOS
 	using FChaosPlane = Chaos::TPlaneConcrete<Chaos::FReal, 3>;
 	if(Chaos::FConvex* RawConvex = ChaosConvex.Get())
 	{
@@ -471,7 +449,6 @@ void FKConvexElem::GetPlanes(TArray<FPlane>& Planes) const
 			Planes.Add({Plane.X(), Plane.Normal()});
 		}
 	}
-#endif
 }
 
 ///////////////////////////////////////
@@ -752,13 +729,11 @@ static void AddEdgeIfNotPresent(TArray<int32>& Edges, int32 Edge0, int32 Edge1)
 
 void FKConvexElem::UpdateElemBox()
 {
-#if WITH_CHAOS
 	// Fixup indices in case an operation has invalidated them
 	{
 		IndexData.Reset();
 		ComputeChaosConvexIndices();
 	}
-#endif
 
 	ElemBox.Init();
 	for(int32 j=0; j<VertexData.Num(); j++)
@@ -987,14 +962,9 @@ FArchive& operator<<(FArchive& Ar,FKConvexElem& Elem)
 		// Initialize the TArray members
 		FMemory::Memzero(&Elem.VertexData, sizeof(Elem.VertexData));
 		FMemory::Memzero(&Elem.ElemBox, sizeof(Elem.ElemBox));
-#if PHYSICS_INTERFACE_PHYSX
-		Elem.ConvexMesh = NULL;
-		Elem.ConvexMeshNegX = NULL;
-#endif
-#if WITH_CHAOS
 		Elem.ChaosConvex.Reset();
-#endif
 	}
+
 	return Ar;
 }
 
diff --git a/Engine/Source/Runtime/Engine/Private/Collision/PhysicsFiltering.cpp b/Engine/Source/Runtime/Engine/Private/Collision/PhysicsFiltering.cpp
--- a/Engine/Source/Runtime/Engine/Private/Collision/PhysicsFiltering.cpp
+++ b/Engine/Source/Runtime/Engine/Private/Collision/PhysicsFiltering.cpp
@@ -3,7 +3,6 @@
 #include "Physics/PhysicsFiltering.h"
 #include "CollisionQueryParams.h"
 
-
 //////////////////////////////////////////////////////////////////////////
 // FPhysicsFilterBuilder
 
@@ -28,7 +27,3 @@ FPhysicsFilterBuilder::FPhysicsFilterBuilder(TEnumAsByte<enum ECollisionChannel>
 
 	Word3 = CreateChannelAndFilter(InObjectType, MaskFilter);
 }
-
-#if PHYSICS_INTERFACE_PHYSX
-PxSimulationFilterShader GSimulationFilterShader = nullptr;
-#endif // WITH_PHYSX
diff --git a/Engine/Source/Runtime/Engine/Private/Collision/SceneQuery.cpp b/Engine/Source/Runtime/Engine/Private/Collision/SceneQuery.cpp
--- a/Engine/Source/Runtime/Engine/Private/Collision/SceneQuery.cpp
+++ b/Engine/Source/Runtime/Engine/Private/Collision/SceneQuery.cpp
@@ -14,10 +14,6 @@
 #include "PhysicsEngine/ScopedSQHitchRepeater.h"
 #include "PhysicsInterfaceDeclaresCore.h"
 
-#if PHYSICS_INTERFACE_PHYSX
-#include "PhysXInterfaceWrapper.h"
-#endif
-
 #include "Collision/CollisionDebugDrawing.h"
 
 float DebugLineLifetime = 2.f;
@@ -164,7 +160,7 @@ struct TSQTraits
 	{
 		FQueryFilterData QueryFilterData = MakeQueryFilterData(FilterData, QueryFlags, Params);
 		FQueryDebugParams DebugParams;
-#if !(UE_BUILD_TEST || UE_BUILD_SHIPPING) && WITH_CHAOS
+#if !(UE_BUILD_TEST || UE_BUILD_SHIPPING)
 		DebugParams.bDebugQuery = Params.bDebugQuery;
 #endif
 		LowLevelRaycast(Scene, StartTM.GetLocation(), Dir, DeltaMag, HitBuffer, OutputFlags, QueryFlags, FilterData, QueryFilterData, QueryCallback, DebugParams);	//todo(ocohen): namespace?
@@ -176,7 +172,7 @@ struct TSQTraits
 	{
 		FQueryFilterData QueryFilterData = MakeQueryFilterData(FilterData, QueryFlags, Params);
 		FQueryDebugParams DebugParams;
-#if !(UE_BUILD_TEST || UE_BUILD_SHIPPING) && WITH_CHAOS
+#if !(UE_BUILD_TEST || UE_BUILD_SHIPPING)
 		DebugParams.bDebugQuery = Params.bDebugQuery;
 #endif
 		LowLevelSweep(Scene, *GeomInputs.GetGeometry(), StartTM, Dir, DeltaMag, HitBuffer, OutputFlags, QueryFlags, FilterData, QueryFilterData, QueryCallback, DebugParams);	//todo(ocohen): namespace?
@@ -202,11 +198,7 @@ struct TSQTraits
 		}
 		else
 		{
-#if PHYSICS_INTERFACE_PHYSX
-			DrawGeomSweeps(World, Start, End, *PGeom, U2PQuat(*PGeomRot), Hits, DebugLineLifetime);
-#else
 			DrawGeomSweeps(World, Start, End, *PGeom, *PGeomRot, Hits, DebugLineLifetime);
-#endif
 		}
 	}
 
@@ -597,7 +589,7 @@ bool GeomOverlapMultiImp(const UWorld* World, const FPhysicsGeometry& Geom, cons
 		FDynamicHitBuffer<TOverlapHit> OverlapBuffer;
 
 		FQueryDebugParams DebugParams;
-#if !(UE_BUILD_TEST || UE_BUILD_SHIPPING) && WITH_CHAOS
+#if !(UE_BUILD_TEST || UE_BUILD_SHIPPING)
 		DebugParams.bDebugQuery = Params.bDebugQuery;
 #endif
 
@@ -627,13 +619,6 @@ bool GeomOverlapMultiImp(const UWorld* World, const FPhysicsGeometry& Geom, cons
 				{
 					bHaveBlockingHit = ConvertOverlapResults(NumHits, OverlapBuffer.GetHits(), Filter, OutOverlaps);
 				}
-
-#if (!(UE_BUILD_SHIPPING || UE_BUILD_TEST) && !WITH_CHAOS)
-				if(World->DebugDrawSceneQueries(Params.TraceTag))
-				{
-					DrawGeomOverlaps(World, Geom, U2PTransform(GeomPose), OutOverlaps, DebugLineLifetime);
-				}
-#endif //!(UE_BUILD_SHIPPING || UE_BUILD_TEST)
 			}
 
 		});
diff --git a/Engine/Source/Runtime/Engine/Private/Collision/SceneQueryLowLevel.cpp b/Engine/Source/Runtime/Engine/Private/Collision/SceneQueryLowLevel.cpp
--- a/Engine/Source/Runtime/Engine/Private/Collision/SceneQueryLowLevel.cpp
+++ b/Engine/Source/Runtime/Engine/Private/Collision/SceneQueryLowLevel.cpp
@@ -1,10 +1,5 @@
 // Copyright Epic Games, Inc. All Rights Reserved.
 
-#if PHYSICS_INTERFACE_PHYSX
-#include "PhysXPublic.h"
-#include "PhysXInterfaceWrapper.h"
-#endif
-
 #include "Physics/PhysicsInterfaceDeclares.h"
 #include "Physics/PhysicsInterfaceCore.h"
 #include "PhysicsInterfaceDeclaresCore.h"
@@ -75,8 +70,6 @@ constexpr int32 EnableSweepSQCapture = 0;
 void FinalizeCapture(FPhysTestSerializer& Serializer) {}
 #endif
 
-
-#if !defined(PHYSICS_INTERFACE_PHYSX) || !PHYSICS_INTERFACE_PHYSX
 namespace
 {
 	void SweepSQCaptureHelper(float QueryDurationSeconds, const FChaosSQAccelerator& SQAccelerator, FPhysScene& Scene, const FPhysicsGeometry& QueryGeom, const FTransform& StartTM, const FVector& Dir, float DeltaMag, const FPhysicsHitCallback<FHitSweep>& HitBuffer, EHitFlags OutputFlags, FQueryFlags QueryFlags, const FCollisionFilterData& Filter, const FQueryFilterData& QueryFilterData, ICollisionQueryFilterCallbackBase* QueryCallback, const FQueryDebugParams& DebugParams)
@@ -181,12 +174,10 @@ namespace
 		}
 	}
 }
-#endif
 
 template <typename THitRaycast>
 void LowLevelRaycastImp(FPhysScene& Scene, const FVector& Start, const FVector& Dir, float DeltaMag, FPhysicsHitCallback<THitRaycast>& HitBuffer, EHitFlags OutputFlags, FQueryFlags QueryFlags, const FCollisionFilterData& Filter, const FQueryFilterData& QueryFilterData, ICollisionQueryFilterCallbackBase* QueryCallback, const FQueryDebugParams& DebugParams)
 {
-#if !defined(PHYSICS_INTERFACE_PHYSX) || !PHYSICS_INTERFACE_PHYSX
 	constexpr bool bGTData = std::is_same<THitRaycast, FHitRaycast>::value;
 	const auto SolverAccelerationStructure = bGTData ? Scene.GetSpacialAcceleration() : Scene.GetSolver()->GetInternalAccelerationStructure_Internal();
 	if (SolverAccelerationStructure)
@@ -206,23 +197,6 @@ void LowLevelRaycastImp(FPhysScene& Scene, const FVector& Start, const FVector&
 			}
 		}
 	}
-#else
-	if (!!SerializeSQs || !!ReplaySQs)
-	{
-		FPhysTestSerializer Serializer;
-		Serializer.SetPhysicsData(*Scene.GetPxScene());
-		FSQCapture& SweepCapture = Serializer.CaptureSQ();
-		SweepCapture.StartCapturePhysXRaycast(*Scene.GetPxScene(), Start, Dir, DeltaMag, OutputFlags, QueryFilterData, Filter, *QueryCallback);
-		Scene.GetPxScene()->raycast(U2PVector(Start), U2PVector(Dir), DeltaMag, HitBuffer, U2PHitFlags(OutputFlags), QueryFilterData, QueryCallback);
-		SweepCapture.EndCapturePhysXRaycast(HitBuffer);
-
-		FinalizeCapture(Serializer);
-	}
-	else
-	{
-		Scene.GetPxScene()->raycast(U2PVector(Start), U2PVector(Dir), DeltaMag, HitBuffer, U2PHitFlags(OutputFlags), QueryFilterData, QueryCallback);
-	}
-#endif
 }
 
 void LowLevelRaycast(FPhysScene& Scene, const FVector& Start, const FVector& Dir, float DeltaMag, FPhysicsHitCallback<FHitRaycast>& HitBuffer, EHitFlags OutputFlags, FQueryFlags QueryFlags, const FCollisionFilterData& Filter, const FQueryFilterData& QueryFilterData, ICollisionQueryFilterCallbackBase* QueryCallback, const FQueryDebugParams& DebugParams)
@@ -238,7 +212,6 @@ void LowLevelRaycast(FPhysScene& Scene, const FVector& Start, const FVector& Dir
 template <typename THitSweep>
 void LowLevelSweepImp(FPhysScene& Scene, const FPhysicsGeometry& QueryGeom, const FTransform& StartTM, const FVector& Dir, float DeltaMag, FPhysicsHitCallback<THitSweep>& HitBuffer, EHitFlags OutputFlags, FQueryFlags QueryFlags, const FCollisionFilterData& Filter, const FQueryFilterData& QueryFilterData, ICollisionQueryFilterCallbackBase* QueryCallback, const FQueryDebugParams& DebugParams)
 {
-#if !defined(PHYSICS_INTERFACE_PHYSX) || !PHYSICS_INTERFACE_PHYSX
 	constexpr bool bGTData = std::is_same<THitSweep, FHitSweep>::value;
 	const auto SolverAccelerationStructure = bGTData ? Scene.GetSpacialAcceleration() : Scene.GetSolver()->GetInternalAccelerationStructure_Internal();
 	if (SolverAccelerationStructure)
@@ -261,23 +234,6 @@ void LowLevelSweepImp(FPhysScene& Scene, const FPhysicsGeometry& QueryGeom, cons
 			}
 		}
 	}
-#else
-	if (!!SerializeSQs || !!ReplaySQs)
-	{
-		FPhysTestSerializer Serializer;
-		Serializer.SetPhysicsData(*Scene.GetPxScene());
-		FSQCapture& SweepCapture = Serializer.CaptureSQ();
-		SweepCapture.StartCapturePhysXSweep(*Scene.GetPxScene(), QueryGeom, StartTM, Dir, DeltaMag, OutputFlags, QueryFilterData, Filter, *QueryCallback);
-		Scene.GetPxScene()->sweep(QueryGeom, U2PTransform(StartTM), U2PVector(Dir), DeltaMag, HitBuffer, U2PHitFlags(OutputFlags), QueryFilterData, QueryCallback);
-		SweepCapture.EndCapturePhysXSweep(HitBuffer);
-
-		FinalizeCapture(Serializer);
-	}
-	else
-	{
-		Scene.GetPxScene()->sweep(QueryGeom, U2PTransform(StartTM), U2PVector(Dir), DeltaMag, HitBuffer, U2PHitFlags(OutputFlags), QueryFilterData, QueryCallback);
-	}
-#endif
 }
 
 void LowLevelSweep(FPhysScene& Scene, const FPhysicsGeometry& QueryGeom, const FTransform& StartTM, const FVector& Dir, float DeltaMag, FPhysicsHitCallback<FHitSweep>& HitBuffer, EHitFlags OutputFlags, FQueryFlags QueryFlags, const FCollisionFilterData& Filter, const FQueryFilterData& QueryFilterData, ICollisionQueryFilterCallbackBase* QueryCallback, const FQueryDebugParams& DebugParams)
@@ -293,7 +249,6 @@ void LowLevelSweep(FPhysScene& Scene, const FPhysicsGeometry& QueryGeom, const F
 template <typename THitOverlap>
 void LowLevelOverlapImp(FPhysScene& Scene, const FPhysicsGeometry& QueryGeom, const FTransform& GeomPose, FPhysicsHitCallback<THitOverlap>& HitBuffer, FQueryFlags QueryFlags, const FCollisionFilterData& Filter, const FQueryFilterData& QueryFilterData, ICollisionQueryFilterCallbackBase* QueryCallback, const FQueryDebugParams& DebugParams)
 {
-#if !defined(PHYSICS_INTERFACE_PHYSX) || !PHYSICS_INTERFACE_PHYSX
 	constexpr bool bGTData = std::is_same<THitOverlap, FHitOverlap>::value;
 	const auto SolverAccelerationStructure = bGTData ? Scene.GetSpacialAcceleration() : Scene.GetSolver()->GetInternalAccelerationStructure_Internal();
 	if (SolverAccelerationStructure)
@@ -313,23 +268,6 @@ void LowLevelOverlapImp(FPhysScene& Scene, const FPhysicsGeometry& QueryGeom, co
 			}
 		}
 	}
-#else
-	if (!!SerializeSQs || !!ReplaySQs)
-	{
-		FPhysTestSerializer Serializer;
-		Serializer.SetPhysicsData(*Scene.GetPxScene());
-		FSQCapture& SweepCapture = Serializer.CaptureSQ();
-		SweepCapture.StartCapturePhysXOverlap(*Scene.GetPxScene(), QueryGeom, GeomPose, QueryFilterData, Filter, *QueryCallback);
-		Scene.GetPxScene()->overlap(QueryGeom, U2PTransform(GeomPose), HitBuffer, QueryFilterData, QueryCallback);
-		SweepCapture.EndCapturePhysXOverlap(HitBuffer);
-
-		FinalizeCapture(Serializer);
-	}
-	else
-	{
-		Scene.GetPxScene()->overlap(QueryGeom, U2PTransform(GeomPose), HitBuffer, QueryFilterData, QueryCallback);
-	}		
-#endif
 }
 
 void LowLevelOverlap(FPhysScene& Scene, const FPhysicsGeometry& QueryGeom, const FTransform& GeomPose, FPhysicsHitCallback<FHitOverlap>& HitBuffer, FQueryFlags QueryFlags, const FCollisionFilterData& Filter, const FQueryFilterData& QueryFilterData, ICollisionQueryFilterCallbackBase* QueryCallback, const FQueryDebugParams& DebugParams)
@@ -341,5 +279,3 @@ void LowLevelOverlap(FPhysScene& Scene, const FPhysicsGeometry& QueryGeom, const
 {
 	return LowLevelOverlapImp(Scene, QueryGeom, GeomPose, HitBuffer, QueryFlags, Filter, QueryFilterData, QueryCallback, DebugParams);
 }
-
-//#endif // WITH_PHYSX 
\ No newline at end of file
diff --git a/Engine/Source/Runtime/Engine/Private/Collision/WorldCollision.cpp b/Engine/Source/Runtime/Engine/Private/Collision/WorldCollision.cpp
--- a/Engine/Source/Runtime/Engine/Private/Collision/WorldCollision.cpp
+++ b/Engine/Source/Runtime/Engine/Private/Collision/WorldCollision.cpp
@@ -15,10 +15,7 @@
 #include "Physics/PhysicsInterfaceCore.h"
 #include "PhysXPublic.h"
 #include "Physics/PhysicsInterfaceTypes.h"
-
-#if WITH_CHAOS
 #include "Chaos/ImplicitObject.h"
-#endif
 
 using namespace PhysicsInterfaceTypes;
 
@@ -399,19 +396,11 @@ bool UWorld::ComponentSweepMulti(TArray<struct FHitResult>& OutHits, class UPrim
 			check(Shape.IsValid());
 			ECollisionShapeType ShapeType = FPhysicsInterface::GetShapeType(Shape);
 
-#if WITH_CHAOS
 			if (!Shape.GetGeometry().IsConvex())
 			{
 				//We skip complex shapes. Should this respect complex as simple?
 				continue;
 			}
-#else
-			if(ShapeType == ECollisionShapeType::Heightfield || ShapeType == ECollisionShapeType::Trimesh)
-			{
-				//We skip complex shapes. Should this respect complex as simple?
-				continue;
-			}
-#endif
 
 			// Calc shape global pose
 			const FTransform ShapeLocalTransform = FPhysicsInterface::GetLocalTransform(Shape);
diff --git a/Engine/Source/Runtime/Engine/Private/Components/ActorComponent.cpp b/Engine/Source/Runtime/Engine/Private/Components/ActorComponent.cpp
--- a/Engine/Source/Runtime/Engine/Private/Components/ActorComponent.cpp
+++ b/Engine/Source/Runtime/Engine/Private/Components/ActorComponent.cpp
@@ -86,8 +86,6 @@ static FAutoConsoleVariableRef CVarDefaultUseSubObjectReplicationList(
 	GDefaultUseSubObjectReplicationList,
 	TEXT("Do actors and actorcomponents replicate subobjects using the registration method by default."));
 
-
-#if WITH_CHAOS
 // Allows for CreatePhysicsState to be deferred, to batch work and parallelize.
 int32 GEnableDeferredPhysicsCreation = 0;
 FAutoConsoleVariableRef CVarEnableDeferredPhysicsCreation(
@@ -95,9 +93,6 @@ FAutoConsoleVariableRef CVarEnableDeferredPhysicsCreation(
 	GEnableDeferredPhysicsCreation,
 	TEXT("Enables/Disables deferred physics creation.")
 );
-#else
-int32 GEnableDeferredPhysicsCreation = 0;
-#endif
 
 void FRegisterComponentContext::Process()
 {
@@ -1230,7 +1225,6 @@ void UActorComponent::RegisterAllComponentTickFunctions(bool bRegister)
 
 void UActorComponent::RegisterAsyncPhysicsTickEnabled(bool bRegister)
 {
-#if WITH_CHAOS
 	if (FPhysScene_Chaos* Scene = static_cast<FPhysScene_Chaos*>(WorldPrivate->GetPhysicsScene()))
 	{
 		if (bRegister)
@@ -1242,12 +1236,10 @@ void UActorComponent::RegisterAsyncPhysicsTickEnabled(bool bRegister)
 			Scene->UnregisterAsyncPhysicsTickComponent(this);
 		}
 	}
-#endif
 }
 
 void UActorComponent::SetAsyncPhysicsTickEnabled(bool bEnable)
 {
-#if WITH_CHAOS
 	// Components don't have async physics functions until they are registered with the world
 	if(bRegistered)
 	{
@@ -1255,7 +1247,6 @@ void UActorComponent::SetAsyncPhysicsTickEnabled(bool bEnable)
 	}
 	
 	bAsyncPhysicsTickEnabled = bEnable;
-#endif
 }
 
 void UActorComponent::TickComponent(float DeltaTime, enum ELevelTick TickType, FActorComponentTickFunction *ThisTickFunction)
@@ -1599,11 +1590,7 @@ void UActorComponent::OnDestroyPhysicsState()
 
 void UActorComponent::CreatePhysicsState(bool bAllowDeferral)
 {
-#if WITH_CHAOS
 	LLM_SCOPE(ELLMTag::Chaos);
-#else
-	LLM_SCOPE(ELLMTag::PhysX);
-#endif
 
 	SCOPE_CYCLE_COUNTER(STAT_ComponentCreatePhysicsState);
 
@@ -1632,11 +1619,7 @@ void UActorComponent::CreatePhysicsState(bool bAllowDeferral)
 
 		if (ShouldDefer)
 		{
-#if WITH_CHAOS
 			WorldPrivate->GetPhysicsScene()->DeferPhysicsStateCreation(Primitive);
-#else
-			check(false);
-#endif
 		}
 		else
 		{
@@ -1671,16 +1654,12 @@ void UActorComponent::DestroyPhysicsState()
 	}
 	else if(GEnableDeferredPhysicsCreation)
 	{
-#if WITH_CHAOS
 		UPrimitiveComponent* PrimitiveComponent = Cast<UPrimitiveComponent>(this);
 		if (PrimitiveComponent && PrimitiveComponent->DeferredCreatePhysicsStateScene != nullptr)
 		{
 			// We had to cache this scene because World ptr is null as we have unregistered already.
 			PrimitiveComponent->DeferredCreatePhysicsStateScene->RemoveDeferredPhysicsStateCreation(PrimitiveComponent);
 		}
-#else
-		check(false);
-#endif
 	}
 }
 
diff --git a/Engine/Source/Runtime/Engine/Private/Components/ShapeComponent.cpp b/Engine/Source/Runtime/Engine/Private/Components/ShapeComponent.cpp
--- a/Engine/Source/Runtime/Engine/Private/Components/ShapeComponent.cpp
+++ b/Engine/Source/Runtime/Engine/Private/Components/ShapeComponent.cpp
@@ -9,10 +9,6 @@
 #include "PhysicsEngine/SphylElem.h"
 #include "PhysicsEngine/BodySetup.h"
 
-#if WITH_PHYSX
-	#include "PhysXPublic.h"
-#endif // WITH_PHYSX
-
 // Custom serialization version
 struct FShapeComponentCustomVersion
 {
diff --git a/Engine/Source/Runtime/Engine/Private/Components/SkeletalMeshComponent.cpp b/Engine/Source/Runtime/Engine/Private/Components/SkeletalMeshComponent.cpp
--- a/Engine/Source/Runtime/Engine/Private/Components/SkeletalMeshComponent.cpp
+++ b/Engine/Source/Runtime/Engine/Private/Components/SkeletalMeshComponent.cpp
@@ -224,7 +224,6 @@ USkeletalMeshComponent::USkeletalMeshComponent(const FObjectInitializer& ObjectI
 	bWaitForParallelClothTask = false;
 	bNotifySyncComponentToRBPhysics = false;
 
-#if WITH_APEX_CLOTHING || WITH_CHAOS_CLOTHING
 	ClothMaxDistanceScale = 1.0f;
 	bResetAfterTeleport = true;
 	TeleportDistanceThreshold = 300.0f;
@@ -240,8 +239,6 @@ USkeletalMeshComponent::USkeletalMeshComponent(const FObjectInitializer& ObjectI
 
 	bBindClothToMasterComponent = false;
 	bClothingSimulationSuspended = false;
-	
-#endif//#if WITH_APEX_CLOTHING || WITH_CHAOS_CLOTHING
 
 #if WITH_EDITORONLY_DATA
 	DefaultPlayRate_DEPRECATED = 1.0f;
@@ -591,7 +588,6 @@ void USkeletalMeshComponent::OnRegister()
 		SetComponentTickEnabled(false);
 	}
 
-#if WITH_APEX_CLOTHING || WITH_CHAOS_CLOTHING
 	// If no simulation factory is currently set, set it to the default factory
 	if (!ClothingSimulationFactory)
 	{
@@ -675,7 +671,6 @@ void USkeletalMeshComponent::OnRegister()
 	}
 
 	RecreateClothingActors();
-#endif  // #if WITH_APEX_CLOTHING || WITH_CHAOS_CLOTHING
 }
 
 void USkeletalMeshComponent::OnUnregister()
@@ -958,12 +953,8 @@ bool USkeletalMeshComponent::InitializeAnimScriptInstance(bool bForceReinit, boo
 
 bool USkeletalMeshComponent::IsWindEnabled() const
 {
-#if WITH_APEX_CLOTHING || WITH_CHAOS_CLOTHING
 	// Wind is enabled in game worlds
 	return GetWorld() && GetWorld()->IsGameWorld();
-#else
-	return false;
-#endif
 }
 
 void USkeletalMeshComponent::ClearAnimScriptInstance()
@@ -2904,9 +2895,7 @@ FBoxSphereBounds USkeletalMeshComponent::CalcBounds(const FTransform& LocalToWor
 			NewBounds = NewBounds + FBoxSphereBounds(ComponentLocation, FVector(1.0f), 1.0f);
 		}
 
-#if WITH_APEX_CLOTHING || WITH_CHAOS_CLOTHING
 		AddClothingBounds(NewBounds, CachedBoundsTransform);
-#endif// #if WITH_APEX_CLOTHING || WITH_CHAOS_CLOTHING
 
 		CachedWorldOrLocalSpaceBounds = NewBounds;
 		bCachedLocalBoundsUpToDate = bCacheLocalSpaceBounds;
@@ -2975,9 +2964,7 @@ void USkeletalMeshComponent::SetSkeletalMesh(USkeletalMesh* InSkelMesh, bool bRe
 		
 		InitAnim(bReinitPose);
 
-#if WITH_APEX_CLOTHING || WITH_CHAOS_CLOTHING
 		RecreateClothingActors();
-#endif
 	}
 
 	// Mark cached material parameter names dirty
diff --git a/Engine/Source/Runtime/Engine/Private/Components/SplineMeshComponent.cpp b/Engine/Source/Runtime/Engine/Private/Components/SplineMeshComponent.cpp
--- a/Engine/Source/Runtime/Engine/Private/Components/SplineMeshComponent.cpp
+++ b/Engine/Source/Runtime/Engine/Private/Components/SplineMeshComponent.cpp
@@ -940,21 +940,12 @@ void USplineMeshComponent::OnCreatePhysicsState()
 
 UBodySetup* USplineMeshComponent::GetBodySetup()
 {
-#if PHYSICS_INTERFACE_PHYSX
-	// Don't return a body setup that has no collision, it means we are interactively moving the spline and don't want to build collision.
-	// Instead we explicitly build collision with USplineMeshComponent::RecreateCollision()
-	if (BodySetup != NULL && (BodySetup->TriMeshes.Num() || BodySetup->AggGeom.GetElementCount() > 0))
-	{
-		return BodySetup;
-	}
-#elif WITH_CHAOS
 	// Don't return a body setup that has no collision, it means we are interactively moving the spline and don't want to build collision.
 	// Instead we explicitly build collision with USplineMeshComponent::RecreateCollision()
 	if (BodySetup != NULL && (BodySetup->ChaosTriMeshes.Num() || BodySetup->AggGeom.GetElementCount() > 0))
 	{
 		return BodySetup;
 	}
-#endif // WITH_CHAOS
 
 	return NULL;
 }
diff --git a/Engine/Source/Runtime/Engine/Private/InstancedStaticMesh.cpp b/Engine/Source/Runtime/Engine/Private/InstancedStaticMesh.cpp
--- a/Engine/Source/Runtime/Engine/Private/InstancedStaticMesh.cpp
+++ b/Engine/Source/Runtime/Engine/Private/InstancedStaticMesh.cpp
@@ -2459,11 +2459,9 @@ void UInstancedStaticMeshComponent::InitInstanceBody(int32 InstanceIdx, FBodyIns
 	// make sure we never enable bSimulatePhysics for ISMComps
 	InstanceBodyInstance->bSimulatePhysics = false;
 
-#if WITH_PHYSX
 	// Create physics body instance.
 	InstanceBodyInstance->bAutoWeld = false;	//We don't support this for instanced meshes.
 	InstanceBodyInstance->InitBody(BodySetup, InstanceTransform, this, GetWorld()->GetPhysicsScene(), nullptr);
-#endif //WITH_PHYSX
 }
 
 void UInstancedStaticMeshComponent::CreateAllInstanceBodies()
@@ -3378,7 +3376,7 @@ void UInstancedStaticMeshComponent::UpdateInstanceBodyTransform(int32 InstanceIn
 	check(bPhysicsStateCreated);
 
 	FBodyInstance*& InstanceBodyInstance = InstanceBodies[InstanceIndex];
-#if WITH_PHYSX
+
 	if (WorldSpaceInstanceTransform.GetScale3D().IsNearlyZero())
 	{
 		if (InstanceBodyInstance)
@@ -3404,7 +3402,6 @@ void UInstancedStaticMeshComponent::UpdateInstanceBodyTransform(int32 InstanceIn
 			InitInstanceBody(InstanceIndex, InstanceBodyInstance);
 		}
 	}
-#endif //WITH_PHYSX
 }
 
 bool UInstancedStaticMeshComponent::UpdateInstanceTransform(int32 InstanceIndex, const FTransform& NewInstanceTransform, bool bWorldSpace, bool bMarkRenderStateDirty, bool bTeleport)
diff --git a/Engine/Source/Runtime/Engine/Private/LevelTick.cpp b/Engine/Source/Runtime/Engine/Private/LevelTick.cpp
--- a/Engine/Source/Runtime/Engine/Private/LevelTick.cpp
+++ b/Engine/Source/Runtime/Engine/Private/LevelTick.cpp
@@ -1662,14 +1662,6 @@ void UWorld::Tick( ELevelTick TickType, float DeltaSeconds )
 
 		FWorldDelegates::OnWorldPostActorTick.Broadcast(this, TickType, DeltaSeconds);
 
-#if PHYSICS_INTERFACE_PHYSX
-		if ( PhysicsScene != nullptr )
-		{
-			TRACE_CPUPROFILER_EVENT_SCOPE(GPhysCommandHandler);
-			GPhysCommandHandler->Flush();
-		}
-#endif // WITH_PHYSX
-		
 		// All tick is done, execute async trace
 		{
 			TRACE_CPUPROFILER_EVENT_SCOPE(FinishAsyncTrace);
diff --git a/Engine/Source/Runtime/Engine/Private/LocalPlayer.cpp b/Engine/Source/Runtime/Engine/Private/LocalPlayer.cpp
--- a/Engine/Source/Runtime/Engine/Private/LocalPlayer.cpp
+++ b/Engine/Source/Runtime/Engine/Private/LocalPlayer.cpp
@@ -1460,7 +1460,6 @@ bool ULocalPlayer::Exec(UWorld* InWorld, const TCHAR* Cmd,FOutputDevice& Ar)
 
 		return true;
 	}
-#if WITH_PHYSX
 	// This will list all awake rigid bodies
 	else if( FParse::Command(&Cmd,TEXT("LISTAWAKEBODIES")) )
 	{
@@ -1471,7 +1470,6 @@ bool ULocalPlayer::Exec(UWorld* InWorld, const TCHAR* Cmd,FOutputDevice& Ar)
 	{
 		return HandleListSimBodiesCommand( Cmd, Ar );
 	}
-#endif
 	else if( FParse::Command(&Cmd, TEXT("MOVECOMPTIMES")) )
 	{
 		return HandleMoveComponentTimesCommand( Cmd, Ar );
