Physics interface cleanup.

diff --git a/Engine/Source/Runtime/Engine/Private/PhysicsEngine/BodyInstance.cpp b/Engine/Source/Runtime/Engine/Private/PhysicsEngine/BodyInstance.cpp
--- a/Engine/Source/Runtime/Engine/Private/PhysicsEngine/BodyInstance.cpp
+++ b/Engine/Source/Runtime/Engine/Private/PhysicsEngine/BodyInstance.cpp
@@ -42,19 +42,8 @@
 #include "Chaos/Utilities.h"
 #include "Physics/Experimental/ChaosInterfaceUtils.h"
 
-#if WITH_CHAOS
 #include "Chaos/ParticleHandle.h"
 #include "PhysicsProxy/SingleParticlePhysicsProxy.h"
-#endif // WITH_CHAOS
-
-#if PHYSICS_INTERFACE_PHYSX
-	#include "PhysXPublic.h"
-	#include "Physics/PhysicsFiltering.h"
-	#include "PhysicsEngine/PhysXSupport.h"
-	#include "Collision/CollisionConversions.h"
-#include "PxShape.h"
-#endif // WITH_PHYSX
-
 
 #define LOCTEXT_NAMESPACE "BodyInstance"
 
@@ -329,9 +318,7 @@ FBodyInstance::FBodyInstance()
 	, ObjectType(ECC_WorldStatic)
 	, MaskFilter(0)
 	, CollisionEnabled(ECollisionEnabled::QueryAndPhysics)
-#if WITH_PHYSX
 	, CurrentSceneState(BodyInstanceSceneState::NotAdded)
-#endif // WITH_PHYSX
 	, SleepFamily(ESleepFamily::Normal)
 	, DOFMode(0)
 	, bUseCCD(false)
@@ -456,10 +443,6 @@ void FBodyInstance::UpdateTriMeshVertices(const TArray<FVector> & NewPositions)
 	{
 		FPhysicsCommand::ExecuteWrite(ActorHandle, [&](const FPhysicsActorHandle& Actor)
 		{
-#if PHYSICS_INTERFACE_PHYSX // Chaos doesn't modify vertices on body setup.
-			GetBodySetup()->UpdateTriMeshVertices(NewPositions);
-#endif
-
 			//after updating the vertices we must call setGeometry again to update any shapes referencing the mesh
 			TArray<FPhysicsShapeHandle> Shapes;
 			const int32 SyncShapeCount = GetAllShapes_AssumesLocked(Shapes);
@@ -468,7 +451,6 @@ void FBodyInstance::UpdateTriMeshVertices(const TArray<FVector> & NewPositions)
 			{
 				if (FPhysicsInterface::GetShapeType(Shape) == ECollisionShapeType::Trimesh)
 				{
-#if WITH_CHAOS
 					using namespace Chaos;
 					const Chaos::FImplicitObject* ShapeImplicit = Shape.Shape->GetGeometry().Get();
 					EImplicitObjectType Type = ShapeImplicit->GetType();
@@ -506,10 +488,6 @@ void FBodyInstance::UpdateTriMeshVertices(const TArray<FVector> & NewPositions)
 					{
 						FPhysicsInterface::SetGeometry(Shape, MoveTemp(TriMeshCopy));
 					}
-#else
-					FPhysicsGeometryCollection GeoCollection = FPhysicsInterface::GetGeometryCollection(Shape);
-					FPhysicsInterface::SetGeometry(Shape, GeoCollection.GetTriMeshGeometry());
-#endif
 				}
 			}
 		});
@@ -1082,7 +1060,6 @@ void FBodyInstance::UpdatePhysicsFilterData()
 			UpdateMassProperties();
 		}
 
-#if WITH_CHAOS
 		//If filtering changed we must update GT structure right away
 		if (FPhysScene* PhysScene = GetPhysicsScene())
 		{
@@ -1090,7 +1067,6 @@ void FBodyInstance::UpdatePhysicsFilterData()
 		}
 		// Always wake actors up when collision filters change
 		FPhysicsInterface::WakeUp_AssumesLocked(Actor);
-#endif
 	});
 
 	UpdateInterpolateWhenSubStepping();
@@ -1219,11 +1195,10 @@ void FInitBodiesHelperBase::CreateActor_AssumesLocked(FBodyInstance* Instance, c
 		FPhysicsInterface::CreateActor(ActorParams, Instance->ActorHandle);
 		FPhysicsInterface::SetCcdEnabled_AssumesLocked(Instance->ActorHandle, Instance->bUseCCD);
 		FPhysicsInterface::SetIsKinematic_AssumesLocked(Instance->ActorHandle, !Instance->ShouldInstanceSimulatingPhysics());
-#if WITH_CHAOS
+
 		FPhysicsInterface::SetMaxLinearVelocity_AssumesLocked(Instance->ActorHandle, TNumericLimits<float>::Max());
 		FPhysicsInterface::SetSmoothEdgeCollisionsEnabled_AssumesLocked(Instance->ActorHandle, Instance->bSmoothEdgeCollisions);
 		FPhysicsInterface::SetInertiaConditioningEnabled_AssumesLocked(Instance->ActorHandle, Instance->bInertiaConditioning);
-#endif
 
 		// Set sleep even notification
 		FPhysicsInterface::SetSendsSleepNotifies_AssumesLocked(Instance->ActorHandle, Instance->bGenerateWakeEvents);
@@ -1237,11 +1212,7 @@ bool FInitBodiesHelperBase::CreateShapes_AssumesLocked(FBodyInstance* Instance)
 	TArray<UPhysicalMaterial*> ComplexPhysMats;
 	TArray<FPhysicalMaterialMaskParams> ComplexPhysMatMasks;
 
-#if WITH_CHAOS
 	ComplexPhysMats = Instance->GetComplexPhysicalMaterials(ComplexPhysMatMasks);
-#else
-	ComplexPhysMats = Instance->GetComplexPhysicalMaterials();
-#endif
 
 	FBodyCollisionData BodyCollisionData;
 	Instance->BuildBodyFilterData(BodyCollisionData.CollisionFilterData);
@@ -1252,9 +1223,7 @@ bool FInitBodiesHelperBase::CreateShapes_AssumesLocked(FBodyInstance* Instance)
 	// #PHYS2 Call interface AddGeometry
 	BodySetup->AddShapesToRigidActor_AssumesLocked(Instance, Instance->Scale3D, SimplePhysMat, ComplexPhysMats, ComplexPhysMatMasks, BodyCollisionData, FTransform::Identity);
 
-#if WITH_CHAOS
 	FPhysicsInterface::SetIgnoreAnalyticCollisions_AssumesLocked(Instance->ActorHandle, CVarIgnoreAnalyticCollisionsOverride.GetValueOnGameThread() ? true : Instance->bIgnoreAnalyticCollisions);
-#endif
 
 	const int32 NumShapes = FPhysicsInterface::GetNumShapes(Instance->ActorHandle);
 	bInitFail |= NumShapes == 0;
@@ -1370,7 +1339,8 @@ bool FInitBodiesHelperBase::CreateShapesAndActors()
 		}
 
 		FPhysicsInterface::SetActorUserData_AssumesLocked(Instance->ActorHandle, &Instance->PhysicsUserData);
-#if WITH_CHAOS && USE_BODYINSTANCE_DEBUG_NAMES
+
+#if USE_BODYINSTANCE_DEBUG_NAMES
 		Instance->ActorHandle->GetParticle_LowLevel()->SetDebugName(DebugName);
 #endif
 	}
@@ -1379,11 +1349,7 @@ bool FInitBodiesHelperBase::CreateShapesAndActors()
 
 void FInitBodiesHelperBase::InitBodies()
 {
-#if WITH_CHAOS
 	LLM_SCOPE(ELLMTag::ChaosBody);
-#else
-	LLM_SCOPE(ELLMTag::PhysX);
-#endif
 
 	//check(IsInGameThread());
 
@@ -1417,7 +1383,6 @@ void FInitBodiesHelperBase::InitBodies()
 					{
 						ActorHandles.Add(ActorHandle);
 
-#if WITH_CHAOS
 						Chaos::FRigidBodyHandle_External& Body_External = ActorHandle->GetGameThreadAPI();
 						const int32 NumShapes = FPhysicsInterface::GetNumShapes(ActorHandle);
 
@@ -1437,29 +1402,16 @@ void FInitBodiesHelperBase::InitBodies()
 								Body_External.SetShapeCollisionTraceType(ShapeIndex, ChaosInterface::ConvertCollisionTraceFlag(BI->BodySetup->CollisionTraceFlag)) ;
 							}
 						}
-
-#endif
-/*
-						With the implementation of AddActorsToScene_AssumesLocked,
-						this call duplicates the AddToSolver operation and is not necessary
-#ifdef WITH_CHAOS
-						if (BI->bSimulatePhysics)
-						{
-							PhysScene->AddActorToSolver(ActorHandle);
-						}
-#endif
-*/
 					}
 				}
 
 				PhysScene->AddActorsToScene_AssumesLocked(ActorHandles);
-#if WITH_CHAOS
+
 				for (FBodyInstance* BI : Bodies)
 				{
 					FPhysicsActorHandle& ActorHandle = BI->GetPhysicsActorHandle();
 					if (FPhysicsInterface::IsValid(ActorHandle))
 					{
-
 						PhysScene->AddToComponentMaps(BI->OwnerComponent.Get(), ActorHandle);
 					}
 					if (BI->bNotifyRigidBodyCollision)
@@ -1471,7 +1423,6 @@ void FInitBodiesHelperBase::InitBodies()
 						}
 					}
 				}
-#endif
 			}
 
 			// Set up dynamic instance data
@@ -1587,7 +1538,6 @@ void FBodyInstance::TermBody(bool bNeverDeferRelease)
 {
 	SCOPE_CYCLE_COUNTER(STAT_TermBody);
 
-#if WITH_CHAOS
 	if (UPrimitiveComponent* PrimComp = OwnerComponent.Get())
 	{
 		if (FPhysScene_Chaos* PhysScene = PrimComp->GetWorld()->GetPhysicsScene())
@@ -1602,7 +1552,6 @@ void FBodyInstance::TermBody(bool bNeverDeferRelease)
 			}
 		}
 	}
-#endif // WITH_CHAOS
 
 	if (IsValidBodyInstance())
 	{
@@ -1658,11 +1607,8 @@ bool FBodyInstance::Weld(FBodyInstance* TheirBody, const FTransform& TheirTM)
 		TArray<UPhysicalMaterial*> ComplexPhysMats;
 		TArray<FPhysicalMaterialMaskParams> ComplexPhysMatMasks;
 	
-#if WITH_CHAOS
 		ComplexPhysMats = TheirBody->GetComplexPhysicalMaterials(ComplexPhysMatMasks);
-#else
-		ComplexPhysMats = TheirBody->GetComplexPhysicalMaterials();
-#endif
+
 		// This builds collision data based on this (parent) body, not their body. This gets fixed  up later though when PostShapeChange() calls UpdatePhysicsFilterData().
 		FBodyCollisionData BodyCollisionData;
 		BuildBodyFilterData(BodyCollisionData.CollisionFilterData);
@@ -1733,12 +1679,10 @@ void FBodyInstance::UnWeld(FBodyInstance* TheirBI)
 		const int32 NumSyncShapes = GetAllShapes_AssumesLocked(Shapes);
 		const int32 NumTotalShapes = Shapes.Num();
 
-#if WITH_CHAOS
 		if(EnsureUnweldModifiesGTOnly && Actor->GetSolverBase() != nullptr)
 		{
 			ensureAlwaysMsgf(false, TEXT("Tried to unweld on body already in solver %s"), *GetBodyDebugName());
 		}
-#endif
 
 		// reversed since FPhysicsInterface::DetachShape is removing shapes
 		for (int Idx = Shapes.Num()-1; Idx >=0; Idx--)
@@ -1896,6 +1840,8 @@ bool FBodyInstance::GetRigidBodyState(FRigidBodyState& OutState)
 
 bool FBodyInstance::UpdateBodyScale(const FVector& InScale3D, bool bForceUpdate)
 {
+	using namespace Chaos;
+
 	SCOPE_CYCLE_COUNTER(STAT_BodyInstanceUpdateBodyScale);
 
 	if (!IsValidBodyInstance())
@@ -1918,9 +1864,6 @@ bool FBodyInstance::UpdateBodyScale(const FVector& InScale3D, bool bForceUpdate)
 
 	FVector UpdatedScale3D;
 
-#if WITH_CHAOS 
-	using namespace Chaos;
-
 	//Get all shapes
 	EScaleMode::Type ScaleMode = EScaleMode::Free;
 
@@ -2263,259 +2206,6 @@ bool FBodyInstance::UpdateBodyScale(const FVector& InScale3D, bool bForceUpdate)
 		}
 	}
 
-#elif PHYSICS_INTERFACE_PHYSX
-	//Get all shapes
-	EScaleMode::Type ScaleMode = EScaleMode::Free;
-
-	FPhysicsCommand::ExecuteWrite(ActorHandle, [&](const FPhysicsActorHandle& Actor)
-	{
-		TArray<PxShape *> PShapes;
-		TArray<FPhysicsShapeHandle> Shapes;
-		GetAllShapes_AssumesLocked(Shapes);
-		ScaleMode = ComputeScaleMode(Shapes);
-
-		FVector AdjustedScale3D;
-		FVector AdjustedScale3DAbs;
-
-		// Apply scaling
-		ComputeScalingVectors(ScaleMode, InScale3D, AdjustedScale3D, AdjustedScale3DAbs);
-		
-		for(FPhysicsShapeHandle& Shape : Shapes)
-		{
-			bool bInvalid = false;	//we only mark invalid if actually found geom and it's invalid scale
-			PxGeometry* UpdatedGeometry = NULL;
-			FTransform LocalTransform = FPhysicsInterface::GetLocalTransform(Shape);
-
-			ECollisionShapeType WrappedOrConcreteType = FPhysicsInterface::GetShapeType(Shape);
-			FKShapeElem* ShapeElem = FPhysxUserData::Get<FKShapeElem>(FPhysicsInterface::GetUserData(Shape));
-			const FTransform& RelativeTM = GetRelativeBodyTransform(Shape);
-
-			FPhysicsGeometryCollection GeoCollection = FPhysicsInterface::GetGeometryCollection(Shape);
-
-			switch (WrappedOrConcreteType)
-			{
-				case ECollisionShapeType::Sphere:
-				{
-					FKSphereElem* SphereElem = ShapeElem->GetShapeCheck<FKSphereElem>();
-					ensure(ScaleMode == EScaleMode::LockedXYZ);
-
-					PxSphereGeometry& PSphereGeom = GeoCollection.GetSphereGeometry();
-					 
-					PSphereGeom.radius = FMath::Max(SphereElem->Radius * AdjustedScale3DAbs.X, FCollisionShape::MinSphereRadius());
-
-					FVector NewTranslation = RelativeTM.TransformPosition(SphereElem->Center) * AdjustedScale3D.X;
-					LocalTransform.SetTranslation(NewTranslation);
-
-					if (PSphereGeom.isValid())
-					{
-						UpdatedGeometry = &PSphereGeom;
-						bSuccess = true;
-					}
-					else
-					{
-						bInvalid = true;
-					}
-					break;
-				}
-				case ECollisionShapeType::Box:
-				{
-					FKBoxElem* BoxElem = ShapeElem->GetShapeCheck<FKBoxElem>();
-					PxBoxGeometry& PBoxGeom = GeoCollection.GetBoxGeometry();
-
-					PBoxGeom.halfExtents.x = FMath::Max((0.5f * BoxElem->X * AdjustedScale3DAbs.X), FCollisionShape::MinBoxExtent());
-					PBoxGeom.halfExtents.y = FMath::Max((0.5f * BoxElem->Y * AdjustedScale3DAbs.Y), FCollisionShape::MinBoxExtent());
-					PBoxGeom.halfExtents.z = FMath::Max((0.5f * BoxElem->Z * AdjustedScale3DAbs.Z), FCollisionShape::MinBoxExtent());
-
-					LocalTransform = BoxElem->GetTransform() * RelativeTM;
-					LocalTransform.ScaleTranslation(AdjustedScale3D);
-
-					if (PBoxGeom.isValid())
-					{
-						UpdatedGeometry = &PBoxGeom;
-						bSuccess = true;
-					}
-					else
-					{
-						bInvalid = true;
-					}
-					break;
-				}
-				case ECollisionShapeType::Capsule:
-				{
-					ensure(ScaleMode == EScaleMode::LockedXY || ScaleMode == EScaleMode::LockedXYZ);
-
-					float ScaleRadius = FMath::Max(AdjustedScale3DAbs.X, AdjustedScale3DAbs.Y);
-					float ScaleLength = AdjustedScale3DAbs.Z;
-
-					PxCapsuleGeometry& PCapsuleGeom = GeoCollection.GetCapsuleGeometry();
-
-					float Radius = 0.0f;
-					float Length = 0.0f;
-
-					if (ShapeElem->GetShapeType() == EAggCollisionShape::TaperedCapsule)
-					{
-						// Handle the case where standard capsules are generated in place of tapered capsules, which are not fully supported.
-						FKTaperedCapsuleElem* TaperedCapsuleElem = ShapeElem->GetShapeCheck<FKTaperedCapsuleElem>();
-
-						const FReal MeanRadius = 0.5f * (TaperedCapsuleElem->Radius0 + TaperedCapsuleElem->Radius1);
-
-						// this is a bit confusing since radius and height is scaled
-						// first apply the scale first 
-						Radius = FMath::Max(MeanRadius * ScaleRadius, 0.1f);
-						Length = TaperedCapsuleElem->Length + MeanRadius * 2.f;
-
-						LocalTransform.SetTranslation(RelativeTM.TransformPosition(TaperedCapsuleElem->Center));
-						LocalTransform.SetRotation(TaperedCapsuleElem->Rotation.Quaternion()* U2PSphylBasis_UE); // #PHYS2 we probably want to put this behind the interface?
-					}
-					else
-					{
-						FKSphylElem* SphylElem = ShapeElem->GetShapeCheck<FKSphylElem>();
-
-						// this is a bit confusing since radius and height is scaled
-						// first apply the scale first 
-						Radius = FMath::Max(SphylElem->Radius * ScaleRadius, 0.1f);
-						Length = SphylElem->Length + SphylElem->Radius * 2.f;
-
-						LocalTransform.SetTranslation(RelativeTM.TransformPosition(SphylElem->Center));
-						LocalTransform.SetRotation(SphylElem->Rotation.Quaternion() * U2PSphylBasis_UE); // #PHYS2 we probably want to put this behind the interface?
-					}
-
-					LocalTransform.ScaleTranslation(AdjustedScale3D);
-
-					float HalfLength = Length * ScaleLength * 0.5f;
-					Radius = FMath::Min(Radius, HalfLength);	//radius is capped by half length
-					Radius = FMath::Max(Radius, FCollisionShape::MinCapsuleRadius()); // bounded by minimum limit.
-					float HalfHeight = HalfLength - Radius;
-					HalfHeight = FMath::Max(FCollisionShape::MinCapsuleAxisHalfHeight(), HalfHeight);
-
-					PCapsuleGeom.halfHeight = FMath::Max(HalfHeight, KINDA_SMALL_NUMBER);
-					PCapsuleGeom.radius = FMath::Max(Radius, KINDA_SMALL_NUMBER);
-
-					if (PCapsuleGeom.isValid())
-					{
-						UpdatedGeometry = &PCapsuleGeom;
-						bSuccess = true;
-					}
-					else
-					{
-						bInvalid = true;
-					}
-
-					break;
-				}
-				case ECollisionShapeType::Convex:
-				{
-					FKConvexElem* ConvexElem = ShapeElem->GetShapeCheck<FKConvexElem>();
-					PxConvexMeshGeometry& PConvexGeom = GeoCollection.GetConvexGeometry();
-
-					bool bUseNegX = CalcMeshNegScaleCompensation(AdjustedScale3D, LocalTransform);
-
-					PConvexGeom.convexMesh = bUseNegX ? ConvexElem->GetMirroredConvexMesh() : ConvexElem->GetConvexMesh();
-					PConvexGeom.scale.scale = U2PVector(AdjustedScale3DAbs);
-
-					LocalTransform.TransformRotation(RelativeTM.GetRotation());
-					LocalTransform.ScaleTranslation(AdjustedScale3D);
-
-					if (PConvexGeom.isValid())
-					{
-						UpdatedGeometry = &PConvexGeom;
-						bSuccess = true;
-					}
-					else
-					{
-						bInvalid = true;
-					}
-
-					break;
-				}
-				case ECollisionShapeType::Trimesh:
-				{
-					check(ShapeElem == nullptr);	//trimesh shape doesn't have userData
-
-					PxTriangleMeshGeometry& PTriMeshGeom = GeoCollection.GetTriMeshGeometry();
-
-					// find which trimesh elems it is
-					// it would be nice to know if the order of PShapes array index is in the order of createShape
-					if (BodySetup.IsValid())
-					{
-						for (PxTriangleMesh* TriMesh : BodySetup->TriMeshes)
-						{
-							// found it
-							if (TriMesh == PTriMeshGeom.triangleMesh)
-							{
-								PTriMeshGeom.scale.scale = U2PVector(AdjustedScale3D);
-
-								LocalTransform = RelativeTM;
-								LocalTransform.ScaleTranslation(AdjustedScale3D);
-
-								if (PTriMeshGeom.isValid())
-								{
-									UpdatedGeometry = &PTriMeshGeom;
-									bSuccess = true;
-								}
-								else
-								{
-									bInvalid = true;
-								}
-							}
-						}
-					}
-
-					break;
-				}
-				case ECollisionShapeType::Heightfield:
-				{
-					// HeightField is only used by Landscape, which does different code path from other primitives
-					break;
-				}
-				default:
-				{
-						   UE_LOG(LogPhysics, Error, TEXT("Unknown geom type."));
-				}
-			}// end switch
-
-			if (UpdatedGeometry)
-			{
-				FPhysicsCommand::ExecuteShapeWrite(this, Shape, [&](FPhysicsShapeHandle& InShape)
-				{
-					FPhysicsInterface::SetLocalTransform(InShape, LocalTransform);
-					FPhysicsInterface::SetGeometry(InShape, *UpdatedGeometry);
-				});
-
-				UpdatedScale3D = AdjustedScale3D;
-			}
-			else if (bInvalid)
-			{
-#if !(UE_BUILD_SHIPPING || UE_BUILD_TEST)
-				FMessageLog("PIE").Warning()
-					->AddToken(FTextToken::Create(LOCTEXT("InvalidScaleStart", "Applying invalid scale")))
-					->AddToken(FTextToken::Create(AdjustedScale3D.ToCompactText()))
-					->AddToken(FTextToken::Create(LOCTEXT("InvalidScaleMid", "to body")))
-					->AddToken(FUObjectToken::Create(OwnerComponent.Get()));
-#endif // !(UE_BUILD_SHIPPING || UE_BUILD_TEST)
-			}
-		}
-	});
-
-	// if success, overwrite old Scale3D, otherwise, just don't do it. It will have invalid scale next time
-	if (bSuccess)
-	{
-		Scale3D = UpdatedScale3D;
-
-		// update mass if required
-		if (bUpdateMassWhenScaleChanges)
-		{
-			bDirtyMassProps = true;
-
-			//if already simulated compute mass immediately
-			if(ShouldInstanceSimulatingPhysics())
-			{
-				UpdateMassProperties();
-			}
-		}
-	}
-#endif
-
 	return bSuccess;
 }
 
@@ -3225,7 +2915,6 @@ void FBodyInstance::UpdateMassProperties()
 
 	UPhysicalMaterial* PhysMat = GetSimplePhysicalMaterial();
 
-#if WITH_PHYSX
 	if (FPhysicsInterface::IsValid(ActorHandle) && FPhysicsInterface::IsRigidBody(ActorHandle))
 	{
 		FPhysicsCommand::ExecuteWrite(ActorHandle, [&](FPhysicsActorHandle& Actor)
@@ -3253,11 +2942,8 @@ void FBodyInstance::UpdateMassProperties()
 					}
 				}
 
-#if WITH_CHAOS
 				Chaos::FMassProperties TotalMassProperties;
-#elif PHYSICS_INTERFACE_PHYSX
-				PxMassProperties TotalMassProperties;
-#endif
+
 				if (ShapeToBodiesMap.IsValid() && ShapeToBodiesMap->Num() > 0)
 				{
 					struct FWeldedBatch
@@ -3297,7 +2983,6 @@ void FBodyInstance::UpdateMassProperties()
 						}
 					}
 
-#if WITH_CHAOS
 					TArray<Chaos::FMassProperties> SubMassProperties;
 					for (auto BodyShapesItr : BodyToShapes)
 					{
@@ -3319,24 +3004,6 @@ void FBodyInstance::UpdateMassProperties()
 					{
 						TotalMassProperties = Chaos::Combine(SubMassProperties);
 					}
-
-#elif PHYSICS_INTERFACE_PHYSX
-					TArray<PxMassProperties> SubMassProperties;
-					TArray<PxTransform> MassTMs;
-					for (auto BodyShapesItr : BodyToShapes)
-					{
-						const FBodyInstance* OwningBI = BodyShapesItr.Key;
-						const FWeldedBatch& WeldedBatch = BodyShapesItr.Value;
-						FTransform MassModifierTransform = WeldedBatch.RelTM;
-						MassModifierTransform.SetScale3D(MassModifierTransform.GetScale3D() * Scale3D);	//Ensure that any scaling that is done on the component is passed into the mass frame modifiers
-
-						PxMassProperties BodyMassProperties = BodyUtils::ComputeMassProperties(OwningBI, WeldedBatch.Shapes, MassModifierTransform);
-						SubMassProperties.Add(BodyMassProperties);
-						MassTMs.Add(PxTransform(PxIdentity));
-					}
-
-					TotalMassProperties = PxMassProperties::sum(SubMassProperties.GetData(), MassTMs.GetData(), SubMassProperties.Num());
-#endif
 				}
 				else
 				{
@@ -3351,8 +3018,6 @@ void FBodyInstance::UpdateMassProperties()
 					}
 				}
 
-				// #PHYS2 Refactor out PxMassProperties (Our own impl?)
-#if WITH_CHAOS
 				// Note: We expect the inertia to be diagonal at this point
 				// Only set mass properties if inertia tensor is valid. TODO Remove this once we track down cause of empty tensors.
 				// (This can happen if all shapes have bContributeToMass set to false which gives an empty Shapes array. There may be other ways).
@@ -3367,16 +3032,6 @@ void FBodyInstance::UpdateMassProperties()
 					FTransform Com(TotalMassProperties.RotationOfMass, TotalMassProperties.CenterOfMass);
 					FPhysicsInterface::SetComLocalPose_AssumesLocked(Actor, Com);
 				}
-#else
-				PxQuat MassOrientation;
-				const FVector MassSpaceInertiaTensor = P2UVector(PxMassProperties::getMassSpaceInertia(TotalMassProperties.inertiaTensor, MassOrientation));
-				FPhysicsInterface::SetMassSpaceInertiaTensor_AssumesLocked(Actor, MassSpaceInertiaTensor);
-
-				FPhysicsInterface::SetMass_AssumesLocked(Actor, TotalMassProperties.mass);
-
-				FTransform Com(P2UQuat(MassOrientation), P2UVector(TotalMassProperties.centerOfMass));
-				FPhysicsInterface::SetComLocalPose_AssumesLocked(Actor, Com);
-#endif
 			}
 		});
 	}
@@ -3386,7 +3041,6 @@ void FBodyInstance::UpdateMassProperties()
 	{
 		BodyInstanceDelegates->OnRecalculatedMassProperties.Broadcast(this);
 	}
-#endif
 }
 
 void FBodyInstance::UpdateDebugRendering()
@@ -3598,7 +3252,6 @@ void FBodyInstance::ClearForces(bool bAllowSubstepping)
 
 void FBodyInstance::SetOneWayInteraction(bool InOneWayInteraction /*= true*/)
 {
-#if WITH_CHAOS
 	FPhysicsCommand::ExecuteWrite(ActorHandle, [&](const FPhysicsActorHandle& Actor)
 		{
 			if (FPhysicsInterface::IsRigidBody(Actor) && !IsRigidBodyKinematic_AssumesLocked(Actor))
@@ -3606,7 +3259,6 @@ void FBodyInstance::SetOneWayInteraction(bool InOneWayInteraction /*= true*/)
 				FPhysicsInterface::SetOneWayInteraction_AssumesLocked(Actor, InOneWayInteraction);
 			}
 		});
-#endif
 }
 
 void FBodyInstance::AddTorqueInRadians(const FVector& Torque, bool bAllowSubstepping, bool bAccelChange)
@@ -3707,7 +3359,6 @@ void FBodyInstance::SetInstanceNotifyRBCollision(bool bNewNotifyCollision)
 		return;
 	}
 	
-#if WITH_CHAOS
 	// make sure to register the component for collision events 
 	if (UPrimitiveComponent* PrimComp = OwnerComponent.Get())
 	{
@@ -3727,8 +3378,6 @@ void FBodyInstance::SetInstanceNotifyRBCollision(bool bNewNotifyCollision)
 			}
 		}
 	}
-#endif // WITH_CHAOS	
-
 
 	bNotifyRigidBodyCollision = bNewNotifyCollision;
 	UpdatePhysicsFilterData();
@@ -3771,7 +3420,7 @@ void FBodyInstance::SetSmoothEdgeCollisionsEnabled(bool bNewSmoothEdgeCollisions
 	if (bNewSmoothEdgeCollisions != bSmoothEdgeCollisions)
 	{
 		bSmoothEdgeCollisions = bNewSmoothEdgeCollisions;
-#if WITH_CHAOS
+
 		FPhysicsCommand::ExecuteWrite(ActorHandle, 
 			[bNewSmoothEdgeCollisions](const FPhysicsActorHandle& Actor)
 			{
@@ -3780,7 +3429,6 @@ void FBodyInstance::SetSmoothEdgeCollisionsEnabled(bool bNewSmoothEdgeCollisions
 					FPhysicsInterface::SetSmoothEdgeCollisionsEnabled_AssumesLocked(Actor, bNewSmoothEdgeCollisions);
 				}
 			});
-#endif
 	}
 }
 
@@ -3917,18 +3565,10 @@ bool FBodyInstance::OverlapTestForBodiesImpl(const FVector& Pos, const FQuat& Ro
 
 		for(const FPhysicsShapeHandle& Shape : TargetShapes)
 		{
-#if WITH_CHAOS
 			if (!Shape.GetGeometry().IsConvex())
 			{
 				continue;	//we skip complex shapes - should this respect ComplexAsSimple?
 			}
-#else
-			ECollisionShapeType ShapeType = FPhysicsInterface::GetShapeType(Shape);
-			if(ShapeType == ECollisionShapeType::Heightfield || ShapeType == ECollisionShapeType::Trimesh)
-			{
-				continue;	//we skip complex shapes - should this respect ComplexAsSimple?
-			}
-#endif
 
 			// Calc shape global pose
 			FTransform PShapeGlobalPose = FPhysicsInterface::GetLocalTransform(Shape) * PTestGlobalPose;
@@ -4065,18 +3705,10 @@ bool FBodyInstance::OverlapMulti(TArray<struct FOverlapResult>& InOutOverlaps, c
 					continue;
 				}
 
-#if WITH_CHAOS
 				if (!ShapeRef.GetGeometry().IsConvex())
 				{
 					continue;	//we skip complex shapes - should this respect ComplexAsSimple?
 				}
-#else
-				ECollisionShapeType ShapeType = FPhysicsInterface::GetShapeType(ShapeRef);
-				if (ShapeType == ECollisionShapeType::Heightfield || ShapeType == ECollisionShapeType::Trimesh)
-				{
-					continue;	//we skip complex shapes - should this respect ComplexAsSimple?
-				}
-#endif
 
 				// Calc shape global pose
 				const FTransform LocalTransform = FPhysicsInterface::GetLocalTransform(ShapeRef);
@@ -4212,7 +3844,6 @@ void FBodyInstance::ApplyMaterialToShape_AssumesLocked(const FPhysicsShapeHandle
 	{
 		if(ComplexPhysMats.Num())
 		{
-#if WITH_CHAOS
 			if (ensure(ComplexPhysMatMasks))
 			{
 				FPhysicsInterface::SetMaterials(InShape, ComplexPhysMats, *ComplexPhysMatMasks);
@@ -4221,9 +3852,6 @@ void FBodyInstance::ApplyMaterialToShape_AssumesLocked(const FPhysicsShapeHandle
 			{
 				FPhysicsInterface::SetMaterials(InShape, ComplexPhysMats);
 			}
-#else
-			FPhysicsInterface::SetMaterials(InShape, ComplexPhysMats);
-#endif
 		}
 		else
 		{
@@ -4379,9 +4007,8 @@ void FBodyInstance::InitDynamicProperties_AssumesLocked()
 				FPhysicsInterface::PutToSleep_AssumesLocked(ActorHandle);
 			}
 		}
-#if WITH_CHAOS
+
 		FPhysicsInterface::SetInitialized_AssumesLocked(ActorHandle, true);
-#endif
 	}
 }
 
@@ -4623,20 +4250,12 @@ FRecalculatedMassProperties& FBodyInstance::OnRecalculatedMassProperties()
 
 bool FBodyInstanceAsyncPhysicsTickHandle::IsValid() const
 {
-#if WITH_CHAOS
 	return Proxy && Proxy->GetPhysicsThreadAPI() != nullptr;
-#else
-	return false;
-#endif
 }
 
 Chaos::FRigidBodyHandle_Internal* FBodyInstanceAsyncPhysicsTickHandle::operator->()
 {
-#if WITH_CHAOS
 	return Proxy->GetPhysicsThreadAPI();
-#else
-	return nullptr;
-#endif
 }
 
 ////////////////////////////////////////////////////////////////////////////
diff --git a/Engine/Source/Runtime/Engine/Private/PhysicsEngine/BodySetup.cpp b/Engine/Source/Runtime/Engine/Private/PhysicsEngine/BodySetup.cpp
--- a/Engine/Source/Runtime/Engine/Private/PhysicsEngine/BodySetup.cpp
+++ b/Engine/Source/Runtime/Engine/Private/PhysicsEngine/BodySetup.cpp
@@ -29,34 +29,18 @@
 #include "Chaos/Convex.h"
 #include "Chaos/Levelset.h"
 
-#include "PhysXCookHelper.h"
-
-#if PHYSICS_INTERFACE_PHYSX
-	#include "PhysXPublic.h"
-	#include "PhysicsEngine/PhysXSupport.h"
-#endif // WITH_PHYSX
-
 #include "Modules/ModuleManager.h"
 
-#if PHYSICS_INTERFACE_PHYSX
-	#include "IPhysXCookingModule.h"
-	#include "IPhysXCooking.h"
-	#include "PhysicsEngine/PhysDerivedData.h"
-#endif
-
 #include "Physics/PhysicsInterfaceUtils.h"
 #include "PhysicalMaterials/PhysicalMaterial.h"
 #include "ProfilingDebugging/CookStats.h"
 #include "UObject/AnimPhysObjectVersion.h"
 
 #include "Chaos/TriangleMeshImplicitObject.h"
-
-#if WITH_CHAOS
-	#include "Experimental/ChaosDerivedData.h"
-	#include "Physics/Experimental/ChaosDerivedDataReader.h"
-	#include "Chaos/CollisionConvexMesh.h"
-	#include "Experimental/ChaosCooking.h"
-#endif
+#include "Experimental/ChaosDerivedData.h"
+#include "Physics/Experimental/ChaosDerivedDataReader.h"
+#include "Chaos/CollisionConvexMesh.h"
+#include "Experimental/ChaosCooking.h"
 
 /** Enable to verify that the cooked data matches the source data as we cook it */
 #define VERIFY_COOKED_PHYS_DATA 0
@@ -65,10 +49,6 @@
 
 
 FCookBodySetupInfo::FCookBodySetupInfo() :
-#if PHYSICS_INTERFACE_PHYSX
-	TriMeshCookFlags(EPhysXMeshCookFlags::Default) ,
-	ConvexCookFlags(EPhysXMeshCookFlags::Default) ,
-#endif // WITH_PHYSX
 	bCookNonMirroredConvex(false),
 	bCookMirroredConvex(false),
 	bConvexDeformableMesh(false),
@@ -100,24 +80,10 @@ namespace PhysXBodySetupCookStats
 
 DEFINE_STAT(STAT_PhysXCooking);
 
-#if WITH_PHYSX
-
 bool IsRuntimeCookingEnabled()
 {
-#if PHYSICS_INTERFACE_PHYSX
-	return FModuleManager::LoadModulePtr<IPhysXCookingModule>("RuntimePhysXCooking") != nullptr;
-#else
 	return false;
-#endif
 }
-#endif //WITH_PHYSX
-
-#if PHYSICS_INTERFACE_PHYSX
-	// Quaternion that converts Sphyls from UE space to PhysX space (negate Y, swap X & Z)
-	// This is equivalent to a 180 degree rotation around the normalized (1, 0, 1) axis
-	const physx::PxQuat U2PSphylBasis( PI, PxVec3( 1.0f / FMath::Sqrt( 2.0f ), 0.0f, 1.0f / FMath::Sqrt( 2.0f ) ) );
-	const FQuat U2PSphylBasis_UE(FVector(1.0f / FMath::Sqrt(2.0f), 0.0f, 1.0f / FMath::Sqrt(2.0f)), PI);
-#endif // WITH_PHYSX
 
 // CVars
 ENGINE_API TAutoConsoleVariable<float> CVarContactOffsetFactor(
@@ -173,9 +139,7 @@ UBodySetup::UBodySetup(const FObjectInitializer& ObjectInitializer)
 	SetFlags(RF_Transactional);
 	bSharedCookedData = false;
 	CookedFormatDataOverride = nullptr;
-#if PHYSICS_INTERFACE_PHYSX
 	CurrentCookHelper = nullptr;
-#endif
 }
 
 void UBodySetup::CopyBodyPropertiesFrom(const UBodySetup* FromSetup)
@@ -186,12 +150,7 @@ void UBodySetup::CopyBodyPropertiesFrom(const UBodySetup* FromSetup)
 	for (int32 i = 0; i < AggGeom.ConvexElems.Num(); i++)
 	{
 		FKConvexElem& ConvexElem = AggGeom.ConvexElems[i];
-#if PHYSICS_INTERFACE_PHYSX
-		ConvexElem.SetConvexMesh(nullptr);
-		ConvexElem.SetMirroredConvexMesh(nullptr);
-#elif WITH_CHAOS
 		ConvexElem.ResetChaosConvexMesh();
-#endif
 	}
 
 	DefaultInstance.CopyBodyInstancePropertiesFrom(&FromSetup->DefaultInstance);
@@ -216,12 +175,7 @@ void UBodySetup::AddCollisionFrom(const FKAggregateGeom& FromAggGeom)
 	for (int32 i = FirstNewConvexIdx; i < AggGeom.ConvexElems.Num(); i++)
 	{
 		FKConvexElem& ConvexElem = AggGeom.ConvexElems[i];
-#if PHYSICS_INTERFACE_PHYSX
-		ConvexElem.SetConvexMesh(nullptr);
-		ConvexElem.SetMirroredConvexMesh(nullptr);
-#elif WITH_CHAOS
 		ConvexElem.ResetChaosConvexMesh();
-#endif
 	}
 
 	AggGeom.LevelSetElems.Append(FromAggGeom.LevelSetElems);
@@ -286,12 +240,6 @@ void UBodySetup::GetCookInfo(FCookBodySetupInfo& OutCookInfo, EPhysXMeshCookFlag
 			// Get cook flags to use
 			OutCookInfo.ConvexCookFlags = InCookFlags;
 			OutCookInfo.bConvexDeformableMesh = GetOuter()->IsA(USplineMeshComponent::StaticClass());
-#if PHYSICS_INTERFACE_PHYSX
-			if (OutCookInfo.bConvexDeformableMesh)
-			{
-				OutCookInfo.ConvexCookFlags |= EPhysXMeshCookFlags::DeformableMesh;
-			}
-#endif
 		}
 	}
 	else
@@ -327,22 +275,6 @@ void UBodySetup::GetCookInfo(FCookBodySetupInfo& OutCookInfo, EPhysXMeshCookFlag
 
 			// Set up cooking flags
 			EPhysXMeshCookFlags CookFlags = InCookFlags;
-#if PHYSICS_INTERFACE_PHYSX
-			if (TriangleMeshDesc.bDeformableMesh)
-			{
-				CookFlags |= EPhysXMeshCookFlags::DeformableMesh;
-			}
-
-			if (TriangleMeshDesc.bFastCook)
-			{
-				CookFlags |= EPhysXMeshCookFlags::FastCook;
-			}
-
-			if (TriangleMeshDesc.bDisableActiveEdgePrecompute)
-			{
-				CookFlags |= EPhysXMeshCookFlags::DisableActiveEdgePrecompute;
-			}
-#endif
 			OutCookInfo.TriMeshCookFlags = CookFlags;
 
 			OutCookInfo.bSupportFaceRemap = bSupportUVsAndFaceRemap;
@@ -445,7 +377,6 @@ void UBodySetup::CreatePhysicsMeshes()
 	bool bClearMeshes = true;
 	bool bSkipProcessFormatData = false;
 
-#if WITH_CHAOS
 	// If the data already has been deserialized from the async loading thread, use that.
 	if (ChaosDerivedDataReader.IsValid())
 	{
@@ -453,27 +384,18 @@ void UBodySetup::CreatePhysicsMeshes()
 		ChaosDerivedDataReader.Reset();
 		bSkipProcessFormatData = true;
 	}
-#endif
 
 	if (!bSkipProcessFormatData)
 	{
 		if (FByteBulkData* FormatData = GetCookedFormatData())
 		{
-#if WITH_PHYSX  && PHYSICS_INTERFACE_PHYSX
-			bClearMeshes = !ProcessFormatData_PhysX(FormatData);
-#elif WITH_CHAOS
 			bClearMeshes = !ProcessFormatData_Chaos(FormatData);
-#endif
 		}
 		else
 		{
 			if (IsRuntime(this))
 			{
-#if WITH_PHYSX  && PHYSICS_INTERFACE_PHYSX
-				bClearMeshes = !RuntimeCookPhysics_PhysX();
-#elif WITH_CHAOS
 				bClearMeshes = !RuntimeCookPhysics_Chaos();
-#endif
 			}
 		}
 	}
@@ -491,7 +413,6 @@ void UBodySetup::CreatePhysicsMeshes()
 		}
 	}
 
-#if WITH_CHAOS
 	// For drawing of convex elements we require an index buffer, previously we could
 	// get this from a PxConvexMesh but Chaos doesn't maintain that data. Instead now
 	// it is a part of the element rather than the physics geometry, if we load in an
@@ -501,8 +422,6 @@ void UBodySetup::CreatePhysicsMeshes()
 	{
 		Convex.ComputeChaosConvexIndices();
 	}
-#endif
-
 
 	if(bClearMeshes)
 	{
@@ -513,118 +432,6 @@ void UBodySetup::CreatePhysicsMeshes()
 
 }
 
-#if WITH_PHYSX && PHYSICS_INTERFACE_PHYSX
-bool UBodySetup::RuntimeCookPhysics_PhysX()
-{
-	FPhysXCookHelper CookHelper(GetPhysXCookingModule());
-
-	GetCookInfo(CookHelper.CookInfo, GetRuntimeOnlyCookOptimizationFlags());
-	if(CookHelper.HasSomethingToCook(CookHelper.CookInfo))
-	{
-		if(!IsRuntimeCookingEnabled())
-		{
-			UE_LOG(LogPhysics, Error, TEXT("Attempting to build physics data for %s at runtime, but runtime cooking is disabled (see the RuntimePhysXCooking plugin)."), *GetPathName());
-		}
-		else
-		{
-			if(CookHelper.CreatePhysicsMeshes_Concurrent())
-			{
-				FinishCreatingPhysicsMeshes_PhysX(CookHelper.OutNonMirroredConvexMeshes, CookHelper.OutMirroredConvexMeshes, CookHelper.OutTriangleMeshes);
-				bFailedToCreatePhysicsMeshes = false;
-				return true;
-			}
-			else
-			{
-				bFailedToCreatePhysicsMeshes = true;
-			}
-		}
-	}			
-	return false;
-}
-
-bool UBodySetup::ProcessFormatData_PhysX(FByteBulkData* FormatData)
-{
-	if(FormatData->IsLocked())
-	{
-		// seems it's being already processed
-		return false;
-	}
-
-	FPhysXCookingDataReader CookedDataReader(*FormatData, &UVInfo);
-
-	if(GetCollisionTraceFlag() != CTF_UseComplexAsSimple)
-	{
-		bool bNeedsCooking = bGenerateNonMirroredCollision && CookedDataReader.ConvexMeshes.Num() != AggGeom.ConvexElems.Num();
-		bNeedsCooking = bNeedsCooking || (bGenerateMirroredCollision && CookedDataReader.ConvexMeshesNegX.Num() != AggGeom.ConvexElems.Num());
-		if(bNeedsCooking)	//Because of bugs it's possible to save with out of sync cooked data. In editor we want to fixup this data
-		{
-			InvalidatePhysicsData();
-			CreatePhysicsMeshes();
-			return false;
-		}
-	}
-
-	FinishCreatingPhysicsMeshes_PhysX(CookedDataReader.ConvexMeshes, CookedDataReader.ConvexMeshesNegX, CookedDataReader.TriMeshes);
-	return true;
-}
-
-void UBodySetup::FinishCreatingPhysicsMeshes_PhysX(const TArray<PxConvexMesh*>& ConvexMeshes, const TArray<PxConvexMesh*>& ConvexMeshesNegX, const TArray<PxTriangleMesh*>& CookedTriMeshes)
-{
-	ClearPhysicsMeshes();
-
-	FPhysxSharedData::LockAccess();
-
-	const FString FullName = GetFullName();
-	if (GetCollisionTraceFlag() != CTF_UseComplexAsSimple)
-	{
-		ensure(!bGenerateNonMirroredCollision || ConvexMeshes.Num() == 0 || ConvexMeshes.Num() == AggGeom.ConvexElems.Num());
-		ensure(!bGenerateMirroredCollision || ConvexMeshesNegX.Num() == 0 || ConvexMeshesNegX.Num() == AggGeom.ConvexElems.Num());
-
-		//If the cooked data no longer has convex meshes, make sure to empty AggGeom.ConvexElems - otherwise we leave NULLS which cause issues, and we also read past the end of CookedDataReader.ConvexMeshes
-		if ((bGenerateNonMirroredCollision && ConvexMeshes.Num() == 0) || (bGenerateMirroredCollision && ConvexMeshesNegX.Num() == 0))
-		{
-			AggGeom.ConvexElems.Empty();
-		}
-
-		for (int32 ElementIndex = 0; ElementIndex < AggGeom.ConvexElems.Num(); ElementIndex++)
-		{
-			FKConvexElem& ConvexElem = AggGeom.ConvexElems[ElementIndex];
-
-			if (bGenerateNonMirroredCollision)
-			{
-				ConvexElem.SetConvexMesh(ConvexMeshes[ElementIndex]);
-				FPhysxSharedData::Get().Add(ConvexElem.GetConvexMesh(), FullName);
-			}
-
-			if (bGenerateMirroredCollision)
-			{
-				ConvexElem.SetMirroredConvexMesh(ConvexMeshesNegX[ElementIndex]);
-				FPhysxSharedData::Get().Add(ConvexElem.GetMirroredConvexMesh(), FullName);
-			}
-		}
-	}
-
-	for (PxTriangleMesh* TriMesh : CookedTriMeshes)
-	{
-		if(TriMesh)
-		{
-			TriMeshes.Add(TriMesh);
-			FPhysxSharedData::Get().Add(TriMesh, FullName);
-		}
-	}
-
-	FPhysxSharedData::UnlockAccess();
-
-	// Clear the cooked data
-	if (!GIsEditor && !bSharedCookedData)
-	{
-		CookedFormatData.FlushData();
-	}
-
-	bCreatedPhysicsMeshes = true;
-}
-#endif //WITH_PHYSX
-
 void UBodySetup::CreatePhysicsMeshesAsync(FOnAsyncPhysicsCookFinished OnAsyncPhysicsCookFinished)
 {
 	check(IsInGameThread());
@@ -632,41 +439,6 @@ void UBodySetup::CreatePhysicsMeshesAsync(FOnAsyncPhysicsCookFinished OnAsyncPhy
 	// Don't start another cook cycle if one's already in progress
 	check(CurrentCookHelper == nullptr);
 
-	// Only perform this check for PhysX as the cooking module is optional
-#if WITH_PHYSX_COOKING && PHYSICS_INTERFACE_PHYSX
-	if (IsRuntime(this) && !IsRuntimeCookingEnabled())
-	{
-		UE_LOG(LogPhysics, Error, TEXT("Attempting to build physics data for %s at runtime, but runtime cooking is disabled (see the RuntimePhysXCooking plugin)."), *GetPathName());
-		FinishCreatePhysicsMeshesAsync(nullptr, OnAsyncPhysicsCookFinished);
-		return;
-	}
-#endif
-
-#if WITH_PHYSX && PHYSICS_INTERFACE_PHYSX
-	if(IPhysXCookingModule* PhysXCookingModule = GetPhysXCookingModule())
-	{
-		FPhysXCookHelper* AsyncPhysicsCookHelper = new FPhysXCookHelper(PhysXCookingModule);
-		GetCookInfo(AsyncPhysicsCookHelper->CookInfo, GetRuntimeOnlyCookOptimizationFlags());	//TODO: pass in different flags?
-
-		if(AsyncPhysicsCookHelper->HasSomethingToCook(AsyncPhysicsCookHelper->CookInfo))
-		{
-			FSimpleDelegateGraphTask::CreateAndDispatchWhenReady(FSimpleDelegateGraphTask::FDelegate::CreateRaw(AsyncPhysicsCookHelper, &FPhysXCookHelper::CreatePhysicsMeshesAsync_Concurrent,
-				/*FinishDelegate=*/FSimpleDelegateGraphTask::FDelegate::CreateUObject(this, &UBodySetup::FinishCreatePhysicsMeshesAsync, AsyncPhysicsCookHelper, OnAsyncPhysicsCookFinished)),
-				GET_STATID(STAT_PhysXCooking), nullptr, ENamedThreads::AnyThread);
-
-			CurrentCookHelper = AsyncPhysicsCookHelper;
-		}
-		else
-		{
-			delete AsyncPhysicsCookHelper;
-			FinishCreatePhysicsMeshesAsync(nullptr, OnAsyncPhysicsCookFinished);
-		}
-	}
-	else
-	{
-		FinishCreatePhysicsMeshesAsync(nullptr, OnAsyncPhysicsCookFinished);
-	}
-#else
 	FAsyncCookHelper* NewCookHelper = new FAsyncCookHelper(this);
 	if(NewCookHelper->HasWork())
 	{
@@ -680,17 +452,11 @@ void UBodySetup::CreatePhysicsMeshesAsync(FOnAsyncPhysicsCookFinished OnAsyncPhy
 		delete NewCookHelper;
 		FinishCreatePhysicsMeshesAsync(nullptr, OnAsyncPhysicsCookFinished);
 	}
-#endif // WITH_PHYSX && PHYSICS_INTERFACE_PHYSX
 }
 
 void UBodySetup::AbortPhysicsMeshAsyncCreation()
 {
-#if PHYSICS_INTERFACE_PHYSX
-	if (CurrentCookHelper)
-	{
-		CurrentCookHelper->Abort();
-	}
-#endif
+	// TODO Chaos - support async cook abort
 }
 
 void UBodySetup::FinishCreatePhysicsMeshesAsync(FAsyncCookHelper* AsyncPhysicsCookHelper, FOnAsyncPhysicsCookFinished OnAsyncPhysicsCookFinished)
@@ -704,15 +470,10 @@ void UBodySetup::FinishCreatePhysicsMeshesAsync(FAsyncCookHelper* AsyncPhysicsCo
 
 	if(AsyncPhysicsCookHelper)
 	{
-#if WITH_PHYSX && PHYSICS_INTERFACE_PHYSX
-		FinishCreatingPhysicsMeshes_PhysX(AsyncPhysicsCookHelper->OutNonMirroredConvexMeshes, AsyncPhysicsCookHelper->OutMirroredConvexMeshes, AsyncPhysicsCookHelper->OutTriangleMeshes);
-		UVInfo = AsyncPhysicsCookHelper->OutUVInfo;
-#elif WITH_CHAOS
 		FinishCreatingPhysicsMeshes_Chaos(*AsyncPhysicsCookHelper);
 		UVInfo = AsyncPhysicsCookHelper->UVInfo;
-#endif // WITH_PHYSX
-		delete AsyncPhysicsCookHelper;
 
+		delete AsyncPhysicsCookHelper;
 	}
 	else
 	{
@@ -723,7 +484,6 @@ void UBodySetup::FinishCreatePhysicsMeshesAsync(FAsyncCookHelper* AsyncPhysicsCo
 	OnAsyncPhysicsCookFinished.ExecuteIfBound(bSuccess);
 }
 
-#if WITH_CHAOS
 bool UBodySetup::ProcessFormatData_Chaos(FChaosDerivedDataReader<float, 3>& Reader)
 {
 	FinishCreatingPhysicsMeshes_Chaos(Reader);
@@ -830,13 +590,11 @@ void UBodySetup::FinishCreatingPhysicsMeshes_Chaos(TArray<TSharedPtr<Chaos::FCon
 	}
 #endif
 
-#if WITH_CHAOS
 	// Force trimesh collisions off
 	for(auto& TriMesh : ChaosTriMeshes)
 	{
 		TriMesh->SetDoCollide(false);
 	}
-#endif
 
 	// Clear the cooked data
 	if(!GIsEditor && !bSharedCookedData)
@@ -847,58 +605,18 @@ void UBodySetup::FinishCreatingPhysicsMeshes_Chaos(TArray<TSharedPtr<Chaos::FCon
 	bCreatedPhysicsMeshes = true;
 }
 
-#endif
 
 void UBodySetup::ClearPhysicsMeshes()
 {
-#if WITH_PHYSX && PHYSICS_INTERFACE_PHYSX
-
-	FPhysxSharedData::LockAccess();
-
-	for(int32 i=0; i<AggGeom.ConvexElems.Num(); i++)
-	{
-		FKConvexElem* ConvexElem = &(AggGeom.ConvexElems[i]);
-
-		if(ConvexElem->GetConvexMesh() != nullptr)
-		{
-			// put in list for deferred release
-			GPhysXPendingKillConvex.Add(ConvexElem->GetConvexMesh());
-			FPhysxSharedData::Get().Remove(ConvexElem->GetConvexMesh());
-			ConvexElem->SetConvexMesh(nullptr);
-		}
-
-		if(ConvexElem->GetMirroredConvexMesh() != nullptr)
-		{
-			// put in list for deferred release
-			GPhysXPendingKillConvex.Add(ConvexElem->GetMirroredConvexMesh());
-			FPhysxSharedData::Get().Remove(ConvexElem->GetMirroredConvexMesh());
-			ConvexElem->SetMirroredConvexMesh(nullptr);
-		}
-	}
-
-	for(int32 ElementIndex = 0; ElementIndex < TriMeshes.Num(); ++ElementIndex)
-	{
-		GPhysXPendingKillTriMesh.Add(TriMeshes[ElementIndex]);
-		FPhysxSharedData::Get().Remove(TriMeshes[ElementIndex]);
-		TriMeshes[ElementIndex] = NULL;
-	}
-
-	FPhysxSharedData::UnlockAccess();
-
-	TriMeshes.Empty();
-
-#elif WITH_CHAOS
 	for (int32 i = 0; i < AggGeom.ConvexElems.Num(); i++)
 	{
 		FKConvexElem* ConvexElem = &(AggGeom.ConvexElems[i]);
 		ConvexElem->ResetChaosConvexMesh();
 	}
 	ChaosTriMeshes.Reset();
-#endif // WITH_PHYSX
 
 	bCreatedPhysicsMeshes = false;
 
-
 	// Also clear render info
 	AggGeom.FreeRenderInfo();
 }
@@ -939,19 +657,12 @@ void UBodySetup::AddShapesToRigidActor_AssumesLocked(
 	AddParams.Scale = Scale3D;
 	AddParams.SimpleMaterial = SimpleMaterial;
 	AddParams.ComplexMaterials = TArrayView<UPhysicalMaterial*>(ComplexMaterials);
-#if WITH_CHAOS
 	AddParams.ComplexMaterialMasks = TArrayView<FPhysicalMaterialMaskParams>(ComplexMaterialMasks);
-#endif
 	AddParams.LocalTransform = RelativeTM;
 	AddParams.WorldTransform = OwningInstance->GetUnrealWorldTransform();
 	AddParams.Geometry = &AggGeom;
-#if PHYSICS_INTERFACE_PHYSX
-	AddParams.TriMeshes = TArrayView<PxTriangleMesh*>(TriMeshes);
-#endif
-
-#if WITH_CHAOS
 	AddParams.ChaosTriMeshes = MakeArrayView(ChaosTriMeshes);
-#endif
+
 	{
 		SCOPE_CYCLE_COUNTER(STAT_AddGeomToSolver);
 		FPhysicsInterface::AddGeometry(OwningInstance->ActorHandle, AddParams, NewShapes);
@@ -1141,7 +852,7 @@ void UBodySetup::Serialize(FArchive& Ar)
 	AggGeom.FixupDeprecated( Ar );
 #endif
 
-#if WITH_CHAOS && WITH_EDITOR
+#if WITH_EDITOR
 
 	if (Ar.IsLoading())
 	{
@@ -1165,9 +876,9 @@ void UBodySetup::Serialize(FArchive& Ar)
 	if (Ar.CustomVer(FFortniteMainBranchObjectVersion::GUID) >= FFortniteMainBranchObjectVersion::LevelsetSerializationSupportForBodySetup
 		&& Ar.CustomVer(FFortniteReleaseBranchCustomObjectVersion::GUID) < FFortniteReleaseBranchCustomObjectVersion::DisableLevelset_v14_10)
 	{
-		TArray<TSharedPtr<Chaos::FImplicitObject, ESPMode::ThreadSafe>> ChaosImplicitObjects;
-#if WITH_CHAOS
 		using namespace Chaos;
+
+		TArray<TSharedPtr<Chaos::FImplicitObject, ESPMode::ThreadSafe>> ChaosImplicitObjects;
 		FChaosArchive ChaosAr(Ar);
 
 		int32 NumImplicits = 0;
@@ -1209,16 +920,8 @@ void UBodySetup::Serialize(FArchive& Ar)
 			}
 		}
 		*/
-#else
-		if(Ar.IsLoading())
-		{
-			int32 DummyCount;
-			Ar << DummyCount;
-		}
-#endif
 	}
 
-#if WITH_CHAOS
 	if (bCooked && Ar.IsLoading())
 	{
 		// Deserialize bulk inline data inside the serialize function to benefit from async loading thread when possible.
@@ -1227,7 +930,6 @@ void UBodySetup::Serialize(FArchive& Ar)
 			ChaosDerivedDataReader = MakeUnique<FChaosDerivedDataReader<float, 3>>(FormatData);
 		}
 	}
-#endif
 }
 
 void UBodySetup::PostLoad()
@@ -1298,42 +1000,18 @@ void UBodySetup::PostLoad()
 		}
 	}
 
-	// make sure that we load the physX data while the linker's loader is still open
-#if PHYSICS_INTERFACE_PHYSX
-	CreatePhysicsMeshes();
-#elif WITH_CHAOS
 	// If Deferring physics creation, skip so we can call CreatePhysicsMeshes in parallel.
 	if (GEnableDeferredPhysicsCreation == false)
 	{
 		CreatePhysicsMeshes();
 	}
-#endif
-
-
 }
 
 void UBodySetup::UpdateTriMeshVertices(const TArray<FVector> & NewPositions)
 {
 	SCOPE_CYCLE_COUNTER(STAT_UpdateTriMeshVertices);
-#if PHYSICS_INTERFACE_PHYSX
-	if (TriMeshes.Num())
-	{
-		check(TriMeshes[0] != nullptr);
-		PxU32 PNumVerts = TriMeshes[0]->getNbVertices(); // Get num of verts we expect
-		PxVec3 * PNewPositions = TriMeshes[0]->getVerticesForModification();	//we only update the first trimesh. We assume this per poly case is not updating welded trimeshes
-
-		int32 NumToCopy = FMath::Min<int32>(PNumVerts, NewPositions.Num()); // Make sure we don't write off end of array provided
-		for (int32 i = 0; i < NumToCopy; ++i)
-		{
-			PNewPositions[i] = U2PVector(NewPositions[i]);
-		}
 
-		TriMeshes[0]->refitBVH();
-	}
-#elif WITH_CHAOS
 	ensure(false);
-
-#endif
 }
 
 template <bool bPositionAndNormal, bool bUseConvexShapes>
@@ -1488,18 +1166,6 @@ void UBodySetup::ClearCachedCookedPlatformData( const ITargetPlatform* TargetPla
 }
 #endif
 
-#if PHYSICS_INTERFACE_PHYSX
-EPhysXMeshCookFlags UBodySetup::GetRuntimeOnlyCookOptimizationFlags() const
-{
-	EPhysXMeshCookFlags RuntimeCookFlags = EPhysXMeshCookFlags::Default;
-	if(UPhysicsSettings::Get()->bSuppressFaceRemapTable)
-	{
-		RuntimeCookFlags |= EPhysXMeshCookFlags::SuppressFaceRemapTable;
-	}
-	return RuntimeCookFlags;
-}
-#endif
-
 bool UBodySetup::CalcUVAtLocation(const FVector& BodySpaceLocation, int32 FaceIndex, int32 UVChannel, FVector2D& UV) const
 {
 	bool bSuccess = false;
@@ -1681,19 +1347,11 @@ FByteBulkData* UBodySetup::GetCookedData(FName Format, bool bRuntimeOnlyOptimize
 			return nullptr;
 		}
 
-#if PHYSICS_INTERFACE_PHYSX
-		const bool bEligibleForRuntimeOptimization = UseCookedData == &CookedFormatDataRuntimeOnlyOptimization;
-		const EPhysXMeshCookFlags CookingFlags = bEligibleForRuntimeOptimization ? GetRuntimeOnlyCookOptimizationFlags() : EPhysXMeshCookFlags::Default;
-		FDerivedDataPhysXCooker* PhysicsDerivedCooker = new FDerivedDataPhysXCooker(Format, CookingFlags, this, bIsRuntime);
-#elif WITH_CHAOS 
 		// We do not want a FGCObject to be created to prevent garbage collection of our own UBodySetup*
 		// because the scope and lifetime is well defined and FGCObject can't be created on other threads
 		// during garbage collection, which would prevent this function from being run asynchronously.
 		const bool bUseRefHolder = false;
 		FChaosDerivedDataCooker* PhysicsDerivedCooker = new FChaosDerivedDataCooker(this, Format, bUseRefHolder);
-#else
-		static_assert(false, "No cooker defined for this physics interface");
-#endif
 
 		GetDDCBuiltData(Result, *PhysicsDerivedCooker, this, bIsRuntime);
 	}
@@ -1841,32 +1499,6 @@ void UBodySetup::GetResourceSizeEx(FResourceSizeEx& CumulativeResourceSize)
 {
 	Super::GetResourceSizeEx(CumulativeResourceSize);
 
-#if PHYSICS_INTERFACE_PHYSX
-	// Count PhysX trimesh mem usage
-	for(PxTriangleMesh* TriMesh : TriMeshes)
-	{
-		CumulativeResourceSize.AddDedicatedSystemMemoryBytes(GetPhysxObjectSize(TriMesh, NULL));
-	}
-
-	// Count PhysX convex mem usage
-	for(int ConvIdx=0; ConvIdx<AggGeom.ConvexElems.Num(); ConvIdx++)
-	{
-		FKConvexElem& ConvexElem = AggGeom.ConvexElems[ConvIdx];
-
-		if(ConvexElem.GetConvexMesh() != NULL)
-		{
-			CumulativeResourceSize.AddDedicatedSystemMemoryBytes(GetPhysxObjectSize(ConvexElem.GetConvexMesh(), NULL));
-		}
-
-		if(ConvexElem.GetMirroredConvexMesh() != NULL)
-		{
-			CumulativeResourceSize.AddDedicatedSystemMemoryBytes(GetPhysxObjectSize(ConvexElem.GetMirroredConvexMesh(), NULL));
-		}
-	}
-
-#endif // WITH_PHYSX
-
-#if WITH_CHAOS
 	// Cooked mesh data is flushed after mesh creation if not sharing and not in editor, account through rough serilization estimates.
 	if (!GIsEditor && !bSharedCookedData)
 	{
@@ -1898,7 +1530,6 @@ void UBodySetup::GetResourceSizeEx(FResourceSizeEx& CumulativeResourceSize)
 			}
 		}
 	}
-#endif
 
 	if (CookedFormatData.Contains(FPlatformProperties::GetPhysicsFormat()))
 	{
@@ -2010,17 +1641,9 @@ FKConvexElem::FKConvexElem()
 	: FKShapeElem(EAggCollisionShape::Convex)
 	, ElemBox(ForceInit)
 	, Transform(FTransform::Identity)
-#if PHYSICS_INTERFACE_PHYSX
-	, ConvexMesh(NULL)
-	, ConvexMeshNegX(NULL)
-#endif
 {}
 
 FKConvexElem::FKConvexElem(const FKConvexElem& Other)
-#if PHYSICS_INTERFACE_PHYSX
-	: ConvexMesh(nullptr)
-	, ConvexMeshNegX(nullptr)
-#endif
 {
 	CloneElem(Other);
 }
@@ -2032,15 +1655,8 @@ FKConvexElem::~FKConvexElem()
 
 const FKConvexElem& FKConvexElem::operator=(const FKConvexElem& Other)
 {
-#if PHYSICS_INTERFACE_PHYSX
-	ensureMsgf(ConvexMesh == nullptr, TEXT("We are leaking memory. Why are we calling the assignment operator on an element that has already allocated resources?"));
-	ensureMsgf(ConvexMeshNegX == nullptr, TEXT("We are leaking memory. Why are we calling the assignment operator on an element that has already allocated resources?"));
-	ConvexMesh = nullptr;
-	ConvexMeshNegX = nullptr;
-#elif WITH_CHAOS
 	ensureMsgf(!ChaosConvex, TEXT("We are leaking memory. Why are we calling the assignment operator on an element that has already allocated resources?"));
 	ResetChaosConvexMesh();
-#endif
 	CloneElem(Other);
 	return *this;
 }
@@ -2069,27 +1685,6 @@ float SignedVolumeOfTriangle(const FVector3f& p1, const FVector3f& p2, const FVe
 {
 	return FVector3f::DotProduct(p1, FVector3f::CrossProduct(p2, p3)) / 6.0f;
 }
-#if PHYSICS_INTERFACE_PHYSX
-physx::PxConvexMesh* FKConvexElem::GetConvexMesh() const
-{
-	return ConvexMesh;
-}
-
-void FKConvexElem::SetConvexMesh(physx::PxConvexMesh* InMesh)
-{
-	ConvexMesh = InMesh;
-}
-
-physx::PxConvexMesh* FKConvexElem::GetMirroredConvexMesh() const
-{
-	return ConvexMeshNegX;
-}
-
-void FKConvexElem::SetMirroredConvexMesh(physx::PxConvexMesh* InMesh)
-{
-	ConvexMeshNegX = InMesh;
-}
-#endif
 
 FVector::FReal FKConvexElem::GetVolume(const FVector& Scale) const
 {
@@ -2100,49 +1695,14 @@ FVector::FReal FKConvexElem::GetScaledVolume(const FVector& Scale) const
 {
 	float Volume = 0.0f;
 
-#if PHYSICS_INTERFACE_PHYSX
-	if (ConvexMesh != NULL)
-	{
-		// Preparation for convex mesh scaling implemented in another changelist
-		FTransform ScaleTransform = FTransform(FQuat::Identity, FVector::ZeroVector, Scale);
-
-		int32 NumPolys = ConvexMesh->getNbPolygons();
-		PxHullPolygon PolyData;
-
-		const PxVec3* Vertices = ConvexMesh->getVertices();
-		const PxU8* Indices = ConvexMesh->getIndexBuffer();
-
-		for (int32 PolyIdx = 0; PolyIdx < NumPolys; ++PolyIdx)
-		{
-			if (ConvexMesh->getPolygonData(PolyIdx, PolyData))
-			{
-				for (int32 VertIdx = 2; VertIdx < PolyData.mNbVerts; ++ VertIdx)
-				{
-					// Grab triangle indices that we hit
-					int32 I0 = Indices[PolyData.mIndexBase + 0];
-					int32 I1 = Indices[PolyData.mIndexBase + (VertIdx - 1)];
-					int32 I2 = Indices[PolyData.mIndexBase + VertIdx];
-
-
-					Volume += SignedVolumeOfTriangle(ScaleTransform.TransformPosition(P2UVector(Vertices[I0])), 
-						ScaleTransform.TransformPosition(P2UVector(Vertices[I1])), 
-						ScaleTransform.TransformPosition(P2UVector(Vertices[I2])));
-				}
-			}
-		}
-	}
-#elif WITH_CHAOS
 	if (ChaosConvex != nullptr)
 	{
 		Volume = FMath::Abs(Scale.X * Scale.Y * Scale.Z) * ChaosConvex->GetVolume();
 	}
-#endif
 
 	return Volume;
 }
 
-#if WITH_CHAOS
-
 void FKConvexElem::SetChaosConvexMesh(TSharedPtr<Chaos::FConvex, ESPMode::ThreadSafe>&& InChaosConvex)
 {
 	ChaosConvex = MoveTemp(InChaosConvex);
@@ -2194,7 +1754,6 @@ TArray<int32> FKConvexElem::GetChaosConvexIndices() const
 
 	return ResultIndexData;
 }
-#endif
 
 #if WITH_EDITORONLY_DATA
 void FKSphereElem::FixupDeprecated( FArchive& Ar )
diff --git a/Engine/Source/Runtime/Engine/Private/PhysicsEngine/BodyUtils.cpp b/Engine/Source/Runtime/Engine/Private/PhysicsEngine/BodyUtils.cpp
--- a/Engine/Source/Runtime/Engine/Private/PhysicsEngine/BodyUtils.cpp
+++ b/Engine/Source/Runtime/Engine/Private/PhysicsEngine/BodyUtils.cpp
@@ -5,18 +5,9 @@
 #include "PhysicsEngine/BodyInstance.h"
 #include "PhysicalMaterials/PhysicalMaterial.h"
 
-#if WITH_CHAOS
 #include "Chaos/MassProperties.h"
 #include "Chaos/Utilities.h"
 #include "Physics/Experimental/ChaosInterfaceUtils.h"
-#endif
-
-#if PHYSICS_INTERFACE_PHYSX
-#include "PhysXPublic.h"
-#include "PhysicsEngine/PhysXSupport.h"
-#endif
-
-//PRAGMA_DISABLE_OPTIMIZATION
 
 namespace BodyUtils
 {
@@ -47,7 +38,6 @@ namespace BodyUtils
 		return DensityKGPerCubicUU;
 	}
 
-#if WITH_CHAOS
 	Chaos::FMassProperties ApplyMassPropertiesModifiers(const FBodyInstance* OwningBodyInstance, Chaos::FMassProperties MassProps, const FTransform& MassModifierTransform, const bool bInertaScaleIncludeMass)
 	{
 		float OldMass = MassProps.Mass;
@@ -121,52 +111,4 @@ namespace BodyUtils
 		// Apply the BodyInstance's mass and inertia modifiers
 		return ApplyMassPropertiesModifiers(OwningBodyInstance, MassProps, MassModifierTransform, bInertaScaleIncludeMass);
 	}
-
-
-#elif PHYSICS_INTERFACE_PHYSX
-
-	/** Computes and adds the mass properties (inertia, com, etc...) based on the mass settings of the body instance. */
-	PxMassProperties ComputeMassProperties(const FBodyInstance* OwningBodyInstance, TArray<FPhysicsShapeHandle> Shapes, const FTransform& MassModifierTransform, const bool bUnused)
-	{
-		// physical material - nothing can weigh less than hydrogen (0.09 kg/m^3)
-		float DensityKGPerCubicUU = 1.0f;
-		float RaiseMassToPower = 0.75f;
-		if (UPhysicalMaterial* PhysMat = OwningBodyInstance->GetSimplePhysicalMaterial())
-		{
-			DensityKGPerCubicUU = FMath::Max(KgPerM3ToKgPerCm3(0.09f), gPerCm3ToKgPerCm3(PhysMat->Density));
-			RaiseMassToPower = PhysMat->RaiseMassToPower;
-		}
-
-		PxMassProperties MassProps;
-		FPhysicsInterface::CalculateMassPropertiesFromShapeCollection(MassProps, Shapes, DensityKGPerCubicUU);
-
-		float OldMass = MassProps.mass;
-		float NewMass = 0.f;
-
-		if (OwningBodyInstance->bOverrideMass == false)
-		{
-			float UsePow = FMath::Clamp<float>(RaiseMassToPower, KINDA_SMALL_NUMBER, 1.f);
-			NewMass = FMath::Pow(OldMass, UsePow);
-
-			// Apply user-defined mass scaling.
-			NewMass = FMath::Max(OwningBodyInstance->MassScale * NewMass, 0.001f);	//min weight of 1g
-		}
-		else
-		{
-			NewMass = FMath::Max(OwningBodyInstance->GetMassOverride(), 0.001f);	//min weight of 1g
-		}
-
-		check(NewMass > 0.f);
-
-		float MassRatio = NewMass / OldMass;
-
-		PxMassProperties FinalMassProps = MassProps * MassRatio;
-
-		FinalMassProps.centerOfMass += U2PVector(MassModifierTransform.TransformVector(OwningBodyInstance->COMNudge));
-		FinalMassProps.inertiaTensor = PxMassProperties::scaleInertia(FinalMassProps.inertiaTensor, PxQuat(PxIdentity), U2PVector(OwningBodyInstance->InertiaTensorScale));
-
-		return FinalMassProps;
-	}
-#endif
-
-}
\ No newline at end of file
+}
diff --git a/Engine/Source/Runtime/Engine/Private/PhysicsEngine/BodyUtils.h b/Engine/Source/Runtime/Engine/Private/PhysicsEngine/BodyUtils.h
--- a/Engine/Source/Runtime/Engine/Private/PhysicsEngine/BodyUtils.h
+++ b/Engine/Source/Runtime/Engine/Private/PhysicsEngine/BodyUtils.h
@@ -7,28 +7,19 @@
 #include "Containers/ContainerAllocationPolicies.h"
 #include "PhysicsInterfaceDeclaresCore.h"
 
-#if PHYSICS_INTERFACE_PHYSX
-#include "extensions/PxMassProperties.h"
-#endif
-
 struct FBodyInstance; 
 
-#if WITH_CHAOS
 namespace Chaos
 {
 	struct FMassProperties;
-
 	using FShapesArray = TArray<TUniquePtr<FPerShapeData>, TInlineAllocator<1>>;
 }
-#endif
 
 /**
  * Utility methods for use by BodyInstance and ImmediatePhysics
  */
 namespace BodyUtils
 {
-#if WITH_CHAOS
-
 	/** 
 	 * Computes and adds the mass properties (inertia, com, etc...) based on the mass settings of the body instance. 
 	 * The inertia returned will be diagonal, and there may be a non-identity rotation of mass.
@@ -36,12 +27,4 @@ namespace BodyUtils
 	 */
 	Chaos::FMassProperties ComputeMassProperties(const FBodyInstance* OwningBodyInstance, const TArray<FPhysicsShapeHandle>& Shapes, const FTransform& MassModifierTransform, const bool bInertaScaleIncludeMass = false);
 	Chaos::FMassProperties ComputeMassProperties(const FBodyInstance* OwningBodyInstance, const Chaos::FShapesArray& Shapes, const TArray<bool>& bContributesToMass, const FTransform& MassModifierTransform, const bool bInertaScaleIncludeMass = false);
-
-#elif PHYSICS_INTERFACE_PHYSX
-	
-	/** Computes and adds the mass properties (inertia, com, etc...) based on the mass settings of the body instance. */
-	physx::PxMassProperties ComputeMassProperties(const FBodyInstance* OwningBodyInstance, TArray<FPhysicsShapeHandle> Shapes, const FTransform& MassModifierTransform, const bool bUnused = false);
-
-#endif
-
 }
\ No newline at end of file
diff --git a/Engine/Source/Runtime/Engine/Private/PhysicsEngine/CollisionAnalyzerCapture.h b/Engine/Source/Runtime/Engine/Private/PhysicsEngine/CollisionAnalyzerCapture.h
--- a/Engine/Source/Runtime/Engine/Private/PhysicsEngine/CollisionAnalyzerCapture.h
+++ b/Engine/Source/Runtime/Engine/Private/PhysicsEngine/CollisionAnalyzerCapture.h
@@ -29,12 +29,7 @@ inline static void ConvertGeometryCollection(const FPhysicsGeometryCollection& G
 		{
 			OutQueryShape = ECAQueryShape::Capsule;
 			const auto& CapsuleGeom = GeomCollection.GetCapsuleGeometry();
-#if PHYSICS_INTERFACE_PHYSX
-			OutDims = FVector(CapsuleGeom.radius, CapsuleGeom.radius, CapsuleGeom.halfHeight + CapsuleGeom.radius);
-			InOutRot = ConvertToUECapsuleRot(InOutRot);
-#else
 			OutDims = FVector(CapsuleGeom.GetRadius(), CapsuleGeom.GetRadius(), CapsuleGeom.GetHeight() * 0.5f + CapsuleGeom.GetRadius());
-#endif
 			break;
 		}
 
@@ -42,11 +37,7 @@ inline static void ConvertGeometryCollection(const FPhysicsGeometryCollection& G
 		{
 			OutQueryShape = ECAQueryShape::Sphere;
 			const auto& SphereGeom = GeomCollection.GetSphereGeometry();
-#if PHYSICS_INTERFACE_PHYSX
-			OutDims = FVector(SphereGeom.radius);
-#else
 			OutDims = FVector(SphereGeom.GetRadius());
-#endif
 			break;
 		}
 
@@ -54,11 +45,7 @@ inline static void ConvertGeometryCollection(const FPhysicsGeometryCollection& G
 		{
 			OutQueryShape = ECAQueryShape::Box;
 			const auto& BoxGeom = GeomCollection.GetBoxGeometry();
-#if PHYSICS_INTERFACE_PHYSX
-			OutDims = P2UVector(BoxGeom.halfExtents);
-#else
 			OutDims = FVector(BoxGeom.Extents() * 0.5f);
-#endif
 			break;
 		}
 
diff --git a/Engine/Source/Runtime/Engine/Private/PhysicsEngine/CollisionQueryFilterCallback.cpp b/Engine/Source/Runtime/Engine/Private/PhysicsEngine/CollisionQueryFilterCallback.cpp
--- a/Engine/Source/Runtime/Engine/Private/PhysicsEngine/CollisionQueryFilterCallback.cpp
+++ b/Engine/Source/Runtime/Engine/Private/PhysicsEngine/CollisionQueryFilterCallback.cpp
@@ -8,10 +8,6 @@
 #include "Components/PrimitiveComponent.h"
 #include "Collision.h"
 
-#if PHYSICS_INTERFACE_PHYSX
-#include "PhysXInterfaceWrapper.h"
-#endif
-
 #include "ChaosInterfaceWrapperCore.h"
 #include "Physics/Experimental/ChaosInterfaceWrapper.h"
 #include "Chaos/GeometryParticles.h"
@@ -87,28 +83,6 @@ ECollisionQueryHitType FCollisionQueryFilterCallback::CalcQueryHitType(const FCo
 	return ECollisionQueryHitType::None;
 }
 
-#if PHYSICS_INTERFACE_PHYSX
-ECollisionQueryHitType FCollisionQueryFilterCallback::PreFilterImp(const FCollisionFilterData& FilterData, const physx::PxShape& Shape, const physx::PxActor& Actor)
-{
-	//SCOPE_CYCLE_COUNTER(STAT_Collision_PreFilter);
-	FCollisionFilterData ShapeFilter = PhysXInterface::GetQueryFilterData(Shape);
-
-	// We usually don't have ignore components so we try to avoid the virtual getSimulationFilterData() call below. 'word2' of shape sim filter data is componentID.
-	uint32 ComponentID = 0;
-	if (IgnoreComponents.Num() > 0)
-	{
-		ComponentID = PhysXInterface::GetSimulationFilterData(Shape).Word2;
-	}
-
-	FBodyInstance* BodyInstance = nullptr;
-#if ENABLE_PREFILTER_LOGGING || DETECT_SQ_HITCHES
-	BodyInstance = PhysXInterface::GetUserData(Actor);
-#endif // ENABLE_PREFILTER_LOGGING || DETECT_SQ_HITCHES
-
-	return PreFilterImp(FilterData, ShapeFilter, ComponentID, BodyInstance);
-}
-#endif
-
 template <typename TParticle>
 ECollisionQueryHitType FCollisionQueryFilterCallback::PreFilterBaseImp(const FCollisionFilterData& FilterData, const Chaos::FPerShapeData& Shape, const TParticle& Actor)
 {
@@ -273,22 +247,6 @@ ECollisionQueryHitType FCollisionQueryFilterCallback::PostFilterImp(const FColli
 	}
 }
 
-#if PHYSICS_INTERFACE_PHYSX
-ECollisionQueryHitType FCollisionQueryFilterCallback::PostFilterImp(const FCollisionFilterData& FilterData, const physx::PxQueryHit& Hit)
-{
-	// Unused in non-sweeps
-	if (!bIsSweep)
-	{
-		return ECollisionQueryHitType::None;
-	}
-
-	const FHitLocation& SweepHit = (const FHitLocation&)Hit;
-	const bool bIsOverlap = HadInitialOverlap(SweepHit);
-
-	return PostFilterImp(FilterData, bIsOverlap);
-}
-#endif
-
 ECollisionQueryHitType FCollisionQueryFilterCallback::PostFilterImp(const FCollisionFilterData& FilterData, const ChaosInterface::FQueryHit& Hit)
 {
 	// Unused in non-sweeps
@@ -316,53 +274,3 @@ ECollisionQueryHitType FCollisionQueryFilterCallback::PostFilterImp(const FColli
 
 	return PostFilterImp(FilterData, bIsOverlap);
 }
-
-
-#if PHYSICS_INTERFACE_PHYSX
-
-PxQueryHitType::Enum FCollisionQueryFilterCallback::preFilter(const PxFilterData& filterData, const PxShape* shape, const PxRigidActor* actor, PxHitFlags& queryFlags)
-{
-	ensureMsgf(shape, TEXT("Invalid shape encountered in FCollisionQueryFilterCallback::preFilter, actor: %p, filterData: %x %x %x %x"), actor, filterData.word0, filterData.word1, filterData.word2, filterData.word3);
-
-	if (!shape)
-	{
-		// Early out to avoid crashing.
-		return U2PCollisionQueryHitType(PreFilterReturnValue = ECollisionQueryHitType::None);
-	}
-
-	FCollisionFilterData FilterData = P2UFilterData(filterData);
-	FCollisionFilterData ShapeFilter = P2UFilterData(shape->getQueryFilterData());
-
-	// We usually don't have ignore components so we try to avoid the virtual getSimulationFilterData() call below. 'word2' of shape sim filter data is componentID.
-	uint32 ComponentID = 0;
-	if (IgnoreComponents.Num() > 0)
-	{
-		ComponentID = shape->getSimulationFilterData().word2;
-	}
-
-	FBodyInstance* BodyInstance = nullptr;
-#if ENABLE_PREFILTER_LOGGING || DETECT_SQ_HITCHES
-	BodyInstance = FPhysxUserData::Get<FBodyInstance>(actor->userData);
-#endif // ENABLE_PREFILTER_LOGGING || DETECT_SQ_HITCHES
-
-	return U2PCollisionQueryHitType(PreFilterImp(FilterData, ShapeFilter, ComponentID, BodyInstance));
-}
-
-PxQueryHitType::Enum FCollisionQueryFilterCallback::postFilter(const PxFilterData& filterData, const PxQueryHit& hit)
-{
-	// Unused in non-sweeps
-	if (!bIsSweep)
-	{
-		return PxQueryHitType::eNONE;
-	}
-
-	FCollisionFilterData FilterData = P2UFilterData(filterData);
-
-	PxSweepHit& SweepHit = (PxSweepHit&)hit;
-	const bool bIsOverlap = SweepHit.hadInitialOverlap();
-
-	const ECollisionQueryHitType HitType = PostFilterImp(FilterData, bIsOverlap);
-	return U2PCollisionQueryHitType(HitType);
-}
-
-#endif
\ No newline at end of file
diff --git a/Engine/Source/Runtime/Engine/Private/PhysicsEngine/CollisionQueryFilterCallback.h b/Engine/Source/Runtime/Engine/Private/PhysicsEngine/CollisionQueryFilterCallback.h
--- a/Engine/Source/Runtime/Engine/Private/PhysicsEngine/CollisionQueryFilterCallback.h
+++ b/Engine/Source/Runtime/Engine/Private/PhysicsEngine/CollisionQueryFilterCallback.h
@@ -1,15 +1,11 @@
 // Copyright Epic Games, Inc. All Rights Reserved.
 
-
 #pragma once
+
 #include "CollisionQueryFilterCallbackCore.h"
 #include "CollisionQueryParams.h"
 #include "Physics/PhysicsInterfaceTypes.h"
 
-#if PHYSICS_INTERFACE_PHYSX
-#include "PhysXInterfaceWrapper.h"
-#endif
-
 #define ENABLE_PREFILTER_LOGGING 0
 
 struct FBodyInstance;
@@ -64,11 +60,6 @@ public:
 
 	ECollisionQueryHitType PostFilterImp(const FCollisionFilterData& FilterData, bool bIsOverlap);
 
-#if PHYSICS_INTERFACE_PHYSX
-	ECollisionQueryHitType PostFilterImp(const FCollisionFilterData& FilterData, const physx::PxQueryHit& Hit);
-	ECollisionQueryHitType PreFilterImp(const FCollisionFilterData& FilterData, const physx::PxShape& Shape, const physx::PxActor& Actor);
-#endif
-
 	ECollisionQueryHitType PreFilterImp(const FCollisionFilterData& FilterData, const Chaos::FPerShapeData& Shape, const Chaos::FGeometryParticle& Actor);
 	ECollisionQueryHitType PreFilterImp(const FCollisionFilterData& FilterData, const Chaos::FPerShapeData& Shape, const Chaos::FGeometryParticleHandle& Actor);
 
@@ -95,13 +86,6 @@ public:
 		return PreFilterImp(FilterData, Shape, Actor);
 	}
 
-#if PHYSICS_INTERFACE_PHYSX
-	virtual ECollisionQueryHitType PostFilter(const FCollisionFilterData& FilterData, const physx::PxQueryHit& Hit) override { return PostFilterImp(FilterData, Hit); }
-	virtual ECollisionQueryHitType PreFilter(const FCollisionFilterData& FilterData, const physx::PxShape& Shape, physx::PxRigidActor& Actor) override { return PreFilterImp(FilterData, Shape, Actor); }
-	virtual PxQueryHitType::Enum preFilter(const PxFilterData& filterData, const PxShape* shape, const PxRigidActor* actor, PxHitFlags& queryFlags) override;
-	virtual PxQueryHitType::Enum postFilter(const PxFilterData& filterData, const PxQueryHit& hit) override;
-#endif
-
 #if DETECT_SQ_HITCHES
 	// Util struct to record what preFilter was called with
 	struct FPreFilterRecord
diff --git a/Engine/Source/Runtime/Engine/Private/PhysicsEngine/ConstraintInstance.cpp b/Engine/Source/Runtime/Engine/Private/PhysicsEngine/ConstraintInstance.cpp
--- a/Engine/Source/Runtime/Engine/Private/PhysicsEngine/ConstraintInstance.cpp
+++ b/Engine/Source/Runtime/Engine/Private/PhysicsEngine/ConstraintInstance.cpp
@@ -196,29 +196,23 @@ void FConstraintInstance::UpdatePlasticity()
 
 void FConstraintInstance::UpdateContactTransferScale()
 {
-#if WITH_CHAOS
 	FPhysicsInterface::ExecuteOnUnbrokenConstraintReadWrite(ConstraintHandle, [&](const FPhysicsConstraintHandle& InConstraint)
 		{
 			ProfileInstance.UpdateContactTransferScale_AssumesLocked(InConstraint);
 		});
-#endif
 }
 
 void FConstraintProfileProperties::UpdatePlasticity_AssumesLocked(const FPhysicsConstraintHandle& InConstraintRef) const
 {
-#if WITH_CHAOS
 	const float LinearPlasticityLimit = bLinearPlasticity ? LinearPlasticityThreshold : FLT_MAX;
 	const float AngularPlasticityLimit = bAngularPlasticity ? FMath::DegreesToRadians(AngularPlasticityThreshold) : UE_MAX_FLT;
 
 	FPhysicsInterface::SetPlasticityLimits_AssumesLocked(InConstraintRef, LinearPlasticityLimit, AngularPlasticityLimit, LinearPlasticityType);
-#endif
 }
 
 void FConstraintProfileProperties::UpdateContactTransferScale_AssumesLocked(const FPhysicsConstraintHandle& InConstraintRef) const
 {
-#if WITH_CHAOS
 	FPhysicsInterface::SetContactTransferScale_AssumesLocked(InConstraintRef, ContactTransferScale);
-#endif
 }
 
 void FConstraintInstance::UpdateDriveTarget()
@@ -238,9 +232,7 @@ FConstraintInstanceBase::FConstraintInstanceBase()
 void FConstraintInstanceBase::Reset()
 {
 	ConstraintIndex = 0;
-#if WITH_CHAOS
 	ConstraintHandle.Reset();
-#endif
 	PhysScene = nullptr;
 }
 
@@ -421,11 +413,7 @@ bool GetActorRefs(FBodyInstance* Body1, FBodyInstance* Body2, FPhysicsActorHandl
 
 bool FConstraintInstance::CreateJoint_AssumesLocked(const FPhysicsActorHandle& InActorRef1, const FPhysicsActorHandle& InActorRef2)
 {
-#if WITH_CHAOS
 	LLM_SCOPE(ELLMTag::ChaosConstraint);
-#else
-	LLM_SCOPE(ELLMTag::PhysX);
-#endif
 
 	FTransform Local1 = GetRefFrame(EConstraintFrame::Frame1);
 	if(FPhysicsInterface::IsValid(InActorRef1))
@@ -465,12 +453,8 @@ void FConstraintProfileProperties::UpdateConstraintFlags_AssumesLocked(const FPh
 #endif
 
 	FPhysicsInterface::SetCollisionEnabled(InConstraintRef, !bDisableCollision);
-#if WITH_CHAOS
 	FPhysicsInterface::SetProjectionEnabled_AssumesLocked(InConstraintRef, bEnableProjection, ProjectionLinearAlpha, ProjectionAngularAlpha, ProjectionLinearTolerance, ProjectionAngularTolerance);
 	FPhysicsInterface::SetShockPropagationEnabled_AssumesLocked(InConstraintRef, bEnableShockPropagation, ShockPropagationAlpha);
-#else
-	FPhysicsInterface::SetProjectionEnabled_AssumesLocked(InConstraintRef, bEnableProjection, ProjectionLinearTolerance, ProjectionAngularTolerance);
-#endif
 	FPhysicsInterface::SetParentDominates_AssumesLocked(InConstraintRef, bParentDominates);
 }
 
@@ -495,12 +479,10 @@ void FConstraintInstance::InitConstraint(FBodyInstance* Body1, FBodyInstance* Bo
 			return;
 		}
 
-#if WITH_CHAOS
 		if (!bAllowKinematicKinematicConstraints && (!FPhysicsInterface::IsValid(Actor1) || FPhysicsInterface::IsKinematic(Actor1)) && (!FPhysicsInterface::IsValid(Actor2) || FPhysicsInterface::IsKinematic(Actor2)))
 		{
 			return;
 		}
-#endif
 
 		FPhysicsCommand::ExecuteWrite(Actor1, Actor2, [&](const FPhysicsActorHandle& ActorA, const FPhysicsActorHandle& ActorB)
 		{
@@ -641,9 +623,7 @@ void FConstraintInstance::CopyConstraintPhysicalPropertiesFrom(const FConstraint
 
 	// Recover internal data we'd like to keep - i.e. bone indices, etc.
 	ConstraintIndex = OldInstance.ConstraintIndex;
-#if WITH_PHYSX
 	ConstraintHandle = OldInstance.ConstraintHandle;
-#endif	//WITH_PHYSX
 	JointName = OldInstance.JointName;
 	ConstraintBone1 = OldInstance.ConstraintBone1;
 	ConstraintBone2 = OldInstance.ConstraintBone2;
@@ -1015,12 +995,10 @@ void FConstraintInstance::SetAngularDOFLimitScale(float InSwing1LimitScale, floa
 void FConstraintInstance::SetLinearLimitSize(float NewLimitSize)
 {
 	//TODO: Is this supposed to be scaling the linear limit? The code just sets it directly.
-#if WITH_PHYSX
 	FPhysicsInterface::ExecuteOnUnbrokenConstraintReadWrite(ConstraintHandle, [&](const FPhysicsConstraintHandle& InUnbrokenConstraint)
 	{
 		FPhysicsInterface::SetLinearLimit(ConstraintHandle, NewLimitSize);
 	});
-#endif
 }
 
 bool FConstraintInstance::Serialize(FArchive& Ar)
@@ -1198,62 +1176,37 @@ void FConstraintInstance::SetProjectionParams(bool bEnableProjection, float Proj
 
 	FPhysicsCommand::ExecuteWrite(ConstraintHandle, [&](const FPhysicsConstraintHandle& Constraint)
 		{
-#if WITH_CHAOS
 			const float LinearAlpha = bEnableProjection ? ProjectionLinearAlpha : 0.0f;
 			const float AngularAlpha = bEnableProjection ? ProjectionAngularAlpha : 0.0f;
 			const float TeleportDistance = bEnableProjection ? ProjectionLinearTolerance : -1.0f;
 			const float TeleportAngle = bEnableProjection ? ProjectionAngularTolerance : -1.0f;
 			FPhysicsInterface::SetProjectionEnabled_AssumesLocked(Constraint, bEnableProjection, LinearAlpha, AngularAlpha, TeleportDistance, TeleportAngle);
-#else
-			if (bEnableProjection)
-			{
-				FPhysicsInterface::SetProjectionEnabled_AssumesLocked(Constraint, true, ProjectionLinearTolerance, ProjectionAngularTolerance);
-			}
-			else
-			{
-				FPhysicsInterface::SetProjectionEnabled_AssumesLocked(Constraint, false);
-			}
-#endif
 		});
 }
 
 float FConstraintInstance::GetShockPropagationAlpha() const
 {
-#if WITH_CHAOS
 	return ProfileInstance.ShockPropagationAlpha;
-#else
-	return 0;
-#endif
 }
 
 void FConstraintInstance::SetShockPropagationParams(bool bEnableShockPropagation, float ShockPropagationAlpha)
 {
-#if WITH_CHAOS
 	ProfileInstance.bEnableShockPropagation = bEnableShockPropagation;
 	ProfileInstance.ShockPropagationAlpha = ShockPropagationAlpha;
 	FPhysicsCommand::ExecuteWrite(ConstraintHandle, [&](const FPhysicsConstraintHandle& Constraint)
 		{
 			FPhysicsInterface::SetShockPropagationEnabled_AssumesLocked(Constraint, bEnableShockPropagation, ShockPropagationAlpha);
 		});
-#endif
 }
 
 void FConstraintInstance::EnableProjection()
 {
-#if WITH_CHAOS
 	SetProjectionParams(true, ProfileInstance.ProjectionLinearAlpha, ProfileInstance.ProjectionAngularAlpha, ProfileInstance.ProjectionLinearTolerance, ProfileInstance.ProjectionAngularTolerance);
-#else
-	SetProjectionParams(true, ProfileInstance.ProjectionLinearTolerance, ProfileInstance.ProjectionAngularTolerance);
-#endif
 }
 
 void FConstraintInstance::DisableProjection()
 {
-#if WITH_CHAOS
 	SetProjectionParams(false, ProfileInstance.ProjectionLinearAlpha, ProfileInstance.ProjectionAngularAlpha, ProfileInstance.ProjectionLinearTolerance, ProfileInstance.ProjectionAngularTolerance);
-#else
-	SetProjectionParams(false, ProfileInstance.ProjectionLinearTolerance, ProfileInstance.ProjectionAngularTolerance);
-#endif
 }
 
 void FConstraintInstance::EnableParentDominates()
diff --git a/Engine/Source/Runtime/Engine/Private/PhysicsEngine/ConstraintTypes.cpp b/Engine/Source/Runtime/Engine/Private/PhysicsEngine/ConstraintTypes.cpp
--- a/Engine/Source/Runtime/Engine/Private/PhysicsEngine/ConstraintTypes.cpp
+++ b/Engine/Source/Runtime/Engine/Private/PhysicsEngine/ConstraintTypes.cpp
@@ -11,29 +11,6 @@ extern TAutoConsoleVariable<float> CVarConstraintLinearStiffnessScale;
 extern TAutoConsoleVariable<float> CVarConstraintAngularDampingScale;
 extern TAutoConsoleVariable<float> CVarConstraintAngularStiffnessScale;
 
-#if PHYSICS_INTERFACE_PHYSX
-
-enum class ESoftLimitTypeHelper
-{
-	Linear,
-	Angular
-};
-
-/** Util for setting soft limit params */
-template <ESoftLimitTypeHelper Type>
-void SetSoftLimitParams_AssumesLocked(PxJointLimitParameters* PLimit, bool bSoft, float Spring, float Damping)
-{
-	if(bSoft)
-	{
-		const float SpringCoeff = Type == ESoftLimitTypeHelper::Angular ? CVarConstraintAngularStiffnessScale.GetValueOnGameThread() : CVarConstraintLinearStiffnessScale.GetValueOnGameThread();
-		const float DampingCoeff = Type == ESoftLimitTypeHelper::Angular ? CVarConstraintAngularDampingScale.GetValueOnGameThread() : CVarConstraintLinearDampingScale.GetValueOnGameThread();
-		PLimit->stiffness = Spring * SpringCoeff;
-		PLimit->damping = Damping * DampingCoeff;
-	}
-}
-
-#endif //WITH_PHYSX
-
 /** Util for setting linear movement for an axis */
 void SetLinearMovement_AssumesLocked(const FPhysicsConstraintHandle& InConstraintRef, PhysicsInterfaceTypes::ELimitAxis InAxis, ELinearConstraintMotion Motion, bool bLockLimitSize, bool bSkipSoftLimit)
 {
diff --git a/Engine/Source/Runtime/Engine/Private/PhysicsEngine/Experimental/ChaosCooking.cpp b/Engine/Source/Runtime/Engine/Private/PhysicsEngine/Experimental/ChaosCooking.cpp
--- a/Engine/Source/Runtime/Engine/Private/PhysicsEngine/Experimental/ChaosCooking.cpp
+++ b/Engine/Source/Runtime/Engine/Private/PhysicsEngine/Experimental/ChaosCooking.cpp
@@ -7,8 +7,6 @@
 #include "Chaos/Convex.h"
 #include "Chaos/TriangleMeshImplicitObject.h"
 
-#if WITH_CHAOS
-
 int32 EnableMeshClean = 1;
 FAutoConsoleVariableRef CVarEnableMeshClean(TEXT("p.EnableMeshClean"), EnableMeshClean, TEXT("Enable/Disable mesh cleanup during cook."));
 
@@ -380,5 +378,3 @@ namespace Chaos
 		return CookInfo.bCookTriMesh || CookInfo.bCookNonMirroredConvex || CookInfo.bCookMirroredConvex;
 	}
 }
-
-#endif
\ No newline at end of file
diff --git a/Engine/Source/Runtime/Engine/Private/PhysicsEngine/Experimental/ChaosCooking.h b/Engine/Source/Runtime/Engine/Private/PhysicsEngine/Experimental/ChaosCooking.h
--- a/Engine/Source/Runtime/Engine/Private/PhysicsEngine/Experimental/ChaosCooking.h
+++ b/Engine/Source/Runtime/Engine/Private/PhysicsEngine/Experimental/ChaosCooking.h
@@ -7,8 +7,6 @@
 #include "Async/TaskGraphInterfaces.h"
 #include "PhysicsEngine/BodySetup.h"
 
-#if WITH_CHAOS
-
 struct FTriMeshCollisionData;
 
 namespace Chaos
@@ -46,5 +44,4 @@ namespace Chaos
 		UBodySetup* SourceSetup;
 		FCookBodySetupInfo CookInfo;
 	};
-}
-#endif
\ No newline at end of file
+}
\ No newline at end of file
diff --git a/Engine/Source/Runtime/Engine/Private/PhysicsEngine/Experimental/ChaosDerivedData.cpp b/Engine/Source/Runtime/Engine/Private/PhysicsEngine/Experimental/ChaosDerivedData.cpp
--- a/Engine/Source/Runtime/Engine/Private/PhysicsEngine/Experimental/ChaosDerivedData.cpp
+++ b/Engine/Source/Runtime/Engine/Private/PhysicsEngine/Experimental/ChaosDerivedData.cpp
@@ -2,8 +2,6 @@
 
 #include "ChaosDerivedData.h"
 
-#if WITH_CHAOS
-
 #include "CoreMinimal.h"
 #include "Chaos/Core.h"
 #include "Chaos/CollisionConvexMesh.h"
@@ -131,5 +129,3 @@ FChaosDerivedDataCooker::FChaosDerivedDataCooker(UBodySetup* InSetup, FName InFo
 }
 
 
-#endif
- 
diff --git a/Engine/Source/Runtime/Engine/Private/PhysicsEngine/Experimental/ChaosDerivedData.h b/Engine/Source/Runtime/Engine/Private/PhysicsEngine/Experimental/ChaosDerivedData.h
--- a/Engine/Source/Runtime/Engine/Private/PhysicsEngine/Experimental/ChaosDerivedData.h
+++ b/Engine/Source/Runtime/Engine/Private/PhysicsEngine/Experimental/ChaosDerivedData.h
@@ -2,8 +2,6 @@
 
 #pragma once
 
-#if WITH_CHAOS
-
 #include "CoreMinimal.h"
 #include "DerivedDataPluginInterface.h"
 #include "UObject/GCObject.h"
@@ -68,4 +66,3 @@ private:
 };
 
 
-#endif
\ No newline at end of file
diff --git a/Engine/Source/Runtime/Engine/Private/PhysicsEngine/Experimental/ChaosDerivedDataUtil.cpp b/Engine/Source/Runtime/Engine/Private/PhysicsEngine/Experimental/ChaosDerivedDataUtil.cpp
--- a/Engine/Source/Runtime/Engine/Private/PhysicsEngine/Experimental/ChaosDerivedDataUtil.cpp
+++ b/Engine/Source/Runtime/Engine/Private/PhysicsEngine/Experimental/ChaosDerivedDataUtil.cpp
@@ -23,6 +23,7 @@ namespace Chaos
 
 	void CleanTrimesh(TArray<FVector3f>& InOutVertices, TArray<int32>& InOutIndices, TArray<int32>* OutOptFaceRemap, TArray<int32>* OutOptVertexRemap)
 	{
+		QUICK_SCOPE_CYCLE_COUNTER(STATQ_CleanTrimesh);
 		TArray<FVector3f> LocalSourceVerts = InOutVertices;
 		TArray<int32> LocalSourceIndices = InOutIndices;
 		
@@ -56,15 +57,18 @@ namespace Chaos
 		};
 
 		int32 NumBadTris = 0;
-		for(int32 SrcTriIndex = 0; SrcTriIndex < NumSourceTriangles; ++SrcTriIndex)
 		{
-			const int32 A = InOutIndices[SrcTriIndex * 3];
-			const int32 B = InOutIndices[SrcTriIndex * 3 + 1];
-			const int32 C = InOutIndices[SrcTriIndex * 3 + 2];
-
-			if(!ValidateTrianglesPre(A, B, C))
+			QUICK_SCOPE_CYCLE_COUNTER(STATQ_Validate);
+			for(int32 SrcTriIndex = 0; SrcTriIndex < NumSourceTriangles; ++SrcTriIndex)
 			{
-				++NumBadTris;
+				const int32 A = InOutIndices[SrcTriIndex * 3];
+				const int32 B = InOutIndices[SrcTriIndex * 3 + 1];
+				const int32 C = InOutIndices[SrcTriIndex * 3 + 2];
+
+				if(!ValidateTrianglesPre(A, B, C))
+				{
+					++NumBadTris;
+				}
 			}
 		}
 		UE_CLOG(NumBadTris > 0, LogChaos, Warning, TEXT("Input trimesh contains %d bad triangles."), NumBadTris);
@@ -81,31 +85,34 @@ namespace Chaos
 		TAABBTree<int32, TAABBTreeLeafArray<int32>> Accel(WrapperVerts);
 		TSet<int32> NonUnique;
 
-		for(int32 SourceVertIndex = 0; SourceVertIndex < NumSourceVerts; ++SourceVertIndex)
 		{
-			if (NonUnique.Contains(SourceVertIndex))
+			QUICK_SCOPE_CYCLE_COUNTER(STATQ_FindDuplicates);
+			for(int32 SourceVertIndex = 0; SourceVertIndex < NumSourceVerts; ++SourceVertIndex)
 			{
-				continue;
-			}
+				if(NonUnique.Contains(SourceVertIndex))
+				{
+					continue;
+				}
 
-			const FVector3f& SourceVert = LocalSourceVerts[SourceVertIndex];
+				const FVector3f& SourceVert = LocalSourceVerts[SourceVertIndex];
 
-			TArray<int32> Duplicates = Accel.FindAllIntersections(FAABB3(SourceVert - WeldThresholdSq, SourceVert + WeldThresholdSq));
-			ensure(Duplicates.Num() > 0);	//Should always find at least original vert
+				TArray<int32> Duplicates = Accel.FindAllIntersections(FAABB3(SourceVert - WeldThresholdSq, SourceVert + WeldThresholdSq));
+				ensure(Duplicates.Num() > 0);	//Should always find at least original vert
 
-			//first index is always considered unique
-			LocalUniqueVerts.Add(SourceVert);
-			LocalUniqueToSourceIndices.Add(SourceVertIndex);
-			LocalVertexRemap[SourceVertIndex] = LocalUniqueVerts.Num() - 1;
+				//first index is always considered unique
+				LocalUniqueVerts.Add(SourceVert);
+				LocalUniqueToSourceIndices.Add(SourceVertIndex);
+				LocalVertexRemap[SourceVertIndex] = LocalUniqueVerts.Num() - 1;
 
 
-			for (int32 Idx : Duplicates)
-			{
-				if (Idx != SourceVertIndex)
+				for(int32 Idx : Duplicates)
 				{
-					ensure(Idx > SourceVertIndex);	//shouldn't be here if a smaller idx already found these duplicates
-					NonUnique.Add(Idx);
-					LocalVertexRemap[Idx] = LocalVertexRemap[SourceVertIndex];
+					if(Idx != SourceVertIndex)
+					{
+						ensure(Idx > SourceVertIndex);	//shouldn't be here if a smaller idx already found these duplicates
+						NonUnique.Add(Idx);
+						LocalVertexRemap[Idx] = LocalVertexRemap[SourceVertIndex];
+					}
 				}
 			}
 		}
@@ -130,40 +137,43 @@ namespace Chaos
 		TArray<int32> LocalTriangleRemap;
 		LocalTriangleRemap.Reserve(NumSourceTriangles);
 
-		for(int32 OriginalTriIndex = 0; OriginalTriIndex < NumSourceTriangles; ++OriginalTriIndex)
 		{
-			const int32 OrigAIndex = LocalSourceIndices[OriginalTriIndex * 3];
-			const int32 OrigBIndex = LocalSourceIndices[OriginalTriIndex * 3 + 1];
-			const int32 OrigCIndex = LocalSourceIndices[OriginalTriIndex * 3 + 2];
-
-			const int32 RemappedAIndex = LocalVertexRemap[OrigAIndex];
-			const int32 RemappedBIndex = LocalVertexRemap[OrigBIndex];
-			const int32 RemappedCIndex = LocalVertexRemap[OrigCIndex];
-
-			const FVector3f& A = LocalUniqueVerts[RemappedAIndex];
-			const FVector3f& B = LocalUniqueVerts[RemappedBIndex];
-			const FVector3f& C = LocalUniqueVerts[RemappedCIndex];
-
-			// Only consider triangles that are actually valid for collision
-			// #BG Consider being able to fix small triangles by collapsing them if we hit this a lot
-			const bool bValidIndices = ValidateTriangleIndices(RemappedAIndex, RemappedBIndex, RemappedCIndex);
-			const bool bValidArea = ValidateTriangleArea(A, B, C);
-			if(bValidIndices && bValidArea)
-			{
-				LocalUniqueIndices.Add(RemappedAIndex);
-				LocalUniqueIndices.Add(RemappedBIndex);
-				LocalUniqueIndices.Add(RemappedCIndex);
-				LocalTriangleRemap.Add(OriginalTriIndex);
-			}
-			else
+			QUICK_SCOPE_CYCLE_COUNTER(STATQ_Remap);
+			for(int32 OriginalTriIndex = 0; OriginalTriIndex < NumSourceTriangles; ++OriginalTriIndex)
 			{
-				if(!bValidIndices)
+				const int32 OrigAIndex = LocalSourceIndices[OriginalTriIndex * 3];
+				const int32 OrigBIndex = LocalSourceIndices[OriginalTriIndex * 3 + 1];
+				const int32 OrigCIndex = LocalSourceIndices[OriginalTriIndex * 3 + 2];
+
+				const int32 RemappedAIndex = LocalVertexRemap[OrigAIndex];
+				const int32 RemappedBIndex = LocalVertexRemap[OrigBIndex];
+				const int32 RemappedCIndex = LocalVertexRemap[OrigCIndex];
+
+				const FVector3f& A = LocalUniqueVerts[RemappedAIndex];
+				const FVector3f& B = LocalUniqueVerts[RemappedBIndex];
+				const FVector3f& C = LocalUniqueVerts[RemappedCIndex];
+
+				// Only consider triangles that are actually valid for collision
+				// #BG Consider being able to fix small triangles by collapsing them if we hit this a lot
+				const bool bValidIndices = ValidateTriangleIndices(RemappedAIndex, RemappedBIndex, RemappedCIndex);
+				const bool bValidArea = ValidateTriangleArea(A, B, C);
+				if(bValidIndices && bValidArea)
 				{
-					++NumDiscardedTriangles_Welded;
+					LocalUniqueIndices.Add(RemappedAIndex);
+					LocalUniqueIndices.Add(RemappedBIndex);
+					LocalUniqueIndices.Add(RemappedCIndex);
+					LocalTriangleRemap.Add(OriginalTriIndex);
 				}
-				else if(!bValidArea)
+				else
 				{
-					++NumDiscardedTriangles_Area;
+					if(!bValidIndices)
+					{
+						++NumDiscardedTriangles_Welded;
+					}
+					else if(!bValidArea)
+					{
+						++NumDiscardedTriangles_Area;
+					}
 				}
 			}
 		}
diff --git a/Engine/Source/Runtime/Engine/Private/PhysicsEngine/Experimental/ChaosInterfaceUtils.cpp b/Engine/Source/Runtime/Engine/Private/PhysicsEngine/Experimental/ChaosInterfaceUtils.cpp
--- a/Engine/Source/Runtime/Engine/Private/PhysicsEngine/Experimental/ChaosInterfaceUtils.cpp
+++ b/Engine/Source/Runtime/Engine/Private/PhysicsEngine/Experimental/ChaosInterfaceUtils.cpp
@@ -25,10 +25,6 @@
 #include "PhysicsEngine/SphereElem.h"
 #include "PhysicsProxy/SingleParticlePhysicsProxy.h"
 
-#if PHYSICS_INTERFACE_PHYSX
-#include "PhysXIncludes.h"
-#endif
-
 #define FORCE_ANALYTICS 0
 #define CREATE_STRAIGHT_CAPSULE_GEOMETRY_FOR_TAPERED_CAPSULES
 
@@ -47,103 +43,6 @@ namespace ChaosInterface
 		check(false);
 	}
 
-#if PHYSICS_INTERFACE_PHYSX
-
-	template<>
-	TArray<Chaos::TVec3<int32>> GetMeshElements(const physx::PxConvexMesh* PhysXMesh)
-	{
-		TArray<Chaos::TVec3<int32>> CollisionMeshElements;
-#if !WITH_CHAOS_NEEDS_TO_BE_FIXED
-		int32 offset = 0;
-		int32 NbPolygons = static_cast<int32>(PhysXMesh->getNbPolygons());
-		for (int32 i = 0; i < NbPolygons; i++)
-		{
-			physx::PxHullPolygon Poly;
-			bool status = PhysXMesh->getPolygonData(i, Poly);
-			const auto Indices = PhysXMesh->getIndexBuffer() + Poly.mIndexBase;
-
-			for (int32 j = 2; j < static_cast<int32>(Poly.mNbVerts); j++)
-			{
-				CollisionMeshElements.Add(Chaos::TVec3<int32>(Indices[offset], Indices[offset + j], Indices[offset + j - 1]));
-			}
-		}
-#endif
-		return CollisionMeshElements;
-	}
-
-	template<>
-	TArray<Chaos::TVec3<int32>> GetMeshElements(const physx::PxTriangleMesh* PhysXMesh)
-	{
-		TArray<Chaos::TVec3<int32>> CollisionMeshElements;
-		const auto MeshFlags = PhysXMesh->getTriangleMeshFlags();
-		for (int32 j = 0; j < static_cast<int32>(PhysXMesh->getNbTriangles()); ++j)
-		{
-			if (MeshFlags | physx::PxTriangleMeshFlag::e16_BIT_INDICES)
-			{
-				const physx::PxU16* Indices = reinterpret_cast<const physx::PxU16*>(PhysXMesh->getTriangles());
-				CollisionMeshElements.Add(Chaos::TVec3<int32>(Indices[3 * j], Indices[3 * j + 1], Indices[3 * j + 2]));
-			}
-			else
-			{
-				const physx::PxU32* Indices = reinterpret_cast<const physx::PxU32*>(PhysXMesh->getTriangles());
-				CollisionMeshElements.Add(Chaos::TVec3<int32>(Indices[3 * j], Indices[3 * j + 1], Indices[3 * j + 2]));
-			}
-		}
-		return CollisionMeshElements;
-	}
-
-	template<class PHYSX_MESH>
-	TUniquePtr<Chaos::FImplicitObject> ConvertPhysXMeshToLevelset(const PHYSX_MESH* PhysXMesh, const FVector& Scale)
-	{
-#if WITH_CHAOS && !WITH_CHAOS_NEEDS_TO_BE_FIXED
-		TArray<Chaos::TVec3<int32>> CollisionMeshElements = GetMeshElements(PhysXMesh);
-		Chaos::FParticles CollisionMeshParticles;
-		CollisionMeshParticles.AddParticles(PhysXMesh->getNbVertices());
-		for (uint32 j = 0; j < CollisionMeshParticles.Size(); ++j)
-		{
-			const auto& Vertex = PhysXMesh->getVertices()[j];
-			CollisionMeshParticles.X(j) = Scale * Chaos::FVec3(Vertex.x, Vertex.y, Vertex.z);
-		}
-		Chaos::FAABB3 BoundingBox(CollisionMeshParticles.X(0), CollisionMeshParticles.X(0));
-		for (uint32 j = 1; j < CollisionMeshParticles.Size(); ++j)
-		{
-			BoundingBox.GrowToInclude(CollisionMeshParticles.X(j));
-		}
-#if FORCE_ANALYTICS
-		return TUniquePtr<Chaos::FImplicitObject>(new Chaos::TBox<FReal, 3>(BoundingBox));
-#else
-		int32 MaxAxisSize = 10;
-		int32 MaxAxis;
-		const auto Extents = BoundingBox.Extents();
-		if (Extents[0] > Extents[1] && Extents[0] > Extents[2])
-		{
-			MaxAxis = 0;
-		}
-		else if (Extents[1] > Extents[2])
-		{
-			MaxAxis = 1;
-		}
-		else
-		{
-			MaxAxis = 2;
-		}
-		Chaos::TVec3<int32> Counts(MaxAxisSize * Extents[0] / Extents[MaxAxis], MaxAxisSize * Extents[1] / Extents[MaxAxis], MaxAxisSize * Extents[2] / Extents[MaxAxis]);
-		Counts[0] = Counts[0] < 1 ? 1 : Counts[0];
-		Counts[1] = Counts[1] < 1 ? 1 : Counts[1];
-		Counts[2] = Counts[2] < 1 ? 1 : Counts[2];
-		Chaos::TUniformGrid<float, 3> Grid(BoundingBox.Min(), BoundingBox.Max(), Counts, 1);
-		Chaos::FTriangleMesh CollisionMesh(MoveTemp(CollisionMeshElements));
-		return TUniquePtr<Chaos::FImplicitObject>(new Chaos::FLevelSet(Grid, CollisionMeshParticles, CollisionMesh));
-#endif
-
-#else
-		return TUniquePtr<Chaos::FImplicitObject>();
-#endif // !WITH_CHAOS_NEEDS_TO_BE_FIXED
-
-	}
-
-#endif
-
 	Chaos::EChaosCollisionTraceFlag ConvertCollisionTraceFlag(ECollisionTraceFlag Flag)
 	{
 		if (Flag == ECollisionTraceFlag::CTF_UseDefault)
@@ -186,17 +85,12 @@ namespace ChaosInterface
 			CollisionMarginMax = Chaos_Collision_MarginMax;
 		}
 
-#if WITH_CHAOS
 		// Complex as simple should not create simple geometry, unless there is no complex geometry.  Otherwise both get queried against.
 		bool bMakeSimpleGeometry = (CollisionTraceType != CTF_UseComplexAsSimple) || (InParams.ChaosTriMeshes.Num() == 0);
 
 		// The reverse is true for Simple as Complex.
 		const int32 SimpleShapeCount = InParams.Geometry->SphereElems.Num() + InParams.Geometry->BoxElems.Num() + InParams.Geometry->ConvexElems.Num() + InParams.Geometry->SphylElems.Num();
 		bool bMakeComplexGeometry = (CollisionTraceType != CTF_UseSimpleAsComplex) || (SimpleShapeCount == 0);
-#else
-		bool bMakeSimpleGeometry = true;
-		bool bMakeComplexGeometry = true;
-#endif
 
 		ensure(bMakeComplexGeometry || bMakeSimpleGeometry);
 
@@ -346,7 +240,7 @@ namespace ChaosInterface
 				}
 			}
 #endif
-#if WITH_CHAOS && !PHYSICS_INTERFACE_PHYSX
+
 			for (uint32 i = 0; i < static_cast<uint32>(InParams.Geometry->ConvexElems.Num()); ++i)
 			{
 				const FKConvexElem& CollisionBody = InParams.Geometry->ConvexElems[i];
@@ -409,21 +303,9 @@ namespace ChaosInterface
 				Shapes.Emplace(MoveTemp(NewShape));
 				Geoms.Add(MoveTemp(Implicit));
 			}
-#endif
-		}
-#if WITH_PHYSX && PHYSICS_INTERFACE_PHYSX
-		for (const auto& PhysXMesh : InParams.TriMeshes)
-		{
-			auto Implicit = ConvertPhysXMeshToLevelset(PhysXMesh, Scale);
-			auto NewShape = NewShapeHelper(MakeSerializable(Implicit), Shapes.Num(), nullptr, ECollisionEnabled::QueryAndPhysics, true);
-			Shapes.Emplace(MoveTemp(NewShape));
-			Geoms.Add(MoveTemp(Implicit));
-
 		}
-#endif
 	}
 
-#if WITH_CHAOS
 	void CalculateMassPropertiesFromShapeCollection(Chaos::FMassProperties& OutProperties, const TArray<FPhysicsShapeHandle>& InShapes, float InDensityKGPerCM)
 	{
 		Chaos::CalculateMassPropertiesFromShapeCollection(
@@ -446,6 +328,6 @@ namespace ChaosInterface
 		);
 	}
 
-#endif // WITH_CHAOS
+
 
 }
\ No newline at end of file
diff --git a/Engine/Source/Runtime/Engine/Private/PhysicsEngine/Experimental/PhysInterface_Chaos.cpp b/Engine/Source/Runtime/Engine/Private/PhysicsEngine/Experimental/PhysInterface_Chaos.cpp
--- a/Engine/Source/Runtime/Engine/Private/PhysicsEngine/Experimental/PhysInterface_Chaos.cpp
+++ b/Engine/Source/Runtime/Engine/Private/PhysicsEngine/Experimental/PhysInterface_Chaos.cpp
@@ -1,7 +1,5 @@
 // Copyright Epic Games, Inc. All Rights Reserved.
 
-#if WITH_CHAOS
-
 #include "Physics/Experimental/PhysInterface_Chaos.h"
 #include "Physics/Experimental/PhysScene_Chaos.h"
 #include "Physics/Experimental/ChaosInterfaceUtils.h"
@@ -46,14 +44,6 @@
 #include "PhysicalMaterials/PhysicalMaterial.h"
 #include "PhysicsProxy/SingleParticlePhysicsProxy.h"
 
-#if PHYSICS_INTERFACE_PHYSX
-#include "geometry/PxConvexMesh.h"
-#include "geometry/PxTriangleMesh.h"
-#include "foundation/PxVec3.h"
-#include "extensions/PxMassProperties.h"
-#include "Containers/ArrayView.h"
-#endif
-
 DEFINE_STAT(STAT_TotalPhysicsTime);
 DEFINE_STAT(STAT_NumCloths);
 DEFINE_STAT(STAT_NumClothVerts);
@@ -848,7 +838,6 @@ void FPhysInterface_Chaos::AddGeometry(FPhysicsActorHandle& InActor, const FGeom
 	Chaos::FShapesArray Shapes;
 	ChaosInterface::CreateGeometry(InParams, Geoms, Shapes);
 
-#if WITH_CHAOS
 	if (InActor && Geoms.Num())
 	{
 		for (TUniquePtr<Chaos::FPerShapeData>& Shape : Shapes)
@@ -889,7 +878,6 @@ void FPhysInterface_Chaos::AddGeometry(FPhysicsActorHandle& InActor, const FGeom
 			InActor->GetGameThreadAPI().SetShapesArray(MoveTemp(Shapes));
 		}
 	}
-#endif
 }
 
 void FPhysInterface_Chaos::SetMaterials(const FPhysicsShapeHandle& InShape, const TArrayView<UPhysicalMaterial*> InMaterials)
@@ -1349,5 +1337,3 @@ void FPhysInterface_Chaos::CalculateMassPropertiesFromShapeCollection(Chaos::FMa
 	ChaosInterface::CalculateMassPropertiesFromShapeCollection(OutProperties,InShapes,InDensityKGPerCM);
 }
 
-#endif
-
diff --git a/Engine/Source/Runtime/Engine/Private/PhysicsEngine/Experimental/PhysScene_Chaos.cpp b/Engine/Source/Runtime/Engine/Private/PhysicsEngine/Experimental/PhysScene_Chaos.cpp
--- a/Engine/Source/Runtime/Engine/Private/PhysicsEngine/Experimental/PhysScene_Chaos.cpp
+++ b/Engine/Source/Runtime/Engine/Private/PhysicsEngine/Experimental/PhysScene_Chaos.cpp
@@ -1,9 +1,7 @@
 // Copyright Epic Games, Inc. All Rights Reserved.
 
-#if WITH_CHAOS
 #include "Physics/Experimental/PhysScene_Chaos.h"
 
-
 #include "CoreMinimal.h"
 #include "GameDelegates.h"
 
@@ -405,9 +403,7 @@
 	)
 	, PhysicsReplication(nullptr)
 	, SolverActor(InSolverActor)
-#if WITH_CHAOS
 	, LastEventDispatchTime(Chaos::FReal(-1))
-#endif
 #if WITH_EDITOR
 	, SingleStepCounter(0)
 #endif
@@ -415,7 +411,6 @@
 	, bIsWorldPaused(false)
 #endif
 {
-#if WITH_CHAOS
 	LLM_SCOPE(ELLMTag::ChaosScene);
 
 	PhysicsProxyToComponentMap.Reset();
@@ -441,13 +436,10 @@
 	PhysicsReplication = PhysicsReplicationFactory.IsValid() ? PhysicsReplicationFactory->Create(this) : new FPhysicsReplication(this);
 
 	FPhysicsDelegates::OnPhysSceneInit.Broadcast(this);
-#endif
 }
 
 FPhysScene_Chaos::~FPhysScene_Chaos()
 {
-#if WITH_CHAOS
-
 	if (AsyncPhysicsTickCallback)
 	{
 		SceneSolver->UnregisterAndFreeSimCallbackObject_External(AsyncPhysicsTickCallback);
@@ -466,7 +458,6 @@
 	{
 		delete PhysicsReplication;
 	}
-#endif
 
 #if CHAOS_WITH_PAUSABLE_SOLVER
 	if (SyncCaller)
@@ -476,7 +467,7 @@
 #endif
 }
 
-#if WITH_EDITOR && WITH_CHAOS
+#if WITH_EDITOR
 bool FPhysScene_Chaos::IsOwningWorldEditor() const
 {
 	const UWorld* WorldPtr = GetOwningWorld();
@@ -629,7 +620,6 @@
 
 void FPhysScene_Chaos::RemoveObject(FGeometryCollectionPhysicsProxy* InObject)
 {
-#if WITH_CHAOS
 	Chaos::FPhysicsSolver* Solver = InObject->GetSolver<Chaos::FPhysicsSolver>();
 
 	for (TUniquePtr<Chaos::TGeometryParticle<Chaos::FReal, 3>>& GTParticleUnique : InObject->GetExternalParticles())
@@ -648,7 +638,6 @@
 	}
 
 	RemoveFromComponentMaps(InObject);
-#endif
 }
 
 FPhysicsReplication* FPhysScene_Chaos::GetPhysicsReplication()
@@ -711,7 +700,6 @@
 
 void FPhysScene_Chaos::HandleCollisionEvents(const Chaos::FCollisionEventData& Event)
 {
-#if WITH_CHAOS
 	ContactPairToPendingNotifyMap.Reset();
 
 	TMap<IPhysicsProxyBase*, TArray<int32>> const& PhysicsProxyToCollisionIndicesMap = Event.PhysicsProxyToCollisionIndices.PhysicsProxyToIndicesMap;
@@ -796,7 +784,6 @@
 
 	// Tell the world and actors about the collisions
 	DispatchPendingCollisionNotifies();
-#endif
 }
 
 void FPhysScene_Chaos::DispatchPendingCollisionNotifies()
@@ -906,8 +893,6 @@
 	PhysicsProxyToComponentMap.Remove(InObject);
 }
 
-#if WITH_CHAOS
-
 void FPhysScene_Chaos::OnWorldBeginPlay()
 {
 	Chaos::FPhysicsSolver* Solver = GetSolver();
@@ -1890,5 +1875,3 @@
 
 TSharedPtr<IPhysicsReplicationFactory> FPhysScene_Chaos::PhysicsReplicationFactory;
 
-#endif // WITH_CHAOS
-#endif
