Physics interface cleanup.

diff --git a/Engine/Source/Runtime/Engine/Private/PhysicsEngine/Experimental/PhysScene_ImmediatePhysX.cpp b/Engine/Source/Runtime/Engine/Private/PhysicsEngine/Experimental/PhysScene_ImmediatePhysX.cpp
--- a/Engine/Source/Runtime/Engine/Private/PhysicsEngine/Experimental/PhysScene_ImmediatePhysX.cpp
+++ b/Engine/Source/Runtime/Engine/Private/PhysicsEngine/Experimental/PhysScene_ImmediatePhysX.cpp
@@ -1,905 +1 @@
 // Copyright Epic Games, Inc. All Rights Reserved.
-
-#if WITH_IMMEDIATE_PHYSX
-
-#include "Physics/PhysScene_ImmediatePhysX.h"
-#include "Misc/CommandLine.h"
-#include "Stats/Stats.h"
-#include "UObject/UObjectGlobals.h"
-#include "UObject/UObjectHash.h"
-#include "UObject/UObjectIterator.h"
-#include "HAL/IConsoleManager.h"
-#include "Async/TaskGraphInterfaces.h"
-#include "EngineDefines.h"
-#include "Engine/EngineTypes.h"
-#include "PhysxUserData.h"
-#include "PhysicsEngine/BodyInstance.h"
-#include "Components/PrimitiveComponent.h"
-#include "Components/SkeletalMeshComponent.h"
-#include "PhysicsEngine/RigidBodyIndexPair.h"
-#include "PhysicsPublic.h"
-#include "CustomPhysXPayload.h"
-#include "HAL/LowLevelMemTracker.h"
-
-#if WITH_PHYSX
-	#include "PhysXPublic.h"
-	#include "PhysicsEngine/PhysXSupport.h"
-#endif
-
-#include "PhysicsEngine/PhysSubstepTasks.h"
-#include "PhysicsEngine/PhysicsCollisionHandler.h"
-#include "Physics/PhysicsInterfaceUtils.h"
-#include "Components/LineBatchComponent.h"
-#include "PhysicsEngine/PhysicsSettings.h"
-#include "PhysicsEngine/BodySetup.h"
-#include "PhysicsEngine/ConstraintInstance.h"
-#include "PhysicsReplication.h"
-#include "ProfilingDebugging/CsvProfiler.h"
-#include "Physics/PhysicsInterfaceImmediatePhysX.h"
-
-/** Physics stats **/
-
-DEFINE_STAT(STAT_TotalPhysicsTime);
-DEFINE_STAT(STAT_NumCloths);
-DEFINE_STAT(STAT_NumClothVerts);
-
-CSV_DECLARE_CATEGORY_MODULE_EXTERN(CORE_API, Basic);
-
-DECLARE_CYCLE_STAT(TEXT("Start Physics Time (sync)"), STAT_PhysicsKickOffDynamicsTime, STATGROUP_Physics);
-DECLARE_CYCLE_STAT(TEXT("Fetch Results Time (sync)"), STAT_PhysicsFetchDynamicsTime, STATGROUP_Physics);
-
-DECLARE_CYCLE_STAT(TEXT("Start Physics Time (async)"), STAT_PhysicsKickOffDynamicsTime_Async, STATGROUP_Physics);
-DECLARE_CYCLE_STAT(TEXT("Fetch Results Time (async)"), STAT_PhysicsFetchDynamicsTime_Async, STATGROUP_Physics);
-
-DECLARE_CYCLE_STAT(TEXT("Update Kinematics On Deferred SkelMeshes"), STAT_UpdateKinematicsOnDeferredSkelMeshes, STATGROUP_Physics);
-
-DECLARE_CYCLE_STAT(TEXT("Phys Events Time"), STAT_PhysicsEventTime, STATGROUP_Physics);
-DECLARE_CYCLE_STAT(TEXT("SyncComponentsToBodies (sync)"), STAT_SyncComponentsToBodies, STATGROUP_Physics);
-DECLARE_CYCLE_STAT(TEXT("SyncComponentsToBodies (async)"), STAT_SyncComponentsToBodies_Async, STATGROUP_Physics);
-
-DECLARE_DWORD_COUNTER_STAT(TEXT("Broadphase Adds"), STAT_NumBroadphaseAdds, STATGROUP_Physics);
-DECLARE_DWORD_COUNTER_STAT(TEXT("Broadphase Removes"), STAT_NumBroadphaseRemoves, STATGROUP_Physics);
-DECLARE_DWORD_COUNTER_STAT(TEXT("Active Constraints"), STAT_NumActiveConstraints, STATGROUP_Physics);
-DECLARE_DWORD_COUNTER_STAT(TEXT("Active Simulated Bodies"), STAT_NumActiveSimulatedBodies, STATGROUP_Physics);
-DECLARE_DWORD_COUNTER_STAT(TEXT("Active Kinematic Bodies"), STAT_NumActiveKinematicBodies, STATGROUP_Physics);
-DECLARE_DWORD_COUNTER_STAT(TEXT("Mobile Bodies"), STAT_NumMobileBodies, STATGROUP_Physics);
-DECLARE_DWORD_COUNTER_STAT(TEXT("Static Bodies"), STAT_NumStaticBodies, STATGROUP_Physics);
-DECLARE_DWORD_COUNTER_STAT(TEXT("Shapes"), STAT_NumShapes, STATGROUP_Physics);
-
-DECLARE_DWORD_COUNTER_STAT(TEXT("(ASync) Broadphase Adds"), STAT_NumBroadphaseAddsAsync, STATGROUP_Physics);
-DECLARE_DWORD_COUNTER_STAT(TEXT("(ASync) Broadphase Removes"), STAT_NumBroadphaseRemovesAsync, STATGROUP_Physics);
-DECLARE_DWORD_COUNTER_STAT(TEXT("(ASync) Active Constraints"), STAT_NumActiveConstraintsAsync, STATGROUP_Physics);
-DECLARE_DWORD_COUNTER_STAT(TEXT("(ASync) Active Simulated Bodies"), STAT_NumActiveSimulatedBodiesAsync, STATGROUP_Physics);
-DECLARE_DWORD_COUNTER_STAT(TEXT("(ASync) Active Kinematic Bodies"), STAT_NumActiveKinematicBodiesAsync, STATGROUP_Physics);
-DECLARE_DWORD_COUNTER_STAT(TEXT("(ASync) Mobile Bodies"), STAT_NumMobileBodiesAsync, STATGROUP_Physics);
-DECLARE_DWORD_COUNTER_STAT(TEXT("(ASync) Static Bodies"), STAT_NumStaticBodiesAsync, STATGROUP_Physics);
-DECLARE_DWORD_COUNTER_STAT(TEXT("(ASync) Shapes"), STAT_NumShapesAsync, STATGROUP_Physics);
-
-EPhysicsSceneType FPhysScene_ImmediatePhysX::SceneType_AssumesLocked(const FBodyInstance* BodyInstance) const
-{
-	return PST_Sync;
-}
-
-/**
-* Return true if we should be running in single threaded mode, ala dedicated server
-**/
-
-/**
-* Return true if we should lag the async scene a frame
-**/
-FORCEINLINE static bool FrameLagAsync()
-{
-	if (IsRunningDedicatedServer())
-	{
-		return false;
-	}
-	return true;
-}
-
-#if WITH_PHYSX
-
-FAutoConsoleTaskPriority CPrio_FPhysXTask(
-	TEXT("TaskGraph.TaskPriorities.PhysXTask"),
-	TEXT("Task and thread priority for FPhysXTask."),
-	ENamedThreads::HighThreadPriority, // if we have high priority task threads, then use them...
-	ENamedThreads::NormalTaskPriority, // .. at normal task priority
-	ENamedThreads::HighTaskPriority // if we don't have hi pri threads, then use normal priority threads at high task priority instead
-	);
-
-
-int32 GPhysXOverrideMbpNumSubdivisions_Client = 0;
-int32 GPhysXOverrideMbpNumSubdivisions_Server = 0;
-int32 GPhysXForceMbp_Client = 0;
-int32 GPhysXForceMbp_Server = 0;
-int32 GPhysXForceNoKinematicStaticPairs = 0;
-int32 GPhysXForceNoKinematicKinematicPairs = 0;
-
-TAutoConsoleVariable<int32> CVarOverrideMbpNumSubdivisionsClient(TEXT("p.OverrideMbpNumSubdivisionsClient"), GPhysXOverrideMbpNumSubdivisions_Client, TEXT("Override for number of subdivisions to perform when building MBP regions on a client, note regions are only generated when a scene is created - this will not update the scene if it's already running (0 = No override, 1>16 - Override number)"), ECVF_Default);
-TAutoConsoleVariable<int32> CVarOverrideMbpNumSubdivisionsServer(TEXT("p.OverrideMbpNumSubdivisionsServer"), GPhysXOverrideMbpNumSubdivisions_Server, TEXT("Override for number of subdivisions to perform when building MBP regions on a server, note regions are only generated when a scene is created - this will not update the scene if it's already running (0 = No override, 1>16 - Override number)"), ECVF_Default);
-TAutoConsoleVariable<int32> CVarForceMbpClient(TEXT("p.ForceMbpClient"), GPhysXForceMbp_Client, TEXT("Forces all created scenes to use MBP on client builds"), ECVF_Default);
-TAutoConsoleVariable<int32> CVarForceMbpServer(TEXT("p.ForceMbpServer"), GPhysXForceMbp_Server, TEXT("Forces all created scenes to use MBP on server builds"), ECVF_Default);
-TAutoConsoleVariable<int32> CVarForceNoKSPairs(TEXT("p.ForceNoKSPairs"), GPhysXForceNoKinematicStaticPairs, TEXT("Disables kinematic-static pairs. This makes converting from static to dynamic a little slower - but provides better broadphase performance because we early reject those pairs."), ECVF_Default);
-TAutoConsoleVariable<int32> CVarForceNoKKPairs(TEXT("p.ForceNoKKPairs"), GPhysXForceNoKinematicKinematicPairs, TEXT("Disables kinematic-kinematic pairs. This is required when using APEX destruction to correctly generate chunk pairs - when not using destruction this speeds up the broadphase by early rejecting KK pairs."), ECVF_Default);
-
-DECLARE_STATS_GROUP(TEXT("PhysXTasks"), STATGROUP_PhysXTasks, STATCAT_Advanced);
-
-struct FPhysXRingBuffer
-{
-	static const int32 Size = 16;
-
-	PxBaseTask* Buffer[Size];
-	int32 Start;
-	int32 End;
-	int32 Num;
-};
-
-int32 GBatchPhysXTasksSize = 3;	//NOTE: FPhysXRingBuffer::Size should be twice as big as this value
-TAutoConsoleVariable<int32> CVarBatchPhysXTasksSize(TEXT("p.BatchPhysXTasksSize"), GBatchPhysXTasksSize, TEXT("Number of tasks to batch together (max 8). 1 will go as wide as possible, but more overhead on small tasks"), ECVF_Default);
-
-struct FBatchPhysXTasks
-{
-	static void SetPhysXTasksSinkFunc()
-	{
-		GBatchPhysXTasksSize = FMath::Max(1, FMath::Min(FPhysXRingBuffer::Size / 2, CVarBatchPhysXTasksSize.GetValueOnGameThread()));
-	}
-};
-
-struct FPhysTaskScopedNamedEvent
-{
-	FPhysTaskScopedNamedEvent() = delete;
-	FPhysTaskScopedNamedEvent(const FPhysTaskScopedNamedEvent& InOther) = delete;
-	FPhysTaskScopedNamedEvent& operator=(const FPhysTaskScopedNamedEvent& InOther) = delete;
-
-	FPhysTaskScopedNamedEvent(PxBaseTask* InTask)
-	{
-#if ENABLE_STATNAMEDEVENTS
-		check(InTask);
-		const char* TaskName = InTask->getName();
-		
-		bEmittedEvent = GCycleStatsShouldEmitNamedEvents != 0;
-
-		if(bEmittedEvent)
-		{
-			FPlatformMisc::BeginNamedEvent(FColor::Green, TaskName);
-		}
-#endif
-	}
-
-	~FPhysTaskScopedNamedEvent()
-	{
-#if ENABLE_STATNAMEDEVENTS
-		if(bEmittedEvent)
-		{
-			FPlatformMisc::EndNamedEvent();
-		}
-#endif
-	}
-
-private:
-
-	bool bEmittedEvent;
-};
-
-static FAutoConsoleVariableSink CVarBatchPhysXTasks(FConsoleCommandDelegate::CreateStatic(&FBatchPhysXTasks::SetPhysXTasksSinkFunc));
-
-namespace DynamicStatsHelper
-{
-	struct FStatLookup
-	{
-		const char* StatName;
-		TStatId Stat;
-	} Stats[100];
-
-	int NumStats = 0;
-	FCriticalSection CS;
-
-	TStatId FindOrCreateStatId(const char* StatName)
-	{
-#if STATS
-		for (int StatIdx = 0; StatIdx < NumStats; ++StatIdx)
-		{
-			FStatLookup& Lookup = Stats[StatIdx];
-			if (Lookup.StatName == StatName)
-			{
-				return Lookup.Stat;
-			}
-		}
-
-		if (ensureMsgf(NumStats < sizeof(Stats) / sizeof(Stats[0]), TEXT("Too many different physx task stats. This will make the stat search slow")))
-		{
-			FScopeLock ScopeLock(&CS);
-
-			//Do the search again in case another thread added
-			for (int StatIdx = 0; StatIdx < NumStats; ++StatIdx)
-			{
-				FStatLookup& Lookup = Stats[StatIdx];
-				if (Lookup.StatName == StatName)
-				{
-					return Lookup.Stat;
-				}
-			}
-
-			FStatLookup& NewStat = Stats[NumStats];
-			NewStat.StatName = StatName;
-			NewStat.Stat = FDynamicStats::CreateStatId<FStatGroup_STATGROUP_PhysXTasks>(FName(StatName));
-			FPlatformMisc::MemoryBarrier();
-			++NumStats;	//make sure to do this at the end in case another thread is currently iterating
-			return NewStat.Stat;
-		}
-#endif // STATS
-		return TStatId();
-	}
-}
-
-struct FPhysXCPUDispatcher;
-
-class FPhysXTask
-{
-public:
-	FPhysXTask(PxBaseTask& InTask, FPhysXCPUDispatcher& InDispatcher);
-	FPhysXTask(FPhysXRingBuffer& RingBuffer, FPhysXCPUDispatcher& InDispatcher);
-	~FPhysXTask();
-
-	static FORCEINLINE TStatId GetStatId()
-	{
-		RETURN_QUICK_DECLARE_CYCLE_STAT(FPhysXTask, STATGROUP_Physics);
-	}
-	static FORCEINLINE ENamedThreads::Type GetDesiredThread()
-	{
-		return CPrio_FPhysXTask.Get();
-	}
-	static FORCEINLINE ESubsequentsMode::Type GetSubsequentsMode()
-	{
-		return ESubsequentsMode::TrackSubsequents;
-	}
-
-	void DoTask(ENamedThreads::Type CurrentThread, const FGraphEventRef& MyCompletionGraphEvent);
-
-	FPhysXRingBuffer RingBuffer;
-	FPhysXCPUDispatcher& Dispatcher;
-};
-
-/** Used to dispatch physx tasks to task graph */
-struct FPhysXCPUDispatcher : public PxCpuDispatcher
-{
-	FPhysXCPUDispatcher()
-	{
-		check(IsInGameThread());
-		TLSKey = FPlatformTLS::AllocTlsSlot();
-	}
-
-	~FPhysXCPUDispatcher()
-	{
-		check(IsInGameThread());
-		FPlatformTLS::FreeTlsSlot(TLSKey);
-	}
-
-	virtual void submitTask(PxBaseTask& Task) override
-	{
-		LLM_SCOPE(ELLMTag::PhysX);
-
-		if (IsInGameThread())
-		{
-			//Game thread enqueues on task graph
-			TGraphTask<FPhysXTask>::CreateTask(NULL).ConstructAndDispatchWhenReady(Task, *this);
-		}
-		else
-		{
-			//See if we can use local queue
-			FPhysXRingBuffer& RingBuffer = *(FPhysXRingBuffer*)FPlatformTLS::GetTlsValue(TLSKey);
-			RingBuffer.Buffer[RingBuffer.End] = &Task;
-			RingBuffer.End = (RingBuffer.End + 1) % FPhysXRingBuffer::Size;
-			RingBuffer.Num++;
-
-			if (RingBuffer.Num >= GBatchPhysXTasksSize * 2)
-			{
-				TGraphTask<FPhysXTask>::CreateTask(NULL).ConstructAndDispatchWhenReady(RingBuffer, *this);
-			}
-		}
-	}
-
-	virtual PxU32 getWorkerCount() const override
-	{
-		return FTaskGraphInterface::Get().GetNumWorkerThreads();
-	}
-
-	uint32 TLSKey;
-};
-
-FPhysXTask::FPhysXTask(PxBaseTask& Task, FPhysXCPUDispatcher& InDispatcher)
-	: Dispatcher(InDispatcher)
-{
-	RingBuffer.Buffer[0] = &Task;
-	RingBuffer.Start = 0;
-	RingBuffer.End = 1;
-	RingBuffer.Num = 1;
-}
-
-FPhysXTask::FPhysXTask(FPhysXRingBuffer& InRingBuffer, FPhysXCPUDispatcher& InDispatcher)
-	: Dispatcher(InDispatcher)
-{
-	int32 NumToSteal = InRingBuffer.Num / 2;
-	ensureMsgf(NumToSteal > 0, TEXT("Trying to steal 0 items"));
-
-	const int32 StartPos = (InRingBuffer.Start + NumToSteal);
-	for (int32 Count = 0; Count < NumToSteal; ++Count)
-	{
-		RingBuffer.Buffer[Count] = InRingBuffer.Buffer[(StartPos + Count) % FPhysXRingBuffer::Size];
-	}
-
-	RingBuffer.Start = 0;
-	RingBuffer.End = NumToSteal;
-	RingBuffer.Num = NumToSteal;
-
-
-	InRingBuffer.Num -= NumToSteal;
-	InRingBuffer.End = (StartPos) % FPhysXRingBuffer::Size;
-}
-
-FPhysXTask::~FPhysXTask()
-{
-	FPlatformTLS::SetTlsValue(Dispatcher.TLSKey, nullptr);
-}
-
-void FPhysXTask::DoTask(ENamedThreads::Type CurrentThread, const FGraphEventRef& MyCompletionGraphEvent)
-{
-	FPlatformTLS::SetTlsValue(Dispatcher.TLSKey, &RingBuffer);
-
-	while (RingBuffer.Num)
-	{
-		PxBaseTask* Task = RingBuffer.Buffer[RingBuffer.Start];
-
-#if ENABLE_STATNAMEDEVENTS || STATS
-		FPhysTaskScopedNamedEvent TaskEvent(Task);
-#endif
-
-#if STATS
-		const char* TaskName = Task->getName();
-		FScopeCycleCounter CycleCounter(DynamicStatsHelper::FindOrCreateStatId(TaskName));
-#endif
-
-		Task->run();
-		Task->release();
-
-		RingBuffer.Start = (RingBuffer.Start + 1) % RingBuffer.Size;
-		--RingBuffer.Num;
-	}
-}
-
-
-DECLARE_CYCLE_STAT(TEXT("PhysX Single Thread Task"), STAT_PhysXSingleThread, STATGROUP_Physics);
-
-/** Used to dispatch physx tasks to the game thread */
-class FPhysXCPUDispatcherSingleThread : public PxCpuDispatcher
-{
-	TArray<PxBaseTask*> TaskStack;
-
-	virtual void submitTask(PxBaseTask& Task) override
-	{
-		SCOPE_CYCLE_COUNTER(STAT_PhysXSingleThread);
-
-		TaskStack.Push(&Task);
-		if (TaskStack.Num() > 1)
-		{
-			return;
-		}
-
-		{
-#if ENABLE_STATNAMEDEVENTS || STATS
-			FPhysTaskScopedNamedEvent TaskEvent(&Task);
-#endif
-
-#if STATS
-			const char* TaskName = Task.getName();
-			FScopeCycleCounter CycleCounter(DynamicStatsHelper::FindOrCreateStatId(TaskName));
-#endif
-
-			Task.run();
-			Task.release();
-		}
-
-		while (TaskStack.Num() > 1)
-		{
-			PxBaseTask& ChildTask = *TaskStack.Pop();
-			{
-
-#if ENABLE_STATNAMEDEVENTS || STATS
-				FPhysTaskScopedNamedEvent TaskEvent(&ChildTask);
-#endif
-
-#if STATS
-				const char* ChildTaskName = ChildTask.getName();
-				FScopeCycleCounter CycleCounter(DynamicStatsHelper::FindOrCreateStatId(ChildTaskName));
-#endif
-				ChildTask.run();
-				ChildTask.release();
-			}
-		}
-		verify(&Task == TaskStack.Pop() && !TaskStack.Num());
-	}
-
-	virtual PxU32 getWorkerCount() const override
-	{
-		return 1;
-	}
-};
-
-TSharedPtr<ISimEventCallbackFactory> FPhysScene_ImmediatePhysX::SimEventCallbackFactory;
-TSharedPtr<IContactModifyCallbackFactory> FPhysScene_ImmediatePhysX::ContactModifyCallbackFactory;
-
-#endif // WITH_PHYSX
-
-static void StaticSetPhysXTreeRebuildRate(const TArray<FString>& Args, UWorld* World)
-{
-	if (Args.Num() > 0)
-	{
-		const int32 NewRate = FCString::Atoi(*Args[0]);
-		if(World && World->GetPhysicsScene())
-		{
-			World->GetPhysicsScene()->SetPhysXTreeRebuildRate(NewRate);
-		}
-	}
-	else
-	{
-		UE_LOG(LogPhysics, Warning, TEXT("Usage: p.PhysXTreeRebuildRate <num_frames>"));
-	}
-}
-
-static FAutoConsoleCommandWithWorldAndArgs GSetPhysXTreeRebuildRate(TEXT("p.PhysXTreeRebuildRate"), TEXT("Utility function to change PhysXTreeRebuildRate, useful when profiling fetchResults vs scene queries."),
-	FConsoleCommandWithWorldAndArgsDelegate::CreateStatic(&StaticSetPhysXTreeRebuildRate)
-);
-
-
-/** Exposes creation of physics-engine scene outside Engine (for use with Physics Asset Editor for example). */
-FPhysScene_ImmediatePhysX::FPhysScene_ImmediatePhysX()
-{
-}
-
-void FPhysScene_ImmediatePhysX::SwapActorData(uint32 Actor1DataIdx, uint32 Actor2DataIdx)
-{
-#if WITH_PHYSX
-	check(Actors.IsValidIndex(Actor1DataIdx));
-	check(Actors.IsValidIndex(Actor2DataIdx));
-
-	Swap(Actors[Actor1DataIdx], Actors[Actor2DataIdx]);
-	Swap(RigidBodiesData[Actor1DataIdx], RigidBodiesData[Actor2DataIdx]);
-	Swap(SolverBodiesData[Actor1DataIdx], SolverBodiesData[Actor2DataIdx]);
-	Swap(PendingAcceleration[Actor1DataIdx], PendingAcceleration[Actor2DataIdx]);
-	Swap(KinematicTargets[Actor1DataIdx], KinematicTargets[Actor2DataIdx]);
-
-	//Update entity index on the handle
-	bDirtyJointData = true;	//reordering of bodies could lead to stale joint data
-#endif // WITH_PHYSX
-
-	bRecreateIterationCache = true;	//reordering of bodies so we need to change iteration order potentially
-}
-
-void FPhysScene_ImmediatePhysX::ResizeActorData(uint32 ActorDataLen)
-{
-#if WITH_PHYSX
-	Actors.SetNum(ActorDataLen);
-	RigidBodiesData.SetNum(ActorDataLen);
-	SolverBodiesData.SetNum(ActorDataLen);
-	PendingAcceleration.SetNum(ActorDataLen);
-	KinematicTargets.SetNum(ActorDataLen);
-#endif // WITH_PHYSX
-}
-
-void FPhysScene_ImmediatePhysX::AddActorsToScene_AssumesLocked(const TArray<FPhysicsActorHandle>& InActors)
-{
-    // Actors have been added to the list so do thing here
-}
-
-void FPhysScene_ImmediatePhysX::AddAggregateToScene(const FPhysicsAggregateHandle& InAggregate, bool bUseAsyncScene)
-{
-}
-
-/** Exposes destruction of physics-engine scene outside Engine. */
-FPhysScene_ImmediatePhysX::~FPhysScene_ImmediatePhysX()
-{
-}
-
-bool FPhysScene_ImmediatePhysX::GetKinematicTarget_AssumesLocked(const FBodyInstance* BodyInstance, FTransform& OutTM) const
-{
-   OutTM = FPhysicsInterface_ImmediatePhysX::GetKinematicTarget_AssumesLocked(BodyInstance->ActorHandle);
-   return true;
-}
-
-void FPhysScene_ImmediatePhysX::SetKinematicTarget_AssumesLocked(FBodyInstance* BodyInstance, const FTransform& TargetTransform, bool bAllowSubstepping)
-{
-    FPhysicsInterface_ImmediatePhysX::SetKinematicTarget_AssumesLocked(BodyInstance->ActorHandle, TargetTransform);
-}
-
-void FPhysScene_ImmediatePhysX::AddCustomPhysics_AssumesLocked(FBodyInstance* BodyInstance, FCalculateCustomPhysics& CalculateCustomPhysics)
-{
-    check(false);
-}
-
-void FPhysScene_ImmediatePhysX::AddForce_AssumesLocked(FBodyInstance* BodyInstance, const FVector& Force, bool bAllowSubstepping, bool bAccelChange)
-{
-    if (bAccelChange)
-    {
-        FPhysicsInterface_ImmediatePhysX::AddForceMassIndependent_AssumesLocked(BodyInstance->ActorHandle, Force);
-    }
-    else
-    {
-        FPhysicsInterface_ImmediatePhysX::AddForce_AssumesLocked(BodyInstance->ActorHandle, Force);
-    }
-}
-
-void FPhysScene_ImmediatePhysX::AddForceAtPosition_AssumesLocked(FBodyInstance* BodyInstance, const FVector& Force, const FVector& Position, bool bAllowSubstepping, bool bIsLocalForce)
-{
-    const auto& RigidBodyData = BodyInstance->ActorHandle.Scene->RigidBodiesData[BodyInstance->ActorHandle.Index];
-    BodyInstance->ActorHandle.Scene->PendingVelocityChange[BodyInstance->ActorHandle.Index] += U2PVector(Force * RigidBodyData.invMass);
-    BodyInstance->ActorHandle.Scene->PendingAngularVelocityChange[BodyInstance->ActorHandle.Index] += U2PVector(FVector::CrossProduct(Force * P2UVector(RigidBodyData.invInertia), Position - P2UTransform(RigidBodyData.body2World).GetTranslation()));
-}
-
-void FPhysScene_ImmediatePhysX::AddRadialForceToBody_AssumesLocked(FBodyInstance* BodyInstance, const FVector& Origin, const float Radius, const float Strength, const uint8 Falloff, bool bAccelChange, bool bAllowSubstepping)
-{
-    const auto& RigidBodyData = BodyInstance->ActorHandle.Scene->RigidBodiesData[BodyInstance->ActorHandle.Index];
-    FPhysicsInterface_ImmediatePhysX::AddRadialImpulse_AssumesLocked(BodyInstance->ActorHandle, Origin, Radius, bAccelChange ? Strength * Falloff : Strength * Falloff * RigidBodyData.invMass, ERadialImpulseFalloff::RIF_Linear, false);
-}
-
-void FPhysScene_ImmediatePhysX::ClearForces_AssumesLocked(FBodyInstance* BodyInstance, bool bAllowSubstepping)
-{
-    BodyInstance->ActorHandle.Scene->PendingAcceleration[BodyInstance->ActorHandle.Index] = physx::PxVec3(0, 0, 0);
-}
-
-void FPhysScene_ImmediatePhysX::AddTorque_AssumesLocked(FBodyInstance* BodyInstance, const FVector& Torque, bool bAllowSubstepping, bool bAccelChange)
-{
-    if (bAccelChange)
-    {
-        FPhysicsInterface_ImmediatePhysX::AddTorqueMassIndependent_AssumesLocked(BodyInstance->ActorHandle, Torque);
-    }
-    else
-    {
-        FPhysicsInterface_ImmediatePhysX::AddTorque_AssumesLocked(BodyInstance->ActorHandle, Torque);
-    }
-}
-
-void FPhysScene_ImmediatePhysX::ClearTorques_AssumesLocked(FBodyInstance* BodyInstance, bool bAllowSubstepping)
-{
-    BodyInstance->ActorHandle.Scene->PendingAngularAcceleration[BodyInstance->ActorHandle.Index] = physx::PxVec3(0, 0, 0);
-}
-
-void FPhysScene_ImmediatePhysX::RemoveBodyInstanceFromPendingLists_AssumesLocked(FBodyInstance* BodyInstance, int32 SceneType)
-{
-    // @todo(mlentine): Do we need to do anything here?
-}
-
-FAutoConsoleTaskPriority CPrio_PhysXStepSimulation(
-	TEXT("TaskGraph.TaskPriorities.PhysXStepSimulation"),
-	TEXT("Task and thread priority for FPhysSubstepTask::StepSimulation."),
-	ENamedThreads::HighThreadPriority, // if we have high priority task threads, then use them...
-	ENamedThreads::NormalTaskPriority, // .. at normal task priority
-	ENamedThreads::HighTaskPriority // if we don't have hi pri threads, then use normal priority threads at high task priority instead
-	);
-
-/** Adds to queue of skelmesh we want to add to collision disable table */
-void FPhysScene_ImmediatePhysX::DeferredAddCollisionDisableTable(uint32 SkelMeshCompID, TMap<struct FRigidBodyIndexPair, bool> * CollisionDisableTable)
-{
-    // @todo(mlentine): For new we ignore this as we probably want a different format for this going forward
-}
-
-/** Adds to queue of skelmesh we want to remove from collision disable table */
-void FPhysScene_ImmediatePhysX::DeferredRemoveCollisionDisableTable(uint32 SkelMeshCompID)
-{
-    // @todo(mlentine): For new we ignore this as we probably want a different format for this going forward
-}
-
-void FPhysScene_ImmediatePhysX::FlushDeferredCollisionDisableTableQueue()
-{
-    // @todo(mlentine): For new we ignore this as we probably want a different format for this going forward
-}
-
-bool FPhysScene_ImmediatePhysX::MarkForPreSimKinematicUpdate(USkeletalMeshComponent* InSkelComp, ETeleportType InTeleport, bool bNeedsSkinning)
-{
-	return false;
-}
-
-void FPhysScene_ImmediatePhysX::ClearPreSimKinematicUpdate(USkeletalMeshComponent* InSkelComp)
-{
-
-}
-
-/** Exposes ticking of physics-engine scene outside Engine. */
-void FPhysScene_ImmediatePhysX::KillVisualDebugger()
-{
-#if WITH_PHYSX
-	if (GPhysXVisualDebugger)
-	{
-		GPhysXVisualDebugger->disconnect();
-	}
-#endif // WITH_PHYSX 
-}
-
-void FPhysScene_ImmediatePhysX::WaitPhysScenes()
-{
-    // @todo(mlentine): For now do nothing as we run sync scenes only
-}
-
-/** Struct to remember a pending component transform change */
-struct FPhysScenePendingComponentTransform_PhysX
-{
-	/** Component to move */
-	TWeakObjectPtr<UPrimitiveComponent> OwningComp;
-	/** New transform from physics engine */
-	FTransform NewTransform;
-
-	FPhysScenePendingComponentTransform_PhysX(UPrimitiveComponent* InOwningComp, const FTransform& InNewTransform)
-		: OwningComp(InOwningComp)
-		, NewTransform(InNewTransform)
-	{}
-};
-
-void FPhysScene_ImmediatePhysX::SyncComponentsToBodies_AssumesLocked(uint32 SceneType)
-{
-	TArray<FPhysScenePendingComponentTransform_PhysX> PendingTransforms;
-
-	for (uint32 Index = 0; Index < static_cast<uint32>(RigidBodiesData.Num()); ++Index)
-	{		
-        FTransform NewTransform = FPhysicsInterface_ImmediatePhysX::GetGlobalPose_AssumesLocked(BodyInstances[Index]->ActorHandle);
-		FPhysScenePendingComponentTransform_PhysX NewEntry(BodyInstances[Index]->OwnerComponent.Get(), NewTransform);
-		PendingTransforms.Add(NewEntry);
-	}
-
-	for (FPhysScenePendingComponentTransform_PhysX& Entry : PendingTransforms)
-	{
-		UPrimitiveComponent* OwnerComponent = Entry.OwningComp.Get();
-		if (OwnerComponent != nullptr)
-		{
-			AActor* Owner = OwnerComponent->GetOwner();
-
-			if (!Entry.NewTransform.EqualsNoScale(OwnerComponent->GetComponentTransform()))
-			{
-				const FVector MoveBy = Entry.NewTransform.GetLocation() - OwnerComponent->GetComponentTransform().GetLocation();
-				const FQuat NewRotation = Entry.NewTransform.GetRotation();
-
-				OwnerComponent->MoveComponent(MoveBy, NewRotation, false, NULL, MOVECOMP_SkipPhysicsMove);
-			}
-
-			if (IsValid(Owner))
-			{
-				Owner->CheckStillInWorld();
-			}
-		}
-	}
-}
-
-void FPhysScene_ImmediatePhysX::DispatchPhysNotifications_AssumesLocked()
-{
-#if WITH_PHYSX
-	SCOPE_CYCLE_COUNTER(STAT_PhysicsEventTime);
-
-	for(int32 SceneType = 0; SceneType < PST_MAX; ++SceneType)
-	{
-		TArray<FCollisionNotifyInfo>& PendingCollisionNotifies = GetPendingCollisionNotifies(SceneType);
-
-		// Let the game-specific PhysicsCollisionHandler process any physics collisions that took place
-		if (OwningWorld != NULL && OwningWorld->PhysicsCollisionHandler != NULL)
-		{
-			OwningWorld->PhysicsCollisionHandler->HandlePhysicsCollisions_AssumesLocked(PendingCollisionNotifies);
-		}
-
-		// Fire any collision notifies in the queue.
-		for (int32 i = 0; i<PendingCollisionNotifies.Num(); i++)
-		{
-			FCollisionNotifyInfo& NotifyInfo = PendingCollisionNotifies[i];
-			if (NotifyInfo.RigidCollisionData.ContactInfos.Num() > 0)
-			{
-				if (NotifyInfo.bCallEvent0 && NotifyInfo.IsValidForNotify() && NotifyInfo.Info0.Actor.IsValid())
-				{
-					NotifyInfo.Info0.Actor->DispatchPhysicsCollisionHit(NotifyInfo.Info0, NotifyInfo.Info1, NotifyInfo.RigidCollisionData);
-				}
-
-				// Need to check IsValidForNotify again in case first call broke something.
-				if (NotifyInfo.bCallEvent1 && NotifyInfo.IsValidForNotify() && NotifyInfo.Info1.Actor.IsValid())
-				{
-					NotifyInfo.RigidCollisionData.SwapContactOrders();
-					NotifyInfo.Info1.Actor->DispatchPhysicsCollisionHit(NotifyInfo.Info1, NotifyInfo.Info0, NotifyInfo.RigidCollisionData);
-				}
-			}
-		}
-		PendingCollisionNotifies.Reset();
-	}
-
-    FPhysicsDelegates::OnPhysDispatchNotifications.Broadcast(this);
-}
-
-void FPhysScene_ImmediatePhysX::SetUpForFrame(const FVector* NewGrav, float InDeltaSeconds, float InMaxPhysicsDeltaTime, float InMaxSubstepDeltaTime, int32 InMaxSubsteps)
-{
-	DeltaSeconds = FMath::Min(InDeltaSeconds, 0.033f);
-	//Create dynamic bodies and integrate their unconstrained velocities
-    if (DeltaSeconds > 0)
-    {
-        ++SimCount;
-
-#if 0
-        ConstructSolverBodies(DeltaSeconds, NewGrav);
-
-        if (bRecreateIterationCache)
-        {
-            PrepareIterationCache();
-        }
-#endif
-    }
-}
-
-FAutoConsoleTaskPriority CPrio_PhyXSceneCompletion(
-	TEXT("TaskGraph.TaskPriorities.PhyXSceneCompletion"),
-	TEXT("Task and thread priority for PhysicsSceneCompletion."),
-	ENamedThreads::HighThreadPriority, // if we have high priority task threads, then use them...
-	ENamedThreads::HighTaskPriority, // .. at high task priority
-	ENamedThreads::HighTaskPriority // if we don't have hi pri threads, then use normal priority threads at high task priority instead
-	);
-
-void FPhysScene_ImmediatePhysX::StartFrame()
-{
-	SCOPE_CYCLE_COUNTER(STAT_TotalPhysicsTime);
-	CSV_SCOPED_TIMING_STAT(Basic, TotalPhysicsTime);
-
-#if 0
-	GenerateContacts();
-	BatchConstraints();
-	PrepareConstraints(DeltaSeconds);
-	SolveAndIntegrate(DeltaSeconds);
-#endif
-}
-
-void FPhysScene_ImmediatePhysX::EndFrame(ULineBatchComponent* InLineBatcher)
-{
-	check(IsInGameThread());
-
-	// Perform any collision notification events
-	DispatchPhysNotifications_AssumesLocked();
-
-    SyncComponentsToBodies_AssumesLocked(0);
-}
-
-#if WITH_PHYSX
-/** Helper struct that puts all awake actors to sleep and then later wakes them back up */
-struct FHelpEnsureCollisionTreeIsBuilt
-{
-	FHelpEnsureCollisionTreeIsBuilt(PxScene* InPScene)
-	: PScene(InPScene)
-	{
-		if(PScene)
-		{
-			SCOPED_SCENE_WRITE_LOCK(PScene);
-			const int32 NumActors = PScene->getNbActors(PxActorTypeFlag::eRIGID_DYNAMIC);
-
-			if (NumActors)
-			{
-				ActorBuffer.AddUninitialized(NumActors);
-				PScene->getActors(PxActorTypeFlag::eRIGID_DYNAMIC, &ActorBuffer[0], NumActors);
-
-				for (PxActor*& PActor : ActorBuffer)
-				{
-					if (PActor)
-					{
-						if (PxRigidDynamic* PDynamic = PActor->is<PxRigidDynamic>())
-						{
-							if (PDynamic->isSleeping() == false)
-							{
-								PDynamic->putToSleep();
-							}
-							else
-							{
-								PActor = nullptr;
-							}
-						}
-					}
-				}
-			}
-		}
-	}
-
-	~FHelpEnsureCollisionTreeIsBuilt()
-	{
-		SCOPED_SCENE_WRITE_LOCK(PScene);
-		for (PxActor* PActor : ActorBuffer)
-		{
-			if (PActor)
-			{
-				if (PxRigidDynamic* PDynamic = PActor->is<PxRigidDynamic>())
-				{
-					PDynamic->wakeUp();
-				}
-			}
-		}
-	}
-
-private:
-
-	TArray<PxActor*> ActorBuffer;
-	PxScene* PScene;
-};
-
-#endif
-
-static void BatchPxRenderBufferLines(class ULineBatchComponent& LineBatcherToUse, const PxRenderBuffer& DebugData)
-{
-	int32 NumPoints = DebugData.getNbPoints();
-	if (NumPoints > 0)
-	{
-		const PxDebugPoint* Points = DebugData.getPoints();
-		for (int32 i = 0; i<NumPoints; i++)
-		{
-			LineBatcherToUse.DrawPoint(P2UVector(Points->pos), FColor((uint32)Points->color), 2, SDPG_World);
-
-			Points++;
-		}
-	}
-
-	// Build a list of all the lines we want to draw
-	TArray<FBatchedLine> DebugLines;
-
-	// Add all the 'lines' from PhysX
-	int32 NumLines = DebugData.getNbLines();
-	if (NumLines > 0)
-	{
-		const PxDebugLine* Lines = DebugData.getLines();
-		for (int32 i = 0; i<NumLines; i++)
-		{
-			new(DebugLines)FBatchedLine(P2UVector(Lines->pos0), P2UVector(Lines->pos1), FColor((uint32)Lines->color0), 0.f, 0.0f, SDPG_World);
-			Lines++;
-		}
-	}
-
-	// Add all the 'triangles' from PhysX
-	int32 NumTris = DebugData.getNbTriangles();
-	if (NumTris > 0)
-	{
-		const PxDebugTriangle* Triangles = DebugData.getTriangles();
-		for (int32 i = 0; i<NumTris; i++)
-		{
-			new(DebugLines)FBatchedLine(P2UVector(Triangles->pos0), P2UVector(Triangles->pos1), FColor((uint32)Triangles->color0), 0.f, 0.0f, SDPG_World);
-			new(DebugLines)FBatchedLine(P2UVector(Triangles->pos1), P2UVector(Triangles->pos2), FColor((uint32)Triangles->color1), 0.f, 0.0f, SDPG_World);
-			new(DebugLines)FBatchedLine(P2UVector(Triangles->pos2), P2UVector(Triangles->pos0), FColor((uint32)Triangles->color2), 0.f, 0.0f, SDPG_World);
-			Triangles++;
-		}
-	}
-
-	// Draw them all in one call.
-	if (DebugLines.Num() > 0)
-	{
-		LineBatcherToUse.DrawLines(DebugLines);
-	}
-}
-#endif // WITH_PHYSX
-
-
-/** Add any debug lines from the physics scene to the supplied line batcher. */
-void FPhysScene_ImmediatePhysX::AddDebugLines(uint32 SceneType, class ULineBatchComponent* LineBatcherToUse)
-{
-	if (LineBatcherToUse)
-	{
-	}
-}
-
-void FPhysScene_ImmediatePhysX::ApplyWorldOffset(FVector InOffset)
-{
-}
-
-
-#if WITH_PHYSX
-void FPhysScene_ImmediatePhysX::AddPendingOnConstraintBreak(FConstraintInstance* ConstraintInstance, int32 SceneType)
-{
-	PendingConstraintData.PendingConstraintBroken.Add( FConstraintBrokenDelegateData(ConstraintInstance) );
-}
-
-FConstraintBrokenDelegateData::FConstraintBrokenDelegateData(FConstraintInstance* ConstraintInstance)
-	: OnConstraintBrokenDelegate(ConstraintInstance->OnConstraintBrokenDelegate)
-	, ConstraintIndex(ConstraintInstance->ConstraintIndex)
-{
-
-}
-
-void FPhysScene_ImmediatePhysX::AddPendingSleepingEvent(FBodyInstance* BI, ESleepEvent SleepEventType, int32 SceneType)
-{
-	PendingSleepEvents.FindOrAdd(BI) = SleepEventType;
-}
-
-#endif
-
-#endif
