Physics interface cleanup.

diff --git a/Engine/Source/Runtime/Engine/Private/PhysicsEngine/Experimental/PhysicsInterfaceImmediatePhysX.cpp b/Engine/Source/Runtime/Engine/Private/PhysicsEngine/Experimental/PhysicsInterfaceImmediatePhysX.cpp
--- a/Engine/Source/Runtime/Engine/Private/PhysicsEngine/Experimental/PhysicsInterfaceImmediatePhysX.cpp
+++ b/Engine/Source/Runtime/Engine/Private/PhysicsEngine/Experimental/PhysicsInterfaceImmediatePhysX.cpp
@@ -1,1872 +1 @@
 // Copyright Epic Games, Inc. All Rights Reserved.
-
-#if WITH_IMMEDIATE_PHYSX
-
-#include "Physics/PhysicsInterfaceImmediatePhysX.h"
-#include "Physics/PhysicsInterfaceUtils.h"
-
-#include "Physics/PhysScene_ImmediatePhysX.h"
-#include "Logging/MessageLog.h"
-#include "Components/SkeletalMeshComponent.h"
-#include "Internationalization/Internationalization.h"
-
-#if WITH_PHYSX
-#include "PhysxUserData.h"
-#include "PhysXPublic.h"
-#include "Physics/PhysicsFiltering.h"
-#include "PhysicsEngine/PhysXSupport.h"
-#include "Collision.h"
-#include "Collision/CollisionConversions.h"
-#endif // WITH_PHYSX
-
-#include "PhysicsEngine/ConstraintDrives.h"
-#include "PhysicsEngine/AggregateGeom.h"
-#include "Physics/PhysicsGeometry.h"
-#include "Engine/Engine.h"
-#include "PhysicalMaterials/PhysicalMaterial.h"
-
-using namespace physx;
-using namespace PhysicsInterfaceTypes;
-
-#define LOCTEXT_NAMESPACE "PhysicsInterface_PhysX"
-
-extern TAutoConsoleVariable<float> CVarConstraintLinearDampingScale;
-extern TAutoConsoleVariable<float> CVarConstraintLinearStiffnessScale;
-extern TAutoConsoleVariable<float> CVarConstraintAngularDampingScale;
-extern TAutoConsoleVariable<float> CVarConstraintAngularStiffnessScale;
-
-enum class EPhysicsInterfaceScopedLockType : uint8
-{
-	Read,
-	Write
-};
-
-struct FPhysicsInterfaceScopedLock_PhysX
-{
-private:
-
-	friend struct FPhysicsCommand_ImmediatePhysX;
-
-	FPhysicsInterfaceScopedLock_PhysX(FPhysicsActorHandle const * InActorReference, EPhysicsInterfaceScopedLockType InLockType);
-	FPhysicsInterfaceScopedLock_PhysX(FPhysicsActorHandle const * InActorReferenceA, FPhysicsActorHandle const * InActorReferenceB, EPhysicsInterfaceScopedLockType InLockType);
-	FPhysicsInterfaceScopedLock_PhysX(FPhysicsConstraintHandle const * InConstraintReference, EPhysicsInterfaceScopedLockType InLockType);
-	FPhysicsInterfaceScopedLock_PhysX(USkeletalMeshComponent* InSkelMeshComp, EPhysicsInterfaceScopedLockType InLockType);
-	FPhysicsInterfaceScopedLock_PhysX(FPhysScene* InScene, EPhysicsInterfaceScopedLockType InLockType);
-
-	~FPhysicsInterfaceScopedLock_PhysX();
-
-	void LockScenes();
-
-	FPhysScene* Scenes[2];
-	EPhysicsInterfaceScopedLockType LockType;
-};
-
-void FPhysicsInterface_ImmediatePhysX::CreateActor(const FActorCreationParams& Params, FPhysicsActorHandle& Handle)
-{
-	FPhysicsActorReference_ImmediatePhysX NewBody;
-	NewBody.Scene = Params.Scene;
-
-	PxSolverBodyData* NewSolverBodyData = new (Params.Scene->SolverBodiesData) PxSolverBodyData();
-	immediate::PxRigidBodyData* NewRigidBodyData = new (Params.Scene->RigidBodiesData) immediate::PxRigidBodyData();
-
-	NewRigidBodyData->body2World = U2PTransform(Params.InitialTM);
-
-	Params.Scene->PendingAcceleration.AddZeroed();
-	Params.Scene->PendingAngularAcceleration.AddZeroed();
-	Params.Scene->KinematicTargets.AddZeroed();
-    Params.Scene->BodyInstances.AddZeroed();
-
-	// @todo(mlentine): How do we treat these differently using immediate mode?
-	if(Params.bStatic || Params.bQueryOnly)
-	{
-		immediate::PxConstructStaticSolverBody(NewRigidBodyData->body2World, *NewSolverBodyData);
-	}
-	else
-	{
-        NewBody.Index = Params.Scene->NumSimulatedBodies;
-	    Params.Scene->SwapActorData(Params.Scene->SolverBodiesData.Num() - 1, Params.Scene->NumSimulatedBodies++);
-	}
-	
-	Handle = NewBody;
-}
-
-void FPhysicsInterface_ImmediatePhysX::ReleaseActor(FPhysicsActorHandle& InActorReference, FPhysScene* InScene)
-{
-	InActorReference.Scene->SwapActorData(InActorReference.Index, --InActorReference.Scene->NumSimulatedBodies);
-	InActorReference.Scene->ResizeActorData(InActorReference.Scene->NumSimulatedBodies);
-}
-
-PxD6Axis::Enum U2PConstraintAxis(PhysicsInterfaceTypes::ELimitAxis InAxis)
-{
-	switch(InAxis)
-	{
-	case PhysicsInterfaceTypes::ELimitAxis::X:
-		return PxD6Axis::Enum::eX;
-	case PhysicsInterfaceTypes::ELimitAxis::Y:
-		return PxD6Axis::Enum::eY;
-	case PhysicsInterfaceTypes::ELimitAxis::Z:
-		return PxD6Axis::Enum::eZ;
-	case PhysicsInterfaceTypes::ELimitAxis::Twist:
-		return PxD6Axis::Enum::eTWIST;
-	case PhysicsInterfaceTypes::ELimitAxis::Swing1:
-		return PxD6Axis::Enum::eSWING1;
-	case PhysicsInterfaceTypes::ELimitAxis::Swing2:
-		return PxD6Axis::Enum::eSWING2;
-	default:
-		check(false);
-	};
-
-	return PxD6Axis::Enum::eX;
-}
-
-/** Util for converting from UE motion enum to physx motion enum */
-PxD6Motion::Enum U2PAngularMotion(EAngularConstraintMotion InMotion)
-{
-	switch(InMotion)
-	{
-		case EAngularConstraintMotion::ACM_Free:
-			return PxD6Motion::eFREE;
-		case EAngularConstraintMotion::ACM_Limited:
-			return PxD6Motion::eLIMITED;
-		case EAngularConstraintMotion::ACM_Locked:
-			return PxD6Motion::eLOCKED;
-		default:
-			check(0);	//unsupported motion type
-	}
-
-	return PxD6Motion::eFREE;
-}
-
-/** Util for converting from UE motion enum to physx motion enum */
-PxD6Motion::Enum U2PLinearMotion(ELinearConstraintMotion InMotion)
-{
-	switch(InMotion)
-	{
-		case ELinearConstraintMotion::LCM_Free:
-			return PxD6Motion::eFREE;
-
-		case ELinearConstraintMotion::LCM_Limited:
-			return PxD6Motion::eLIMITED;
-		case ELinearConstraintMotion::LCM_Locked:
-			return PxD6Motion::eLOCKED;
-		default: 
-			check(0);	//unsupported motion type
-	}
-
-	return PxD6Motion::eFREE;
-}
-
-physx::PxJointActorIndex::Enum U2PConstraintFrame(EConstraintFrame::Type InFrame)
-{
-	// Swap frame order, since Unreal reverses physx order
-	return (InFrame == EConstraintFrame::Frame1) ? physx::PxJointActorIndex::eACTOR1 : physx::PxJointActorIndex::eACTOR0;
-}
-
-physx::PxD6Drive::Enum U2PDriveType(EDriveType InDriveType)
-{
-	switch(InDriveType)
-	{
-		case EDriveType::X:
-			return PxD6Drive::eX;
-		case EDriveType::Y:
-			return PxD6Drive::eY;
-		case EDriveType::Z:
-			return PxD6Drive::eZ;
-		case EDriveType::Swing:
-			return PxD6Drive::eSWING;
-		case EDriveType::Twist:
-			return PxD6Drive::eTWIST;
-		case EDriveType::Slerp:
-			return PxD6Drive::eSLERP;
-		default:
-			check(false);	// Invalid drive type
-			break;
-	};
-
-	return PxD6Drive::eX;
-}
-
-// PhysX interface definition //////////////////////////////////////////////////////////////////////////
-
-FPhysicsActorReference_ImmediatePhysX::FPhysicsActorReference_ImmediatePhysX() 
-	: Scene(nullptr)
-{
-
-}
-
-bool FPhysicsActorReference_ImmediatePhysX::IsValid() const
-{
-	return Scene != nullptr && Index < Scene->NumSimulatedBodies;
-}
-
-bool FPhysicsActorReference_ImmediatePhysX::Equals(const FPhysicsActorReference_ImmediatePhysX& Other) const
-{
-	return Scene == Other.Scene
-		&& Index == Other.Index;
-}
-
-FPhysicsConstraintReference_ImmediatePhysX::FPhysicsConstraintReference_ImmediatePhysX()
-	:Scene(nullptr)
-{
-
-}
-
-bool FPhysicsConstraintReference_ImmediatePhysX::IsValid() const
-{
-	return Scene != nullptr && Index < static_cast<uint32>(Scene->Joints.Num());
-}
-
-bool FPhysicsConstraintReference_ImmediatePhysX::Equals(const FPhysicsConstraintReference_ImmediatePhysX& Other) const
-{
-	return Scene == Other.Scene
-		&& Index == Other.Index;
-}
-
-
-FPhysicsAggregateReference_ImmediatePhysX::FPhysicsAggregateReference_ImmediatePhysX()
-	: Scene(nullptr)
-{
-
-}
-
-bool FPhysicsAggregateReference_ImmediatePhysX::IsValid() const
-{
-	return Scene != nullptr && Indices.Num() > 0;
-}
-
-// @todo(mlentine): Do we need locks?
-FPhysicsInterfaceScopedLock_PhysX::FPhysicsInterfaceScopedLock_PhysX(FPhysicsActorHandle const * InActorReference, EPhysicsInterfaceScopedLockType InLockType)
-	: LockType(InLockType)
-{
-	Scenes[0] = InActorReference ? InActorReference->Scene : nullptr;
-    Scenes[1] = nullptr;
-	LockScenes();
-}
-
-FPhysicsInterfaceScopedLock_PhysX::FPhysicsInterfaceScopedLock_PhysX(FPhysicsActorHandle const * InActorReferenceA, FPhysicsActorHandle const * InActorReferenceB, EPhysicsInterfaceScopedLockType InLockType)
-	: LockType(InLockType)
-{
-	Scenes[0] = InActorReferenceA ? InActorReferenceA->Scene : nullptr;
-	Scenes[1] = InActorReferenceB ? InActorReferenceB->Scene : nullptr;
-
-	// Only lock if we have unique scenes, either one vs. nullptr or both are equal
-	if(Scenes[0] == Scenes[1] || (!Scenes[0] || !Scenes[1]))
-	{
-		LockScenes();
-	}
-	else
-	{
-		UE_LOG(LogPhysics, Warning, TEXT("Attempted to aquire a physics scene lock for two paired actors that were not in the same scene. Skipping lock"));
-	}
-}
-
-
-FPhysicsInterfaceScopedLock_PhysX::FPhysicsInterfaceScopedLock_PhysX(FPhysicsConstraintHandle const * InConstraintReference, EPhysicsInterfaceScopedLockType InLockType)
-	: LockType(InLockType)
-{
-	if(InConstraintReference)
-	{
-		Scenes[0] = InConstraintReference->Scene;
-		Scenes[1] = nullptr;
-		LockScenes();
-	}
-}
-
-FPhysicsInterfaceScopedLock_PhysX::FPhysicsInterfaceScopedLock_PhysX(USkeletalMeshComponent* InSkelMeshComp, EPhysicsInterfaceScopedLockType InLockType)
-	: LockType(InLockType)
-{
-	Scenes[0] = nullptr;
-	Scenes[1] = nullptr;
-
-	LockScenes();
-}
-
-FPhysicsInterfaceScopedLock_PhysX::FPhysicsInterfaceScopedLock_PhysX(FPhysScene* InScene, EPhysicsInterfaceScopedLockType InLockType)
-	: LockType(InLockType)
-{
-	Scenes[0] = InScene;
-	Scenes[1] = nullptr;
-
-	LockScenes();
-}
-
-FPhysicsInterfaceScopedLock_PhysX::~FPhysicsInterfaceScopedLock_PhysX()
-{
-	if (Scenes[0])
-	{
-		if (LockType == EPhysicsInterfaceScopedLockType::Read)
-		{
-		}
-		else if (LockType == EPhysicsInterfaceScopedLockType::Write)
-		{
-		}
-	}
-
-	if (Scenes[1])
-	{
-		if (LockType == EPhysicsInterfaceScopedLockType::Read)
-		{
-		}
-		else if (LockType == EPhysicsInterfaceScopedLockType::Write)
-		{
-		}
-	}
-}
-
-
-void FPhysicsInterfaceScopedLock_PhysX::LockScenes()
-{
-	if (Scenes[0])
-	{
-		if (LockType == EPhysicsInterfaceScopedLockType::Read)
-		{
-		}
-		else if(LockType == EPhysicsInterfaceScopedLockType::Write)
-		{
-		}
-	}
-
-	if (Scenes[1])
-	{
-		if (LockType == EPhysicsInterfaceScopedLockType::Read)
-		{
-		}
-		else if (LockType == EPhysicsInterfaceScopedLockType::Write)
-		{
-		}
-	}
-}
-
-bool FPhysicsCommand_ImmediatePhysX::ExecuteRead(const FPhysicsActorHandle& InActorReference, TFunctionRef<void(const FPhysicsActorHandle& Actor)> InCallable)
-{
-	if(InActorReference.IsValid())
-	{
-		FPhysicsInterfaceScopedLock_PhysX ScopeLock(&InActorReference, EPhysicsInterfaceScopedLockType::Read);
-		InCallable(InActorReference);
-
-		return true;
-	}
-
-	return false;
-}
-
-bool FPhysicsCommand_ImmediatePhysX::ExecuteRead(USkeletalMeshComponent* InMeshComponent, TFunctionRef<void()> InCallable)
-{
-	FPhysicsInterfaceScopedLock_PhysX ScopeLock(InMeshComponent, EPhysicsInterfaceScopedLockType::Read);
-	InCallable();
-
-	return ScopeLock.Scenes[0] || ScopeLock.Scenes[1];
-}
-
-bool FPhysicsCommand_ImmediatePhysX::ExecuteRead(const FPhysicsActorHandle& InActorReferenceA, const FPhysicsActorHandle& InActorReferenceB, TFunctionRef<void(const FPhysicsActorHandle& ActorA, const FPhysicsActorHandle& ActorB)> InCallable)
-{
-	if(InActorReferenceA.IsValid() || InActorReferenceB.IsValid())
-	{
-		FPhysicsInterfaceScopedLock_PhysX ScopeLock(&InActorReferenceA, &InActorReferenceB, EPhysicsInterfaceScopedLockType::Read);
-		InCallable(InActorReferenceA, InActorReferenceB);
-
-		return ScopeLock.Scenes[0] || ScopeLock.Scenes[1];
-	}
-
-	return false;
-}
-
-bool FPhysicsCommand_ImmediatePhysX::ExecuteRead(const FPhysicsConstraintHandle& InConstraintRef, TFunctionRef<void(const FPhysicsConstraintHandle& Constraint)> InCallable)
-{
-	if(InConstraintRef.IsValid())
-	{
-		FPhysicsInterfaceScopedLock_PhysX ScopeLock(&InConstraintRef, EPhysicsInterfaceScopedLockType::Read);
-		InCallable(InConstraintRef);
-
-		return true;
-	}
-	
-	return false;
-}
-
-bool FPhysicsCommand_ImmediatePhysX::ExecuteRead(FPhysScene* InScene, TFunctionRef<void()> InCallable)
-{
- 	if(InScene)
-	{
-		FPhysicsInterfaceScopedLock_PhysX ScopeLock(InScene, EPhysicsInterfaceScopedLockType::Read);
-		InCallable();
-
-		return true;
-	}
-
-	return false;   
-}
-
-bool FPhysicsCommand_ImmediatePhysX::ExecuteWrite(const FPhysicsActorHandle& InActorReference, TFunctionRef<void(const FPhysicsActorHandle& Actor)> InCallable)
-{
-	if(InActorReference.IsValid())
-	{
-		FPhysicsInterfaceScopedLock_PhysX ScopeLock(&InActorReference, EPhysicsInterfaceScopedLockType::Write);
-		InCallable(InActorReference);
-
-		return true;
-	}
-
-	return false;
-}
-
-bool FPhysicsCommand_ImmediatePhysX::ExecuteWrite(USkeletalMeshComponent* InMeshComponent, TFunctionRef<void()> InCallable)
-{
-	FPhysicsInterfaceScopedLock_PhysX ScopeLock(InMeshComponent, EPhysicsInterfaceScopedLockType::Write);
-	InCallable();
-
-	return ScopeLock.Scenes[0] || ScopeLock.Scenes[1];
-}
-
-bool FPhysicsCommand_ImmediatePhysX::ExecuteWrite(const FPhysicsActorHandle& InActorReferenceA, const FPhysicsActorHandle& InActorReferenceB, TFunctionRef<void(const FPhysicsActorHandle& ActorA, const FPhysicsActorHandle& ActorB)> InCallable)
-{
-	if(InActorReferenceA.IsValid() || InActorReferenceB.IsValid())
-	{
-		FPhysicsInterfaceScopedLock_PhysX ScopeLock(&InActorReferenceA, &InActorReferenceB, EPhysicsInterfaceScopedLockType::Write);
-		InCallable(InActorReferenceA, InActorReferenceB);
-
-		return ScopeLock.Scenes[0] || ScopeLock.Scenes[1];
-	}
-
-	return false;
-}
-
-bool FPhysicsCommand_ImmediatePhysX::ExecuteWrite(const FPhysicsConstraintHandle& InConstraintRef, TFunctionRef<void(const FPhysicsConstraintHandle& Constraint)> InCallable)
-{
-	if(InConstraintRef.IsValid())
-	{
-		FPhysicsInterfaceScopedLock_PhysX ScopeLock(&InConstraintRef, EPhysicsInterfaceScopedLockType::Write);
-		InCallable(InConstraintRef);
-
-		return true;
-	}
-
-	return false;
-}
-
-bool FPhysicsCommand_ImmediatePhysX::ExecuteWrite(FPhysScene* InScene, TFunctionRef<void()> InCallable)
-{
- 	if(InScene)
-	{
-		FPhysicsInterfaceScopedLock_PhysX ScopeLock(InScene, EPhysicsInterfaceScopedLockType::Write);
-		InCallable();
-
-		return true;
-	}
-
-	return false;   
-}
-
-struct FScopedSharedShapeHandler
-{
-	FScopedSharedShapeHandler() = delete;
-	FScopedSharedShapeHandler(const FScopedSharedShapeHandler& Other) = delete;
-	FScopedSharedShapeHandler& operator=(const FScopedSharedShapeHandler& Other) = delete;
-
-	FScopedSharedShapeHandler(FBodyInstance* InInstance, FPhysicsShapeHandle& InShape)
-		: Instance(InInstance)
-		, Shape(InShape)
-		, bShared(false)
-	{
-		bShared = Instance && Instance->HasSharedShapes() && Instance->ActorHandle.IsValid();
-
-		if(bShared)
-		{
-			Actor = Instance->ActorHandle;
-
-			FPhysicsShapeHandle NewShape = FPhysicsInterface::CloneShape(Shape);
-			FPhysicsInterface::DetachShape(Actor, Shape);
-			Shape = NewShape;
-		}
-	}
-
-	~FScopedSharedShapeHandler()
-	{
-		if(bShared)
-		{
-			FPhysicsInterface::AttachShape(Actor, Shape);
-			FPhysicsInterface::ReleaseShape(Shape);
-		}
-	}
-
-private:
-	FBodyInstance* Instance;
-	FPhysicsShapeHandle& Shape;
-	FPhysicsActorHandle Actor;
-	bool bShared;
-};
-
-void FPhysicsCommand_ImmediatePhysX::ExecuteShapeWrite(FBodyInstance* InInstance, FPhysicsShapeHandle& InShape, TFunctionRef<void(const FPhysicsShapeHandle& InShape)> InCallable)
-{
-	if(InShape.IsValid())
-	{
-		FScopedSharedShapeHandler SharedShapeHandler(InInstance, InShape);
-		InCallable(InShape);
-	}
-}
-
-template<typename AllocatorType>
-int32 FPhysicsInterface_ImmediatePhysX::GetAllShapes_AssumedLocked(const FPhysicsActorReference_ImmediatePhysX& InActorReference, TArray<FPhysicsShapeHandle, AllocatorType>& OutShapes, EPhysicsSceneType InSceneType)
-{
-	OutShapes.Empty();
-
-	// @todo(mlentine): Fix Memory Leak for new Materials
-	for (int32 i = 0; i < InActorReference.Scene->Actors[InActorReference.Index].Shapes.Num(); ++i)
-	{
-        FPhysicsShapeHandle NewHandle(FPhysScene_ImmediatePhysX::FShape(InActorReference.Scene->Actors[InActorReference.Index].Shapes[i]));
-        NewHandle.Actor = const_cast<FPhysicsActorReference_ImmediatePhysX*>(&InActorReference);
-        NewHandle.Index = i;
-        OutShapes.Add(NewHandle);
-		//OutShapes.Append(GPhysXSDK->createShape(*Shape.Geometry, *GPhysXSDK->createMaterial(Shape.Material.StaticFriction, Shape.Material.DynamicFriction, Shape.Material.Restitution), 1));
-	}
-
-	return OutShapes.Num();
-}
-
-template int32 FPhysicsInterface_ImmediatePhysX::GetAllShapes_AssumedLocked(const FPhysicsActorReference_ImmediatePhysX& InActorHandle, TArray<FPhysicsShapeHandle>& OutShapes, EPhysicsSceneType InSceneType);
-template int32 FPhysicsInterface_ImmediatePhysX::GetAllShapes_AssumedLocked(const FPhysicsActorReference_ImmediatePhysX& InActorHandle, PhysicsInterfaceTypes::FInlineShapeArray& OutShapes, EPhysicsSceneType InSceneType);
-
-void FPhysicsInterface_ImmediatePhysX::GetNumShapes(const FPhysicsActorHandle& InHandle, int32& OutNumSyncShapes, int32& OutNumAsyncShapes)
-{
-    // @todo(mlentine): What to do in this case with sync/async?
-    OutNumSyncShapes = InHandle.Scene->Actors[InHandle.Index].Shapes.Num();
-    OutNumAsyncShapes = 0;
-}
-
-void FPhysicsInterface_ImmediatePhysX::ReleaseShape(const FPhysicsShapeHandle& InShape)
-{
-    check(InShape.Actor == nullptr);
-}
-
-void FPhysicsInterface_ImmediatePhysX::AttachShape(const FPhysicsActorHandle& InActor, const FPhysicsShapeHandle& InNewShape)
-{
-    const_cast<FPhysicsShapeHandle&>(InNewShape).Index = InActor.Scene->Actors[InActor.Index].Shapes.Num();
-    const_cast<FPhysicsShapeHandle&>(InNewShape).Actor = const_cast<FPhysicsActorHandle*>(&InActor);
-    InActor.Scene->Actors[InActor.Index].Shapes.Add(InNewShape.Shape);
-}
-
-void FPhysicsInterface_ImmediatePhysX::AttachShape(const FPhysicsActorHandle& InActor, const FPhysicsShapeHandle& InNewShape, EPhysicsSceneType SceneType)
-{
-    AttachShape(InActor, InNewShape);
-}
-
-void FPhysicsInterface_ImmediatePhysX::DetachShape(const FPhysicsActorHandle& InActor, FPhysicsShapeHandle& InShape, bool bWakeTouching)
-{
-    // @todo(mlentine): We need to renumber shapes before we can remove it
-    //InActor.Scene->Actors[InActor.Index].Shapes.RemoveAt(InShape.Index);
-    InShape.Actor = nullptr;
-}
-
-FPhysicsAggregateReference_ImmediatePhysX FPhysicsInterface_ImmediatePhysX::CreateAggregate(int32 MaxBodies)
-{
-	FPhysicsAggregateHandle NewAggregate;
-	return NewAggregate;
-}
-
-void FPhysicsInterface_ImmediatePhysX::ReleaseAggregate(FPhysicsAggregateReference_ImmediatePhysX& InAggregate)
-{
-	if (InAggregate.IsValid())
-	{
-		InAggregate.Indices.SetNum(0);
-		InAggregate.Scene = nullptr;
-	}
-}
-
-
-int32 FPhysicsInterface_ImmediatePhysX::GetNumActorsInAggregate(const FPhysicsAggregateReference_ImmediatePhysX& InAggregate)
-{
-	int32 NumActors = 0;
-	if (InAggregate.IsValid())
-	{
-		NumActors = InAggregate.Indices.Num();
-	}
-	return NumActors;
-}
-
-// @todo(mlentine): InAggregate should be modifiable
-void FPhysicsInterface_ImmediatePhysX::AddActorToAggregate_AssumesLocked(const FPhysicsAggregateReference_ImmediatePhysX& InAggregate, const FPhysicsActorReference_ImmediatePhysX& InActor)
-{
-	if (InAggregate.Scene)
-	{
-		check(InAggregate.Scene == InActor.Scene);
-	}
-	else
-	{
-		const_cast<FPhysicsAggregateReference_ImmediatePhysX&>(InAggregate).Scene = InActor.Scene;
-	}
-	const_cast<FPhysicsAggregateReference_ImmediatePhysX&>(InAggregate).Indices.Add(InActor.Index);
-}
-
-////////////////////////////////////////////////////
-PxMaterial* GetDefaultPhysMaterial()
-{
-	check(GEngine->DefaultPhysMaterial != NULL);
-#if WITH_IMMEDIATE_PHYSX
-    return nullptr;
-#else
-	return GEngine->DefaultPhysMaterial->GetPhysicsMaterial();
-#endif
-}
-
-
-FPhysicsShapeReference_ImmediatePhysX FPhysicsInterface_ImmediatePhysX::CreateShape(physx::PxGeometry* InGeom, bool bSimulation, bool bQuery, UPhysicalMaterial* InSimpleMaterial, TArray<UPhysicalMaterial*>* InComplexMaterials, bool bShared)
-{
-    //@todo(mlentine): What do we do with simulation and query here?
-    FPhysScene_ImmediatePhysX::FMaterial NewMaterial;
-    if (InComplexMaterials)
-    {
-        check(InComplexMaterials->Num() == 1);
-        NewMaterial.StaticFriction = (*InComplexMaterials)[0]->Friction;
-        NewMaterial.DynamicFriction = (*InComplexMaterials)[0]->Friction;
-        NewMaterial.Restitution = (*InComplexMaterials)[0]->Restitution;
-    }
-    FPhysScene_ImmediatePhysX::FShape NewShape(PxTransform(), PxVec3(), 1, InGeom, NewMaterial);
-    FPhysicsShapeReference_ImmediatePhysX NewHandle(MoveTemp(NewShape));
-    return NewHandle;
-}
-
-void FPhysicsInterface_ImmediatePhysX::AddGeometry(const FPhysicsActorHandle& InActor, const FGeometryAddParams& InParams, TArray<FPhysicsShapeReference_ImmediatePhysX>* OutOptShapes)
-{
-	auto AttachShape_AssumesLocked = [&InParams, &OutOptShapes, InActor](const PxGeometry& PGeom, const PxTransform& PLocalPose, const float ContactOffset, const float RestOffset)
-	{
-		const bool bShapeSharing = InParams.bSharedShapes;
-		const FBodyCollisionData& BodyCollisionData = InParams.CollisionData;
-
-        FPhysScene_ImmediatePhysX::FMaterial Material(GetDefaultPhysMaterial());
-
-        if (OutOptShapes)
-        {
-            FPhysicsShapeHandle NewShapeRef(FPhysScene_ImmediatePhysX::FShape(PLocalPose, PxVec3(ContactOffset), RestOffset, &PGeom, Material));
-            NewShapeRef.Index = InActor.Scene->Actors[InActor.Index].Shapes.Num();
-            NewShapeRef.Actor = const_cast<FPhysicsActorHandle*>(&InActor);
-            OutOptShapes->Add(NewShapeRef);
-        }
-		
-		InActor.Scene->Actors[InActor.Index].Shapes.Add(FPhysScene_ImmediatePhysX::FShape(PLocalPose, PxVec3(ContactOffset), RestOffset, &PGeom, Material));
-	};
-
-	auto IterateSimpleShapes = [AttachShape_AssumesLocked](const FKShapeElem& Elem, const PxGeometry& Geom, const PxTransform& PLocalPose, float ContactOffset, float RestOffset)
-	{
-		AttachShape_AssumesLocked(Geom, PLocalPose, ContactOffset, RestOffset);
-	};
-
-	auto IterateTrimeshes = [AttachShape_AssumesLocked](PxTriangleMesh*, const PxGeometry& Geom, const PxTransform& PLocalPose, float ContactOffset, float RestOffset)
-	{
-		// Create without 'sim shape' flag, problematic if it's kinematic, and it gets set later anyway.
-        AttachShape_AssumesLocked(Geom, PLocalPose, ContactOffset, RestOffset);
-	};
-
-	FBodySetupShapeIterator AddShapesHelper(InParams.Scale, InParams.LocalTransform, InParams.bDoubleSided);
-
-	// Create shapes for simple collision if we do not want to use the complex collision mesh 
-	// for simple queries as well
-	FKAggregateGeom* AggGeom = InParams.Geometry;
-	check(AggGeom);
-
-	if (InParams.CollisionTraceType != ECollisionTraceFlag::CTF_UseComplexAsSimple)
-	{
-		AddShapesHelper.ForEachShape<FKSphereElem, PxSphereGeometry>(AggGeom->SphereElems, IterateSimpleShapes);
-		AddShapesHelper.ForEachShape<FKSphylElem, PxCapsuleGeometry>(AggGeom->SphylElems, IterateSimpleShapes);
-		AddShapesHelper.ForEachShape<FKBoxElem, PxBoxGeometry>(AggGeom->BoxElems, IterateSimpleShapes);
-		AddShapesHelper.ForEachShape<FKConvexElem, PxConvexMeshGeometry>(AggGeom->ConvexElems, IterateSimpleShapes);
-	}
-	
-	// Create tri-mesh shape, when we are not using simple collision shapes for 
-	// complex queries as well
-	if (InParams.CollisionTraceType != ECollisionTraceFlag::CTF_UseSimpleAsComplex)
-	{
-		AddShapesHelper.ForEachShape<PxTriangleMesh*, PxTriangleMeshGeometry>(InParams.TriMeshes, IterateTrimeshes);
-	}
-}
-
-FPhysicsShapeHandle FPhysicsInterface_ImmediatePhysX::CloneShape(const FPhysicsShapeHandle& InShape)
-{
-    FPhysicsShapeHandle NewShapeRef(FPhysScene_ImmediatePhysX::FShape(InShape.Shape.LocalTM, InShape.Shape.BoundsOffset, InShape.Shape.BoundsMagnitude, InShape.Shape.Geometry, InShape.Shape.Material));
-    if (InShape.Actor)
-    {
-        const auto& Shape = InShape.Actor->Scene->Actors[InShape.Actor->Index].Shapes[InShape.Index];
-        NewShapeRef.Index = InShape.Actor->Scene->Actors[InShape.Actor->Index].Shapes.Num();
-        NewShapeRef.Actor = InShape.Actor;
-        InShape.Actor->Scene->Actors[InShape.Actor->Index].Shapes.Add(FPhysScene_ImmediatePhysX::FShape(Shape.LocalTM, Shape.BoundsOffset, Shape.BoundsMagnitude, Shape.Geometry, Shape.Material));
-    }
-    else
-    {
-        NewShapeRef.Actor = nullptr;
-    }
-    return NewShapeRef;
-}
-
-bool FPhysicsInterface_ImmediatePhysX::IsSimulationShape(const FPhysicsShapeHandle& InShape)
-{
-    return InShape.Actor != nullptr;
-}
-
-bool FPhysicsInterface_ImmediatePhysX::IsQueryShape(const FPhysicsShapeHandle& InShape)
-{
-    return InShape.Actor != nullptr;
-}
-
-ECollisionShapeType FPhysicsInterface_ImmediatePhysX::GetShapeType(const FPhysicsShapeHandle& InShape)
-{
-    if (!InShape.Actor)
-    {
-        return ECollisionShapeType::None;
-    }
-    const FPhysScene_ImmediatePhysX::FShape& Shape = InShape.Actor->Scene->Actors[InShape.Actor->Index].Shapes[InShape.Index];
-    if (Shape.Geometry->getType() == PxGeometryType::Enum::eSPHERE)
-    {
-        return ECollisionShapeType::Sphere;
-    }
-    if (Shape.Geometry->getType() == PxGeometryType::Enum::eBOX)
-    {
-        return ECollisionShapeType::Box;
-    }
-    if (Shape.Geometry->getType() == PxGeometryType::Enum::eCONVEXMESH)
-    {
-        return ECollisionShapeType::Convex;
-    }
-    if (Shape.Geometry->getType() == PxGeometryType::Enum::eTRIANGLEMESH)
-    {
-        return ECollisionShapeType::Trimesh;
-    }
-    if (Shape.Geometry->getType() == PxGeometryType::Enum::eHEIGHTFIELD)
-    {
-        return ECollisionShapeType::Heightfield;
-    }
-    if (Shape.Geometry->getType() == PxGeometryType::Enum::eCAPSULE)
-    {
-        return ECollisionShapeType::TCapsule;
-    }
-    return ECollisionShapeType::None;
-}
-
-FPhysicsGeometryCollection FPhysicsInterface_ImmediatePhysX::GetGeometryCollection(const FPhysicsShapeHandle& InShape)
-{
-    FPhysicsGeometryCollection Geometry;
-    Geometry.Geometry = nullptr;
-    if (InShape.Actor)
-    {
-        const FPhysScene_ImmediatePhysX::FShape& Shape = InShape.Actor->Scene->Actors[InShape.Actor->Index].Shapes[InShape.Index];
-        Geometry.Geometry = const_cast<physx::PxGeometry*>(Shape.Geometry);
-    }
-    return Geometry;
-}
-
-FTransform FPhysicsInterface_ImmediatePhysX::GetLocalTransform(const FPhysicsShapeHandle& InShape)
-{
-    if (!InShape.Actor)
-    {
-        return FTransform();
-    }
-    const FPhysScene_ImmediatePhysX::FShape& Shape = InShape.Actor->Scene->Actors[InShape.Actor->Index].Shapes[InShape.Index];
-    return P2UTransform(Shape.LocalTM);
-}
-
-void* FPhysicsInterface_ImmediatePhysX::GetUserData(const FPhysicsShapeHandle& InShape)
-{
-    return nullptr;
-}
-
-// @todo(mlentine): Do we need to do anything for these?
-void FPhysicsInterface_ImmediatePhysX::SetMaskFilter(const FPhysicsShapeHandle& InShape, FMaskFilter InFilter)
-{
-}
-
-void FPhysicsInterface_ImmediatePhysX::SetSimulationFilter(const FPhysicsShapeHandle& InShape, const FCollisionFilterData& InFilter)
-{
-}
-
-void FPhysicsInterface_ImmediatePhysX::SetQueryFilter(const FPhysicsShapeHandle& InShape, const FCollisionFilterData& InFilter)
-{
-}
-
-void FPhysicsInterface_ImmediatePhysX::SetIsSimulationShape(const FPhysicsShapeHandle& InShape, bool bIsSimShape)
-{
-}
-
-void FPhysicsInterface_ImmediatePhysX::SetIsQueryShape(const FPhysicsShapeHandle& InShape, bool bIsQueryShape)
-{
-}
-
-void FPhysicsInterface_ImmediatePhysX::SetUserData(const FPhysicsShapeHandle& InShape, void* InUserData)
-{
-}
-
-void FPhysicsInterface_ImmediatePhysX::SetGeometry(const FPhysicsShapeHandle& InShape, physx::PxGeometry& InGeom)
-{
-    if (!InShape.Actor)
-    {
-        return;
-    }
-    const FPhysScene_ImmediatePhysX::FShape& Shape = InShape.Actor->Scene->Actors[InShape.Actor->Index].Shapes[InShape.Index];
-    const_cast<FPhysicsShapeHandle&>(InShape).Shape = InShape.Actor->Scene->Actors[InShape.Actor->Index].Shapes[InShape.Index] = FPhysScene_ImmediatePhysX::FShape(Shape.LocalTM, Shape.BoundsOffset, Shape.BoundsMagnitude, &InGeom, Shape.Material);
-}
-
-void FPhysicsInterface_ImmediatePhysX::SetLocalTransform(const FPhysicsShapeHandle& InShape, const FTransform& NewLocalTransform)
-{
-    if (!InShape.Actor)
-    {
-        return;
-    }
-    const FPhysScene_ImmediatePhysX::FShape& Shape = InShape.Actor->Scene->Actors[InShape.Actor->Index].Shapes[InShape.Index];
-    const_cast<FPhysicsShapeHandle&>(InShape).Shape = InShape.Actor->Scene->Actors[InShape.Actor->Index].Shapes[InShape.Index] = FPhysScene_ImmediatePhysX::FShape(U2PTransform(NewLocalTransform), Shape.BoundsOffset, Shape.BoundsMagnitude, Shape.Geometry, Shape.Material);
-}
-
-void FPhysicsInterface_ImmediatePhysX::SetMaterials(const FPhysicsShapeHandle& InShape, const TArrayView<UPhysicalMaterial*>InMaterials)
-{
-    FPhysScene_ImmediatePhysX::FMaterial NewMaterial;
-    check(InMaterials.Num() == 1);
-    NewMaterial.StaticFriction = InMaterials[0]->Friction;
-    NewMaterial.DynamicFriction = InMaterials[0]->Friction;
-    NewMaterial.Restitution = InMaterials[0]->Restitution;
-    if (!InShape.Actor)
-    {
-        return;
-    }
-    const FPhysScene_ImmediatePhysX::FShape& Shape = InShape.Actor->Scene->Actors[InShape.Actor->Index].Shapes[InShape.Index];
-    const_cast<FPhysicsShapeHandle&>(InShape).Shape = InShape.Actor->Scene->Actors[InShape.Actor->Index].Shapes[InShape.Index] = FPhysScene_ImmediatePhysX::FShape(Shape.LocalTM, Shape.BoundsOffset, Shape.BoundsMagnitude, Shape.Geometry, NewMaterial);
-}
-
-FPhysicsMaterialHandle FPhysicsInterface_ImmediatePhysX::CreateMaterial(const UPhysicalMaterial* InMaterial)
-{
-    check(InMaterial);
-    FPhysScene_ImmediatePhysX::FMaterial NewMaterial;
-    NewMaterial.StaticFriction = InMaterial->Friction;
-    NewMaterial.DynamicFriction = InMaterial->Friction;
-    NewMaterial.Restitution = InMaterial->Restitution;
-    FPhysicsMaterialHandle Handle;
-    Handle.Material = NewMaterial;
-    return Handle;
-}
-
-void FPhysicsInterface_ImmediatePhysX::ReleaseMaterial(FPhysicsMaterialHandle& InHandle)
-{
-
-}
-
-void FPhysicsInterface_ImmediatePhysX::UpdateMaterial(const FPhysicsMaterialHandle& InHandle, UPhysicalMaterial* InMaterial)
-{
-    check(InMaterial);
-    const_cast<FPhysicsMaterialHandle&>(InHandle).Material.StaticFriction = InMaterial->Friction;
-    const_cast<FPhysicsMaterialHandle&>(InHandle).Material.DynamicFriction = InMaterial->Friction;
-    const_cast<FPhysicsMaterialHandle&>(InHandle).Material.Restitution = InMaterial->Restitution;
-}
-
-void FPhysicsInterface_ImmediatePhysX::SetUserData(const FPhysicsMaterialHandle& InHandle, void* InUserData)
-{
-    check(false);
-}
-
-void FPhysicsInterface_ImmediatePhysX::SetActorUserData_AssumesLocked(const FPhysicsActorHandle& InActorReference, FPhysxUserData* InUserData)
-{
-    if (FPhysScene* Scene = InUserData->Get<FPhysScene>(InUserData))
-    {
-        const_cast<FPhysicsActorHandle&>(InActorReference).Scene = Scene;
-    }
-    if (FBodyInstance* BodyInstance = InUserData->Get<FBodyInstance>(InUserData))
-    {
-        InActorReference.Scene->BodyInstances[InActorReference.Index] = BodyInstance;
-    }
-}
-
-bool FPhysicsInterface_ImmediatePhysX::IsRigidBody(const FPhysicsActorHandle& InActorReference)
-{
-    return InActorReference.IsValid();
-}
-
-bool FPhysicsInterface_ImmediatePhysX::IsDynamic(const FPhysicsActorHandle& InActorReference)
-{
-    return InActorReference.Index < InActorReference.Scene->NumSimulatedBodies;
-}
-
-bool FPhysicsInterface_ImmediatePhysX::IsStatic(const FPhysicsActorHandle& InActorReference)
-{
-    return InActorReference.Index >= (InActorReference.Scene->NumSimulatedBodies + InActorReference.Scene->NumKinematicBodies);
-}
-
-bool FPhysicsInterface_ImmediatePhysX::IsKinematic_AssumesLocked(const FPhysicsActorHandle& InActorReference)
-{
-    return InActorReference.Index >= InActorReference.Scene->NumSimulatedBodies &&
-        InActorReference.Index < (InActorReference.Scene->NumSimulatedBodies + InActorReference.Scene->NumKinematicBodies);
-}
-
-bool FPhysicsInterface_ImmediatePhysX::IsSleeping(const FPhysicsActorHandle& InActorReference)
-{
-    return !IsDynamic(InActorReference);
-}
-
-bool FPhysicsInterface_ImmediatePhysX::IsCcdEnabled(const FPhysicsActorHandle& InActorReference)
-{
-    // @todo(mlentine): It looks like immediate mode doesn't support this
-    return false;
-}
-
-bool FPhysicsInterface_ImmediatePhysX::IsInScene(const FPhysicsActorHandle& InActorReference)
-{
-    return InActorReference.Scene != nullptr;
-}
-
-FPhysScene* FPhysicsInterface_ImmediatePhysX::GetCurrentScene(const FPhysicsActorHandle& InActorReference)
-{
-    return InActorReference.Scene;
-}
-
-bool FPhysicsInterface_ImmediatePhysX::CanSimulate_AssumesLocked(const FPhysicsActorHandle& InActorReference)
-{
-    return InActorReference.IsValid();
-}
-
-float FPhysicsInterface_ImmediatePhysX::GetMass_AssumesLocked(const FPhysicsActorHandle& InActorReference)
-{
-    return 1 / InActorReference.Scene->RigidBodiesData[InActorReference.Index].invMass;
-}
-
-void FPhysicsInterface_ImmediatePhysX::SetSendsSleepNotifies_AssumesLocked(const FPhysicsActorHandle& InActorReference, bool bSendSleepNotifies)
-{
-    // @todo(mlentine): Is there a way to "sleep" bodies in immediate mode?
-    check(false);
-}
-
-
-void FPhysicsInterface_ImmediatePhysX::PutToSleep_AssumesLocked(const FPhysicsActorHandle& InActorReference)
-{
-    // @todo(mlentine): Is there a way to "sleep" bodies in immediate mode?
-    check(false);
-}
-
-void FPhysicsInterface_ImmediatePhysX::WakeUp_AssumesLocked(const FPhysicsActorHandle& InActorReference)
-{
-    // @todo(mlentine): Is there a way to "sleep" bodies in immediate mode?
-    check(false);
-}
-
-void FPhysicsInterface_ImmediatePhysX::SetIsKinematic_AssumesLocked(const FPhysicsActorHandle& InActorReference, bool bIsKinematic)
-{
-    if (bIsKinematic)
-    {
-        int32 NewIndex = IsDynamic(InActorReference) ? --InActorReference.Scene->NumSimulatedBodies : InActorReference.Scene->NumSimulatedBodies + InActorReference.Scene->NumKinematicBodies;
-        InActorReference.Scene->SwapActorData(NewIndex, InActorReference.Index);
-        const_cast<FPhysicsActorHandle&>(InActorReference).Index = NewIndex;
-        InActorReference.Scene->NumKinematicBodies++;
-    }
-    else
-    {
-        // @todo(mlentine): We are assuming making it not kinematic means it is now dynamic
-        int32 NewIndex = InActorReference.Scene->NumSimulatedBodies++;
-        InActorReference.Scene->SwapActorData(NewIndex, InActorReference.Index);
-        const_cast<FPhysicsActorHandle&>(InActorReference).Index = NewIndex;
-        InActorReference.Scene->NumKinematicBodies--;
-    }
-}
-
-void FPhysicsInterface_ImmediatePhysX::SetCcdEnabled_AssumesLocked(const FPhysicsActorHandle& InActorReference, bool bIsCcdEnabled)
-{
-    check(bIsCcdEnabled == false);
-}
-
-FTransform FPhysicsInterface_ImmediatePhysX::GetGlobalPose_AssumesLocked(const FPhysicsActorHandle& InActorReference)
-{
-    return P2UTransform(InActorReference.Scene->RigidBodiesData[InActorReference.Index].body2World);
-}
-
-void FPhysicsInterface_ImmediatePhysX::SetGlobalPose_AssumesLocked(const FPhysicsActorHandle& InActorReference, const FTransform& InNewPose, bool bAutoWake)
-{
-    InActorReference.Scene->RigidBodiesData[InActorReference.Index].body2World = U2PTransform(InNewPose);
-}
-
-FTransform FPhysicsInterface_ImmediatePhysX::GetTransform_AssumesLocked(const FPhysicsActorHandle& InRef, bool bForceGlobalPose /*= false*/)
-{
-	if(!bForceGlobalPose)
-	{
-		if(IsDynamic(InRef))
-		{
-			if(HasKinematicTarget_AssumesLocked(InRef))
-			{
-				return GetKinematicTarget_AssumesLocked(InRef);
-			}
-		}
-	}
-
-	return GetGlobalPose_AssumesLocked(InRef);
-}
-
-bool FPhysicsInterface_ImmediatePhysX::HasKinematicTarget_AssumesLocked(const FPhysicsActorHandle& InActorReference)
-{
-    return !IsDynamic(InActorReference) && !IsStatic(InActorReference);
-}
-
-FTransform FPhysicsInterface_ImmediatePhysX::GetKinematicTarget_AssumesLocked(const FPhysicsActorHandle& InActorReference)
-{
-    return P2UTransform(InActorReference.Scene->KinematicTargets[InActorReference.Index].BodyToWorld);
-}
-
-void FPhysicsInterface_ImmediatePhysX::SetKinematicTarget_AssumesLocked(const FPhysicsActorHandle& InActorReference, const FTransform& InNewTarget)
-{
-    InActorReference.Scene->KinematicTargets[InActorReference.Index].BodyToWorld = U2PTransform(InNewTarget);
-}
-
-FVector FPhysicsInterface_ImmediatePhysX::GetLinearVelocity_AssumesLocked(const FPhysicsActorHandle& InActorReference)
-{
-    return P2UVector(InActorReference.Scene->RigidBodiesData[InActorReference.Index].linearVelocity);
-}
-
-void FPhysicsInterface_ImmediatePhysX::SetLinearVelocity_AssumesLocked(const FPhysicsActorHandle& InActorReference, const FVector& InNewVelocity, bool bAutoWake)
-{
-    InActorReference.Scene->RigidBodiesData[InActorReference.Index].linearVelocity = U2PVector(InNewVelocity);
-}
-
-FVector FPhysicsInterface_ImmediatePhysX::GetAngularVelocity_AssumesLocked(const FPhysicsActorHandle& InActorReference)
-{
-    return P2UVector(InActorReference.Scene->RigidBodiesData[InActorReference.Index].angularVelocity);
-}
-
-void FPhysicsInterface_ImmediatePhysX::SetAngularVelocity_AssumesLocked(const FPhysicsActorHandle& InActorReference, const FVector& InNewVelocity, bool bAutoWake)
-{
-    InActorReference.Scene->RigidBodiesData[InActorReference.Index].angularVelocity = U2PVector(InNewVelocity);
-}
-
-float FPhysicsInterface_ImmediatePhysX::GetMaxAngularVelocity_AssumesLocked(const FPhysicsActorHandle& InActorReference)
-{
-    return sqrt(InActorReference.Scene->RigidBodiesData[InActorReference.Index].maxAngularVelocitySq);
-}
-
-void FPhysicsInterface_ImmediatePhysX::SetMaxAngularVelocity_AssumesLocked(const FPhysicsActorHandle& InActorReference, float InMaxAngularVelocity)
-{
-    InActorReference.Scene->RigidBodiesData[InActorReference.Index].maxAngularVelocitySq = InMaxAngularVelocity * InMaxAngularVelocity;
-}
-
-float FPhysicsInterface_ImmediatePhysX::GetMaxDepenetrationVelocity_AssumesLocked(const FPhysicsActorHandle& InActorReference)
-{
-    return InActorReference.Scene->RigidBodiesData[InActorReference.Index].maxDepenetrationVelocity;
-}
-
-void FPhysicsInterface_ImmediatePhysX::SetMaxDepenetrationVelocity_AssumesLocked(const FPhysicsActorHandle& InActorReference, float InMaxDepenetrationVelocity)
-{
-    InActorReference.Scene->RigidBodiesData[InActorReference.Index].maxDepenetrationVelocity = InMaxDepenetrationVelocity;
-}
-
-FVector FPhysicsInterface_ImmediatePhysX::GetWorldVelocityAtPoint_AssumesLocked(const FPhysicsActorHandle& InActorReference, const FVector& InPoint)
-{
-    const auto& RigidBodyData = InActorReference.Scene->RigidBodiesData[InActorReference.Index];
-    return P2UVector(RigidBodyData.linearVelocity) + FVector::CrossProduct(P2UVector(RigidBodyData.angularVelocity), InPoint - P2UTransform(RigidBodyData.body2World).GetTranslation());
-}
-
-FTransform FPhysicsInterface_ImmediatePhysX::GetComTransform_AssumesLocked(const FPhysicsActorHandle& InActorReference)
-{
-    // @todo(mlentine): Need to get Com from Shape
-    check(false);
-	return FTransform::Identity;
-}
-
-FVector FPhysicsInterface_ImmediatePhysX::GetLocalInertiaTensor_AssumesLocked(const FPhysicsActorHandle& InActorReference)
-{
-    FVector InvInertia = P2UVector(InActorReference.Scene->RigidBodiesData[InActorReference.Index].invInertia);
-    return FVector(1.f / InvInertia.X, 1.f / InvInertia.Y, 1.f / InvInertia.Z);
-}
-
-FBox FPhysicsInterface_ImmediatePhysX::GetBounds_AssumesLocked(const FPhysicsActorHandle& InActorReference)
-{
-    // @todo(mlentine): Need to get Bounds from Shape
-    check(false);
-	return FBox();
-}
-
-void FPhysicsInterface_ImmediatePhysX::SetLinearDamping_AssumesLocked(const FPhysicsActorHandle& InActorReference, float InDamping)
-{
-    InActorReference.Scene->RigidBodiesData[InActorReference.Index].linearDamping = InDamping;
-}
-
-void FPhysicsInterface_ImmediatePhysX::SetAngularDamping_AssumesLocked(const FPhysicsActorHandle& InActorReference, float InDamping)
-{
-    InActorReference.Scene->RigidBodiesData[InActorReference.Index].angularDamping = InDamping;
-}
-
-void FPhysicsInterface_ImmediatePhysX::AddForce_AssumesLocked(const FPhysicsActorHandle& InActorReference, const FVector& InForce)
-{
-    check(IsDynamic(InActorReference));
-    InActorReference.Scene->PendingAcceleration[InActorReference.Index] += U2PVector(InForce * InActorReference.Scene->RigidBodiesData[InActorReference.Index].invMass);
-}
-
-void FPhysicsInterface_ImmediatePhysX::AddTorque_AssumesLocked(const FPhysicsActorHandle& InActorReference, const FVector& InTorque)
-{
-    check(IsDynamic(InActorReference));
-    InActorReference.Scene->PendingAngularAcceleration[InActorReference.Index] += U2PVector(InTorque * P2UVector(InActorReference.Scene->RigidBodiesData[InActorReference.Index].invInertia));
-}
-
-// @todo(mlentine): Rename this to Impulse
-void FPhysicsInterface_ImmediatePhysX::AddForceMassIndependent_AssumesLocked(const FPhysicsActorHandle& InActorReference, const FVector& InForce)
-{
-    check(IsDynamic(InActorReference));
-    InActorReference.Scene->PendingVelocityChange[InActorReference.Index] += U2PVector(InForce);
-}
-
-void FPhysicsInterface_ImmediatePhysX::AddTorqueMassIndependent_AssumesLocked(const FPhysicsActorHandle& InActorReference, const FVector& InTorque)
-{
-    check(IsDynamic(InActorReference));
-    InActorReference.Scene->PendingAngularAcceleration[InActorReference.Index] += U2PVector(InTorque);
-}
-
-void FPhysicsInterface_ImmediatePhysX::AddImpulseAtLocation_AssumesLocked(const FPhysicsActorHandle& InActorReference, const FVector& InImpulse, const FVector& InLocation)
-{
-    const auto& RigidBodyData = InActorReference.Scene->RigidBodiesData[InActorReference.Index];
-    InActorReference.Scene->PendingVelocityChange[InActorReference.Index] += U2PVector(InImpulse);
-    InActorReference.Scene->PendingAngularVelocityChange[InActorReference.Index] += U2PVector(FVector::CrossProduct(InImpulse, InLocation - P2UTransform(RigidBodyData.body2World).GetTranslation()));
-}
-
-void FPhysicsInterface_ImmediatePhysX::AddRadialImpulse_AssumesLocked(const FPhysicsActorHandle& InActorReference, const FVector& InOrigin, float InRadius, float InStrength, ERadialImpulseFalloff InFalloff, bool bInVelChange)
-{
-    const auto& RigidBodyData = InActorReference.Scene->RigidBodiesData[InActorReference.Index];
-    FVector Direction = (P2UTransform(RigidBodyData.body2World).GetTranslation() - InOrigin);
-    const float Distance = Direction.Size();
-    if (Distance > InRadius)
-    {
-        return;
-    }
-    Direction = Direction.GetSafeNormal();
-    FVector Force(0, 0, 0);
-    check(InFalloff == RIF_Constant || InFalloff == RIF_Linear);
-    if (InFalloff == RIF_Constant)
-    {
-        Force = InStrength * Direction;
-    }
-    if (InFalloff == RIF_Linear)
-    {
-        Force = (InRadius - Distance) / InRadius * InStrength * Direction;
-    }
-    if (bInVelChange)
-    {
-        InActorReference.Scene->PendingVelocityChange[InActorReference.Index] += U2PVector(Force);
-    }
-    else 
-    {
-        InActorReference.Scene->PendingAcceleration[InActorReference.Index] += U2PVector(Force);
-    }
-}
-
-bool FPhysicsInterface_ImmediatePhysX::IsGravityEnabled_AssumesLocked(const FPhysicsActorHandle& InActorReference)
-{
-	return true;
-}
-
-void FPhysicsInterface_ImmediatePhysX::SetGravityEnabled_AssumesLocked(const FPhysicsActorHandle& InActorReference, bool bEnabled)
-{
-    // @todo(mlentine): We do not currently support a way to turn off gravity
-    check(bEnabled);
-}
-
-float FPhysicsInterface_ImmediatePhysX::GetSleepEnergyThreshold_AssumesLocked(const FPhysicsActorHandle& InActorReference)
-{
-    // @todo(mlentine): How is sleeping supported in immediate mode?
-	return 0.0f;
-}
-
-void FPhysicsInterface_ImmediatePhysX::SetSleepEnergyThreshold_AssumesLocked(const FPhysicsActorHandle& InActorReference, float InEnergyThreshold)
-{
-    // @todo(mlentine): How is sleeping supported in immediate mode?
-    check(false);
-}
-
-void FPhysicsInterface_ImmediatePhysX::SetMass_AssumesLocked(const FPhysicsActorReference_ImmediatePhysX& InHandle, float InMass)
-{
-    InHandle.Scene->RigidBodiesData[InHandle.Index].invMass = 1.f / InMass;
-}
-
-void FPhysicsInterface_ImmediatePhysX::SetMassSpaceInertiaTensor_AssumesLocked(const FPhysicsActorReference_ImmediatePhysX& InHandle, const FVector& InTensor)
-{
-    InHandle.Scene->RigidBodiesData[InHandle.Index].invInertia = U2PVector(FVector(1.f / InTensor.X, 1.f / InTensor.Y, 1.f / InTensor.Z));
-}
-
-void FPhysicsInterface_ImmediatePhysX::SetComLocalPose_AssumesLocked(const FPhysicsActorReference_ImmediatePhysX& InHandle, const FTransform& InComLocalPose)
-{
-    // @todo(mlentine): Similar to Chaos this shouldn't be possible as this makes initia tensor non diagonal
-    check(false);
-}
-
-float FPhysicsInterface_ImmediatePhysX::GetStabilizationEnergyThreshold_AssumesLocked(const FPhysicsActorReference_ImmediatePhysX& InHandle)
-{
-	// #PHYS2 implement
-	return 0.0f;
-}
-
-void FPhysicsInterface_ImmediatePhysX::SetStabilizationEnergyThreshold_AssumesLocked(const FPhysicsActorReference_ImmediatePhysX& InHandle, float InThreshold)
-{
-	// #PHYS2 implement
-}
-
-uint32 FPhysicsInterface_ImmediatePhysX::GetSolverPositionIterationCount_AssumesLocked(const FPhysicsActorReference_ImmediatePhysX& InHandle)
-{
-	// #PHYS2 implement
-	return 0;
-}
-
-void FPhysicsInterface_ImmediatePhysX::SetSolverPositionIterationCount_AssumesLocked(const FPhysicsActorReference_ImmediatePhysX& InHandle, uint32 InSolverIterationCount)
-{
-	// #PHYS2 implement
-}
-
-uint32 FPhysicsInterface_ImmediatePhysX::GetSolverVelocityIterationCount_AssumesLocked(const FPhysicsActorReference_ImmediatePhysX& InHandle)
-{
-	// #PHYS2 implement
-	return 0;
-}
-
-void FPhysicsInterface_ImmediatePhysX::SetSolverVelocityIterationCount_AssumesLocked(const FPhysicsActorReference_ImmediatePhysX& InHandle, uint32 InSolverIterationCount)
-{
-	// #PHYS2 implement
-}
-
-float FPhysicsInterface_ImmediatePhysX::GetWakeCounter_AssumesLocked(const FPhysicsActorReference_ImmediatePhysX& InHandle)
-{
-	// #PHYS2 implement
-	return 0.0f;
-}
-
-void FPhysicsInterface_ImmediatePhysX::SetWakeCounter_AssumesLocked(const FPhysicsActorReference_ImmediatePhysX& InHandle, float InWakeCounter)
-{
-	// #PHYS2 implement
-}
-
-SIZE_T FPhysicsInterface_ImmediatePhysX::GetResourceSizeEx(const FPhysicsActorReference_ImmediatePhysX& InActorReference)
-{
-    // @todo(mlentine): What uses this and what does this need to be?
-    check(false);
-	return 0;
-}
-
-// Constraint free functions //////////////////////////////////////////////////////////////////////////
-
-const bool bDrivesUseAcceleration = true;
-
-bool GetSceneForConstraintActors_LockFree(const FPhysicsActorHandle& InActor1, const FPhysicsActorHandle& InActor2, FPhysScene* OutScene)
-{
-    if (InActor1.Scene && InActor2.Scene)
-    {
-        if (InActor1.Scene != InActor2.Scene)
-        {
-            return false;
-        }
-    }
-    OutScene = InActor1.Scene ? InActor1.Scene : InActor2.Scene;
-	return OutScene != nullptr;
-}
-
-void GetSoftLimitParams_Linear(float& InOutDamping, float& InOutStiffness)
-{
-	InOutDamping = InOutDamping * CVarConstraintLinearDampingScale.GetValueOnGameThread();
-	InOutStiffness = InOutStiffness * CVarConstraintLinearStiffnessScale.GetValueOnGameThread();
-}
-
-void GetSoftLimitParams_Angular(float& InOutDamping, float& InOutStiffness)
-{
-	InOutDamping = InOutDamping * CVarConstraintAngularDampingScale.GetValueOnGameThread();
-	InOutStiffness = InOutStiffness * CVarConstraintAngularStiffnessScale.GetValueOnGameThread();
-}
-
-// Constraint interface functions
-
-FPhysicsConstraintReference_ImmediatePhysX FPhysicsInterface_ImmediatePhysX::CreateConstraint(const FPhysicsActorHandle& InActorRef1, const FPhysicsActorHandle& InActorRef2, const FTransform& InLocalFrame1, const FTransform& InLocalFrame2)
-{
-	FPhysicsConstraintReference_ImmediatePhysX OutReference;
-    OutReference.Scene = InActorRef1.Scene;
-    OutReference.Index = InActorRef1.Scene->Joints.Num();
-    
-    OutReference.Scene->Joints.Add(FPhysScene_ImmediatePhysX::FJoint(InActorRef1.Index, InActorRef2.Index, InLocalFrame1, InLocalFrame2));
-	return OutReference;
-}
-
-void FPhysicsInterface_ImmediatePhysX::SetConstraintUserData(const FPhysicsConstraintHandle& InConstraintRef, void* InUserData)
-{
-    // @todo(mlentine): What do we use InUserData for?
-    check(false);
-}
-
-void FPhysicsInterface_ImmediatePhysX::ReleaseConstraint(FPhysicsConstraintHandle& InConstraintRef)
-{
-    // @todo(mlentine): I don't think we need to do anything here.
-    InConstraintRef.Scene->Joints.RemoveAtSwap(InConstraintRef.Index);
-}
-
-FTransform FPhysicsInterface_ImmediatePhysX::GetLocalPose(const FPhysicsConstraintHandle& InConstraintRef, EConstraintFrame::Type InFrame)
-{
-    if (InFrame == EConstraintFrame::Frame1)
-    {
-        return InConstraintRef.Scene->Joints[InConstraintRef.Index].JointToParent;
-    }
-    else if (InFrame == EConstraintFrame::Frame2)
-    {
-        return InConstraintRef.Scene->Joints[InConstraintRef.Index].JointToChild;
-    }
-	return FTransform::Identity;
-}
-
-FTransform FPhysicsInterface_ImmediatePhysX::GetGlobalPose(const FPhysicsConstraintHandle& InConstraintRef, EConstraintFrame::Type InFrame)
-{
-    if (InFrame == EConstraintFrame::Frame1)
-    {
-        return P2UTransform(InConstraintRef.Scene->RigidBodiesData[InConstraintRef.Scene->Joints[InConstraintRef.Index].ParentIndex].body2World);
-    }
-    else if (InFrame == EConstraintFrame::Frame2)
-    {
-        return P2UTransform(InConstraintRef.Scene->RigidBodiesData[InConstraintRef.Scene->Joints[InConstraintRef.Index].ChildIndex].body2World);
-    }
-	return FTransform::Identity;
-}
-
-FVector FPhysicsInterface_ImmediatePhysX::GetLocation(const FPhysicsConstraintHandle& InConstraintRef)
-{
-    return 0.5 * (GetGlobalPose(InConstraintRef, EConstraintFrame::Frame1).GetTranslation() + GetGlobalPose(InConstraintRef, EConstraintFrame::Frame2).GetTranslation());
-}
-
-void FPhysicsInterface_ImmediatePhysX::GetForce(const FPhysicsConstraintHandle& InConstraintRef, FVector& OutLinForce, FVector& OutAngForce)
-{
-    // @todo(mlentine): How do I get this from immediate mode?
-    check(false);
-}
-
-void FPhysicsInterface_ImmediatePhysX::GetDriveLinearVelocity(const FPhysicsConstraintHandle& InConstraintRef, FVector& OutLinVelocity)
-{
-    // @todo(mlentine): How do I get this from immediate mode?
-    check(false);
-}
-
-void FPhysicsInterface_ImmediatePhysX::GetDriveAngularVelocity(const FPhysicsConstraintHandle& InConstraintRef, FVector& OutAngVelocity)
-{
-    // @todo(mlentine): How do I get this from immediate mode?
-    check(false);
-}
-
-float FPhysicsInterface_ImmediatePhysX::GetCurrentSwing1(const FPhysicsConstraintHandle& InConstraintRef)
-{
-    // @todo(mlentine): How do I get this from immediate mode?
-    check(false);
-    return 0;
-}
-
-float FPhysicsInterface_ImmediatePhysX::GetCurrentSwing2(const FPhysicsConstraintHandle& InConstraintRef)
-{
-    // @todo(mlentine): How do I get this from immediate mode?
-    check(false);
-    return 0;
-}
-
-float FPhysicsInterface_ImmediatePhysX::GetCurrentTwist(const FPhysicsConstraintHandle& InConstraintRef)
-{
-    // @todo(mlentine): How do I get this from immediate mode?
-    check(false);
-    return 0;
-}
-
-void FPhysicsInterface_ImmediatePhysX::SetCanVisualize(const FPhysicsConstraintHandle& InConstraintRef, bool bInCanVisualize)
-{
-    // @todo(mlentine): Can we enable visualization from immediate mode?
-    check(bInCanVisualize == false);
-}
-
-void FPhysicsInterface_ImmediatePhysX::SetCollisionEnabled(const FPhysicsConstraintHandle& InConstraintRef, bool bInCollisionEnabled)
-{
-    // @todo(mlentine): Allow collisions to be disabled
-    check(bInCollisionEnabled == true);
-}
-
-void FPhysicsInterface_ImmediatePhysX::SetProjectionEnabled_AssumesLocked(const FPhysicsConstraintHandle& InConstraintRef, bool bInProjectionEnabled, float InLinearTolerance, float InAngularToleranceDegrees)
-{
-    // @todo(mlentine): How do we set this from immediate mode?
-    check(false);
-}
-
-void FPhysicsInterface_ImmediatePhysX::SetParentDominates_AssumesLocked(const FPhysicsConstraintHandle& InConstraintRef, bool bInParentDominates)
-{
-    InConstraintRef.Scene->RigidBodiesData[InConstraintRef.Scene->Joints[InConstraintRef.Index].ParentIndex].invMass = bInParentDominates ? 0.f : 1.f;
-    // @todo(mlentine): We will have to save the origional inertia somehow as the physx type doesn't have a scale. Right now we just treat it as a sphere.
-    InConstraintRef.Scene->RigidBodiesData[InConstraintRef.Scene->Joints[InConstraintRef.Index].ParentIndex].invInertia = U2PVector(bInParentDominates ? FVector(0, 0, 0) : FVector(1, 1, 1));
-}
-
-void FPhysicsInterface_ImmediatePhysX::SetBreakForces_AssumesLocked(const FPhysicsConstraintHandle& InConstraintRef, float InLinearBreakForce, float InAngularBreakForce)
-{
-    // @todo(mlentine): Need to store force and run something like this in order to determine if it breaks.
-    /*
-    FVector LinearForce, AngularForce;
-    GetForce(InConstraintRef, LinearForce, AngularForce);
-    if (InLinearBreakForce * InLinearBreakForce > LinearForce.SizeSquared() || InAngularBreakForce * InAngularBreakForce > AngularForce.SizeSquared())
-    {
-        ReleaseConstraint(InConstraintRef);
-    }
-    */
-}
-
-void FPhysicsInterface_ImmediatePhysX::SetPlasticityLimits_AssumesLocked(const FPhysicsConstraintHandle& InConstraintRef, float InLinearBreakForce, float InAngularBreakTorque)
-{
-	// @todo(chaos) : Implement
-}
-
-
-void FPhysicsInterface_ImmediatePhysX::SetLocalPose(const FPhysicsConstraintHandle& InConstraintRef, const FTransform& InPose, EConstraintFrame::Type InFrame)
-{
-    if (InFrame == EConstraintFrame::Frame1)
-    {
-        InConstraintRef.Scene->Joints[InConstraintRef.Index].JointToParent = InPose;
-    }
-    if (InFrame == EConstraintFrame::Frame2)
-    {
-        InConstraintRef.Scene->Joints[InConstraintRef.Index].JointToChild = InPose;
-    }
-}
-
-void FPhysicsInterface_ImmediatePhysX::SetLinearMotionLimitType_AssumesLocked(const FPhysicsConstraintHandle& InConstraintRef, PhysicsInterfaceTypes::ELimitAxis InAxis, ELinearConstraintMotion InMotion)
-{
-    // @todo(mlentine): How to do this in immediate mode
-    check(false);
-}
-
-void FPhysicsInterface_ImmediatePhysX::SetAngularMotionLimitType_AssumesLocked(const FPhysicsConstraintHandle& InConstraintRef, PhysicsInterfaceTypes::ELimitAxis InAxis, EAngularConstraintMotion InMotion)
-{
-    // @todo(mlentine): How to do this in immediate mode
-    check(false);
-}
-
-void FPhysicsInterface_ImmediatePhysX::UpdateLinearLimitParams_AssumesLocked(const FPhysicsConstraintHandle& InConstraintRef, float InLimit, float InAverageMass, const FLinearConstraint& InParams)
-{
-    // @todo(mlentine): How to do this in immediate mode
-    check(false);
-}
-
-void FPhysicsInterface_ImmediatePhysX::UpdateConeLimitParams_AssumesLocked(const FPhysicsConstraintHandle& InConstraintRef, float InAverageMass, const FConeConstraint& InParams)
-{
-    // @todo(mlentine): How to do this in immediate mode
-    check(false);
-}
-
-void FPhysicsInterface_ImmediatePhysX::UpdateTwistLimitParams_AssumesLocked(const FPhysicsConstraintHandle& InConstraintRef, float InAverageMass, const FTwistConstraint& InParams)
-{
-    // @todo(mlentine): How to do this in immediate mode
-    check(false);
-}
-
-void FPhysicsInterface_ImmediatePhysX::UpdateLinearDrive_AssumesLocked(const FPhysicsConstraintHandle& InConstraintRef, const FLinearDriveConstraint& InDriveParams)
-{
-    // @todo(mlentine): How to do this in immediate mode
-    check(false);
-}
-
-void FPhysicsInterface_ImmediatePhysX::UpdateAngularDrive_AssumesLocked(const FPhysicsConstraintHandle& InConstraintRef, const FAngularDriveConstraint& InDriveParams)
-{
-    // @todo(mlentine): How to do this in immediate mode
-    check(false);
-}
-
-void FPhysicsInterface_ImmediatePhysX::UpdateDriveTarget_AssumesLocked(const FPhysicsConstraintHandle& InConstraintRef, const FLinearDriveConstraint& InLinDrive, const FAngularDriveConstraint& InAngDrive)
-{
-    // @todo(mlentine): How to do this in immediate mode
-    check(false);
-}
-
-void FPhysicsInterface_ImmediatePhysX::SetDrivePosition(const FPhysicsConstraintHandle& InConstraintRef, const FVector& InPosition)
-{
-    // @todo(mlentine): How to do this in immediate mode
-    check(false);
-}
-
-void FPhysicsInterface_ImmediatePhysX::SetDriveOrientation(const FPhysicsConstraintHandle& InConstraintRef, const FQuat& InOrientation)
-{
-    // @todo(mlentine): How to do this in immediate mode
-    check(false);
-}
-
-void FPhysicsInterface_ImmediatePhysX::SetDriveLinearVelocity(const FPhysicsConstraintHandle& InConstraintRef, const FVector& InLinVelocity)
-{
-    // @todo(mlentine): How to do this in immediate mode
-    check(false);
-}
-
-void FPhysicsInterface_ImmediatePhysX::SetDriveAngularVelocity(const FPhysicsConstraintHandle& InConstraintRef, const FVector& InAngVelocity)
-{
-    // @todo(mlentine): How to do this in immediate mode
-    check(false);
-}
-
-void FPhysicsInterface_ImmediatePhysX::SetTwistLimit(const FPhysicsConstraintHandle& InConstraintRef, float InLowerLimit, float InUpperLimit, float InContactDistance)
-{
-    // @todo(mlentine): How to do this in immediate mode
-    check(false);
-}
-
-void FPhysicsInterface_ImmediatePhysX::SetSwingLimit(const FPhysicsConstraintHandle& InConstraintRef, float InYLimit, float InZLimit, float InContactDistance)
-{
-    // @todo(mlentine): How to do this in immediate mode
-    check(false);
-}
-
-void FPhysicsInterface_ImmediatePhysX::SetLinearLimit(const FPhysicsConstraintHandle& InConstraintRef, float InLimit)
-{
-    // @todo(mlentine): How to do this in immediate mode
-    check(false);
-}
-
-bool FPhysicsInterface_ImmediatePhysX::IsBroken(const FPhysicsConstraintHandle& InConstraintRef)
-{
-    return !InConstraintRef.IsValid();
-}
-
-bool FPhysicsInterface_ImmediatePhysX::ExecuteOnUnbrokenConstraintReadOnly(const FPhysicsConstraintHandle& InConstraintRef, TFunctionRef<void(const FPhysicsConstraintHandle&)> Func)
-{
-    if (!IsBroken(InConstraintRef))
-    {
-        Func(InConstraintRef);
-        return true;
-    }
-	return false;
-}
-
-bool FPhysicsInterface_ImmediatePhysX::ExecuteOnUnbrokenConstraintReadWrite(const FPhysicsConstraintHandle& InConstraintRef, TFunctionRef<void(const FPhysicsConstraintHandle&)> Func)
-{
-    if (!IsBroken(InConstraintRef))
-    {
-        Func(InConstraintRef);
-        return true;
-    }
-	return false;
-}
-
-void FinishSceneStat(uint32 Scene)
-{
-}
-
-void FPhysicsInterface_ImmediatePhysX::CalculateMassPropertiesFromShapeCollection(physx::PxMassProperties& OutProperties, const TArray<FPhysicsShapeHandle>& InShapes, float InDensityKGPerCM)
-{
-    // What does it mean when if there is more than one collision object?
-    check(InShapes.Num() == 1);
-    if (InShapes[0].Actor->IsValid())
-    {
-        const auto& Data = InShapes[0].Actor->Scene->RigidBodiesData[InShapes[0].Actor->Index];
-        OutProperties.centerOfMass = Data.body2World.p;
-        OutProperties.inertiaTensor = physx::PxMat33();
-        OutProperties.inertiaTensor(0, 0) = 1.f / Data.invInertia.x;
-        OutProperties.inertiaTensor(1, 1) = 1.f / Data.invInertia.y;
-        OutProperties.inertiaTensor(2, 2) = 1.f / Data.invInertia.z;
-        OutProperties.mass = 1.f / Data.invMass;
-    }
-}
-
-bool FPhysicsInterface_ImmediatePhysX::LineTrace_Geom(FHitResult& OutHit, const FBodyInstance* InInstance, const FVector& InStart, const FVector& InEnd, bool bTraceComplex, bool bExtractPhysMaterial)
-{
-	check(InInstance);
-
-	OutHit.TraceStart = InStart;
-	OutHit.TraceEnd = InEnd;
-
-	bool bHitSomething = false;
-
-	const FVector Delta = InEnd - InStart;
-	const float DeltaMag = Delta.Size();
-	if(DeltaMag > KINDA_SMALL_NUMBER)
-	{
-		{
-			// #PHYS2 Really need a concept for "multi" locks here - as we're locking ActorRef but not TargetInstance->ActorRef
-			FPhysicsCommand::ExecuteRead(InInstance->ActorHandle, [&](const FPhysicsActorHandle& Actor)
-			{
-				// If we're welded then the target instance is actually our parent
-				const FBodyInstance* TargetInstance = InInstance->WeldParent ? InInstance->WeldParent : InInstance;
-
-				if(Actor.IsValid())
-				{
-					// Create filter data used to filter collisions, should always return eTOUCH for LineTraceComponent		
-					PxHitFlags PHitFlags = PxHitFlag::ePOSITION | PxHitFlag::eNORMAL | PxHitFlag::eDISTANCE | PxHitFlag::eFACE_INDEX;
-
-					PxRaycastHit BestHit;
-					float BestHitDistance = BIG_NUMBER;
-
-					// Get all the shapes from the actor
-					FInlineShapeArray PShapes;
-					const int32 NumShapes = FillInlineShapeArray_AssumesLocked(PShapes, Actor);
-
-					// Iterate over each shape
-					for(int32 ShapeIdx = 0; ShapeIdx < NumShapes; ShapeIdx++)
-					{
-						FPhysicsShapeHandle& ShapeRef = PShapes[ShapeIdx];
-
-						const PxU32 HitBufferSize = 1;
-						PxRaycastHit PHits[HitBufferSize];
-
-						// Filter so we trace against the right kind of collision
-						const bool bShapeIsComplex = ShapeRef.Shape.Geometry->getType() == PxGeometryType::Enum::eTRIANGLEMESH;
-						if((bTraceComplex && bShapeIsComplex) || (!bTraceComplex && !bShapeIsComplex))
-						{
-							const int32 ArraySize = UE_ARRAY_COUNT(PHits);
-							// #PHYS2 This may not work with shared shapes (GetTransform requires getActor to return non-nullptr) verify
-							PxTransform ShapeTransform = ShapeRef.Shape.LocalTM;
-							const PxI32 NumHits = PxGeometryQuery::raycast(U2PVector(InStart), U2PVector(Delta / DeltaMag), *ShapeRef.Shape.Geometry, ShapeTransform, DeltaMag, PHitFlags, ArraySize, PHits);
-
-							if(ensure(NumHits <= ArraySize))
-							{
-								for(int HitIndex = 0; HitIndex < NumHits; HitIndex++)
-								{
-									PxRaycastHit& Hit = PHits[HitIndex];
-									if(Hit.distance < BestHitDistance)
-									{
-										BestHitDistance = PHits[HitIndex].distance;
-										BestHit = PHits[HitIndex];
-									}
-								}
-							}
-						}
-					}
-
-					if(BestHitDistance < BIG_NUMBER)
-					{
-						// we just like to make sure if the hit is made, set to test touch
-						PxFilterData QueryFilter;
-						QueryFilter.word2 = 0xFFFFF;
-
-						PxTransform PStartTM(U2PVector(InStart));
-						const UPrimitiveComponent* OwnerComponentInst = InInstance->OwnerComponent.Get();
-						ConvertQueryImpactHit(OwnerComponentInst ? OwnerComponentInst->GetWorld() : nullptr, BestHit, OutHit, DeltaMag, QueryFilter, InStart, InEnd, nullptr, PStartTM, true, bExtractPhysMaterial);
-						bHitSomething = true;
-					}
-				}
-			});
-		}
-	}
-
-	return bHitSomething;
-}
-
-bool FPhysicsInterface_ImmediatePhysX::Sweep_Geom(FHitResult& OutHit, const FBodyInstance* InInstance, const FVector& InStart, const FVector& InEnd, const FQuat& InShapeRotation, const FCollisionShape& InShape, bool bSweepComplex)
-{
-	bool bSweepHit = false;
-
-	if(InShape.IsNearlyZero())
-	{
-		bSweepHit = LineTrace_Geom(OutHit, InInstance, InStart, InEnd, bSweepComplex);
-	}
-	else
-	{
-		OutHit.TraceStart = InStart;
-		OutHit.TraceEnd = InEnd;
-
-		const FBodyInstance* TargetInstance = InInstance->WeldParent ? InInstance->WeldParent : InInstance;
-
-		FPhysicsCommand::ExecuteRead(TargetInstance->ActorHandle, [&](const FPhysicsActorHandle& Actor)
-		{
-			if(Actor.IsValid() && InInstance->OwnerComponent != nullptr)
-			{
-				FPhysXShapeAdapter ShapeAdaptor(InShapeRotation, InShape);
-				
-				const FVector Delta = InEnd - InStart;
-				const float DeltaMag = Delta.Size();
-				if(DeltaMag > KINDA_SMALL_NUMBER)
-				{
-					PxHitFlags POutputFlags = PxHitFlag::ePOSITION | PxHitFlag::eNORMAL | PxHitFlag::eDISTANCE | PxHitFlag::eFACE_INDEX | PxHitFlag::eMTD;
-
-					UPrimitiveComponent* OwnerComponentInst = InInstance->OwnerComponent.Get();
-					PxTransform PStartTM(U2PVector(InStart), ShapeAdaptor.GetGeomOrientation());
-					PxTransform PCompTM(U2PTransform(OwnerComponentInst->GetComponentTransform()));
-
-					PxVec3 PDir = U2PVector(Delta / DeltaMag);
-
-					PxSweepHit PHit;
-
-					// Get all the shapes from the actor
-					FInlineShapeArray PShapes;
-					// #PHYS2 - SHAPES - Resolve this function to not use px stuff
-					const int32 NumShapes = FillInlineShapeArray_AssumesLocked(PShapes, Actor); // #PHYS2 - Need a lock/execute here?
-
-					// Iterate over each shape
-					for(int32 ShapeIdx = 0; ShapeIdx < NumShapes; ShapeIdx++)
-					{
-						FPhysicsShapeHandle& ShapeRef = PShapes[ShapeIdx];
-
-						// Filter so we trace against the right kind of collision
-						const bool bShapeIsComplex = ShapeRef.Shape.Geometry->getType() == PxGeometryType::Enum::eTRIANGLEMESH;
-						if((bSweepComplex && bShapeIsComplex) || (!bSweepComplex && !bShapeIsComplex))
-						{
-							PxTransform PGlobalPose = PCompTM.transform(ShapeRef.Shape.LocalTM);
-							const PxGeometry& Geometry = ShapeAdaptor.GetGeometry();
-							if(PxGeometryQuery::sweep(PDir, DeltaMag, Geometry, PStartTM, *ShapeRef.Shape.Geometry, PGlobalPose, PHit, POutputFlags))
-							{
-								// we just like to make sure if the hit is made
-								PxFilterData QueryFilter;
-								QueryFilter.word2 = 0xFFFFF;
-
-								PxTransform PStartTransform(U2PVector(InStart));
-								PHit.faceIndex = FindFaceIndex(PHit, PDir);
-								ConvertQueryImpactHit(OwnerComponentInst->GetWorld(), PHit, OutHit, DeltaMag, QueryFilter, InStart, InEnd, nullptr, PStartTransform, false, false);
-								bSweepHit = true;
-							}
-						}
-					}
-				}
-			}
-		});
-	}
-
-	return bSweepHit;
-}
-
-bool Overlap_GeomInternal(const FBodyInstance* InInstance, const PxGeometry& InPxGeom, const FTransform& InShapeTransform, FMTDResult* OutOptResult)
-{
-	PxTransform ShapePose = U2PTransform(InShapeTransform);
-	const FBodyInstance* TargetInstance = InInstance->WeldParent ? InInstance->WeldParent : InInstance;
-
-	// Get all the shapes from the actor
-	FInlineShapeArray PShapes;
-	const int32 NumShapes = FillInlineShapeArray_AssumesLocked(PShapes, TargetInstance->ActorHandle);
-
-	// Iterate over each shape
-	for(int32 ShapeIdx = 0; ShapeIdx < NumShapes; ++ShapeIdx)
-	{
-		FPhysicsShapeHandle& ShapeRef = PShapes[ShapeIdx];
-
-		if(TargetInstance->IsShapeBoundToBody(ShapeRef))
-		{
-			PxVec3 POutDirection;
-			float OutDistance;
-
-			PxTransform PTransform = U2PTransform(FPhysicsInterface::GetTransform_AssumesLocked(TargetInstance->ActorHandle) * P2UTransform(ShapeRef.Shape.LocalTM));
-			if(OutOptResult)
-			{
-				if(PxGeometryQuery::computePenetration(POutDirection, OutDistance, InPxGeom, ShapePose, *ShapeRef.Shape.Geometry, PTransform))
-				{
-					//TODO: there are some edge cases that give us nan results. In these cases we skip
-					if(!POutDirection.isFinite())
-					{
-						POutDirection.x = 0.f;
-						POutDirection.y = 0.f;
-						POutDirection.z = 0.f;
-					}
-
-					OutOptResult->Direction = P2UVector(POutDirection);
-					OutOptResult->Distance = FMath::Abs(OutDistance);
-
-					return true;
-				}
-			}
-			else
-			{
-				if(PxGeometryQuery::overlap(InPxGeom, ShapePose, *ShapeRef.Shape.Geometry, PTransform))
-				{
-					return true;
-				}
-			}
-
-		}
-	}
-
-	return false;
-}
-
-bool FPhysicsInterface_ImmediatePhysX::Overlap_Geom(const FBodyInstance* InBodyInstance, const FPhysicsGeometryCollection& InGeometry, const FTransform& InShapeTransform, FMTDResult* OutOptResult)
-{
-	PxGeometry& PGeom = *InGeometry.Geometry;
-
-	return Overlap_GeomInternal(InBodyInstance, PGeom, InShapeTransform, OutOptResult);
-}
-
-bool FPhysicsInterface_ImmediatePhysX::Overlap_Geom(const FBodyInstance* InBodyInstance, const FCollisionShape& InCollisionShape, const FQuat& InShapeRotation, const FTransform& InShapeTransform, FMTDResult* OutOptResult)
-{
-	FPhysXShapeAdapter Adaptor(InShapeRotation, InCollisionShape);
-
-	return Overlap_GeomInternal(InBodyInstance, Adaptor.GetGeometry(), InShapeTransform, OutOptResult);
-}
-
-bool FPhysicsInterface_ImmediatePhysX::GetSquaredDistanceToBody(const FBodyInstance* InInstance, const FVector& InPoint, float& OutDistanceSquared, FVector* OutOptPointOnBody)
-{
-	if(OutOptPointOnBody)
-	{
-		*OutOptPointOnBody = InPoint;
-	}
-
-	float ReturnDistance = -1.f;
-	float MinDistanceSqr = BIG_NUMBER;
-	bool bFoundValidBody = false;
-	bool bEarlyOut = true;
-
-	const FBodyInstance* UseBI = InInstance->WeldParent ? InInstance->WeldParent : InInstance;
-
-	FPhysicsCommand::ExecuteRead(UseBI->ActorHandle, [&](const FPhysicsActorHandle& Actor)
-	{
-		bool bSyncData = FPhysicsInterface::HasSyncSceneData(Actor);
-		int32 NumSyncShapes = 0;
-		int32 NumAsyncShapes = 0;
-		FPhysicsInterface::GetNumShapes(Actor, NumSyncShapes, NumAsyncShapes);
-
-		int32 NumShapes = bSyncData ? NumSyncShapes : NumAsyncShapes;
-
-		if(NumShapes == 0 || UseBI->OwnerComponent == NULL)
-		{
-			return;
-		}
-
-		bEarlyOut = false;
-
-		// Get all the shapes from the actor
-		FInlineShapeArray PShapes;
-		const int32 NumTotalShapes = FillInlineShapeArray_AssumesLocked(PShapes, Actor);
-
-		const PxVec3 PPoint = U2PVector(InPoint);
-
-		// Iterate over each shape
-		for(int32 ShapeIdx = 0; ShapeIdx < NumTotalShapes; ShapeIdx++)
-		{
-			// #PHYS2 - resolve px stuff here
-			FPhysicsShapeHandle& ShapeRef = PShapes[ShapeIdx];
-
-			FPhysicsGeometryCollection GeoCollection = FPhysicsInterface::GetGeometryCollection(ShapeRef);
-
-			PxTransform PGlobalPose = U2PTransform(FPhysicsInterface::GetTransform_AssumesLocked(Actor) * P2UTransform(ShapeRef.Shape.LocalTM));
-
-			ECollisionShapeType GeomType = FPhysicsInterface::GetShapeType(ShapeRef);
-
-			if(GeomType == ECollisionShapeType::Trimesh)
-			{
-				// Type unsupported for this function, but some other shapes will probably work. 
-				continue;
-			}
-
-			bFoundValidBody = true;
-
-			PxVec3 PClosestPoint;
-			float SqrDistance = PxGeometryQuery::pointDistance(PPoint, GeoCollection.GetGeometry(), PGlobalPose, &PClosestPoint);
-			// distance has valid data and smaller than mindistance
-			if(SqrDistance > 0.f && MinDistanceSqr > SqrDistance)
-			{
-				MinDistanceSqr = SqrDistance;
-
-				if(OutOptPointOnBody)
-				{
-					*OutOptPointOnBody = P2UVector(PClosestPoint);
-				}
-			}
-			else if(SqrDistance == 0.f)
-			{
-				MinDistanceSqr = 0.f;
-				break;
-			}
-		}
-	});
-
-	if(!bFoundValidBody && !bEarlyOut)
-	{
-		UE_LOG(LogPhysics, Verbose, TEXT("GetDistanceToBody: Component (%s) has no simple collision and cannot be queried for closest point."), InInstance->OwnerComponent.Get() ? *(InInstance->OwnerComponent->GetPathName()) : TEXT("NONE"));
-	}
-
-	if(bFoundValidBody)
-	{
-		OutDistanceSquared = MinDistanceSqr;
-	}
-	return bFoundValidBody;
-}
-
-FPhysicsShapeReference_ImmediatePhysX::FPhysicsShapeReference_ImmediatePhysX(FPhysScene_ImmediatePhysX::FShape&& InShape)
-	: Shape(MoveTemp(InShape)), Actor(nullptr)
-{
-
-}
-
-bool FPhysicsShapeReference_ImmediatePhysX::IsValid() const
-{
-	return Actor != nullptr;
-}
-
-#undef LOCTEXT_NAMESPACE
-
-#endif
\ No newline at end of file
diff --git a/Engine/Source/Runtime/Engine/Private/PhysicsEngine/ImmediatePhysics/ImmediatePhysicsChaos/ImmediatePhysicsActorHandle_Chaos.cpp b/Engine/Source/Runtime/Engine/Private/PhysicsEngine/ImmediatePhysics/ImmediatePhysicsChaos/ImmediatePhysicsActorHandle_Chaos.cpp
--- a/Engine/Source/Runtime/Engine/Private/PhysicsEngine/ImmediatePhysics/ImmediatePhysicsChaos/ImmediatePhysicsActorHandle_Chaos.cpp
+++ b/Engine/Source/Runtime/Engine/Private/PhysicsEngine/ImmediatePhysics/ImmediatePhysicsChaos/ImmediatePhysicsActorHandle_Chaos.cpp
@@ -53,7 +53,6 @@ namespace ImmediatePhysics_Chaos
 		return true;
 	}
 
-#if WITH_CHAOS
 	TUniquePtr<Chaos::FImplicitObject> CloneGeometry(const Chaos::FImplicitObject* Geom, TArray<TUniquePtr<Chaos::FPerShapeData>>& OutShapes)
 	{
 		using namespace Chaos;
@@ -88,12 +87,10 @@ namespace ImmediatePhysics_Chaos
 
 		return nullptr;
 	}
-#endif
 
 	// Intended for use with Tri Mesh and Heightfields when cloning world simulation objects into the immediate scene
 	bool CloneGeometry(FBodyInstance* BodyInstance, EActorType ActorType, const FVector& Scale, FReal& OutMass, Chaos::FVec3& OutInertia, Chaos::FRigidTransform3& OutCoMTransform, TUniquePtr<Chaos::FImplicitObject>& OutGeom, TArray<TUniquePtr<Chaos::FPerShapeData>>& OutShapes)
 	{
-#if WITH_CHAOS
 		// We should only get non-simulated objects through this path, but you never know...
 		if ((BodyInstance != nullptr) && !BodyInstance->bSimulatePhysics && BodyInstance->ActorHandle)
 		{
@@ -106,7 +103,6 @@ namespace ImmediatePhysics_Chaos
 				return true;
 			}
 		}
-#endif
 
 		return CreateDefaultGeometry(Scale, OutMass, OutInertia, OutCoMTransform, OutGeom, OutShapes);
 	}
@@ -152,12 +148,7 @@ namespace ImmediatePhysics_Chaos
 #endif
 		AddParams.WorldTransform = BodyInstance->GetUnrealWorldTransform();
 		AddParams.Geometry = &BodySetup->AggGeom;
-#if PHYSICS_INTERFACE_PHYSX
-		AddParams.TriMeshes = TArrayView<PxTriangleMesh*>(BodySetup->TriMeshes);
-#endif
-#if WITH_CHAOS
 		AddParams.ChaosTriMeshes = MakeArrayView(BodySetup->ChaosTriMeshes);
-#endif
 
 		TArray<TUniquePtr<FImplicitObject>> Geoms;
 		FShapesArray Shapes;
@@ -168,7 +159,6 @@ namespace ImmediatePhysics_Chaos
 			return false;
 		}
 
-#if WITH_CHAOS && !PHYSICS_INTERFACE_PHYSX
 		if (ActorType == EActorType::DynamicActor)
 		{
 			// Whether each shape contributes to mass
@@ -191,11 +181,6 @@ namespace ImmediatePhysics_Chaos
 			OutInertia = MassProperties.InertiaTensor.GetDiagonal();
 			OutCoMTransform = FTransform(MassProperties.RotationOfMass, MassProperties.CenterOfMass);
 		}
-#else
-		OutMass = BodyInstance->GetBodyMass();
-		OutInertia = BodyInstance->GetBodyInertiaTensor();
-		OutCoMTransform = BodyInstance->GetMassSpaceLocal();
-#endif
 
 		// If we have multiple root shapes, wrap them in a union
 		if (Geoms.Num() == 1)
diff --git a/Engine/Source/Runtime/Engine/Private/PhysicsEngine/ImmediatePhysics/ImmediatePhysicsPhysX/ImmediatePhysicsActor_Physx.cpp b/Engine/Source/Runtime/Engine/Private/PhysicsEngine/ImmediatePhysics/ImmediatePhysicsPhysX/ImmediatePhysicsActor_Physx.cpp
--- a/Engine/Source/Runtime/Engine/Private/PhysicsEngine/ImmediatePhysics/ImmediatePhysicsPhysX/ImmediatePhysicsActor_Physx.cpp
+++ b/Engine/Source/Runtime/Engine/Private/PhysicsEngine/ImmediatePhysics/ImmediatePhysicsPhysX/ImmediatePhysicsActor_Physx.cpp
@@ -1,137 +1,3 @@
 // Copyright Epic Games, Inc. All Rights Reserved.
 
 #include "Physics/ImmediatePhysics/ImmediatePhysicsPhysX/ImmediatePhysicsActor_PhysX.h"
-
-#if WITH_PHYSX && PHYSICS_INTERFACE_PHYSX
-
-#include "PhysicsPublic.h"
-#include "PhysicsEngine/BodySetup.h"
-#include "PhysicalMaterials/PhysicalMaterial.h"
-
-namespace ImmediatePhysics_PhysX
-{
-
-	FMaterial* GetBaseMaterialFromUnrealMaterial(UPhysicalMaterial* InMaterial)
-	{
-		return &FMaterial::Default;
-	}
-
-	FActor::~FActor()
-	{
-		TerminateGeometry();
-	}
-
-	void FActor::CreateGeometry(PxRigidActor* RigidActor, const PxTransform& ActorToBodyTM)
-	{
-		const uint32 NumShapes = RigidActor->getNbShapes();
-		TArray<PxShape*> ActorShapes;
-		ActorShapes.AddUninitialized(NumShapes);
-		RigidActor->getShapes(ActorShapes.GetData(), sizeof(PxShape*)*NumShapes);
-		PxTransform BodyToActorTM = ActorToBodyTM.getInverse();
-
-		Shapes.Empty(NumShapes);
-		TArray<PxMaterial*, TInlineAllocator<1>> Materials;
-
-		for (PxShape* Shape : ActorShapes)
-		{
-			if (!(Shape->getFlags() & PxShapeFlag::eSIMULATION_SHAPE))
-			{
-				continue;
-			}
-
-			const PxTransform LocalPose = Shape->getLocalPose();
-			const PxTransform BodyLocalShape = BodyToActorTM * Shape->getLocalPose();
-			const PxGeometryHolder& GeomHolder = Shape->getGeometry();
-			const PxBounds3 Bounds = PxGeometryQuery::getWorldBounds(GeomHolder.any(), PxTransform(PxIdentity), /*inflation=*/1.f);
-			const float BoundsMag = Bounds.getExtents().magnitude();
-			const PxVec3 BoundsCenter = Bounds.getCenter();
-			int32 NumMaterials = Shape->getNbMaterials();
-
-			Materials.SetNumUninitialized(NumMaterials);
-			Shape->getMaterials(Materials.GetData(), sizeof(Materials[0]) * NumMaterials);
-
-			FMaterial NewMaterial;
-			if (NumMaterials > 0)
-			{
-				NewMaterial = FMaterial(Materials[0]);
-			}
-
-			switch (GeomHolder.getType())
-			{
-			case PxGeometryType::eSPHERE:		Shapes.Emplace(BodyLocalShape, BoundsCenter, BoundsMag, new PxSphereGeometry(GeomHolder.sphere().radius), NewMaterial); break;
-			case PxGeometryType::eCAPSULE:		Shapes.Emplace(BodyLocalShape, BoundsCenter, BoundsMag, new PxCapsuleGeometry(GeomHolder.capsule().radius, GeomHolder.capsule().halfHeight), NewMaterial); break;
-			case PxGeometryType::eBOX:			Shapes.Emplace(BodyLocalShape, BoundsCenter, BoundsMag, new PxBoxGeometry(GeomHolder.box().halfExtents), NewMaterial); break;
-			case PxGeometryType::eCONVEXMESH:	Shapes.Emplace(BodyLocalShape, BoundsCenter, BoundsMag, new PxConvexMeshGeometry(GeomHolder.convexMesh().convexMesh, GeomHolder.convexMesh().scale, GeomHolder.convexMesh().meshFlags), NewMaterial); break;
-			case PxGeometryType::eHEIGHTFIELD:	Shapes.Emplace(BodyLocalShape, BoundsCenter, BoundsMag, new PxHeightFieldGeometry(GeomHolder.heightField().heightField, GeomHolder.heightField().heightFieldFlags, GeomHolder.heightField().heightScale, GeomHolder.heightField().rowScale, GeomHolder.heightField().columnScale), NewMaterial); break;
-			case PxGeometryType::eTRIANGLEMESH: Shapes.Emplace(BodyLocalShape, BoundsCenter, BoundsMag, new PxTriangleMeshGeometry(GeomHolder.triangleMesh().triangleMesh, GeomHolder.triangleMesh().scale, GeomHolder.triangleMesh().meshFlags), NewMaterial); break;
-			default: continue;	//we don't bother with other types
-			}
-		}
-	}
-
-	bool FActor::AddShape(PxShape* InShape)
-	{
-		if (!(InShape->getFlags() & PxShapeFlag::eSIMULATION_SHAPE))
-		{
-			return false;
-		}
-
-		const PxTransform LocalPose = InShape->getLocalPose();
-		const PxGeometryHolder& GeomHolder = InShape->getGeometry();
-		const PxBounds3 ShapeBounds = PxGeometryQuery::getWorldBounds(GeomHolder.any(), PxTransform(PxIdentity), 1.0f);
-		const float BoundsMagnitude = ShapeBounds.getExtents().magnitude();
-		const PxVec3 BoundsCenter = ShapeBounds.getCenter();
-		const int32 NumMaterials = InShape->getNbMaterials();
-		TArray<PxMaterial*, TInlineAllocator<1>> Materials;
-
-		// #PHYS2 Add multi-material support
-		Materials.SetNumUninitialized(NumMaterials);
-		InShape->getMaterials(Materials.GetData(), sizeof(Materials[0]) * NumMaterials);
-
-		FMaterial NewMaterial;
-		if (NumMaterials > 0)
-		{
-			NewMaterial = FMaterial(Materials[0]);
-		}
-
-		switch (GeomHolder.getType())
-		{
-		case PxGeometryType::eSPHERE:		Shapes.Emplace(LocalPose, BoundsCenter, BoundsMagnitude, new PxSphereGeometry(GeomHolder.sphere().radius), NewMaterial); break;
-		case PxGeometryType::eCAPSULE:		Shapes.Emplace(LocalPose, BoundsCenter, BoundsMagnitude, new PxCapsuleGeometry(GeomHolder.capsule().radius, GeomHolder.capsule().halfHeight), NewMaterial); break;
-		case PxGeometryType::eBOX:			Shapes.Emplace(LocalPose, BoundsCenter, BoundsMagnitude, new PxBoxGeometry(GeomHolder.box().halfExtents), NewMaterial); break;
-		case PxGeometryType::eCONVEXMESH:	Shapes.Emplace(LocalPose, BoundsCenter, BoundsMagnitude, new PxConvexMeshGeometry(GeomHolder.convexMesh().convexMesh, GeomHolder.convexMesh().scale, GeomHolder.convexMesh().meshFlags), NewMaterial); break;
-		case PxGeometryType::eHEIGHTFIELD:	Shapes.Emplace(LocalPose, BoundsCenter, BoundsMagnitude, new PxHeightFieldGeometry(GeomHolder.heightField().heightField, GeomHolder.heightField().heightFieldFlags, GeomHolder.heightField().heightScale, GeomHolder.heightField().rowScale, GeomHolder.heightField().columnScale), NewMaterial); break;
-		case PxGeometryType::eTRIANGLEMESH: Shapes.Emplace(LocalPose, BoundsCenter, BoundsMagnitude, new PxTriangleMeshGeometry(GeomHolder.triangleMesh().triangleMesh, GeomHolder.triangleMesh().scale, GeomHolder.triangleMesh().meshFlags), NewMaterial); break;
-		default: return false;	//we don't bother with other types
-		}
-
-		// If we get here then we succeeded in creating a shape.
-		Shapes.Last().UserData = InShape->userData;
-
-		return true;
-	}
-
-	void FActor::TerminateGeometry()
-	{
-		for (FShape& Shape : Shapes)
-		{
-			switch (Shape.Geometry->getType())
-			{
-			case PxGeometryType::eSPHERE: delete static_cast<PxSphereGeometry*>(Shape.Geometry); break;
-			case PxGeometryType::eCAPSULE: delete static_cast<PxCapsuleGeometry*>(Shape.Geometry); break;
-			case PxGeometryType::eBOX: delete static_cast<PxBoxGeometry*>(Shape.Geometry); break;
-			case PxGeometryType::eCONVEXMESH: delete static_cast<PxConvexMeshGeometry*>(Shape.Geometry); break;
-			case PxGeometryType::eTRIANGLEMESH: delete static_cast<PxTriangleMeshGeometry*>(Shape.Geometry); break;
-			case PxGeometryType::eHEIGHTFIELD: delete static_cast<PxHeightFieldGeometry*>(Shape.Geometry); break;
-			default: break;
-			}
-		}
-
-		Shapes.Empty();
-	}
-
-}
-
-#endif // WITH_PHYSX
-
-
diff --git a/Engine/Source/Runtime/Engine/Private/PhysicsEngine/ImmediatePhysics/ImmediatePhysicsPhysX/ImmediatePhysicsContactPointRecorder.cpp b/Engine/Source/Runtime/Engine/Private/PhysicsEngine/ImmediatePhysics/ImmediatePhysicsPhysX/ImmediatePhysicsContactPointRecorder.cpp
--- a/Engine/Source/Runtime/Engine/Private/PhysicsEngine/ImmediatePhysics/ImmediatePhysicsPhysX/ImmediatePhysicsContactPointRecorder.cpp
+++ b/Engine/Source/Runtime/Engine/Private/PhysicsEngine/ImmediatePhysics/ImmediatePhysicsPhysX/ImmediatePhysicsContactPointRecorder.cpp
@@ -1,67 +1,3 @@
 // Copyright Epic Games, Inc. All Rights Reserved.
 
 #include "Physics/ImmediatePhysics/ImmediatePhysicsPhysX/ImmediatePhysicsContactPointRecorder_PhysX.h"
-
-#if WITH_PHYSX && PHYSICS_INTERFACE_PHYSX
-
-#include "Physics/ImmediatePhysics/ImmediatePhysicsPhysX/ImmediatePhysicsSimulation_PhysX.h"
-#include "Physics/ImmediatePhysics/ImmediatePhysicsPhysX/ImmediatePhysicsContactPair_PhysX.h"
-
-namespace ImmediatePhysics_PhysX
-{
-
-	EFrictionCombineMode::Type GetCombineMode(EFrictionCombineMode::Type A, EFrictionCombineMode::Type B)
-	{
-		return FMath::Max(A, B);
-	}
-
-	float UseCombineMode(EFrictionCombineMode::Type Mode, float A, float B)
-	{
-		switch (Mode)
-		{
-		case EFrictionCombineMode::Average: return (A * 0.5f) + (B * 0.5f);
-		case EFrictionCombineMode::Multiply: return A * B;
-		case EFrictionCombineMode::Min: return FMath::Min(A, B);
-		case EFrictionCombineMode::Max: return FMath::Max(A, B);
-		default: return 0.f;	//error
-		}
-	}
-
-	bool FContactPointRecorder::recordContacts(const Gu::ContactPoint* ContactPoints, const PxU32 NumContacts, const PxU32 Index)
-	{
-		FContactPair ContactPair;
-		ContactPair.DynamicActorDataIndex = DynamicActorDataIndex;
-		ContactPair.OtherActorDataIndex = OtherActorDataIndex;
-		ContactPair.StartContactIndex = Simulation.ContactPoints.Num();
-		ContactPair.NumContacts = NumContacts;
-		ContactPair.PairIdx = PairIdx;
-
-		//TODO: can probably optimize this for the default material case
-		EFrictionCombineMode::Type FrictionCombine = GetCombineMode(SimulatedShapeMaterial.FrictionCombineMode, OtherShapeMaterial.FrictionCombineMode);
-		EFrictionCombineMode::Type RestitutionCombine = GetCombineMode(SimulatedShapeMaterial.FrictionCombineMode, OtherShapeMaterial.FrictionCombineMode);
-
-		const float StaticFriction = UseCombineMode(FrictionCombine, SimulatedShapeMaterial.StaticFriction, OtherShapeMaterial.StaticFriction);
-		const float DynamicFriction = UseCombineMode(FrictionCombine, SimulatedShapeMaterial.DynamicFriction, OtherShapeMaterial.DynamicFriction);
-		const float Restitution = UseCombineMode(RestitutionCombine, SimulatedShapeMaterial.Restitution, OtherShapeMaterial.Restitution);
-
-		for (PxU32 ContactIndex = 0; ContactIndex < NumContacts; ++ContactIndex)
-		{
-			//Fill in solver-specific data that our contact gen does not produce...
-
-			Gu::ContactPoint NewPoint = ContactPoints[ContactIndex];
-			NewPoint.maxImpulse = PX_MAX_F32;
-			NewPoint.targetVel = PxVec3(0.f);
-			NewPoint.staticFriction = StaticFriction;
-			NewPoint.dynamicFriction = DynamicFriction;
-			NewPoint.restitution = Restitution;
-			NewPoint.materialFlags = 0;
-			Simulation.ContactPoints.Add(NewPoint);
-		}
-
-		Simulation.ContactPairs.Add(ContactPair);
-		return true;
-	}
-
-}
-
-#endif // WITH_PHYSX
