Physics interface cleanup.

diff --git a/Engine/Source/Runtime/Engine/Private/PhysicsEngine/ImmediatePhysics/ImmediatePhysicsPhysX/ImmediatePhysicsD6JointData_PhysX.cpp b/Engine/Source/Runtime/Engine/Private/PhysicsEngine/ImmediatePhysics/ImmediatePhysicsPhysX/ImmediatePhysicsD6JointData_PhysX.cpp
--- a/Engine/Source/Runtime/Engine/Private/PhysicsEngine/ImmediatePhysics/ImmediatePhysicsPhysX/ImmediatePhysicsD6JointData_PhysX.cpp
+++ b/Engine/Source/Runtime/Engine/Private/PhysicsEngine/ImmediatePhysics/ImmediatePhysicsPhysX/ImmediatePhysicsD6JointData_PhysX.cpp
@@ -1,751 +1,3 @@
 // Copyright Epic Games, Inc. All Rights Reserved.
 
 #include "Physics/ImmediatePhysics/ImmediatePhysicsPhysX/ImmediatePhysicsD6JointData_PhysX.h"
-
-#if WITH_PHYSX && PHYSICS_INTERFACE_PHYSX
-
-#include "PhysicsPublic.h"
-
-/****** 
- ************ :(
-///CODE COPIED DIRECTLY OUT OF PHYSX (until refactor)
- ************
- ******/
-
-
-namespace ImmediatePhysics_PhysX
-{
-
-	D6JointData::D6JointData(PxD6Joint* Joint)
-		: linearLimit(GPhysXSDK->getTolerancesScale(), 0.f)
-		, twistLimit(0.f, 0.f)
-		, swingLimit(0.f, 0.f)
-	{
-		PxRigidActor* Actor0 = nullptr;
-		PxRigidActor* Actor1 = nullptr;
-		Joint->getActors(Actor0, Actor1);
-
-		auto GetCOM = [](PxRigidActor* Actor) -> PxTransform
-		{
-			if (Actor)
-			{
-				if (Actor->getType() == PxActorType::eRIGID_DYNAMIC || Actor->getType() == PxActorType::eARTICULATION_LINK)
-				{
-					return static_cast<PxRigidBody*>(Actor)->getCMassLocalPose();
-				}
-				else
-				{
-					return static_cast<PxRigidStatic*>(Actor)->getGlobalPose().getInverse();
-				}
-			}
-
-			return PxTransform(PxIdentity);
-		};
-
-		c2b[0] = GetCOM(Actor0).transformInv(Joint->getLocalPose(PxJointActorIndex::eACTOR0));
-		c2b[1] = GetCOM(Actor1).transformInv(Joint->getLocalPose(PxJointActorIndex::eACTOR1));
-
-		linearLimit = Joint->getLinearLimit();
-		swingLimit = Joint->getSwingLimit();
-		twistLimit = Joint->getTwistLimit();
-		linearMinDist = 0.001f;
-
-		for (uint32 DriveIdx = 0; DriveIdx < PxD6Drive::eCOUNT; ++DriveIdx)
-		{
-			drive[DriveIdx] = Joint->getDrive((PxD6Drive::Enum)DriveIdx);
-		}
-
-		Joint->getDriveVelocity(driveLinearVelocity, driveAngularVelocity);
-		drivePosition = Joint->getDrivePosition();
-
-		for (uint32 MotionIdx = 0; MotionIdx < PxD6Axis::eCOUNT; ++MotionIdx)
-		{
-			motion[MotionIdx] = Joint->getMotion((PxD6Axis::Enum)MotionIdx);
-		}
-
-		invMassScale.linear0 = Joint->getInvMassScale0();
-		invMassScale.linear1 = Joint->getInvMassScale1();
-		invMassScale.angular0 = Joint->getInvInertiaScale0();
-		invMassScale.angular1 = Joint->getInvInertiaScale1();
-	}
-
-	bool IsActive(const D6JointData& JointData, const PxD6Drive::Enum Index)
-	{
-		const PxD6JointDrive& DriveData = JointData.drive[Index];
-		return DriveData.stiffness != 0 || DriveData.damping != 0;
-	}
-
-	void PrepareJointData(D6JointData& JointData)
-	{
-		//angular
-		JointData.thSwingY = PxTan(JointData.swingLimit.yAngle / 2);
-		JointData.thSwingZ = PxTan(JointData.swingLimit.zAngle / 2);
-		JointData.thSwingPad = PxTan(JointData.swingLimit.contactDistance / 2);
-
-		JointData.tqSwingY = PxTan(JointData.swingLimit.yAngle / 4);
-		JointData.tqSwingZ = PxTan(JointData.swingLimit.zAngle / 4);
-		JointData.tqSwingPad = PxTan(JointData.swingLimit.contactDistance / 4);
-
-		JointData.tqTwistLow = PxTan(JointData.twistLimit.lower / 4);
-		JointData.tqTwistHigh = PxTan(JointData.twistLimit.upper / 4);
-		JointData.tqTwistPad = PxTan(JointData.twistLimit.contactDistance / 4);
-
-		//linear
-		JointData.driving = 0;
-		JointData.limited = 0;
-		JointData.locked = 0;
-
-		for (PxU32 i = 0; i < PxD6Axis::eCOUNT; i++)
-		{
-			if (JointData.motion[i] == PxD6Motion::eLIMITED)
-				JointData.limited |= 1 << i;
-			else if (JointData.motion[i] == PxD6Motion::eLOCKED)
-				JointData.locked |= 1 << i;
-		}
-
-		// a linear direction isn't driven if it's locked
-		if (IsActive(JointData, PxD6Drive::eX) && JointData.motion[PxD6Axis::eX] != PxD6Motion::eLOCKED) JointData.driving |= 1 << PxD6Drive::eX;
-		if (IsActive(JointData, PxD6Drive::eY) && JointData.motion[PxD6Axis::eY] != PxD6Motion::eLOCKED) JointData.driving |= 1 << PxD6Drive::eY;
-		if (IsActive(JointData, PxD6Drive::eZ) && JointData.motion[PxD6Axis::eZ] != PxD6Motion::eLOCKED) JointData.driving |= 1 << PxD6Drive::eZ;
-
-		// SLERP drive requires all angular dofs unlocked, and inhibits swing/twist
-
-		bool swing1Locked = JointData.motion[PxD6Axis::eSWING1] == PxD6Motion::eLOCKED;
-		bool swing2Locked = JointData.motion[PxD6Axis::eSWING2] == PxD6Motion::eLOCKED;
-		bool twistLocked = JointData.motion[PxD6Axis::eTWIST] == PxD6Motion::eLOCKED;
-
-		if (IsActive(JointData, PxD6Drive::eSLERP) && !swing1Locked && !swing2Locked && !twistLocked)
-			JointData.driving |= 1 << PxD6Drive::eSLERP;
-		else
-		{
-			if (IsActive(JointData, PxD6Drive::eTWIST) && !twistLocked)
-				JointData.driving |= 1 << PxD6Drive::eTWIST;
-			if (IsActive(JointData, PxD6Drive::eSWING) && (!swing1Locked || !swing2Locked))
-				JointData.driving |= 1 << PxD6Drive::eSWING;
-		}
-	}
-
-	void computeJacobianAxes(PxVec3 row[3], const PxQuat& qa, const PxQuat& qb)
-	{
-		// Compute jacobian matrix for (qa* qb)  [[* means conjugate in this expr]]
-		// d/dt (qa* qb) = 1/2 L(qa*) R(qb) (omega_b - omega_a)
-		// result is L(qa*) R(qb), where L(q) and R(q) are left/right q multiply matrix
-
-		PxReal wa = qa.w, wb = qb.w;
-		const PxVec3 va(qa.x, qa.y, qa.z), vb(qb.x, qb.y, qb.z);
-
-		const PxVec3 c = vb * wa + va * wb;
-		const PxReal d0 = wa * wb;
-		const PxReal d1 = va.dot(vb);
-		const PxReal d = d0 - d1;
-
-		row[0] = (va * vb.x + vb * va.x + PxVec3(d, c.z, -c.y)) * 0.5f;
-		row[1] = (va * vb.y + vb * va.y + PxVec3(-c.z, d, c.x)) * 0.5f;
-		row[2] = (va * vb.z + vb * va.z + PxVec3(c.y, -c.x, d)) * 0.5f;
-
-		if ((d0 + d1) != 0.0f)  // check if relative rotation is 180 degrees which can lead to singular matrix
-			return;
-		else
-		{
-			row[0].x += PX_EPS_F32;
-			row[1].y += PX_EPS_F32;
-			row[2].z += PX_EPS_F32;
-		}
-	}
-
-
-	PxReal tanAdd(PxReal tan1, PxReal tan2)
-	{
-		PX_ASSERT(PxAbs(1 - tan1 * tan2) > 1e-6f);
-		return (tan1 + tan2) / (1 - tan1 * tan2);
-	}
-
-	PX_CUDA_CALLABLE PxF32 sqr(const PxF32 a)
-	{
-		return a * a;
-	}
-
-	PX_CUDA_CALLABLE PxReal tanHalf(PxReal sin, PxReal cos)
-	{
-		return sin / (1 + cos);
-	}
-
-
-	class ConstraintHelper
-	{
-		Px1DConstraint* mConstraints;
-		Px1DConstraint* mCurrent;
-		PxVec3 mRa, mRb;
-
-	public:
-		ConstraintHelper(Px1DConstraint* c, const PxVec3& ra, const PxVec3& rb)
-			: mConstraints(c), mCurrent(c), mRa(ra), mRb(rb) {}
-
-
-		// hard linear & angular
-		void linearHard(const PxVec3& axis, PxReal posErr)
-		{
-			Px1DConstraint *c = linear(axis, posErr, PxConstraintSolveHint::eEQUALITY);
-			c->flags |= Px1DConstraintFlag::eOUTPUT_FORCE;
-		}
-
-		void angularHard(const PxVec3& axis, PxReal posErr)
-		{
-			Px1DConstraint *c = angular(axis, posErr, PxConstraintSolveHint::eEQUALITY);
-			c->flags |= Px1DConstraintFlag::eOUTPUT_FORCE;
-		}
-
-		// limited linear & angular
-		void linearLimit(const PxVec3& axis, PxReal ordinate, PxReal limitValue, const PxJointLimitParameters& limit)
-		{
-			PxReal pad = limit.isSoft() ? 0 : limit.contactDistance;
-
-			if (ordinate + pad > limitValue)
-				addLimit(linear(axis, limitValue - ordinate, PxConstraintSolveHint::eNONE), limit);
-		}
-
-		void angularLimit(const PxVec3& axis, PxReal ordinate, PxReal limitValue, PxReal pad, const PxJointLimitParameters& limit)
-		{
-			if (limit.isSoft())
-				pad = 0;
-
-			if (ordinate + pad > limitValue)
-				addLimit(angular(axis, limitValue - ordinate, PxConstraintSolveHint::eNONE), limit);
-		}
-
-
-		void angularLimit(const PxVec3& axis, PxReal error, const PxJointLimitParameters& limit)
-		{
-			addLimit(angular(axis, error, PxConstraintSolveHint::eNONE), limit);
-		}
-
-		void halfAnglePair(PxReal halfAngle, PxReal lower, PxReal upper, PxReal pad, const PxVec3& axis, const PxJointLimitParameters& limit)
-		{
-			PX_ASSERT(lower < upper);
-			if (limit.isSoft())
-				pad = 0;
-
-			if (halfAngle < lower + pad)
-				angularLimit(-axis, -(lower - halfAngle) * 2, limit);
-			if (halfAngle > upper - pad)
-				angularLimit(axis, (upper - halfAngle) * 2, limit);
-		}
-
-		void quarterAnglePair(PxReal quarterAngle, PxReal lower, PxReal upper, PxReal pad, const PxVec3& axis, const PxJointLimitParameters& limit)
-		{
-			if (limit.isSoft())
-				pad = 0;
-
-			PX_ASSERT(lower < upper);
-			if (quarterAngle < lower + pad)
-				angularLimit(-axis, -(lower - quarterAngle) * 4, limit);
-			if (quarterAngle > upper - pad)
-				angularLimit(axis, (upper - quarterAngle) * 4, limit);
-		}
-
-		// driven linear & angular
-
-		void linear(const PxVec3& axis, PxReal velTarget, PxReal error, const PxD6JointDrive& drive)
-		{
-			addDrive(linear(axis, error, PxConstraintSolveHint::eNONE), velTarget, drive);
-		}
-
-		void angular(const PxVec3& axis, PxReal velTarget, PxReal error, const PxD6JointDrive& drive, PxConstraintSolveHint::Enum hint = PxConstraintSolveHint::eNONE)
-		{
-			addDrive(angular(axis, error, hint), velTarget, drive);
-		}
-
-
-		PxU32 getCount() { return PxU32(mCurrent - mConstraints); }
-
-		void prepareLockedAxes(const PxQuat& qA, const PxQuat& qB, const PxVec3& cB2cAp, PxU32 lin, PxU32 ang)
-		{
-			Px1DConstraint* current = mCurrent;
-			if (ang)
-			{
-				PxQuat qB2qA = qA.getConjugate() * qB;
-
-				PxVec3 row[3];
-				computeJacobianAxes(row, qA, qB);
-				PxVec3 imp = qB2qA.getImaginaryPart();
-				if (ang & 1) angular(row[0], -imp.x, PxConstraintSolveHint::eEQUALITY, current++);
-				if (ang & 2) angular(row[1], -imp.y, PxConstraintSolveHint::eEQUALITY, current++);
-				if (ang & 4) angular(row[2], -imp.z, PxConstraintSolveHint::eEQUALITY, current++);
-			}
-
-			if (lin)
-			{
-				PxMat33 axes(qA);
-				if (lin & 1) linear(axes[0], -cB2cAp[0], PxConstraintSolveHint::eEQUALITY, current++);
-				if (lin & 2) linear(axes[1], -cB2cAp[1], PxConstraintSolveHint::eEQUALITY, current++);
-				if (lin & 4) linear(axes[2], -cB2cAp[2], PxConstraintSolveHint::eEQUALITY, current++);
-			}
-
-			for (Px1DConstraint* front = mCurrent; front < current; front++)
-				front->flags = Px1DConstraintFlag::eOUTPUT_FORCE;
-
-			mCurrent = current;
-		}
-
-		Px1DConstraint *getConstraintRow()
-		{
-			return mCurrent++;
-		}
-
-	private:
-		Px1DConstraint* linear(const PxVec3& axis, PxReal posErr, PxConstraintSolveHint::Enum hint)
-		{
-			Px1DConstraint* c = mCurrent++;
-
-			c->solveHint = PxU16(hint);
-			c->linear0 = axis;					c->angular0 = mRa.cross(axis);
-			c->linear1 = axis;					c->angular1 = mRb.cross(axis);
-			PX_ASSERT(c->linear0.isFinite());
-			PX_ASSERT(c->linear1.isFinite());
-			PX_ASSERT(c->angular0.isFinite());
-			PX_ASSERT(c->angular1.isFinite());
-
-			c->geometricError = posErr;
-
-			return c;
-		}
-
-		Px1DConstraint* angular(const PxVec3& axis, PxReal posErr, PxConstraintSolveHint::Enum hint)
-		{
-			Px1DConstraint* c = mCurrent++;
-
-			c->solveHint = PxU16(hint);
-			c->linear0 = PxVec3(0);		c->angular0 = axis;
-			c->linear1 = PxVec3(0);		c->angular1 = axis;
-
-			c->geometricError = posErr;
-			return c;
-		}
-
-		Px1DConstraint* linear(const PxVec3& axis, PxReal posErr, PxConstraintSolveHint::Enum hint, Px1DConstraint* c)
-		{
-			c->solveHint = PxU16(hint);
-			c->linear0 = axis;					c->angular0 = mRa.cross(axis);
-			c->linear1 = axis;					c->angular1 = mRb.cross(axis);
-			PX_ASSERT(c->linear0.isFinite());
-			PX_ASSERT(c->linear1.isFinite());
-			PX_ASSERT(c->angular0.isFinite());
-			PX_ASSERT(c->angular1.isFinite());
-
-			c->geometricError = posErr;
-
-			return c;
-		}
-
-		Px1DConstraint* angular(const PxVec3& axis, PxReal posErr, PxConstraintSolveHint::Enum hint, Px1DConstraint* c)
-		{
-			c->solveHint = PxU16(hint);
-			c->linear0 = PxVec3(0.f);		c->angular0 = axis;
-			c->linear1 = PxVec3(0.f);		c->angular1 = axis;
-
-			c->geometricError = posErr;
-
-			return c;
-		}
-
-		void addLimit(Px1DConstraint* c, const PxJointLimitParameters& limit)
-		{
-			PxU16 flags = PxU16(c->flags | Px1DConstraintFlag::eOUTPUT_FORCE);
-
-			if (limit.isSoft())
-			{
-				flags |= Px1DConstraintFlag::eSPRING;
-				c->mods.spring.stiffness = limit.stiffness;
-				c->mods.spring.damping = limit.damping;
-			}
-			else
-			{
-				c->solveHint = PxConstraintSolveHint::eINEQUALITY;
-				c->mods.bounce.restitution = limit.restitution;
-				c->mods.bounce.velocityThreshold = limit.bounceThreshold;
-				if (c->geometricError > 0)
-					flags |= Px1DConstraintFlag::eKEEPBIAS;
-				if (limit.restitution > 0)
-					flags |= Px1DConstraintFlag::eRESTITUTION;
-			}
-
-			c->flags = flags;
-			c->minImpulse = 0;
-		}
-
-		void addDrive(Px1DConstraint* c, PxReal velTarget, const PxD6JointDrive& drive)
-		{
-			c->velocityTarget = velTarget;
-
-			PxU16 flags = PxU16(c->flags | Px1DConstraintFlag::eSPRING | Px1DConstraintFlag::eHAS_DRIVE_LIMIT);
-			if (drive.flags & PxD6JointDriveFlag::eACCELERATION)
-				flags |= Px1DConstraintFlag::eACCELERATION_SPRING;
-			c->flags = flags;
-			c->mods.spring.stiffness = drive.stiffness;
-			c->mods.spring.damping = drive.damping;
-
-			c->minImpulse = -drive.forceLimit;
-			c->maxImpulse = drive.forceLimit;
-
-			PX_ASSERT(c->linear0.isFinite());
-			PX_ASSERT(c->angular0.isFinite());
-		}
-	};
-
-	void separateSwingTwist(const PxQuat& q, PxQuat& swing, PxQuat& twist)
-	{
-		twist = q.x != 0.0f ? PxQuat(q.x, 0, 0, q.w).getNormalized() : PxQuat(PxIdentity);
-		swing = q * twist.getConjugate();
-	}
-
-	PX_CUDA_CALLABLE PxVec3 ellipseClamp(const PxVec3& point, const PxVec3& radii)
-	{
-		// This function need to be implemented in the header file because
-		// it is included in a spu shader program.
-
-		// finds the closest point on the ellipse to a given point
-
-		// (p.y, p.z) is the input point
-		// (e.y, e.z) are the radii of the ellipse
-
-		// lagrange multiplier method with Newton/Halley hybrid root-finder.
-		// see http://www.geometrictools.com/Documentation/DistancePointToEllipse2.pdf
-		// for proof of Newton step robustness and initial estimate.
-		// Halley converges much faster but sometimes overshoots - when that happens we take
-		// a newton step instead
-
-		// converges in 1-2 iterations where D&C works well, and it's good with 4 iterations
-		// with any ellipse that isn't completely crazy
-
-		const PxU32 MAX_ITERATIONS = 20;
-		const PxReal convergenceThreshold = 1e-4f;
-
-		// iteration requires first quadrant but we recover generality later
-
-		PxVec3 q(0, PxAbs(point.y), PxAbs(point.z));
-		const PxReal tinyEps = 1e-6f; // very close to minor axis is numerically problematic but trivial
-		if (radii.y >= radii.z)
-		{
-			if (q.z < tinyEps)
-				return PxVec3(0, point.y > 0 ? radii.y : -radii.y, 0);
-		}
-		else
-		{
-			if (q.y < tinyEps)
-				return PxVec3(0, 0, point.z > 0 ? radii.z : -radii.z);
-		}
-
-		PxVec3 denom, e2 = radii.multiply(radii), eq = radii.multiply(q);
-
-		// we can use any initial guess which is > maximum(-e.y^2,-e.z^2) and for which f(t) is > 0.
-		// this guess works well near the axes, but is weak along the diagonals.
-
-		PxReal t = PxMax(eq.y - e2.y, eq.z - e2.z);
-
-		for (PxU32 i = 0; i < MAX_ITERATIONS; i++)
-		{
-			denom = PxVec3(0, 1 / (t + e2.y), 1 / (t + e2.z));
-			PxVec3 denom2 = eq.multiply(denom);
-
-			PxVec3 fv = denom2.multiply(denom2);
-			PxReal f = fv.y + fv.z - 1;
-
-			// although in exact arithmetic we are guaranteed f>0, we can get here
-			// on the first iteration via catastrophic cancellation if the point is
-			// very close to the origin. In that case we just behave as if f=0
-
-			if (f < convergenceThreshold)
-				return e2.multiply(point).multiply(denom);
-
-			PxReal df = fv.dot(denom) * -2.0f;
-			t = t - f / df;
-		}
-
-		// we didn't converge, so clamp what we have
-		PxVec3 r = e2.multiply(point).multiply(denom);
-		return r * PxRecipSqrt(sqr(r.y / radii.y) + sqr(r.z / radii.z));
-	}
-
-
-	class ConeLimitHelper
-	{
-	public:
-		ConeLimitHelper(PxReal tanQSwingY, PxReal tanQSwingZ, PxReal tanQPadding)
-			: mTanQYMax(tanQSwingY), mTanQZMax(tanQSwingZ), mTanQPadding(tanQPadding) {}
-
-		// whether the point is inside the (inwardly) padded cone - if it is, there's no limit
-		// constraint
-
-		bool contains(const PxVec3& tanQSwing)
-		{
-			PxReal tanQSwingYPadded = tanAdd(PxAbs(tanQSwing.y), mTanQPadding);
-			PxReal tanQSwingZPadded = tanAdd(PxAbs(tanQSwing.z), mTanQPadding);
-			return sqr(tanQSwingYPadded / mTanQYMax) + sqr(tanQSwingZPadded / mTanQZMax) <= 1;
-		}
-
-		PxVec3 clamp(const PxVec3& tanQSwing,
-			PxVec3& normal)
-		{
-			PxVec3 p = ellipseClamp(tanQSwing, PxVec3(0, mTanQYMax, mTanQZMax));
-			normal = PxVec3(0, p.y / sqr(mTanQYMax), p.z / sqr(mTanQZMax));
-#ifdef PX_PARANOIA_ELLIPSE_CHECK
-			PxReal err = PxAbs(sqr(p.y / mTanQYMax) + sqr(p.z / mTanQZMax) - 1);
-			PX_ASSERT(err < 1e-3);
-#endif
-
-			return p;
-		}
-
-
-		// input is a swing quat, such that swing.x = twist.y = twist.z = 0, q = swing * twist
-		// The routine is agnostic to the sign of q.w (i.e. we don't need the minimal-rotation swing)
-
-		// output is an axis such that positive rotation increases the angle outward from the
-		// limit (i.e. the image of the x axis), the error is the sine of the angular difference,
-		// positive if the twist axis is inside the cone 
-
-		bool getLimit(const PxQuat& swing, PxVec3& axis, PxReal& error)
-		{
-			PX_ASSERT(swing.w > 0);
-			PxVec3 twistAxis = swing.getBasisVector0();
-			PxVec3 tanQSwing = PxVec3(0, tanHalf(swing.z, swing.w), -tanHalf(swing.y, swing.w));
-			if (contains(tanQSwing))
-				return false;
-
-			PxVec3 normal, clamped = clamp(tanQSwing, normal);
-
-			// rotation vector and ellipse normal
-			PxVec3 r(0, -clamped.z, clamped.y), d(0, -normal.z, normal.y);
-
-			// the point on the cone defined by the tanQ swing vector r
-			PxVec3 p(1.f, 0, 0);
-			PxReal r2 = r.dot(r), a = 1 - r2, b = 1 / (1 + r2), b2 = b * b;
-			PxReal v1 = 2 * a*b2;
-			PxVec3 v2(a, 2 * r.z, -2 * r.y);		// a*p + 2*r.cross(p);
-			PxVec3 coneLine = v1 * v2 - p;		// already normalized
-
-												// the derivative of coneLine in the direction d	
-			PxReal rd = r.dot(d);
-			PxReal dv1 = -4 * rd*(3 - r2)*b2*b;
-			PxVec3 dv2(-2 * rd, 2 * d.z, -2 * d.y);
-
-			PxVec3 coneNormal = v1 * dv2 + dv1 * v2;
-
-			axis = coneLine.cross(coneNormal) / coneNormal.magnitude();
-			error = coneLine.cross(axis).dot(twistAxis);
-
-			PX_ASSERT(PxAbs(axis.magnitude() - 1) < 1e-5f);
-
-#ifdef PX_PARANOIA_ELLIPSE_CHECK
-			bool inside = sqr(tanQSwing.y / mTanQYMax) + sqr(tanQSwing.z / mTanQZMax) <= 1;
-			PX_ASSERT(inside && error > -1e-4f || !inside && error < 1e-4f);
-#endif
-
-			return true;
-		}
-
-	private:
-
-
-		PxReal mTanQYMax, mTanQZMax, mTanQPadding;
-	};
-
-
-	PxU32 D6JointSolverPrep(Px1DConstraint* constraints,
-		PxVec3& body0WorldOffset,
-		PxU32 maxConstraints,
-		PxConstraintInvMassScale& invMassScale,
-		const void* constantBlock,
-		const PxTransform& bA2w,
-		const PxTransform& bB2w)
-	{
-		PX_UNUSED(maxConstraints);
-
-		PX_UNUSED(maxConstraints);
-		const D6JointData& data =
-			*reinterpret_cast<const D6JointData*>(constantBlock);
-
-		invMassScale = data.invMassScale;
-
-		const PxU32 SWING1_FLAG = 1 << PxD6Axis::eSWING1,
-			SWING2_FLAG = 1 << PxD6Axis::eSWING2,
-			TWIST_FLAG = 1 << PxD6Axis::eTWIST;
-
-		const PxU32 ANGULAR_MASK = SWING1_FLAG | SWING2_FLAG | TWIST_FLAG;
-		const PxU32 LINEAR_MASK = 1 << PxD6Axis::eX | 1 << PxD6Axis::eY | 1 << PxD6Axis::eZ;
-
-		const PxD6JointDrive* drives = data.drive;
-		PxU32 locked = data.locked, limited = data.limited, driving = data.driving;
-
-		PxTransform cA2w = bA2w.transform(data.c2b[0]);
-		PxTransform cB2w = bB2w.transform(data.c2b[1]);
-
-		body0WorldOffset = cB2w.p - bA2w.p;
-		ConstraintHelper g(constraints, cB2w.p - bA2w.p, cB2w.p - bB2w.p);
-
-		if (cA2w.q.dot(cB2w.q) < 0)	// minimum dist quat (equiv to flipping cB2bB.q, which we don't use anywhere)
-			cB2w.q = -cB2w.q;
-
-		PxTransform cB2cA = cA2w.transformInv(cB2w);
-
-		PX_ASSERT(data.c2b[0].isValid());
-		PX_ASSERT(data.c2b[1].isValid());
-		PX_ASSERT(cA2w.isValid());
-		PX_ASSERT(cB2w.isValid());
-		PX_ASSERT(cB2cA.isValid());
-
-		PxMat33 cA2w_m(cA2w.q), cB2w_m(cB2w.q);
-
-		// handy for swing computation
-		PxVec3 bX = cB2w_m[0], aY = cA2w_m[1], aZ = cA2w_m[2];
-
-		if (driving & ((1 << PxD6Drive::eX) | (1 << PxD6Drive::eY) | (1 << PxD6Drive::eZ)))
-		{
-			// TODO: make drive unilateral if we are outside the limit
-			PxVec3 posErr = data.drivePosition.p - cB2cA.p;
-			for (PxU32 i = 0; i < 3; i++)
-			{
-				// -driveVelocity because velTarget is child (body1) - parent (body0) and Jacobian is 1 for body0 and -1 for parent
-				if (driving & (1 << (PxD6Drive::eX + i)))
-					g.linear(cA2w_m[i], -data.driveLinearVelocity[i], posErr[i], drives[PxD6Drive::eX + i]);
-			}
-		}
-
-		if (driving & ((1 << PxD6Drive::eSLERP) | (1 << PxD6Drive::eSWING) | (1 << PxD6Drive::eTWIST)))
-		{
-			PxQuat d2cA_q = cB2cA.q.dot(data.drivePosition.q) > 0 ? data.drivePosition.q : -data.drivePosition.q;
-
-			const PxVec3& v = data.driveAngularVelocity;
-			PxQuat delta = d2cA_q.getConjugate() * cB2cA.q;
-
-			if (driving & (1 << PxD6Drive::eSLERP))
-			{
-				PxVec3 velTarget = -cA2w.rotate(data.driveAngularVelocity);
-
-				PxVec3 axis[3] = { PxVec3(1.f,0,0), PxVec3(0,1.f,0), PxVec3(0,0,1.f) };
-
-				if (drives[PxD6Drive::eSLERP].stiffness != 0)
-					computeJacobianAxes(axis, cA2w.q * d2cA_q, cB2w.q);	// converges faster if there is only velocity drive
-
-				for (PxU32 i = 0; i < 3; i++)
-					g.angular(axis[i], axis[i].dot(velTarget), -delta.getImaginaryPart()[i], drives[PxD6Drive::eSLERP], PxConstraintSolveHint::eSLERP_SPRING);
-			}
-			else
-			{
-				if (driving & (1 << PxD6Drive::eTWIST))
-					g.angular(bX, v.x, -2.0f * delta.x, drives[PxD6Drive::eTWIST]);
-
-				if (driving & (1 << PxD6Drive::eSWING))
-				{
-					PxVec3 err = delta.rotate(PxVec3(1.f, 0, 0));
-
-					if (!(locked & SWING1_FLAG))
-						g.angular(cB2w_m[1], v.y, err.z, drives[PxD6Drive::eSWING]);
-
-					if (!(locked & SWING2_FLAG))
-						g.angular(cB2w_m[2], v.z, -err.y, drives[PxD6Drive::eSWING]);
-				}
-			}
-		}
-
-		if (limited & ANGULAR_MASK)
-		{
-			PxQuat swing, twist;
-			separateSwingTwist(cB2cA.q, swing, twist);
-
-			// swing limits: if just one is limited: if the other is free, we support 
-			// (-pi/2, +pi/2) limit, using tan of the half-angle as the error measure parameter. 
-			// If the other is locked, we support (-pi, +pi) limits using the tan of the quarter-angle
-			// Notation: th == Ps::tanHalf, tq = tanQuarter
-
-			if (limited & SWING1_FLAG && limited & SWING2_FLAG)
-			{
-				ConeLimitHelper coneHelper(data.tqSwingZ, data.tqSwingY, data.tqSwingPad);
-
-				PxVec3 axis;
-				PxReal error;
-				if (coneHelper.getLimit(swing, axis, error))
-					g.angularLimit(cA2w.rotate(axis), error, data.swingLimit);
-			}
-			else
-			{
-				const PxJointLimitCone &limit = data.swingLimit;
-
-				PxReal tqPad = data.tqSwingPad, thPad = data.thSwingPad;
-
-				if (limited & SWING1_FLAG)
-				{
-					if (locked & SWING2_FLAG)
-					{
-						g.quarterAnglePair(tanHalf(swing.y, swing.w), -data.tqSwingY, data.tqSwingY, tqPad, aY, limit);
-					}
-					else
-					{
-						PxReal dot = -aZ.dot(bX);
-						g.halfAnglePair(tanHalf(dot, 1 - dot * dot), -data.thSwingY, data.thSwingY, thPad, aZ.cross(bX), limit);
-					}
-				}
-				if (limited & SWING2_FLAG)
-				{
-					if (locked & SWING1_FLAG)
-					{
-						g.quarterAnglePair(tanHalf(swing.z, swing.w), -data.tqSwingZ, data.tqSwingZ, tqPad, aZ, limit);
-					}
-					else
-					{
-						PxReal dot = aY.dot(bX);
-						g.halfAnglePair(tanHalf(dot, 1 - dot * dot), -data.thSwingZ, data.thSwingZ, thPad, -aY.cross(bX), limit);
-					}
-				}
-			}
-
-			if (limited & TWIST_FLAG)
-			{
-				g.quarterAnglePair(tanHalf(twist.x, twist.w), data.tqTwistLow, data.tqTwistHigh, data.tqTwistPad,
-					cB2w_m[0], data.twistLimit);
-			}
-		}
-
-		if (limited & LINEAR_MASK)
-		{
-			PxVec3 limitDir = PxVec3(0);
-
-			for (PxU32 i = 0; i < 3; i++)
-			{
-				if (limited & (1 << (PxD6Axis::eX + i)))
-					limitDir += cA2w_m[i] * cB2cA.p[i];
-			}
-
-			PxReal distance = limitDir.magnitude();
-			if (distance > data.linearMinDist)
-				g.linearLimit(limitDir * (1.0f / distance), distance, data.linearLimit.value, data.linearLimit);
-		}
-
-		// we handle specially the case of just one swing dof locked
-
-		PxU32 angularLocked = locked & ANGULAR_MASK;
-
-		if (angularLocked == SWING1_FLAG)
-		{
-			g.angularHard(bX.cross(aZ), -bX.dot(aZ));
-			locked &= ~SWING1_FLAG;
-		}
-		else if (angularLocked == SWING2_FLAG)
-		{
-			locked &= ~SWING2_FLAG;
-			g.angularHard(bX.cross(aY), -bX.dot(aY));
-		}
-
-		g.prepareLockedAxes(cA2w.q, cB2w.q, cB2cA.p, locked & 7, locked >> 3);
-
-		return g.getCount();
-	}
-
-	/*PxMassProperties FLowLevelPhysicsEntity::ComputeMassProperties() const
-	{
-		//return PxRigidBodyExt::computeMassPropertiesFromGeometries(Geometries.GetData(), LocalTMs.GetData(), Geometries.Num());
-	}*/
-
-}
-
-#endif // WITH_PHYSX
diff --git a/Engine/Source/Runtime/Engine/Private/PhysicsEngine/ImmediatePhysics/ImmediatePhysicsPhysX/ImmediatePhysicsMaterial_PhysX.cpp b/Engine/Source/Runtime/Engine/Private/PhysicsEngine/ImmediatePhysics/ImmediatePhysicsPhysX/ImmediatePhysicsMaterial_PhysX.cpp
--- a/Engine/Source/Runtime/Engine/Private/PhysicsEngine/ImmediatePhysics/ImmediatePhysicsPhysX/ImmediatePhysicsMaterial_PhysX.cpp
+++ b/Engine/Source/Runtime/Engine/Private/PhysicsEngine/ImmediatePhysics/ImmediatePhysicsPhysX/ImmediatePhysicsMaterial_PhysX.cpp
@@ -1,23 +1,3 @@
 // Copyright Epic Games, Inc. All Rights Reserved.
 
 #include "Physics/ImmediatePhysics/ImmediatePhysicsPhysX/ImmediatePhysicsMaterial_PhysX.h"
-
-#if WITH_PHYSX && PHYSICS_INTERFACE_PHYSX
-#include "PhysXPublic.h"
-
-namespace ImmediatePhysics_PhysX
-{
-	FMaterial::FMaterial(physx::PxMaterial* InPxMaterial)
-		: StaticFriction(InPxMaterial->getStaticFriction())
-		, DynamicFriction(InPxMaterial->getDynamicFriction())
-		, Restitution(InPxMaterial->getRestitution())
-		, FrictionCombineMode((EFrictionCombineMode::Type)InPxMaterial->getFrictionCombineMode())
-		, RestitutionCombineMode((EFrictionCombineMode::Type)InPxMaterial->getRestitutionCombineMode())
-	{
-	}
-
-	/** Default shape material. Created from the CDO of UPhysicalMaterial */
-	FMaterial FMaterial::Default;
-}
-
-#endif // WITH_PHYSX
diff --git a/Engine/Source/Runtime/Engine/Private/PhysicsEngine/ImmediatePhysics/ImmediatePhysicsPhysX/ImmediatePhysicsSimulation_PhysX.cpp b/Engine/Source/Runtime/Engine/Private/PhysicsEngine/ImmediatePhysics/ImmediatePhysicsPhysX/ImmediatePhysicsSimulation_PhysX.cpp
--- a/Engine/Source/Runtime/Engine/Private/PhysicsEngine/ImmediatePhysics/ImmediatePhysicsPhysX/ImmediatePhysicsSimulation_PhysX.cpp
+++ b/Engine/Source/Runtime/Engine/Private/PhysicsEngine/ImmediatePhysics/ImmediatePhysicsPhysX/ImmediatePhysicsSimulation_PhysX.cpp
@@ -1,1179 +1,3 @@
 // Copyright Epic Games, Inc. All Rights Reserved.
 
 #include "Physics/ImmediatePhysics/ImmediatePhysicsPhysX/ImmediatePhysicsSimulation_PhysX.h"
-
-#if WITH_PHYSX && PHYSICS_INTERFACE_PHYSX
-
-#include "Physics/ImmediatePhysics/ImmediatePhysicsPhysX/ImmediatePhysicsActorHandle_PhysX.h"
-#include "Physics/ImmediatePhysics/ImmediatePhysicsPhysX/ImmediatePhysicsJointHandle_PhysX.h"
-#include "Physics/ImmediatePhysics/ImmediatePhysicsPhysX/ImmediatePhysicsContactPointRecorder_PhysX.h"
-#include "Physics/ImmediatePhysics/ImmediatePhysicsPhysX/ResourceManager_PhysX.h"
-
-#include "Physics/ImmediatePhysics/ImmediatePhysicsStats.h"
-
-#include "PhysicsEngine/BodyInstance.h"
-#include "PhysicsEngine/ConstraintInstance.h"
-
-#include "PhysicalMaterials/PhysicalMaterial.h"
-
-namespace ImmediatePhysics_PhysX
-{
-
-FSimulation::FSimulation()
-{
-	NumActiveSimulatedBodies = 0;
-	NumSimulatedBodies = 0;
-	NumKinematicBodies = 0;
-	NumSimulatedShapesWithCollision = 0;
-	NumPositionIterations = 1;
-	NumVelocityIterations = 1;
-	bRecreateIterationCache = false;
-	SimCount = 0;
-
-	NumContactHeaders = 0;
-	NumJointHeaders = 0;
-	NumActiveJoints = 0;
-	bDirtyJointData = false;
-}
-
-FSimulation::~FSimulation()
-{
-	for(FActorHandle* ActorHandle : ActorHandles)
-	{
-		delete ActorHandle;
-	}
-
-	for (FJointHandle* JointHandle : JointHandles)
-	{
-		delete JointHandle;
-	}
-	ActorHandles.Empty();
-	JointHandles.Empty();
-}
-
-void FSimulation::SetPositionIterationCount(int32 InIterationCount)
-{
-	NumVelocityIterations = InIterationCount;
-}
-
-void FSimulation::SetVelocityIterationCount(int32 InIterationCount)
-{
-	NumPositionIterations = InIterationCount;
-}
-
-void FSimulation::SetNumActiveBodies(int32 InNumActiveBodies)
-{
-	NumActiveSimulatedBodies = InNumActiveBodies;
-
-	bDirtyJointData = true;
-}
-
-FJointHandle* FSimulation::CreateJoint(FConstraintInstance* ConstraintInstance, FActorHandle* Actor1, FActorHandle* Actor2)
-{
-	//TODO: make this robust
-	if (Actor1 == nullptr) return nullptr;
-
-	check(&Actor1->OwningSimulation == this);
-	check(!Actor2 || &Actor2->OwningSimulation == this);	//Actor2 can be null
-
-	FJointHandle* NewHandle = nullptr;
-
-#if PHYSICS_INTERFACE_PHYSX
-	PxD6Joint* Joint = ConstraintInstance->GetPhysicsConstraintRef().ConstraintData;
-	int32 JointDataIdx = JointData.Emplace(Joint);
-	PrepareJointData(JointData[JointDataIdx]);
-
-	if(!JointData[JointDataIdx].HasConstraints())
-	{
-		JointData.RemoveAt(JointDataIdx, /*Count=*/1, /*bAllowShrinking=*/false);
-	}
-	else
-	{
-		Joints.Emplace();
-		FJoint& NewJoint = Joints[JointDataIdx];
-		//TODO: make dynamic vs static API more robust
-		NewJoint.DynamicActor = Actor2;
-		NewJoint.OtherActor = Actor1;
-
-		NewHandle = new FJointHandle(*this, Joints.Num() - 1);
-		JointHandles.Add(NewHandle);
-		bDirtyJointData = true;
-	}
-
-	ValidateArrays();
-#endif
-
-	return NewHandle;
-}
-
-void FSimulation::SetIgnoreCollisionPairTable(const TArray<FIgnorePair>& InIgnoreTable)
-{
-	IgnoreCollisionPairTable.Reset();
-	for(const FIgnorePair& IgnorePair : InIgnoreTable)
-	{
-		TSet<FActorHandle*>& IgnoreActorsA = IgnoreCollisionPairTable.FindOrAdd(IgnorePair.A);
-		IgnoreActorsA.Add(IgnorePair.B);
-
-		TSet<FActorHandle*>& IgnoreActorsB = IgnoreCollisionPairTable.FindOrAdd(IgnorePair.B);
-		IgnoreActorsB.Add(IgnorePair.A);
-	}
-
-	bRecreateIterationCache = true;
-}
-
-void FSimulation::SetIgnoreCollisionActors(const TArray<FActorHandle*>& InIgnoreCollisionActors)
-{
-	IgnoreCollisionActors.Reset();
-	IgnoreCollisionActors.Append(InIgnoreCollisionActors);
-	bRecreateIterationCache = true;
-}
-
-bool FSimulation::GetIsKinematic(int32 ActorDataIndex)
-{
-	const int32 NumMobileBodies = NumSimulatedBodies + NumKinematicBodies;
-	const bool bIsSimulated = ActorDataIndex < NumSimulatedBodies;
-	const bool bIsKinematic = !bIsSimulated && ActorDataIndex < NumMobileBodies;
-
-	return bIsKinematic;
-}
-
-void FSimulation::SetIsKinematic(int32 ActorDataIndex, bool bKinematic)
-{
-	const int32 NumMobileBodies = NumSimulatedBodies + NumKinematicBodies;
-	const bool bIsSimulated = ActorDataIndex < NumSimulatedBodies;
-	const bool bIsKinematic = !bIsSimulated && ActorDataIndex < NumKinematicBodies;
-	const bool bStatic = !bIsSimulated && !bIsKinematic;
-
-	if(bStatic)
-	{
-		// Should we support this case? Currently we build static solver data on actor creation
-		// so likely can't just force this to be kinematic here.
-		return;
-	}
-
-	if(bIsSimulated && bKinematic)
-	{
-		// A to B
-
-		// Need to clear out some stuff first, no velocities when kinematic
-		immediate::PxRigidBodyData& RBData = RigidBodiesData[ActorDataIndex];
-		RBData.angularVelocity = PxVec3(0.0f);
-		RBData.linearVelocity = PxVec3(0.0f);
-
-		int32 LastSimIndex = NumSimulatedBodies - 1;
-		if(ActorDataIndex != LastSimIndex)
-		{
-			SwapActorData(ActorDataIndex, LastSimIndex);
-		}
-
-		// Rearrange groups so kinematic now contains the body
-		--NumSimulatedBodies;
-		--NumActiveSimulatedBodies;
-		++NumKinematicBodies;
-	}
-	else if(bIsKinematic && !bKinematic)
-	{
-		// B to A
-		int32 FirstKinematicIndex = NumSimulatedBodies;
-
-		if(ActorDataIndex != FirstKinematicIndex)
-		{
-			SwapActorData(ActorDataIndex, FirstKinematicIndex);
-		}
-
-		++NumSimulatedBodies;
-		++NumActiveSimulatedBodies;
-		--NumKinematicBodies;
-	}
-}
-
-template <EActorType ActorType>
-int32 FSimulation::CreateActor(PxRigidActor* RigidActor, const FTransform& WorldTM)
-{
-	bDirtyJointData = true;	//new entity potentially re-orders bodies so joint data becomes stale. TODO: can this be optimized for adding static cases that don't re-order?
-	bRecreateIterationCache = true;	//new entity potentially re-orders bodies so our iteration cache becomes stale. TODO: can this be optimized for adding static cases that don't re-order?
-
-	FActor* NewEntity = new (Actors) FActor();
-	const int32 ActorDataIndex = Actors.Num() - 1;
-
-	FActorHandle* NewActorHandle = new FActorHandle(*this, ActorDataIndex);
-	ActorHandles.Add(NewActorHandle);
-
-	PxSolverBodyData* NewSolverBodyData = new (SolverBodiesData) PxSolverBodyData();
-	immediate::PxRigidBodyData* NewRigidBodyData = new (RigidBodiesData) immediate::PxRigidBodyData();
-
-	//Fill rigid body data based on the actor we've been given
-	PxRigidBody* const RigidBody = ActorType == EActorType::StaticActor ? nullptr : static_cast<PxRigidBody*>(RigidActor);	//Note the const pointer - this is for templated code
-	PxRigidDynamic* const RigidDynamic = ActorType == EActorType::DynamicActor ? static_cast<PxRigidDynamic*>(RigidActor) : nullptr;	//Note the const pointer - this is for templated code
-
-	NewRigidBodyData->invMass = RigidDynamic ? RigidDynamic->getInvMass() : 0.f;
-	NewRigidBodyData->angularVelocity = PxVec3(PxZero);
-	NewRigidBodyData->maxDepenetrationVelocity = RigidDynamic ? RigidDynamic->getMaxDepenetrationVelocity() : PX_MAX_F32;
-	NewRigidBodyData->invInertia = RigidDynamic ? RigidDynamic->getMassSpaceInvInertiaTensor() : PxVec3(PxZero);
-	NewRigidBodyData->maxContactImpulse = RigidBody ? RigidBody->getMaxContactImpulse() : PX_MAX_F32;
-	
-	NewActorHandle->ActorToBody = RigidBody ? P2UTransform(RigidBody->getCMassLocalPose()) : FTransform::Identity;	//TODO: Why is this an unreal transform?
-	NewRigidBodyData->body2World = U2PTransform(NewActorHandle->ActorToBody * WorldTM);
-
-	NewRigidBodyData->linearDamping = RigidDynamic ? RigidDynamic->getLinearDamping() : 0.f;
-	NewRigidBodyData->angularDamping = RigidDynamic ? RigidDynamic->getAngularDamping() : 0.f;
-	NewRigidBodyData->maxLinearVelocitySq = PX_MAX_F32;	//NOTE: rigid dynamics don't seem to have this max property.
-
-	PendingAcceleration.AddZeroed();
-	KinematicTargets.AddZeroed();
-
-	if(RigidDynamic)
-	{
-		NewRigidBodyData->maxAngularVelocitySq = RigidDynamic->getMaxAngularVelocity();
-		NewRigidBodyData->maxAngularVelocitySq *= NewRigidBodyData->maxAngularVelocitySq;
-
-		PxU32 PositionIterations, VelocityIterations;
-		RigidDynamic->getSolverIterationCounts(PositionIterations, VelocityIterations);
-
-		NumPositionIterations = FMath::Max((int32)PositionIterations, NumPositionIterations);
-		NumVelocityIterations = FMath::Max((int32)VelocityIterations, NumVelocityIterations);
-
-	}else
-	{
-		NewRigidBodyData->maxAngularVelocitySq = PX_MAX_F32;
-	}
-
-	if (ActorType == EActorType::StaticActor)
-	{
-		//If static we fill this struct here, for dynamic we fill during every call to simulate
-		immediate::PxConstructStaticSolverBody(NewRigidBodyData->body2World, *NewSolverBodyData);
-	}
-
-	NewEntity->CreateGeometry(RigidActor, U2PTransform(NewActorHandle->ActorToBody));
-
-	ValidateArrays();
-
-	//We sort entities so that dynamics are first, kinematics second, and static third
-	if (ActorType == EActorType::DynamicActor)	//dynamic first
-	{
-		if (NumSimulatedBodies < ActorDataIndex)	//There's at least one kinematic/static entity so swap
-		{
-			const int32 FirstNonSimulatedIndex = NumSimulatedBodies;
-			SwapActorData(FirstNonSimulatedIndex, ActorDataIndex);
-		}
-
-		++NumSimulatedBodies;
-	}
-	else if(ActorType == EActorType::KinematicActor)  //then kinematics
-	{
-		const int32 NumMobileBodies = NumKinematicBodies + NumSimulatedBodies;
-		if (NumMobileBodies < ActorDataIndex)	//There's at least one static entity so swap
-		{
-			const int32 FirstStaticIndex = NumMobileBodies;
-			SwapActorData(FirstStaticIndex, ActorDataIndex);
-		}
-
-		++NumKinematicBodies;
-	}
-
-	return NewActorHandle->ActorDataIndex;
-}
-
-FActorHandle* FSimulation::InsertActorData(const FActor& InActor, const FActorData& InData)
-{
-	// @todo(ccaulfield): does not seem to be called anywhere...
-
-	FActor* NewActor = new(Actors) FActor(InActor);
-	const int32 ActorDataIndex = Actors.Num() - 1;
-
-	FActorHandle* NewActorHandle = new FActorHandle(*this, ActorDataIndex);
-	ActorHandles.Add(NewActorHandle);
-
-	PxSolverBodyData* NewSolverBodyData = new (SolverBodiesData) PxSolverBodyData();
-	RigidBodiesData.Add(InData.RigidBodyData);
-
-	immediate::PxRigidBodyData& InternalData = RigidBodiesData.Last();
-
-	NewActorHandle->ActorToBody = FTransform::Identity; // #PHYS2 should be the centre of mass of the actor
-	InternalData.body2World = U2PTransform(NewActorHandle->ActorToBody * InData.InitialTransform);
-
-	PendingAcceleration.AddZeroed();
-	KinematicTargets.AddZeroed();
-
-	const bool bDynamic = !InData.bStatic && !InData.bKinematic;
-	if(bDynamic)
-	{
-		// Dynamic object
-
-		// #PHYS2 Setup max angular velocity
-	}
-	else if(InData.bKinematic)
-	{
-		// Kinematic object
-	}
-	else
-	{
-		// Static object
-		immediate::PxConstructStaticSolverBody(InternalData.body2World, *NewSolverBodyData);
-	}
-
-	ValidateArrays();
-
-	//We sort entities so that dynamics are first, kinematics second, and static third
-	if(bDynamic)	//dynamic first
-	{
-		if(NumSimulatedBodies < ActorDataIndex)	//There's at least one kinematic/static entity so swap
-		{
-			const int32 FirstNonSimulatedIndex = NumSimulatedBodies;
-			SwapActorData(FirstNonSimulatedIndex, ActorDataIndex);
-		}
-
-		++NumSimulatedBodies;
-	}
-	else if(InData.bKinematic)  //then kinematics
-	{
-		const int32 NumMobileBodies = NumKinematicBodies + NumSimulatedBodies;
-		if(NumMobileBodies < ActorDataIndex)	//There's at least one static entity so swap
-		{
-			const int32 FirstStaticIndex = NumMobileBodies;
-			SwapActorData(FirstStaticIndex, ActorDataIndex);
-		}
-
-		++NumKinematicBodies;
-	}
-
-	NumActiveSimulatedBodies = NumSimulatedBodies;
-
-	return NewActorHandle;
-}
-
-void FSimulation::RemoveActor(FActorHandle* InHandle)
-{
-	// @todo(ccaulfield): this is never called. Also I think it should delete the handle?
-
-	int32 DataIndex = InHandle->ActorDataIndex;
-
-	check(Actors.IsValidIndex(DataIndex));
-	check(DataIndex >= 0);
-
-	bool bDynamic = DataIndex < NumSimulatedBodies;
-	bool bKinematic = !bDynamic && DataIndex < NumKinematicBodies;
-	bool bStatic = !bDynamic && !bKinematic;
-
-	Actors.RemoveAt(DataIndex);
-
-	delete ActorHandles[DataIndex];
-	ActorHandles.RemoveAt(DataIndex);
-
-	RigidBodiesData.RemoveAt(DataIndex);
-	SolverBodiesData.RemoveAt(DataIndex);
-	PendingAcceleration.RemoveAt(DataIndex);
-	KinematicTargets.RemoveAt(DataIndex);
-
-	// Starting where we removed the actor we need to update any handles beyond it
-	const int32 NumTotalHandles = ActorHandles.Num();
-	for(int32 CurrIndex = DataIndex; CurrIndex < NumTotalHandles; ++CurrIndex)
-	{
-		ActorHandles[CurrIndex]->ActorDataIndex--;
-	}
-
-	bDirtyJointData = true;
-
-	if(bDynamic)
-	{
-		NumSimulatedBodies--;
-
-		// This is used for RBAN LOD optimizations, in normal simulations this should be identical to NumSimulatedBodies
-		NumActiveSimulatedBodies--;
-	}
-	else if(bKinematic)
-	{
-		NumKinematicBodies--;
-	}
-
-	bRecreateIterationCache = true;	//reordering of bodies so we need to change iteration order potentially
-}
-
-FActorHandle* FSimulation::CreateDynamicActor(FBodyInstance* BodyInstance, const FTransform& TM)
-{
-	PxRigidActor* PRigidActor = FPhysicsInterface::GetPxRigidActor_AssumesLocked(BodyInstance->GetPhysicsActorHandle());
-	if (PxRigidDynamic* PDynamic = PRigidActor->is<PxRigidDynamic>())
-	{
-		const int32 ActorDataIndex = CreateActor<EActorType::DynamicActor>(PDynamic, TM);
-		NumActiveSimulatedBodies = NumSimulatedBodies;
-		bDirtyJointData = true;
-		return ActorHandles[ActorDataIndex];
-	}
-	return nullptr;
-}
-
-FActorHandle* FSimulation::CreateKinematicActor(FBodyInstance* BodyInstance, const FTransform& TM)
-{
-	PxRigidActor* PRigidActor = FPhysicsInterface::GetPxRigidActor_AssumesLocked(BodyInstance->GetPhysicsActorHandle());
-	if (PxRigidBody* PRigidBody = PRigidActor->is<PxRigidBody>())
-	{
-		const int32 ActorDataIndex = CreateActor<EActorType::KinematicActor>(PRigidBody, TM);
-		return ActorHandles[ActorDataIndex];
-	}
-	return nullptr;
-}
-
-FActorHandle* FSimulation::CreateStaticActor(FBodyInstance* BodyInstance)
-{
-	if (PxRigidActor* RigidActor = FPhysicsInterface_PhysX::GetPxRigidActor_AssumesLocked(BodyInstance->ActorHandle))
-	{
-		const int32 ActorDataIndex = CreateActor<EActorType::StaticActor>(RigidActor, P2UTransform(RigidActor->getGlobalPose()));
-		return ActorHandles[ActorDataIndex];
-	}
-	return nullptr;
-}
-
-FActorHandle* FSimulation::CreateActor(EActorType ActorType, FBodyInstance* BodyInstance, const FTransform& TM)
-{
-	switch (ActorType)
-	{
-	case EActorType::StaticActor:
-		return CreateStaticActor(BodyInstance);
-	case EActorType::KinematicActor:
-		return CreateKinematicActor(BodyInstance, TM);
-	case EActorType::DynamicActor:
-		return CreateDynamicActor(BodyInstance, TM);
-	}
-	return nullptr;
-}
-
-void FSimulation::SwapActorData(int32 Actor1DataIdx, int32 Actor2DataIdx)
-{
-	check(Actors.IsValidIndex(Actor1DataIdx));
-	check(Actors.IsValidIndex(Actor2DataIdx));
-
-	Swap(Actors[Actor1DataIdx], Actors[Actor2DataIdx]);
-	Swap(ActorHandles[Actor1DataIdx], ActorHandles[Actor2DataIdx]);
-	Swap(RigidBodiesData[Actor1DataIdx], RigidBodiesData[Actor2DataIdx]);
-	Swap(SolverBodiesData[Actor1DataIdx], SolverBodiesData[Actor2DataIdx]);
-	Swap(PendingAcceleration[Actor1DataIdx], PendingAcceleration[Actor2DataIdx]);
-	Swap(KinematicTargets[Actor1DataIdx], KinematicTargets[Actor2DataIdx]);
-
-	//Update entity index on the handle
-	ActorHandles[Actor1DataIdx]->ActorDataIndex = Actor1DataIdx;
-	ActorHandles[Actor2DataIdx]->ActorDataIndex = Actor2DataIdx;
-	bDirtyJointData = true;	//reordering of bodies could lead to stale joint data
-
-	bRecreateIterationCache = true;	//reordering of bodies so we need to change iteration order potentially
-}
-
-void FSimulation::SwapJointData(int32 Joint1Idx, int32 Joint2Idx)
-{
-	check(JointData.IsValidIndex(Joint1Idx));
-	check(JointData.IsValidIndex(Joint2Idx));
-
-	Swap(JointData[Joint1Idx], JointData[Joint2Idx]);
-	Swap(Joints[Joint1Idx], Joints[Joint2Idx]);
-	Swap(JointHandles[Joint1Idx], JointHandles[Joint2Idx]);
-
-	//Update joint index on the handle
-	JointHandles[Joint1Idx]->JointDataIndex = Joint1Idx;
-	JointHandles[Joint2Idx]->JointDataIndex = Joint2Idx;
-	bDirtyJointData = true;	//reordering of joints could lead to stale joint data
-}
-
-void FSimulation::ValidateArrays() const
-{
-	check(Actors.Num() == ActorHandles.Num());
-	check(Actors.Num() == RigidBodiesData.Num());
-	check(Actors.Num() == SolverBodiesData.Num());
-	check(Actors.Num() == PendingAcceleration.Num());
-	check(Actors.Num() == KinematicTargets.Num());
-	check(Joints.Num() == JointData.Num());
-	check(Joints.Num() == JointHandles.Num());
-}
-
-/*static char*  CacheHack = nullptr;
-int32 GCacheSumHack = 0;
-
-TAutoConsoleVariable<int32> GHackEvictCache(TEXT("p.ForceCacheEvict"), 1, TEXT("blah"));
-
-void FSimulation::EvictCache()
-{
-	if(GHackEvictCache.GetValueOnAnyThread() == 0)
-	{
-		return;
-	}
-
-	if(CacheHack == nullptr)
-	{
-		CacheHack = (char*) FMemory::Malloc(1024*1024*2);
-	}
-
-	const int32 NumCacheLines = 1024*1024*2/64;
-	for(int i=0; i<NumCacheLines; ++i)
-	{
-		GCacheSumHack += CacheHack[i*64];
-	}
-	
-}*/
-
-DECLARE_DWORD_COUNTER_STAT(TEXT("Simulated Bodies"), STAT_IPNumSimulatedBodies, STATGROUP_ImmediatePhysics);
-DECLARE_DWORD_COUNTER_STAT(TEXT("Active Simulated Bodies"), STAT_IPNumActiveSimulatedBodies, STATGROUP_ImmediatePhysics);
-DECLARE_DWORD_COUNTER_STAT(TEXT("Kinematic Bodies"), STAT_IPNumKinematicBodies, STATGROUP_ImmediatePhysics);
-
-DECLARE_CYCLE_STAT(TEXT("FSimulation::Simulate"), STAT_ImmediateSimulate, STATGROUP_ImmediatePhysics);
-
-void FSimulation::Simulate(float DeltaTime, const FVector& Gravity)
-{
-	// Lock the resource manager before we simulate - any writes need to wait until we're done
-	FScopedSharedResourceLock<LockMode::Read> ScopeLock;
-	Simulate_AssumesLocked(DeltaTime, Gravity);
-}
-
-void FSimulation::Simulate_AssumesLocked(float DeltaTime, const FVector& Gravity)
-{
-	SET_DWORD_STAT(STAT_IPNumSimulatedBodies, NumSimulatedBodies);
-	SET_DWORD_STAT(STAT_IPNumActiveSimulatedBodies, NumActiveSimulatedBodies);
-	SET_DWORD_STAT(STAT_IPNumKinematicBodies, NumKinematicBodies);
-
-	DeltaTime = FMath::Min(DeltaTime, 0.033f);
-	//Create dynamic bodies and integrate their unconstrained velocities
-	if(DeltaTime > 0)
-	{
-		SCOPE_CYCLE_COUNTER(STAT_ImmediateSimulate);
-
-		++SimCount;
-
-		ConstructSolverBodies(DeltaTime, Gravity);
-	
-		if(bRecreateIterationCache)
-		{
-			PrepareIterationCache();
-		}
-		
-		GenerateContacts();
-		
-		BatchConstraints();
-
-		PrepareConstraints(DeltaTime);
-
-		SolveAndIntegrate(DeltaTime);
-
-		Workspace.Reset();
-	}
-
-	//EvictCache();
-}
-
-DECLARE_CYCLE_STAT(TEXT("ConstructSolverBodies"), STAT_ImmediateConstructSolverBodies, STATGROUP_ImmediatePhysics);
-
-void FSimulation::ConstructSolverBodies(float DeltaTime, const FVector& Gravity)
-{
-	SCOPE_CYCLE_COUNTER(STAT_ImmediateConstructSolverBodies);
-	const int32 NumBytes = Actors.Num() * sizeof(PxSolverBody);
-	SolverBodies = (PxSolverBody*)Workspace.Alloc(NumBytes);
-
-	for(int32 BodyIdx = 0; BodyIdx < NumActiveSimulatedBodies; ++BodyIdx)
-	{
-		RigidBodiesData[BodyIdx].linearVelocity += PendingAcceleration[BodyIdx] * DeltaTime;
-	}
-
-	FMemory::Memzero(PendingAcceleration.GetData(), sizeof(PendingAcceleration[0]) * NumSimulatedBodies);
-
-	FMemory::Memzero(SolverBodies, NumBytes);
-	immediate::PxConstructSolverBodies(RigidBodiesData.GetData(), SolverBodiesData.GetData(), NumActiveSimulatedBodies, U2PVector(Gravity), DeltaTime);
-	
-	const float DeltaTimeInv = 1.f / DeltaTime;
-
-	//compute velocity for kinematic targets
-	const int32 NumActors = Actors.Num();
-	for (int32 KinematicBodyIdx = NumSimulatedBodies; KinematicBodyIdx < NumActors; ++KinematicBodyIdx)
-	{
-		const FImmediateKinematicTarget& KinematicTarget = KinematicTargets[KinematicBodyIdx];
-		if(KinematicTarget.bTargetSet)
-		{
-			immediate::PxRigidBodyData& RigidBodyData = RigidBodiesData[KinematicBodyIdx];
-			const PxTransform& CurrentTM = RigidBodyData.body2World;
-			PxVec3 LinearDelta = KinematicTarget.BodyToWorld.p - CurrentTM.p;
-			PxQuat AngularDelta = KinematicTarget.BodyToWorld.q * CurrentTM.q.getConjugate();
-
-			if (AngularDelta.w < 0)	//shortest angle.
-				AngularDelta = -AngularDelta;
-
-			PxReal Angle;
-			PxVec3 Axis;
-			AngularDelta.toRadiansAndUnitAxis(Angle, Axis);
-			
-			RigidBodyData.linearVelocity = LinearDelta * DeltaTimeInv;
-			RigidBodyData.angularVelocity = Axis * Angle * DeltaTimeInv;
-		}
-	}
-
-	immediate::PxConstructSolverBodies(RigidBodiesData.GetData() + NumActiveSimulatedBodies, SolverBodiesData.GetData() + NumActiveSimulatedBodies, NumKinematicBodies + (NumSimulatedBodies - NumActiveSimulatedBodies), PxVec3(PxZero), DeltaTime);
-}
-
-DECLARE_CYCLE_STAT(TEXT("PrepareIterationCache"), STAT_ImmediatePrepareIterationCache, STATGROUP_ImmediatePhysics);
-
-void FSimulation::PrepareIterationCache()
-{
-	SCOPE_CYCLE_COUNTER(STAT_ImmediatePrepareIterationCache);
-	// Iteration cache is not yet available, so we need to iterate over bodies and prepare the data as needed
-	const int32 NumActors = Actors.Num();
-	ShapeSOA.LocalTMs.Empty(NumActors);
-	ShapeSOA.Materials.Empty(NumActors);
-	ShapeSOA.Geometries.Empty(NumActors);
-	ShapeSOA.Bounds.Empty(NumActors);
-	ShapeSOA.OwningActors.Empty(NumActors);
-	ShapeSOA.BoundsOffsets.Empty(NumActors);
-
-	int32 NumShapes = 0;
-	NumSimulatedShapesWithCollision = 0;
-
-	//flatten out all shapes of actors that have collision
-	for(int32 ActorIdx = 0; ActorIdx < NumActors; ++ActorIdx)
-	{
-		const FActor& Actor = Actors[ActorIdx];
-		if(IgnoreCollisionActors.Contains(ActorHandles[ActorIdx]))
-		{
-			if (ActorIdx+1 == NumSimulatedBodies)
-			{
-				NumSimulatedShapesWithCollision = NumShapes;
-			}
-
-			continue;
-		}
-
-		for(const FShape& Shape : Actor.Shapes)
-		{
-			ShapeSOA.LocalTMs.Add(Shape.LocalTM);
-			ShapeSOA.Materials.Add(Shape.Material ? *Shape.Material : FMaterial());
-			ShapeSOA.Geometries.Add(Shape.Geometry);
-			ShapeSOA.Bounds.Add(Shape.BoundsMagnitude);
-			ShapeSOA.BoundsOffsets.Add(Shape.BoundsOffset);
-			ShapeSOA.OwningActors.Add(ActorIdx);
-			++NumShapes;
-		}
-
-		if (ActorIdx+1 == NumSimulatedBodies)
-		{
-			NumSimulatedShapesWithCollision = NumShapes;
-		}
-	}
-
-	//build cache of which shapes to skip
-	int32 IterationCount = 0;
-	SkipCollisionCache.Empty(NumShapes);
-
-	for(int32 SimShapeIdx = 0; SimShapeIdx < NumSimulatedShapesWithCollision; ++SimShapeIdx)
-	{
-		const int32 SimActorIdx = ShapeSOA.OwningActors[SimShapeIdx];
-		const TSet<FActorHandle*>* IgnoreActorsForSimulated = IgnoreCollisionPairTable.Find(ActorHandles[SimActorIdx]);
-
-		for(int32 OtherShapeIdx = SimShapeIdx + 1; OtherShapeIdx < NumShapes; ++OtherShapeIdx)
-		{
-			const int32 OtherActorIdx = ShapeSOA.OwningActors[OtherShapeIdx];
-
-			if(SimActorIdx == OtherActorIdx || (IgnoreActorsForSimulated && IgnoreActorsForSimulated->Find(ActorHandles[OtherActorIdx])))
-			{
-				SkipCollisionCache.Add(IterationCount++);
-			}
-			else
-			{
-				++IterationCount;
-			}
-		}
-	}
-
-#if PERSISTENT_CONTACT_PAIRS
-	const int32 NumPairs = (NumShapes) * (NumShapes - 1) / 2;
-	ShapeSOA.ContactPairData.Empty(NumPairs);
-	ShapeSOA.ContactPairData.AddZeroed(NumPairs);
-#endif
-
-	bRecreateIterationCache = false;
-}
-
-DECLARE_CYCLE_STAT(TEXT("GenerateContacts"), STAT_ImmediateGenerateContacts, STATGROUP_ImmediatePhysics);
-
-void FSimulation::GenerateContacts()
-{
-	SCOPE_CYCLE_COUNTER(STAT_ImmediateGenerateContacts);
-
-	//narrow phase, find actual contact points
-	CacheAllocator.Reset();
-
-	ContactPairs.Reset();
-	ContactPoints.Reset();
-	
-	const int32 NumShapes = ShapeSOA.LocalTMs.Num();
-	PxTransform* ShapeWorldTMs = (PxTransform*)Workspace.Alloc(sizeof(PxTransform) * NumShapes);
-	for(int32 ShapeIdx = 0; ShapeIdx < NumShapes; ++ShapeIdx)
-	{
-		const immediate::PxRigidBodyData& Body = RigidBodiesData[ShapeSOA.OwningActors[ShapeIdx]];
-		ShapeWorldTMs[ShapeIdx] = Body.body2World * ShapeSOA.LocalTMs[ShapeIdx];
-	}
-
-	int32 IterationCount = 0;
-	int32 PotentialPairCount = 0;
-	const int32 IterationCacheSize = SkipCollisionCache.Num();
-	int32 CurrentIterationCacheIdx = 0;
-	const int32* SkipCollisionCacheRaw = SkipCollisionCache.GetData();
-
-	for(int32 SimShapeIdx = 0; SimShapeIdx < NumSimulatedShapesWithCollision; ++SimShapeIdx)
-	{
-		const float SimulatedRadius = ShapeSOA.Bounds[SimShapeIdx];
-		const PxTransform& SimulatedShapeTM = ShapeWorldTMs[SimShapeIdx];
-		const PxVec3 SimulatedBoundsOffset = ShapeSOA.BoundsOffsets[SimShapeIdx];
-		const PxVec3 SimulatedShapeBoundsOrigin = SimulatedShapeTM.transform(SimulatedBoundsOffset);
-		
-		const int32 SimulatedActorIdx = ShapeSOA.OwningActors[SimShapeIdx];
-		const FMaterial& SimulatedShapeMaterial = ShapeSOA.Materials[SimShapeIdx];
-		const PxGeometry* SimulatedGeometry = ShapeSOA.Geometries[SimShapeIdx];
-
-		if(SimulatedActorIdx >= NumActiveSimulatedBodies)
-		{
-			break;
-		}
-
-		for(int32 OtherShapeIdx = SimShapeIdx + 1; OtherShapeIdx < NumShapes; ++OtherShapeIdx)
-		{
-			if(CurrentIterationCacheIdx < IterationCacheSize && SkipCollisionCache[CurrentIterationCacheIdx] == IterationCount++)
-			{
-				++CurrentIterationCacheIdx;
-				continue;
-			}
-
-			const int32 OtherActorIdx = ShapeSOA.OwningActors[OtherShapeIdx];
-			if(OtherActorIdx >= NumActiveSimulatedBodies && OtherActorIdx < NumSimulatedBodies)
-			{
-				continue;
-			}
-
-			PxCache* InCache;
-			const int32 PotentialPairIdx = PotentialPairCount++;
-#if PERSISTENT_CONTACT_PAIRS
-			FPersistentContactPairData& PersistentPairData = ShapeSOA.ContactPairData[PotentialPairIdx];
-			InCache = &PersistentPairData.Cache;
-
-			if(PersistentPairData.SimCount + 1 != SimCount)	//Note: in the case where we roll over to 0 does this matter? Could use % but seems unneeded for now
-			{
-				PersistentPairData.Clear();
-			}
-#else
-			PxCache Cache;
-			InCache = &Cache;
-#endif
-			const FMaterial& OtherShapeMaterial = ShapeSOA.Materials[OtherShapeIdx];
-			const float OtherRadius = ShapeSOA.Bounds[OtherShapeIdx];
-			const float TotalRadius = (SimulatedRadius + OtherRadius);
-			const float TotalRadius2 = TotalRadius * TotalRadius;
-			const PxTransform& OtherShapeTM = ShapeWorldTMs[OtherShapeIdx];
-			const PxVec3 OtherShapeBoundsOrigin = OtherShapeTM.transform(ShapeSOA.BoundsOffsets[OtherShapeIdx]);
-
-			const float Distance2 = (SimulatedShapeBoundsOrigin  - OtherShapeBoundsOrigin).magnitudeSquared();
-
-			if ( Distance2 > TotalRadius2)
-			{
-#if PERSISTENT_CONTACT_PAIRS
-				PersistentPairData.Clear();
-#endif
-				continue;	//no intersection so skip narrow phase
-			}
-
-			FContactPointRecorder ContactRecorder(*this, SimulatedActorIdx, OtherActorIdx, PotentialPairIdx, SimulatedShapeMaterial, OtherShapeMaterial);
-			if (!immediate::PxGenerateContacts(&SimulatedGeometry, &ShapeSOA.Geometries[OtherShapeIdx], &SimulatedShapeTM, &OtherShapeTM, InCache, 1, ContactRecorder, 4.f, 1.f, 100.f, CacheAllocator))
-			{
-#if PERSISTENT_CONTACT_PAIRS
-				//no contacts found so clear friction
-				PersistentPairData.Frictions = nullptr;
-				PersistentPairData.NumFrictions = 0;
-#endif
-			}
-		}
-	}
-}
-
-bool GBatchJoints = false;	//TODO: temp hack to fix crash in immediate mode due to joint batching
-
-DECLARE_CYCLE_STAT(TEXT("BatchConstraints"), STAT_ImmediateBatchConstraints, STATGROUP_ImmediatePhysics);
-
-void FSimulation::BatchConstraints()
-{
-	SCOPE_CYCLE_COUNTER(STAT_ImmediateBatchConstraints);
-	
-	const int32 NumContactPairs = ContactPairs.Num();
-	const int32 NumJoints = JointData.Num();
-	const int32 NumConstraints = NumJoints + NumContactPairs;
-	OrderedDescriptors.SetNum(NumConstraints); //Ordered descriptors are semi-persistent. We only reorder joint constraints when needed, but we always update contact constraints which come at the end of the array.
-	BatchHeaders.SetNum(NumConstraints);	   //batch headers follow the same approach
-	
-	// Joint constraints if needed
-	if (bDirtyJointData)
-	{
-		NumActiveJoints = 0;
-		if(NumJoints > 0)
-		{
-			PxSolverConstraintDesc* JointDescriptors = (PxSolverConstraintDesc*)Workspace.Alloc(sizeof(PxSolverConstraintDesc) * NumJoints);
-
-			for (int32 JointIdx = 0; JointIdx < NumJoints; ++JointIdx)
-			{
-				PxSolverConstraintDesc& JointDescriptor = JointDescriptors[NumActiveJoints];
-
-				const FJoint& Joint = Joints[JointIdx];
-				const int32 DynamicActorIdx = Joint.DynamicActor->ActorDataIndex;
-				const int32 OtherActorIdx = Joint.OtherActor != nullptr ? Joint.OtherActor->ActorDataIndex : INDEX_NONE;
-
-				if( (DynamicActorIdx >= NumActiveSimulatedBodies && DynamicActorIdx < NumSimulatedBodies) || (OtherActorIdx >= NumActiveSimulatedBodies && OtherActorIdx < NumSimulatedBodies) )
-				{
-					continue;
-				}
-
-				JointDescriptor.bodyA = &SolverBodies[DynamicActorIdx];
-				JointDescriptor.bodyB = OtherActorIdx != INDEX_NONE ? &SolverBodies[OtherActorIdx] : nullptr;
-
-				JointDescriptor.bodyADataIndex = PxU16(DynamicActorIdx);
-				JointDescriptor.bodyBDataIndex = PxU16(OtherActorIdx);
-				JointDescriptor.linkIndexA = PxSolverConstraintDesc::NO_LINK;
-				JointDescriptor.linkIndexB = PxSolverConstraintDesc::NO_LINK;
-				JointDescriptor.writeBack = nullptr;
-
-				JointDescriptor.constraint = (PxU8*)JointHandles[JointIdx];
-				JointDescriptor.constraintLengthOver16 = PxSolverConstraintDesc::eJOINT_CONSTRAINT;
-
-				++NumActiveJoints;
-			}
-
-			if(GBatchJoints)
-			{
-				NumJointHeaders = physx::immediate::PxBatchConstraints(JointDescriptors, NumActiveJoints, SolverBodies, Actors.Num(), BatchHeaders.GetData(), OrderedDescriptors.GetData());
-			}
-			else
-			{
-				NumJointHeaders = NumActiveJoints;
-				OrderedDescriptors.AddUninitialized(NumJointHeaders);
-				FMemory::Memcpy(OrderedDescriptors.GetData(), JointDescriptors, sizeof(OrderedDescriptors[0]) * NumJointHeaders);
-
-				for(int32 Offset = 0; Offset < NumActiveJoints; ++Offset)
-				{
-					PxConstraintBatchHeader& BatchHeader = BatchHeaders[Offset];
-					BatchHeader.mStartIndex = Offset;
-					BatchHeader.mStride = 1;
-				}
-			}
-			
-
-			//ensure D6 joint data and low level joint data are matching the Ordered descriptors
-			for(int32 DescriptorIdx = 0; DescriptorIdx < NumActiveJoints; ++DescriptorIdx)
-			{
-				FJointHandle* JointHandle = (FJointHandle*) OrderedDescriptors[DescriptorIdx].constraint;
-				if(JointHandle->JointDataIndex != DescriptorIdx)
-				{
-					//The joint or the descriptor have moved - either way fix it
-					SwapJointData(DescriptorIdx, JointHandle->JointDataIndex);
-				}
-			}
-		}
-		else
-		{
-			NumJointHeaders = 0;
-		}
-	}
-
-	if(NumContactPairs > 0)
-	{
-		//Contact constraints
-		PxSolverConstraintDesc* ContactDescriptors = (PxSolverConstraintDesc*)Workspace.Alloc(sizeof(PxSolverConstraintDesc) * NumContactPairs);
-		for (int32 ContactPairIdx = 0; ContactPairIdx < NumContactPairs; ++ContactPairIdx)
-		{
-			const FContactPair& ContactPair = ContactPairs[ContactPairIdx];
-			PxSolverConstraintDesc& ContactDescriptor = ContactDescriptors[ContactPairIdx];
-
-			//Set body pointers and bodyData idxs
-			ContactDescriptor.bodyA = &SolverBodies[ContactPair.DynamicActorDataIndex];
-			ContactDescriptor.bodyB = &SolverBodies[ContactPair.OtherActorDataIndex];
-
-			ContactDescriptor.bodyADataIndex = PxU16(ContactPair.DynamicActorDataIndex);
-			ContactDescriptor.bodyBDataIndex = PxU16(ContactPair.OtherActorDataIndex);
-			ContactDescriptor.linkIndexA = PxSolverConstraintDesc::NO_LINK;
-			ContactDescriptor.linkIndexB = PxSolverConstraintDesc::NO_LINK;
-			ContactDescriptor.writeBack = nullptr;
-
-			ContactDescriptor.constraint = (PxU8*)(&ContactPair);	//Need to save this off for later when the constraint descriptors get re-ordered
-			ContactDescriptor.constraintLengthOver16 = PxSolverConstraintDesc::eCONTACT_CONSTRAINT;
-		}
-
-		NumContactHeaders = physx::immediate::PxBatchConstraints(ContactDescriptors, NumContactPairs, SolverBodies, Actors.Num(), BatchHeaders.GetData() + NumJointHeaders, OrderedDescriptors.GetData() + NumActiveJoints);
-	}
-	else
-	{
-		NumContactHeaders = 0;
-	}
-}
-
-//float GReadCacheHack = 0.f;
-
-DECLARE_CYCLE_STAT(TEXT("PrepareConstraints"), STAT_ImmediatePrepareConstraints, STATGROUP_ImmediatePhysics);
-
-void FSimulation::PrepareConstraints(float DeltaTime)
-{
-	SCOPE_CYCLE_COUNTER(STAT_ImmediatePrepareConstraints);
-
-	//Prepares constraint data based on latest position data and contact pairs
-	const int32 NumBodies = Actors.Num();
-	const float InvDeltaTime = 1.f / DeltaTime;
-	ConstraintAllocator.Reset();
-
-	PxSolverConstraintPrepDesc JointConstraints[4];
-	Px1DConstraint Rows[24];	//24 rows because we have at most 6 rows per joint and we could be solving at most 4 joints in one batch
-
-	/*FPlatformMisc::Prefetch(OrderedDescriptors.GetData());
-	FPlatformMisc::Prefetch(BatchHeaders.GetData());
-	FPlatformMisc::Prefetch(SolverBodiesData.GetData());
-	FPlatformMisc::Prefetch(JointData.GetData());*/
-
-	const int32 NumMobileBodies = NumSimulatedBodies + NumKinematicBodies;
-	
-	//Joint constraints
-	for (PxU32 HeaderIdx = 0; HeaderIdx < NumJointHeaders; ++HeaderIdx)
-	{
-		PxConstraintBatchHeader& Header = BatchHeaders[HeaderIdx];
-		Header.mConstraintType = PxSolverConstraintDesc::eJOINT_CONSTRAINT;	//Since joint headers are cached and solver re-uses this we have to reset it every time
-
-		FMemory::Memzero(Rows, sizeof(Rows));	//TODO: do we actually need to do this?
-		for (Px1DConstraint& Constraint : Rows)
-		{
-			Constraint.minImpulse = -FLT_MAX;
-			Constraint.maxImpulse = FLT_MAX;
-		}
-
-		int32 CurRows = 0;
-
-		for (PxU32 BatchInnerIdx = 0; BatchInnerIdx < Header.mStride; ++BatchInnerIdx)
-		{
-			const PxU32 DescriptorIdx = Header.mStartIndex + BatchInnerIdx;
-			//FPlatformMisc::PrefetchBlock(&JointData[DescriptorIdx], 128);
-			//GReadCacheHack = JointData[DescriptorIdx].invMassScale.angular0;
-
-			PxSolverConstraintDesc& OrderedDescriptor = OrderedDescriptors[DescriptorIdx];
-			PxSolverConstraintPrepDesc& JointDescriptor = JointConstraints[BatchInnerIdx];
-
-			//TODO: is it worth caching some of this? Most of the data doesn't change each frame
-			JointDescriptor.body0 = OrderedDescriptor.bodyA;
-			JointDescriptor.body1 = OrderedDescriptor.bodyB;
-			JointDescriptor.data0 = &SolverBodiesData[OrderedDescriptor.bodyADataIndex];
-			JointDescriptor.data1 = &SolverBodiesData[OrderedDescriptor.bodyBDataIndex];
-
-			JointDescriptor.bodyFrame0 = JointDescriptor.data0->body2World;
-			JointDescriptor.bodyFrame1 = JointDescriptor.data1->body2World;
-
-			JointDescriptor.bodyState0 = PxSolverConstraintPrepDescBase::eDYNAMIC_BODY;
-			JointDescriptor.bodyState1 = JointDescriptor.body1 == nullptr ? PxSolverConstraintPrepDescBase::eSTATIC_BODY : (OrderedDescriptor.bodyBDataIndex < (uint32)NumSimulatedBodies ? PxSolverConstraintPrepDescBase::eDYNAMIC_BODY : PxSolverConstraintPrepDescBase::eSTATIC_BODY);
-			JointDescriptor.desc = &OrderedDescriptor;
-			JointDescriptor.mInvMassScales.angular0 = JointDescriptor.mInvMassScales.angular1 = JointDescriptor.mInvMassScales.linear0 = JointDescriptor.mInvMassScales.linear1 = 1.f;
-			JointDescriptor.writeback = NULL;
-			//Constraint.getBreakForce(JointDescriptor.linBreakForce, JointDescriptor.angBreakForce);
-			JointDescriptor.linBreakForce = FLT_MAX;
-			JointDescriptor.angBreakForce = FLT_MAX;
-
-			JointDescriptor.minResponseThreshold = 0;//Constraint.getMinResponseThreshold();
-			JointDescriptor.disablePreprocessing = false;//!!(Constraint.getFlags() & PxConstraintFlag::eDISABLE_PREPROCESSING);
-			JointDescriptor.improvedSlerp = true;//!!(Constraint.getFlags() & PxConstraintFlag::eIMPROVED_SLERP);
-			JointDescriptor.driveLimitsAreForces = false;//!!(Constraint.getFlags() & PxConstraintFlag::eDRIVE_LIMITS_ARE_FORCES);
-
-			PxVec3 Body0WorldOffset;
-			PxConstraintInvMassScale InvMassScale;
-			const D6JointData& Joint = JointData[DescriptorIdx];	
-			int32 NumRows = D6JointSolverPrep(&Rows[CurRows], Body0WorldOffset, 0, InvMassScale, (const void*)&Joint, JointDescriptor.bodyFrame0, JointDescriptor.bodyFrame1);
-
-			JointDescriptor.mInvMassScales = InvMassScale;
-
-			JointDescriptor.numRows = NumRows;
-			JointDescriptor.rows = &Rows[CurRows];
-			CurRows += NumRows;
-		}
-
-		immediate::PxCreateJointConstraints(&Header, 1, JointConstraints, ConstraintAllocator, DeltaTime, InvDeltaTime);
-	}
-
-	bDirtyJointData = false;
-
-	//Contact constraints
-	const PxU32 NumTotalHeaders = NumContactHeaders + NumJointHeaders;
-	PxSolverContactDesc ContactDescriptors[4];
-	const FContactPair* InnerContactPairs[4];
-	const PxU32 NumJoints = JointData.Num();
-
-	for (PxU32 HeaderIdx = NumJointHeaders; HeaderIdx < NumTotalHeaders; ++HeaderIdx)
-	{
-		PxConstraintBatchHeader& Header = BatchHeaders[HeaderIdx];
-		check(Header.mConstraintType == PxSolverConstraintDesc::eCONTACT_CONSTRAINT);
-
-		Header.mStartIndex += NumActiveJoints;
-		
-		for (PxU32 BatchInnerIdx = 0; BatchInnerIdx < Header.mStride; ++BatchInnerIdx)
-		{
-			PxSolverConstraintDesc& OrderedDescriptor = OrderedDescriptors[Header.mStartIndex + BatchInnerIdx];
-			PxSolverContactDesc& ContactDescriptor = ContactDescriptors[BatchInnerIdx];
-			//Extract the contact pair that we saved in this structure earlier.
-			const FContactPair& ContactPair = *reinterpret_cast<const FContactPair*>(OrderedDescriptor.constraint);
-			InnerContactPairs[BatchInnerIdx] = &ContactPair;
-
-			ContactDescriptor.body0 = OrderedDescriptor.bodyA;
-			ContactDescriptor.body1 = OrderedDescriptor.bodyB;
-			ContactDescriptor.data0 = &SolverBodiesData[OrderedDescriptor.bodyADataIndex];
-			ContactDescriptor.data1 = &SolverBodiesData[OrderedDescriptor.bodyBDataIndex];
-
-			ContactDescriptor.bodyFrame0 = ContactDescriptor.data0->body2World;
-			ContactDescriptor.bodyFrame1 = ContactDescriptor.data1->body2World;
-
-			ContactDescriptor.contactForces = nullptr;//&ContactForces[ContactPair.StartContactIndex];
-			ContactDescriptor.contacts = &ContactPoints[ContactPair.StartContactIndex];
-			ContactDescriptor.numContacts = ContactPair.NumContacts;
-
-#if  PERSISTENT_CONTACT_PAIRS
-			FPersistentContactPairData& PersistentPairData = ShapeSOA.ContactPairData[ContactPair.PairIdx];
-
-			ContactDescriptor.frictionPtr = PersistentPairData.Frictions;
-			ContactDescriptor.frictionCount = PersistentPairData.NumFrictions;
-#else
-			ContactDescriptor.frictionPtr = nullptr;
-			ContactDescriptor.frictionCount = 0;
-#endif
-
-			ContactDescriptor.disableStrongFriction = false;
-			ContactDescriptor.hasMaxImpulse = false;
-			ContactDescriptor.hasForceThresholds = false;
-			ContactDescriptor.shapeInteraction = NULL;
-			ContactDescriptor.restDistance = 0.f;
-			ContactDescriptor.maxCCDSeparation = PX_MAX_F32;
-
-			ContactDescriptor.bodyState0 = PxSolverConstraintPrepDescBase::eDYNAMIC_BODY;
-			ContactDescriptor.bodyState1 = ContactPair.OtherActorDataIndex < (uint32)NumSimulatedBodies ? PxSolverConstraintPrepDescBase::eDYNAMIC_BODY : PxSolverConstraintPrepDescBase::eSTATIC_BODY;
-			ContactDescriptor.desc = &OrderedDescriptor;
-			ContactDescriptor.mInvMassScales.angular0 = ContactDescriptor.mInvMassScales.angular1 = ContactDescriptor.mInvMassScales.linear0 = ContactDescriptor.mInvMassScales.linear1 = 1.f;
-		}
-
-		immediate::PxCreateContactConstraints(&Header, 1, &ContactDescriptors[0], ConstraintAllocator, InvDeltaTime, -200.f, 4.f, 1.f);
-
-#if PERSISTENT_CONTACT_PAIRS
-		// Now cache friction results
-		for (PxU32 BatchInnerIdx = 0; BatchInnerIdx < Header.mStride; ++BatchInnerIdx)
-		{
-			PxSolverContactDesc& ContactDescriptor = ContactDescriptors[BatchInnerIdx];
-
-			//Extract the contact pair that we saved in this structure earlier.
-			const FContactPair& ContactPair = *InnerContactPairs[BatchInnerIdx];
-			FPersistentContactPairData& PersistentPairData = ShapeSOA.ContactPairData[ContactPair.PairIdx];
-
-			PersistentPairData.Frictions = ContactDescriptor.frictionPtr;
-			PersistentPairData.NumFrictions = ContactDescriptor.frictionCount;
-		}
-#endif
-	}
-}
-
-DECLARE_CYCLE_STAT(TEXT("SolveAndIntegrate"), STAT_ImmediateSolveAndIntegrate, STATGROUP_ImmediatePhysics);
-
-
-void FSimulation::SolveAndIntegrate(float DeltaTime)
-{
-	SCOPE_CYCLE_COUNTER(STAT_ImmediateSolveAndIntegrate);
-
-	PxVec3* LinearMotionVelocity = (PxVec3*) Workspace.Alloc(sizeof(PxVec3) * NumActiveSimulatedBodies * 2);
-	PxVec3* AngularMotionVelocity = &LinearMotionVelocity[NumActiveSimulatedBodies];
-
-	//Solve all constraints
-	immediate::PxSolveConstraints(BatchHeaders.GetData(), NumContactHeaders + NumJointHeaders, OrderedDescriptors.GetData(), SolverBodies, LinearMotionVelocity, AngularMotionVelocity, NumActiveSimulatedBodies, NumPositionIterations, NumVelocityIterations);
-
-	//Integrate velocities
-	immediate::PxIntegrateSolverBodies(SolverBodiesData.GetData(), SolverBodies, LinearMotionVelocity, AngularMotionVelocity, NumActiveSimulatedBodies, DeltaTime);
-
-	//Copy positions of dynamic bodies back into rigid body
-	for (int32 DynamicsBodyIdx = 0; DynamicsBodyIdx < NumActiveSimulatedBodies; ++DynamicsBodyIdx)
-	{
-		RigidBodiesData[DynamicsBodyIdx].linearVelocity = SolverBodiesData[DynamicsBodyIdx].linearVelocity;
-		RigidBodiesData[DynamicsBodyIdx].angularVelocity = SolverBodiesData[DynamicsBodyIdx].angularVelocity;
-		RigidBodiesData[DynamicsBodyIdx].body2World = SolverBodiesData[DynamicsBodyIdx].body2World;
-	}
-
-	//Copy kinematic targets
-	const int32 NumActors = Actors.Num();
-	for(int32 KinematicBodyIdx = NumSimulatedBodies; KinematicBodyIdx < NumActors; ++KinematicBodyIdx)
-	{
-		if(KinematicTargets[KinematicBodyIdx].bTargetSet)
-		{
-			RigidBodiesData[KinematicBodyIdx].body2World = KinematicTargets[KinematicBodyIdx].BodyToWorld;
-			RigidBodiesData[KinematicBodyIdx].linearVelocity = PxVec3(PxZero);
-			RigidBodiesData[KinematicBodyIdx].angularVelocity = PxVec3(PxZero);
-		}
-	}
-
-	FMemory::Memzero(KinematicTargets.GetData(), NumActors * sizeof(KinematicTargets[0]));
-}
-
-void FSimulation::AddRadialForce(int32 ActorDataIndex, const FVector& InOrigin, float Strength, float Radius, ERadialImpulseFalloff Falloff, EForceType ForceType)
-{
-	if(IsSimulated(ActorDataIndex))
-	{
-		PxVec3 Origin = U2PVector(InOrigin);
-		immediate::PxRigidBodyData& LowLevelBody = GetLowLevelBody(ActorDataIndex);
-		const PxTransform& PCOMTransform = LowLevelBody.body2World;
-		PxVec3 Delta = LowLevelBody.body2World.p - Origin;
-
-		float Mag = Delta.magnitude(); // Distance from COM to origin
-		if (Mag > Radius)
-		{
-			return;
-		}
-
-		Delta.normalize();
-
-		// Scale by U2PScale here, because units are velocity * mass. 
-		float ImpulseMag = Strength;
-		if (Falloff == RIF_Linear)
-		{
-			ImpulseMag *= (1.0f - (Mag / Radius));
-		}
-
-		const PxVec3 PImpulse = Delta * ImpulseMag;
-		const PxVec3 ApplyDelta = (ForceType == EForceType::AddAcceleration || ForceType == EForceType::AddVelocity) ? PImpulse : PImpulse * LowLevelBody.invMass;
-		
-		if(ForceType == EForceType::AddImpulse || ForceType == EForceType::AddVelocity)
-		{
-			LowLevelBody.linearVelocity += ApplyDelta;
-		}
-		else
-		{
-			PendingAcceleration[ActorDataIndex] += ApplyDelta;
-		}
-	}
-}
-
-void FSimulation::AddForce(int32 ActorDataIndex, const FVector& Force)
-{
-	if (IsSimulated(ActorDataIndex))
-	{
-		immediate::PxRigidBodyData& LowLevelBody = GetLowLevelBody(ActorDataIndex);
-		PendingAcceleration[ActorDataIndex] += U2PVector(Force) * LowLevelBody.invMass;
-	}
-}
-
-FActorData CreateActorData(const FActorCreationParams& InParams)
-{
-	FActorData OutData;
-
-	immediate::PxRigidBodyData& RigidData = OutData.RigidBodyData;
-
-	RigidData.linearVelocity = PxVec3(PxZero);
-	RigidData.invMass = 0.0f;
-	RigidData.angularVelocity = PxVec3(PxZero);
-	RigidData.maxDepenetrationVelocity = PX_MAX_F32;
-	RigidData.invInertia = PxVec3(PxZero);
-	RigidData.maxContactImpulse = PX_MAX_F32;
-	RigidData.linearDamping = 0.0f;
-	RigidData.angularDamping = 0.0f;
-	RigidData.maxLinearVelocitySq = PX_MAX_F32;
-	RigidData.maxAngularVelocitySq = PX_MAX_F32;
-	RigidData.body2World = U2PTransform(InParams.InitialTM);
-
-	OutData.InitialTransform = InParams.InitialTM;
-	OutData.bStatic = InParams.bStatic;
-	OutData.bKinematic = false; //#PHYS2 Add kinematic support
-
-	return OutData;
-}
-
-}
-
-#endif // WITH_PHYSX
\ No newline at end of file
diff --git a/Engine/Source/Runtime/Engine/Private/PhysicsEngine/ImmediatePhysics/ImmediatePhysicsPhysX/ResourceManager_PhysX.cpp b/Engine/Source/Runtime/Engine/Private/PhysicsEngine/ImmediatePhysics/ImmediatePhysicsPhysX/ResourceManager_PhysX.cpp
--- a/Engine/Source/Runtime/Engine/Private/PhysicsEngine/ImmediatePhysics/ImmediatePhysicsPhysX/ResourceManager_PhysX.cpp
+++ b/Engine/Source/Runtime/Engine/Private/PhysicsEngine/ImmediatePhysics/ImmediatePhysicsPhysX/ResourceManager_PhysX.cpp
@@ -1,52 +1,3 @@
 // Copyright Epic Games, Inc. All Rights Reserved.
 
 #include "Physics/ImmediatePhysics/ImmediatePhysicsPhysX/ResourceManager_PhysX.h"
-
-#if WITH_PHYSX && PHYSICS_INTERFACE_PHYSX
-
-namespace ImmediatePhysics_PhysX
-{
-
-	FSharedResourceManager FSharedResourceManager::Instance;
-
-	FResourceHandle FSharedResourceManager::CreateMaterial()
-	{
-		int32 Index = Materials.Add(TResourceWithId<FMaterial>());
-		int32 Id = Materials[Index].Id;
-
-		return FResourceHandle(EResourceType::Material, Index, Id);
-	}
-
-	void FSharedResourceManager::ReleaseMaterial(int32 InIndex)
-	{
-		Materials.RemoveAt(InIndex);
-	}
-
-	int32 FSharedResourceManager::GetMaterialId(int32 InIndex)
-	{
-		if (InIndex != INDEX_NONE && Materials.IsAllocated(InIndex))
-		{
-			return Materials[InIndex].Id;
-		}
-
-		return INDEX_NONE;
-	}
-
-	FMaterial* FSharedResourceManager::GetMaterial(int32 InIndex)
-	{
-		if (InIndex != INDEX_NONE && Materials.IsAllocated(InIndex))
-		{
-			return &Materials[InIndex].Resource;
-		}
-
-		return nullptr;
-	}
-
-	FRWLock& FSharedResourceManager::GetLockObject()
-	{
-		return ResourceLock;
-	}
-
-}
-
-#endif // WITH_PHYSX
\ No newline at end of file
diff --git a/Engine/Source/Runtime/Engine/Private/PhysicsEngine/PhysAnim.cpp b/Engine/Source/Runtime/Engine/Private/PhysicsEngine/PhysAnim.cpp
--- a/Engine/Source/Runtime/Engine/Private/PhysicsEngine/PhysAnim.cpp
+++ b/Engine/Source/Runtime/Engine/Private/PhysicsEngine/PhysAnim.cpp
@@ -19,9 +19,6 @@
 #include "PhysicsPublic.h"
 #include "Rendering/SkeletalMeshRenderData.h"
 #include "Physics/PhysicsInterfaceCore.h"
-#if WITH_PHYSX
-	#include "PhysXPublic.h"
-#endif // WITH_PHYSX
 #include "PhysicsEngine/PhysicsConstraintTemplate.h"
 #include "PhysicsEngine/PhysicsAsset.h"
 
@@ -568,7 +565,6 @@ void USkeletalMeshComponent::UpdateKinematicBonesToAnim(const TArray<FTransform>
 #endif
 			const int32 NumComponentSpaceTransforms = GetNumComponentSpaceTransforms();
 			const int32 NumBodies = Bodies.Num();
-#if WITH_PHYSX
 
 			// Lock the scenes we need (flags set in InitArticulated)
 			FPhysicsCommand::ExecuteWrite(this, [&]()
@@ -674,9 +670,6 @@ void USkeletalMeshComponent::UpdateKinematicBonesToAnim(const TArray<FTransform>
 					}
 				}
 			});
-
-#endif // WITH_PHYSX
-
 		}
 	}
 	else
diff --git a/Engine/Source/Runtime/Engine/Private/PhysicsEngine/PhysCommandHandler.cpp b/Engine/Source/Runtime/Engine/Private/PhysicsEngine/PhysCommandHandler.cpp
--- a/Engine/Source/Runtime/Engine/Private/PhysicsEngine/PhysCommandHandler.cpp
+++ b/Engine/Source/Runtime/Engine/Private/PhysicsEngine/PhysCommandHandler.cpp
@@ -6,10 +6,6 @@
 #include "PhysicsPublic.h"
 #include "Physics/PhysicsInterfaceCore.h"
 
-#if PHYSICS_INTERFACE_PHYSX
-	#include "PhysicsEngine/PhysXSupport.h"
-#endif
-
 FPhysCommandHandler::~FPhysCommandHandler()
 {
 	if(PendingCommands.Num() != 0)
@@ -37,89 +33,33 @@ void FPhysCommandHandler::ExecuteCommands()
 		const FPhysPendingCommand& Command = PendingCommands[i];
 		switch (Command.CommandType)
 		{
-#if WITH_APEX
-		case PhysCommand::Release:
-		{
-			nvidia::apex::ApexInterface * ApexInterface = Command.Pointer.ApexInterface;
-			ApexInterface->release();
-			break;
-		}
-#endif
-#if WITH_PHYSX
 		case PhysCommand::ReleasePScene:
 		{
-#if PHYSICS_INTERFACE_PHYSX
-			physx::PxScene * PScene = Command.Pointer.PScene;
-			PScene->release();
-#endif
 			break;
 		}
 		case PhysCommand::DeleteSimEventCallback:
 		{
-#if !WITH_CHAOS // Chaos doesn't implement these callbacks
-			physx::PxSimulationEventCallback * SimEventCallback = Command.Pointer.SimEventCallback;
-			if (FPhysScene::SimEventCallbackFactory.IsValid())
-			{
-				FPhysScene::SimEventCallbackFactory->Destroy(SimEventCallback);
-			}
-			else
-			{
-				delete SimEventCallback;
-			}
-#endif
 			break;
 		}
 		case PhysCommand::DeleteContactModifyCallback:
 		{
-#if !WITH_CHAOS // Chaos doesn't implement these callbacks
-			FContactModifyCallback* ContactModifyCallback = Command.Pointer.ContactModifyCallback;
-			if (FPhysScene::ContactModifyCallbackFactory.IsValid())
-			{
-				FPhysScene::ContactModifyCallbackFactory->Destroy(ContactModifyCallback);
-			}
-			else
-			{
-				delete ContactModifyCallback;
-			}
-#endif
 			break;
 		}
 
 		case PhysCommand::DeleteCCDContactModifyCallback:
 		{
-#if !WITH_CHAOS // Chaos doesn't implement these callbacks
-			FCCDContactModifyCallback* CCDContactModifyCallback = Command.Pointer.CCDContactModifyCallback;
-
-			if (FPhysScene::CCDContactModifyCallbackFactory.IsValid())
-			{
-				FPhysScene::CCDContactModifyCallbackFactory->Destroy(CCDContactModifyCallback);
-			}
-			else
-			{
-				delete CCDContactModifyCallback;
-			}
-#endif
 			break;
 		}
 
 		case PhysCommand::DeleteCPUDispatcher:
 		{
-#if PHYSICS_INTERFACE_PHYSX
-			physx::PxCpuDispatcher * CPUDispatcher = Command.Pointer.CPUDispatcher;
-			delete CPUDispatcher;
-#endif
 			break;
 		}
 
 		case PhysCommand::DeleteMbpBroadphaseCallback:
 		{
-#if PHYSICS_INTERFACE_PHYSX
-			FPhysXMbpBroadphaseCallback* Callback = Command.Pointer.MbpCallback;
-			delete Callback;
-#endif
 			break;
 		}
-#endif
 		case PhysCommand::Max:	//this is just here because right now all commands are APEX and having a switch with only default is bad
 		default:
 			check(0);	//Unsupported command
@@ -132,78 +72,4 @@ void FPhysCommandHandler::EnqueueCommand(const FPhysPendingCommand& Command)
 {
 	check(IsInGameThread());
 	PendingCommands.Add(Command);
-}
-
-
-#if PHYSICS_INTERFACE_PHYSX
-void FPhysCommandHandler::DeferredRelease(physx::PxScene* PScene)
-{
-	check(PScene);
-
-	FPhysPendingCommand Command;
-	Command.Pointer.PScene = PScene;
-	Command.CommandType = PhysCommand::ReleasePScene;
-
-	EnqueueCommand(Command);
-
-}
-
-void  FPhysCommandHandler::DeferredDeleteContactModifyCallback(FContactModifyCallback* ContactModifyCallback)
-{
-	if (ContactModifyCallback)	//default is nullptr
-	{
-		FPhysPendingCommand Command;
-		Command.Pointer.ContactModifyCallback = ContactModifyCallback;
-		Command.CommandType = PhysCommand::DeleteContactModifyCallback;
-
-		EnqueueCommand(Command);
-	}
-}
-
-void FPhysCommandHandler::DeferredDeleteMbpBroadphaseCallback(FPhysXMbpBroadphaseCallback* MbpCallback)
-{
-	if(MbpCallback)	//default is nullptr
-	{
-		FPhysPendingCommand Command;
-		Command.Pointer.MbpCallback = MbpCallback;
-		Command.CommandType = PhysCommand::DeleteMbpBroadphaseCallback;
-
-		EnqueueCommand(Command);
-	}
-}
-
-void  FPhysCommandHandler::DeferredDeleteSimEventCallback(physx::PxSimulationEventCallback * SimEventCallback)
-{
-	check(SimEventCallback)
-
-	FPhysPendingCommand Command;
-	Command.Pointer.SimEventCallback = SimEventCallback;
-	Command.CommandType = PhysCommand::DeleteSimEventCallback;
-
-	EnqueueCommand(Command);
-}
-void  FPhysCommandHandler::DeferredDeleteCPUDispathcer(physx::PxCpuDispatcher * CPUDispatcher)
-{
-	check(CPUDispatcher);
-
-	FPhysPendingCommand Command;
-	Command.Pointer.CPUDispatcher = CPUDispatcher;
-	Command.CommandType = PhysCommand::DeleteCPUDispatcher;
-
-	EnqueueCommand(Command);
-}
-#endif
-
-#if WITH_APEX
-void FPhysCommandHandler::DeferredRelease(nvidia::apex::ApexInterface* ApexInterface)
-{
-	check(ApexInterface);
-
-	FPhysPendingCommand Command;
-	Command.Pointer.ApexInterface = ApexInterface;
-	Command.CommandType = PhysCommand::Release;
-
-	EnqueueCommand(Command);
-
-}
-#endif
\ No newline at end of file
+}
\ No newline at end of file
diff --git a/Engine/Source/Runtime/Engine/Private/PhysicsEngine/PhysDerivedData.cpp b/Engine/Source/Runtime/Engine/Private/PhysicsEngine/PhysDerivedData.cpp
--- a/Engine/Source/Runtime/Engine/Private/PhysicsEngine/PhysDerivedData.cpp
+++ b/Engine/Source/Runtime/Engine/Private/PhysicsEngine/PhysDerivedData.cpp
@@ -8,201 +8,3 @@
 #include "PhysicsEngine/PhysicsSettings.h"
 #include "PhysicsEngine/BodySetup.h"
 #include "Interfaces/ITargetPlatformManagerModule.h"
-
-#if PHYSICS_INTERFACE_PHYSX && WITH_EDITOR
-
-#include "IPhysXCookingModule.h"
-
-FDerivedDataPhysXCooker::FDerivedDataPhysXCooker(FName InFormat, EPhysXMeshCookFlags InRuntimeCookFlags, UBodySetup* InBodySetup, bool InIsRuntime)
-	: BodySetup( InBodySetup )
-	, CollisionDataProvider( NULL )
-	, Format( InFormat )
-	, RuntimeCookFlags(InRuntimeCookFlags)
-	, Cooker( NULL )
-	, bIsRuntime(InIsRuntime)
-{
-	check( BodySetup != NULL );
-	CollisionDataProvider = BodySetup->GetOuter();
-	DataGuid = BodySetup->BodySetupGuid;
-	bGenerateNormalMesh = BodySetup->bGenerateNonMirroredCollision;
-	bGenerateMirroredMesh = BodySetup->bGenerateMirroredCollision;
-	bGenerateUVInfo = UPhysicsSettings::Get()->bSupportUVFromHitResults;
-	BodyComplexity = (int32)BodySetup->GetCollisionTraceFlag();
-	IInterface_CollisionDataProvider* CDP = Cast<IInterface_CollisionDataProvider>(CollisionDataProvider);
-	if (CDP)
-	{
-		CDP->GetMeshId(MeshId);
-	}
-	InitCooker();
-
-	bVerifyDDC = FParse::Param(FCommandLine::Get(), TEXT("VerifyDDC"));
-}
-
-void FDerivedDataPhysXCooker::InitCooker()
-{
-	// static here as an optimization
-	static ITargetPlatformManagerModule* TPM = GetTargetPlatformManager();
-	if (TPM)
-	{
-		Cooker = TPM->FindPhysXCooking( Format );
-	}
-}
-
-FString FDerivedDataPhysXCooker::GetDebugContextString() const
-{
-	if (BodySetup)
-	{
-		UObject* Outer = BodySetup->GetOuter();
-		if (Outer)
-		{
-			return Outer->GetFullName();
-		}
-	}
-	return TEXT("Unknown Context");
-}
-
-bool FDerivedDataPhysXCooker::Build( TArray<uint8>& OutData )
-{
-	TRACE_CPUPROFILER_EVENT_SCOPE(FDerivedDataPhysXCooker::Build);
-
-	SCOPE_CYCLE_COUNTER(STAT_PhysXCooking);
-
-	check(Cooker != NULL);
-
-	if(bIsRuntime && !bVerifyDDC)
-	{
-		return false;
-	}
-
-	FMemoryWriter Ar( OutData );
-	uint8 bLittleEndian = PLATFORM_LITTLE_ENDIAN;	//TODO: We should pass the target platform into this function and write it. Then swap the endian on the writer so the reader doesn't have to do it at runtime
-	int32 NumConvexElementsCooked = 0;
-	int32 NumMirroredElementsCooked = 0;
-	int32 NumTriMeshesCooked = 0;
-	Ar << bLittleEndian;
-	int64 CookedMeshInfoOffset = Ar.Tell();
-	Ar << NumConvexElementsCooked;	
-	Ar << NumMirroredElementsCooked;
-	Ar << NumTriMeshesCooked;
-
-	//TODO: we must save an id with convex and tri meshes for serialization. We must save this here and patch it up at runtime somehow
-
-	bool bSuccess = true;
-
-	FCookBodySetupInfo CookInfo;
-	BodySetup->GetCookInfo(CookInfo, RuntimeCookFlags);
-
-	if(CookInfo.bCookNonMirroredConvex)
-	{
-		bSuccess = BuildConvex(OutData, CookInfo.bConvexDeformableMesh, false, CookInfo.NonMirroredConvexVertices, CookInfo.ConvexCookFlags, NumConvexElementsCooked) && bSuccess;
-	}
-
-	if(CookInfo.bCookMirroredConvex)
-	{
-		bSuccess = BuildConvex(OutData, CookInfo.bConvexDeformableMesh, true, CookInfo.MirroredConvexVertices, CookInfo.ConvexCookFlags, NumMirroredElementsCooked) && bSuccess;
-	}
-
-	FBodySetupUVInfo UVInfo;
-	if(CookInfo.bCookTriMesh)
-	{
-		if(!CookInfo.bTriMeshError)
-		{
-			bSuccess = BuildTriMesh(OutData, CookInfo.TriangleMeshDesc, CookInfo.TriMeshCookFlags, CookInfo.bSupportUVFromHitResults ? &UVInfo : nullptr, NumTriMeshesCooked) && bSuccess;
-		}
-		else
-		{
-			bSuccess = false;
-		}
-	}
-
-	// Seek to end, serialize UV info
-	Ar.Seek(OutData.Num());
-	Ar << UVInfo;
-
-	// Update info on what actually got cooked
-	Ar.Seek( CookedMeshInfoOffset );
-	Ar << NumConvexElementsCooked;	
-	Ar << NumMirroredElementsCooked;
-	Ar << NumTriMeshesCooked;
-
-	return bSuccess;
-}
-
-bool FDerivedDataPhysXCooker::BuildConvex( TArray<uint8>& OutData, bool bDeformableMesh, bool InMirrored, const TArray<TArray<FVector>>& Elements, EPhysXMeshCookFlags CookFlags, int32& NumConvexCooked )
-{	
-	bool bSuccess = true;
-	for( int32 ElementIndex = 0; ElementIndex < Elements.Num(); ElementIndex++ )
-	{
-		// Store info on the cooking result (1 byte)
-		int32 ResultInfoOffset = OutData.Add(false);
-
-		// Cook and store Result at ResultInfoOffset
-		UE_LOG(LogPhysics, Log, TEXT("Cook Convex: %s %d (FlipX:%d)"), *BodySetup->GetOuter()->GetPathName(), ElementIndex, InMirrored);		
-		const EPhysXCookingResult Result = Cooker->CookConvex(Format, CookFlags, Elements[ElementIndex], OutData);
-		switch (Result)
-		{
-		case EPhysXCookingResult::Succeeded:
-			break;
-		case EPhysXCookingResult::Failed:
-			UE_LOG(LogPhysics, Warning, TEXT("Failed to cook convex: %s %d (FlipX:%d). The remaining elements will not get cooked."), *BodySetup->GetOuter()->GetPathName(), ElementIndex, InMirrored ? 1 : 0);
-			bSuccess = false;
-			break;
-		case EPhysXCookingResult::SucceededWithInflation:
-			if (!bDeformableMesh)
-			{
-				bSuccess = false;
-				UE_LOG(LogPhysics, Warning, TEXT("Cook convex: %s %d (FlipX:%d) failed but succeeded with inflation.  The mesh should be looked at."), *BodySetup->GetOuter()->GetPathName(), ElementIndex, InMirrored ? 1 : 0);
-			}
-			else
-			{
-				UE_LOG(LogPhysics, Log, TEXT("Cook convex: %s %d (FlipX:%d) required inflation. You may wish to adjust the mesh so this is not necessary."), *BodySetup->GetOuter()->GetPathName(), ElementIndex, InMirrored ? 1 : 0);
-			}
-			break;
-		default:
-			// Unknown/unsupported enum value
-			bSuccess = false;
-			check(false);
-		}
-		OutData[ ResultInfoOffset ] = (Result != EPhysXCookingResult::Failed) ? 1 : 0;
-	}
-
-	NumConvexCooked = Elements.Num();
-	return bSuccess;
-}
-
-bool FDerivedDataPhysXCooker::ShouldGenerateTriMeshData(bool InUseAllTriData)
-{
-	check(Cooker != NULL);
-
-	IInterface_CollisionDataProvider* CDP = Cast<IInterface_CollisionDataProvider>(CollisionDataProvider);
-	const bool bPerformCook = ( CDP != NULL ) ? CDP->ContainsPhysicsTriMeshData(InUseAllTriData) : false;
-	return bPerformCook;
-}
-
-bool FDerivedDataPhysXCooker::BuildTriMesh( TArray<uint8>& OutData,  const FTriMeshCollisionData& TriangleMeshDesc, EPhysXMeshCookFlags CookFlags, FBodySetupUVInfo* UVInfo, int32& NumTrimeshCooked)
-{
-	check(Cooker != NULL);
-
-	bool bError = false;
-	bool bResult = false;
-
-	UE_LOG(LogPhysics, Log, TEXT("Cook TriMesh: %s"), *CollisionDataProvider->GetPathName());
-	bResult = Cooker->CookTriMesh( Format, CookFlags, TriangleMeshDesc.Vertices, TriangleMeshDesc.Indices, TriangleMeshDesc.MaterialIndices, TriangleMeshDesc.bFlipNormals, OutData);
-	if( !bResult )
-	{
-		bError = true;
-		UE_LOG(LogPhysics, Warning, TEXT("Failed to cook TriMesh: %s."), *CollisionDataProvider->GetPathName());
-	}
-
-	// If we want UV info, copy that now
-	if (bResult && UVInfo)
-	{
-		UVInfo->FillFromTriMesh(TriangleMeshDesc);
-	}
-
-	NumTrimeshCooked = bResult == true ? 1 : 0;	//the cooker only generates 1 or 0 trimeshes. We save an int because we support multiple trimeshes for welding and we might want to do this per static mesh in the future.
-	return !bError;	//use error instead of bResult because we do not warn if the trimesh data is simply empty, and in that case we don't consider it to be a failure
-}
-
-
-#endif	//WITH_PHYSX && WITH_EDITOR
diff --git a/Engine/Source/Runtime/Engine/Private/PhysicsEngine/PhysDerivedData.h b/Engine/Source/Runtime/Engine/Private/PhysicsEngine/PhysDerivedData.h
--- a/Engine/Source/Runtime/Engine/Private/PhysicsEngine/PhysDerivedData.h
+++ b/Engine/Source/Runtime/Engine/Private/PhysicsEngine/PhysDerivedData.h
@@ -8,104 +8,3 @@
 #include "PhysXIncludes.h"
 #include "Stats/Stats.h"
 #include "PhysDerivedDataPublic.h"
-
-#if PHYSICS_INTERFACE_PHYSX && WITH_EDITOR
-#include "DerivedDataPluginInterface.h"
-#include "IPhysXCooking.h"
-
-class UBodySetup;
-struct FBodyInstance;
-struct FBodySetupUVInfo;
-struct FTriMeshCollisionData;
-
-//////////////////////////////////////////////////////////////////////////
-// PhysX Cooker
-class FDerivedDataPhysXCooker : public FDerivedDataPluginInterface
-{
-private:
-
-	UBodySetup* BodySetup;
-	UObject* CollisionDataProvider;
-	FName Format;
-	bool bGenerateNormalMesh;
-	bool bGenerateMirroredMesh;
-	bool bGenerateUVInfo;
-	int32 BodyComplexity;
-	EPhysXMeshCookFlags RuntimeCookFlags;
-	const class IPhysXCooking* Cooker;
-	FGuid DataGuid;
-	FString MeshId;
-	bool bIsRuntime;
-	bool bVerifyDDC;
-
-public:
-	FDerivedDataPhysXCooker(FName InFormat, EPhysXMeshCookFlags InRuntimeCookFlags, UBodySetup* InBodySetup, bool InIsRuntime);
-
-	virtual const TCHAR* GetPluginName() const override
-	{
-		return TEXT("PhysX");
-	}
-
-	virtual const TCHAR* GetVersionString() const override
-	{
-		// This is a version string that mimics the old versioning scheme. If you
-		// want to bump this version, generate a new guid using VS->Tools->Create GUID and
-		// return it here. Ex.
-		return PHYSX_DDC;	
-	}
-
-	virtual FString GetPluginSpecificCacheKeySuffix() const override
-	{
-		//  1 - base version
-		//  2 - cook out small area trimesh triangles from BSP (see UPhysicsSettings::TriangleMeshTriangleMinAreaThreshold)
-		//  3 - increase default small area threshold and force recook.
-		enum { UE_PHYSX_DERIVEDDATA_VER = 3 };
-
-		const uint16 PhysXVersion = ((PX_PHYSICS_VERSION_MAJOR  & 0xF) << 12) |
-				((PX_PHYSICS_VERSION_MINOR  & 0xF) << 8) |
-				((PX_PHYSICS_VERSION_BUGFIX & 0xF) << 4) |
-				((UE_PHYSX_DERIVEDDATA_VER	& 0xF));
-
-		return FString::Printf( TEXT("%s_%s_%s_%d_%d_%d_%d_%d_%hu_%hu"),
-			*Format.ToString(),
-			*DataGuid.ToString(),
-			*MeshId,
-			(int32)bGenerateNormalMesh,
-			(int32)bGenerateMirroredMesh,
-			(int32)bGenerateUVInfo,
-			(int32)RuntimeCookFlags,
-			BodyComplexity,
-			PhysXVersion,
-			Cooker ? Cooker->GetVersion( Format ) : 0xffff
-			);
-	}
-
-
-	virtual bool IsBuildThreadsafe() const override
-	{
-		return false;
-	}
-
-	virtual bool IsDeterministic() const override
-	{
-		return true;
-	}
-
-	virtual FString GetDebugContextString() const override;
-
-	virtual bool Build( TArray<uint8>& OutData ) override;
-
-	/** Return true if we can build **/
-	bool CanBuild()
-	{
-		return !!Cooker;
-	}
-private:
-
-	void InitCooker();
-	bool BuildConvex( TArray<uint8>& OutData, bool bDeformableMesh, bool InMirrored, const TArray<TArray<FVector>>& Elements, EPhysXMeshCookFlags CookFlags, int32& NumConvexCooked);
-	bool BuildTriMesh( TArray<uint8>& OutData, const FTriMeshCollisionData& TriangleMeshDesc, EPhysXMeshCookFlags CookFlags, FBodySetupUVInfo* UVInfo, int32& NumTriMeshCooked);
-	bool ShouldGenerateTriMeshData(bool InUseAllTriData);
-};
-
-#endif	//WITH_PHYSX && WITH_EDITOR
diff --git a/Engine/Source/Runtime/Engine/Private/PhysicsEngine/PhysDrawing.cpp b/Engine/Source/Runtime/Engine/Private/PhysicsEngine/PhysDrawing.cpp
--- a/Engine/Source/Runtime/Engine/Private/PhysicsEngine/PhysDrawing.cpp
+++ b/Engine/Source/Runtime/Engine/Private/PhysicsEngine/PhysDrawing.cpp
@@ -677,58 +677,6 @@ void FKTaperedCapsuleElem::DrawElemSolid(FPrimitiveDrawInterface* PDI, const FTr
 
 void FKConvexElem::DrawElemWire(FPrimitiveDrawInterface* PDI, const FTransform& ElemTM, const float Scale, const FColor Color) const
 {
-#if WITH_PHYSX && PHYSICS_INTERFACE_PHYSX
-
-	PxConvexMesh* Mesh = ConvexMesh;
-
-	if(Mesh)
-	{
-		// Draw each triangle that makes up the convex hull
-		PxU32 NbVerts = Mesh->getNbVertices();
-		const PxVec3* Vertices = Mesh->getVertices();
-		
-		// ElemTM is element transform, but geometry is stored in body space, so we need to remove body->element transform
-		FTransform RenderTM = Transform.GetRelativeTransformReverse(ElemTM);
-
-		TArray<FVector> TransformedVerts;
-		TransformedVerts.AddUninitialized(NbVerts);
-		for(PxU32 i=0; i<NbVerts; i++)
-		{
-			TransformedVerts[i] = RenderTM.TransformPosition(P2UVector(Vertices[i]) * Scale);
-		}
-						
-		const PxU8* PIndexBuffer = Mesh->getIndexBuffer();
-		PxU32 NbPolygons = Mesh->getNbPolygons();
-
-		for(PxU32 i=0;i<NbPolygons;i++)
-		{
-			PxHullPolygon Data;
-			bool bStatus = Mesh->getPolygonData(i, Data);
-			check(bStatus);
-
-			const PxU8* PIndices = PIndexBuffer + Data.mIndexBase;
-		
-			for(PxU16 j=0;j<Data.mNbVerts;j++)
-			{
-				// Get the verts that make up this line.
-				int32 I0 = PIndices[j];
-				int32 I1 = (j == Data.mNbVerts - 1) ? PIndices[0] : PIndices[j + 1];
-
-				// Loop back last and first vertices
-				if(j==Data.mNbVerts - 1)
-				{
-					I1 = PIndices[0];
-				}
-
-				PDI->DrawLine( TransformedVerts[I0], TransformedVerts[I1], Color, SDPG_World );
-			}
-		}
-	}
-	else
-	{
-		UE_LOG(LogPhysics, Log, TEXT("FKConvexElem::DrawElemWire : No ConvexMesh, so unable to draw."));
-	}
-#elif WITH_CHAOS
 	const int32 NumIndices = IndexData.Num();
 	if(NumIndices > 0 && ensure(NumIndices % 3 == 0))
 	{
@@ -759,12 +707,10 @@ void FKConvexElem::DrawElemWire(FPrimitiveDrawInterface* PDI, const FTransform&
 	{
 		UE_LOG(LogPhysics, Log, TEXT("FKConvexElem::DrawElemWire : No ConvexMesh, so unable to draw."));
 	}
-#endif // WITH_PHYSX
 }
 
 void FKConvexElem::DrawElemSolid(FPrimitiveDrawInterface* PDI, const FTransform& ElemTM, const float Scale, const FMaterialRenderProxy* MaterialRenderProxy) const
 {
-#if WITH_CHAOS
 	const int32 NumIndices = IndexData.Num();
 	if (NumIndices > 0 && ensure(NumIndices % 3 == 0))
 	{
@@ -798,91 +744,10 @@ void FKConvexElem::DrawElemSolid(FPrimitiveDrawInterface* PDI, const FTransform&
 	{
 		UE_LOG(LogPhysics, Log, TEXT("FKConvexElem::DrawElemSolid : No ConvexMesh, so unable to draw."));
 	}
-#endif
 }
 
 void FKConvexElem::AddCachedSolidConvexGeom(TArray<FDynamicMeshVertex>& VertexBuffer, TArray<uint32>& IndexBuffer, const FColor VertexColor) const
 {
-#if WITH_PHYSX && PHYSICS_INTERFACE_PHYSX
-	// We always want to generate 'non-mirrored geometry', so if all we have is flipped, we have to un-flip it in this function
-	bool bIsMirrored = false;
-	const PxConvexMesh* ConvexMeshToUse = nullptr; 
-	if (ConvexMesh != nullptr)
-	{
-		ConvexMeshToUse = ConvexMesh;
-	}
-	else if (ConvexMeshNegX != nullptr)
-	{
-		ConvexMeshToUse = ConvexMeshNegX;
-		bIsMirrored = true;
-	}
-
-	if(ConvexMeshToUse)
-	{
-		int32 StartVertOffset = VertexBuffer.Num();
-
-		// get PhysX data
-		const PxVec3* PVertices = ConvexMeshToUse->getVertices();
-		const PxU8* PIndexBuffer = ConvexMeshToUse->getIndexBuffer();
-		PxU32 NbPolygons = ConvexMeshToUse->getNbPolygons();
-
-		FVector Scale3D = bIsMirrored ? FVector(-1, 1, 1) : FVector(1, 1, 1);
-
-		for(PxU32 i=0;i<NbPolygons;i++)
-		{
-			PxHullPolygon Data;
-			bool bStatus = ConvexMeshToUse->getPolygonData(i, Data);
-			check(bStatus);
-
-			const PxU8* indices = PIndexBuffer + Data.mIndexBase;
-
-			// create tangents from the first and second vertices of each polygon
-			const FVector TangentX = P2UVector(PVertices[indices[1]]-PVertices[indices[0]]).GetSafeNormal();
-			const FVector TangentZ = FVector(Data.mPlane[0], Data.mPlane[1], Data.mPlane[2]).GetSafeNormal();
-			const FVector TangentY = (TangentX ^ TangentZ).GetSafeNormal();
-
-			// add vertices 
-			for(PxU32 j=0;j<Data.mNbVerts;j++)
-			{
-				int32 VertIndex = indices[j];
-
-				FDynamicMeshVertex Vert1;
-				Vert1.Position = P2UVector(PVertices[VertIndex]) * Scale3D;
-				Vert1.Color = VertexColor;
-				Vert1.SetTangents(
-					TangentX,
-					TangentY,
-					TangentZ
-					);
-				VertexBuffer.Add(Vert1);
-			}
-
-			// Add indices
-			PxU32 nbTris = Data.mNbVerts - 2;
-			for(PxU32 j=0;j<nbTris;j++)
-			{
-				IndexBuffer.Add(StartVertOffset+0);
-
-				if (bIsMirrored)
-				{
-					IndexBuffer.Add(StartVertOffset + j + 1);
-					IndexBuffer.Add(StartVertOffset + j + 2);
-				}
-				else
-				{
-					IndexBuffer.Add(StartVertOffset + j + 2);
-					IndexBuffer.Add(StartVertOffset + j + 1);
-				}
-			}
-
-			StartVertOffset += Data.mNbVerts;
-		}
-	}
-	else
-	{
-		UE_LOG(LogPhysics, Log, TEXT("FKConvexElem::AddCachedSolidConvexGeom : No ConvexMesh, so unable to draw."));
-	}
-#elif WITH_CHAOS
 	const int32 NumIndices = IndexData.Num();
 	if(NumIndices > 0 && ensure(NumIndices % 3 == 0))
 	{
@@ -916,7 +781,6 @@ void FKConvexElem::AddCachedSolidConvexGeom(TArray<FDynamicMeshVertex>& VertexBu
 	{
 		UE_LOG(LogPhysics, Log, TEXT("FKConvexElem::AddCachedSolidConvexGeom : No ConvexMesh, so unable to draw."));
 	}
-#endif // WITH_PHYSX
 }
 
 /////////////////////////////////////////////////////////////////////////////////////
