Physics interface cleanup.

diff --git a/Engine/Source/Runtime/Engine/Private/PhysicsEngine/PhysLevel.cpp b/Engine/Source/Runtime/Engine/Private/PhysicsEngine/PhysLevel.cpp
--- a/Engine/Source/Runtime/Engine/Private/PhysicsEngine/PhysLevel.cpp
+++ b/Engine/Source/Runtime/Engine/Private/PhysicsEngine/PhysLevel.cpp
@@ -10,20 +10,12 @@
 #include "EngineDefines.h"
 #include "Engine/World.h"
 #include "Physics/PhysicsInterfaceCore.h"
-#if PHYSICS_INTERFACE_PHYSX
-#include "IPhysXCookingModule.h"
-#include "IPhysXCooking.h"
-#endif
 #include "Modules/ModuleManager.h"
 #include "PhysicsInitialization.h"
 #include "PhysicsEngine/PhysicsSettings.h"
 
 #include "ChaosSolversModule.h"
 
-#if PHYSICS_INTERFACE_PHYSX
-	#include "PhysicsEngine/PhysXSupport.h"
-#endif
-
 #include "PhysicsEngine/PhysicsSettings.h"
 #include "Misc/CoreDelegates.h"
 
@@ -96,7 +88,7 @@ void UWorld::SetupPhysicsTickFunctions(float DeltaSeconds)
 	EndPhysicsTickFunction.Target = this;
 
 // Chaos ticks solver for trace collisions
-#if (WITH_CHAOS && WITH_EDITOR)
+#if (WITH_EDITOR)
 	bool bEnablePhysics = (bShouldSimulatePhysics || bEnableTraceCollision);
 #else
 	bool bEnablePhysics = bShouldSimulatePhysics;
@@ -168,11 +160,7 @@ void UWorld::FinishPhysicsSim()
 		return;
 	}
 
-#if WITH_CHAOS
 	PhysScene->EndFrame();
-#else
-	PhysScene->EndFrame(LineBatcher);
-#endif
 }
 
 // the physics tick functions
@@ -227,17 +215,6 @@ void FEndPhysicsTickFunction::ExecuteTick(float DeltaTime, enum ELevelTick TickT
 		// it was already done, so let just do it.
 		Target->FinishPhysicsSim();
 	}
-
-#if PHYSICS_INTERFACE_PHYSX
-#if PHYSX_MEMORY_VALIDATION
-	static int32 Frequency = 0;
-	if (Frequency++ > 10)
-	{
-		Frequency = 0;
-		GPhysXAllocator->ValidateHeaders();
-	}
-#endif
-#endif // WITH_PHYSX 
 }
 
 FString FEndPhysicsTickFunction::DiagnosticMessage()
@@ -304,15 +281,6 @@ void TermGamePhys()
 		GPostInitHandle.Reset();
 	}
 
-#if PHYSICS_INTERFACE_PHYSX
-	// Do nothing if they were never initialized
-	if(GPhysXFoundation == NULL)
-	{
-		FPhysxSharedData::Terminate();	//early out before TermGamePhysCore so kill this - not sure if this is a real case we even care about
-		return;
-	}
-#endif
-
 	if (GPhysCommandHandler != NULL)
 	{
 		GPhysCommandHandler->Flush();	//finish off any remaining commands
@@ -330,118 +298,5 @@ void TermGamePhys()
 */
 void DeferredPhysResourceCleanup()
 {
-#if PHYSICS_INTERFACE_PHYSX
-
-	// Release all tri meshes and reset array
-	for(int32 MeshIdx=0; MeshIdx<GPhysXPendingKillTriMesh.Num(); MeshIdx++)
-	{
-		PxTriangleMesh* PTriMesh = GPhysXPendingKillTriMesh[MeshIdx];
-
-		// Check this as it shouldn't be null, but then gate on it so we can
-		// avoid a crash if we end up in this state in shipping
-		check(PTriMesh);
-		if(PTriMesh)
-		{
-			PTriMesh->release();
-
-			if(GPhysXPendingKillTriMesh.IsValidIndex(MeshIdx))
-			{
-				GPhysXPendingKillTriMesh[MeshIdx] = NULL;
-			}
-			else
-			{
-				UE_LOG(LogPhysics, Warning, TEXT("DeferredPhysResourceCleanup found invalid index into GPhysXPendingKillTriMesh, another thread may have modified the array."), MeshIdx);
-			}
-		}
-		else
-		{
-			UE_LOG(LogPhysics, Warning, TEXT("DeferredPhysResourceCleanup found null PxTriangleMesh in pending kill array, another thread may have modified the array."), MeshIdx);
-		}
-	}
-	GPhysXPendingKillTriMesh.Reset();
-
-	// Release all convex meshes and reset array
-	for(int32 MeshIdx=0; MeshIdx<GPhysXPendingKillConvex.Num(); MeshIdx++)
-	{
-		PxConvexMesh* PConvexMesh = GPhysXPendingKillConvex[MeshIdx];
-
-		// Check this as it shouldn't be null, but then gate on it so we can
-		// avoid a crash if we end up in this state in shipping
-		check(PConvexMesh);
-		if(PConvexMesh)
-		{
-			PConvexMesh->release();
-
-			if(GPhysXPendingKillConvex.IsValidIndex(MeshIdx))
-			{
-				GPhysXPendingKillConvex[MeshIdx] = NULL;
-			}
-			else
-			{
-				UE_LOG(LogPhysics, Warning, TEXT("DeferredPhysResourceCleanup found invalid index into GPhysXPendingKillConvex (%d), another thread may have modified the array."), MeshIdx);
-			}
-		}
-		else
-		{
-			UE_LOG(LogPhysics, Warning, TEXT("DeferredPhysResourceCleanup found null PxConvexMesh in pending kill array (at %d), another thread may have modified the array."), MeshIdx);
-		}
-	}
-	GPhysXPendingKillConvex.Reset();
-
-	// Release all heightfields and reset array
-	for(int32 HfIdx=0; HfIdx<GPhysXPendingKillHeightfield.Num(); HfIdx++)
-	{
-		PxHeightField* PHeightfield = GPhysXPendingKillHeightfield[HfIdx];
-
-		// Check this as it shouldn't be null, but then gate on it so we can
-		// avoid a crash if we end up in this state in shipping
-		check(PHeightfield);
-		if(PHeightfield)
-		{
-			PHeightfield->release();
-
-			if(GPhysXPendingKillHeightfield.IsValidIndex(HfIdx))
-			{
-				GPhysXPendingKillHeightfield[HfIdx] = NULL;
-			}
-			else
-			{
-				UE_LOG(LogPhysics, Warning, TEXT("DeferredPhysResourceCleanup found invalid index into GPhysXPendingKillHeightfield (%d), another thread may have modified the array."), HfIdx);
-			}
-		}
-		else
-		{
-			UE_LOG(LogPhysics, Warning, TEXT("DeferredPhysResourceCleanup found null PxHeightField in pending kill array (at %d), another thread may have modified the array."), HfIdx);
-		}
-	}
-	GPhysXPendingKillHeightfield.Reset();
-
-	// Release all materials and reset array
-	for(int32 MeshIdx=0; MeshIdx<GPhysXPendingKillMaterial.Num(); MeshIdx++)
-	{
-		PxMaterial* PMaterial = GPhysXPendingKillMaterial[MeshIdx];
-
-		// Check this as it shouldn't be null, but then gate on it so we can
-		// avoid a crash if we end up in this state in shipping
-		check(PMaterial);
-		if(PMaterial)
-		{
-			PMaterial->release();
-			if(GPhysXPendingKillMaterial.IsValidIndex(MeshIdx))
-			{
-				GPhysXPendingKillMaterial[MeshIdx] = NULL;
-			}
-			else
-			{
-				UE_LOG(LogPhysics, Warning, TEXT("DeferredPhysResourceCleanup found invalid index into GPhysXPendingKillMaterial(%d), another thread may have modified the array."), MeshIdx);
-			}
-		}
-		else
-		{
-			UE_LOG(LogPhysics, Warning, TEXT("DeferredPhysResourceCleanup found null PxMaterial in pending kill array (at %d), another thread may have modified the array."), MeshIdx);
-		}
-	}
-	GPhysXPendingKillMaterial.Reset();
-#endif
 }
 
diff --git a/Engine/Source/Runtime/Engine/Private/PhysicsEngine/PhysScene_PhysX.cpp b/Engine/Source/Runtime/Engine/Private/PhysicsEngine/PhysScene_PhysX.cpp
--- a/Engine/Source/Runtime/Engine/Private/PhysicsEngine/PhysScene_PhysX.cpp
+++ b/Engine/Source/Runtime/Engine/Private/PhysicsEngine/PhysScene_PhysX.cpp
@@ -1,2333 +1 @@
 // Copyright Epic Games, Inc. All Rights Reserved.
-
-#if !WITH_CHAOS && !WITH_IMMEDIATE_PHYSX
-
-#include "Physics/PhysScene_PhysX.h"
-#include "Misc/CommandLine.h"
-#include "Stats/Stats.h"
-#include "UObject/UObjectGlobals.h"
-#include "UObject/UObjectHash.h"
-#include "UObject/UObjectIterator.h"
-#include "HAL/IConsoleManager.h"
-#include "Async/TaskGraphInterfaces.h"
-#include "EngineDefines.h"
-#include "Engine/EngineTypes.h"
-#include "PhysxUserData.h"
-#include "PhysicsEngine/BodyInstance.h"
-#include "Components/PrimitiveComponent.h"
-#include "Components/SkeletalMeshComponent.h"
-#include "PhysicsEngine/RigidBodyIndexPair.h"
-#include "PhysicsPublic.h"
-#include "CustomPhysXPayload.h"
-#include "HAL/LowLevelMemTracker.h"
-
-#include "PhysicsInterfaceDeclaresCore.h"
-#if !WITH_CHAOS_NEEDS_TO_BE_FIXED
-
-#include "SQAccelerator.h"
-
-#if WITH_PHYSX
-#include "PhysXPublic.h"
-#include "PhysicsEngine/PhysXSupport.h"
-#endif
-
-#include "PhysicsEngine/PhysSubstepTasks.h"
-#include "PhysicsEngine/PhysicsCollisionHandler.h"
-#include "Physics/PhysicsInterfaceUtils.h"
-#include "Components/LineBatchComponent.h"
-#include "PhysicsEngine/PhysicsSettings.h"
-#include "PhysicsEngine/BodySetup.h"
-#include "PhysicsEngine/ConstraintInstance.h"
-#include "PhysicsReplication.h"
-#include "ProfilingDebugging/CsvProfiler.h"
-
-#include "PhysTestSerializer.h"
-
-/** Physics stats **/
-
-DEFINE_STAT(STAT_TotalPhysicsTime);
-DEFINE_STAT(STAT_NumCloths);
-DEFINE_STAT(STAT_NumClothVerts);
-
-
-CSV_DECLARE_CATEGORY_MODULE_EXTERN(CORE_API, Basic);
-
-DECLARE_CYCLE_STAT(TEXT("Start Physics Time"), STAT_PhysicsKickOffDynamicsTime, STATGROUP_Physics);
-DECLARE_CYCLE_STAT(TEXT("Fetch Results Time"), STAT_PhysicsFetchDynamicsTime, STATGROUP_Physics);
-
-
-DECLARE_CYCLE_STAT(TEXT("Update Kinematics On Deferred SkelMeshes"), STAT_UpdateKinematicsOnDeferredSkelMeshes, STATGROUP_Physics);
-
-DECLARE_CYCLE_STAT(TEXT("Phys Events Time"), STAT_PhysicsEventTime, STATGROUP_Physics);
-DECLARE_CYCLE_STAT(TEXT("SyncComponentsToBodies"), STAT_SyncComponentsToBodies, STATGROUP_Physics);
-
-DECLARE_DWORD_COUNTER_STAT(TEXT("Broadphase Adds"), STAT_NumBroadphaseAdds, STATGROUP_Physics);
-DECLARE_DWORD_COUNTER_STAT(TEXT("Broadphase Removes"), STAT_NumBroadphaseRemoves, STATGROUP_Physics);
-DECLARE_DWORD_COUNTER_STAT(TEXT("Active Constraints"), STAT_NumActiveConstraints, STATGROUP_Physics);
-DECLARE_DWORD_COUNTER_STAT(TEXT("Active Simulated Bodies"), STAT_NumActiveSimulatedBodies, STATGROUP_Physics);
-DECLARE_DWORD_COUNTER_STAT(TEXT("Active Kinematic Bodies"), STAT_NumActiveKinematicBodies, STATGROUP_Physics);
-DECLARE_DWORD_COUNTER_STAT(TEXT("Mobile Bodies"), STAT_NumMobileBodies, STATGROUP_Physics);
-DECLARE_DWORD_COUNTER_STAT(TEXT("Static Bodies"), STAT_NumStaticBodies, STATGROUP_Physics);
-DECLARE_DWORD_COUNTER_STAT(TEXT("Shapes"), STAT_NumShapes, STATGROUP_Physics);
-
-#if WITH_PHYSX
-
-FAutoConsoleTaskPriority CPrio_FPhysXTask(
-	TEXT("TaskGraph.TaskPriorities.PhysXTask"),
-	TEXT("Task and thread priority for FPhysXTask."),
-	ENamedThreads::HighThreadPriority, // if we have high priority task threads, then use them...
-	ENamedThreads::NormalTaskPriority, // .. at normal task priority
-	ENamedThreads::HighTaskPriority // if we don't have hi pri threads, then use normal priority threads at high task priority instead
-);
-
-
-int32 GPhysXOverrideMbpNumSubdivisions_Client = 0;
-int32 GPhysXOverrideMbpNumSubdivisions_Server = 0;
-int32 GPhysXForceMbp_Client = 0;
-int32 GPhysXForceMbp_Server = 0;
-int32 GPhysXForceNoKinematicStaticPairs = 0;
-int32 GPhysXForceNoKinematicKinematicPairs = 0;
-
-FAutoConsoleVariableRef CVarOverrideMbpNumSubdivisionsClient(TEXT("p.OverrideMbpNumSubdivisionsClient"), GPhysXOverrideMbpNumSubdivisions_Client, TEXT("Override for number of subdivisions to perform when building MBP regions on a client, note regions are only generated when a scene is created - this will not update the scene if it's already running (0 = No override, 1>16 - Override number)"), ECVF_Default);
-FAutoConsoleVariableRef CVarOverrideMbpNumSubdivisionsServer(TEXT("p.OverrideMbpNumSubdivisionsServer"), GPhysXOverrideMbpNumSubdivisions_Server, TEXT("Override for number of subdivisions to perform when building MBP regions on a server, note regions are only generated when a scene is created - this will not update the scene if it's already running (0 = No override, 1>16 - Override number)"), ECVF_Default);
-FAutoConsoleVariableRef CVarForceMbpClient(TEXT("p.ForceMbpClient"), GPhysXForceMbp_Client, TEXT("Forces all created scenes to use MBP on client builds"), ECVF_Default);
-FAutoConsoleVariableRef CVarForceMbpServer(TEXT("p.ForceMbpServer"), GPhysXForceMbp_Server, TEXT("Forces all created scenes to use MBP on server builds"), ECVF_Default);
-FAutoConsoleVariableRef CVarForceNoKSPairs(TEXT("p.ForceNoKSPairs"), GPhysXForceNoKinematicStaticPairs, TEXT("Disables kinematic-static pairs. This makes converting from static to dynamic a little slower - but provides better broadphase performance because we early reject those pairs."), ECVF_Default);
-FAutoConsoleVariableRef CVarForceNoKKPairs(TEXT("p.ForceNoKKPairs"), GPhysXForceNoKinematicKinematicPairs, TEXT("Disables kinematic-kinematic pairs. This is required when using APEX destruction to correctly generate chunk pairs - when not using destruction this speeds up the broadphase by early rejecting KK pairs."), ECVF_Default);
-
-
-DECLARE_STATS_GROUP(TEXT("PhysXTasks"), STATGROUP_PhysXTasks, STATCAT_Advanced);
-
-struct FPhysXRingBuffer
-{
-	static const int32 Size = 16;
-
-	PxBaseTask* Buffer[Size];
-	int32 Start;
-	int32 End;
-	int32 Num;
-};
-
-int32 GBatchPhysXTasksSize = 3;	//NOTE: FPhysXRingBuffer::Size should be twice as big as this value
-TAutoConsoleVariable<int32> CVarBatchPhysXTasksSize(TEXT("p.BatchPhysXTasksSize"), GBatchPhysXTasksSize, TEXT("Number of tasks to batch together (max 8). 1 will go as wide as possible, but more overhead on small tasks"), ECVF_Default);
-
-struct FBatchPhysXTasks
-{
-	static void SetPhysXTasksSinkFunc()
-	{
-		GBatchPhysXTasksSize = FMath::Max(1, FMath::Min(FPhysXRingBuffer::Size / 2, CVarBatchPhysXTasksSize.GetValueOnGameThread()));
-	}
-};
-
-struct FPhysTaskScopedNamedEvent
-{
-	FPhysTaskScopedNamedEvent() = delete;
-	FPhysTaskScopedNamedEvent(const FPhysTaskScopedNamedEvent& InOther) = delete;
-	FPhysTaskScopedNamedEvent& operator=(const FPhysTaskScopedNamedEvent& InOther) = delete;
-
-	FPhysTaskScopedNamedEvent(PxBaseTask* InTask)
-	{
-#if ENABLE_STATNAMEDEVENTS
-		check(InTask);
-		const char* TaskName = InTask->getName();
-
-		bEmittedEvent = GCycleStatsShouldEmitNamedEvents != 0;
-
-		if(bEmittedEvent)
-		{
-			FPlatformMisc::BeginNamedEvent(FColor::Green, TaskName);
-		}
-#endif
-	}
-
-	~FPhysTaskScopedNamedEvent()
-	{
-#if ENABLE_STATNAMEDEVENTS
-		if(bEmittedEvent)
-		{
-			FPlatformMisc::EndNamedEvent();
-		}
-#endif
-	}
-
-private:
-
-	bool bEmittedEvent;
-};
-
-static FAutoConsoleVariableSink CVarBatchPhysXTasks(FConsoleCommandDelegate::CreateStatic(&FBatchPhysXTasks::SetPhysXTasksSinkFunc));
-
-namespace DynamicStatsHelper
-{
-	struct FStatLookup
-	{
-		const char* StatName;
-		TStatId Stat;
-	} Stats[100];
-
-	int NumStats = 0;
-	FCriticalSection CS;
-
-	TStatId FindOrCreateStatId(const char* StatName)
-	{
-#if STATS
-		for (int StatIdx = 0; StatIdx < NumStats; ++StatIdx)
-		{
-			FStatLookup& Lookup = Stats[StatIdx];
-			if (Lookup.StatName == StatName)
-			{
-				return Lookup.Stat;
-			}
-		}
-
-		if (ensureMsgf(NumStats < sizeof(Stats) / sizeof(Stats[0]), TEXT("Too many different physx task stats. This will make the stat search slow")))
-		{
-			FScopeLock ScopeLock(&CS);
-
-			//Do the search again in case another thread added
-			for (int StatIdx = 0; StatIdx < NumStats; ++StatIdx)
-			{
-				FStatLookup& Lookup = Stats[StatIdx];
-				if (Lookup.StatName == StatName)
-				{
-					return Lookup.Stat;
-				}
-			}
-
-			FStatLookup& NewStat = Stats[NumStats];
-			NewStat.StatName = StatName;
-			NewStat.Stat = FDynamicStats::CreateStatId<FStatGroup_STATGROUP_PhysXTasks>(FName(StatName));
-			FPlatformMisc::MemoryBarrier();
-			++NumStats;	//make sure to do this at the end in case another thread is currently iterating
-			return NewStat.Stat;
-		}
-#endif // STATS
-		return TStatId();
-	}
-}
-
-struct FPhysXCPUDispatcher;
-
-class FPhysXTask
-{
-public:
-	FPhysXTask(PxBaseTask& InTask, FPhysXCPUDispatcher& InDispatcher);
-	FPhysXTask(FPhysXRingBuffer& RingBuffer, FPhysXCPUDispatcher& InDispatcher);
-	~FPhysXTask();
-
-	static FORCEINLINE TStatId GetStatId()
-	{
-		RETURN_QUICK_DECLARE_CYCLE_STAT(FPhysXTask, STATGROUP_Physics);
-	}
-	static FORCEINLINE ENamedThreads::Type GetDesiredThread()
-	{
-		return CPrio_FPhysXTask.Get();
-	}
-	static FORCEINLINE ESubsequentsMode::Type GetSubsequentsMode()
-	{
-		return ESubsequentsMode::TrackSubsequents;
-	}
-
-	void DoTask(ENamedThreads::Type CurrentThread, const FGraphEventRef& MyCompletionGraphEvent);
-
-	FPhysXRingBuffer RingBuffer;
-	FPhysXCPUDispatcher& Dispatcher;
-};
-
-/** Used to dispatch physx tasks to task graph */
-struct FPhysXCPUDispatcher : public PxCpuDispatcher
-{
-	FPhysXCPUDispatcher()
-	{
-		check(IsInGameThread());
-		TLSKey = FPlatformTLS::AllocTlsSlot();
-	}
-
-	~FPhysXCPUDispatcher()
-	{
-		check(IsInGameThread());
-		FPlatformTLS::FreeTlsSlot(TLSKey);
-	}
-
-	virtual void submitTask(PxBaseTask& Task) override
-	{
-		if (IsInGameThread())
-		{
-			//Game thread enqueues on task graph
-			TGraphTask<FPhysXTask>::CreateTask(NULL).ConstructAndDispatchWhenReady(Task, *this);
-		}
-		else
-		{
-			//See if we can use local queue
-			FPhysXRingBuffer& RingBuffer = *(FPhysXRingBuffer*)FPlatformTLS::GetTlsValue(TLSKey);
-			RingBuffer.Buffer[RingBuffer.End] = &Task;
-			RingBuffer.End = (RingBuffer.End + 1) % FPhysXRingBuffer::Size;
-			RingBuffer.Num++;
-
-			if (RingBuffer.Num >= GBatchPhysXTasksSize * 2)
-			{
-				TGraphTask<FPhysXTask>::CreateTask(NULL).ConstructAndDispatchWhenReady(RingBuffer, *this);
-			}
-		}
-	}
-
-	virtual PxU32 getWorkerCount() const override
-	{
-		return FTaskGraphInterface::Get().GetNumWorkerThreads();
-	}
-
-	uint32 TLSKey;
-};
-
-FPhysXTask::FPhysXTask(PxBaseTask& Task, FPhysXCPUDispatcher& InDispatcher)
-	: Dispatcher(InDispatcher)
-{
-	RingBuffer.Buffer[0] = &Task;
-	RingBuffer.Start = 0;
-	RingBuffer.End = 1;
-	RingBuffer.Num = 1;
-}
-
-FPhysXTask::FPhysXTask(FPhysXRingBuffer& InRingBuffer, FPhysXCPUDispatcher& InDispatcher)
-	: Dispatcher(InDispatcher)
-{
-	int32 NumToSteal = InRingBuffer.Num / 2;
-	ensureMsgf(NumToSteal > 0, TEXT("Trying to steal 0 items"));
-
-	const int32 StartPos = (InRingBuffer.Start + NumToSteal);
-	for (int32 Count = 0; Count < NumToSteal; ++Count)
-	{
-		RingBuffer.Buffer[Count] = InRingBuffer.Buffer[(StartPos + Count) % FPhysXRingBuffer::Size];
-	}
-
-	RingBuffer.Start = 0;
-	RingBuffer.End = NumToSteal;
-	RingBuffer.Num = NumToSteal;
-
-
-	InRingBuffer.Num -= NumToSteal;
-	InRingBuffer.End = (StartPos) % FPhysXRingBuffer::Size;
-}
-
-FPhysXTask::~FPhysXTask()
-{
-	FPlatformTLS::SetTlsValue(Dispatcher.TLSKey, nullptr);
-}
-
-void FPhysXTask::DoTask(ENamedThreads::Type CurrentThread, const FGraphEventRef& MyCompletionGraphEvent)
-{
-	FPlatformTLS::SetTlsValue(Dispatcher.TLSKey, &RingBuffer);
-
-	while (RingBuffer.Num)
-	{
-		PxBaseTask* Task = RingBuffer.Buffer[RingBuffer.Start];
-
-#if ENABLE_STATNAMEDEVENTS || STATS
-		FPhysTaskScopedNamedEvent TaskEvent(Task);
-#endif
-
-#if STATS
-		const char* TaskName = Task->getName();
-		FScopeCycleCounter CycleCounter(DynamicStatsHelper::FindOrCreateStatId(TaskName));
-#endif
-
-		Task->run();
-		Task->release();
-
-		RingBuffer.Start = (RingBuffer.Start + 1) % RingBuffer.Size;
-		--RingBuffer.Num;
-	}
-}
-
-
-DECLARE_CYCLE_STAT(TEXT("PhysX Single Thread Task"), STAT_PhysXSingleThread, STATGROUP_Physics);
-
-/** Used to dispatch physx tasks to the game thread */
-class FPhysXCPUDispatcherSingleThread : public PxCpuDispatcher
-{
-	TArray<PxBaseTask*> TaskStack;
-
-	virtual void submitTask(PxBaseTask& Task) override
-	{
-		SCOPE_CYCLE_COUNTER(STAT_PhysXSingleThread);
-		CSV_SCOPED_TIMING_STAT_EXCLUSIVE(Physics);
-		CSV_SCOPED_TIMING_STAT(PhysicsVerbose, StepSolver);
-
-		TaskStack.Push(&Task);
-		if (TaskStack.Num() > 1)
-		{
-			return;
-		}
-
-		{
-#if ENABLE_STATNAMEDEVENTS || STATS
-			FPhysTaskScopedNamedEvent TaskEvent(&Task);
-#endif
-
-#if STATS
-			const char* TaskName = Task.getName();
-			FScopeCycleCounter CycleCounter(DynamicStatsHelper::FindOrCreateStatId(TaskName));
-#endif
-
-			Task.run();
-			Task.release();
-		}
-
-		while (TaskStack.Num() > 1)
-		{
-			PxBaseTask& ChildTask = *TaskStack.Pop();
-			{
-
-#if ENABLE_STATNAMEDEVENTS || STATS
-				FPhysTaskScopedNamedEvent TaskEvent(&ChildTask);
-#endif
-
-#if STATS
-				const char* ChildTaskName = ChildTask.getName();
-				FScopeCycleCounter CycleCounter(DynamicStatsHelper::FindOrCreateStatId(ChildTaskName));
-#endif
-				ChildTask.run();
-				ChildTask.release();
-			}
-		}
-		verify(&Task == TaskStack.Pop() && !TaskStack.Num());
-	}
-
-	virtual PxU32 getWorkerCount() const override
-	{
-		return 1;
-	}
-};
-
-TSharedPtr<ISimEventCallbackFactory> FPhysScene_PhysX::SimEventCallbackFactory;
-TSharedPtr<IContactModifyCallbackFactory> FPhysScene_PhysX::ContactModifyCallbackFactory;
-TSharedPtr<ICCDContactModifyCallbackFactory> FPhysScene_PhysX::CCDContactModifyCallbackFactory;
-
-#endif // WITH_PHYSX
-
-TSharedPtr<IPhysicsReplicationFactory> FPhysScene_PhysX::PhysicsReplicationFactory;
-
-static void StaticSetPhysXTreeRebuildRate(const TArray<FString>& Args, UWorld* World)
-{
-	if (Args.Num() > 0)
-	{
-		const int32 NewRate = FCString::Atoi(*Args[0]);
-		if(World && World->GetPhysicsScene())
-		{
-			World->GetPhysicsScene()->SetPhysXTreeRebuildRate(NewRate);
-		}
-	}
-	else
-	{
-		UE_LOG(LogPhysics, Warning, TEXT("Usage: p.PhysXTreeRebuildRate <num_frames>"));
-	}
-}
-
-static FAutoConsoleCommandWithWorldAndArgs GSetPhysXTreeRebuildRate(TEXT("p.PhysXTreeRebuildRate"), TEXT("Utility function to change PhysXTreeRebuildRate, useful when profiling fetchResults vs scene queries."),
-	FConsoleCommandWithWorldAndArgsDelegate::CreateStatic(&StaticSetPhysXTreeRebuildRate)
-);
-
-
-/** Exposes creation of physics-engine scene outside Engine (for use with Physics Asset Editor for example). */
-FPhysScene_PhysX::FPhysScene_PhysX(const AWorldSettings* Settings)
-#if WITH_CUSTOM_SQ_STRUCTURE
-	: SQAccelerator(nullptr)
-#if WITH_PHYSX
-	, PhysXSQAccelerator(nullptr)
-#endif
-#endif 
-{
-	LLM_SCOPE(ELLMTag::PhysX);
-	LineBatcher = NULL;
-	OwningWorld = NULL;
-#if WITH_PHYSX
-	PhysxUserData = FPhysxUserData(this);
-#endif	//#if WITH_PHYSX
-
-	UPhysicsSettings * PhysSetting = UPhysicsSettings::Get();
-	FMemory::Memzero(FrameTimeSmoothingFactor);
-	FrameTimeSmoothingFactor= PhysSetting->SyncSceneSmoothingFactor;
-
-	bSubstepping = PhysSetting->bSubstepping;
-
-	PhysXTreeRebuildRate = PhysSetting->PhysXTreeRebuildRate;
-
-	// Create the physics scene
-	InitPhysScene(Settings);
-
-	// Also initialize scene data
-	bPhysXSceneExecuting = false;
-
-	// Initialize to a value which would be acceptable if FrameTimeSmoothingFactor[i] = 1.0f, i.e. constant simulation substeps
-	AveragedFrameTime = PhysSetting->InitialAverageFrameRate;
-
-	// gets from console variable, and clamp to [0, 1] - 1 should be fixed time as 30 fps
-	FrameTimeSmoothingFactor = FMath::Clamp<float>(FrameTimeSmoothingFactor, 0.f, 1.f);
-
-	// Create replication manager
-	PhysicsReplication = PhysicsReplicationFactory.IsValid() ? PhysicsReplicationFactory->Create(this) : new FPhysicsReplication(this);
-
-	PreGarbageCollectDelegateHandle = FCoreUObjectDelegates::GetPreGarbageCollectDelegate().AddRaw(this, &FPhysScene_PhysX::WaitPhysScenes);
-
-#if WITH_PHYSX
-	// Initialise PhysX scratch buffers (only if size > 0)
-	int32 SceneScratchBufferSize = PhysSetting->SimulateScratchMemorySize;
-	if(SceneScratchBufferSize > 0)
-	{
-		// Make sure that SceneScratchBufferSize is a multiple of 16K as requested by PhysX.
-		SceneScratchBufferSize = FMath::DivideAndRoundUp<int32>(SceneScratchBufferSize, SimScratchBufferBoundary) * SimScratchBufferBoundary;
-		
-		if(PxScene* Scene = GetPxScene())
-		{
-			// We have a valid scene, so allocate the buffer for it
-			SimScratchBuffer.Buffer = (uint8*)FMemory::Malloc(SceneScratchBufferSize, 16);
-			SimScratchBuffer.BufferSize = SceneScratchBufferSize;
-		}
-	}
-#endif
-}
-
-
-void FPhysScene_PhysX::AddActorsToPhysXScene_AssumesLocked(const TArray<FPhysicsActorHandle>& InActors)
-{
-	// Check we have a sync scene
-	if (PxScene* PScene = GetPxScene())
-	{
-		// If not simulating at the moment, can use batch add
-		if (!bIsSceneSimulating)
-		{
-			TArray<PxActor*> PActors;
-			for (const FPhysicsActorHandle& ActorRef : InActors)
-			{
-				if (ActorRef.SyncActor != nullptr)
-				{
-					PActors.Add(ActorRef.SyncActor);
-				}
-			}
-
-			PScene->addActors(PActors.GetData(), PActors.Num());
-		}
-		// If we are simulating, add one at a time
-		else
-		{
-			for (const FPhysicsActorHandle& ActorRef : InActors)
-			{
-				if (ActorRef.SyncActor != nullptr)
-				{
-					PScene->addActor(*ActorRef.SyncActor);
-				}
-			}
-		}
-	}
-
-	/** Disabled the following case until we have a good broadphase. Right now using FPhysXSQAccelerator to just query the scene */
-//#if WITH_CUSTOM_SQ_STRUCTURE
-//	for (const FPhysicsActorHandle& ActorRef : InActors)
-//	{
-//		if (PxRigidActor* RigidActor = ActorRef.SyncActor)
-//		{
-//			RigidActorToSQEntries.Add(RigidActor, SQAccelerator->AddEntry(RigidActor));
-//		}
-//	}
-//#endif
-}
-
-ISQAccelerator* FPhysScene_PhysX::GetSQAccelerator() const
-{
-#if WITH_CUSTOM_SQ_STRUCTURE
-	return SQAcceleratorUnion;
-#endif
-
-	return nullptr;
-}
-
-FSQAcceleratorUnion* FPhysScene_PhysX::GetSQAcceleratorUnion() const
-{
-#if WITH_CUSTOM_SQ_STRUCTURE
-	return SQAcceleratorUnion;
-#endif
-
-	return nullptr;
-}
-
-void FPhysScene_PhysX::AddActorsToScene_AssumesLocked(const TArray<FPhysicsActorHandle>& InActors)
-{
-	AddActorsToPhysXScene_AssumesLocked(InActors);
-}
-
-void FPhysScene_PhysX::AddAggregateToScene(const FPhysicsAggregateHandle& InAggregate)
-{
-	if (PxScene* PScene = GetPxScene())
-	{
-		SCOPED_SCENE_WRITE_LOCK(PScene);
-		// add Aggregate into the scene
-		if (InAggregate.IsValid() && InAggregate.Aggregate->getNbActors() > 0)
-		{
-			PScene->addAggregate(*InAggregate.Aggregate);
-		}
-	}
-}
-
-
-void FPhysScene_PhysX::SetOwningWorld(UWorld* InOwningWorld)
-{
-	OwningWorld = InOwningWorld;
-}
-
-/** Exposes destruction of physics-engine scene outside Engine. */
-FPhysScene_PhysX::~FPhysScene_PhysX()
-{
-	FCoreUObjectDelegates::GetPreGarbageCollectDelegate().Remove(PreGarbageCollectDelegateHandle);
-	// Make sure no scenes are left simulating (no-ops if not simulating)
-	WaitPhysScenes();
-
-	if (IPhysicsReplicationFactory* RawReplicationFactory = PhysicsReplicationFactory.Get())
-	{
-		RawReplicationFactory->Destroy(PhysicsReplication);
-	}
-	else
-	{
-		delete PhysicsReplication;
-	}
-
-	// Destroy the physics scene
-	TermPhysScene();
-
-#if WITH_PHYSX
-	GPhysCommandHandler->DeferredDeleteCPUDispathcer(CPUDispatcher);
-#endif	//#if WITH_PHYSX
-
-#if WITH_PHYSX
-	// Free the scratch buffers
-	if(SimScratchBuffer.Buffer != nullptr)
-	{
-		FMemory::Free(SimScratchBuffer.Buffer);
-		SimScratchBuffer.Buffer = nullptr;
-		SimScratchBuffer.BufferSize = 0;
-	}
-#endif
-}
-
-bool FPhysScene_PhysX::GetKinematicTarget_AssumesLocked(const FBodyInstance* BodyInstance, FTransform& OutTM) const
-{
-#if WITH_PHYSX
-	if (PxRigidDynamic * PRigidDynamic = FPhysicsInterface_PhysX::GetPxRigidDynamic_AssumesLocked(BodyInstance->GetPhysicsActorHandle()))
-	{
-		if (IsSubstepping())
-		{
-			return PhysSubStepper->GetKinematicTarget_AssumesLocked(BodyInstance, OutTM);
-		}
-		else
-		{
-			PxTransform POutTM;
-			bool validTM = PRigidDynamic->getKinematicTarget(POutTM);
-			if (validTM)
-			{
-				OutTM = P2UTransform(POutTM);
-				return true;
-			}
-		}
-	}
-#endif
-
-	return false;
-}
-
-void FPhysScene_PhysX::SetKinematicTarget_AssumesLocked(FBodyInstance* BodyInstance, const FTransform& TargetTransform, bool bAllowSubstepping)
-{
-	TargetTransform.DiagnosticCheck_IsValid();
-
-#if WITH_PHYSX
-	if (PxRigidDynamic * PRigidDynamic = FPhysicsInterface_PhysX::GetPxRigidDynamic_AssumesLocked(BodyInstance->GetPhysicsActorHandle()))
-	{
-		const bool bIsKinematicTarget = IsRigidBodyKinematicAndInSimulationScene_AssumesLocked(PRigidDynamic);
-		if(bIsKinematicTarget)
-		{
-			if (bAllowSubstepping && IsSubstepping())
-			{
-				PhysSubStepper->SetKinematicTarget_AssumesLocked(BodyInstance, TargetTransform);
-			}
-
-			const PxTransform PNewPose = U2PTransform(TargetTransform);
-			PRigidDynamic->setKinematicTarget(PNewPose);	//If we interpolate, we will end up setting the kinematic target once per sub-step. However, for the sake of scene queries we should do this right away
-		}
-		else
-		{
-			const PxTransform PNewPose = U2PTransform(TargetTransform);
-			PRigidDynamic->setGlobalPose(PNewPose);
-		}
-	}
-#endif
-}
-
-void FPhysScene_PhysX::AddCustomPhysics_AssumesLocked(FBodyInstance* BodyInstance, FCalculateCustomPhysics& CalculateCustomPhysics)
-{
-#if WITH_PHYSX
-	if (IsSubstepping())
-	{
-		PhysSubStepper->AddCustomPhysics_AssumesLocked(BodyInstance, CalculateCustomPhysics);
-	}
-	else
-	{
-		// Since physics frame is set up before "pre-physics" tick group is called, can just fetch delta time from there
-		CalculateCustomPhysics.ExecuteIfBound(this->DeltaSeconds, BodyInstance);
-	}
-#endif
-}
-
-void FPhysScene_PhysX::AddForce_AssumesLocked(FBodyInstance* BodyInstance, const FVector& Force, bool bAllowSubstepping, bool bAccelChange)
-{
-#if WITH_PHYSX
-
-	if (PxRigidBody * PRigidBody = FPhysicsInterface_PhysX::GetPxRigidBody_AssumesLocked(BodyInstance->GetPhysicsActorHandle()))
-	{
-		if (bAllowSubstepping && IsSubstepping())
-		{
-			PhysSubStepper->AddForce_AssumesLocked(BodyInstance, Force, bAccelChange);
-		}
-		else
-		{
-			PRigidBody->addForce(U2PVector(Force), bAccelChange ? PxForceMode::eACCELERATION : PxForceMode::eFORCE, true);
-		}
-	}
-#endif
-}
-
-void FPhysScene_PhysX::AddForceAtPosition_AssumesLocked(FBodyInstance* BodyInstance, const FVector& Force, const FVector& Position, bool bAllowSubstepping, bool bIsLocalForce)
-{
-#if WITH_PHYSX
-
-	if (PxRigidBody * PRigidBody = FPhysicsInterface_PhysX::GetPxRigidBody_AssumesLocked(BodyInstance->GetPhysicsActorHandle()))
-	{
-		if (bAllowSubstepping && IsSubstepping())
-		{
-			PhysSubStepper->AddForceAtPosition_AssumesLocked(BodyInstance, Force, Position, bIsLocalForce);
-		}
-		else if (!bIsLocalForce)
-		{
-			PxRigidBodyExt::addForceAtPos(*PRigidBody, U2PVector(Force), U2PVector(Position), PxForceMode::eFORCE, true);
-		}
-		else
-		{
-			PxRigidBodyExt::addLocalForceAtLocalPos(*PRigidBody, U2PVector(Force), U2PVector(Position), PxForceMode::eFORCE, true);
-		}
-	}
-#endif
-}
-
-void FPhysScene_PhysX::AddRadialForceToBody_AssumesLocked(FBodyInstance* BodyInstance, const FVector& Origin, const float Radius, const float Strength, const uint8 Falloff, bool bAccelChange, bool bAllowSubstepping)
-{
-#if WITH_PHYSX
-
-	if (PxRigidBody * PRigidBody = FPhysicsInterface_PhysX::GetPxRigidBody_AssumesLocked(BodyInstance->GetPhysicsActorHandle()))
-	{
-		if (bAllowSubstepping && IsSubstepping())
-		{
-			PhysSubStepper->AddRadialForceToBody_AssumesLocked(BodyInstance, Origin, Radius, Strength, Falloff, bAccelChange);
-		}
-		else
-		{
-			AddRadialForceToPxRigidBody_AssumesLocked(*PRigidBody, Origin, Radius, Strength, Falloff, bAccelChange);
-		}
-	}
-#endif
-}
-
-void FPhysScene_PhysX::ClearForces_AssumesLocked(FBodyInstance* BodyInstance, bool bAllowSubstepping)
-{
-#if WITH_PHYSX
-	if (PxRigidBody * PRigidBody = FPhysicsInterface_PhysX::GetPxRigidBody_AssumesLocked(BodyInstance->GetPhysicsActorHandle()))
-	{
-		PRigidBody->clearForce();
-		if (bAllowSubstepping && IsSubstepping())
-		{
-			PhysSubStepper->ClearForces_AssumesLocked(BodyInstance);
-		}
-	}
-#endif
-}
-
-
-void FPhysScene_PhysX::AddTorque_AssumesLocked(FBodyInstance* BodyInstance, const FVector& Torque, bool bAllowSubstepping, bool bAccelChange)
-{
-#if WITH_PHYSX
-
-	if (PxRigidBody * PRigidBody = FPhysicsInterface_PhysX::GetPxRigidBody_AssumesLocked(BodyInstance->GetPhysicsActorHandle()))
-	{
-		if (bAllowSubstepping && IsSubstepping())
-		{
-			PhysSubStepper->AddTorque_AssumesLocked(BodyInstance, Torque, bAccelChange);
-		}
-		else
-		{
-			PRigidBody->addTorque(U2PVector(Torque), bAccelChange ? PxForceMode::eACCELERATION : PxForceMode::eFORCE, true);
-		}
-	}
-#endif
-}
-
-void FPhysScene_PhysX::ClearTorques_AssumesLocked(FBodyInstance* BodyInstance, bool bAllowSubstepping)
-{
-#if WITH_PHYSX
-	if (PxRigidBody * PRigidBody = FPhysicsInterface_PhysX::GetPxRigidBody_AssumesLocked(BodyInstance->GetPhysicsActorHandle()))
-	{
-		PRigidBody->clearTorque();
-		if (bAllowSubstepping && IsSubstepping())
-		{
-			PhysSubStepper->ClearTorques_AssumesLocked(BodyInstance);
-		}
-	}
-#endif
-}
-
-void FPhysScene_PhysX::RemoveBodyInstanceFromPendingLists_AssumesLocked(FBodyInstance* BodyInstance)
-{
-#if WITH_PHYSX
-
-	/** Disabled the following case until we have a good broadphase. Right now using FPhysXSQAccelerator to just query the scene */
-//#if WITH_CUSTOM_SQ_STRUCTURE
-//	if (PxRigidActor* RigidActor = BodyInstance->GetPhysicsActorHandle().SyncActor)
-//	{
-//		FSQAcceleratorEntry* Entry = nullptr;
-//		RigidActorToSQEntries.RemoveAndCopyValue(RigidActor, Entry);
-//		if (Entry)
-//		{
-//			SQAccelerator->RemoveEntry(Entry);
-//		}
-//	}
-//#endif
-
-
-	if (FPhysicsInterface_PhysX::IsRigidBody(BodyInstance->GetPhysicsActorHandle()))
-	{
-		PhysSubStepper->RemoveBodyInstance_AssumesLocked(BodyInstance);
-	}
-
-	PendingSleepEvents.Remove(BodyInstance);
-#endif // WITH_PHYSX
-}
-
-FAutoConsoleTaskPriority CPrio_PhysXStepSimulation(
-	TEXT("TaskGraph.TaskPriorities.PhysXStepSimulation"),
-	TEXT("Task and thread priority for FPhysSubstepTask::StepSimulation."),
-	ENamedThreads::HighThreadPriority, // if we have high priority task threads, then use them...
-	ENamedThreads::NormalTaskPriority, // .. at normal task priority
-	ENamedThreads::HighTaskPriority // if we don't have hi pri threads, then use normal priority threads at high task priority instead
-);
-
-bool FPhysScene_PhysX::SubstepSimulation(FGraphEventRef &InOutCompletionEvent)
-{
-#if WITH_PHYSX
-	float UseDelta = DeltaSeconds;
-	float SubTime = PhysSubStepper->UpdateTime(UseDelta);
-	PxScene* PScene = GetPxScene();
-	if(SubTime <= 0.f)
-	{
-		return false;
-	}else
-	{
-		//we have valid scene and subtime so enqueue task
-		PhysXCompletionTask* Task = new PhysXCompletionTask(InOutCompletionEvent, PScene->getTaskManager(), &SimScratchBuffer);
-		ENamedThreads::Type NamedThread = PhysSingleThreadedMode() ? ENamedThreads::GameThread : ENamedThreads::SetTaskPriority(ENamedThreads::GameThread, ENamedThreads::HighTaskPriority);
-
-		DECLARE_CYCLE_STAT(TEXT("FSimpleDelegateGraphTask.SubstepSimulationImp"),
-		STAT_FSimpleDelegateGraphTask_SubstepSimulationImp,
-			STATGROUP_TaskGraphTasks);
-
-		FSimpleDelegateGraphTask::CreateAndDispatchWhenReady(
-			FSimpleDelegateGraphTask::FDelegate::CreateRaw(PhysSubStepper, &FPhysSubstepTask::StepSimulation, Task),
-			GET_STATID(STAT_FSimpleDelegateGraphTask_SubstepSimulationImp), NULL, NamedThread
-		);
-		return true;
-	}
-#else
-	return false;
-#endif
-
-}
-
-
-/** Adds to queue of skelmesh we want to add to collision disable table */
-void FPhysScene_PhysX::DeferredAddCollisionDisableTable(uint32 SkelMeshCompID, TMap<struct FRigidBodyIndexPair, bool> * CollisionDisableTable)
-{
-	check(IsInGameThread());
-
-	FPendingCollisionDisableTable PendingCollisionDisableTable;
-	PendingCollisionDisableTable.SkelMeshCompID = SkelMeshCompID;
-	PendingCollisionDisableTable.CollisionDisableTable = CollisionDisableTable;
-
-	DeferredCollisionDisableTableQueue.Add(PendingCollisionDisableTable);
-}
-
-/** Adds to queue of skelmesh we want to remove from collision disable table */
-void FPhysScene_PhysX::DeferredRemoveCollisionDisableTable(uint32 SkelMeshCompID)
-{
-	check(IsInGameThread());
-
-	FPendingCollisionDisableTable PendingDisableCollisionTable;
-	PendingDisableCollisionTable.SkelMeshCompID = SkelMeshCompID;
-	PendingDisableCollisionTable.CollisionDisableTable = NULL;
-
-	DeferredCollisionDisableTableQueue.Add(PendingDisableCollisionTable);
-}
-
-void FPhysScene_PhysX::FlushDeferredCollisionDisableTableQueue()
-{
-	check(IsInGameThread());
-	for (int32 i = 0; i < DeferredCollisionDisableTableQueue.Num(); ++i)
-	{
-		FPendingCollisionDisableTable & PendingCollisionDisableTable = DeferredCollisionDisableTableQueue[i];
-
-		if (PendingCollisionDisableTable.CollisionDisableTable)
-		{
-			CollisionDisableTableLookup.Add(PendingCollisionDisableTable.SkelMeshCompID, PendingCollisionDisableTable.CollisionDisableTable);
-		}
-		else
-		{
-			CollisionDisableTableLookup.Remove(PendingCollisionDisableTable.SkelMeshCompID);
-		}
-	}
-
-	DeferredCollisionDisableTableQueue.Empty();
-}
-
-#if WITH_PHYSX
-
-void GatherPhysXStats_AssumesLocked(PxScene* PSyncScene)
-{
-	/** Gather PhysX stats */
-	if (PSyncScene)
-	{
-		PxSimulationStatistics SimStats;
-		PSyncScene->getSimulationStatistics(SimStats);
-
-		SET_DWORD_STAT(STAT_NumActiveConstraints, SimStats.nbActiveConstraints);
-		SET_DWORD_STAT(STAT_NumActiveSimulatedBodies, SimStats.nbActiveDynamicBodies);
-		SET_DWORD_STAT(STAT_NumActiveKinematicBodies, SimStats.nbActiveKinematicBodies);
-		SET_DWORD_STAT(STAT_NumStaticBodies, SimStats.nbStaticBodies);
-		SET_DWORD_STAT(STAT_NumMobileBodies, SimStats.nbDynamicBodies);
-
-		
-		CSV_CUSTOM_STAT(PhysicsCounters, NumPotentialContacts, (int32)SimStats.nbDiscreteContactPairsTotal, ECsvCustomStatOp::Set);
-		CSV_CUSTOM_STAT(PhysicsCounters, NumContacts, (int32)SimStats.nbDiscreteContactPairsWithContacts, ECsvCustomStatOp::Set);
-		//CSV_CUSTOM_STAT(PhysicsVerbose, NumActiveConstraints, (int32)SimStats.nbActiveConstraints, ECsvCustomStatOp::Set);
-		CSV_CUSTOM_STAT(PhysicsCounters, NumActiveDynamicBodies, (int32)SimStats.nbActiveDynamicBodies, ECsvCustomStatOp::Set);
-		CSV_CUSTOM_STAT(PhysicsCounters, NumActiveKinematicBodies, (int32)SimStats.nbActiveKinematicBodies, ECsvCustomStatOp::Set);
-		CSV_CUSTOM_STAT(PhysicsCounters, NumStaticBodies, (int32)SimStats.nbStaticBodies, ECsvCustomStatOp::Set);
-		CSV_CUSTOM_STAT(PhysicsCounters, NumDynamicBodies, (int32)SimStats.nbDynamicBodies, ECsvCustomStatOp::Set);
-
-		//SET_DWORD_STAT(STAT_NumBroadphaseAdds, SimStats.getNbBroadPhaseAdds(PxSimulationStatistics::VolumeType::eRIGID_BODY));	//TODO: These do not seem to work
-		//SET_DWORD_STAT(STAT_NumBroadphaseRemoves, SimStats.getNbBroadPhaseRemoves(PxSimulationStatistics::VolumeType::eRIGID_BODY));
-
-		uint32 NumShapes = 0;
-		for (int32 GeomType = 0; GeomType < PxGeometryType::eGEOMETRY_COUNT; ++GeomType)
-		{
-			NumShapes += SimStats.nbShapes[GeomType];
-		}
-
-		SET_DWORD_STAT(STAT_NumShapes, NumShapes);
-		CSV_CUSTOM_STAT(PhysicsCounters, NumShapes, (int32)NumShapes, ECsvCustomStatOp::Set);
-
-	}
-}
-#endif // WITH_PHYSX
-
-DECLARE_FLOAT_COUNTER_STAT(TEXT("Sim Time (ms)"), STAT_PhysSim, STATGROUP_Physics);
-
-double GSimStartTime = 0.f;
-
-void FinishSceneStat()
-{
-	float SceneTime = float(FPlatformTime::Seconds() - GSimStartTime) * 1000.0f;
-	INC_FLOAT_STAT_BY(STAT_PhysSim, SceneTime);
-}
-
-void GatherClothingStats(const UWorld* World)
-{
-#if WITH_PHYSX
-#if STATS
-	QUICK_SCOPE_CYCLE_COUNTER(STAT_GatherApexStats);
-
-	SET_DWORD_STAT(STAT_NumCloths, 0);
-	SET_DWORD_STAT(STAT_NumClothVerts, 0);
-
-	if ( FThreadStats::IsCollectingData(GET_STATID(STAT_NumCloths)) ||  FThreadStats::IsCollectingData(GET_STATID(STAT_NumClothVerts)) )
-	{
-		for (TObjectIterator<USkeletalMeshComponent> Itr; Itr; ++Itr)
-		{
-			if (Itr->GetWorld() != World) { continue; }
-
-			if(const IClothingSimulation* Simulation = Itr->GetClothingSimulation())
-			{
-				Simulation->GatherStats();
-			}
-		}
-	}
-#endif
-#endif // WITH_PHYSX
-}
-
-
-bool FPhysScene_PhysX::MarkForPreSimKinematicUpdate(USkeletalMeshComponent* InSkelComp, ETeleportType InTeleport, bool bNeedsSkinning)
-{
-	// If null, or pending kill, do nothing
-	if (IsValid(InSkelComp))
-	{
-		// If we are already flagged, just need to update info
-		if (InSkelComp->bDeferredKinematicUpdate)
-		{
-			TPair<USkeletalMeshComponent*, FDeferredKinematicUpdateInfo>* FoundItem = DeferredKinematicUpdateSkelMeshes.FindByPredicate([InSkelComp](const TPair<USkeletalMeshComponent*, FDeferredKinematicUpdateInfo>& InItem) { return InSkelComp == InItem.Key; });
-			if (!ensure(FoundItem != nullptr))// If the bool was set, we must be in the array!
-			{
-				FDeferredKinematicUpdateInfo Info;
-				Info.TeleportType = InTeleport;
-				Info.bNeedsSkinning = bNeedsSkinning;
-				DeferredKinematicUpdateSkelMeshes.Emplace(InSkelComp, Info);
-				return true;
-			}
-
-			FDeferredKinematicUpdateInfo& Info = FoundItem->Value;
-
-			// If we are currently not going to teleport physics, but this update wants to, we 'upgrade' it
-			if (Info.TeleportType == ETeleportType::None && InTeleport == ETeleportType::TeleportPhysics)
-			{
-				Info.TeleportType = ETeleportType::TeleportPhysics;
-			}
-
-			// If we need skinning, remember that
-			if (bNeedsSkinning)
-			{
-				Info.bNeedsSkinning = true;
-			}
-		}
-		// We are not flagged yet..
-		else
-		{
-			// Set info and add to map
-			FDeferredKinematicUpdateInfo Info;
-			Info.TeleportType = InTeleport;
-			Info.bNeedsSkinning = bNeedsSkinning;
-			DeferredKinematicUpdateSkelMeshes.Emplace(InSkelComp, Info);
-
-			// Set flag on component
-			InSkelComp->bDeferredKinematicUpdate = true;
-		}
-	}
-
-	return true;
-}
-
-void FPhysScene_PhysX::ClearPreSimKinematicUpdate(USkeletalMeshComponent* InSkelComp)
-{
-	// If non-null, and flagged for deferred update..
-	if (InSkelComp != nullptr && InSkelComp->bDeferredKinematicUpdate)
-	{
-		// Remove from map
-		int32 NumRemoved = DeferredKinematicUpdateSkelMeshes.RemoveAll([InSkelComp](const TPair<USkeletalMeshComponent*, FDeferredKinematicUpdateInfo>& InItem) { return InSkelComp == InItem.Key; });
-
-		ensure(NumRemoved == 1); // Should be in array if flag was set!
-
-		// Clear flag
-		InSkelComp->bDeferredKinematicUpdate = false;
-	}
-}
-
-
-void FPhysScene_PhysX::UpdateKinematicsOnDeferredSkelMeshes()
-{
-	SCOPE_CYCLE_COUNTER(STAT_UpdateKinematicsOnDeferredSkelMeshes);
-
-	for (const TPair<USkeletalMeshComponent*, FDeferredKinematicUpdateInfo>& DeferredKinematicUpdate : DeferredKinematicUpdateSkelMeshes)
-	{
-		USkeletalMeshComponent* SkelComp = DeferredKinematicUpdate.Key;
-		const FDeferredKinematicUpdateInfo& Info = DeferredKinematicUpdate.Value;
-
-		ensure(SkelComp->bDeferredKinematicUpdate); // Should be true if in map!
-
-		// Perform kinematic updates
-		SkelComp->UpdateKinematicBonesToAnim(SkelComp->GetComponentSpaceTransforms(), Info.TeleportType, Info.bNeedsSkinning, EAllowKinematicDeferral::DisallowDeferral);
-
-		// Clear deferred flag
-		SkelComp->bDeferredKinematicUpdate = false;
-	}
-
-	// Empty map now all is done
-	DeferredKinematicUpdateSkelMeshes.Reset();
-}
-
-/** Exposes ticking of physics-engine scene outside Engine. */
-void FPhysScene_PhysX::TickPhysScene(FGraphEventRef& InOutCompletionEvent)
-{
-	SCOPE_CYCLE_COUNTER(STAT_TotalPhysicsTime);
-	CSV_SCOPED_TIMING_STAT_EXCLUSIVE(Physics);
-
-	SCOPE_CYCLE_COUNTER(STAT_PhysicsKickOffDynamicsTime);
-
-	GSimStartTime = FPlatformTime::Seconds();
-
-	if (bPhysXSceneExecuting != 0)
-	{
-		// Already executing this scene, must call WaitPhysScene before calling this function again.
-		UE_LOG(LogPhysics, Log, TEXT("TickPhysScene: Already executing scene - aborting."));
-		return;
-	}
-
-	/**
-	* clamp down... if this happens we are simming physics slower than real-time, so be careful with it.
-	* it can improve framerate dramatically (really, it is the same as scaling all velocities down and
-	* enlarging all timesteps) but at the same time, it will screw with networking (client and server will
-	* diverge a lot more.)
-	*/
-
-	float UseDelta = FMath::Min(DeltaSeconds, MaxPhysicsDeltaTime);
-
-	// Only simulate a positive time step.
-	if (UseDelta <= 0.f)
-	{
-		if (UseDelta < 0.f)
-		{
-			// only do this if negative. Otherwise, whenever we pause, this will come up
-			UE_LOG(LogPhysics, Warning, TEXT("TickPhysScene: Negative timestep (%f) - aborting."), UseDelta);
-		}
-		return;
-	}
-
-	/**
-	* Weight frame time according to PhysScene settings.
-	*/
-	AveragedFrameTime *= FrameTimeSmoothingFactor;
-	AveragedFrameTime += (1.0f - FrameTimeSmoothingFactor)*UseDelta;
-
-	// Set execution flag
-	bPhysXSceneExecuting = true;
-
-	check(!InOutCompletionEvent.GetReference()); // these should be gone because nothing is outstanding
-	InOutCompletionEvent = FGraphEvent::CreateGraphEvent();
-	bool bTaskOutstanding = false;
-
-	// Update any skeletal meshes that need their bone transforms sent to physics sim
-	UpdateKinematicsOnDeferredSkelMeshes();
-
-#if !WITH_PHYSX
-	const bool bSimulateScene = false;
-#else
-#if !WITH_APEX
-	PxScene* PScene = GetPxScene();
-	const bool bSimulateScene = PScene && (UseDelta > 0.f);
-#else
-	apex::Scene* ApexScene = GetApexScene();
-	const bool bSimulateScene = ApexScene && UseDelta > 0.f;
-#endif
-#endif
-
-	// Replicate physics
-#if WITH_PHYSX
-	if (bSimulateScene && PhysicsReplication)
-	{
-		PhysicsReplication->Tick(AveragedFrameTime);
-	}
-#endif
-
-	float PreTickTime = IsSubstepping() ? UseDelta : AveragedFrameTime;
-
-	// Broadcast 'pre tick' delegate
-	OnPhysScenePreTick.Broadcast(this, PreTickTime);
-
-	// If not substepping, call this delegate here. Otherwise we call it in FPhysSubstepTask::SubstepSimulationStart
-	if (IsSubstepping() == false)
-	{
-		OnPhysSceneStep.Broadcast(this, PreTickTime);
-	}
-	else
-	{
-		//We're about to start stepping so swap buffers. Might want to find a better place for this?
-		PhysSubStepper->SwapBuffers();
-	}
-
-#if WITH_PHYSX
-	bIsSceneSimulating = true;
-
-	if (bSimulateScene)
-	{
-		if(IsSubstepping()) //we don't bother sub-stepping cloth
-		{
-			bTaskOutstanding = SubstepSimulation(InOutCompletionEvent);
-		}
-		else
-		{
-#if !WITH_APEX
-			PhysXCompletionTask* Task = new PhysXCompletionTask(InOutCompletionEvent, PScene->getTaskManager());
-			PScene->lockWrite();
-			PScene->simulate(AveragedFrameTime, Task, SimScratchBuffer.Buffer, SimScratchBuffer.BufferSize);
-			PScene->unlockWrite();
-			Task->removeReference();
-			bTaskOutstanding = true;
-#else
-			PhysXCompletionTask* Task = new PhysXCompletionTask(InOutCompletionEvent, ApexScene->getTaskManager());
-			ApexScene->simulate(AveragedFrameTime, true, Task, SimScratchBuffer.Buffer, SimScratchBuffer.BufferSize);
-			Task->removeReference();
-			bTaskOutstanding = true;
-#endif
-		}
-	}
-
-#endif // WITH_PHYSX
-
-	if (!bTaskOutstanding)
-	{
-		InOutCompletionEvent->DispatchSubsequents(); // nothing to do, so nothing to wait for
-	}
-
-	bSubstepping = UPhysicsSettings::Get()->bSubstepping;
-}
-
-void FPhysScene_PhysX::KillVisualDebugger()
-{
-#if WITH_PHYSX
-	if (GPhysXVisualDebugger)
-	{
-		GPhysXVisualDebugger->disconnect();
-	}
-#endif // WITH_PHYSX 
-}
-
-void FPhysScene_PhysX::WaitPhysScenes()
-{
-	check(IsInGameThread());
-
-	FGraphEventArray ThingsToComplete;
-	if (PhysicsSceneCompletion.GetReference())
-	{
-		ThingsToComplete.Add(PhysicsSceneCompletion);
-	}
-	{
-		if (PhysicsSubsceneCompletion.GetReference())
-		{
-			ThingsToComplete.Add(PhysicsSubsceneCompletion);
-		}
-	}
-	if (ThingsToComplete.Num())
-	{
-		QUICK_SCOPE_CYCLE_COUNTER(STAT_FPhysScene_WaitPhysScenes);
-		FTaskGraphInterface::Get().WaitUntilTasksComplete(ThingsToComplete, ENamedThreads::GameThread);
-	}
-
-
-}
-
-FGraphEventArray FPhysScene_PhysX::GetCompletionEvents()
-{
-	FGraphEventArray CompletionEvents;
-	CompletionEvents.Add(PhysicsSceneCompletion);
-	return CompletionEvents;
-
-}
-
-bool FPhysScene_PhysX::IsCompletionEventComplete() const
-{
-	if (PhysicsSceneCompletion && !PhysicsSceneCompletion->IsComplete())
-	{
-		return false;
-	}
-
-	return true;
-}
-
-void FPhysScene_PhysX::SceneCompletionTask(ENamedThreads::Type CurrentThread, const FGraphEventRef& MyCompletionGraphEvent)
-{
-	ProcessPhysScene();
-}
-
-void FPhysScene_PhysX::ProcessPhysScene()
-{
-	LLM_SCOPE(ELLMTag::PhysX);
-
-	CSV_SCOPED_TIMING_STAT_EXCLUSIVE(Physics);
-
-	SCOPED_NAMED_EVENT(FPhysScene_ProcessPhysScene, FColor::Orange);
-
-	SCOPE_CYCLE_COUNTER(STAT_TotalPhysicsTime);
-	SCOPE_CYCLE_COUNTER(STAT_PhysicsFetchDynamicsTime);
-
-	if (bPhysXSceneExecuting == 0)
-	{
-		// Not executing this scene, must call TickPhysScene before calling this function again.
-		UE_LOG(LogPhysics, Log, TEXT("WaitPhysScene`: Not executing this scene - aborting."));
-		return;
-	}
-
-	// Reset execution flag
-
-	bool bSuccess = false;
-
-#if WITH_PHYSX
-	//This fetches and gets active transforms. It's important that the function that calls this locks because getting the transforms and using the data must be an atomic operation
-	PxScene* PScene = GetPxScene();
-	check(PScene);
-	PxU32 OutErrorCode = 0;
-
-	PScene->lockWrite();
-#if !WITH_APEX
-	bSuccess = PScene->fetchResults(true, &OutErrorCode);
-#else	//	#if !WITH_APEX
-	// The APEX scene calls the fetchResults function for the PhysX scene, so we only call ApexScene->fetchResults().
-	apex::Scene* ApexScene = GetApexScene();
-	check(ApexScene);
-	bSuccess = ApexScene->fetchResults(true, &OutErrorCode);
-#endif	//	#if !WITH_APEX
-
-	GPhysXHackCurrentLoopCounter = 0;
-	if (OutErrorCode != 0)
-	{
-		UE_LOG(LogPhysics, Log, TEXT("PHYSX FETCHRESULTS ERROR: %d"), OutErrorCode);
-	}
-
-	SyncComponentsToBodies_AssumesLocked();
-	PScene->unlockWrite();
-#endif // WITH_PHYSX
-
-	PhysicsSubsceneCompletion = nullptr;
-	bPhysXSceneExecuting = false;
-
-#if WITH_PHYSX
-	bIsSceneSimulating = false;
-#endif
-
-	// Broadcast 'post tick' delegate
-	OnPhysScenePostTick.Broadcast(this);
-}
-
-/** Struct to remember a pending component transform change */
-struct FPhysScenePendingComponentTransform_PhysX
-{
-	/** Component to move */
-	TWeakObjectPtr<UPrimitiveComponent> OwningComp;
-	/** New transform from physics engine */
-	FTransform NewTransform;
-
-	FPhysScenePendingComponentTransform_PhysX(UPrimitiveComponent* InOwningComp, const FTransform& InNewTransform)
-		: OwningComp(InOwningComp)
-		, NewTransform(InNewTransform)
-	{}
-};
-
-void FPhysScene_PhysX::SyncComponentsToBodies_AssumesLocked()
-{
-	SCOPE_CYCLE_COUNTER(STAT_TotalPhysicsTime);
-	SCOPE_CYCLE_COUNTER(STAT_SyncComponentsToBodies);
-
-#if WITH_PHYSX
-	PxScene* PScene = GetPxScene();
-	check(PScene);
-
-	/** Array of custom sync handlers (plugins) */
-	TArray<FCustomPhysXSyncActors*> CustomPhysXSyncActors;
-
-	PxU32 NumActors = 0;
-	PxActor** PActiveActors = PScene->getActiveActors(NumActors);
-
-	TArray<FPhysScenePendingComponentTransform_PhysX> PendingTransforms;
-
-	for (PxU32 TransformIdx = 0; TransformIdx < NumActors; ++TransformIdx)
-	{
-		PxActor* PActiveActor = PActiveActors[TransformIdx];
-		PxRigidActor* RigidActor = PActiveActor->is<PxRigidActor>();
-
-		ensure(!RigidActor->userData || !FPhysxUserData::IsGarbage(RigidActor->userData));
-
-		if (FBodyInstance* BodyInstance = FPhysxUserData::Get<FBodyInstance>(RigidActor->userData))
-		{
-			if (BodyInstance->InstanceBodyIndex == INDEX_NONE && BodyInstance->OwnerComponent.IsValid())
-			{
-				check(BodyInstance->OwnerComponent->IsRegistered()); // shouldn't have a physics body for a non-registered component!
-
-				const FTransform NewTransform = BodyInstance->GetUnrealWorldTransform_AssumesLocked();
-
-				// Add to set of transforms to process
-				// We can't actually move the component now (or check for out of world), because that could destroy a body
-				// elsewhere in the PActiveActors array, resulting in a bad pointer
-				FPhysScenePendingComponentTransform_PhysX NewEntry(BodyInstance->OwnerComponent.Get(), NewTransform);
-				PendingTransforms.Add(NewEntry);
-			}
-		}
-		else if (const FCustomPhysXPayload* CustomPayload = FPhysxUserData::Get<FCustomPhysXPayload>(RigidActor->userData))
-		{
-			if(CustomPayload->CustomSyncActors)
-			{
-				CustomPhysXSyncActors.AddUnique(CustomPayload->CustomSyncActors);	//NOTE: AddUnique because the assumed number of plugins that rely on this is very small
-				CustomPayload->CustomSyncActors->Actors.Add(RigidActor);
-			}
-		}
-	}
-
-	//Give custom plugins the chance to build the sync data
-	for (FCustomPhysXSyncActors* CustomSync : CustomPhysXSyncActors)
-	{
-		CustomSync->BuildSyncData_AssumesLocked(CustomSync->Actors);
-		CustomSync->Actors.Empty(CustomSync->Actors.Num());
-	}
-
-	//Allow custom plugins to actually act on the sync data
-	for (FCustomPhysXSyncActors* CustomSync : CustomPhysXSyncActors)
-	{
-		CustomSync->FinalizeSync();
-	}
-
-	/// Now actually move components
-	for (FPhysScenePendingComponentTransform_PhysX& Entry : PendingTransforms)
-	{
-		// Check if still valid (ie not destroyed)
-		UPrimitiveComponent* OwnerComponent = Entry.OwningComp.Get();
-		if (OwnerComponent != nullptr)
-		{
-			AActor* Owner = OwnerComponent->GetOwner();
-
-			// See if the transform is actually different, and if so, move the component to match physics
-			if (!Entry.NewTransform.EqualsNoScale(OwnerComponent->GetComponentTransform()))
-			{
-				const FVector MoveBy = Entry.NewTransform.GetLocation() - OwnerComponent->GetComponentTransform().GetLocation();
-				const FQuat NewRotation = Entry.NewTransform.GetRotation();
-
-				//@warning: do not reference BodyInstance again after calling MoveComponent() - events from the move could have made it unusable (destroying the actor, SetPhysics(), etc)
-				OwnerComponent->MoveComponent(MoveBy, NewRotation, false, NULL, MOVECOMP_SkipPhysicsMove);
-			}
-
-			// Check if we didn't fall out of the world
-			if (IsValid(Owner))
-			{
-				Owner->CheckStillInWorld();
-			}
-		}
-	}
-
-#endif // WITH_PHYSX 
-}
-
-void FPhysScene_PhysX::DispatchPhysNotifications_AssumesLocked()
-{
-#if WITH_PHYSX
-	SCOPE_CYCLE_COUNTER(STAT_PhysicsEventTime);
-
-	{
-		TArray<FCollisionNotifyInfo>& PendingCollisionNotifies = GetPendingCollisionNotifies();
-
-		// Let the game-specific PhysicsCollisionHandler process any physics collisions that took place
-		if (OwningWorld != NULL && OwningWorld->PhysicsCollisionHandler != NULL)
-		{
-			OwningWorld->PhysicsCollisionHandler->HandlePhysicsCollisions_AssumesLocked(PendingCollisionNotifies);
-		}
-
-		// Fire any collision notifies in the queue.
-		for (int32 i = 0; i<PendingCollisionNotifies.Num(); i++)
-		{
-			FCollisionNotifyInfo& NotifyInfo = PendingCollisionNotifies[i];
-			if (NotifyInfo.RigidCollisionData.ContactInfos.Num() > 0)
-			{
-				if (NotifyInfo.bCallEvent0 && NotifyInfo.IsValidForNotify() && NotifyInfo.Info0.Actor.IsValid())
-				{
-					NotifyInfo.Info0.Actor->DispatchPhysicsCollisionHit(NotifyInfo.Info0, NotifyInfo.Info1, NotifyInfo.RigidCollisionData);
-				}
-
-				// Need to check IsValidForNotify again in case first call broke something.
-				if (NotifyInfo.bCallEvent1 && NotifyInfo.IsValidForNotify() && NotifyInfo.Info1.Actor.IsValid())
-				{
-					NotifyInfo.RigidCollisionData.SwapContactOrders();
-					NotifyInfo.Info1.Actor->DispatchPhysicsCollisionHit(NotifyInfo.Info1, NotifyInfo.Info0, NotifyInfo.RigidCollisionData);
-				}
-			}
-		}
-		PendingCollisionNotifies.Reset();
-	}
-
-	{
-		for (auto MapItr = PendingSleepEvents.CreateIterator(); MapItr; ++MapItr)
-		{
-			FBodyInstance* BodyInstance = MapItr.Key();
-			if(UPrimitiveComponent* PrimitiveComponent = BodyInstance->OwnerComponent.Get())
-			{
-				PrimitiveComponent->DispatchWakeEvents(MapItr.Value(), BodyInstance->BodySetup->BoneName);
-			}
-		}
-
-		PendingSleepEvents.Empty();
-	}
-
-	{
-		FPendingConstraintData& ConstraintData = PendingConstraintData;
-		for(FConstraintBrokenDelegateData& ConstraintBrokenData : ConstraintData.PendingConstraintBroken)
-		{
-			ConstraintBrokenData.DispatchOnBroken();
-		}
-
-		ConstraintData.PendingConstraintBroken.Empty();
-	}
-#endif // WITH_PHYSX 
-
-#if WITH_CHAOS
-	check(false);
-#else
-	FPhysicsDelegates::OnPhysDispatchNotifications.Broadcast(this);
-#endif
-
-}
-
-void FPhysScene_PhysX::SetUpForFrame(const FVector* NewGrav, float InDeltaSeconds, float UnusedMinPhysicsDeltaTime, float InMaxPhysicsDeltaTime, float InMaxSubstepDeltaTime, int32 InMaxSubsteps, bool bUnused)
-{
-	DeltaSeconds = InDeltaSeconds;
-	MaxPhysicsDeltaTime = InMaxPhysicsDeltaTime;
-#if WITH_PHYSX
-	if (NewGrav)
-	{
-		if (PxScene* PScene = GetPxScene())
-		{
-			//@todo phys_thread don't do this if gravity changes
-
-			//@todo, to me it looks like we should avoid this if the gravity has not changed, the lock is probably expensive
-			// Lock scene lock, in case it is required
-			SCENE_LOCK_WRITE(PScene);
-
-			PScene->setGravity(U2PVector(*NewGrav));
-
-			// Unlock scene lock, in case it is required
-			SCENE_UNLOCK_WRITE(PScene);
-		}
-	}
-#endif
-}
-
-FAutoConsoleTaskPriority CPrio_PhyXSceneCompletion(
-	TEXT("TaskGraph.TaskPriorities.PhyXSceneCompletion"),
-	TEXT("Task and thread priority for PhysicsSceneCompletion."),
-	ENamedThreads::HighThreadPriority, // if we have high priority task threads, then use them...
-	ENamedThreads::HighTaskPriority, // .. at high task priority
-	ENamedThreads::HighTaskPriority // if we don't have hi pri threads, then use normal priority threads at high task priority instead
-);
-
-
-
-void FPhysScene_PhysX::StartFrame()
-{
-	FGraphEventArray FinishPrerequisites;
-
-	//Update the collision disable table before ticking
-	FlushDeferredCollisionDisableTableQueue();
-
-	// Run the sync scene
-	TickPhysScene(PhysicsSubsceneCompletion);
-	{
-		FGraphEventArray MainScenePrerequisites;
-		
-		if (PhysicsSubsceneCompletion.GetReference())
-		{
-			MainScenePrerequisites.Add(PhysicsSubsceneCompletion);
-
-			DECLARE_CYCLE_STAT(TEXT("FDelegateGraphTask.ProcessPhysScene_Sync"),
-			STAT_FDelegateGraphTask_ProcessPhysScene_Sync,
-				STATGROUP_TaskGraphTasks);
-
-			new (FinishPrerequisites)FGraphEventRef(
-				FDelegateGraphTask::CreateAndDispatchWhenReady(
-					FDelegateGraphTask::FDelegate::CreateRaw(this, &FPhysScene_PhysX::SceneCompletionTask),
-					GET_STATID(STAT_FDelegateGraphTask_ProcessPhysScene_Sync), &MainScenePrerequisites,
-					ENamedThreads::GameThread, ENamedThreads::GameThread
-				)
-			);
-		}
-	}
-
-	check(!PhysicsSceneCompletion.GetReference()); // this should have been cleared
-	if (FinishPrerequisites.Num())
-	{
-		// #BG Not sure this is needed anymore without async scene
-		if (FinishPrerequisites.Num() > 1)  // we don't need to create a new task if we only have one prerequisite
-		{
-			DECLARE_CYCLE_STAT(TEXT("FNullGraphTask.ProcessPhysScene_Join"),
-			STAT_FNullGraphTask_ProcessPhysScene_Join,
-				STATGROUP_TaskGraphTasks);
-
-			PhysicsSceneCompletion = TGraphTask<FNullGraphTask>::CreateTask(&FinishPrerequisites, ENamedThreads::GameThread).ConstructAndDispatchWhenReady(
-				GET_STATID(STAT_FNullGraphTask_ProcessPhysScene_Join), PhysSingleThreadedMode() ? ENamedThreads::GameThread : CPrio_PhyXSceneCompletion.Get());
-		}
-		else
-		{
-			PhysicsSceneCompletion = FinishPrerequisites[0]; // we don't need a join
-		}
-	}
-
-
-	// Query clothing stats from skel mesh components in this world
-	// This is done outside TickPhysScene because clothing is
-	// not related to a scene.
-	GatherClothingStats(this->OwningWorld);
-}
-
-void FPhysScene_PhysX::EndFrame(ULineBatchComponent* InLineBatcher)
-{
-	check(IsInGameThread());
-
-	PhysicsSceneCompletion = NULL;
-
-	/**
-	* At this point physics simulation has finished. We obtain both scene locks so that the various read/write operations needed can be done quickly.
-	* This means that anyone attempting to write on other threads will be blocked. This is OK because accessing any of these game objects from another thread is probably a bad idea!
-	*/
-
-#if WITH_PHYSX
-	SCOPED_SCENE_WRITE_LOCK(GetPxScene());
-#endif // WITH_PHYSX 
-
-
-	// Perform any collision notification events
-	DispatchPhysNotifications_AssumesLocked();
-	
-	GatherPhysXStats_AssumesLocked(GetPxScene());
-
-#if !(UE_BUILD_SHIPPING || UE_BUILD_TEST)
-	// Handle debug rendering
-	if (InLineBatcher)
-	{
-		AddDebugLines(InLineBatcher);
-	}
-#endif // !(UE_BUILD_SHIPPING || UE_BUILD_TEST)
-}
-
-#if WITH_PHYSX
-/** Helper struct that puts all awake actors to sleep and then later wakes them back up */
-struct FHelpEnsureCollisionTreeIsBuilt
-{
-	FHelpEnsureCollisionTreeIsBuilt(PxScene* InPScene)
-		: PScene(InPScene)
-	{
-		if(PScene)
-		{
-			SCOPED_SCENE_WRITE_LOCK(PScene);
-			const int32 NumActors = PScene->getNbActors(PxActorTypeFlag::eRIGID_DYNAMIC);
-
-			if (NumActors)
-			{
-				ActorBuffer.AddUninitialized(NumActors);
-				PScene->getActors(PxActorTypeFlag::eRIGID_DYNAMIC, &ActorBuffer[0], NumActors);
-
-				for (PxActor*& PActor : ActorBuffer)
-				{
-					if (PActor)
-					{
-						if (PxRigidDynamic* PDynamic = PActor->is<PxRigidDynamic>())
-						{
-							if (PDynamic->isSleeping() == false)
-							{
-								PDynamic->putToSleep();
-							}
-							else
-							{
-								PActor = nullptr;
-							}
-						}
-					}
-				}
-			}
-		}
-	}
-
-	~FHelpEnsureCollisionTreeIsBuilt()
-	{
-		SCOPED_SCENE_WRITE_LOCK(PScene);
-		for (PxActor* PActor : ActorBuffer)
-		{
-			if (PActor)
-			{
-				if (PxRigidDynamic* PDynamic = PActor->is<PxRigidDynamic>())
-				{
-					PDynamic->wakeUp();
-				}
-			}
-		}
-	}
-
-private:
-
-	TArray<PxActor*> ActorBuffer;
-	PxScene* PScene;
-};
-
-#endif
-
-
-DECLARE_CYCLE_STAT(TEXT("EnsureCollisionTreeIsBuilt"), STAT_PhysicsEnsureCollisionTreeIsBuilt, STATGROUP_Physics);
-
-void FPhysScene_PhysX::EnsureCollisionTreeIsBuilt(UWorld* World)
-{
-	check(IsInGameThread());
-
-
-	SCOPE_CYCLE_COUNTER(STAT_PhysicsEnsureCollisionTreeIsBuilt);
-	//We have to call fetchResults several times to update the internal data structures. PhysX doesn't have an API for this so we have to make all actors sleep before doing this
-
-	SetIsStaticLoading(true);
-
-#if WITH_PHYSX
-	FHelpEnsureCollisionTreeIsBuilt SyncSceneHelper(GetPxScene());
-#endif
-
-	for (int Iteration = 0; Iteration < 6; ++Iteration)
-	{
-		World->SetupPhysicsTickFunctions(0.1f);
-		StartFrame();
-		WaitPhysScenes();
-		EndFrame(nullptr);
-	}
-
-	SetIsStaticLoading(false);
-}
-
-void FPhysScene_PhysX::SetIsStaticLoading(bool bStaticLoading)
-{
-	SetPhysXTreeRebuildRateImp(bStaticLoading ? 5 : PhysXTreeRebuildRate);
-}
-
-void FPhysScene_PhysX::SetPhysXTreeRebuildRate(int32 RebuildRate)
-{
-	PhysXTreeRebuildRate = FMath::Max(4, RebuildRate);
-	SetPhysXTreeRebuildRateImp(RebuildRate);
-}
-
-void FPhysScene_PhysX::SetPhysXTreeRebuildRateImp(int32 RebuildRate)
-{
-#if WITH_PHYSX
-	if (PxScene* PScene = GetPxScene())
-	{
-		// Lock scene lock, in case it is required
-		SCENE_LOCK_WRITE(PScene);
-
-		// Sets the rebuild rate hint, to 1 frame if static loading
-		PScene->setDynamicTreeRebuildRateHint(PhysXTreeRebuildRate);
-
-		// Unlock scene lock, in case it is required
-		SCENE_UNLOCK_WRITE(PScene);
-	}
-#endif
-}
-
-#if WITH_PHYSX
-/** Utility for looking up the PxScene associated with this FPhysScene. */
-PxScene* FPhysScene_PhysX::GetPxScene() const
-{
-#if WITH_APEX
-	nvidia::apex::Scene* ApexScene = PhysXScene;
-	return (ApexScene != nullptr) ? ApexScene->getPhysXScene() : nullptr;
-#else
-	return PhysXScene;
-#endif
-
-	return nullptr;
-}
-
-static void BatchPxRenderBufferLines(class ULineBatchComponent& LineBatcherToUse, const PxRenderBuffer& DebugData)
-{
-	int32 NumPoints = DebugData.getNbPoints();
-	if (NumPoints > 0)
-	{
-		const PxDebugPoint* Points = DebugData.getPoints();
-		for (int32 i = 0; i<NumPoints; i++)
-		{
-			LineBatcherToUse.DrawPoint(P2UVector(Points->pos), FColor((uint32)Points->color), 2, SDPG_World);
-
-			Points++;
-		}
-	}
-
-	// Build a list of all the lines we want to draw
-	TArray<FBatchedLine> DebugLines;
-
-	// Add all the 'lines' from PhysX
-	int32 NumLines = DebugData.getNbLines();
-	if (NumLines > 0)
-	{
-		const PxDebugLine* Lines = DebugData.getLines();
-		for (int32 i = 0; i<NumLines; i++)
-		{
-			new(DebugLines)FBatchedLine(P2UVector(Lines->pos0), P2UVector(Lines->pos1), FColor((uint32)Lines->color0), 0.f, 0.0f, SDPG_World);
-			Lines++;
-		}
-	}
-
-	// Add all the 'triangles' from PhysX
-	int32 NumTris = DebugData.getNbTriangles();
-	if (NumTris > 0)
-	{
-		const PxDebugTriangle* Triangles = DebugData.getTriangles();
-		for (int32 i = 0; i<NumTris; i++)
-		{
-			new(DebugLines)FBatchedLine(P2UVector(Triangles->pos0), P2UVector(Triangles->pos1), FColor((uint32)Triangles->color0), 0.f, 0.0f, SDPG_World);
-			new(DebugLines)FBatchedLine(P2UVector(Triangles->pos1), P2UVector(Triangles->pos2), FColor((uint32)Triangles->color1), 0.f, 0.0f, SDPG_World);
-			new(DebugLines)FBatchedLine(P2UVector(Triangles->pos2), P2UVector(Triangles->pos0), FColor((uint32)Triangles->color2), 0.f, 0.0f, SDPG_World);
-			Triangles++;
-		}
-	}
-
-	// Draw them all in one call.
-	if (DebugLines.Num() > 0)
-	{
-		LineBatcherToUse.DrawLines(DebugLines);
-	}
-}
-#endif // WITH_PHYSX
-
-
-/** Add any debug lines from the physics scene to the supplied line batcher. */
-void FPhysScene_PhysX::AddDebugLines(class ULineBatchComponent* LineBatcherToUse)
-{
-	if (LineBatcherToUse)
-	{
-#if WITH_PHYSX
-		// Render PhysX debug data
-		PxScene* PScene = GetPxScene();
-		const PxRenderBuffer& DebugData = PScene->getRenderBuffer();
-		BatchPxRenderBufferLines(*LineBatcherToUse, DebugData);
-#if WITH_APEX
-		// Render APEX debug data
-		apex::Scene* ApexScene = GetApexScene();
-		const PxRenderBuffer* RenderBuffer = ApexScene->getRenderBuffer();
-		if (RenderBuffer != NULL)
-		{
-			BatchPxRenderBufferLines(*LineBatcherToUse, *RenderBuffer);
-			ApexScene->updateRenderResources();
-		}
-#endif	// WITH_APEX
-#endif	// WITH_PHYSX
-	}
-}
-
-#if !UE_BUILD_SHIPPING
-int32 ForceSubstep = 0;
-FAutoConsoleVariableRef CVarSubStep(
-	TEXT("p.ForceSubstep"),
-	ForceSubstep,
-	TEXT("Whether to force substepping on")
-	TEXT("0: Ignore, 1: Force"),
-	ECVF_Default);
-#else
-constexpr int32 ForceSubstep = 0;
-#endif
-
-bool FPhysScene_PhysX::IsSubstepping() const
-{
-	// Substepping relies on interpolating transforms over frames, but only game worlds will be ticked,
-	// so we disallow this feature in non-game worlds.
-	if (!OwningWorld || !OwningWorld->IsGameWorld())
-	{
-		return false;
-	}
-
-	return ForceSubstep == 1 || bSubstepping;
-
-	return false;
-}
-
-
-void FPhysScene_PhysX::ApplyWorldOffset(FVector InOffset)
-{
-#if WITH_PHYSX
-	if (PxScene* PScene = GetPxScene())
-	{
-		// Lock scene lock, in case it is required
-		SCENE_LOCK_WRITE(PScene);
-
-		PScene->shiftOrigin(U2PVector(-InOffset));
-
-		// Unlock scene lock, in case it is required
-		SCENE_UNLOCK_WRITE(PScene);
-	}
-#endif
-}
-
-void FPhysScene_PhysX::InitPhysScene(const AWorldSettings* Settings)
-{
-#if WITH_PHYSX
-
-	int64 NumPhysxDispatcher = 0;
-	FParse::Value(FCommandLine::Get(), TEXT("physxDispatcher="), NumPhysxDispatcher);
-	if (NumPhysxDispatcher == 0 && FParse::Param(FCommandLine::Get(), TEXT("physxDispatcher")))
-	{
-		NumPhysxDispatcher = 4;	//by default give physx 4 threads
-	}
-
-	// Create dispatcher for tasks
-	if (PhysSingleThreadedMode())
-	{
-		CPUDispatcher = new FPhysXCPUDispatcherSingleThread();
-	}
-	else
-	{
-		if (NumPhysxDispatcher)
-		{
-			CPUDispatcher = PxDefaultCpuDispatcherCreate(NumPhysxDispatcher);
-		}
-		else
-		{
-			CPUDispatcher = new FPhysXCPUDispatcher();
-		}
-
-	}
-
-
-	PhysxUserData = FPhysxUserData(this);
-
-	// Create sim event callback
-	SimEventCallback = SimEventCallbackFactory.IsValid() ? SimEventCallbackFactory->Create(this) : new FPhysXSimEventCallback(this);
-	ContactModifyCallback = ContactModifyCallbackFactory.IsValid() ? ContactModifyCallbackFactory->Create(this) : nullptr;
-	CCDContactModifyCallback = CCDContactModifyCallbackFactory.IsValid() ? CCDContactModifyCallbackFactory->Create(this) : nullptr;
-
-	// Include scene descriptor in loop, so that we might vary it with scene type
-	PxSceneDesc PSceneDesc(GPhysXSDK->getTolerancesScale());
-	PSceneDesc.cpuDispatcher = CPUDispatcher;
-
-	FPhysSceneShaderInfo PhysSceneShaderInfo;
-	PhysSceneShaderInfo.PhysScene = this;
-	PSceneDesc.filterShaderData = &PhysSceneShaderInfo;
-	PSceneDesc.filterShaderDataSize = sizeof(PhysSceneShaderInfo);
-
-	PSceneDesc.filterShader = GSimulationFilterShader ? GSimulationFilterShader : PhysXSimFilterShader;
-	PSceneDesc.simulationEventCallback = SimEventCallback;
-	PSceneDesc.contactModifyCallback = ContactModifyCallback;
-	PSceneDesc.ccdContactModifyCallback = CCDContactModifyCallback;
-
-	if(UPhysicsSettings::Get()->bEnablePCM)
-	{
-		PSceneDesc.flags |= PxSceneFlag::eENABLE_PCM;
-	}
-	else
-	{
-		PSceneDesc.flags &= ~PxSceneFlag::eENABLE_PCM;
-	}
-
-	if (UPhysicsSettings::Get()->bEnableStabilization)
-	{
-		PSceneDesc.flags |= PxSceneFlag::eENABLE_STABILIZATION;
-	}
-	else
-	{
-		PSceneDesc.flags &= ~PxSceneFlag::eENABLE_STABILIZATION;
-	}
-
-	// Set bounce threshold
-	PSceneDesc.bounceThresholdVelocity = UPhysicsSettings::Get()->BounceThresholdVelocity;
-
-#if USE_SCENE_LOCK
-	if(UPhysicsSettings::Get()->bWarnMissingLocks)
-	{
-		PSceneDesc.flags |= PxSceneFlag::eREQUIRE_RW_LOCK;
-	}
-
-#endif
-
-	if(!UPhysicsSettings::Get()->bDisableActiveActors)
-	{
-		// We want to use 'active actors'
-		PSceneDesc.flags |= PxSceneFlag::eENABLE_ACTIVE_ACTORS;
-		PSceneDesc.flags |= PxSceneFlag::eEXCLUDE_KINEMATICS_FROM_ACTIVE_ACTORS;
-	}
-
-	// enable CCD at scene level
-	if (UPhysicsSettings::Get()->bDisableCCD == false)
-	{
-		PSceneDesc.flags |= PxSceneFlag::eENABLE_CCD;
-	}
-
-	if(!UPhysicsSettings::Get()->bDisableKinematicStaticPairs && GPhysXForceNoKinematicStaticPairs == 0)
-	{
-		// Need to turn this on to consider kinematics turning into dynamic. Otherwise, you'll need to call resetFiltering to do the expensive broadphase reinserting
-		PSceneDesc.flags |= PxSceneFlag::eENABLE_KINEMATIC_STATIC_PAIRS;
-	}
-
-	if(!UPhysicsSettings::Get()->bDisableKinematicKinematicPairs && GPhysXForceNoKinematicKinematicPairs == 0)
-	{
-		PSceneDesc.flags |= PxSceneFlag::eENABLE_KINEMATIC_PAIRS;	//this is only needed for destruction, but unfortunately this flag cannot be modified after creation and the plugin has no hook (yet)
-	}
-
-	// @TODO Should we set up PSceneDesc.limits? How?
-
-	// Do this to improve loading times, esp. for streaming in sublevels
-	PSceneDesc.staticStructure = PxPruningStructureType::eDYNAMIC_AABB_TREE;
-	// Default to rebuilding tree slowly
-	PSceneDesc.dynamicTreeRebuildRateHint = PhysXTreeRebuildRate;
-
-	if (UPhysicsSettings::Get()->bEnableEnhancedDeterminism) {
-		PSceneDesc.flags |= PxSceneFlag::eENABLE_ENHANCED_DETERMINISM;
-	}
-
-	bool bIsValid = PSceneDesc.isValid();
-	if (!bIsValid)
-	{
-		UE_LOG(LogPhysics, Log, TEXT("Invalid PSceneDesc"));
-	}
-
-	// Setup MBP desc settings if required
-	const FBroadphaseSettings& BroadphaseSettings = (Settings && Settings->bOverrideDefaultBroadphaseSettings) ? Settings->BroadphaseSettings : UPhysicsSettings::Get()->DefaultBroadphaseSettings;
-	bool bUseMBP = IsRunningDedicatedServer() ? BroadphaseSettings.bUseMBPOnServer : BroadphaseSettings.bUseMBPOnClient;
-
-	if(bUseMBP)
-	{
-		MbpBroadphaseCallback = new FPhysXMbpBroadphaseCallback();
-		PSceneDesc.broadPhaseType = PxBroadPhaseType::eMBP;
-		PSceneDesc.broadPhaseCallback = MbpBroadphaseCallback;
-	}
-	else
-	{
-		MbpBroadphaseCallback = nullptr;
-	}
-
-	// Create scene, and add to map
-	PxScene* PScene = GPhysXSDK->createScene(PSceneDesc);
-	if(PxPvdSceneClient* PVDClient = PScene->getScenePvdClient())
-	{
-		PVDClient->setScenePvdFlags(PxPvdSceneFlag::eTRANSMIT_CONSTRAINTS | PxPvdSceneFlag::eTRANSMIT_CONTACTS | PxPvdSceneFlag::eTRANSMIT_SCENEQUERIES);
-	}
-
-	// Setup actual MBP data on live scene
-	if(bUseMBP)
-	{
-		uint32 NumSubdivisions = BroadphaseSettings.MBPNumSubdivs;
-
-		if(IsRunningDedicatedServer())
-		{
-			if(GPhysXOverrideMbpNumSubdivisions_Server > 0)
-			{
-				NumSubdivisions = GPhysXOverrideMbpNumSubdivisions_Server;
-			}
-		}
-		else
-		{
-			if(GPhysXOverrideMbpNumSubdivisions_Client > 0)
-			{
-				NumSubdivisions = GPhysXOverrideMbpNumSubdivisions_Client;
-			}
-		}
-
-		// Must have at least one and no more than 256 regions, subdivision is num^2 so only up to 16
-		NumSubdivisions = FMath::Clamp<uint32>(NumSubdivisions, 1, BroadphaseSettings.bUseMBPOuterBounds ? 15 : 16);
-
-		const FBox& Bounds = BroadphaseSettings.MBPBounds;
-		PxBounds3 MbpBounds(U2PVector(Bounds.Min), U2PVector(Bounds.Max));
-
-		// Storage for generated regions, the generation function will create num^2 regions
-		TArray<PxBounds3> GeneratedRegions;
-		GeneratedRegions.AddZeroed(NumSubdivisions * NumSubdivisions);
-
-		// Final parameter is up axis (2 == Z for Unreal Engine)
-		PxBroadPhaseExt::createRegionsFromWorldBounds(GeneratedRegions.GetData(), MbpBounds, NumSubdivisions, 2);
-
-		for(const PxBounds3& Region : GeneratedRegions)
-		{
-			PxBroadPhaseRegion NewRegion;
-			NewRegion.bounds = Region;
-			NewRegion.userData = nullptr; // No need to track back to an Unreal instance at the moment
-
-			PScene->addBroadPhaseRegion(NewRegion);
-		}
-			
-		if (BroadphaseSettings.bUseMBPOuterBounds)
-		{
-			const FBox& OuterBounds = BroadphaseSettings.MBPOuterBounds;
-			if (!(OuterBounds.Min.X >= Bounds.Min.X || OuterBounds.Min.Y >= Bounds.Min.Y ||
-				OuterBounds.Max.X <= Bounds.Max.X || OuterBounds.Max.Y <= Bounds.Max.Y))
-			{
-				{
-					//outer 1
-					const FBox LeftBounds(FVector(Bounds.Min.X, OuterBounds.Min.Y, Bounds.Min.Z), FVector(OuterBounds.Max.X, Bounds.Min.Y, Bounds.Max.Z));
-					PxBounds3 MbpLeftBounds(U2PVector(LeftBounds.Min), U2PVector(LeftBounds.Max));
-
-					// Storage for generated regions, the generation function will create num^2 regions
-					TArray<PxBounds3> GeneratedRegionsLeft;
-					GeneratedRegionsLeft.AddZeroed(1);
-
-					// Final parameter is up axis (2 == Z for Unreal Engine)
-					PxBroadPhaseExt::createRegionsFromWorldBounds(GeneratedRegionsLeft.GetData(), MbpLeftBounds, 1, 2);
-
-					for (const PxBounds3& Region : GeneratedRegionsLeft)
-					{
-						PxBroadPhaseRegion NewRegion;
-						NewRegion.bounds = Region;
-						NewRegion.userData = nullptr; // No need to track back to an Unreal instance at the moment
-
-						PScene->addBroadPhaseRegion(NewRegion);
-					}
-				}
-				{
-					//outer 2
-					const FBox RightBounds(FVector(OuterBounds.Min.X, Bounds.Max.Y, Bounds.Min.Z), FVector(Bounds.Max.X, OuterBounds.Max.Y, Bounds.Max.Z));
-					PxBounds3 MbpRightBounds(U2PVector(RightBounds.Min), U2PVector(RightBounds.Max));
-
-					// Storage for generated regions, the generation function will create num^2 regions
-					TArray<PxBounds3> GeneratedRegionsRight;
-					GeneratedRegionsRight.AddZeroed(1);
-
-					// Final parameter is up axis (2 == Z for Unreal Engine)
-					PxBroadPhaseExt::createRegionsFromWorldBounds(GeneratedRegionsRight.GetData(), MbpRightBounds, 1, 2);
-
-					for (const PxBounds3& Region : GeneratedRegionsRight)
-					{
-						PxBroadPhaseRegion NewRegion;
-						NewRegion.bounds = Region;
-						NewRegion.userData = nullptr; // No need to track back to an Unreal instance at the moment
-
-						PScene->addBroadPhaseRegion(NewRegion);
-					}
-				}
-				{
-					//outer 3
-					const FBox TopBounds(FVector(Bounds.Max.X, Bounds.Min.Y, Bounds.Min.Z), FVector(OuterBounds.Max.X, OuterBounds.Max.Y, Bounds.Max.Z));
-					PxBounds3 MbpTopBounds(U2PVector(TopBounds.Min), U2PVector(TopBounds.Max));
-
-					// Storage for generated regions, the generation function will create num^2 regions
-					TArray<PxBounds3> GeneratedRegionsTop;
-					GeneratedRegionsTop.AddZeroed(1);
-
-					// Final parameter is up axis (2 == Z for Unreal Engine)
-					PxBroadPhaseExt::createRegionsFromWorldBounds(GeneratedRegionsTop.GetData(), MbpTopBounds, 1, 2);
-
-					for (const PxBounds3& Region : GeneratedRegionsTop)
-					{
-						PxBroadPhaseRegion NewRegion;
-						NewRegion.bounds = Region;
-						NewRegion.userData = nullptr; // No need to track back to an Unreal instance at the moment
-
-						PScene->addBroadPhaseRegion(NewRegion);
-					}
-				}
-				{
-					//outer 4
-					const FBox BottomBounds(FVector(OuterBounds.Min.X, OuterBounds.Min.Y, Bounds.Min.Z), FVector(Bounds.Min.X, Bounds.Max.Y, Bounds.Max.Z));
-					PxBounds3 MbpBottomBounds(U2PVector(BottomBounds.Min), U2PVector(BottomBounds.Max));
-
-					// Storage for generated regions, the generation function will create num^2 regions
-					TArray<PxBounds3> GeneratedRegionsBottom;
-					GeneratedRegionsBottom.AddZeroed(1);
-
-					// Final parameter is up axis (2 == Z for Unreal Engine)
-					PxBroadPhaseExt::createRegionsFromWorldBounds(GeneratedRegionsBottom.GetData(), MbpBottomBounds, 1, 2);
-
-					for (const PxBounds3& Region : GeneratedRegionsBottom)
-					{
-						PxBroadPhaseRegion NewRegion;
-						NewRegion.bounds = Region;
-						NewRegion.userData = nullptr; // No need to track back to an Unreal instance at the moment
-
-						PScene->addBroadPhaseRegion(NewRegion);
-					}
-				}
-			}
-		}
-		
-	}
-
-#if WITH_APEX
-	// Build the APEX scene descriptor for the PhysX scene
-	apex::SceneDesc ApexSceneDesc;
-	ApexSceneDesc.scene = PScene;
-	// This interface allows us to modify the PhysX simulation filter shader data with contact pair flags
-	ApexSceneDesc.physX3Interface = GPhysX3Interface;
-
-	// Create the APEX scene from our descriptor
-	apex::Scene* ApexScene = GApexSDK->createScene(ApexSceneDesc);
-
-	// This enables debug rendering using the "legacy" method, not using the APEX render API
-	ApexScene->setUseDebugRenderable(true);
-
-	// Allocate a view matrix for APEX scene LOD
-	ApexScene->allocViewMatrix(apex::ViewMatrixType::LOOK_AT_RH);
-
-	// Store index of APEX scene in this FPhysScene
-	PhysXScene = ApexScene;
-#else
-	// Store index of PhysX Scene in this FPhysScene
-	PhysXScene = PScene;
-#endif	// #if WITH_APEX
-
-	// Save pointer to FPhysScene in userdata
-	PScene->userData = &PhysxUserData;
-#if WITH_APEX
-	ApexScene->userData = &PhysxUserData;
-#endif
-
-	//Initialize substeppers
-#if WITH_PHYSX
-#if WITH_APEX
-	PhysSubStepper = new FPhysSubstepTask(ApexScene, this);
-#else
-	PhysSubStepper = new FPhysSubstepTask(PScene, this);
-#endif
-
-	if (PxPvdSceneClient* PVDSceneClient = PScene->getScenePvdClient())
-	{
-		PVDSceneClient->setScenePvdFlags(PxPvdSceneFlag::eTRANSMIT_CONTACTS | PxPvdSceneFlag::eTRANSMIT_SCENEQUERIES | PxPvdSceneFlag::eTRANSMIT_CONSTRAINTS);
-	}
-#endif
-
-#if WITH_CUSTOM_SQ_STRUCTURE
-	SQAcceleratorUnion = new FSQAcceleratorUnion();
-	
-	/** Disabled the following case until we have a good broadphase. Right now using FPhysXSQAccelerator to just query the scene */
-	//SQAccelerator = new FSQAccelerator();
-	//SQAcceleratorUnion->AddSQAccelerator(SQAccelerator);
-
-#if WITH_PHYSX && !WITH_CHAOS
-	PhysXSQAccelerator = new FPhysXSQAccelerator(PScene);
-	SQAcceleratorUnion->AddSQAccelerator(PhysXSQAccelerator);
-#endif
-#endif
-
-#if WITH_CHAOS
-	check(false);
-#else
-	FPhysicsDelegates::OnPhysSceneInit.Broadcast(this);
-#endif
-
-#endif // WITH_PHYSX
-}
-
-void FPhysScene_PhysX::TermPhysScene()
-{
-#if WITH_CUSTOM_SQ_STRUCTURE
-	delete SQAcceleratorUnion;
-	SQAcceleratorUnion = nullptr;
-
-	/** Disabled the following case until we have a good broadphase. Right now using FPhysXSQAccelerator to just query the scene */
-	//delete SQAccelerator;
-	//SQAccelerator = nullptr;
-#if WITH_PHYSX
-	delete PhysXSQAccelerator;
-	PhysXSQAccelerator = nullptr;
-#endif
-#endif
-
-#if WITH_PHYSX
-	if (PxScene* PScene = GetPxScene())
-	{
-#if WITH_APEX
-		if (apex::Scene* ApexScene = GetApexScene())
-		{
-			GPhysCommandHandler->DeferredRelease(ApexScene);
-		}
-#endif // #if WITH_APEX
-
-#if WITH_CHAOS
-		check(false);
-#else
-		FPhysicsDelegates::OnPhysSceneTerm.Broadcast(this);
-#endif
-
-		delete PhysSubStepper;
-		PhysSubStepper = NULL;
-
-		// @todo block on any running scene before calling this
-		GPhysCommandHandler->DeferredRelease(PScene);
-		GPhysCommandHandler->DeferredDeleteSimEventCallback(SimEventCallback);
-		GPhysCommandHandler->DeferredDeleteContactModifyCallback(ContactModifyCallback);
-		GPhysCommandHandler->DeferredDeleteMbpBroadphaseCallback(MbpBroadphaseCallback);
-
-		// Commands may have accumulated as the scene is terminated - flush any commands for this scene.
-		GPhysCommandHandler->Flush();
-
-		PhysXScene = nullptr;
-	}
-#endif
-}
-
-#if WITH_PHYSX
-void FPhysScene_PhysX::AddPendingOnConstraintBreak(FConstraintInstance* ConstraintInstance)
-{
-	PendingConstraintData.PendingConstraintBroken.Add( FConstraintBrokenDelegateData(ConstraintInstance) );
-}
-
-FConstraintBrokenDelegateData::FConstraintBrokenDelegateData(FConstraintInstance* ConstraintInstance)
-	: OnConstraintBrokenDelegate(ConstraintInstance->OnConstraintBrokenDelegate)
-	, ConstraintIndex(ConstraintInstance->ConstraintIndex)
-{
-
-}
-
-void FPhysScene_PhysX::AddPendingSleepingEvent(FBodyInstance* BI, ESleepEvent SleepEventType)
-{
-	PendingSleepEvents.FindOrAdd(BI) = SleepEventType;
-}
-
-#if WITH_PHYSX
-void ListAwakeRigidBodiesFromScene(bool bIncludeKinematic, PxScene* PhysXScene, int32& totalCount)
-{
-	check(PhysXScene != NULL);
-
-	SCOPED_SCENE_READ_LOCK(PhysXScene);
-
-	PxActor* PhysXActors[2048];
-	int32 NumberActors = PhysXScene->getActors(PxActorTypeFlag::eRIGID_DYNAMIC, PhysXActors, 2048);
-	for(int32 i = 0; i < NumberActors; ++i)
-	{
-		PxRigidDynamic* RgActor = (PxRigidDynamic*)PhysXActors[i];
-		if(!RgActor->isSleeping() && (bIncludeKinematic || RgActor->getRigidBodyFlags() != PxRigidBodyFlag::eKINEMATIC))
-		{
-			++totalCount;
-			FBodyInstance* BodyInst = FPhysxUserData::Get<FBodyInstance>(RgActor->userData);
-			if(BodyInst)
-			{
-				UE_LOG(LogPhysics, Log, TEXT("BI %s %d"), BodyInst->OwnerComponent.Get() ? *BodyInst->OwnerComponent.Get()->GetPathName() : TEXT("NONE"), BodyInst->InstanceBodyIndex);
-			}
-			else
-			{
-				UE_LOG(LogPhysics, Log, TEXT("BI %s"), TEXT("NONE"));
-			}
-		}
-	}
-}
-
-void FPhysScene_PhysX::SerializeForTesting(FArchive& Ar)
-{
-	FPhysTestSerializer Serializer;
-	Chaos::FChaosArchive ChaosAr(Ar);
-	Serializer.SetPhysicsData(*GetPxScene());
-	Serializer.Serialize(ChaosAr);
-}
-#endif // WITH_PHYSX
-
-/** Util to list to log all currently awake rigid bodies */
-void FPhysScene_PhysX::ListAwakeRigidBodies(bool bIncludeKinematic)
-{
-#if WITH_PHYSX
-	int32 BodyCount = 0;
-	UE_LOG(LogPhysics, Log, TEXT("TOTAL: ListAwakeRigidBodies needs fixing."));
-	ListAwakeRigidBodiesFromScene(bIncludeKinematic, GetPxScene(), BodyCount);
-	
-	UE_LOG(LogPhysics, Log, TEXT("TOTAL: %d awake bodies."), BodyCount);
-#endif // WITH_PHYSX
-}
-
-#if !WITH_CHAOS
-int32 FPhysScene::GetNumAwakeBodies()
-{
-	int32 NumAwake = 0;
-
-	if(PxScene* PScene = GetPxScene())
-	{
-		TArray<PxActor*> PxActors;
-		int32 NumActors = PScene->getNbActors(PxActorTypeFlag::eRIGID_DYNAMIC);
-		PxActors.AddZeroed(NumActors);
-
-		PScene->getActors(PxActorTypeFlag::eRIGID_DYNAMIC, PxActors.GetData(), NumActors * sizeof(PxActors[0]));
-		for(PxActor* PActor : PxActors)
-		{
-			if(!PActor->is<PxRigidDynamic>()->isSleeping())
-			{
-				++NumAwake;
-			}
-		}
-	}
-
-	return NumAwake;
-}
-#endif
-
-#endif // WITH_PHYSX
-
-#endif // !WITH_CHAOS_NEEDS_TO_BE_FIXED
-
-#endif //  !WITH_CHAOS && !WITH_IMMEDIATE_PHYSX
diff --git a/Engine/Source/Runtime/Engine/Private/PhysicsEngine/PhysSubstepTasks.cpp b/Engine/Source/Runtime/Engine/Private/PhysicsEngine/PhysSubstepTasks.cpp
--- a/Engine/Source/Runtime/Engine/Private/PhysicsEngine/PhysSubstepTasks.cpp
+++ b/Engine/Source/Runtime/Engine/Private/PhysicsEngine/PhysSubstepTasks.cpp
@@ -4,20 +4,6 @@
 #include "PhysicsEngine/PhysicsSettings.h"
 #include "Physics/PhysScene_PhysX.h"
 
-#if PHYSICS_INTERFACE_PHYSX
-#include "PhysXPublic.h"
-
-uint8* PhysXCompletionTask::GetScratchBufferData()
-{
-	return ScratchBuffer ? ScratchBuffer->Buffer : nullptr;
-}
-
-int32 PhysXCompletionTask::GetScratchBufferSize()
-{
-	return ScratchBuffer ? ScratchBuffer->BufferSize : 0;
-}
-#endif
-
 struct FSubstepCallbackGuard
 {
 #if !UE_BUILD_SHIPPING
@@ -39,24 +25,6 @@ struct FSubstepCallbackGuard
 #endif
 };
 
-#if PHYSICS_INTERFACE_PHYSX
-FPhysSubstepTask::FPhysSubstepTask(PxApexScene * GivenScene, FPhysScene* InPhysScene) :
-	NumSubsteps(0),
-	SubTime(0.f),
-	DeltaSeconds(0.f),
-	FullSimulationTask(0),
-	Alpha(0.f),
-	StepScale(0.f),
-	TotalSubTime(0.f),
-	CurrentSubStep(0),
-	SubstepCallbackGuard(0),
-	PhysScene(InPhysScene),
-	PAScene(GivenScene)
-{
-	check(PAScene);
-}
-#endif
-
 void FPhysSubstepTask::SwapBuffers()
 {
 	External = !External;
@@ -70,7 +38,6 @@ void FPhysSubstepTask::RemoveBodyInstance_AssumesLocked(FBodyInstance* BodyInsta
 
 void FPhysSubstepTask::SetKinematicTarget_AssumesLocked(FBodyInstance* Body, const FTransform& TM)
 {
-#if WITH_PHYSX
 	TM.DiagnosticCheck_IsValid();
 
 	//We only interpolate kinematic actors that need it
@@ -81,12 +48,10 @@ void FPhysSubstepTask::SetKinematicTarget_AssumesLocked(FBodyInstance* Body, con
 		TargetState.bKinematicTarget = true;
 		TargetState.KinematicTarget = KinmaticTarget;
 	}
-#endif
 }
 
 bool FPhysSubstepTask::GetKinematicTarget_AssumesLocked(const FBodyInstance* Body, FTransform& OutTM) const
 {
-#if WITH_PHYSX
 	if (const FPhysTarget* TargetState = PhysTargetBuffers[External].Find(Body))
 	{
 		if (TargetState->bKinematicTarget)
@@ -95,14 +60,12 @@ bool FPhysSubstepTask::GetKinematicTarget_AssumesLocked(const FBodyInstance* Bod
 			return true;
 		}
 	}
-#endif
 
 	return false;
 }
 
 void FPhysSubstepTask::AddCustomPhysics_AssumesLocked(FBodyInstance* Body, const FCalculateCustomPhysics& CalculateCustomPhysics)
 {
-#if WITH_PHYSX
 	//Limit custom physics to non kinematic actors
 	if (Body->IsNonKinematic())
 	{
@@ -111,7 +74,6 @@ void FPhysSubstepTask::AddCustomPhysics_AssumesLocked(FBodyInstance* Body, const
 		FPhysTarget & TargetState = PhysTargetBuffers[External].FindOrAdd(Body);
 		TargetState.CustomPhysics.Add(CustomTarget);
 	}
-#endif
 }
 
 #if !UE_BUILD_SHIPPING
@@ -122,7 +84,6 @@ void FPhysSubstepTask::AddCustomPhysics_AssumesLocked(FBodyInstance* Body, const
 
 void FPhysSubstepTask::AddForce_AssumesLocked(FBodyInstance* Body, const FVector& Force, bool bAccelChange)
 {
-#if WITH_PHYSX
 	//We should only apply forces on non kinematic actors
 	if (Body->IsNonKinematic())
 	{
@@ -135,12 +96,10 @@ void FPhysSubstepTask::AddForce_AssumesLocked(FBodyInstance* Body, const FVector
 		FPhysTarget & TargetState = PhysTargetBuffers[External].FindOrAdd(Body);
 		TargetState.Forces.Add(ForceTarget);
 	}
-#endif
 }
 
 void FPhysSubstepTask::AddForceAtPosition_AssumesLocked(FBodyInstance* Body, const FVector& Force, const FVector& Position, bool bIsLocalForce)
 {
-#if WITH_PHYSX
 	if (Body->IsNonKinematic())
 	{
 		SUBSTEPPING_WARNING()
@@ -153,11 +112,9 @@ void FPhysSubstepTask::AddForceAtPosition_AssumesLocked(FBodyInstance* Body, con
 		FPhysTarget & TargetState = PhysTargetBuffers[External].FindOrAdd(Body);
 		TargetState.Forces.Add(ForceTarget);
 	}
-#endif
 }
 void FPhysSubstepTask::AddTorque_AssumesLocked(FBodyInstance* Body, const FVector& Torque, bool bAccelChange)
 {
-#if WITH_PHYSX
 	//We should only apply torque on non kinematic actors
 	if (Body->IsNonKinematic())
 	{
@@ -169,24 +126,20 @@ void FPhysSubstepTask::AddTorque_AssumesLocked(FBodyInstance* Body, const FVecto
 		FPhysTarget & TargetState = PhysTargetBuffers[External].FindOrAdd(Body);
 		TargetState.Torques.Add(TorqueTarget);
 	}
-#endif
 }
 
 void FPhysSubstepTask::ClearTorques_AssumesLocked(FBodyInstance* Body)
 {
-#if WITH_PHYSX
 	if (Body->IsNonKinematic())
 	{
 		SUBSTEPPING_WARNING()
 		FPhysTarget & TargetState = PhysTargetBuffers[External].FindOrAdd(Body);
 		TargetState.Torques.Empty();
 	}
-#endif
 }
 
 void FPhysSubstepTask::AddRadialForceToBody_AssumesLocked(FBodyInstance* Body, const FVector& Origin, const float Radius, const float Strength, const uint8 Falloff, const bool bAccelChange)
 {
-#if WITH_PHYSX
 	//We should only apply torque on non kinematic actors
 	if (Body->IsNonKinematic())
 	{
@@ -201,12 +154,10 @@ void FPhysSubstepTask::AddRadialForceToBody_AssumesLocked(FBodyInstance* Body, c
 		FPhysTarget & TargetState = PhysTargetBuffers[External].FindOrAdd(Body);
 		TargetState.RadialForces.Add(RadialForceTarget);
 	}
-#endif
 }
 
 void FPhysSubstepTask::ClearForces_AssumesLocked(FBodyInstance* Body)
 {
-#if WITH_PHYSX
 	if (Body->IsNonKinematic())
 	{
 		SUBSTEPPING_WARNING()
@@ -214,13 +165,11 @@ void FPhysSubstepTask::ClearForces_AssumesLocked(FBodyInstance* Body)
 		TargetState.Forces.Empty();
 		TargetState.RadialForces.Empty();
 	}
-#endif
 }
 
 /** Applies custom physics - Assumes caller has obtained writer lock */
 void FPhysSubstepTask::ApplyCustomPhysics(const FPhysTarget& PhysTarget, FBodyInstance* BodyInstance, float DeltaTime)
 {
-#if WITH_PHYSX
 	FSubstepCallbackGuard Guard(*this);
 	for (int32 i = 0; i < PhysTarget.CustomPhysics.Num(); ++i)
 	{
@@ -228,182 +177,36 @@ void FPhysSubstepTask::ApplyCustomPhysics(const FPhysTarget& PhysTarget, FBodyIn
 
 		CustomTarget.CalculateCustomPhysics->ExecuteIfBound(DeltaTime, BodyInstance);
 	}
-#endif
 }
 
-#if PHYSICS_INTERFACE_PHYSX
-bool IsKinematicHelper(const PxRigidBody* PRigidBody)
-{
-	const bool bIsKinematic = PRigidBody->getRigidBodyFlags() & PxRigidBodyFlag::eKINEMATIC;
-	return bIsKinematic;
-}
-#endif
 
 /** Applies forces - Assumes caller has obtained writer lock */
 void FPhysSubstepTask::ApplyForces_AssumesLocked(const FPhysTarget& PhysTarget, FBodyInstance* BodyInstance)
 {
-#if WITH_PHYSX
-#if WITH_CHAOS || WITH_IMMEDIATE_PHYSX
     check(false);
-#else
-	/** Apply Forces */
-	PxRigidBody* PRigidBody = FPhysicsInterface::GetPxRigidBody_AssumesLocked(BodyInstance->GetPhysicsActorHandle());
-
-	for (int32 i = 0; i < PhysTarget.Forces.Num(); ++i)
-	{
-		const FForceTarget& ForceTarget = PhysTarget.Forces[i];
-
-		if (ForceTarget.bPosition)
-		{
-			if (!ForceTarget.bIsLocalForce)
-			{
-				PxRigidBodyExt::addForceAtPos(*PRigidBody, U2PVector(ForceTarget.Force), U2PVector(ForceTarget.Position), PxForceMode::eFORCE, true);
-			}
-			else
-			{
-				PxRigidBodyExt::addLocalForceAtLocalPos(*PRigidBody, U2PVector(ForceTarget.Force), U2PVector(ForceTarget.Position), PxForceMode::eFORCE, true);
-			}
-		}
-		else
-		{
-			PRigidBody->addForce(U2PVector(ForceTarget.Force), ForceTarget.bAccelChange ? PxForceMode::eACCELERATION : PxForceMode::eFORCE, true);
-		}
-	}
-#endif
-#endif
 }
 
 /** Applies torques - Assumes caller has obtained writer lock */
 void FPhysSubstepTask::ApplyTorques_AssumesLocked(const FPhysTarget& PhysTarget, FBodyInstance* BodyInstance)
 {
-#if WITH_PHYSX
-#if WITH_CHAOS || WITH_IMMEDIATE_PHYSX
     check(false);
-#else
-	/** Apply Torques */
-	PxRigidBody* PRigidBody = FPhysicsInterface_PhysX::GetPxRigidBody_AssumesLocked(BodyInstance->GetPhysicsActorHandle());
-
-	for (int32 i = 0; i < PhysTarget.Torques.Num(); ++i)
-	{
-		const FTorqueTarget& TorqueTarget = PhysTarget.Torques[i];
-		PRigidBody->addTorque(U2PVector(TorqueTarget.Torque), TorqueTarget.bAccelChange ? PxForceMode::eACCELERATION : PxForceMode::eFORCE, true);
-	}
-#endif
-#endif
 }
 
 /** Applies radial forces - Assumes caller has obtained writer lock */
 void FPhysSubstepTask::ApplyRadialForces_AssumesLocked(const FPhysTarget& PhysTarget, FBodyInstance* BodyInstance)
 {
-#if WITH_PHYSX
-#if WITH_CHAOS || WITH_IMMEDIATE_PHYSX
     check(false);
-#else
-	/** Apply Torques */
-	PxRigidBody* PRigidBody = FPhysicsInterface_PhysX::GetPxRigidBody_AssumesLocked(BodyInstance->GetPhysicsActorHandle());
-
-	for (int32 i = 0; i < PhysTarget.RadialForces.Num(); ++i)
-	{
-		const FRadialForceTarget& RadialForceTArget= PhysTarget.RadialForces[i];
-		AddRadialForceToPxRigidBody_AssumesLocked(*PRigidBody, RadialForceTArget.Origin, RadialForceTArget.Radius, RadialForceTArget.Strength, RadialForceTArget.Falloff, RadialForceTArget.bAccelChange);
-	}
-#endif
-#endif
 }
 
 
 /** Interpolates kinematic actor transform - Assumes caller has obtained writer lock */
 void FPhysSubstepTask::InterpolateKinematicActor_AssumesLocked(const FPhysTarget& PhysTarget, FBodyInstance* BodyInstance, float InAlpha)
 {
-#if WITH_PHYSX
-#if WITH_CHAOS || WITH_IMMEDIATE_PHYSX
     check(false);
-#else
-	PxRigidDynamic * PRigidDynamic = FPhysicsInterface_PhysX::GetPxRigidDynamic_AssumesLocked(BodyInstance->GetPhysicsActorHandle());
-	InAlpha = FMath::Clamp(InAlpha, 0.f, 1.f);
-
-	/** Interpolate kinematic actors */
-	if (PhysTarget.bKinematicTarget)
-	{
-		//It's possible that the actor is no longer kinematic and is now simulating. In that case do nothing
-		if (!BodyInstance->IsNonKinematic())
-		{
-			const FKinematicTarget& KinematicTarget = PhysTarget.KinematicTarget;
-			const FTransform& TargetTM = KinematicTarget.TargetTM;
-			const FTransform& StartTM = KinematicTarget.OriginalTM;
-			FTransform InterTM = FTransform::Identity;
-
-			InterTM.SetLocation(FMath::Lerp(StartTM.GetLocation(), TargetTM.GetLocation(), InAlpha));
-			InterTM.SetRotation(FMath::Lerp(StartTM.GetRotation(), TargetTM.GetRotation(), InAlpha));
-
-			PxTransform PNewPose = U2PTransform(InterTM);
-			if (!ensureMsgf(PNewPose.isValid(), TEXT("Found an Invalid Pose for %s. Using previous pose."), *BodyInstance->GetBodyDebugName()))
-			{
-				PRigidDynamic->getKinematicTarget(PNewPose);
-				if (!ensureMsgf(PNewPose.isValid(), TEXT("Previous pose is invalid. Using identity.")))
-				{
-					PNewPose = U2PTransform(FTransform::Identity);
-				}
-			}
-			PRigidDynamic->setKinematicTarget(PNewPose);
-		}
-	}
-#endif
-#endif
 }
 
 void FPhysSubstepTask::SubstepInterpolation(float InAlpha, float DeltaTime)
 {
-#if PHYSICS_INTERFACE_PHYSX
-#if WITH_APEX
-	SCOPED_APEX_SCENE_WRITE_LOCK(PAScene);
-	PxScene * PScene = PAScene->getPhysXScene();
-#else
-	PxScene * PScene = PAScene;
-	SCOPED_SCENE_WRITE_LOCK(PScene);
-#endif
-#if WITH_CHAOS || WITH_IMMEDIATE_PHYSX
-    check(false);
-#else
-
-	/** Note: We lock the entire scene before iterating. The assumption is that removing an FBodyInstance from the map will also be wrapped by this lock */
-	
-	
-	PhysTargetMap& Targets = PhysTargetBuffers[!External];
-
-	for (PhysTargetMap::TIterator Itr = Targets.CreateIterator(); Itr; ++Itr)
-	{
-		FPhysTarget & PhysTarget = Itr.Value();
-		FBodyInstance * BodyInstance = Itr.Key();
-		PxRigidBody* PRigidBody = FPhysicsInterface_PhysX::GetPxRigidBody_AssumesLocked(BodyInstance->GetPhysicsActorHandle());
-
-		if (PRigidBody == NULL)
-		{
-			continue;
-		}
-
-		//We should only be iterating over actors that belong to this scene
-		check(PRigidBody->getScene() == PScene);
-
-		if (!IsKinematicHelper(PRigidBody))
-		{
-			ApplyCustomPhysics(PhysTarget, BodyInstance, DeltaTime);
-			ApplyForces_AssumesLocked(PhysTarget, BodyInstance);
-			ApplyTorques_AssumesLocked(PhysTarget, BodyInstance);
-			ApplyRadialForces_AssumesLocked(PhysTarget, BodyInstance);
-		}else
-		{
-			InterpolateKinematicActor_AssumesLocked(PhysTarget, BodyInstance, InAlpha);
-		}
-	}
-
-	/** Final substep */
-	if (InAlpha >= 1.f)
-	{
-		Targets.Empty(Targets.Num());
-	}
-#endif
-#endif
 }
 
 float FPhysSubstepTask::UpdateTime(float UseDelta)
@@ -426,22 +229,6 @@ float FPhysSubstepTask::UpdateTime(float UseDelta)
 	return SubTime;
 }
 
-#if PHYSICS_INTERFACE_PHYSX
-void FPhysSubstepTask::StepSimulation(PhysXCompletionTask * Task)
-{
-	check(SubTime > 0.f);
-	check(DeltaSeconds > 0.f);
-
-	FullSimulationTask = Task;
-	Alpha = 0.f;
-	StepScale = SubTime / DeltaSeconds;
-	TotalSubTime = 0.f;
-	CurrentSubStep = 0;
-
-	SubstepSimulationStart();
-}
-#endif
-
 DECLARE_CYCLE_STAT(TEXT("Phys SubstepStart"), STAT_SubstepSimulationStart, STATGROUP_Physics);
 DECLARE_CYCLE_STAT(TEXT("Phys SubstepEnd"), STAT_SubstepSimulationEnd, STATGROUP_Physics);
 
@@ -449,104 +236,8 @@ void FPhysSubstepTask::SubstepSimulationStart()
 {
 	SCOPE_CYCLE_COUNTER(STAT_TotalPhysicsTime);
 	SCOPE_CYCLE_COUNTER(STAT_SubstepSimulationStart);
-#if PHYSICS_INTERFACE_PHYSX
-	check(SubTime > 0.f);
-	check(DeltaSeconds > 0.f);
-	
-	check(!CompletionEvent.GetReference());	//should be done
-	CompletionEvent = FGraphEvent::CreateGraphEvent();
-	PhysXCompletionTask* SubstepTask = new PhysXCompletionTask(CompletionEvent,PAScene->getTaskManager());
-	ENamedThreads::Type NamedThread = PhysSingleThreadedMode() ? ENamedThreads::GameThread : ENamedThreads::SetTaskPriority(ENamedThreads::GameThread, ENamedThreads::HighTaskPriority);
-
-	DECLARE_CYCLE_STAT(TEXT("FDelegateGraphTask.ProcessPhysSubstepSimulation"),
-		STAT_FDelegateGraphTask_ProcessPhysSubstepSimulation,
-		STATGROUP_TaskGraphTasks);
-
-	FDelegateGraphTask::CreateAndDispatchWhenReady(
-		FDelegateGraphTask::FDelegate::CreateRaw(this, &FPhysSubstepTask::SubstepSimulationEnd),
-		GET_STATID(STAT_FDelegateGraphTask_ProcessPhysSubstepSimulation), CompletionEvent, ENamedThreads::GameThread, NamedThread);
-
-	++CurrentSubStep;	
-
-	bool bLastSubstep = CurrentSubStep >= NumSubsteps;
-
-	if (!bLastSubstep)
-	{
-
-		Alpha += StepScale;
-		TotalSubTime += SubTime;
-	}
-
-	float DeltaTime = bLastSubstep ? (DeltaSeconds - TotalSubTime) : SubTime;
-	float Interpolation = bLastSubstep ? 1.f : Alpha;
-
-	// Call scene step delegate
-#if !WITH_CHAOS
-	if (PhysScene != nullptr)
-	{
-		FSubstepCallbackGuard Guard(*this);
-		PhysScene->OnPhysSceneStep.Broadcast(PhysScene, DeltaTime);
-	}
-#endif
-
-	SubstepInterpolation(Interpolation, DeltaTime);
-
-#if WITH_APEX
-	PAScene->simulate(DeltaTime, bLastSubstep, SubstepTask, FullSimulationTask->GetScratchBufferData(), FullSimulationTask->GetScratchBufferSize());
-#else
-	PAScene->lockWrite();
-	PAScene->simulate(DeltaTime, SubstepTask, FullSimulationTask->GetScratchBufferData(), FullSimulationTask->GetScratchBufferSize());
-	PAScene->unlockWrite();
-#endif
-
-	SubstepTask->removeReference();
-#endif
 }
 
 void FPhysSubstepTask::SubstepSimulationEnd(ENamedThreads::Type CurrentThread, const FGraphEventRef& MyCompletionGraphEvent)
 {
-#if PHYSICS_INTERFACE_PHYSX
-	CompletionEvent = NULL;
-	if (CurrentSubStep < NumSubsteps)
-	{
-		uint32 OutErrorCode = 0;
-		{
-			SCOPE_CYCLE_COUNTER(STAT_TotalPhysicsTime);
-			SCOPE_CYCLE_COUNTER(STAT_SubstepSimulationEnd);
-
-			{
-#if WITH_APEX
-				SCOPED_APEX_SCENE_WRITE_LOCK(PAScene);
-				PxScene* PScene = PAScene->getPhysXScene();
-#else
-				PxScene* PScene = PAScene;
-				SCOPED_SCENE_WRITE_LOCK(PScene);
-#endif
-
-				// Dont rebuild scene queries here as this is an intermediate step
-				PScene->setSceneQueryUpdateMode(PxSceneQueryUpdateMode::eBUILD_DISABLED_COMMIT_DISABLED);
-
-				PAScene->fetchResults(true, &OutErrorCode);
-
-				// re-enable query updates (the next fetchResults will rebuild the SQ tree)
-				PScene->setSceneQueryUpdateMode(PxSceneQueryUpdateMode::eBUILD_ENABLED_COMMIT_ENABLED);
-			}
-		}
-
-		if (OutErrorCode != 0)
-		{
-			UE_LOG(LogPhysics, Log, TEXT("PHYSX FETCHRESULTS ERROR: %d"), OutErrorCode);
-		}
-
-		SubstepSimulationStart();
-	}
-	else
-	{
-		SCOPE_CYCLE_COUNTER(STAT_TotalPhysicsTime);
-		SCOPE_CYCLE_COUNTER(STAT_SubstepSimulationEnd);
-
-		//final step we call fetch on in game thread
-		FullSimulationTask->removeReference();
-	}
-#endif
 }
