Physics interface cleanup.

diff --git a/Engine/Source/Runtime/Engine/Private/PhysicsEngine/PhysSubstepTasks.h b/Engine/Source/Runtime/Engine/Private/PhysicsEngine/PhysSubstepTasks.h
--- a/Engine/Source/Runtime/Engine/Private/PhysicsEngine/PhysSubstepTasks.h
+++ b/Engine/Source/Runtime/Engine/Private/PhysicsEngine/PhysSubstepTasks.h
@@ -7,60 +7,11 @@
 #include "HAL/ThreadSafeBool.h"
 #include "Physics/PhysicsInterfaceCore.h"
 #include "PhysicsEngine/BodyInstance.h"
-#if PHYSICS_INTERFACE_PHYSX
-#include "PhysicsEngine/PhysXSupport.h"
-#endif
 
 struct FSimulationScratchBuffer;
 
 void FinishSceneStat();
 
-//This is only here for now while we transition into substepping
-#if PHYSICS_INTERFACE_PHYSX
-class PhysXCompletionTask : public PxLightCpuTask
-{
-	FGraphEventRef EventToFire;
-
-	// Scratch memory for call to PxScene::simulate. This is owned by the FPhysScene that spawned this
-	// task and is guaranteed to remain valid till FPhysScene destruction - can be nullptr if size
-	// was set to zero or a buffer isn't required for this scene
-	FSimulationScratchBuffer* ScratchBuffer;
-
-public:
-	PhysXCompletionTask()
-		: ScratchBuffer(nullptr)
-	{}
-
-	PhysXCompletionTask(FGraphEventRef& InEventToFire, PxTaskManager* InTaskManager, FSimulationScratchBuffer* InScratchBuffer = nullptr)
-		: EventToFire(InEventToFire)
-		, ScratchBuffer(InScratchBuffer)
-	{
-		setContinuation(*InTaskManager, NULL);
-	}
-	virtual void run()
-	{
-	}
-	virtual void release()
-	{
-		PxLightCpuTask::release();
-		FinishSceneStat();
-		if (EventToFire.GetReference())
-		{
-			EventToFire->DispatchSubsequents();
-		}
-		delete this;
-	}
-	virtual const char *getName() const
-	{
-		return "CompleteSimulate";
-	}
-
-	uint8* GetScratchBufferData();
-	int32 GetScratchBufferSize();
-};
-#endif	//#if WITH_PHYSX
-
-
 /** Hold information about kinematic target */
 struct FKinematicTarget
 {
@@ -148,22 +99,10 @@ struct FPhysTarget
 	
 };
 
-#if PHYSICS_INTERFACE_PHYSX
-#if WITH_APEX
-typedef apex::Scene PxApexScene;	//helper typedef so we don't have to use as many ifdefs
-#else
-typedef PxScene PxApexScene;
-#endif
-#endif
-
 /** Holds information used for substepping a scene */
 class FPhysSubstepTask
 {
 public:
-#if PHYSICS_INTERFACE_PHYSX
-	FPhysSubstepTask(PxApexScene * GivenScene, FPhysScene* InPhysScene);
-#endif
-
 	void SetKinematicTarget_AssumesLocked(FBodyInstance* Body, const FTransform& TM);
 	bool GetKinematicTarget_AssumesLocked(const FBodyInstance* Body, FTransform& OutTM) const;
 	void AddCustomPhysics_AssumesLocked(FBodyInstance* Body, const FCalculateCustomPhysics& CalculateCustomPhysics);
@@ -183,9 +122,6 @@ public:
 
 	void SubstepSimulationStart();
 	void SubstepSimulationEnd(ENamedThreads::Type CurrentThread, const FGraphEventRef& MyCompletionGraphEvent);
-#if PHYSICS_INTERFACE_PHYSX
-	void StepSimulation(PhysXCompletionTask * Task);
-#endif
 		
 private:
 
@@ -214,9 +150,4 @@ private:
 	FGraphEventRef CompletionEvent;
 
 	FPhysScene* PhysScene;
-
-#if PHYSICS_INTERFACE_PHYSX
-	PxApexScene * PAScene;
-#endif
-
 };
\ No newline at end of file
diff --git a/Engine/Source/Runtime/Engine/Private/PhysicsEngine/PhysUtils.cpp b/Engine/Source/Runtime/Engine/Private/PhysicsEngine/PhysUtils.cpp
--- a/Engine/Source/Runtime/Engine/Private/PhysicsEngine/PhysUtils.cpp
+++ b/Engine/Source/Runtime/Engine/Private/PhysicsEngine/PhysUtils.cpp
@@ -28,13 +28,9 @@
  */
 static float SuggestConvexVertexSnapDistance(UModel* InModel)
 {
-#if PHYSICS_INTERFACE_PHYSX
-	return SMALL_NUMBER;
-#else
 	// For an object around 1cmx1cmx1cm this is the vertex snap distance, then scaled to match the input model
 	constexpr float UnitSphereSnapDistance = 0.0015f;
 	return UnitSphereSnapDistance * InModel->Bounds.SphereRadius;
-#endif
 }
 
 /** Returns false if ModelToHulls operation should halt because of vertex count overflow. */
@@ -342,213 +338,16 @@ FCollisionEnabledMask operator|(const ECollisionEnabled::Type A, const FCollisio
 	return FCollisionEnabledMask(A) | B;
 }
 
-#if WITH_CHAOS
 bool FPhysScene::ExecPxVis(uint32 SceneType, const TCHAR* Cmd, FOutputDevice* Ar)
 {
     return false;
 }
-#else
-/** PxScene visualization*/
-bool FPhysScene::ExecPxVis(const TCHAR* Cmd, FOutputDevice* Ar)
-{
-#if WITH_PHYSX
-	// Get the scene to set flags on
-	PxScene* PScene = GetPxScene();
-
-	struct { const TCHAR* Name; PxVisualizationParameter::Enum Flag; float Size; } Flags[] =
-	{
-		// Axes
-		{ TEXT("WORLDAXES"),			PxVisualizationParameter::eWORLD_AXES,			1.f },
-		{ TEXT("BODYAXES"),				PxVisualizationParameter::eBODY_AXES,			1.f },
-		{ TEXT("MASSAXES"),             PxVisualizationParameter::eBODY_MASS_AXES,		1.f },
-
-		// Contacts
-		{ TEXT("CONTACTPOINT"),			PxVisualizationParameter::eCONTACT_POINT,		1.f },
-		{ TEXT("CONTACTS"),				PxVisualizationParameter::eCONTACT_NORMAL,		1.f },
-		{ TEXT("CONTACTERROR"),			PxVisualizationParameter::eCONTACT_ERROR,		100.f },
-		{ TEXT("CONTACTFORCE"),			PxVisualizationParameter::eCONTACT_FORCE,		1.f },
-
-		// Joints
-		{ TEXT("JOINTLIMITS"),			PxVisualizationParameter::eJOINT_LIMITS,		1.f },
-		{ TEXT("JOINTLOCALFRAMES"),		PxVisualizationParameter::eJOINT_LOCAL_FRAMES,	1.f },
-
-		// Collision
-		{ TEXT("COLLISION"),			PxVisualizationParameter::eCOLLISION_SHAPES,	1.f },
-	};
-
-	SCOPED_SCENE_WRITE_LOCK(PScene);
-
-	bool bDebuggingActive = false;
-	bool bFoundFlag = false;
-	if ( FParse::Command(&Cmd,TEXT("PHYSX_CLEAR_ALL")) )
-	{
-		Ar->Logf(TEXT("Clearing all PhysX Debug Flags."));
-		for (int32 i = 0; i < UE_ARRAY_COUNT(Flags); i++)
-		{
-			PScene->setVisualizationParameter(Flags[i].Flag, 0.0f);
-			bFoundFlag = true;
-		}
-	}
-	else
-	{
-		for (int32 i = 0; i < UE_ARRAY_COUNT(Flags); i++)
-		{
-			// Parse out the command sent in and set only those flags
-			if (FParse::Command(&Cmd, Flags[i].Name))
-			{
-				PxReal Result = PScene->getVisualizationParameter(Flags[i].Flag);
-				if (Result == 0.0f)
-				{
-					PScene->setVisualizationParameter(Flags[i].Flag, Flags[i].Size);
-					Ar->Logf(TEXT("Flag set."));
-				}
-				else
-				{
-					PScene->setVisualizationParameter(Flags[i].Flag, 0.0f);
-					Ar->Logf(TEXT("Flag un-set."));
-				}
-
-				bFoundFlag = true;
-			}
-
-			// See if any flags are true
-			PxReal Result = PScene->getVisualizationParameter(Flags[i].Flag);
-			if(Result > 0.f)
-			{
-				bDebuggingActive = true;
-			}
-		}
-	}
-
-	// If no debugging going on - disable it using NX_VISUALIZATION_SCALE
-	if(bDebuggingActive)
-	{
-		PScene->setVisualizationParameter(PxVisualizationParameter::eSCALE, 20.0f);
-	}
-	else
-	{
-		PScene->setVisualizationParameter(PxVisualizationParameter::eSCALE, 0.0f);
-	}
-
-	if(!bFoundFlag)
-	{
-		Ar->Logf(TEXT("Unknown PhysX visualization flag specified."));
-	}
-#endif	// WITH_PHYSX
-
-	return 1;
-}
-#endif
 
-#if WITH_CHAOS
 bool FPhysScene::ExecApexVis(uint32 SceneType, const TCHAR* Cmd, FOutputDevice* Ar)
 {
     return false;
 }
-#else
-/** PxScene visualization */
-bool FPhysScene::ExecApexVis(const TCHAR* Cmd, FOutputDevice* Ar)
-{
-	check(Cmd);
-#if WITH_PHYSX
-#if WITH_APEX
-	// Get the scene to set flags on
-	apex::Scene* ApexScene = GetApexScene();
-
-	if (ApexScene == NULL)
-	{
-		return false;
-	}
-
-	NvParameterized::Interface* DebugRenderParams = ApexScene->getDebugRenderParams();
-	check(DebugRenderParams != NULL);
-
-	// Toggle global flags if there are no further arguments
-	const bool bToggle = *Cmd == TCHAR(0);
-
-	// Enable or toggle visualization
-	NvParameterized::Handle EnableDebugRenderHandle(*DebugRenderParams, "Enable");
-	check(EnableDebugRenderHandle.isValid());
-	bool bEnableValue = true;
-	if (bToggle)
-	{
-		EnableDebugRenderHandle.getParamBool(bEnableValue);
-		bEnableValue = !bEnableValue;
-	}
-	EnableDebugRenderHandle.setParamBool(bEnableValue);
-	NvParameterized::Handle ScaleDebugRenderHandle(*DebugRenderParams, "Scale");
-	check(ScaleDebugRenderHandle.isValid());
-	float ScaleValue = 1.0f;
-	if (bToggle)
-	{
-		ScaleDebugRenderHandle.getParamF32(ScaleValue);
-		ScaleValue = ScaleValue > 0.0f ? 0.0f : 1.0f;
-	}
-	ScaleDebugRenderHandle.setParamF32(ScaleValue);
-
-	// See if there's a '/', which means we have a module-specific visualization parameter
-	const TCHAR* Slash = Cmd;
-	while (*Slash != TCHAR(0) && *Slash != TCHAR('/'))
-	{
-		++Slash;
-	}
-
-	if (*Slash != TCHAR(0))
-	{
-		FString ModuleName = FString(Cmd).Left((int32)(Slash-Cmd));
-		DebugRenderParams = ApexScene->getModuleDebugRenderParams(StringCast<ANSICHAR>(*ModuleName).Get());
-	}
-
-	if (DebugRenderParams == NULL)
-	{
-		Ar->Logf(TEXT("Unknown APEX module requested for apex debug visualization."));
-		return false;
-	}
-
-	auto FlagNameAnsi = StringCast<ANSICHAR>((*Slash == TCHAR(0) ? Cmd : Slash+1));
-	NvParameterized::Handle DebugRenderHandle(*DebugRenderParams, FlagNameAnsi.Get());
-
-	if (!DebugRenderHandle.isValid())
-	{
-		Ar->Logf(TEXT("Unknown APEX visualization flag specified."));
-		return false;
-	}
-
-	if (DebugRenderHandle.parameterDefinition()->type() == NvParameterized::TYPE_F32)
-	{
-		PxF32 Value;
-		DebugRenderHandle.getParamF32(Value);
-		DebugRenderHandle.setParamF32(Value > 0.0f ? 0.0f : 1.0f);
-	}
-	else
-	if (DebugRenderHandle.parameterDefinition()->type() == NvParameterized::TYPE_U32)
-	{
-		PxU32 Value;
-		DebugRenderHandle.getParamU32(Value);
-		DebugRenderHandle.setParamU32(Value > 0 ? 0 : 1);
-	}
-	else
-	if (DebugRenderHandle.parameterDefinition()->type() == NvParameterized::TYPE_BOOL)
-	{
-		bool bValue;
-		DebugRenderHandle.getParamBool(bValue);
-		DebugRenderHandle.setParamBool(!bValue);
-	}
-	else
-	{
-		Ar->Logf(TEXT("Unknown APEX visualization flag type."));
-		return false;
-	}
-
-#endif	// WITH_APEX
-#endif	// WITH_PHYSX
-
-	return true;
-}
-#endif
 
-
-#if WITH_CHAOS
 bool FPhysicsInterface::ExecPhysCommands(const TCHAR* Cmd, FOutputDevice* OutputDevice, UWorld* InWorld)
 {
 	if (FParse::Command(&Cmd, TEXT("ChaosGeometryMemory")))
@@ -599,97 +398,4 @@ bool FPhysicsInterface::ExecPhysCommands(const TCHAR* Cmd, FOutputDevice* Output
 #endif
 
     return false;
-}
-#else
-bool FPhysScene::HandleExecCommands(const TCHAR* Cmd, FOutputDevice* Ar)
-{
-	if (FParse::Command(&Cmd, TEXT("PXVIS")) || FParse::Command(&Cmd, TEXT("APEXVIS")))
-	{
-		return ExecPxVis(Cmd, Ar);
-	}
-
-	return false;
-}
-
-//// EXEC
-bool FPhysicsInterface::ExecPhysCommands(const TCHAR* Cmd, FOutputDevice* Ar, UWorld* InWorld)
-{
-#if WITH_PHYSX
-	FPhysScene* PhysScene = InWorld->GetPhysicsScene();
-	// Give PhysScene change to handle commands
-	if (PhysScene && PhysScene->HandleExecCommands(Cmd, Ar))
-	{
-		return true;
-	}
-	else if(!IsRunningCommandlet() && GPhysXSDK && FParse::Command(&Cmd, TEXT("PVD")) )
-	{
-		// check if PvdConnection manager is available on this platform
-		if(GPhysXVisualDebugger != NULL)
-		{
-			if(FParse::Command(&Cmd, TEXT("CONNECT")))
-			{
-				
-				
-				const bool bVizualization = !FParse::Command(&Cmd, TEXT("NODEBUG"));
-
-				// setup connection parameters
-				FString Host = TEXT("localhost");
-				if (*Cmd)
-				{
-					Host = Cmd;
-				}
-
-				
-
-				PvdConnect(Host, bVizualization);
-
-			}
-			else if(FParse::Command(&Cmd, TEXT("DISCONNECT")))
-			{
-				GPhysXVisualDebugger->disconnect();
-			}
-		}
-
-		return 1;
-	}
-#if PHYSX_MEMORY_STATS
-	else if(GPhysXAllocator && FParse::Command(&Cmd, TEXT("PHYSXALLOC")) )
-	{
-		GPhysXAllocator->DumpAllocations(Ar);
-		return 1;
-	}
-#endif
-	else if(FParse::Command(&Cmd, TEXT("PHYSXSHARED")) )
-	{
-		FPhysxSharedData::Get().DumpSharedMemoryUsage(Ar);
-		return 1;
-	}
-	else if(FParse::Command(&Cmd, TEXT("PHYSXINFO")))
-	{
-		Ar->Logf(TEXT("PhysX Info:"));
-		Ar->Logf(TEXT("  Version: %d.%d.%d"), PX_PHYSICS_VERSION_MAJOR, PX_PHYSICS_VERSION_MINOR, PX_PHYSICS_VERSION_BUGFIX);
-#if UE_BUILD_DEBUG && !defined(NDEBUG)
-		Ar->Logf(TEXT("  Configuration: DEBUG"));
-#elif WITH_PHYSX_RELEASE
-		Ar->Logf(TEXT("  Configuration: RELEASE"));
-#else
-		Ar->Logf(TEXT("  Configuration: PROFILE"));
-#endif
-		if(GetPhysXCookingModule())
-		{
-			Ar->Logf(TEXT("  Cooking Module: TRUE"));
-		}
-		else
-		{
-			Ar->Logf(TEXT("  Cooking Module: FALSE"));
-		}
-
-		return 1;
-	}
-
-#endif // WITH_PHYSX
-
-	return 0;
-}
-#endif
-
+}
\ No newline at end of file
diff --git a/Engine/Source/Runtime/Engine/Private/PhysicsEngine/PhysXCookHelper.cpp b/Engine/Source/Runtime/Engine/Private/PhysicsEngine/PhysXCookHelper.cpp
--- a/Engine/Source/Runtime/Engine/Private/PhysicsEngine/PhysXCookHelper.cpp
+++ b/Engine/Source/Runtime/Engine/Private/PhysicsEngine/PhysXCookHelper.cpp
@@ -1,76 +1 @@
 // Copyright Epic Games, Inc. All Rights Reserved.
-
-
-#if PHYSICS_INTERFACE_PHYSX
-#include "PhysXCookHelper.h"
-#include "PhysXSupport.h"
-#include "IPhysXCookingModule.h"
-#include "IPhysXCooking.h"
-
-FPhysXCookHelper::FPhysXCookHelper(IPhysXCookingModule* InPhysXCookingModule)
-	: PhysXCookingModule(InPhysXCookingModule)
-{
-}
-
-bool FPhysXCookHelper::CreatePhysicsMeshes_Concurrent()
-{
-	bool bSuccess = true;
-
-	CreateConvexElements_Concurrent(CookInfo.NonMirroredConvexVertices, OutNonMirroredConvexMeshes, false);
-	CreateConvexElements_Concurrent(CookInfo.MirroredConvexVertices, OutMirroredConvexMeshes, true);
-
-	if (CookInfo.bCookTriMesh && !CookInfo.bTriMeshError)
-	{
-		OutTriangleMeshes.AddZeroed();
-		const bool bError = !PhysXCookingModule->GetPhysXCooking()->CreateTriMesh(FPlatformProperties::GetPhysicsFormat(), CookInfo.TriMeshCookFlags, CookInfo.TriangleMeshDesc.Vertices, CookInfo.TriangleMeshDesc.Indices, CookInfo.TriangleMeshDesc.MaterialIndices, CookInfo.TriangleMeshDesc.bFlipNormals, OutTriangleMeshes[0]);
-		if (bError)
-		{
-			bSuccess = false;
-			UE_LOG(LogPhysics, Warning, TEXT("Failed to cook TriMesh: %s."), *CookInfo.OuterDebugName);
-		}
-		else if (CookInfo.bSupportUVFromHitResults)
-		{
-			OutUVInfo.FillFromTriMesh(CookInfo.TriangleMeshDesc);
-		}
-	}
-
-    return bSuccess;
-}
-
-void FPhysXCookHelper::CreatePhysicsMeshesAsync_Concurrent(FSimpleDelegateGraphTask::FDelegate FinishDelegate)
-{
-	CreatePhysicsMeshes_Concurrent();
-	FSimpleDelegateGraphTask::CreateAndDispatchWhenReady(FinishDelegate, GET_STATID(STAT_PhysXCooking), nullptr, ENamedThreads::GameThread);
-}
-
-void FPhysXCookHelper::CreateConvexElements_Concurrent(const TArray<TArray<FVector>>& Elements, TArray<PxConvexMesh*>& OutConvexMeshes, bool bFlipped)
-{
-	OutMirroredConvexMeshes.Reserve(Elements.Num());
-	for (int32 ElementIndex = 0; ElementIndex < Elements.Num(); ++ElementIndex)
-	{
-		OutConvexMeshes.AddZeroed();
-		const EPhysXCookingResult Result = PhysXCookingModule->GetPhysXCooking()->CreateConvex(FPlatformProperties::GetPhysicsFormat(), CookInfo.ConvexCookFlags, Elements[ElementIndex], OutConvexMeshes.Last());
-		switch (Result)
-		{
-		case EPhysXCookingResult::Succeeded:
-			break;
-		case EPhysXCookingResult::Failed:
-			UE_LOG(LogPhysics, Warning, TEXT("Failed to cook convex: %s %d (FlipX:%d). The remaining elements will not get cooked."), *CookInfo.OuterDebugName, ElementIndex, bFlipped ? 1 : 0);
-			break;
-		case EPhysXCookingResult::SucceededWithInflation:
-			if (!CookInfo.bConvexDeformableMesh)
-			{
-				UE_LOG(LogPhysics, Warning, TEXT("Cook convex: %s %d (FlipX:%d) failed but succeeded with inflation.  The mesh should be looked at."), *CookInfo.OuterDebugName, ElementIndex, bFlipped ? 1 : 0);
-			}
-			else
-			{
-				UE_LOG(LogPhysics, Log, TEXT("Cook convex: %s %d (FlipX:%d) required inflation. You may wish to adjust the mesh so this is not necessary."), *CookInfo.OuterDebugName, ElementIndex, bFlipped ? 1 : 0);
-			}
-			break;
-		default:
-			check(false);
-		}
-	}
-}
-
-#endif //WITH_PHYSX
diff --git a/Engine/Source/Runtime/Engine/Private/PhysicsEngine/PhysXSupport.cpp b/Engine/Source/Runtime/Engine/Private/PhysicsEngine/PhysXSupport.cpp
--- a/Engine/Source/Runtime/Engine/Private/PhysicsEngine/PhysXSupport.cpp
+++ b/Engine/Source/Runtime/Engine/Private/PhysicsEngine/PhysXSupport.cpp
@@ -11,677 +11,3 @@
 #include "Physics/PhysicsInterfaceCore.h"
 #include "PhysicalMaterials/PhysicalMaterial.h"
 #include "CustomPhysXPayload.h"
-
-#if PHYSICS_INTERFACE_PHYSX
-
-#include "PhysXPublic.h"
-#include "PhysicsEngine/ConstraintInstance.h"
-#include "PhysicsEngine/BodySetup.h"
-#include "PhysicsPublicCore.h"
-
-int32						GNumPhysXConvexMeshes = 0;
-
-TArray<PxConvexMesh*>	GPhysXPendingKillConvex;
-TArray<PxTriangleMesh*>	GPhysXPendingKillTriMesh;
-TArray<PxHeightField*>	GPhysXPendingKillHeightfield;
-///////////////////// Utils /////////////////////
-
-
-void AddRadialImpulseToPxRigidBody_AssumesLocked(PxRigidBody& PRigidBody, const FVector& Origin, float Radius, float Strength, uint8 Falloff, bool bVelChange)
-{
-#if WITH_PHYSX
-	if (!(PRigidBody.getRigidBodyFlags() & PxRigidBodyFlag::eKINEMATIC))
-	{
-		float Mass = PRigidBody.getMass();
-		PxTransform PCOMTransform = PRigidBody.getGlobalPose().transform(PRigidBody.getCMassLocalPose());
-		PxVec3 PCOMPos = PCOMTransform.p; // center of mass in world space
-		PxVec3 POrigin = U2PVector(Origin); // origin of radial impulse, in world space
-		PxVec3 PDelta = PCOMPos - POrigin; // vector from origin to COM
-
-		float Mag = PDelta.magnitude(); // Distance from COM to origin, in Unreal scale : @todo: do we still need conversion scale?
-
-		// If COM is outside radius, do nothing.
-		if (Mag > Radius)
-		{
-			return;
-		}
-
-		PDelta.normalize();
-
-		// Scale by U2PScale here, because units are velocity * mass. 
-		float ImpulseMag = Strength;
-		if (Falloff == RIF_Linear)
-		{
-			ImpulseMag *= (1.0f - (Mag / Radius));
-		}
-
-		PxVec3 PImpulse = PDelta * ImpulseMag;
-
-		PxForceMode::Enum Mode = bVelChange ? PxForceMode::eVELOCITY_CHANGE : PxForceMode::eIMPULSE;
-		PRigidBody.addForce(PImpulse, Mode);
-	}
-#endif // WITH_PHYSX
-}
-
-void AddRadialForceToPxRigidBody_AssumesLocked(PxRigidBody& PRigidBody, const FVector& Origin, float Radius, float Strength, uint8 Falloff, bool bAccelChange)
-{
-#if WITH_PHYSX
-	if (!(PRigidBody.getRigidBodyFlags() & PxRigidBodyFlag::eKINEMATIC))
-	{
-		float Mass = PRigidBody.getMass();
-		PxTransform PCOMTransform = PRigidBody.getGlobalPose().transform(PRigidBody.getCMassLocalPose());
-		PxVec3 PCOMPos = PCOMTransform.p; // center of mass in world space
-		PxVec3 POrigin = U2PVector(Origin); // origin of radial impulse, in world space
-		PxVec3 PDelta = PCOMPos - POrigin; // vector from
-
-		float Mag = PDelta.magnitude(); // Distance from COM to origin, in Unreal scale : @todo: do we still need conversion scale?
-
-		// If COM is outside radius, do nothing.
-		if (Mag > Radius)
-		{
-			return;
-		}
-
-		PDelta.normalize();
-
-		// If using linear falloff, scale with distance.
-		float ForceMag = Strength;
-		if (Falloff == RIF_Linear)
-		{
-			ForceMag *= (1.0f - (Mag / Radius));
-		}
-
-		// Apply force
-		PxVec3 PImpulse = PDelta * ForceMag;
-		PRigidBody.addForce(PImpulse, bAccelChange ? PxForceMode::eACCELERATION : PxForceMode::eFORCE);
-	}
-#endif // WITH_PHYSX
-}
-
-bool IsRigidBodyKinematicAndInSimulationScene_AssumesLocked(const PxRigidBody* PRigidBody)
-{
-	if (PRigidBody)
-	{
-		//For some cases we only consider an actor kinematic if it's in the simulation scene. This is in cases where we set a kinematic target
-		return (PRigidBody->getRigidBodyFlags() & PxRigidBodyFlag::eKINEMATIC) && !(PRigidBody->getActorFlags() & PxActorFlag::eDISABLE_SIMULATION);
-	}
-
-	return false;
-}
-
-//////////////////////////////////////////////////////////////////////////
-// PHYSXSIMFILTERSHADER 
-
-/** Util to return a string for the type of a query (for debugging) */
-FString ObjTypeToString(PxFilterObjectAttributes PAtt)
-{
-	PxFilterObjectType::Enum Type = PxGetFilterObjectType(PAtt);
-
-	if(Type == PxFilterObjectType::eRIGID_STATIC)
-	{
-		return TEXT("rigid static");
-	}
-	else if(Type == PxFilterObjectType::eRIGID_DYNAMIC)
-	{
-		return TEXT("rigid dynamic");
-	}
-
-	return TEXT("unknown");
-}
-
-PxFilterFlags PhysXSimFilterShader(	PxFilterObjectAttributes attributes0, PxFilterData filterData0, 
-									PxFilterObjectAttributes attributes1, PxFilterData filterData1,
-									PxPairFlags& pairFlags, const void* constantBlock, PxU32 constantBlockSize )
-{
-	//UE_LOG(LogPhysics, Log, TEXT("filterData0 (%s): %x %x %x %x"), *ObjTypeToString(attributes0), filterData0.word0, filterData0.word1, filterData0.word2, filterData0.word3);
-	//UE_LOG(LogPhysics, Log, TEXT("filterData1 (%s): %x %x %x %x"), *ObjTypeToString(attributes1), filterData1.word0, filterData1.word1, filterData1.word2, filterData1.word3);
-
-	bool k0 = PxFilterObjectIsKinematic(attributes0);
-	bool k1 = PxFilterObjectIsKinematic(attributes1);
-
-	PxU32 FilterFlags0 = (filterData0.word3 & 0xFFFFFF);
-	PxU32 FilterFlags1 = (filterData1.word3 & 0xFFFFFF);
-
-	if (k0 && k1)
-	{
-		//Ignore kinematic kinematic pairs unless they are explicitly requested
-		if(!(FilterFlags0&EPDF_KinematicKinematicPairs) && !(FilterFlags1&EPDF_KinematicKinematicPairs))
-		{
-			return PxFilterFlag::eSUPPRESS;	//NOTE: Waiting on physx fix for refiltering on aggregates. For now use supress which automatically tests when changes to simulation happen
-		}
-	}
-	
-	bool s0 = PxGetFilterObjectType(attributes0) == PxFilterObjectType::eRIGID_STATIC;
-	bool s1 = PxGetFilterObjectType(attributes1) == PxFilterObjectType::eRIGID_STATIC;
-
-	//ignore static-kinematic (this assumes that statics can't be flagged as kinematics)
-	// should return eSUPPRESS here instead eKILL so that kinematics vs statics will still be considered once kinematics become dynamic (dying ragdoll case)
-	if((k0 || k1) && (s0 || s1))
-	{
-		return PxFilterFlag::eSUPPRESS;
-	}
-	
-	// if these bodies are from the same component, use the disable table to see if we should disable collision. This case should only happen for things like skeletalmesh and destruction. The table is only created for skeletal mesh components at the moment
-#if !WITH_CHAOS
-	if(filterData0.word2 == filterData1.word2)
-	{
-		check(constantBlockSize == sizeof(FPhysSceneShaderInfo));
-		const FPhysSceneShaderInfo* PhysSceneShaderInfo = (const FPhysSceneShaderInfo*) constantBlock;
-		check(PhysSceneShaderInfo);
-		FPhysScene * PhysScene = PhysSceneShaderInfo->PhysScene;
-		check(PhysScene);
-
-		const TMap<uint32, TMap<FRigidBodyIndexPair, bool> *> & CollisionDisableTableLookup = PhysScene->GetCollisionDisableTableLookup();
-		TMap<FRigidBodyIndexPair, bool>* const * DisableTablePtrPtr = CollisionDisableTableLookup.Find(filterData1.word2);
-		if (DisableTablePtrPtr)		//Since collision table is deferred during sub-stepping it's possible that we won't get the collision disable table until the next frame
-		{
-			TMap<FRigidBodyIndexPair, bool>* DisableTablePtr = *DisableTablePtrPtr;
-			FRigidBodyIndexPair BodyPair(filterData0.word0, filterData1.word0); // body indexes are stored in word 0
-			if (DisableTablePtr->Find(BodyPair))
-			{
-				return PxFilterFlag::eKILL;
-			}
-
-		}
-	}
-#endif
-
-	// Find out which channels the objects are in
-	ECollisionChannel Channel0 = GetCollisionChannel(filterData0.word3);
-	ECollisionChannel Channel1 = GetCollisionChannel(filterData1.word3);
-	
-	// see if 0/1 would like to block the other 
-	PxU32 BlockFlagTo1 = (ECC_TO_BITFIELD(Channel1) & filterData0.word1);
-	PxU32 BlockFlagTo0 = (ECC_TO_BITFIELD(Channel0) & filterData1.word1);
-
-	bool bDoesWantToBlock = (BlockFlagTo1 && BlockFlagTo0);
-
-	// if don't want to block, suppress
-	if ( !bDoesWantToBlock )
-	{
-		return PxFilterFlag::eSUPPRESS;
-	}
-
-	
-
-	pairFlags = PxPairFlag::eCONTACT_DEFAULT;
-
-	//todo enabling CCD objects against everything else for now
-	if(!(k0 && k1) && ((FilterFlags0&EPDF_CCD) || (FilterFlags1&EPDF_CCD)))
-	{
-		pairFlags |= PxPairFlag::eDETECT_CCD_CONTACT | PxPairFlag::eSOLVE_CONTACT;
-	}
-
-
-	if((FilterFlags0&EPDF_ContactNotify) || (FilterFlags1&EPDF_ContactNotify))
-	{
-		pairFlags |= (PxPairFlag::eNOTIFY_TOUCH_FOUND | PxPairFlag::eNOTIFY_TOUCH_PERSISTS | PxPairFlag::eNOTIFY_CONTACT_POINTS );
-	}
-
-
-	if ((FilterFlags0&EPDF_ModifyContacts) || (FilterFlags1&EPDF_ModifyContacts))
-	{
-		pairFlags |= (PxPairFlag::eMODIFY_CONTACTS);
-	}
-
-	return PxFilterFlags();
-}
-
-#if !WITH_CHAOS
-
-/** Figures out the new FCollisionNotifyInfo needed for pending notification. It adds it, and then returns an array that goes from pair index to notify collision index */
-TArray<int32> AddCollisionNotifyInfo(const FBodyInstance* Body0, const FBodyInstance* Body1, const physx::PxContactPair * Pairs, uint32 NumPairs, TArray<FCollisionNotifyInfo> & PendingNotifyInfos)
-{
-	TArray<int32> PairNotifyMapping;
-	PairNotifyMapping.Empty(NumPairs);
-
-	TMap<const FBodyInstance*, TMap<const FBodyInstance*, int32> > BodyPairNotifyMap;
-	for(uint32 PairIdx = 0; PairIdx < NumPairs; ++PairIdx)
-	{
-		const PxContactPair* Pair = Pairs + PairIdx;
-		PairNotifyMapping.Add(-1);	//start as -1 because we can have collisions that we don't want to actually record collision
-
-									// Check if either shape has been removed
-		if(!Pair->events.isSet(PxPairFlag::eNOTIFY_TOUCH_LOST) &&
-			!Pair->events.isSet(PxPairFlag::eNOTIFY_THRESHOLD_FORCE_LOST) &&
-			!Pair->flags.isSet(PxContactPairFlag::eREMOVED_SHAPE_0) &&
-			!Pair->flags.isSet(PxContactPairFlag::eREMOVED_SHAPE_1))
-		{
-			// Get the two shapes that are involved in the collision
-			const PxShape* Shape0 = Pair->shapes[0];
-			check(Shape0);
-			const PxShape* Shape1 = Pair->shapes[1];
-			check(Shape1);
-
-			PxU32 FilterFlags0 = Shape0->getSimulationFilterData().word3 & 0xFFFFFF;
-			PxU32 FilterFlags1 = Shape1->getSimulationFilterData().word3 & 0xFFFFFF;
-
-			const bool bBody0Notify = (FilterFlags0 & EPDF_ContactNotify) != 0;
-			const bool bBody1Notify = (FilterFlags1 & EPDF_ContactNotify) != 0;
-
-			if(bBody0Notify || bBody1Notify)
-			{
-#if WITH_IMMEDIATE_PHYSX
-                check(false);
-#else
-				const FBodyInstance* SubBody0 = FPhysicsInterface_PhysX::ShapeToOriginalBodyInstance(Body0, Shape0);
-				const FBodyInstance* SubBody1 = FPhysicsInterface_PhysX::ShapeToOriginalBodyInstance(Body1, Shape1);
-
-				TMap<const FBodyInstance *, int32> & SubBodyNotifyMap = BodyPairNotifyMap.FindOrAdd(SubBody0);
-				int32* NotifyInfoIndex = SubBodyNotifyMap.Find(SubBody1);
-
-				if(NotifyInfoIndex == NULL)
-				{
-					FCollisionNotifyInfo * NotifyInfo = new (PendingNotifyInfos) FCollisionNotifyInfo;
-					NotifyInfo->bCallEvent0 = bBody0Notify;
-					NotifyInfo->Info0.SetFrom(SubBody0);
-					NotifyInfo->bCallEvent1 = bBody1Notify;
-					NotifyInfo->Info1.SetFrom(SubBody1);
-
-					NotifyInfoIndex = &SubBodyNotifyMap.Add(SubBody0, PendingNotifyInfos.Num() - 1);
-				}
-
-				PairNotifyMapping[PairIdx] = *NotifyInfoIndex;
-#endif
-			}
-		}
-	}
-
-	return PairNotifyMapping;
-}
-
-///////// FPhysXSimEventCallback //////////////////////////////////
-
-void FPhysXSimEventCallback::onContact(const PxContactPairHeader& PairHeader, const PxContactPair* Pairs, PxU32 NumPairs)
-{
-	// Check actors are not destroyed
-	if( PairHeader.flags & (PxContactPairHeaderFlag::eREMOVED_ACTOR_0 | PxContactPairHeaderFlag::eREMOVED_ACTOR_1) )
-	{
-		UE_LOG(LogPhysics, Log, TEXT("%llu onContact(): Actors have been deleted!"), (uint64)GFrameCounter );
-		return;
-	}
-
-	const PxActor* PActor0 = PairHeader.actors[0];
-	const PxActor* PActor1 = PairHeader.actors[1];
-	check(PActor0 && PActor1);
-
-	const PxRigidBody* PRigidBody0 = PActor0->is<PxRigidBody>();
-	const PxRigidBody* PRigidBody1 = PActor1->is<PxRigidBody>();
-
-	const FBodyInstance* BodyInst0 = FPhysxUserData::Get<FBodyInstance>(PActor0->userData);
-	const FBodyInstance* BodyInst1 = FPhysxUserData::Get<FBodyInstance>(PActor1->userData);
-	
-	bool bEitherCustomPayload = false;
-
-	// check if it is a custom payload with special body instance conversion
-	if (BodyInst0 == nullptr)
-	{
-		if (const FCustomPhysXPayload* CustomPayload = FPhysxUserData::Get<FCustomPhysXPayload>(PActor0->userData))
-		{
-			bEitherCustomPayload = true;
-			BodyInst0 = CustomPayload->GetBodyInstance();
-		}
-	}
-
-	if (BodyInst1 == nullptr)
-	{
-		if (const FCustomPhysXPayload* CustomPayload = FPhysxUserData::Get<FCustomPhysXPayload>(PActor1->userData))
-		{
-			bEitherCustomPayload = true;
-			BodyInst1 = CustomPayload->GetBodyInstance();
-		}
-	}
-
-	//if nothing valid just exit
-	//if a custom payload (like apex destruction) generates collision between the same body instance we ignore it. This is potentially bad, but in general we have not had a need for this
-	if(BodyInst0 == nullptr || BodyInst1 == nullptr || BodyInst0 == BodyInst1)
-	{
-		return;
-	}
-
-	//custom payloads may (hackily) rely on the onContact flag. Apex Destruction needs this for being able to apply damage as a result of collision.
-	//Because of this we only want onContact events to happen if the user actually selected bNotifyRigidBodyCollision so we have to check if this is the case
-	if (bEitherCustomPayload)
-	{
-		if (BodyInst0->bNotifyRigidBodyCollision == false && BodyInst1->bNotifyRigidBodyCollision == false)
-		{
-			return;
-		}
-	}
-
-	TArray<FCollisionNotifyInfo>& PendingCollisionNotifies = OwningScene->GetPendingCollisionNotifies();
-
-	uint32 PreAddingCollisionNotify = PendingCollisionNotifies.Num() - 1;
-	TArray<int32> PairNotifyMapping = AddCollisionNotifyInfo(BodyInst0, BodyInst1, Pairs, NumPairs, PendingCollisionNotifies);
-
-	// Iterate through contact points
-	for(uint32 PairIdx=0; PairIdx<NumPairs; PairIdx++)
-	{
-		int32 NotifyIdx = PairNotifyMapping[PairIdx];
-		if (NotifyIdx == -1)	//the body instance this pair belongs to is not listening for events
-		{
-			continue;
-		}
-
-		FCollisionNotifyInfo * NotifyInfo = &PendingCollisionNotifies[NotifyIdx];
-		FCollisionImpactData* ImpactInfo = &(NotifyInfo->RigidCollisionData);
-
-		const PxContactPair* Pair = Pairs + PairIdx;
-
-		// Get the two shapes that are involved in the collision
-		const PxShape* Shape0 = Pair->shapes[0];
-		check(Shape0);
-		const PxShape* Shape1 = Pair->shapes[1];
-		check(Shape1);
-
-		// Get materials
-		PxMaterial* Material0 = nullptr;
-		UPhysicalMaterial* PhysMat0  = nullptr;
-		if(Shape0->getNbMaterials() == 1)	//If we have simple geometry or only 1 material we set it here. Otherwise do it per face
-		{
-			Shape0->getMaterials(&Material0, 1);		
-			PhysMat0 = Material0 ? FPhysxUserData::Get<UPhysicalMaterial>(Material0->userData) : nullptr;
-		}
-
-		PxMaterial* Material1 = nullptr;
-		UPhysicalMaterial* PhysMat1  = nullptr;
-		if (Shape1->getNbMaterials() == 1)	//If we have simple geometry or only 1 material we set it here. Otherwise do it per face
-		{
-			Shape1->getMaterials(&Material1, 1);
-			PhysMat1 = Material1 ? FPhysxUserData::Get<UPhysicalMaterial>(Material1->userData) : nullptr;
-		}
-
-		// Iterate over contact points
-		PxContactPairPoint ContactPointBuffer[16];
-		int32 NumContactPoints = Pair->extractContacts(ContactPointBuffer, 16);
-		for(int32 PointIdx=0; PointIdx<NumContactPoints; PointIdx++)
-		{
-			const PxContactPairPoint& Point = ContactPointBuffer[PointIdx];
-
-			const PxVec3 NormalImpulse = Point.impulse.dot(Point.normal) * Point.normal; // project impulse along normal
-			ImpactInfo->TotalNormalImpulse += P2UVector(NormalImpulse);
-			ImpactInfo->TotalFrictionImpulse += P2UVector(Point.impulse - NormalImpulse); // friction is component not along contact normal
-
-			// Get per face materials
-			if(!Material0)	//there is complex geometry or multiple materials so resolve the physical material here
-			{
-				if(PxMaterial* Material0PerFace = Shape0->getMaterialFromInternalFaceIndex(Point.internalFaceIndex0))
-				{
-					PhysMat0 = FPhysxUserData::Get<UPhysicalMaterial>(Material0PerFace->userData);
-				}
-			}
-
-			if (!Material1)	//there is complex geometry or multiple materials so resolve the physical material here
-			{
-				if(PxMaterial* Material1PerFace = Shape1->getMaterialFromInternalFaceIndex(Point.internalFaceIndex1))
-				{
-					PhysMat1 = FPhysxUserData::Get<UPhysicalMaterial>(Material1PerFace->userData);
-				}
-				
-			}
-			
-			new(ImpactInfo->ContactInfos) FRigidBodyContactInfo(
-				P2UVector(Point.position), 
-				P2UVector(Point.normal), 
-				-1.f * Point.separation, 
-				PhysMat0, 
-				PhysMat1);
-		}	
-	}
-
-	for (int32 NotifyIdx = PreAddingCollisionNotify + 1; NotifyIdx < PendingCollisionNotifies.Num(); NotifyIdx++)
-	{
-		FCollisionNotifyInfo * NotifyInfo = &PendingCollisionNotifies[NotifyIdx];
-		FCollisionImpactData* ImpactInfo = &(NotifyInfo->RigidCollisionData);
-		// Discard pairs that don't generate any force (eg. have been rejected through a modify contact callback).
-		if (ImpactInfo->TotalNormalImpulse.SizeSquared() < KINDA_SMALL_NUMBER)
-		{
-			PendingCollisionNotifies.RemoveAt(NotifyIdx);
-			NotifyIdx--;
-		}
-	}
-}
-
-void FPhysXSimEventCallback::onConstraintBreak( PxConstraintInfo* constraints, PxU32 count )
-{
-	for (int32 i=0; i < (int32)count; ++i)
-	{
-		PxJoint* Joint = (PxJoint*)(constraints[i].externalReference);
-
-		if (Joint && Joint->userData)
-		{
-			if (FConstraintInstance* Constraint = FPhysxUserData::Get<FConstraintInstance>(Joint->userData))
-			{
-				OwningScene->AddPendingOnConstraintBreak(Constraint);
-			}
-		}
-	}
-}
-
-void FPhysXSimEventCallback::onWake(PxActor** Actors, PxU32 Count)
-{
-	for(PxU32 ActorIdx = 0; ActorIdx < Count; ++ActorIdx)
-	{
-		if (FBodyInstance* BodyInstance = FPhysxUserData::Get<FBodyInstance>(Actors[ActorIdx]->userData))
-		{
-			OwningScene->AddPendingSleepingEvent(BodyInstance, ESleepEvent::SET_Wakeup);
-		}
-	}
-}
-
-void FPhysXSimEventCallback::onSleep(PxActor** Actors, PxU32 Count)
-{
-	for (PxU32 ActorIdx = 0; ActorIdx < Count; ++ActorIdx)
-	{
-		if (FBodyInstance* BodyInstance = FPhysxUserData::Get<FBodyInstance>(Actors[ActorIdx]->userData))
-		{
-			OwningScene->AddPendingSleepingEvent(BodyInstance, ESleepEvent::SET_Sleep);
-		}
-	}
-}
-#endif
-
-//////////////////////////////////////////////////////////////////////////
-// FPhysXCookingDataReader
-
-FPhysXCookingDataReader::FPhysXCookingDataReader( FByteBulkData& InBulkData, FBodySetupUVInfo* UVInfo )
-{
-	// Read cooked physics data
-	uint8* DataPtr = (uint8*)InBulkData.Lock( LOCK_READ_ONLY );
-	FBufferReader Ar( DataPtr, InBulkData.GetBulkDataSize(), false );
-	
-	uint8 bLittleEndian = true;
-	int32 NumConvexElementsCooked = 0;
-	int32 NumMirroredElementsCooked = 0;
-	int32 NumTriMeshesCooked = 0;
-
-	Ar << bLittleEndian;
-	Ar.SetByteSwapping( PLATFORM_LITTLE_ENDIAN ? !bLittleEndian : !!bLittleEndian );
-	Ar << NumConvexElementsCooked;	
-	Ar << NumMirroredElementsCooked;
-	Ar << NumTriMeshesCooked;
-	
-	ConvexMeshes.Empty(NumConvexElementsCooked);
-	for( int32 ElementIndex = 0; ElementIndex < NumConvexElementsCooked; ElementIndex++ )
-	{
-		//NOTE: these tags currently get ignored because of low level allocator
-		PxConvexMesh* ConvexMesh;
-		{
-			LLM_SCOPE(ELLMTag::PhysXConvex);
-			ConvexMesh = ReadConvexMesh(Ar,DataPtr,InBulkData.GetBulkDataSize());
-		}
-		
-		ConvexMeshes.Add( ConvexMesh );
-	}
-
-	ConvexMeshesNegX.Empty(NumMirroredElementsCooked);
-	for( int32 ElementIndex = 0; ElementIndex < NumMirroredElementsCooked; ElementIndex++ )
-	{
-		PxConvexMesh* ConvexMeshNegX;
-		{
-			LLM_SCOPE(ELLMTag::PhysXConvex);
-			ConvexMeshNegX = ReadConvexMesh( Ar, DataPtr, InBulkData.GetBulkDataSize() );
-		}
-		
-		ConvexMeshesNegX.Add( ConvexMeshNegX );
-	}
-
-	TriMeshes.Empty(NumTriMeshesCooked);
-	for(int32 ElementIndex = 0; ElementIndex < NumTriMeshesCooked; ++ElementIndex)
-	{
-		PxTriangleMesh* TriMesh;
-		{
-			LLM_SCOPE(ELLMTag::PhysXTrimesh);
-			TriMesh = ReadTriMesh( Ar, DataPtr, InBulkData.GetBulkDataSize() );
-		}
-		
-		TriMeshes.Add(TriMesh);
-	}
-
-	// Init UVInfo pointer
-	check(UVInfo);
-	{
-		LLM_SCOPE(ELLMTag::PhysXTrimesh);
-		Ar << *UVInfo;
-	}
-	
-
-	InBulkData.Unlock();
-}
-
-PxConvexMesh* FPhysXCookingDataReader::ReadConvexMesh( FBufferReader& Ar, uint8* InBulkDataPtr, int32 InBulkDataSize )
-{
-	LLM_SCOPE(ELLMTag::PhysX);
-
-	PxConvexMesh* CookedMesh = NULL;
-	uint8 IsMeshCooked = false;
-	Ar << IsMeshCooked;
-	if( IsMeshCooked )
-	{
-		FPhysXInputStream Buffer( InBulkDataPtr + Ar.Tell(), InBulkDataSize - Ar.Tell() );		
-		CookedMesh = GPhysXSDK->createConvexMesh( Buffer );
-		check( CookedMesh != NULL );
-		Ar.Seek( Ar.Tell() + Buffer.ReadPos );
-	}
-	return CookedMesh;
-}
-
-PxTriangleMesh* FPhysXCookingDataReader::ReadTriMesh( FBufferReader& Ar, uint8* InBulkDataPtr, int32 InBulkDataSize )
-{
-	LLM_SCOPE(ELLMTag::PhysX);
-
-	FPhysXInputStream Buffer( InBulkDataPtr + Ar.Tell(), InBulkDataSize - Ar.Tell() );
-	PxTriangleMesh* CookedMesh = GPhysXSDK->createTriangleMesh(Buffer);
-	check(CookedMesh);
-	Ar.Seek( Ar.Tell() + Buffer.ReadPos );
-	return CookedMesh;
-}
-
-
-void AddToCollection(PxCollection* PCollection, PxBase* PBase)
-{
-	if (PBase)
-	{
-		PCollection->add(*PBase);
-	}
-}
-
-PxCollection* MakePhysXCollection(const TArray<UPhysicalMaterial*>& PhysicalMaterials, const TArray<UBodySetup*>& BodySetups, uint64 BaseId)
-{
-#if WITH_CHAOS || WITH_IMMEDIATE_PHYSX
-    ensure(false);
-    return nullptr;
-#else
-	QUICK_SCOPE_CYCLE_COUNTER(STAT_CreateSharedData);
-	PxCollection* PCollection = PxCreateCollection();
-	for (UPhysicalMaterial* PhysicalMaterial : PhysicalMaterials)
-	{
-		if (PhysicalMaterial)
-		{
-			PCollection->add(*PhysicalMaterial->GetPhysicsMaterial().Material);
-		}
-	}
-
-	for (UBodySetup* BodySetup : BodySetups)
-	{
-		for(PxTriangleMesh* TriMesh : BodySetup->TriMeshes)
-		{
-			AddToCollection(PCollection, TriMesh);
-		}
-
-		for (const FKConvexElem& ConvexElem : BodySetup->AggGeom.ConvexElems)
-		{
-			AddToCollection(PCollection, ConvexElem.GetConvexMesh());
-			AddToCollection(PCollection, ConvexElem.GetMirroredConvexMesh());
-		}
-	}
-
-	PxSerialization::createSerialObjectIds(*PCollection, PxSerialObjectId(BaseId));
-
-	return PCollection;
-#endif
-}
-
-void* FPhysXProfilerCallback::zoneStart(const char* eventName, bool detached, uint64_t contextId)
-{
-	if(GCycleStatsShouldEmitNamedEvents)
-	{
-		FPlatformMisc::BeginNamedEvent(FColor::Red, *FString::Printf(TEXT("PHYSX: %s"), StringCast<TCHAR>(eventName).Get()));
-	}
-
-	return nullptr;
-}
-
-void FPhysXProfilerCallback::zoneEnd(void* profilerData, const char* eventName, bool detached, uint64_t contextId)
-{
-	if(GCycleStatsShouldEmitNamedEvents)
-	{
-		FPlatformMisc::EndNamedEvent();
-	}
-}
-
-void FPhysXMbpBroadphaseCallback::onObjectOutOfBounds(PxShape& InShape, PxActor& InActor)
-{
-	FBodyInstance* ActorBodyInstance = FPhysxUserData::Get<FBodyInstance>(InActor.userData);
-	if(ActorBodyInstance)
-	{
-		UPrimitiveComponent* OwnerComponent = ActorBodyInstance->OwnerComponent.Get();
-
-		if(OwnerComponent)
-		{
-			UE_LOG(LogPhysics, Warning, TEXT("Component %s at location %s has physics bodies outside of MBP bounds. Check MBP bounds are correct for this world, collisions are disabled for bodies outside of MBP bounds when MBP is enabled."), *OwnerComponent->GetName(), *OwnerComponent->GetComponentLocation().ToString());
-			return;
-		}
-	}
-
-	UE_LOG(LogPhysics, Warning, TEXT("Unknown component has physics bodies outside of MBP bounds. Check MBP bounds are correct for this world, collisions are disabled for bodies outside of MBP bounds when MBP is enabled."));
-}
-
-void FPhysXMbpBroadphaseCallback::onObjectOutOfBounds(PxAggregate& InAggregate)
-{
-	const PxU32 NumActors = InAggregate.getNbActors();
-	if(NumActors > 0)
-	{
-		// The following code assumes that an aggregate does not span multiple components, this code will need to be updated if this changes
-		PxActor* FirstActor;
-
-		if(InAggregate.getActors(&FirstActor, 1) > 0)
-		{
-			FBodyInstance* ActorBodyInstance = FPhysxUserData::Get<FBodyInstance>(FirstActor->userData);
-			if(ActorBodyInstance)
-			{
-				UPrimitiveComponent* OwnerComponent = ActorBodyInstance->OwnerComponent.Get();
-
-				if(OwnerComponent)
-				{
-					UE_LOG(LogPhysics, Warning, TEXT("Component %s at location %s has physics bodies outside of MBP bounds. Check MBP bounds are correct for this world, collisions are disabled for bodies outside of MBP bounds when MBP is enabled."), *OwnerComponent->GetName(), *OwnerComponent->GetComponentLocation().ToString());
-					return;
-				}
-			}
-		}
-	}
-
-	UE_LOG(LogPhysics, Warning, TEXT("Unknown component has physics bodies outside of MBP bounds. Check MBP bounds are correct for this world, collisions are disabled for bodies outside of MBP bounds when MBP is enabled."));
-}
-
-#endif // WITH_PHYSX
diff --git a/Engine/Source/Runtime/Engine/Private/PhysicsEngine/PhysXSupport.h b/Engine/Source/Runtime/Engine/Private/PhysicsEngine/PhysXSupport.h
--- a/Engine/Source/Runtime/Engine/Private/PhysicsEngine/PhysXSupport.h
+++ b/Engine/Source/Runtime/Engine/Private/PhysicsEngine/PhysXSupport.h
@@ -19,152 +19,3 @@
 class UBodySetup;
 class UPhysicalMaterial;
 struct FCollisionShape;
-class FPhysScene_PhysX;
-
-#if PHYSICS_INTERFACE_PHYSX
-
-// binary serialization requires 128 byte alignment
-#ifndef PHYSX_SERIALIZATION_ALIGNMENT
-#define PHYSX_SERIALIZATION_ALIGNMENT 128
-#endif
-
-/** Thresholds for aggregates  */
-const uint32 AggregateMaxSize	   = 128;
-const uint32 AggregateBodyShapesThreshold	   = 999999999;
-
-
-/////// UTILS
-
-
-
-
-/** Calculates correct impulse at the body's center of mass and adds the impulse to the body. */
-/** Util to see if a PxRigidBody is non-kinematic */
-bool IsRigidBodyKinematic_AssumesLocked(const PxRigidBody* PRigidBody);
-
-bool IsRigidBodyKinematicAndInSimulationScene_AssumesLocked(const PxRigidBody* PRigidBody);
-
-/////// GLOBAL POINTERS
-	
-/** Total number of PhysX convex meshes around currently. */
-extern int32					GNumPhysXConvexMeshes;
-
-// The following are used for deferred cleanup - object that cannot be destroyed until all uses have been destroyed, but GC guarantees no order.
-
-/** Array of PxConvexMesh objects which are awaiting cleaning up. */
-extern TArray<PxConvexMesh*>	GPhysXPendingKillConvex;
-
-/** Array of PxTriangleMesh objects which are awaiting cleaning up. */
-extern ENGINE_API TArray<PxTriangleMesh*>	GPhysXPendingKillTriMesh;
-
-/** Array of PxHeightField objects which are awaiting cleaning up. */
-extern ENGINE_API TArray<PxHeightField*>	GPhysXPendingKillHeightfield;
-
-
-extern const physx::PxQuat U2PSphylBasis;
-extern const FQuat U2PSphylBasis_UE;
-
-ENGINE_API PxCollection* MakePhysXCollection(const TArray<UPhysicalMaterial*>& PhysicalMaterials, const TArray<UBodySetup*>& BodySetups, uint64 BaseId);
-
-/** Utility wrapper for a uint8 TArray for loading into PhysX. */
-class FPhysXInputStream : public PxInputStream
-{
-public:
-	/** Raw byte data */
-	const uint8	*Data;
-	/** Number of bytes */
-	int32				DataSize;
-	/** Current read position withing Data buffer */
-	mutable uint32			ReadPos;
-
-	FPhysXInputStream()
-		: Data(NULL)
-		, DataSize(0)
-		, ReadPos(0)
-	{}
-
-	FPhysXInputStream(const uint8 *InData, const int32 InSize)
-		: Data(InData)
-		, DataSize( InSize )
-		, ReadPos(0)
-	{}
-
-	virtual PxU32 read(void* Dest, PxU32 Count) override
-	{
-		check(Data);
-		check(Dest);
-		check(Count);
-		const uint32 EndPos = ReadPos + Count;
-		if( EndPos <= (uint32)DataSize )
-		{
-			FMemory::Memcpy( Dest, Data + ReadPos, Count );
-			ReadPos = EndPos;
-			return Count;
-		}
-		else
-		{
-			return 0;
-		}
-	}
-};
-
-/** Utility class for reading cooked physics data. */
-class FPhysXCookingDataReader
-{
-public:
-	TArray<PxConvexMesh*> ConvexMeshes;
-	TArray<PxConvexMesh*> ConvexMeshesNegX;
-	TArray<PxTriangleMesh*> TriMeshes;
-
-	FPhysXCookingDataReader( FByteBulkData& InBulkData, struct FBodySetupUVInfo* UVInfo );
-
-private:
-
-	PxConvexMesh* ReadConvexMesh( FBufferReader& Ar, uint8* InBulkDataPtr, int32 InBulkDataSize );
-	PxTriangleMesh* ReadTriMesh( FBufferReader& Ar, uint8* InBulkDataPtr, int32 InBulkDataSize );
-};
-
-/** 'Shader' used to filter simulation collisions. Could be called on any thread. */
-PxFilterFlags PhysXSimFilterShader(	PxFilterObjectAttributes attributes0, PxFilterData filterData0, 
-									PxFilterObjectAttributes attributes1, PxFilterData filterData1,
-									PxPairFlags& pairFlags, const void* constantBlock, PxU32 constantBlockSize );
-
-#if !WITH_CHAOS
-/** Event callback used to notify engine about various collision events */
-class ENGINE_API FPhysXSimEventCallback : public PxSimulationEventCallback
-{
-public:
-	FPhysXSimEventCallback(FPhysScene* InOwningScene) : OwningScene(InOwningScene){}
-
-	virtual void onConstraintBreak(PxConstraintInfo* constraints, PxU32 count) override;
-	virtual void onWake(PxActor** actors, PxU32 count) override;
-	virtual void onSleep(PxActor** actors, PxU32 count) override;
-	virtual void onTrigger(PxTriggerPair* pairs, PxU32 count) override {}
-	virtual void onContact(const PxContactPairHeader& PairHeader, const PxContactPair* Pairs, PxU32 NumPairs) override;
-	virtual void onAdvance(const PxRigidBody*const* bodyBuffer, const PxTransform* poseBuffer, const PxU32 count) override {}
-
-private:	
-	FPhysScene* OwningScene;
-};
-#endif
-
-class FPhysXProfilerCallback : public PxProfilerCallback
-{
-
-public:
-	virtual void* zoneStart(const char* eventName, bool detached, uint64_t contextId) override;
-	virtual void zoneEnd(void* profilerData, const char* eventName, bool detached, uint64_t contextId) override;
-
-private:
-};
-
-struct FPhysXMbpBroadphaseCallback : public PxBroadPhaseCallback
-{
-
-public:
-	virtual void onObjectOutOfBounds(PxShape& InShape, PxActor& InActor) override;
-	virtual void onObjectOutOfBounds(PxAggregate& InAggregate) override;
-
-};
-
-#endif // WITH_PHYSX
diff --git a/Engine/Source/Runtime/Engine/Private/PhysicsEngine/PhysicalAnimationComponent.cpp b/Engine/Source/Runtime/Engine/Private/PhysicsEngine/PhysicalAnimationComponent.cpp
--- a/Engine/Source/Runtime/Engine/Private/PhysicsEngine/PhysicalAnimationComponent.cpp
+++ b/Engine/Source/Runtime/Engine/Private/PhysicsEngine/PhysicalAnimationComponent.cpp
@@ -179,33 +179,17 @@ FTransform UPhysicalAnimationComponent::GetBodyTargetTransform(FName BodyName) c
 {
 	if (SkeletalMeshComponent)
 	{
-
 		for (int32 DataIdx = 0; DataIdx < DriveData.Num(); ++DataIdx)
 		{
 			const FPhysicalAnimationData& PhysAnimData = DriveData[DataIdx];
 			const FPhysicalAnimationInstanceData& InstanceData = RuntimeInstanceData[DataIdx];
 			if (BodyName == PhysAnimData.BodyName)
 			{
-#if PHYSICS_INTERFACE_PHYSX
-				if (PxRigidDynamic* TargetActor = InstanceData.TargetActor)
-				{
-					PxTransform PKinematicTarget;
-					if (TargetActor->getKinematicTarget(PKinematicTarget))
-					{
-						return P2UTransform(PKinematicTarget);
-					}
-					else
-					{
-						return P2UTransform(TargetActor->getGlobalPose());
-					}
-				}
-#elif WITH_CHAOS
 				if (FPhysicsActorHandle TargetActor = InstanceData.TargetActor)
 				{
 					// TODO: If kinematic targets implemented, fetch target and don't use position.
 					return FTransform(TargetActor->GetGameThreadAPI().R(), TargetActor->GetGameThreadAPI().X());
 				}
-#endif
 
 				break;
 			}
@@ -280,7 +264,6 @@ void UPhysicalAnimationComponent::UpdateTargetActors(ETeleportType TeleportType)
 		// various anim ticks, before buffers are flipped (which happens in the skel mesh component's post-physics tick)
 		const TArray<FTransform>& SpaceBases = SkeletalMeshComponent->GetEditableComponentSpaceTransforms();
 
-
 		FPhysicsCommand::ExecuteWrite(SkeletalMeshComponent, [&]()
 		{
 			TArray<FTransform> LocalTransforms = SkeletalMeshComponent->GetBoneSpaceTransforms();
@@ -288,23 +271,7 @@ void UPhysicalAnimationComponent::UpdateTargetActors(ETeleportType TeleportType)
 			{
 				const FPhysicalAnimationData& PhysAnimData = DriveData[DataIdx];
 				FPhysicalAnimationInstanceData& InstanceData = RuntimeInstanceData[DataIdx];
-#if PHYSICS_INTERFACE_PHYSX
-				if (PxRigidDynamic* TargetActor = InstanceData.TargetActor)
-				{
-					const int32 BoneIdx = RefSkeleton.FindBoneIndex(PhysAnimData.BodyName);
-					if (BoneIdx != INDEX_NONE)	//It's possible the skeletal mesh has changed out from under us. In that case we should probably reset, but at the very least don't do work on non-existent bones
-					{
-						const FTransform TargetTM = ComputeTargetTM(PhysAnimData, *SkeletalMeshComponent, *PhysAsset, LocalTransforms, SpaceBases, BoneIdx);
-						TargetActor->setKinematicTarget(U2PTransform(TargetTM));	//TODO: this doesn't work with sub-stepping!
-						
-						if(TeleportType == ETeleportType::TeleportPhysics)
-						{
-							TargetActor->setGlobalPose(U2PTransform(TargetTM));	//Note that we still set the kinematic target because physx doesn't clear this
-						}
-						
-					}
-				}
-#else
+
 				if (FPhysicsActorHandle TargetActor = InstanceData.TargetActor)
 				{
 					const int32 BoneIdx = RefSkeleton.FindBoneIndex(PhysAnimData.BodyName);
@@ -320,7 +287,6 @@ void UPhysicalAnimationComponent::UpdateTargetActors(ETeleportType TeleportType)
 						}
 					}
 				}
-#endif
 			}
 		});
 	}
@@ -352,13 +318,11 @@ void SetMotorStrength(FConstraintInstance& ConstraintInstance, const FPhysicalAn
 	float VelocityStrengthMultiplier = StrengthMultiplier;
 	float OrientationStrengthMultiplier = StrengthMultiplier;
 	float AngularVelocityStrengthMultiplier = StrengthMultiplier;
-#if WITH_CHAOS
-	// Chaos has it's own global adjustments
+
 	PositionStrengthMultiplier *= Chaos::ConstraintSettings::LinearDriveStiffnessScale();
 	VelocityStrengthMultiplier *= Chaos::ConstraintSettings::LinearDriveDampingScale();
 	OrientationStrengthMultiplier *= Chaos::ConstraintSettings::AngularDriveStiffnessScale();
 	AngularVelocityStrengthMultiplier *= Chaos::ConstraintSettings::AngularDriveDampingScale();
-#endif
 
 	ConstraintInstance.SetAngularDriveParams(
 		PhysAnimData.OrientationStrength * OrientationStrengthMultiplier, 
@@ -406,7 +370,6 @@ void UPhysicalAnimationComponent::UpdatePhysicsEngineImp()
 		const TArray<FTransform>& SpaceBases = SkeletalMeshComponent->GetEditableComponentSpaceTransforms();
 		const FReferenceSkeleton& RefSkeleton = SkeletalMeshComponent->SkeletalMesh->GetRefSkeleton();
 
-#if WITH_PHYSX
 		FPhysicsCommand::ExecuteWrite(SkeletalMeshComponent, [&]()
 		{
 			TArray<FTransform> LocalTransforms = SkeletalMeshComponent->GetBoneSpaceTransforms();
@@ -433,7 +396,6 @@ void UPhysicalAnimationComponent::UpdatePhysicsEngineImp()
 					int32 ChildBodyIdx = PhysAsset->FindBodyIndex(PhysAnimData.BodyName);
 					if (FBodyInstance* ChildBody = (ChildBodyIdx == INDEX_NONE ? nullptr : SkeletalMeshComponent->Bodies[ChildBodyIdx]))
 					{
-#if WITH_CHAOS || WITH_IMMEDIATE_PHYSX
 						if (FPhysicsActorHandle ActorHandle = ChildBody->ActorHandle)
 						{
 							FPhysScene* Scene = ChildBody->GetPhysicsScene();
@@ -466,43 +428,11 @@ void UPhysicalAnimationComponent::UpdatePhysicsEngineImp()
 							InstanceData.TargetActor = KineActor;
 
 							ConstraintInstance->InitConstraint_AssumesLocked(ChildBody->ActorHandle, InstanceData.TargetActor, 1.f);
-					}
-
-#elif PHYSICS_INTERFACE_PHYSX
-						if (PxRigidActor* PRigidActor = FPhysicsInterface_PhysX::GetPxRigidActor_AssumesLocked(ChildBody->ActorHandle))
-						{
-							ConstraintInstance->SetRefFrame(EConstraintFrame::Frame1, FTransform::Identity);
-							ConstraintInstance->SetRefFrame(EConstraintFrame::Frame2, FTransform::Identity);
-
-							const FTransform TargetTM = ComputeTargetTM(PhysAnimData, *SkeletalMeshComponent, *PhysAsset, LocalTransforms, SpaceBases, ChildBody->InstanceBoneIndex);
-
-							// Create kinematic actor we are going to create joint with. This will be moved around with calls to SetLocation/SetRotation.
-							PxScene* PScene = PRigidActor->getScene();
-							PxRigidDynamic* KineActor = PScene->getPhysics().createRigidDynamic(U2PTransform(TargetTM));
-							KineActor->setRigidBodyFlag(PxRigidBodyFlag::eKINEMATIC, true);
-							KineActor->setMass(1.0f);
-							KineActor->setMassSpaceInertiaTensor(PxVec3(1.0f, 1.0f, 1.0f));
-
-							// No bodyinstance
-							KineActor->userData = NULL;
-
-							// Add to Scene
-							PScene->addActor(*KineActor);
-
-							// Save reference to the kinematic actor.
-							InstanceData.TargetActor = KineActor;
-
-							FPhysicsActorHandle TargetRef;
-							TargetRef.SyncActor = InstanceData.TargetActor;
-							
-							ConstraintInstance->InitConstraint_AssumesLocked(ChildBody->ActorHandle, TargetRef, 1.f);
 						}
-#endif
 					}
 				}
 			}
 		});
-#endif
 	}
 }
 
@@ -549,23 +479,12 @@ void UPhysicalAnimationComponent::ReleasePhysicsEngine()
 
 			if(Instance.TargetActor)
 			{
-#if PHYSICS_INTERFACE_PHYSX
-				PxScene* PScene = Instance.TargetActor->getScene();
-				if(PScene)
-				{
-					SCOPED_SCENE_WRITE_LOCK(PScene);
-					PScene->removeActor(*Instance.TargetActor);
-				}
-				Instance.TargetActor->release();
-				Instance.TargetActor = nullptr;
-#else
 				FChaosScene* PhysScene = FChaosEngineInterface::GetCurrentScene(Instance.TargetActor);
 				if (ensure(PhysScene))
 				{
 					FPhysInterface_Chaos::ReleaseActor(Instance.TargetActor, PhysScene);
 				}
 				Instance.TargetActor = nullptr;
-#endif
 			}
 		}
 
@@ -582,11 +501,7 @@ void UPhysicalAnimationComponent::DebugDraw(FPrimitiveDrawInterface* PDI) const
 	{
 		if (PhysAnimData.TargetActor)
 		{
-#if PHYSICS_INTERFACE_PHYSX
-			PDI->DrawPoint(P2UVector(PhysAnimData.TargetActor->getGlobalPose().p), TargetActorColor, 3.f, SDPG_World);
-#elif WITH_CHAOS
 			PDI->DrawPoint(PhysAnimData.TargetActor->GetGameThreadAPI().X(), TargetActorColor, 3.f, SDPG_World);
-#endif
 		}
 	}
 }
diff --git a/Engine/Source/Runtime/Engine/Private/PhysicsEngine/PhysicalMaterialMask.cpp b/Engine/Source/Runtime/Engine/Private/PhysicsEngine/PhysicalMaterialMask.cpp
--- a/Engine/Source/Runtime/Engine/Private/PhysicsEngine/PhysicalMaterialMask.cpp
+++ b/Engine/Source/Runtime/Engine/Private/PhysicsEngine/PhysicalMaterialMask.cpp
@@ -29,10 +29,7 @@
 
 #endif
 
-#if WITH_CHAOS
 #include "Chaos/PhysicalMaterials.h"
-#endif
-
 
 DEFINE_LOG_CATEGORY_STATIC(LogPhysicalMaterialMask, Log, All);
 
@@ -109,12 +106,10 @@ UPhysicalMaterialMask::~UPhysicalMaterialMask() = default;
 
 void UPhysicalMaterialMask::PostEditChangeProperty(FPropertyChangedEvent& PropertyChangedEvent)
 {
-#if WITH_CHAOS
 	if (MaterialMaskHandle && MaterialMaskHandle->IsValid())
 	{
 		FPhysicsInterface::UpdateMaterialMask(*MaterialMaskHandle, this);
 	}
-#endif
 
 	Super::PostEditChangeProperty(PropertyChangedEvent);
 }
@@ -147,12 +142,11 @@ void UPhysicalMaterialMask::PostLoad()
 
 void UPhysicalMaterialMask::FinishDestroy()
 {
-#if WITH_CHAOS
 	if(MaterialMaskHandle)
 	{
 		FPhysicsInterface::ReleaseMaterialMask(*MaterialMaskHandle);
 	}
-#endif
+
 	Super::FinishDestroy();
 }
 
@@ -239,8 +233,6 @@ void UPhysicalMaterialMask::DumpMaskData()
 
 #endif // WITH_EDITOR
 
-#if	WITH_CHAOS
-
 FPhysicsMaterialMaskHandle& UPhysicalMaterialMask::GetPhysicsMaterialMask()
 {
 	if(!MaterialMaskHandle)
@@ -260,8 +252,6 @@ FPhysicsMaterialMaskHandle& UPhysicalMaterialMask::GetPhysicsMaterialMask()
 	return *MaterialMaskHandle;
 }
 
-#endif
-
 // This template generates mask data from the texture mask, converting colors to mask ids.
 template<typename PixelDataType, int32 RIdx, int32 GIdx, int32 BIdx, int32 AIdx> class MaskDataGenerator
 {
diff --git a/Engine/Source/Runtime/Engine/Private/PhysicsEngine/PhysicsGeometry.cpp b/Engine/Source/Runtime/Engine/Private/PhysicsEngine/PhysicsGeometry.cpp
--- a/Engine/Source/Runtime/Engine/Private/PhysicsEngine/PhysicsGeometry.cpp
+++ b/Engine/Source/Runtime/Engine/Private/PhysicsEngine/PhysicsGeometry.cpp
@@ -1,307 +1 @@
 // Copyright Epic Games, Inc. All Rights Reserved.
-
-#if WITH_PHYSX && PHYSICS_INTERFACE_PHYSX
-
-#include "Physics/PhysicsGeometry.h"
-
-#include "PhysicsEngine/AggregateGeom.h"
-
-//#include "PhysXPublic.h"
-#include "PhysicsEngine/PhysXSupport.h"
-#include "PhysicsEngine/PhysicsSettings.h"
-#include "Engine/Engine.h"
-#include "PhysicalMaterials/PhysicalMaterial.h"
-#include "PhysicsEngine/BodySetup.h"
-#include "PxSphereGeometry.h"
-
-using namespace physx;
-
-extern TAutoConsoleVariable<float> CVarContactOffsetFactor;
-extern TAutoConsoleVariable<float> CVarMaxContactOffset;
-
-void FBodySetupShapeIterator::GetContactOffsetParams(float& InOutContactOffsetFactor, float& InOutMinContactOffset, float& InOutMaxContactOffset)
-{
-	// Get contact offset params
-	InOutContactOffsetFactor = CVarContactOffsetFactor.GetValueOnAnyThread();
-	InOutMaxContactOffset = CVarMaxContactOffset.GetValueOnAnyThread();
-
-	InOutContactOffsetFactor = InOutContactOffsetFactor < 0.f ? UPhysicsSettings::Get()->ContactOffsetMultiplier : InOutContactOffsetFactor;
-	InOutMaxContactOffset = InOutMaxContactOffset < 0.f ? UPhysicsSettings::Get()->MaxContactOffset : InOutMaxContactOffset;
-
-	InOutMinContactOffset = UPhysicsSettings::Get()->MinContactOffset;
-}
-
-FBodySetupShapeIterator::FBodySetupShapeIterator(const FVector& InScale3D, const FTransform& InRelativeTM, bool bDoubleSidedTrimeshes)
-	: Scale3D(InScale3D)
-	, RelativeTM(InRelativeTM)
-	, bDoubleSidedTriMeshGeo(bDoubleSidedTrimeshes)
-{
-	SetupNonUniformHelper(Scale3D, MinScale, MinScaleAbs, ShapeScale3DAbs);
-	{
-		float MinScaleRelative;
-		float MinScaleAbsRelative;
-		FVector Scale3DAbsRelative;
-		FVector Scale3DRelative = RelativeTM.GetScale3D();
-
-		SetupNonUniformHelper(Scale3DRelative, MinScaleRelative, MinScaleAbsRelative, Scale3DAbsRelative);
-
-		MinScaleAbs *= MinScaleAbsRelative;
-		ShapeScale3DAbs.X *= Scale3DAbsRelative.X;
-		ShapeScale3DAbs.Y *= Scale3DAbsRelative.Y;
-		ShapeScale3DAbs.Z *= Scale3DAbsRelative.Z;
-
-		ShapeScale3D = Scale3D;
-		ShapeScale3D.X *= Scale3DAbsRelative.X;
-		ShapeScale3D.Y *= Scale3DAbsRelative.Y;
-		ShapeScale3D.Z *= Scale3DAbsRelative.Z;
-	}
-
-	GetContactOffsetParams(ContactOffsetFactor, MinContactOffset, MaxContactOffset);
-}
-
-FBodySetupShapeIterator::FBodySetupShapeIterator()
-	: Scale3D(FVector::OneVector)
-	, RelativeTM(FTransform::Identity)
-	, bDoubleSidedTriMeshGeo(false)
-{
-
-}
-
-template <typename ElemType, typename GeomType>
-void FBodySetupShapeIterator::ForEachShape(const TArrayView<ElemType>& Elements, TFunctionRef<void(const ElemType& Elem, const GeomType& Geom, const FBodySetupShapeIterator::TTransform& LocalPose, float ContactOffset, float RestOffset)> VisitorFunc) const
-{
-	for(int32 ElemIdx = 0; ElemIdx < Elements.Num(); ElemIdx++)
-	{
-		const ElemType& Elem = Elements[ElemIdx];
-		GeomType Geom;
-		PxTransform PLocalPose;
-
-		if(PopulateGeometryAndTransform(Elem, Geom, PLocalPose))
-		{
-			const float RestOffset = ComputeRestOffset(Elem);
-			const float ContactOffset = ComputeContactOffset(Geom) + RestOffset;
-			VisitorFunc(Elem, Geom, PLocalPose, ContactOffset, RestOffset);
-		}
-		else
-		{
-			UE_LOG(LogPhysics, Warning, TEXT("ForeachShape(%s): ScaledElem[%d] invalid"), *GetDebugName<ElemType>(), ElemIdx);
-		}
-	}
-}
-
-
-//////////////////////// Sphere elements ////////////////////////
-template <> bool FBodySetupShapeIterator::PopulateGeometryAndTransform(const FKSphereElem& SphereElem, PxSphereGeometry& OutGeometry, FBodySetupShapeIterator::TTransform& OutTM) const
-{
-	const FKSphereElem ScaledSphereElem = SphereElem.GetFinalScaled(Scale3D, RelativeTM);
-	OutGeometry.radius = FMath::Max(ScaledSphereElem.Radius, KINDA_SMALL_NUMBER);
-
-	if(ensure(OutGeometry.isValid()))
-	{
-		OutTM = PxTransform(U2PVector(ScaledSphereElem.Center));
-		return true;
-	}
-	else
-	{
-		return false;
-	}
-}
-
-template <> float FBodySetupShapeIterator::ComputeContactOffset(const PxSphereGeometry& PSphereGeom) const
-{
-	return FMath::Clamp(ContactOffsetFactor * PSphereGeom.radius, MinContactOffset, MaxContactOffset);
-}
-
-template <> FString FBodySetupShapeIterator::GetDebugName<FKSphereElem>()  const
-{
-	return TEXT("Sphere");
-}
-
-template <typename ElemType> float FBodySetupShapeIterator::ComputeRestOffset(const ElemType& Elem) const
-{
-	return Elem.RestOffset;
-}
-
-template <> float FBodySetupShapeIterator::ComputeRestOffset(physx::PxTriangleMesh* const&) const
-{
-	return 0.f;
-}
-
-/////////////////// Box elements //////////////////////////////
-template <> bool FBodySetupShapeIterator::PopulateGeometryAndTransform(const FKBoxElem& BoxElem, PxBoxGeometry& OutGeometry, FBodySetupShapeIterator::TTransform& OutTM) const
-{
-	const FKBoxElem ScaledBoxElem = BoxElem.GetFinalScaled(Scale3D, RelativeTM);
-	const FTransform& BoxTransform = ScaledBoxElem.GetTransform();
-
-	OutGeometry.halfExtents.x = FMath::Max(ScaledBoxElem.X * 0.5f, KINDA_SMALL_NUMBER);
-	OutGeometry.halfExtents.y = FMath::Max(ScaledBoxElem.Y * 0.5f, KINDA_SMALL_NUMBER);
-	OutGeometry.halfExtents.z = FMath::Max(ScaledBoxElem.Z * 0.5f, KINDA_SMALL_NUMBER);
-
-	if(OutGeometry.isValid() && BoxTransform.IsValid())
-	{
-		OutTM = U2PTransform(BoxTransform);
-		if(ensure(OutTM.isValid()))
-		{
-			return true;
-		}
-	}
-
-	return false;
-}
-
-template <> float FBodySetupShapeIterator::ComputeContactOffset(const PxBoxGeometry& PBoxGeom) const
-{
-	return FMath::Clamp(ContactOffsetFactor * PBoxGeom.halfExtents.minElement(), MinContactOffset, MaxContactOffset);
-}
-
-template <> FString FBodySetupShapeIterator::GetDebugName<FKBoxElem>()  const
-{
-	return TEXT("Box");
-}
-
-/////////////////////// Capsule elements /////////////////////////////
-template <> bool FBodySetupShapeIterator::PopulateGeometryAndTransform(const FKSphylElem& SphylElem, PxCapsuleGeometry& OutGeometry, FBodySetupShapeIterator::TTransform& OutTM) const
-{
-	const FKSphylElem ScaledSphylElem = SphylElem.GetFinalScaled(Scale3D, RelativeTM);
-
-	OutGeometry.halfHeight = FMath::Max(ScaledSphylElem.Length * 0.5f, KINDA_SMALL_NUMBER);
-	OutGeometry.radius = FMath::Max(ScaledSphylElem.Radius, KINDA_SMALL_NUMBER);
-
-	if(OutGeometry.isValid())
-	{
-		// The stored capsule transform assumes the capsule axis is down Z. In PhysX, it points down X, so we twiddle the matrix a bit here (swap X and Z and negate Y).
-		OutTM = PxTransform(U2PVector(ScaledSphylElem.Center), U2PQuat(ScaledSphylElem.Rotation.Quaternion()) * U2PSphylBasis);
-
-		if(ensure(OutTM.isValid()))
-		{
-			return true;
-		}
-	}
-
-	return false;
-}
-
-template <> float FBodySetupShapeIterator::ComputeContactOffset(const PxCapsuleGeometry& PCapsuleGeom) const
-{
-	return FMath::Clamp(ContactOffsetFactor * PCapsuleGeom.radius, MinContactOffset, MaxContactOffset);
-}
-
-template <> FString FBodySetupShapeIterator::GetDebugName<FKSphylElem>() const
-{
-	return TEXT("Capsule");
-}
-
-////////////////////////////// Convex elements ////////////////////////////
-template <> bool FBodySetupShapeIterator::PopulateGeometryAndTransform(const FKConvexElem& ConvexElem, PxConvexMeshGeometry& OutGeometry, FBodySetupShapeIterator::TTransform& OutTM) const
-{
-	FTransform NegativeScaleCompensation;
-	const bool bUseNegX = CalcMeshNegScaleCompensation(Scale3D * RelativeTM.GetScale3D(), NegativeScaleCompensation);
-	OutTM = U2PTransform(NegativeScaleCompensation);
-
-	PxConvexMesh* UseConvexMesh = bUseNegX ? ConvexElem.GetMirroredConvexMesh() : ConvexElem.GetConvexMesh();
-	if(UseConvexMesh)
-	{
-		OutGeometry.convexMesh = UseConvexMesh;
-		OutGeometry.scale.scale = U2PVector(ShapeScale3DAbs);	//scale shape about the origin
-
-																													//Scale the position independent of shape scale. This is because physx transforms have no concept of scale
-		PxTransform PElementTransform = U2PTransform(RelativeTM);
-		OutTM.q *= PElementTransform.q;
-		OutTM.p = PElementTransform.p;
-		OutTM.p.x *= Scale3D.X;
-		OutTM.p.y *= Scale3D.Y;
-		OutTM.p.z *= Scale3D.Z;
-
-		if(OutGeometry.isValid())
-		{
-			PxVec3 PBoundsExtents = OutGeometry.convexMesh->getLocalBounds().getExtents();
-
-			if(ensure(OutTM.isValid()))
-			{
-				return true;
-			}
-			else
-			{
-				UE_LOG(LogPhysics, Warning, TEXT("PopulatePhysXGeometryAndTransform(Convex): ConvexElem invalid"));
-			}
-		}
-		else
-		{
-			UE_LOG(LogPhysics, Warning, TEXT("PopulatePhysXGeometryAndTransform(Convex): ConvexElem has invalid transform"));
-		}
-	}
-	else
-	{
-		UE_LOG(LogPhysics, Warning, TEXT("PopulatePhysXGeometryAndTransform(Convex): ConvexElem is missing ConvexMesh"));
-	}
-
-	return false;
-}
-
-template <> float FBodySetupShapeIterator::ComputeContactOffset(const PxConvexMeshGeometry& PConvexGeom) const
-{
-	PxVec3 PBoundsExtents = PConvexGeom.convexMesh->getLocalBounds().getExtents();
-	return FMath::Clamp(ContactOffsetFactor * PBoundsExtents.minElement(), MinContactOffset, MaxContactOffset);
-}
-
-template <> FString FBodySetupShapeIterator::GetDebugName<FKConvexElem>() const
-{
-	return TEXT("Convex");
-}
-
-
-///////////////////////////////////Trimesh elements ////////////////////////////
-template <> bool FBodySetupShapeIterator::PopulateGeometryAndTransform(PxTriangleMesh* const & TriMesh, PxTriangleMeshGeometry& OutGeometry, FBodySetupShapeIterator::TTransform& OutTM) const
-{
-	OutGeometry.triangleMesh = TriMesh;
-	OutGeometry.scale.scale = U2PVector(ShapeScale3D); //scale shape about the origin
-
-	auto ClampScale = [](float& Val) -> void
-	{
-		Val = Val <= 0.f ? FMath::Min(Val, -KINDA_SMALL_NUMBER) : FMath::Max(KINDA_SMALL_NUMBER, Val);
-	};
-
-	ClampScale(OutGeometry.scale.scale.x);
-	ClampScale(OutGeometry.scale.scale.y);
-	ClampScale(OutGeometry.scale.scale.z);
-
-	if(bDoubleSidedTriMeshGeo)
-	{
-		OutGeometry.meshFlags |= PxMeshGeometryFlag::eDOUBLE_SIDED;
-	}
-
-	if(OutGeometry.isValid())
-	{
-		//Scale the position independent of shape scale. This is because physx transforms have no concept of scale
-		OutTM = U2PTransform(RelativeTM);
-		OutTM.p.x *= Scale3D.X;
-		OutTM.p.y *= Scale3D.Y;
-		OutTM.p.z *= Scale3D.Z;
-
-		return true;
-	}
-	else
-	{
-		UE_LOG(LogPhysics, Log, TEXT("PopulatePhysXGeometryAndTransform(TriMesh): TriMesh invalid"));
-	}
-
-	return false;
-}
-
-template <> float FBodySetupShapeIterator::ComputeContactOffset(const PxTriangleMeshGeometry& PTriMeshGeom) const
-{
-	return MaxContactOffset;
-}
-
-template <> FString FBodySetupShapeIterator::GetDebugName<PxTriangleMesh*>() const
-{
-	return TEXT("Trimesh");
-}
-
-template void FBodySetupShapeIterator::ForEachShape(const TArrayView<FKSphereElem>&, TFunctionRef<void(const FKSphereElem&, const physx::PxSphereGeometry&, const physx::PxTransform&, float, float)>) const;
-template void FBodySetupShapeIterator::ForEachShape(const TArrayView<FKBoxElem>&, TFunctionRef<void(const FKBoxElem&, const physx::PxBoxGeometry&, const physx::PxTransform&, float, float)>) const;
-template void FBodySetupShapeIterator::ForEachShape(const TArrayView<FKSphylElem>&, TFunctionRef<void(const FKSphylElem&, const physx::PxCapsuleGeometry&, const physx::PxTransform&, float, float)>) const;
-template void FBodySetupShapeIterator::ForEachShape(const TArrayView<FKConvexElem>&, TFunctionRef<void(const FKConvexElem&, const physx::PxConvexMeshGeometry&, const physx::PxTransform&, float, float)>) const;
-template void FBodySetupShapeIterator::ForEachShape(const TArrayView<physx::PxTriangleMesh*>&, TFunctionRef<void(physx::PxTriangleMesh* const &, const physx::PxTriangleMeshGeometry&, const physx::PxTransform&, float, float)>) const;
-
-#endif
\ No newline at end of file
diff --git a/Engine/Source/Runtime/Engine/Private/PhysicsEngine/PhysicsHandleComponent.cpp b/Engine/Source/Runtime/Engine/Private/PhysicsEngine/PhysicsHandleComponent.cpp
--- a/Engine/Source/Runtime/Engine/Private/PhysicsEngine/PhysicsHandleComponent.cpp
+++ b/Engine/Source/Runtime/Engine/Private/PhysicsEngine/PhysicsHandleComponent.cpp
@@ -7,29 +7,22 @@
 #include "PhysicsPublic.h"
 #include "Physics/PhysicsInterfaceCore.h"
 
-#if PHYSICS_INTERFACE_PHYSX
-	#include "PhysXPublic.h"
-#elif WITH_CHAOS
-	#include "Chaos/ParticleHandle.h"
-	#include "Chaos/KinematicGeometryParticles.h"
-	#include "Chaos/PBDJointConstraintTypes.h"
-	#include "Chaos/PBDJointConstraintData.h"
-	#include "Chaos/Sphere.h"
-	#include "PhysicsProxy/SingleParticlePhysicsProxy.h"
-#endif
-
+#include "Chaos/ParticleHandle.h"
+#include "Chaos/KinematicGeometryParticles.h"
+#include "Chaos/PBDJointConstraintTypes.h"
+#include "Chaos/PBDJointConstraintData.h"
+#include "Chaos/Sphere.h"
+#include "PhysicsProxy/SingleParticlePhysicsProxy.h"
 #include "ChaosCheck.h"
 
 UPhysicsHandleComponent::UPhysicsHandleComponent(const FObjectInitializer& ObjectInitializer)
 	: Super(ObjectInitializer)
-#if WITH_CHAOS
 	, bPendingConstraint(false)
 	, PhysicsUserData(&ConstraintInstance)
 	, GrabbedHandle(nullptr)
 	, KinematicHandle(nullptr)
 	, ConstraintLocalPosition(FVector::ZeroVector)
 	, ConstraintLocalRotation(FRotator::ZeroRotator)
-#endif
 {
 	bAutoActivate = true;
 	PrimaryComponentTick.bCanEverTick = true;
@@ -51,30 +44,6 @@ void UPhysicsHandleComponent::OnUnregister()
 		ReleaseComponent();
 	}
 
-#if PHYSICS_INTERFACE_PHYSX
-	if(HandleData)
-	{
-		check(KinActorData);
-
-		// use correct scene
-		PxScene* PScene = HandleData->getScene();
-		if(PScene)
-		{
-			PScene->lockWrite();
-
-			// destroy joint
-			HandleData->release();
-			HandleData = NULL;
-
-			// Destroy temporary actor.
-			KinActorData->release();
-			KinActorData = NULL;
-
-			PScene->unlockWrite();
-		}
-	}
-#endif // WITH_PHYSX
-
 	Super::OnUnregister();
 }
 
@@ -137,80 +106,6 @@ void UPhysicsHandleComponent::GrabComponentImp(UPrimitiveComponent* InComponent,
 		return;
 	}
 
-#if WITH_PHYSX && PHYSICS_INTERFACE_PHYSX
-
-	FPhysicsCommand::ExecuteWrite(BodyInstance->ActorHandle, [&](const FPhysicsActorHandle& Actor)
-	{
-		if(PxRigidActor* PActor = FPhysicsInterface::GetPxRigidActor_AssumesLocked(Actor))
-		{
-			PxScene* Scene = PActor->getScene();
-
-			// Get transform of actor we are grabbing
-			PxVec3 KinLocation = U2PVector(Location);
-			PxQuat KinOrientation = U2PQuat(Rotation.Quaternion());
-			PxTransform GrabbedActorPose = PActor->getGlobalPose();
-			PxTransform KinPose(KinLocation, KinOrientation);
-
-			// set target and current, so we don't need another "Tick" call to have it right
-			TargetTransform = CurrentTransform = P2UTransform(KinPose);
-
-			// If we don't already have a handle - make one now.
-			if(!HandleData)
-			{
-				// Create kinematic actor we are going to create joint with. This will be moved around with calls to SetLocation/SetRotation.
-				PxRigidDynamic* KinActor = Scene->getPhysics().createRigidDynamic(KinPose);
-				KinActor->setRigidBodyFlag(PxRigidBodyFlag::eKINEMATIC, true);
-				KinActor->setMass(1.0f);
-				KinActor->setMassSpaceInertiaTensor(PxVec3(1.0f, 1.0f, 1.0f));
-
-				// No bodyinstance
-				KinActor->userData = NULL;
-
-				// Add to Scene
-				Scene->addActor(*KinActor);
-
-				// Save reference to the kinematic actor.
-				KinActorData = KinActor;
-
-				// Create the joint
-				PxD6Joint* NewJoint = PxD6JointCreate(Scene->getPhysics(), KinActor, PxTransform(PxIdentity), PActor, GrabbedActorPose.transformInv(KinPose));
-
-				if(!NewJoint)
-				{
-					HandleData = 0;
-				}
-				else
-				{
-					// No constraint instance
-					NewJoint->userData = NULL;
-					HandleData = NewJoint;
-
-					// Remember the scene index that the handle joint/actor are in.
-					FPhysScene* RBScene = FPhysxUserData::Get<FPhysScene>(Scene->userData);
-
-					// Setting up the joint
-
-					PxD6Motion::Enum const LocationMotionType = bSoftLinearConstraint ? PxD6Motion::eFREE : PxD6Motion::eLOCKED;
-					PxD6Motion::Enum const RotationMotionType = (bSoftAngularConstraint || !bInConstrainRotation) ? PxD6Motion::eFREE : PxD6Motion::eLOCKED;
-
-					NewJoint->setMotion(PxD6Axis::eX, LocationMotionType);
-					NewJoint->setMotion(PxD6Axis::eY, LocationMotionType);
-					NewJoint->setMotion(PxD6Axis::eZ, LocationMotionType);
-					NewJoint->setDrivePosition(PxTransform(PxVec3(0, 0, 0)));
-
-					NewJoint->setMotion(PxD6Axis::eTWIST, RotationMotionType);
-					NewJoint->setMotion(PxD6Axis::eSWING1, RotationMotionType);
-					NewJoint->setMotion(PxD6Axis::eSWING2, RotationMotionType);
-
-					bRotationConstrained = bInConstrainRotation;
-
-					UpdateDriveSettings();
-				}
-			}
-		}
-	});
-	
-#elif WITH_CHAOS
 	// simulatable bodies should have handles.
 	FPhysicsActorHandle& InHandle = BodyInstance->GetPhysicsActorHandle();
 	if (!InHandle)
@@ -251,108 +146,54 @@ void UPhysicsHandleComponent::GrabComponentImp(UPrimitiveComponent* InComponent,
 
 	bRotationConstrained = bInConstrainRotation;
 	GrabbedHandle = InHandle; 
-#endif
-
 	GrabbedComponent = InComponent;
 	GrabbedBoneName = InBoneName;
 }
 
 void UPhysicsHandleComponent::UpdateDriveSettings()
 {
-#if PHYSICS_INTERFACE_PHYSX
-	if(HandleData != nullptr)
-	{
-		if (bSoftLinearConstraint)
-		{
-		HandleData->setDrive(PxD6Drive::eX, PxD6JointDrive(LinearStiffness, LinearDamping, PX_MAX_F32, PxD6JointDriveFlag::eACCELERATION));
-		HandleData->setDrive(PxD6Drive::eY, PxD6JointDrive(LinearStiffness, LinearDamping, PX_MAX_F32, PxD6JointDriveFlag::eACCELERATION));
-		HandleData->setDrive(PxD6Drive::eZ, PxD6JointDrive(LinearStiffness, LinearDamping, PX_MAX_F32, PxD6JointDriveFlag::eACCELERATION));
-		}
-
-		if (bSoftAngularConstraint && bRotationConstrained)
-		{
-			HandleData->setDrive(PxD6Drive::eSLERP, PxD6JointDrive(AngularStiffness, AngularDamping, PX_MAX_F32, PxD6JointDriveFlag::eACCELERATION));
-
-			//NewJoint->setDrive(PxD6Drive::eTWIST, PxD6JointDrive(AngularStiffness, AngularDamping, PX_MAX_F32, PxD6JointDriveFlag::eACCELERATION));
-			//NewJoint->setDrive(PxD6Drive::eSWING, PxD6JointDrive(AngularStiffness, AngularDamping, PX_MAX_F32, PxD6JointDriveFlag::eACCELERATION));
-		}
-	}
-#elif WITH_CHAOS
-
 	if (ConstraintHandle.IsValid() && ConstraintHandle.Constraint->IsType(Chaos::EConstraintType::JointConstraintType))
 	{
 		FPhysicsCommand::ExecuteWrite(ConstraintHandle, [&](const FPhysicsConstraintHandle& InConstraintHandle)
+		{
+			if (Chaos::FJointConstraint* Constraint = static_cast<Chaos::FJointConstraint*>(ConstraintHandle.Constraint))
 			{
-				if (Chaos::FJointConstraint* Constraint = static_cast<Chaos::FJointConstraint*>(ConstraintHandle.Constraint))
+				Chaos::EJointMotionType LocationMotionType = bSoftLinearConstraint ? Chaos::EJointMotionType::Free : Chaos::EJointMotionType::Locked;
+				Chaos::EJointMotionType RotationMotionType = (bSoftAngularConstraint || !bRotationConstrained) ? Chaos::EJointMotionType::Free : Chaos::EJointMotionType::Locked;
+
+				Constraint->SetCollisionEnabled(false);
+				Constraint->SetLinearVelocityDriveEnabled(Chaos::TVec3<bool>(LocationMotionType != Chaos::EJointMotionType::Locked));
+				Constraint->SetLinearPositionDriveEnabled(Chaos::TVec3<bool>(LocationMotionType != Chaos::EJointMotionType::Locked));
+				Constraint->SetLinearMotionTypesX(LocationMotionType);
+				Constraint->SetLinearMotionTypesY(LocationMotionType);
+				Constraint->SetLinearMotionTypesZ(LocationMotionType);
+
+				Constraint->SetAngularSLerpPositionDriveEnabled(bRotationConstrained && RotationMotionType != Chaos::EJointMotionType::Locked);
+				Constraint->SetAngularSLerpVelocityDriveEnabled(bRotationConstrained && RotationMotionType != Chaos::EJointMotionType::Locked);
+				Constraint->SetAngularMotionTypesX(RotationMotionType);
+				Constraint->SetAngularMotionTypesY(RotationMotionType);
+				Constraint->SetAngularMotionTypesZ(RotationMotionType);
+				FTransform GrabConstraintLocalTransform(ConstraintLocalRotation, ConstraintLocalPosition);
+				Constraint->SetJointTransforms({ FTransform::Identity , GrabConstraintLocalTransform });
+
+				if (LocationMotionType != Chaos::EJointMotionType::Locked)
 				{
-					Chaos::EJointMotionType LocationMotionType = bSoftLinearConstraint ? Chaos::EJointMotionType::Free : Chaos::EJointMotionType::Locked;
-					Chaos::EJointMotionType RotationMotionType = (bSoftAngularConstraint || !bRotationConstrained) ? Chaos::EJointMotionType::Free : Chaos::EJointMotionType::Locked;
-
-					Constraint->SetCollisionEnabled(false);
-					Constraint->SetLinearVelocityDriveEnabled(Chaos::TVec3<bool>(LocationMotionType != Chaos::EJointMotionType::Locked));
-					Constraint->SetLinearPositionDriveEnabled(Chaos::TVec3<bool>(LocationMotionType != Chaos::EJointMotionType::Locked));
-					Constraint->SetLinearMotionTypesX(LocationMotionType);
-					Constraint->SetLinearMotionTypesY(LocationMotionType);
-					Constraint->SetLinearMotionTypesZ(LocationMotionType);
-
-					Constraint->SetAngularSLerpPositionDriveEnabled(bRotationConstrained && RotationMotionType != Chaos::EJointMotionType::Locked);
-					Constraint->SetAngularSLerpVelocityDriveEnabled(bRotationConstrained && RotationMotionType != Chaos::EJointMotionType::Locked);
-					Constraint->SetAngularMotionTypesX(RotationMotionType);
-					Constraint->SetAngularMotionTypesY(RotationMotionType);
-					Constraint->SetAngularMotionTypesZ(RotationMotionType);
-					FTransform GrabConstraintLocalTransform(ConstraintLocalRotation, ConstraintLocalPosition);
-					Constraint->SetJointTransforms({ FTransform::Identity , GrabConstraintLocalTransform });
-
-					if (LocationMotionType != Chaos::EJointMotionType::Locked)
-					{
-						Constraint->SetLinearDriveStiffness(LinearStiffness);
-						Constraint->SetLinearDriveDamping(LinearDamping);
-					}
-
-					if (bRotationConstrained && RotationMotionType != Chaos::EJointMotionType::Locked)
-					{
-						Constraint->SetAngularDriveStiffness(AngularStiffness);
-						Constraint->SetAngularDriveDamping(AngularDamping);
-					}
+					Constraint->SetLinearDriveStiffness(LinearStiffness);
+					Constraint->SetLinearDriveDamping(LinearDamping);
 				}
-			});
+
+				if (bRotationConstrained && RotationMotionType != Chaos::EJointMotionType::Locked)
+				{
+					Constraint->SetAngularDriveStiffness(AngularStiffness);
+					Constraint->SetAngularDriveDamping(AngularDamping);
+				}
+			}
+		});
 	}
-#endif // WITH_PHYSX
 }
 
 void UPhysicsHandleComponent::ReleaseComponent()
 {
-#if PHYSICS_INTERFACE_PHYSX
-	if(GrabbedComponent)
-	{
-		if(HandleData)
-		{
-			check(KinActorData);
-
-			// use correct scene
-			PxScene* PScene = HandleData->getScene();
-			if(PScene)
-			{
-				SCOPED_SCENE_WRITE_LOCK(PScene);
-				// Destroy joint.
-				HandleData->release();
-				
-				// Destroy temporary actor.
-				KinActorData->release();
-				
-			}
-			KinActorData = NULL;
-			HandleData = NULL;
-		}
-
-		bRotationConstrained = false;
-
-		GrabbedComponent->WakeRigidBody(GrabbedBoneName);
-
-		GrabbedComponent = NULL;
-		GrabbedBoneName = NAME_None;
-	}
-#elif WITH_CHAOS
 	if (ConstraintHandle.IsValid())
 	{
 		FPhysicsInterface::ReleaseConstraint(ConstraintHandle);
@@ -372,8 +213,6 @@ void UPhysicsHandleComponent::ReleaseComponent()
 	}
 
 	ConstraintInstance.Reset();
-
-#endif
 }
 
 UPrimitiveComponent* UPhysicsHandleComponent::GetGrabbedComponent() const
@@ -399,54 +238,6 @@ void UPhysicsHandleComponent::SetTargetLocationAndRotation(FVector NewLocation,
 
 void UPhysicsHandleComponent::UpdateHandleTransform(const FTransform& NewTransform)
 {
-#if PHYSICS_INTERFACE_PHYSX
-	if (!KinActorData)
-	{
-		return;
-	}
-
-	bool bChangedPosition = true;
-	bool bChangedRotation = true;
-
-	PxRigidDynamic* KinActor = KinActorData;
-	PxScene* PScene = KinActor->getScene();
-	SCOPED_SCENE_WRITE_LOCK(PScene);
-
-	// Check if the new location is worthy of change
-	PxVec3 PNewLocation = U2PVector(NewTransform.GetTranslation());
-	PxVec3 PCurrentLocation = KinActor->getGlobalPose().p;
-	if((PNewLocation - PCurrentLocation).magnitudeSquared() <= 0.01f*0.01f)
-	{
-		PNewLocation = PCurrentLocation;
-		bChangedPosition = false;
-	}
-
-	// Check if the new rotation is worthy of change
-	PxQuat PNewOrientation = U2PQuat(NewTransform.GetRotation());
-	PxQuat PCurrentOrientation = KinActor->getGlobalPose().q;
-	if((FMath::Abs(PNewOrientation.dot(PCurrentOrientation)) > (1.f - SMALL_NUMBER)))
-	{
-		PNewOrientation = PCurrentOrientation;
-		bChangedRotation = false;
-	}
-
-	// Don't call moveKinematic if it hasn't changed - that will stop bodies from going to sleep.
-	if (bChangedPosition || bChangedRotation)
-	{
-		KinActor->setKinematicTarget(PxTransform(PNewLocation, PNewOrientation));
-
-		//LOC_MOD
-		//PxD6Joint* Joint = (PxD6Joint*) HandleData;
-		//if(Joint)// && (PNewLocation - PCurrentLocation).magnitudeSquared() > 0.01f*0.01f)
-		//{
-		//	PxRigidActor* Actor0, *Actor1;
-		//	Joint->getActors(Actor0, Actor1);
-		//	//Joint->setDrivePosition(PxTransform(Actor0->getGlobalPose().transformInv(PNewLocation)));
-		//	Joint->setDrivePosition(PxTransform::createIdentity());
-		//	//Joint->setDriveVelocity(PxVec3(0), PxVec3(0));
-		//}
-	}
-#elif WITH_CHAOS
 	if (!CurrentTransform.Equals(PreviousTransform))
 	{
 		FPhysicsCommand::ExecuteWrite(KinematicHandle, [&](const FPhysicsActorHandle& InKinematicHandle)
@@ -457,15 +248,12 @@ void UPhysicsHandleComponent::UpdateHandleTransform(const FTransform& NewTransfo
 
 		PreviousTransform = CurrentTransform;
 	}
-#endif
 }
 
 void UPhysicsHandleComponent::TickComponent(float DeltaTime, enum ELevelTick TickType, FActorComponentTickFunction *ThisTickFunction)
 {
 	Super::TickComponent(DeltaTime, TickType, ThisTickFunction);
 
-
-#if WITH_CHAOS
 	if (bPendingConstraint)
 	{
 		if (!ConstraintHandle.IsValid())
@@ -475,7 +263,6 @@ void UPhysicsHandleComponent::TickComponent(float DeltaTime, enum ELevelTick Tic
 
 	if (ConstraintHandle.IsValid())
 	{
-#endif
 		if (bInterpolateTarget)
 		{
 			const float Alpha = FMath::Clamp(DeltaTime * InterpolationSpeed, 0.f, 1.f);
@@ -491,7 +278,6 @@ void UPhysicsHandleComponent::TickComponent(float DeltaTime, enum ELevelTick Tic
 		}
 
 		UpdateHandleTransform(CurrentTransform);
-#if WITH_CHAOS
 	}
 	else if (KinematicHandle && GrabbedHandle)
 	{
@@ -511,7 +297,6 @@ void UPhysicsHandleComponent::TickComponent(float DeltaTime, enum ELevelTick Tic
 		}
 		bPendingConstraint = true;
 	}
-#endif
 }
 
 void UPhysicsHandleComponent::GetTargetLocationAndRotation(FVector& OutLocation, FRotator& OutRotation) const 
