Physics interface cleanup.

diff --git a/Engine/Source/Runtime/Engine/Private/PhysicsEngine/PhysicsInterfacePhysX.cpp b/Engine/Source/Runtime/Engine/Private/PhysicsEngine/PhysicsInterfacePhysX.cpp
--- a/Engine/Source/Runtime/Engine/Private/PhysicsEngine/PhysicsInterfacePhysX.cpp
+++ b/Engine/Source/Runtime/Engine/Private/PhysicsEngine/PhysicsInterfacePhysX.cpp
@@ -1,2825 +1 @@
 // Copyright Epic Games, Inc. All Rights Reserved.
-
-#if !WITH_CHAOS && !WITH_IMMEDIATE_PHYSX
-
-#include "Physics/PhysicsInterfacePhysX.h"
-#include "Physics/PhysicsInterfaceUtils.h"
-
-#include "Physics/PhysScene_PhysX.h"
-#include "Logging/MessageLog.h"
-#include "Components/SkeletalMeshComponent.h"
-#include "Internationalization/Internationalization.h"
-
-#include "PhysicsInterfaceDeclaresCore.h"
-
-#if !WITH_CHAOS_NEEDS_TO_BE_FIXED
-
-#include "SQAccelerator.h"
-
-#if WITH_PHYSX
-#include "PhysXPublic.h"
-#include "PhysxUserData.h"
-#include "PhysXPublic.h"
-#include "Physics/PhysicsFiltering.h"
-#include "PhysicsEngine/PhysXSupport.h"
-#include "Collision.h"
-#include "Collision/CollisionConversions.h"
-#endif // WITH_PHYSX
-
-#include "PhysicsEngine/ConstraintDrives.h"
-#include "PhysicsEngine/AggregateGeom.h"
-#include "Physics/PhysicsGeometry.h"
-#include "Engine/Engine.h"
-#include "PhysicalMaterials/PhysicalMaterial.h"
-#include "PhysicsEngine/PhysicsSettings.h"
-#include "PhysicsEngine/BodyInstance.h"
-
-using namespace physx;
-using namespace PhysicsInterfaceTypes;
-
-#define LOCTEXT_NAMESPACE "PhysicsInterface_PhysX"
-
-extern TAutoConsoleVariable<float> CVarConstraintLinearDampingScale;
-extern TAutoConsoleVariable<float> CVarConstraintLinearStiffnessScale;
-extern TAutoConsoleVariable<float> CVarConstraintAngularDampingScale;
-extern TAutoConsoleVariable<float> CVarConstraintAngularStiffnessScale;
-
-enum class EPhysicsInterfaceScopedLockType : uint8
-{
-	Read,
-	Write
-};
-
-template<typename PxType>
-PxType* GetPx(const FPhysicsActorHandle_PhysX& InActorHandle)
-{
-	PxRigidActor* Actor = FPhysicsInterface::GetPxRigidActor_AssumesLocked(InActorHandle);
-	return Actor ? Actor->is<PxType>() : nullptr;
-}
-
-struct FPhysicsInterfaceScopedLock_PhysX
-{
-private:
-
-	friend struct FPhysicsCommand_PhysX;
-
-	FPhysicsInterfaceScopedLock_PhysX(FPhysicsActorHandle_PhysX const * InActorHandle, EPhysicsInterfaceScopedLockType InLockType);
-	FPhysicsInterfaceScopedLock_PhysX(FPhysicsActorHandle_PhysX const * InActorHandleA, FPhysicsActorHandle_PhysX const * InActorHandleB, EPhysicsInterfaceScopedLockType InLockType);
-	FPhysicsInterfaceScopedLock_PhysX(FPhysicsConstraintHandle_PhysX const * InHandleerence, EPhysicsInterfaceScopedLockType InLockType);
-	FPhysicsInterfaceScopedLock_PhysX(USkeletalMeshComponent* InSkelMeshComp, EPhysicsInterfaceScopedLockType InLockType);
-	FPhysicsInterfaceScopedLock_PhysX(FPhysScene_PhysX* InScene, EPhysicsInterfaceScopedLockType InLockType);
-
-	~FPhysicsInterfaceScopedLock_PhysX();
-
-	void LockScenes();
-
-	physx::PxScene* Scenes[2];
-	EPhysicsInterfaceScopedLockType LockType;
-};
-
-FPhysicsActorHandle_PhysX DefaultPhysicsActorHandle()
-{
-	return FPhysicsActorHandle_PhysX();
-}
-
-bool IsValidHandle(const FPhysicsActorHandle_PhysX& Handle)
-{
-	return Handle.IsValid();
-}
-
-PxRigidActor* FPhysicsInterface_PhysX::GetPxRigidActor_AssumesLocked(const FPhysicsActorHandle_PhysX& InRef)
-{
-	return InRef.SyncActor;
-}
-
-physx::PxRigidDynamic* FPhysicsInterface_PhysX::GetPxRigidDynamic_AssumesLocked(const FPhysicsActorHandle_PhysX& InHandle)
-{
-	return GetPx<PxRigidDynamic>(InHandle);
-}
-
-physx::PxRigidBody* FPhysicsInterface_PhysX::GetPxRigidBody_AssumesLocked(const FPhysicsActorHandle_PhysX& InHandle)
-{
-	return GetPx<PxRigidBody>(InHandle);
-}
-
-const FBodyInstance* FPhysicsInterface_PhysX::ShapeToOriginalBodyInstance(const FBodyInstance* InCurrentInstance, const physx::PxShape* InShape)
-{
-	check(InCurrentInstance);
-	check(InShape);
-
-	const FBodyInstance* TargetInstance = InCurrentInstance->WeldParent ? InCurrentInstance->WeldParent : InCurrentInstance;
-	const FBodyInstance* OutInstance = TargetInstance;
-
-	if(const TMap<FPhysicsShapeHandle, FBodyInstance::FWeldInfo>* WeldInfo = InCurrentInstance->GetCurrentWeldInfo())
-	{
-		for(const TPair<FPhysicsShapeHandle, FBodyInstance::FWeldInfo>& Pair : *WeldInfo)
-		{
-			if(Pair.Key.Shape == InShape)
-			{
-				TargetInstance = Pair.Value.ChildBI;
-			}
-		}
-	}
-
-	return TargetInstance;
-}
-
-void FPhysicsInterface_PhysX::CreateActor(const FActorCreationParams& Params, FPhysicsActorHandle& Handle)
-{
-	FPhysicsActorHandle NewActor;
-
-	PxTransform PTransform = U2PTransform(Params.InitialTM);
-    FPhysScene* PhysScene = Params.Scene;
-
-	// Create Static
-	if(Params.bStatic)
-	{
-		NewActor.SyncActor = GPhysXSDK->createRigidStatic(PTransform);
-		if(Params.bQueryOnly)
-		{
-			ModifyActorFlag_Isolated<PxActorFlag::eDISABLE_SIMULATION>(NewActor.SyncActor, true);
-		}
-		NewActor.SyncActor->setName(Params.DebugName);
-	}
-	// Create Dynamic
-	else
-	{
-		PxRigidDynamic* PNewDynamic = GPhysXSDK->createRigidDynamic(PTransform);
-		NewActor.SyncActor = PNewDynamic;
-
-		PNewDynamic->setName(Params.DebugName);
-
-		ModifyRigidBodyFlag_Isolated<PxRigidBodyFlag::eUSE_KINEMATIC_TARGET_FOR_SCENE_QUERIES>(PNewDynamic, true);
-
-		if(Params.bQueryOnly)
-		{
-			ModifyActorFlag_Isolated<PxActorFlag::eDISABLE_SIMULATION>(PNewDynamic, true);
-		}
-
-		if(!Params.bEnableGravity)
-		{
-			ModifyActorFlag_Isolated<PxActorFlag::eDISABLE_GRAVITY>(PNewDynamic, true);
-		}
-	}
-
-	Handle = NewActor;
-}
-
-//helper function for TermBody to avoid code duplication between scenes
-void TermBodyHelper(FPhysScene* PhysScene, PxRigidActor*& PRigidActor, bool bNeverDeferRelease = false)
-{
-	if(PRigidActor)
-	{
-		// #Phys2 fixed hitting the check below because body scene was null, check was invalid in this case.
-		PxScene* PScene = PhysScene ? PhysScene->GetPxScene() : nullptr;
-		PxScene* BodyPScene = PRigidActor->getScene();
-		if(PScene && BodyPScene)
-		{
-			checkSlow(PScene == PRigidActor->getScene());
-
-			// Enable scene lock
-			SCOPED_SCENE_WRITE_LOCK(PScene);
-
-			// Let FPhysScene know
-			FBodyInstance* BodyInst = FPhysxUserData::Get<FBodyInstance>(PRigidActor->userData);
-			if(BodyInst)
-			{
-				PhysScene->RemoveBodyInstanceFromPendingLists_AssumesLocked(BodyInst);
-			}
-
-			PRigidActor->release();
-			//we must do this within the lock because we use it in the sub-stepping thread to determine that RigidActor is still valid
-			PRigidActor = nullptr;
-		}
-		else
-		{
-			if(bNeverDeferRelease)
-			{
-				PRigidActor->release();
-			}
-
-			PRigidActor = nullptr;
-		}
-	}
-
-	checkSlow(PRigidActor == NULL);
-}
-
-PxMaterial* GetDefaultPhysMaterial()
-{
-	check(GEngine->DefaultPhysMaterial != NULL);
-	const FPhysicsMaterialHandle_PhysX& MaterialHandle = GEngine->DefaultPhysMaterial->GetPhysicsMaterial();
-
-	return MaterialHandle.Material;
-}
-
-void FPhysicsInterface_PhysX::ReleaseActor(FPhysicsActorHandle_PhysX& InActorHandle, FPhysScene* InScene, bool bNeverDeferRelease)
-{
-	TermBodyHelper(InScene, InActorHandle.SyncActor, bNeverDeferRelease);
-}
-
-PxRigidActor* FPhysicsInterface_PhysX::GetPxRigidActorFromScene_AssumesLocked(const FPhysicsActorHandle_PhysX& InActorHandle)
-{
-	return InActorHandle.SyncActor;
-}
-
-PxD6Axis::Enum U2PConstraintAxis(PhysicsInterfaceTypes::ELimitAxis InAxis)
-{
-	switch(InAxis)
-	{
-	case PhysicsInterfaceTypes::ELimitAxis::X:
-		return PxD6Axis::Enum::eX;
-	case PhysicsInterfaceTypes::ELimitAxis::Y:
-		return PxD6Axis::Enum::eY;
-	case PhysicsInterfaceTypes::ELimitAxis::Z:
-		return PxD6Axis::Enum::eZ;
-	case PhysicsInterfaceTypes::ELimitAxis::Twist:
-		return PxD6Axis::Enum::eTWIST;
-	case PhysicsInterfaceTypes::ELimitAxis::Swing1:
-		return PxD6Axis::Enum::eSWING1;
-	case PhysicsInterfaceTypes::ELimitAxis::Swing2:
-		return PxD6Axis::Enum::eSWING2;
-	default:
-		check(false);
-	};
-
-	return PxD6Axis::Enum::eX;
-}
-
-/** Util for converting from UE motion enum to physx motion enum */
-PxD6Motion::Enum U2PAngularMotion(EAngularConstraintMotion InMotion)
-{
-	switch(InMotion)
-	{
-		case EAngularConstraintMotion::ACM_Free:
-			return PxD6Motion::eFREE;
-		case EAngularConstraintMotion::ACM_Limited:
-			return PxD6Motion::eLIMITED;
-		case EAngularConstraintMotion::ACM_Locked:
-			return PxD6Motion::eLOCKED;
-		default:
-			check(0);	//unsupported motion type
-	}
-
-	return PxD6Motion::eFREE;
-}
-
-/** Util for converting from UE motion enum to physx motion enum */
-PxD6Motion::Enum U2PLinearMotion(ELinearConstraintMotion InMotion)
-{
-	switch(InMotion)
-	{
-		case ELinearConstraintMotion::LCM_Free:
-			return PxD6Motion::eFREE;
-		case ELinearConstraintMotion::LCM_Limited:
-			return PxD6Motion::eLIMITED;
-		case ELinearConstraintMotion::LCM_Locked:
-			return PxD6Motion::eLOCKED;
-		default: 
-			check(0);	//unsupported motion type
-	}
-
-	return PxD6Motion::eFREE;
-}
-
-physx::PxJointActorIndex::Enum U2PConstraintFrame(EConstraintFrame::Type InFrame)
-{
-	// Swap frame order, since Unreal reverses physx order
-	return (InFrame == EConstraintFrame::Frame1) ? physx::PxJointActorIndex::eACTOR1 : physx::PxJointActorIndex::eACTOR0;
-}
-
-physx::PxD6Drive::Enum U2PDriveType(EDriveType InDriveType)
-{
-	switch(InDriveType)
-	{
-		case EDriveType::X:
-			return PxD6Drive::eX;
-		case EDriveType::Y:
-			return PxD6Drive::eY;
-		case EDriveType::Z:
-			return PxD6Drive::eZ;
-		case EDriveType::Swing:
-			return PxD6Drive::eSWING;
-		case EDriveType::Twist:
-			return PxD6Drive::eTWIST;
-		case EDriveType::Slerp:
-			return PxD6Drive::eSLERP;
-		default:
-			check(false);	// Invalid drive type
-			break;
-	};
-
-	return PxD6Drive::eX;
-}
-
-template<const PxRigidBodyFlag::Enum FlagToSet>
-void SetRigidBodyFlag(const FPhysicsActorHandle_PhysX& InRef, bool bInValue)
-{
-	if(PxRigidActor* Actor = FPhysicsInterface::GetPxRigidActor_AssumesLocked(InRef))
-	{
-		if(PxRigidBody* Body = Actor->is<PxRigidBody>())
-		{
-			PxRigidBodyFlags Flags = Body->getRigidBodyFlags();
-			ModifyRigidBodyFlag<FlagToSet>(Flags, bInValue);
-			Body->setRigidBodyFlags(Flags);
-		}
-	}
-}
-
-bool GetRigidBodyFlag(const FPhysicsActorHandle_PhysX& InRef, PxRigidBodyFlag::Enum InFlag)
-{
-	if(PxRigidActor* Actor = FPhysicsInterface::GetPxRigidActor_AssumesLocked(InRef))
-	{
-		if(PxRigidBody* Body = Actor->is<PxRigidBody>())
-		{
-			return (Body->getRigidBodyFlags() & InFlag);
-		}
-	}
-
-	return false;
-}
-
-PxTransform GetKinematicOrGlobalTransform_AssumesLocked(const PxRigidActor* InActor, bool bForceGlobalPose = false)
-{
-	check(InActor)
-	if(!bForceGlobalPose)
-	{
-		const PxRigidDynamic* Dynamic = InActor ? InActor->is<PxRigidDynamic>() : nullptr;
-
-		PxTransform PTarget;
-		if(Dynamic && Dynamic->getKinematicTarget(PTarget))
-		{
-			return PTarget;
-		}
-	}
-
-	return InActor->getGlobalPose();
-}
-
-// PhysX interface definition //////////////////////////////////////////////////////////////////////////
-
-FPhysicsActorHandle_PhysX::FPhysicsActorHandle_PhysX() 
-	: SyncActor(nullptr)
-{
-
-}
-
-bool FPhysicsActorHandle_PhysX::IsValid() const
-{
-	return FPhysicsInterface::GetPxRigidActor_AssumesLocked(*this) != nullptr;
-}
-
-bool FPhysicsActorHandle_PhysX::Equals(const FPhysicsActorHandle_PhysX& Other) const
-{
-	return SyncActor == Other.SyncActor;
-}
-
-FPhysicsConstraintHandle_PhysX::FPhysicsConstraintHandle_PhysX()
-	: ConstraintData(nullptr)
-{
-
-}
-
-bool FPhysicsConstraintHandle_PhysX::IsValid() const
-{
-	return ConstraintData != nullptr;
-}
-
-bool FPhysicsConstraintHandle_PhysX::Equals(const FPhysicsConstraintHandle_PhysX& Other) const
-{
-	return ConstraintData == Other.ConstraintData;
-}
-
-
-FPhysicsAggregateHandle_PhysX::FPhysicsAggregateHandle_PhysX()
-	: Aggregate(nullptr)
-{
-
-}
-
-bool FPhysicsAggregateHandle_PhysX::IsValid() const
-{
-	return Aggregate != nullptr;
-}
-
-
-FPhysicsInterfaceScopedLock_PhysX::FPhysicsInterfaceScopedLock_PhysX(FPhysicsActorHandle_PhysX const * InActorHandle, EPhysicsInterfaceScopedLockType InLockType)
-	: LockType(InLockType)
-{
-	Scenes[0] = (InActorHandle && InActorHandle->SyncActor) ? InActorHandle->SyncActor->getScene() : nullptr;
-	Scenes[1] = nullptr;
-	LockScenes();
-}
-
-PxScene* GetPxSceneForPhysActor(FPhysicsActorHandle_PhysX const * InActorHandle)
-{
-	PxScene* PScene = nullptr;
-	if (InActorHandle)
-	{
-		if (InActorHandle->SyncActor)
-		{
-			PScene = InActorHandle->SyncActor->getScene();
-		}
-	}
-
-	return PScene;
-}
-
-FPhysicsInterfaceScopedLock_PhysX::FPhysicsInterfaceScopedLock_PhysX(FPhysicsActorHandle_PhysX const * InActorHandleA, FPhysicsActorHandle_PhysX const * InActorHandleB, EPhysicsInterfaceScopedLockType InLockType)
-	: LockType(InLockType)
-{
-	Scenes[0] = GetPxSceneForPhysActor(InActorHandleA);
-	Scenes[1] = GetPxSceneForPhysActor(InActorHandleB);
-
-	// Only lock if we have unique scenes, either one vs. nullptr or both are equal
-	if(Scenes[0] == Scenes[1] || (!Scenes[0] || !Scenes[1]))
-	{
-		LockScenes();
-	}
-	else
-	{
-		UE_LOG(LogPhysics, Warning, TEXT("Attempted to aquire a physics scene lock for two paired actors that were not in the same scene. Skipping lock"));
-	}
-}
-
-
-FPhysicsInterfaceScopedLock_PhysX::FPhysicsInterfaceScopedLock_PhysX(FPhysicsConstraintHandle_PhysX const * InHandleerence, EPhysicsInterfaceScopedLockType InLockType)
-	: LockType(InLockType)
-{
-	if(InHandleerence)
-	{
-		Scenes[0] = InHandleerence->ConstraintData ? InHandleerence->ConstraintData->getScene() : nullptr;
-		Scenes[1] = nullptr;
-		LockScenes();
-	}
-}
-
-FPhysicsInterfaceScopedLock_PhysX::FPhysicsInterfaceScopedLock_PhysX(USkeletalMeshComponent* InSkelMeshComp, EPhysicsInterfaceScopedLockType InLockType)
-	: LockType(InLockType)
-{
-	Scenes[0] = nullptr;
-	Scenes[1] = nullptr;
-
-	// Iterate over bodies until we find a valid scene
-	// This assume that all bodies in the SkelMeshComp are in the same scene, which really should be the case, but should we verify in debug maybe?
-	if (InSkelMeshComp != nullptr)
-	{
-		for (FBodyInstance* BI : InSkelMeshComp->Bodies)
-		{
-			const FPhysicsActorHandle& ActorHandle = BI->GetPhysicsActorHandle();
-			PxScene* PScene = GetPxSceneForPhysActor(&ActorHandle);
-			if (PScene != nullptr)
-			{
-				Scenes[0] = PScene;
-				break;
-			}
-		}
-	}
-
-	LockScenes();
-}
-
-FPhysicsInterfaceScopedLock_PhysX::FPhysicsInterfaceScopedLock_PhysX(FPhysScene_PhysX* InScene, EPhysicsInterfaceScopedLockType InLockType)
-	: LockType(InLockType)
-{
-	Scenes[0] = nullptr;
-	Scenes[1] = nullptr;
-
-	if(InScene)
-	{
-		Scenes[0] = InScene->GetPxScene();
-		Scenes[1] = nullptr;
-	}
-
-	LockScenes();
-}
-
-FPhysicsInterfaceScopedLock_PhysX::~FPhysicsInterfaceScopedLock_PhysX()
-{
-	if (Scenes[0])
-	{
-		if (LockType == EPhysicsInterfaceScopedLockType::Read)
-		{
-			Scenes[0]->unlockRead();
-		}
-		else if (LockType == EPhysicsInterfaceScopedLockType::Write)
-		{
-			Scenes[0]->unlockWrite();
-		}
-	}
-
-	if (Scenes[1])
-	{
-		if (LockType == EPhysicsInterfaceScopedLockType::Read)
-		{
-			Scenes[1]->unlockRead();
-		}
-		else if (LockType == EPhysicsInterfaceScopedLockType::Write)
-		{
-			Scenes[1]->unlockWrite();
-		}
-	}
-}
-
-
-void FPhysicsInterfaceScopedLock_PhysX::LockScenes()
-{
-	if (Scenes[0])
-	{
-		if (LockType == EPhysicsInterfaceScopedLockType::Read)
-		{
-			Scenes[0]->lockRead();
-		}
-		else if(LockType == EPhysicsInterfaceScopedLockType::Write)
-		{
-			Scenes[0]->lockWrite();
-		}
-	}
-
-	if (Scenes[1])
-	{
-		if (LockType == EPhysicsInterfaceScopedLockType::Read)
-		{
-			Scenes[1]->lockRead();
-		}
-		else if (LockType == EPhysicsInterfaceScopedLockType::Write)
-		{
-			Scenes[1]->lockWrite();
-		}
-	}
-}
-
-bool FPhysicsCommand_PhysX::ExecuteRead(const FPhysicsActorHandle_PhysX& InActorHandle, TFunctionRef<void(const FPhysicsActorHandle_PhysX& Actor)> InCallable)
-{
-	if(InActorHandle.IsValid())
-	{
-		FPhysicsInterfaceScopedLock_PhysX ScopeLock(&InActorHandle, EPhysicsInterfaceScopedLockType::Read);
-		InCallable(InActorHandle);
-
-		return true;
-	}
-
-	return false;
-}
-
-bool FPhysicsCommand_PhysX::ExecuteRead(USkeletalMeshComponent* InMeshComponent, TFunctionRef<void()> InCallable)
-{
-	FPhysicsInterfaceScopedLock_PhysX ScopeLock(InMeshComponent, EPhysicsInterfaceScopedLockType::Read);
-	InCallable();
-
-	return ScopeLock.Scenes[0] || ScopeLock.Scenes[1];
-}
-
-bool FPhysicsCommand_PhysX::ExecuteRead(const FPhysicsActorHandle_PhysX& InActorHandleA, const FPhysicsActorHandle_PhysX& InActorHandleB, TFunctionRef<void(const FPhysicsActorHandle_PhysX& ActorA, const FPhysicsActorHandle_PhysX& ActorB)> InCallable)
-{
-	if(InActorHandleA.IsValid() || InActorHandleB.IsValid())
-	{
-		FPhysicsInterfaceScopedLock_PhysX ScopeLock(&InActorHandleA, &InActorHandleB, EPhysicsInterfaceScopedLockType::Read);
-		InCallable(InActorHandleA, InActorHandleB);
-
-		return ScopeLock.Scenes[0] || ScopeLock.Scenes[1];
-	}
-
-	return false;
-}
-
-bool FPhysicsCommand_PhysX::ExecuteRead(const FPhysicsConstraintHandle_PhysX& InHandle, TFunctionRef<void(const FPhysicsConstraintHandle_PhysX& Constraint)> InCallable)
-{
-	if(InHandle.IsValid())
-	{
-		FPhysicsInterfaceScopedLock_PhysX ScopeLock(&InHandle, EPhysicsInterfaceScopedLockType::Read);
-		InCallable(InHandle);
-
-		return true;
-	}
-	
-	return false;
-}
-
-bool FPhysicsCommand_PhysX::ExecuteRead(FPhysScene_PhysX* InScene, TFunctionRef<void()> InCallable)
-{
-	if(InScene)
-	{
-		FPhysicsInterfaceScopedLock_PhysX ScopeLock(InScene, EPhysicsInterfaceScopedLockType::Read);
-		InCallable();
-
-		return ScopeLock.Scenes[0] || ScopeLock.Scenes[1];
-	}
-
-	return false;
-}
-
-bool FPhysicsCommand_PhysX::ExecuteWrite(const FPhysicsActorHandle_PhysX& InActorHandle, TFunctionRef<void(const FPhysicsActorHandle_PhysX& Actor)> InCallable)
-{
-	if(InActorHandle.IsValid())
-	{
-		FPhysicsInterfaceScopedLock_PhysX ScopeLock(&InActorHandle, EPhysicsInterfaceScopedLockType::Write);
-		InCallable(InActorHandle);
-
-		return true;
-	}
-
-	return false;
-}
-
-bool FPhysicsCommand_PhysX::ExecuteWrite(USkeletalMeshComponent* InMeshComponent, TFunctionRef<void()> InCallable)
-{
-	FPhysicsInterfaceScopedLock_PhysX ScopeLock(InMeshComponent, EPhysicsInterfaceScopedLockType::Write);
-	InCallable();
-
-	return ScopeLock.Scenes[0] || ScopeLock.Scenes[1];
-}
-
-bool FPhysicsCommand_PhysX::ExecuteWrite(const FPhysicsActorHandle_PhysX& InActorHandleA, const FPhysicsActorHandle_PhysX& InActorHandleB, TFunctionRef<void(const FPhysicsActorHandle_PhysX& ActorA, const FPhysicsActorHandle_PhysX& ActorB)> InCallable)
-{
-	if(InActorHandleA.IsValid() || InActorHandleB.IsValid())
-	{
-		FPhysicsInterfaceScopedLock_PhysX ScopeLock(&InActorHandleA, &InActorHandleB, EPhysicsInterfaceScopedLockType::Write);
-		InCallable(InActorHandleA, InActorHandleB);
-
-		return ScopeLock.Scenes[0] || ScopeLock.Scenes[1];
-	}
-
-	return false;
-}
-
-bool FPhysicsCommand_PhysX::ExecuteWrite(const FPhysicsConstraintHandle_PhysX& InHandle, TFunctionRef<void(const FPhysicsConstraintHandle_PhysX& Constraint)> InCallable)
-{
-	if(InHandle.IsValid())
-	{
-		FPhysicsInterfaceScopedLock_PhysX ScopeLock(&InHandle, EPhysicsInterfaceScopedLockType::Write);
-		InCallable(InHandle);
-
-		return true;
-	}
-
-	return false;
-}
-
-bool FPhysicsCommand_PhysX::ExecuteWrite(FPhysScene_PhysX* InScene, TFunctionRef<void()> InCallable)
-{
-	FPhysicsInterfaceScopedLock_PhysX ScopeLock(InScene, EPhysicsInterfaceScopedLockType::Write);
-	InCallable();
-
-	return ScopeLock.Scenes[0] || ScopeLock.Scenes[1];
-
-	return false;
-}
-
-bool FPhysicsCommand_PhysX::ExecuteWrite(FPhysicsActorHandle_PhysX& InHandle, TFunctionRef<void(FPhysicsActorHandle_PhysX& Actor)> InCallable)
-{
-	if(InHandle.IsValid())
-	{
-		FPhysicsInterfaceScopedLock_PhysX ScopeLock(&InHandle, EPhysicsInterfaceScopedLockType::Write);
-		InCallable(InHandle);
-
-		return true;
-	}
-
-	return false;
-}
-
-void FPhysicsCommand_PhysX::ExecuteShapeWrite(FBodyInstance* InInstance, FPhysicsShapeHandle_PhysX& InShape, TFunctionRef<void(FPhysicsShapeHandle_PhysX& InShape)> InCallable)
-{
-	if(InShape.IsValid())
-	{
-		InCallable(InShape);
-	}
-}
-
-template<typename AllocatorType>
-int32 GetAllShapesInternal_AssumedLocked(const FPhysicsActorHandle_PhysX& InActorHandle, TArray<FPhysicsShapeHandle, AllocatorType>& OutShapes)
-{
-	int32 NumSyncShapes = 0;
-	TArray<PxShape*> TempShapes;
-	OutShapes.Empty();
-
-	// grab shapes from sync actor
-	if(InActorHandle.SyncActor)
-	{
-		NumSyncShapes = InActorHandle.SyncActor->getNbShapes();
-		TempShapes.AddUninitialized(NumSyncShapes);
-		InActorHandle.SyncActor->getShapes(TempShapes.GetData(), NumSyncShapes);
-	}
-
-	OutShapes.Reset(TempShapes.Num());
-	for(PxShape* Shape : TempShapes)
-	{
-		OutShapes.Add(FPhysicsShapeHandle_PhysX(Shape));
-	}
-
-	return NumSyncShapes;
-}
-
-template <>
-int32 FPhysicsInterface_PhysX::GetAllShapes_AssumedLocked(const FPhysicsActorHandle_PhysX& InActorHandle, TArray<FPhysicsShapeHandle_PhysX, FDefaultAllocator>& OutShapes)
-{
-	return GetAllShapesInternal_AssumedLocked(InActorHandle, OutShapes);
-}
-
-template <>
-int32 FPhysicsInterface_PhysX::GetAllShapes_AssumedLocked(const FPhysicsActorHandle_PhysX& InActorHandle, PhysicsInterfaceTypes::FInlineShapeArray& OutShapes)
-{
-	return GetAllShapesInternal_AssumedLocked(InActorHandle, OutShapes);
-}
-
-int32 FPhysicsInterface_PhysX::GetNumShapes(const FPhysicsActorHandle_PhysX& InActorHandle)
-{
-	return InActorHandle.SyncActor ? InActorHandle.SyncActor->getNbShapes() : 0;
-}
-
-void FPhysicsInterface_PhysX::ReleaseShape(const FPhysicsShapeHandle_PhysX& InShape)
-{
-	if(InShape.IsValid())
-	{
-		InShape.Shape->release();
-	}
-}
-
-void FPhysicsInterface_PhysX::AttachShape(const FPhysicsActorHandle_PhysX& InActor, const FPhysicsShapeHandle_PhysX& InNewShape)
-{
-	if(InActor.IsValid() && InNewShape.IsValid())
-	{
-		if(PxRigidActor* InternalActor = InActor.SyncActor)
-		{
-			InternalActor->attachShape(*InNewShape.Shape);
-		}
-	}
-}
-
-void FPhysicsInterface_PhysX::DetachShape(const FPhysicsActorHandle_PhysX& InActor, FPhysicsShapeHandle_PhysX& InShape, bool bWakeTouching /*= true*/)
-{
-	if(InActor.IsValid() && InShape.IsValid())
-	{
-		if(PxRigidActor* SyncActor = InActor.SyncActor)
-		{
-			SyncActor->detachShape(*InShape.Shape, bWakeTouching);
-		}
-	}
-}
-
-FPhysicsAggregateHandle_PhysX FPhysicsInterface_PhysX::CreateAggregate(int32 MaxBodies)
-{
-	FPhysicsAggregateHandle NewAggregate;
-	NewAggregate.Aggregate = GPhysXSDK->createAggregate(MaxBodies, true);
-	return NewAggregate;
-}
-
-void FPhysicsInterface_PhysX::ReleaseAggregate(FPhysicsAggregateHandle_PhysX& InAggregate)
-{
-	if (InAggregate.IsValid())
-	{
-		InAggregate.Aggregate->release();
-		InAggregate.Aggregate = nullptr;
-	}
-}
-
-
-int32 FPhysicsInterface_PhysX::GetNumActorsInAggregate(const FPhysicsAggregateHandle_PhysX& InAggregate)
-{
-	int32 NumActors = 0;
-	if (InAggregate.IsValid())
-	{
-		NumActors = InAggregate.Aggregate->getNbActors();
-	}
-	return NumActors;
-}
-
-
-void FPhysicsInterface_PhysX::AddActorToAggregate_AssumesLocked(const FPhysicsAggregateHandle_PhysX& InAggregate, const FPhysicsActorHandle_PhysX& InActor)
-{
-	if (InAggregate.Aggregate != nullptr)
-	{
-		if (InActor.SyncActor != nullptr)
-		{
-			InAggregate.Aggregate->addActor(*InActor.SyncActor);
-		}
-	}
-}
-
-FPhysicsShapeHandle_PhysX FPhysicsInterface_PhysX::CreateShape(physx::PxGeometry* InGeom, bool bSimulation /*= true*/, bool bQuery /*= true*/, UPhysicalMaterial* InSimpleMaterial /*= nullptr*/, TArray<UPhysicalMaterial*>* InComplexMaterials /*= nullptr*/)
-{
-	FPhysicsShapeHandle_PhysX OutHandle;
-
-	if(InGeom)
-	{
-		check(GPhysXSDK);
-		const PxMaterial* DefaultMaterial = GetDefaultPhysMaterial();
-
-		PxShapeFlags Flags;
-		if(bSimulation)
-		{
-			Flags |= PxShapeFlag::Enum::eSIMULATION_SHAPE;
-		}
-
-		if(bQuery)
-		{
-			Flags |= PxShapeFlag::Enum::eSCENE_QUERY_SHAPE;
-		}
-
-		Flags |= PxShapeFlag::Enum::eVISUALIZATION;
-
-		PxShape* NewShape = GPhysXSDK->createShape(*InGeom, *DefaultMaterial, true, Flags);
-
-		if(NewShape && (InSimpleMaterial || InComplexMaterials))
-		{
-			OutHandle = FPhysicsShapeHandle_PhysX(NewShape);
-			TArrayView<UPhysicalMaterial*> ComplexMaterialsView(InComplexMaterials ? InComplexMaterials->GetData() : nullptr, InComplexMaterials ? InComplexMaterials->Num() : 0);
-			FBodyInstance::ApplyMaterialToShape_AssumesLocked(OutHandle, InSimpleMaterial, ComplexMaterialsView);
-		}
-	}
-
-	return OutHandle;
-}
-
-void FPhysicsInterface_PhysX::AddGeometry(const FPhysicsActorHandle& InActor, const FGeometryAddParams& InParams, TArray<FPhysicsShapeHandle_PhysX>* OutOptShapes)
-{
-	PxRigidActor* PDestActor = InActor.SyncActor;
-
-	auto AttachShape_AssumesLocked = [&InParams, &OutOptShapes, PDestActor](const PxGeometry& PGeom, const PxTransform& PLocalPose, const float ContactOffset, const float RestOffset, const FPhysxUserData* ShapeElemUserData, PxShapeFlags PShapeFlags)
-	{
-		const FBodyCollisionData& BodyCollisionData = InParams.CollisionData;
-
-		const PxMaterial* PMaterial = GetDefaultPhysMaterial();
-		PxShape* PNewShape = GPhysXSDK->createShape(PGeom, *PMaterial, true, PShapeFlags);
-
-		if(PNewShape)
-		{
-			PNewShape->userData = (void*)ShapeElemUserData;
-			PNewShape->setLocalPose(PLocalPose);
-
-			if(OutOptShapes)
-			{
-				OutOptShapes->Add(FPhysicsShapeHandle_PhysX(PNewShape));
-			}
-
-			PNewShape->setContactOffset(ContactOffset);
-			PNewShape->setRestOffset(RestOffset);
-
-			const bool bComplexShape = PNewShape->getGeometryType() == PxGeometryType::eTRIANGLEMESH;
-			const bool bIsStatic = (PDestActor->is<PxRigidStatic>() != nullptr);
-
-			PxShapeFlags ShapeFlags = BuildPhysXShapeFlags(BodyCollisionData.CollisionFlags, bIsStatic, bComplexShape);
-
-			PNewShape->setQueryFilterData(U2PFilterData(bComplexShape ? BodyCollisionData.CollisionFilterData.QueryComplexFilter : BodyCollisionData.CollisionFilterData.QuerySimpleFilter));
-			PNewShape->setFlags(ShapeFlags);
-			PNewShape->setSimulationFilterData(U2PFilterData(BodyCollisionData.CollisionFilterData.SimFilter));
-			FBodyInstance::ApplyMaterialToShape_AssumesLocked(FPhysicsShapeHandle_PhysX(PNewShape), InParams.SimpleMaterial, InParams.ComplexMaterials);
-
-			PDestActor->attachShape(*PNewShape);
-			PNewShape->release();
-		}
-
-		return PNewShape;
-	};
-
-	auto IterateSimpleShapes = [AttachShape_AssumesLocked](const FKShapeElem& Elem, const PxGeometry& Geom, const PxTransform& PLocalPose, float ContactOffset, float RestOffset)
-	{
-		AttachShape_AssumesLocked(Geom, PLocalPose, ContactOffset, RestOffset, Elem.GetUserData(), PxShapeFlag::eVISUALIZATION | PxShapeFlag::eSCENE_QUERY_SHAPE | PxShapeFlag::eSIMULATION_SHAPE);
-	};
-
-	auto IterateTrimeshes = [AttachShape_AssumesLocked](PxTriangleMesh*, const PxGeometry& Geom, const PxTransform& PLocalPose, float ContactOffset, float RestOffset)
-	{
-		// Create without 'sim shape' flag, problematic if it's kinematic, and it gets set later anyway.
-		if(!AttachShape_AssumesLocked(Geom, PLocalPose, ContactOffset, RestOffset, nullptr, PxShapeFlag::eSCENE_QUERY_SHAPE | PxShapeFlag::eVISUALIZATION))
-		{
-			UE_LOG(LogPhysics, Log, TEXT("Can't create new mesh shape in AddShapesToRigidActor"));
-		}
-	};
-
-	FBodySetupShapeIterator AddShapesHelper(InParams.Scale, InParams.LocalTransform, InParams.bDoubleSided);
-
-	// Create shapes for simple collision if we do not want to use the complex collision mesh 
-	// for simple queries as well
-	FKAggregateGeom* AggGeom = InParams.Geometry;
-	check(AggGeom);
-
-	if (InParams.CollisionTraceType != ECollisionTraceFlag::CTF_UseComplexAsSimple)
-	{
-		AddShapesHelper.ForEachShape<FKSphereElem, PxSphereGeometry>(AggGeom->SphereElems, IterateSimpleShapes);
-		AddShapesHelper.ForEachShape<FKSphylElem, PxCapsuleGeometry>(AggGeom->SphylElems, IterateSimpleShapes);
-		AddShapesHelper.ForEachShape<FKBoxElem, PxBoxGeometry>(AggGeom->BoxElems, IterateSimpleShapes);
-		AddShapesHelper.ForEachShape<FKConvexElem, PxConvexMeshGeometry>(AggGeom->ConvexElems, IterateSimpleShapes);
-	}
-	
-	// Create tri-mesh shape, when we are not using simple collision shapes for 
-	// complex queries as well
-	if (InParams.CollisionTraceType != ECollisionTraceFlag::CTF_UseSimpleAsComplex)
-	{
-		AddShapesHelper.ForEachShape<PxTriangleMesh*, PxTriangleMeshGeometry>(InParams.TriMeshes, IterateTrimeshes);
-	}
-}
-
-FPhysicsShapeHandle_PhysX FPhysicsInterface_PhysX::CloneShape(const FPhysicsShapeHandle_PhysX& InShape)
-{
-	PxShape* PShape = InShape.Shape;
-
-	PxU16 PMaterialCount = PShape->getNbMaterials();
-
-	TArray<PxMaterial*, TInlineAllocator<64>> PMaterials;
-
-	PMaterials.AddZeroed(PMaterialCount);
-	PShape->getMaterials(&PMaterials[0], PMaterialCount);
-
-	PxShape* PNewShape = GPhysXSDK->createShape(PShape->getGeometry().any(), &PMaterials[0], PMaterialCount, false, PShape->getFlags());
-	PNewShape->setLocalPose(PShape->getLocalPose());
-	PNewShape->setContactOffset(PShape->getContactOffset());
-	PNewShape->setRestOffset(PShape->getRestOffset());
-	PNewShape->setSimulationFilterData(PShape->getSimulationFilterData());
-	PNewShape->setQueryFilterData(PShape->getQueryFilterData());
-	PNewShape->userData = PShape->userData;
-
-	return FPhysicsShapeHandle_PhysX(PNewShape);
-}
-
-FCollisionFilterData FPhysicsInterface_PhysX::GetSimulationFilter(const FPhysicsShapeHandle_PhysX& InShape)
-{
-	if(InShape.IsValid())
-	{
-		return P2UFilterData(InShape.Shape->getSimulationFilterData());
-	}
-
-	return FCollisionFilterData();
-}
-
-FCollisionFilterData FPhysicsInterface_PhysX::GetQueryFilter(const FPhysicsShapeHandle_PhysX& InShape)
-{
-	if(InShape.IsValid())
-	{
-		return P2UFilterData(InShape.Shape->getQueryFilterData());
-	}
-
-	return FCollisionFilterData();
-}
-
-bool FPhysicsInterface_PhysX::IsSimulationShape(const FPhysicsShapeHandle_PhysX& InShape)
-{
-	if(InShape.IsValid())
-	{
-		return InShape.Shape->getFlags() & PxShapeFlag::eSIMULATION_SHAPE;
-	}
-
-	return false;
-}
-
-bool FPhysicsInterface_PhysX::IsQueryShape(const FPhysicsShapeHandle_PhysX& InShape)
-{
-	if(InShape.IsValid())
-	{
-		return InShape.Shape->getFlags() & PxShapeFlag::eSCENE_QUERY_SHAPE;
-	}
-
-	return false;
-}
-
-ECollisionShapeType FPhysicsInterface_PhysX::GetShapeType(const FPhysicsShapeHandle_PhysX& InShape)
-{
-	if(InShape.IsValid())
-	{
-		return P2UCollisionShapeType(InShape.Shape->getGeometryType());
-	}
-
-	return ECollisionShapeType::None;
-}
-
-FPhysicsGeometryCollection_PhysX FPhysicsInterface_PhysX::GetGeometryCollection(const FPhysicsShapeHandle_PhysX& InShape)
-{
-	FPhysicsGeometryCollection_PhysX NewCollection(InShape);
-	return NewCollection;
-}
-
-FTransform FPhysicsInterface_PhysX::GetLocalTransform(const FPhysicsShapeHandle_PhysX& InShape)
-{
-	if(InShape.IsValid())
-	{
-		return P2UTransform(InShape.Shape->getLocalPose());
-	}
-	return FTransform::Identity;
-}
-
-FTransform FPhysicsInterface_PhysX::GetTransform(const FPhysicsShapeHandle_PhysX& InShape)
-{
-	if(InShape.IsValid())
-	{
-		PxRigidActor* OwningPxRigidActor = InShape.Shape->getActor();
-		check(OwningPxRigidActor);
-
-		FBodyInstance* BodyInst = FPhysxUserData::Get<FBodyInstance>(OwningPxRigidActor->userData);
-
-		if(BodyInst && BodyInst->ActorHandle.IsValid())
-		{
-			return GetLocalTransform(InShape) * GetTransform_AssumesLocked(BodyInst->ActorHandle);
-		}
-	}
-
-	return FTransform::Identity;
-}
-
-void* FPhysicsInterface_PhysX::GetUserData(const FPhysicsShapeHandle_PhysX& InShape)
-{
-	if(InShape.IsValid())
-	{
-		return InShape.Shape->userData;
-	}
-
-	return nullptr;
-}
-
-void FPhysicsInterface_PhysX::SetMaskFilter(const FPhysicsShapeHandle_PhysX& InShape, FMaskFilter InFilter)
-{
-	if(InShape.IsValid())
-	{
-		PxShape* PShape = InShape.Shape;
-		PxFilterData PQueryFilterData = PShape->getQueryFilterData();
-		UpdateMaskFilter(PQueryFilterData.word3, InFilter);
-		PShape->setQueryFilterData(PQueryFilterData);
-
-		PxFilterData PSimFilterData = PShape->getSimulationFilterData();
-		UpdateMaskFilter(PSimFilterData.word3, InFilter);
-		PShape->setSimulationFilterData(PSimFilterData);
-	}
-}
-
-void FPhysicsInterface_PhysX::SetSimulationFilter(const FPhysicsShapeHandle_PhysX& InShape, const FCollisionFilterData& InFilter)
-{
-	if(InShape.IsValid())
-	{
-		InShape.Shape->setSimulationFilterData(U2PFilterData(InFilter));
-	}
-}
-
-void FPhysicsInterface_PhysX::SetQueryFilter(const FPhysicsShapeHandle_PhysX& InShape, const FCollisionFilterData& InFilter)
-{
-	if(InShape.IsValid())
-	{
-		InShape.Shape->setQueryFilterData(U2PFilterData(InFilter));
-	}
-}
-
-void FPhysicsInterface_PhysX::SetIsSimulationShape(const FPhysicsShapeHandle_PhysX& InShape, bool bIsSimShape)
-{
-	if(InShape.IsValid())
-	{
-		PxShape* PShape = InShape.Shape;
-		ModifyShapeFlag_Isolated<PxShapeFlag::eSIMULATION_SHAPE>(PShape, bIsSimShape);
-	}
-}
-
-void FPhysicsInterface_PhysX::SetIsQueryShape(const FPhysicsShapeHandle_PhysX& InShape, bool bIsQueryShape)
-{
-	if(InShape.IsValid())
-	{
-		PxShape* PShape = InShape.Shape;
-		ModifyShapeFlag_Isolated<PxShapeFlag::eSCENE_QUERY_SHAPE>(PShape, bIsQueryShape);
-	}
-}
-
-void FPhysicsInterface_PhysX::SetUserData(const FPhysicsShapeHandle_PhysX& InShape, void* InUserData)
-{
-	if(InShape.IsValid())
-	{
-		InShape.Shape->userData = InUserData;
-	}
-}
-
-void FPhysicsInterface_PhysX::SetGeometry(const FPhysicsShapeHandle_PhysX& InShape, physx::PxGeometry& InGeom)
-{
-	if(InShape.IsValid())
-	{
-		InShape.Shape->setGeometry(InGeom);
-	}
-}
-
-void FPhysicsInterface_PhysX::SetLocalTransform(const FPhysicsShapeHandle_PhysX& InShape, const FTransform& NewLocalTransform)
-{
-	if(InShape.IsValid())
-	{
-		InShape.Shape->setLocalPose(U2PTransform(NewLocalTransform));
-	}
-}
-
-void FPhysicsInterface_PhysX::SetMaterials(const FPhysicsShapeHandle_PhysX& InShape, const TArrayView<UPhysicalMaterial*>InMaterials)
-{
-	if(InShape.IsValid())
-	{
-		TArray<PxMaterial*, TInlineAllocator<16>> PhysXMaterials;
-
-		for(UPhysicalMaterial* UnrealMat : InMaterials)
-		{
-			check(UnrealMat);
-			PhysXMaterials.Add(UnrealMat->GetPhysicsMaterial().Material);
-			check(PhysXMaterials.Last());
-		}
-
-		InShape.Shape->setMaterials(PhysXMaterials.GetData(), PhysXMaterials.Num());
-	}
-}
-
-void FPhysicsInterface_PhysX::SetActorUserData_AssumesLocked(const FPhysicsActorHandle_PhysX& InActorHandle, FPhysxUserData* InUserData)
-{
-	if (InActorHandle.SyncActor)
-	{
-		InActorHandle.SyncActor->userData = InUserData;
-	}
-}
-
-bool FPhysicsInterface_PhysX::IsRigidBody(const FPhysicsActorHandle_PhysX& InActorHandle)
-{
-	if (PxRigidActor* Actor = GetPxRigidActor_AssumesLocked(InActorHandle))
-	{
-		return Actor->is<PxRigidBody>() != nullptr;
-	}
-
-	return false;
-}
-
-bool FPhysicsInterface_PhysX::IsDynamic(const FPhysicsActorHandle_PhysX& InActorHandle)
-{
-	if(PxRigidActor* Actor = GetPxRigidActor_AssumesLocked(InActorHandle))
-	{
-		return Actor->is<PxRigidDynamic>() != nullptr;
-	}
-
-	return false;
-}
-
-bool FPhysicsInterface_PhysX::IsStatic(const FPhysicsActorHandle_PhysX& InActorHandle)
-{
-	if (PxRigidActor* Actor = GetPxRigidActor_AssumesLocked(InActorHandle))
-	{
-		return Actor->is<PxRigidStatic>() != nullptr;
-	}
-
-	return false;
-}
-
-bool FPhysicsInterface_PhysX::IsKinematic_AssumesLocked(const FPhysicsActorHandle_PhysX& InActorHandle)
-{
-	return GetRigidBodyFlag(InActorHandle, PxRigidBodyFlag::eKINEMATIC);
-}
-
-bool FPhysicsInterface_PhysX::IsSleeping(const FPhysicsActorHandle_PhysX& InActorHandle)
-{
-	PxRigidActor* Actor = GetPxRigidActor_AssumesLocked(InActorHandle);
-	PxRigidDynamic* Dynamic = Actor ? Actor->is<PxRigidDynamic>() : nullptr;
-	return !Dynamic || (Dynamic->getScene() && Dynamic->isSleeping());
-}
-
-bool FPhysicsInterface_PhysX::IsCcdEnabled(const FPhysicsActorHandle_PhysX& InActorHandle)
-{
-	return GetRigidBodyFlag(InActorHandle, PxRigidBodyFlag::eENABLE_CCD);
-}
-
-bool FPhysicsInterface_PhysX::IsDisabled(const FPhysicsActorHandle_PhysX& InHandle)
-{
-	// @todo(chaos) todo - see SetDisabled() bewlow
-	return false;
-}
-
-void FPhysicsInterface_PhysX::SetDisabled(const FPhysicsActorHandle_PhysX& InHandle, bool bSetDisabled)
-{
-	// @todo(chaos) todo - the closest way may be to set the PxRigidBodyFlag::eDISABLE_SIMULATION, but this may have implication on what you can or not call on the body after that 
-	ensureMsgf(false, TEXT("SetDisabled is unsupported by the PhysX interface"));
-}
-
-
-bool FPhysicsInterface_PhysX::IsInScene(const FPhysicsActorHandle_PhysX& InActorHandle)
-{
-	if(PxRigidActor* Actor = GetPxRigidActor_AssumesLocked(InActorHandle))
-	{
-		return Actor->getScene() != nullptr;
-	}
-
-	return false;
-}
-
-FPhysScene* FPhysicsInterface_PhysX::GetCurrentScene(const FPhysicsActorHandle_PhysX& InActorHandle)
-{
-	if(PxRigidActor* Actor = GetPxRigidActor_AssumesLocked(InActorHandle))
-	{
-		if(PxScene* PScene = Actor->getScene())
-		{
-			return FPhysxUserData::Get<FPhysScene>(PScene->userData);
-		}
-	}
-
-	return nullptr;
-}
-
-bool FPhysicsInterface_PhysX::CanSimulate_AssumesLocked(const FPhysicsActorHandle_PhysX& InActorHandle)
-{
-	if(PxActor* Actor = GetPxRigidActor_AssumesLocked(InActorHandle))
-	{
-		return !(Actor->getActorFlags() & PxActorFlag::eDISABLE_SIMULATION);
-	}
-	return false;
-}
-
-float FPhysicsInterface_PhysX::GetMass_AssumesLocked(const FPhysicsActorHandle_PhysX& InActorHandle)
-{
-	if(PxRigidActor* Actor = GetPxRigidActor_AssumesLocked(InActorHandle))
-	{
-		if(PxRigidBody* Body = Actor->is<PxRigidBody>())
-		{
-			return Body->getMass();
-		}
-	}
-
-	return 0.0f;
-}
-
-void FPhysicsInterface_PhysX::SetSendsSleepNotifies_AssumesLocked(const FPhysicsActorHandle_PhysX& InActorHandle, bool bSendSleepNotifies)
-{
-	if (PxRigidBody* Body = GetPx<PxRigidBody>(InActorHandle))
-	{
-		ModifyActorFlag_Isolated<PxActorFlag::eSEND_SLEEP_NOTIFIES>(Body, bSendSleepNotifies);
-	}
-}
-
-
-void FPhysicsInterface_PhysX::PutToSleep_AssumesLocked(const FPhysicsActorHandle_PhysX& InActorHandle)
-{
-	PxRigidActor* Actor = GetPxRigidActor_AssumesLocked(InActorHandle);
-	PxRigidDynamic* Dynamic = Actor ? Actor->is<PxRigidDynamic>() : nullptr;
-
-	if(Dynamic)
-	{
-		Dynamic->putToSleep();
-	}
-}
-
-void FPhysicsInterface_PhysX::WakeUp_AssumesLocked(const FPhysicsActorHandle_PhysX& InActorHandle)
-{
-	PxRigidActor* Actor = GetPxRigidActor_AssumesLocked(InActorHandle);
-	PxRigidDynamic* Dynamic = Actor ? Actor->is<PxRigidDynamic>() : nullptr;
-
-	if(Dynamic)
-	{
-		Dynamic->wakeUp();
-	}
-}
-
-void FPhysicsInterface_PhysX::SetIsKinematic_AssumesLocked(const FPhysicsActorHandle_PhysX& InActorHandle, bool bIsKinematic)
-{
-	SetRigidBodyFlag<PxRigidBodyFlag::eKINEMATIC>(InActorHandle, bIsKinematic);
-}
-
-void FPhysicsInterface_PhysX::SetCcdEnabled_AssumesLocked(const FPhysicsActorHandle_PhysX& InActorHandle, bool bIsCcdEnabled)
-{
-	SetRigidBodyFlag<PxRigidBodyFlag::eENABLE_CCD>(InActorHandle, bIsCcdEnabled);
-}
-
-FTransform FPhysicsInterface_PhysX::GetGlobalPose_AssumesLocked(const FPhysicsActorHandle_PhysX& InActorHandle)
-{
-	PxRigidActor* Actor = GetPxRigidActor_AssumesLocked(InActorHandle);
-	return Actor ? P2UTransform(Actor->getGlobalPose()) : FTransform::Identity;
-}
-
-void FPhysicsInterface_PhysX::SetGlobalPose_AssumesLocked(const FPhysicsActorHandle_PhysX& InActorHandle, const FTransform& InNewPose, bool bAutoWake)
-{
-	PxRigidActor* Actor = GetPxRigidActor_AssumesLocked(InActorHandle);
-
-	if(Actor)
-	{
-		Actor->setGlobalPose(U2PTransform(InNewPose));
-	}
-}
-
-FTransform FPhysicsInterface_PhysX::GetTransform_AssumesLocked(const FPhysicsActorHandle& InRef, bool bForceGlobalPose /*= false*/)
-{
-	if(!bForceGlobalPose)
-	{
-		if(IsDynamic(InRef))
-		{
-			if(HasKinematicTarget_AssumesLocked(InRef))
-			{
-				return GetKinematicTarget_AssumesLocked(InRef);
-			}
-		}
-	}
-
-	return GetGlobalPose_AssumesLocked(InRef);
-}
-
-bool FPhysicsInterface_PhysX::HasKinematicTarget_AssumesLocked(const FPhysicsActorHandle_PhysX& InActorHandle)
-{
-	PxRigidActor* Actor = GetPxRigidActor_AssumesLocked(InActorHandle);
-	PxRigidDynamic* Dynamic = Actor ? Actor->is<PxRigidDynamic>() : nullptr;
-
-	PxTransform PTarget;
-	return Dynamic && Dynamic->getKinematicTarget(PTarget);
-}
-
-FTransform FPhysicsInterface_PhysX::GetKinematicTarget_AssumesLocked(const FPhysicsActorHandle_PhysX& InActorHandle)
-{
-	PxRigidActor* Actor = GetPxRigidActor_AssumesLocked(InActorHandle);
-	PxRigidDynamic* Dynamic = Actor ? Actor->is<PxRigidDynamic>() : nullptr;
-
-	PxTransform PTarget;
-	if(Dynamic && Dynamic->getKinematicTarget(PTarget))
-	{
-		return P2UTransform(PTarget);
-	}
-	
-	return FTransform::Identity;
-}
-
-void FPhysicsInterface_PhysX::SetKinematicTarget_AssumesLocked(const FPhysicsActorHandle_PhysX& InActorHandle, const FTransform& InNewTarget)
-{
-	PxRigidActor* Actor = GetPxRigidActor_AssumesLocked(InActorHandle);
-	PxRigidDynamic* Dynamic = Actor ? Actor->is<PxRigidDynamic>() : nullptr;
-
-	if(Dynamic)
-	{
-		Dynamic->setKinematicTarget(U2PTransform(InNewTarget));
-	}
-}
-
-FVector FPhysicsInterface_PhysX::GetLinearVelocity_AssumesLocked(const FPhysicsActorHandle_PhysX& InActorHandle)
-{
-	PxRigidActor* Actor = GetPxRigidActor_AssumesLocked(InActorHandle);
-	PxRigidBody* Body = Actor ? Actor->is<PxRigidBody>() : nullptr;
-
-	if(Body)
-	{
-		return P2UVector(Body->getLinearVelocity());
-	}
-
-	return FVector::ZeroVector;
-}
-
-void FPhysicsInterface_PhysX::SetLinearVelocity_AssumesLocked(const FPhysicsActorHandle_PhysX& InActorHandle, const FVector& InNewVelocity, bool bAutoWake)
-{
-	PxRigidActor* Actor = GetPxRigidActor_AssumesLocked(InActorHandle);
-	PxRigidBody* Body = Actor ? Actor->is<PxRigidBody>() : nullptr;
-
-	if(Body)
-	{
-		Body->setLinearVelocity(U2PVector(InNewVelocity));
-	}
-}
-
-FVector FPhysicsInterface_PhysX::GetAngularVelocity_AssumesLocked(const FPhysicsActorHandle_PhysX& InActorHandle)
-{
-	PxRigidActor* Actor = GetPxRigidActor_AssumesLocked(InActorHandle);
-	PxRigidBody* Body = Actor ? Actor->is<PxRigidBody>() : nullptr;
-
-	if(Body)
-	{
-		return P2UVector(Body->getAngularVelocity());
-	}
-
-	return FVector::ZeroVector;
-}
-
-void FPhysicsInterface_PhysX::SetAngularVelocity_AssumesLocked(const FPhysicsActorHandle_PhysX& InActorHandle, const FVector& InNewVelocity, bool bAutoWake)
-{
-	PxRigidActor* Actor = GetPxRigidActor_AssumesLocked(InActorHandle);
-	PxRigidBody* Body = Actor ? Actor->is<PxRigidBody>() : nullptr;
-
-	if(Body)
-	{
-		Body->setAngularVelocity(U2PVector(InNewVelocity));
-	}
-}
-
-float FPhysicsInterface_PhysX::GetMaxAngularVelocity_AssumesLocked(const FPhysicsActorHandle_PhysX& InActorHandle)
-{
-	if(PxRigidDynamic* Dynamic = GetPx<PxRigidDynamic>(InActorHandle))
-	{
-		return Dynamic->getMaxAngularVelocity();
-	}
-
-	return 0.0f;
-}
-
-void FPhysicsInterface_PhysX::SetMaxAngularVelocity_AssumesLocked(const FPhysicsActorHandle_PhysX& InActorHandle, float InMaxAngularVelocity)
-{
-	if(PxRigidDynamic* Dynamic = GetPx<PxRigidDynamic>(InActorHandle))
-	{
-		Dynamic->setMaxAngularVelocity(InMaxAngularVelocity);
-	}
-}
-
-float FPhysicsInterface_PhysX::GetMaxDepenetrationVelocity_AssumesLocked(const FPhysicsActorHandle_PhysX& InActorHandle)
-{
-	if(PxRigidBody* Body = GetPx<PxRigidBody>(InActorHandle))
-	{
-		return Body->getMaxDepenetrationVelocity();
-	}
-
-	return 0.0f;
-}
-
-void FPhysicsInterface_PhysX::SetMaxDepenetrationVelocity_AssumesLocked(const FPhysicsActorHandle_PhysX& InActorHandle, float InMaxDepenetrationVelocity)
-{
-	if(PxRigidBody* Body = GetPx<PxRigidBody>(InActorHandle))
-	{
-		float UseMaxVelocity = InMaxDepenetrationVelocity == 0.f ? PX_MAX_F32 : InMaxDepenetrationVelocity;
-		Body->setMaxDepenetrationVelocity(UseMaxVelocity);
-	}
-}
-
-FVector FPhysicsInterface_PhysX::GetWorldVelocityAtPoint_AssumesLocked(const FPhysicsActorHandle_PhysX& InActorHandle, const FVector& InPoint)
-{
-	if(PxRigidBody* Body = GetPx<PxRigidBody>(InActorHandle))
-	{
-		return P2UVector(PxRigidBodyExt::getVelocityAtPos(*Body, U2PVector(InPoint)));
-	}
-
-	return FVector::ZeroVector;
-}
-
-FTransform FPhysicsInterface_PhysX::GetComTransform_AssumesLocked(const FPhysicsActorHandle_PhysX& InActorHandle)
-{
-	PxRigidActor* Actor = GetPxRigidActor_AssumesLocked(InActorHandle);
-	PxRigidBody* Body = Actor ? Actor->is<PxRigidBody>() : nullptr;
-
-	if(Body)
-	{
-		PxTransform PLocalCom = Body->getCMassLocalPose();
-		return P2UTransform(GetKinematicOrGlobalTransform_AssumesLocked(Actor) * PLocalCom);
-	}
-	else if(Actor)
-	{
-		return P2UTransform(GetKinematicOrGlobalTransform_AssumesLocked(Actor));
-	}
-
-	return FTransform::Identity;
-}
-
-FTransform FPhysicsInterface_PhysX::GetComTransformLocal_AssumesLocked(const FPhysicsActorHandle_PhysX& InHandle)
-{
-	FTransform OutTransform = FTransform::Identity;
-
-	PxRigidActor* Actor = GetPxRigidActor_AssumesLocked(InHandle);
-	PxRigidBody* Body = Actor ? Actor->is<PxRigidBody>() : nullptr;
-
-	if(Body)
-	{
-		OutTransform = P2UTransform(Body->getCMassLocalPose());
-	}
-
-	return OutTransform;
-}
-
-FVector FPhysicsInterface_PhysX::GetLocalInertiaTensor_AssumesLocked(const FPhysicsActorHandle_PhysX& InActorHandle)
-{
-	if(PxRigidBody* Body = GetPx<PxRigidBody>(InActorHandle))
-	{
-		return P2UVector(Body->getMassSpaceInertiaTensor());
-	}
-
-	return FVector::ZeroVector;
-}
-
-FBox FPhysicsInterface_PhysX::GetBounds_AssumesLocked(const FPhysicsActorHandle_PhysX& InActorHandle)
-{
-	if(PxRigidActor* Actor = GetPxRigidActor_AssumesLocked(InActorHandle))
-	{
-		PxBounds3 PBounds = Actor->getWorldBounds();
-		return FBox(P2UVector(PBounds.minimum), P2UVector(PBounds.maximum));
-	}
-
-	return FBox();
-}
-
-void FPhysicsInterface_PhysX::SetLinearDamping_AssumesLocked(const FPhysicsActorHandle_PhysX& InActorHandle, float InDamping)
-{
-	if(PxRigidDynamic* Dynamic = GetPx<PxRigidDynamic>(InActorHandle))
-	{
-		Dynamic->setLinearDamping(InDamping);
-	}
-}
-
-void FPhysicsInterface_PhysX::SetAngularDamping_AssumesLocked(const FPhysicsActorHandle_PhysX& InActorHandle, float InDamping)
-{
-	if(PxRigidDynamic* Dynamic = GetPx<PxRigidDynamic>(InActorHandle))
-	{
-		Dynamic->setAngularDamping(InDamping);
-	}
-}
-
-void FPhysicsInterface_PhysX::AddImpulse_AssumesLocked(const FPhysicsActorHandle_PhysX& InActorHandle, const FVector& InForce)
-{
-	if(PxRigidBody* Body = GetPx<PxRigidBody>(InActorHandle))
-	{
-		Body->addForce(U2PVector(InForce), PxForceMode::eIMPULSE);
-	}
-}
-
-void FPhysicsInterface_PhysX::AddAngularImpulseInRadians_AssumesLocked(const FPhysicsActorHandle_PhysX& InActorHandle, const FVector& InTorque)
-{
-	if(PxRigidBody* Body = GetPx<PxRigidBody>(InActorHandle))
-	{
-		Body->addTorque(U2PVector(InTorque), PxForceMode::eIMPULSE);
-	}
-}
-
-void FPhysicsInterface_PhysX::AddVelocity_AssumesLocked(const FPhysicsActorHandle_PhysX& InActorHandle, const FVector& InForce)
-{
-	if(PxRigidBody* Body = GetPx<PxRigidBody>(InActorHandle))
-	{
-		Body->addForce(U2PVector(InForce), PxForceMode::eVELOCITY_CHANGE);
-	}
-}
-
-void FPhysicsInterface_PhysX::AddAngularVelocityInRadians_AssumesLocked(const FPhysicsActorHandle_PhysX& InActorHandle, const FVector& InTorque)
-{
-	if(PxRigidBody* Body = GetPx<PxRigidBody>(InActorHandle))
-	{
-		Body->addTorque(U2PVector(InTorque), PxForceMode::eVELOCITY_CHANGE);
-	}
-}
-
-void FPhysicsInterface_PhysX::AddImpulseAtLocation_AssumesLocked(const FPhysicsActorHandle_PhysX& InActorHandle, const FVector& InImpulse, const FVector& InLocation)
-{
-	if(PxRigidBody* Body = GetPx<PxRigidBody>(InActorHandle))
-	{
-		PxRigidBodyExt::addForceAtPos(*Body, U2PVector(InImpulse), U2PVector(InLocation), PxForceMode::eIMPULSE);
-	}
-}
-
-void FPhysicsInterface_PhysX::AddRadialImpulse_AssumesLocked(const FPhysicsActorHandle_PhysX& InActorHandle, const FVector& InOrigin, float InRadius, float InStrength, ERadialImpulseFalloff InFalloff, bool bInVelChange)
-{
-	if(PxRigidBody* Body = GetPx<PxRigidBody>(InActorHandle))
-	{
-		AddRadialImpulseToPxRigidBody_AssumesLocked(*Body, InOrigin, InRadius, InStrength, InFalloff, bInVelChange);
-	}
-}
-
-bool FPhysicsInterface_PhysX::IsGravityEnabled_AssumesLocked(const FPhysicsActorHandle_PhysX& InActorHandle)
-{
-	if(PxRigidActor* Actor = GetPxRigidActor_AssumesLocked(InActorHandle))
-	{
-		return !(Actor->getActorFlags() & PxActorFlag::eDISABLE_GRAVITY);
-	}
-
-	return false;
-}
-
-void FPhysicsInterface_PhysX::SetGravityEnabled_AssumesLocked(const FPhysicsActorHandle_PhysX& InActorHandle, bool bEnabled)
-{
-	if(PxRigidBody* Body = GetPx<PxRigidBody>(InActorHandle))
-	{
-		ModifyActorFlag_Isolated<PxActorFlag::eDISABLE_GRAVITY>(Body, !bEnabled);
-	}
-}
-
-float FPhysicsInterface_PhysX::GetSleepEnergyThreshold_AssumesLocked(const FPhysicsActorHandle_PhysX& InActorHandle)
-{
-	if(PxRigidDynamic* Dynamic = GetPx<PxRigidDynamic>(InActorHandle))
-	{
-		return Dynamic->getSleepThreshold();
-	}
-
-	return 0.0f;
-}
-
-void FPhysicsInterface_PhysX::SetSleepEnergyThreshold_AssumesLocked(const FPhysicsActorHandle_PhysX& InActorHandle, float InEnergyThreshold)
-{
-	if(PxRigidDynamic* Dynamic = GetPx<PxRigidDynamic>(InActorHandle))
-	{
-		Dynamic->setSleepThreshold(InEnergyThreshold);
-	}
-}
-
-void FPhysicsInterface_PhysX::SetMass_AssumesLocked(const FPhysicsActorHandle_PhysX& InActorHandle, float InMass)
-{
-	if(PxRigidBody* Body = GetPx<PxRigidBody>(InActorHandle))
-	{
-		Body->setMass(InMass);
-	}
-}
-
-void FPhysicsInterface_PhysX::SetMassSpaceInertiaTensor_AssumesLocked(const FPhysicsActorHandle_PhysX& InActorHandle, const FVector& InTensor)
-{
-	if(PxRigidBody* Body = GetPx<PxRigidBody>(InActorHandle))
-	{
-		Body->setMassSpaceInertiaTensor(U2PVector(InTensor));
-	}
-}
-
-void FPhysicsInterface_PhysX::SetComLocalPose_AssumesLocked(const FPhysicsActorHandle_PhysX& InActorHandle, const FTransform& InComLocalPose)
-{
-	if(PxRigidBody* Body = GetPx<PxRigidBody>(InActorHandle))
-	{
-		Body->setCMassLocalPose(U2PTransform(InComLocalPose));
-	}
-}
-
-float FPhysicsInterface_PhysX::GetStabilizationEnergyThreshold_AssumesLocked(const FPhysicsActorHandle_PhysX& InHandle)
-{
-	if(PxRigidDynamic* Dynamic = GetPx<PxRigidDynamic>(InHandle))
-	{
-		return Dynamic->getStabilizationThreshold();
-	}
-
-	return 0.0f;
-}
-
-void FPhysicsInterface_PhysX::SetStabilizationEnergyThreshold_AssumesLocked(const FPhysicsActorHandle_PhysX& InHandle, float InThreshold)
-{
-	if(PxRigidDynamic* Dynamic = GetPx<PxRigidDynamic>(InHandle))
-	{
-		Dynamic->setStabilizationThreshold(InThreshold);
-	}
-}
-
-uint32 FPhysicsInterface_PhysX::GetSolverPositionIterationCount_AssumesLocked(const FPhysicsActorHandle_PhysX& InHandle)
-{
-	if(PxRigidDynamic* Dynamic = GetPx<PxRigidDynamic>(InHandle))
-	{
-		uint32 PositionIters = 0;
-		uint32 VelocityIters = 0;
-		Dynamic->getSolverIterationCounts(PositionIters, VelocityIters);
-
-		return PositionIters;
-	}
-
-	return 0;
-}
-
-void FPhysicsInterface_PhysX::SetSolverPositionIterationCount_AssumesLocked(const FPhysicsActorHandle_PhysX& InHandle, uint32 InSolverIterationCount)
-{
-	if(PxRigidDynamic* Dynamic = GetPx<PxRigidDynamic>(InHandle))
-	{
-		uint32 PositionIters = 0;
-		uint32 VelocityIters = 0;
-		Dynamic->getSolverIterationCounts(PositionIters, VelocityIters);
-		Dynamic->setSolverIterationCounts(InSolverIterationCount, VelocityIters);
-	}
-}
-
-uint32 FPhysicsInterface_PhysX::GetSolverVelocityIterationCount_AssumesLocked(const FPhysicsActorHandle_PhysX& InHandle)
-{
-	if(PxRigidDynamic* Dynamic = GetPx<PxRigidDynamic>(InHandle))
-	{
-		uint32 PositionIters = 0;
-		uint32 VelocityIters = 0;
-		Dynamic->getSolverIterationCounts(PositionIters, VelocityIters);
-
-		return VelocityIters;
-	}
-
-	return 0;
-}
-
-void FPhysicsInterface_PhysX::SetSolverVelocityIterationCount_AssumesLocked(const FPhysicsActorHandle_PhysX& InHandle, uint32 InSolverIterationCount)
-{
-	if(PxRigidDynamic* Dynamic = GetPx<PxRigidDynamic>(InHandle))
-	{
-		uint32 PositionIters = 0;
-		uint32 VelocityIters = 0;
-		Dynamic->getSolverIterationCounts(PositionIters, VelocityIters);
-		Dynamic->setSolverIterationCounts(PositionIters, InSolverIterationCount);
-	}
-}
-
-float FPhysicsInterface_PhysX::GetWakeCounter_AssumesLocked(const FPhysicsActorHandle_PhysX& InHandle)
-{
-	if(PxRigidDynamic* Dynamic = GetPx<PxRigidDynamic>(InHandle))
-	{
-		return Dynamic->getWakeCounter();
-	}
-
-	return 0.0f;
-}
-
-void FPhysicsInterface_PhysX::SetWakeCounter_AssumesLocked(const FPhysicsActorHandle_PhysX& InHandle, float InWakeCounter)
-{
-	if(PxRigidDynamic* Dynamic = GetPx<PxRigidDynamic>(InHandle))
-	{
-		Dynamic->setWakeCounter(InWakeCounter);
-	}
-}
-
-SIZE_T FPhysicsInterface_PhysX::GetResourceSizeEx(const FPhysicsActorHandle_PhysX& InHandle)
-{
-	SIZE_T OutSize = 0;
-
-	if(InHandle.SyncActor)
-	{
-		OutSize += GetPhysxObjectSize(InHandle.SyncActor, FPhysxSharedData::Get().GetCollection());
-	}
-
-	return OutSize;
-}
-
-// Constraint free functions //////////////////////////////////////////////////////////////////////////
-
-const bool bDrivesUseAcceleration = true;
-
-bool GetSceneForConstraintActors_LockFree(const FPhysicsActorHandle_PhysX& InActor1, const FPhysicsActorHandle_PhysX& InActor2, PxScene* OutScene)
-{
-	PxScene* Scene1 = GetPxSceneForPhysActor(&InActor1);
-	PxScene* Scene2 = GetPxSceneForPhysActor(&InActor2);
-
-	OutScene = nullptr;
-
-	if(Scene1 && Scene2)
-	{
-		if(Scene1 != Scene2)
-		{
-#if !(UE_BUILD_SHIPPING || UE_BUILD_TEST)
-			FMessageLog("PIE").Warning()
-				->AddToken(FTextToken::Create(LOCTEXT("JointBetweenScenesStart", "Constraint")))
-				->AddToken(FTextToken::Create(LOCTEXT("JointBetweenScenesMid", "attempting to create a joint between two actors in different scenes (")))
-				->AddToken(FTextToken::Create(LOCTEXT("JointBetweenScenesEnd", ").  No joint created.")));
-#endif // !(UE_BUILD_SHIPPING || UE_BUILD_TEST)
-			return false;
-		}
-
-		OutScene = Scene1 ? Scene1 : Scene2;
-	}
-
-	return true;
-}
-
-void GetSoftLimitParams_Linear(float& InOutDamping, float& InOutStiffness)
-{
-	InOutDamping = InOutDamping * CVarConstraintLinearDampingScale.GetValueOnGameThread();
-	InOutStiffness = InOutStiffness * CVarConstraintLinearStiffnessScale.GetValueOnGameThread();
-}
-
-void GetSoftLimitParams_Angular(float& InOutDamping, float& InOutStiffness)
-{
-	InOutDamping = InOutDamping * CVarConstraintAngularDampingScale.GetValueOnGameThread();
-	InOutStiffness = InOutStiffness * CVarConstraintAngularStiffnessScale.GetValueOnGameThread();
-}
-
-void WakeActor_AssumesLocked(PxRigidActor* InActor)
-{
-	PxRigidDynamic* Dynamic = InActor ? InActor->is<PxRigidDynamic>() : nullptr;
-
-	if(Dynamic && Dynamic->getScene() && !(Dynamic->getRigidBodyFlags() & PxRigidBodyFlag::eKINEMATIC))
-	{
-		Dynamic->wakeUp();
-	}
-}
-
-void WakeupJointedActors_AssumesLocked(const FPhysicsConstraintHandle_PhysX& InHandle)
-{
-	PxD6Joint* Joint = InHandle.ConstraintData;
-
-	if(Joint)
-	{
-		PxRigidActor* Actor1 = nullptr;
-		PxRigidActor* Actor2 = nullptr;
-
-		Joint->getActors(Actor1, Actor2);
-
-		WakeActor_AssumesLocked(Actor1);
-		WakeActor_AssumesLocked(Actor2);
-	}
-}
-
-void UpdateSingleDrive_AssumesLocked(const FPhysicsConstraintHandle_PhysX& InHandle, const FConstraintDrive& InDrive, PhysicsInterfaceTypes::EDriveType InDriveType, bool bDriveEnabled)
-{
-	PxD6Joint* Joint = InHandle.ConstraintData;
-
-	if(Joint)
-	{
-		PxD6Drive::Enum PxDriveType = U2PDriveType(InDriveType);
-
-		if(bDriveEnabled)
-		{
-			const float Stiffness = InDrive.bEnablePositionDrive ? InDrive.Stiffness : 0.0f;
-			const float Damping = InDrive.bEnableVelocityDrive ? InDrive.Damping : 0.0f;
-			const float MaxForce = InDrive.MaxForce > 0.0f ? InDrive.MaxForce : PX_MAX_F32;
-			Joint->setDrive(U2PDriveType(InDriveType), PxD6JointDrive(Stiffness, Damping, MaxForce, bDrivesUseAcceleration));
-		}
-		else
-		{
-			Joint->setDrive(U2PDriveType(InDriveType), PxD6JointDrive());
-		}
-	}
-}
-
-// Constraint interface functions
-
-FPhysicsConstraintHandle_PhysX FPhysicsInterface_PhysX::CreateConstraint(const FPhysicsActorHandle_PhysX& InActorHandle1, const FPhysicsActorHandle_PhysX& InActorHandle2, const FTransform& InLocalFrame1, const FTransform& InLocalFrame2)
-{
-	PxRigidActor* Actor1 = GetPxRigidActor_AssumesLocked(InActorHandle1);
-	PxRigidActor* Actor2 = GetPxRigidActor_AssumesLocked(InActorHandle2);
-
-	PxScene* PScene = nullptr;
-	if(GetSceneForConstraintActors_LockFree(InActorHandle1, InActorHandle2, PScene))
-	{
-		SCOPED_SCENE_WRITE_LOCK(PScene);
-
-		// Create the internal joint
-		PxD6Joint* NewJoint = PxD6JointCreate(*GPhysXSDK, Actor2, U2PTransform(InLocalFrame2), Actor1, U2PTransform(InLocalFrame1));
-
-		if(!NewJoint)
-		{
-			UE_LOG(LogPhysics, Log, TEXT("FPhysicsInterface_PhysX::CreateConstraint - Failed to create constraint."));
-			return FPhysicsConstraintHandle_PhysX();
-		}
-
-		FPhysicsConstraintHandle_PhysX OutReference;
-		OutReference.ConstraintData = NewJoint;
-		return OutReference;
-	}
-
-	// Return invalid constraint
-	return FPhysicsConstraintHandle_PhysX();
-}
-
-void FPhysicsInterface_PhysX::SetConstraintUserData(const FPhysicsConstraintHandle_PhysX& InHandle, void* InUserData)
-{
-	if(InHandle.ConstraintData)
-	{
-		PxScene* ConstraintScene = InHandle.ConstraintData->getScene();
-
-		if(ConstraintScene)
-		{
-			SCOPED_SCENE_WRITE_LOCK(ConstraintScene);
-			InHandle.ConstraintData->userData = InUserData;
-		}
-	}
-	else
-	{
-		UE_LOG(LogPhysics, Log, TEXT("Failed to set constraint data for an invalid constraint."));
-	}
-}
-
-void FPhysicsInterface_PhysX::ReleaseConstraint(FPhysicsConstraintHandle_PhysX& InHandle)
-{
-	if(InHandle.IsValid())
-	{
-		PxScene* PScene = InHandle.ConstraintData->getScene();
-		// Scene may be null if constraint was never actually added to a scene
-
-		SCOPED_SCENE_WRITE_LOCK(PScene);
-		InHandle.ConstraintData->release();
-	}
-	
-	InHandle.ConstraintData = nullptr;
-}
-
-FTransform FPhysicsInterface_PhysX::GetLocalPose(const FPhysicsConstraintHandle_PhysX& InHandle, EConstraintFrame::Type InFrame)
-{
-	PxD6Joint* Joint = InHandle.ConstraintData;
-
-	if(Joint)
-	{
-		PxJointActorIndex::Enum PxFrame = U2PConstraintFrame(InFrame);
-		return P2UTransform(Joint->getLocalPose(PxFrame));
-	}
-
-	return FTransform::Identity;
-}
-
-FTransform FPhysicsInterface_PhysX::GetGlobalPose(const FPhysicsConstraintHandle_PhysX& InHandle, EConstraintFrame::Type InFrame)
-{
-	PxD6Joint* Joint = InHandle.ConstraintData;
-
-	if(Joint)
-	{
-		PxJointActorIndex::Enum PxFrame = U2PConstraintFrame(InFrame);
-
-		PxRigidActor* Actor1 = nullptr;
-		PxRigidActor* Actor2 = nullptr;
-
-		Joint->getActors(Actor1, Actor2);
-
-		switch(InFrame)
-		{
-			case EConstraintFrame::Frame1:
-				if(Actor1)
-				{
-					return P2UTransform(Actor1->getGlobalPose());
-				}
-				break;
-			case EConstraintFrame::Frame2:
-				if(Actor2)
-				{
-					return P2UTransform(Actor2->getGlobalPose());
-				}
-				break;
-			default:
-				break;
-		}
-	}
-
-	return FTransform::Identity;
-}
-
-FVector FPhysicsInterface_PhysX::GetLocation(const FPhysicsConstraintHandle_PhysX& InHandle)
-{
-	PxD6Joint* Joint = InHandle.ConstraintData;
-
-	if(!Joint)
-	{
-		return FVector::ZeroVector;
-	}
-
-	PxRigidActor* Actor1 = nullptr;
-	PxRigidActor* Actor2 = nullptr;
-
-	Joint->getActors(Actor1, Actor2);
-
-	PxVec3 Location(0);
-
-	if(Actor1)
-	{
-		Location = Actor1->getGlobalPose().transform(Joint->getLocalPose(PxJointActorIndex::eACTOR0).p);
-	}
-
-	if(Actor2)
-	{
-		Location += Actor2->getGlobalPose().transform(Joint->getLocalPose(PxJointActorIndex::eACTOR1).p);
-	}
-
-	Location *= 0.5f;
-
-	return P2UVector(Location);
-}
-
-void FPhysicsInterface_PhysX::GetForce(const FPhysicsConstraintHandle_PhysX& InHandle, FVector& OutLinForce, FVector& OutAngForce)
-{
-	PxD6Joint* Joint = InHandle.ConstraintData;
-
-	OutLinForce = FVector::ZeroVector;
-	OutAngForce = FVector::ZeroVector;
-
-	if(Joint)
-	{
-		ExecuteOnUnbrokenConstraintReadWrite(InHandle, [&](const FPhysicsConstraintHandle_PhysX& InUnbrokenConstraint)
-		{
-			PxVec3 PxLinForce;
-			PxVec3 PxAngForce;
-			InUnbrokenConstraint.ConstraintData->getConstraint()->getForce(PxLinForce, PxAngForce);
-
-			OutLinForce = P2UVector(PxLinForce);
-			OutAngForce = P2UVector(PxAngForce);
-		});
-	}
-}
-
-void FPhysicsInterface_PhysX::GetDriveLinearVelocity(const FPhysicsConstraintHandle_PhysX& InHandle, FVector& OutLinVelocity)
-{
-	PxD6Joint* Joint = InHandle.ConstraintData;
-
-	OutLinVelocity = FVector::ZeroVector;
-
-	if (Joint)
-	{
-		ExecuteOnUnbrokenConstraintReadWrite(InHandle, [&](const FPhysicsConstraintHandle_PhysX& InUnbrokenConstraint)
-		{
-			PxVec3 PxLinVelocity;
-			PxVec3 PxAngVelocity;
-			InUnbrokenConstraint.ConstraintData->getDriveVelocity(PxLinVelocity, PxAngVelocity);
-
-			OutLinVelocity = P2UVector(PxLinVelocity);
-		});
-	}
-}
-
-void FPhysicsInterface_PhysX::GetDriveAngularVelocity(const FPhysicsConstraintHandle_PhysX& InHandle, FVector& OutAngVelocity)
-{
-	PxD6Joint* Joint = InHandle.ConstraintData;
-
-	OutAngVelocity = FVector::ZeroVector;
-
-	if (Joint)
-	{
-		ExecuteOnUnbrokenConstraintReadWrite(InHandle, [&](const FPhysicsConstraintHandle_PhysX& InUnbrokenConstraint)
-		{
-			PxVec3 PxLinVelocity;
-			PxVec3 PxAngVelocity;
-			InUnbrokenConstraint.ConstraintData->getDriveVelocity(PxLinVelocity, PxAngVelocity);
-
-			OutAngVelocity = P2UVector(PxAngVelocity);
-		});
-	}
-}
-
-float FPhysicsInterface_PhysX::GetCurrentSwing1(const FPhysicsConstraintHandle_PhysX& InHandle)
-{
-	float Swing1 = 0;
-
-	ExecuteOnUnbrokenConstraintReadOnly(InHandle, [&Swing1](const FPhysicsConstraintHandle_PhysX& InUnbrokenConstraint)
-	{
-		PxD6Joint* Joint = InUnbrokenConstraint.ConstraintData;
-		if(Joint)
-		{
-			Swing1 = Joint->getSwingZAngle();
-		}
-	});
-
-	return Swing1;
-}
-
-float FPhysicsInterface_PhysX::GetCurrentSwing2(const FPhysicsConstraintHandle_PhysX& InHandle)
-{
-	float Swing2 = 0;
-
-	ExecuteOnUnbrokenConstraintReadOnly(InHandle, [&Swing2](const FPhysicsConstraintHandle_PhysX& InUnbrokenConstraint)
-	{
-		PxD6Joint* Joint = InUnbrokenConstraint.ConstraintData;
-		if(Joint)
-		{
-			Swing2 = Joint->getSwingYAngle();
-		}
-	});
-
-	return Swing2;
-}
-
-float FPhysicsInterface_PhysX::GetCurrentTwist(const FPhysicsConstraintHandle_PhysX& InHandle)
-{
-	float Twist = 0;
-
-	ExecuteOnUnbrokenConstraintReadOnly(InHandle, [&Twist](const FPhysicsConstraintHandle_PhysX& InUnbrokenConstraint)
-	{
-		PxD6Joint* Joint = InUnbrokenConstraint.ConstraintData;
-		if(Joint)
-		{
-			Twist = Joint->getTwist();
-		}
-	});
-
-	return Twist;
-}
-
-void FPhysicsInterface_PhysX::SetCanVisualize(const FPhysicsConstraintHandle_PhysX& InHandle, bool bInCanVisualize)
-{
-	if(PxD6Joint* Joint = InHandle.ConstraintData)
-	{
-		Joint->setConstraintFlag(PxConstraintFlag::eVISUALIZATION, bInCanVisualize);
-	}
-}
-
-void FPhysicsInterface_PhysX::SetCollisionEnabled(const FPhysicsConstraintHandle_PhysX& InHandle, bool bInCollisionEnabled)
-{
-	if(PxD6Joint* Joint = InHandle.ConstraintData)
-	{
-		Joint->setConstraintFlag(PxConstraintFlag::eCOLLISION_ENABLED, bInCollisionEnabled);
-	}
-}
-
-void FPhysicsInterface_PhysX::SetProjectionEnabled_AssumesLocked(const FPhysicsConstraintHandle_PhysX& InHandle, bool bInProjectionEnabled, float InLinearTolerance, float InAngularToleranceDegrees)
-{
-	if(PxD6Joint* Joint = InHandle.ConstraintData)
-	{
-		Joint->setConstraintFlag(PxConstraintFlag::ePROJECTION, bInProjectionEnabled);
-		Joint->setProjectionLinearTolerance(InLinearTolerance);
-		Joint->setProjectionAngularTolerance(FMath::DegreesToRadians(InAngularToleranceDegrees));
-	}
-}
-
-void FPhysicsInterface_PhysX::SetParentDominates_AssumesLocked(const FPhysicsConstraintHandle_PhysX& InHandle, bool bInParentDominates)
-{
-	if(PxD6Joint* Joint = InHandle.ConstraintData)
-	{
-		float InertiaScale = bInParentDominates ? 0.0f : 1.0f;
-
-		Joint->setInvMassScale0(InertiaScale);
-		Joint->setInvMassScale1(1.0f);
-
-		Joint->setInvInertiaScale0(InertiaScale);
-		Joint->setInvInertiaScale1(1.0f);
-	}
-}
-
-void FPhysicsInterface_PhysX::SetBreakForces_AssumesLocked(const FPhysicsConstraintHandle_PhysX& InHandle, float InLinearBreakForce, float InAngularBreakForce)
-{
-	PxD6Joint* Joint = InHandle.ConstraintData;
-
-	if(Joint)
-	{
-		Joint->setBreakForce(InLinearBreakForce, InAngularBreakForce);
-	}
-}
-
-void FPhysicsInterface_PhysX::SetLocalPose(const FPhysicsConstraintHandle_PhysX& InHandle, const FTransform& InPose, EConstraintFrame::Type InFrame)
-{
-	PxD6Joint* Joint = InHandle.ConstraintData;
-	if(Joint)
-	{
-		PxJointActorIndex::Enum PxFrame = U2PConstraintFrame(InFrame);
-		PxTransform PxPose = U2PTransform(InPose);
-		Joint->setLocalPose(PxFrame, PxPose);
-	}
-}
-
-void FPhysicsInterface_PhysX::SetLinearMotionLimitType_AssumesLocked(const FPhysicsConstraintHandle_PhysX& InHandle, PhysicsInterfaceTypes::ELimitAxis InAxis, ELinearConstraintMotion InMotion)
-{
-	PxD6Joint* Joint = InHandle.ConstraintData;
-
-	if(Joint)
-	{
-		Joint->setMotion(U2PConstraintAxis(InAxis), U2PLinearMotion(InMotion));
-	}
-}
-
-void FPhysicsInterface_PhysX::SetAngularMotionLimitType_AssumesLocked(const FPhysicsConstraintHandle_PhysX& InHandle, PhysicsInterfaceTypes::ELimitAxis InAxis, EAngularConstraintMotion InMotion)
-{
-	PxD6Joint* Joint = InHandle.ConstraintData;
-
-	if(Joint)
-	{
-		Joint->setMotion(U2PConstraintAxis(InAxis), U2PAngularMotion(InMotion));
-	}
-}
-
-void FPhysicsInterface_PhysX::UpdateLinearLimitParams_AssumesLocked(const FPhysicsConstraintHandle_PhysX& InHandle, float InLimit, float InAverageMass, const FLinearConstraint& InParams)
-{
-	PxD6Joint* Joint = InHandle.ConstraintData;
-
-	if(Joint)
-	{
-		PxJointLinearLimit PxLinLimit(GPhysXSDK->getTolerancesScale(), InLimit, FMath::Clamp(InParams.ContactDistance, 5.f, InLimit * 0.49f));
-		PxLinLimit.restitution = InParams.Restitution;
-
-		if(InParams.bSoftConstraint)
-		{
-			PxLinLimit.damping = InParams.Damping * InAverageMass;
-			PxLinLimit.stiffness = InParams.Stiffness * InAverageMass;
-			GetSoftLimitParams_Linear(PxLinLimit.damping, PxLinLimit.stiffness);
-		}
-
-		Joint->setLinearLimit(PxLinLimit);
-	}
-}
-
-void FPhysicsInterface_PhysX::UpdateConeLimitParams_AssumesLocked(const FPhysicsConstraintHandle_PhysX& InHandle, float InAverageMass, const FConeConstraint& InParams)
-{
-	PxD6Joint* Joint = InHandle.ConstraintData;
-
-	if(Joint)
-	{
-		//Clamp the limit value to valid range which PhysX won't ignore, both value have to be clamped even if there is only one degree limit in constraint
-		const float Limit1Rad = FMath::DegreesToRadians(FMath::ClampAngle(InParams.Swing1LimitDegrees, KINDA_SMALL_NUMBER, 179.9999f));
-		const float Limit2Rad = FMath::DegreesToRadians(FMath::ClampAngle(InParams.Swing2LimitDegrees, KINDA_SMALL_NUMBER, 179.9999f));
-
-		//Clamp the contact distance so that it's not too small (jittery joint) or too big (always active joint)
-		const float ContactRad = FMath::DegreesToRadians(FMath::Clamp(InParams.ContactDistance, 1.0f, FMath::Min(InParams.Swing1LimitDegrees, InParams.Swing2LimitDegrees) * 0.49f));
-
-		PxJointLimitCone PxLimitCone(Limit2Rad, Limit1Rad, ContactRad);
-		PxLimitCone.restitution = InParams.Restitution;
-
-		if(InParams.bSoftConstraint)
-		{
-			PxLimitCone.damping = InParams.Damping * InAverageMass;
-			PxLimitCone.stiffness = InParams.Stiffness * InAverageMass;
-			GetSoftLimitParams_Angular(PxLimitCone.damping, PxLimitCone.stiffness);
-		}
-
-		Joint->setSwingLimit(PxLimitCone);
-	}
-}
-
-void FPhysicsInterface_PhysX::UpdateTwistLimitParams_AssumesLocked(const FPhysicsConstraintHandle_PhysX& InHandle, float InAverageMass, const FTwistConstraint& InParams)
-{
-	PxD6Joint* Joint = InHandle.ConstraintData;
-
-	if(Joint)
-	{
-		const float  TwistLimitRad = FMath::DegreesToRadians(InParams.TwistLimitDegrees);
-
-		//Clamp the contact distance so that it's not too small (jittery joint) or too big (always active joint)
-		const float ContactRad = FMath::DegreesToRadians(FMath::Clamp(InParams.ContactDistance, 1.f, InParams.TwistLimitDegrees * .95f));
-
-		PxJointAngularLimitPair PxLimitTwist(-TwistLimitRad, TwistLimitRad, ContactRad);
-		PxLimitTwist.restitution = InParams.Restitution;
-
-		if(InParams.bSoftConstraint)
-		{
-			PxLimitTwist.damping = InParams.Damping * InAverageMass;
-			PxLimitTwist.stiffness = InParams.Stiffness * InAverageMass;
-			GetSoftLimitParams_Angular(PxLimitTwist.damping, PxLimitTwist.stiffness);
-		}
-
-		Joint->setTwistLimit(PxLimitTwist);
-	}
-}
-
-void FPhysicsInterface_PhysX::UpdateLinearDrive_AssumesLocked(const FPhysicsConstraintHandle_PhysX& InHandle, const FLinearDriveConstraint& InDriveParams)
-{
-	UpdateSingleDrive_AssumesLocked(InHandle, InDriveParams.XDrive, EDriveType::X, true);
-	UpdateSingleDrive_AssumesLocked(InHandle, InDriveParams.YDrive, EDriveType::Y, true);
-	UpdateSingleDrive_AssumesLocked(InHandle, InDriveParams.ZDrive, EDriveType::Z, true);
-
-	PxD6Joint* Joint = InHandle.ConstraintData;
-	if(Joint)
-	{
-		WakeupJointedActors_AssumesLocked(InHandle);
-	}
-}
-
-void FPhysicsInterface_PhysX::UpdateAngularDrive_AssumesLocked(const FPhysicsConstraintHandle& InHandle, const FAngularDriveConstraint& InDriveParams)
-{
-	PxD6Joint* Joint = InHandle.ConstraintData;
-
-	const bool bUseSlerp = InDriveParams.AngularDriveMode == EAngularDriveMode::SLERP;
-	UpdateSingleDrive_AssumesLocked(InHandle, InDriveParams.SlerpDrive, EDriveType::Slerp, bUseSlerp);
-	UpdateSingleDrive_AssumesLocked(InHandle, InDriveParams.SwingDrive, EDriveType::Swing, !bUseSlerp);
-	UpdateSingleDrive_AssumesLocked(InHandle, InDriveParams.TwistDrive, EDriveType::Twist, !bUseSlerp);
-
-	if(Joint)
-	{
-		WakeupJointedActors_AssumesLocked(InHandle);
-	}
-}
-
-void FPhysicsInterface_PhysX::UpdateDriveTarget_AssumesLocked(const FPhysicsConstraintHandle_PhysX& InHandle, const FLinearDriveConstraint& InLinDrive, const FAngularDriveConstraint& InAngDrive, bool bInitialized_IGNORED)
-{
-	PxD6Joint* Joint = InHandle.ConstraintData;
-
-	if(Joint)
-	{
-		const FQuat OrientationTarget(InAngDrive.OrientationTarget);
-
-		// Convert revolutions to radians
-		FVector AngularVelocityRads = InAngDrive.AngularVelocityTarget * 2.0f * PI;
-
-		Joint->setDrivePosition(PxTransform(U2PVector(InLinDrive.PositionTarget), U2PQuat(OrientationTarget)));
-		Joint->setDriveVelocity(U2PVector(InLinDrive.VelocityTarget), U2PVector(AngularVelocityRads));
-	}
-}
-
-void FPhysicsInterface_PhysX::SetDrivePosition(const FPhysicsConstraintHandle_PhysX& InHandle, const FVector& InPosition)
-{
-	ExecuteOnUnbrokenConstraintReadWrite(InHandle, [&InPosition](const FPhysicsConstraintHandle_PhysX& InUnbrokenConstraint)
-	{
-		PxD6Joint* Joint = InUnbrokenConstraint.ConstraintData;
-
-		if(Joint)
-		{
-			Joint->setDrivePosition(PxTransform(U2PVector(InPosition), Joint->getDrivePosition().q));
-		}
-	});
-}
-
-void FPhysicsInterface_PhysX::SetDriveOrientation(const FPhysicsConstraintHandle_PhysX& InHandle, const FQuat& InOrientation)
-{
-	ExecuteOnUnbrokenConstraintReadWrite(InHandle, [&InOrientation](const FPhysicsConstraintHandle_PhysX& InUnbrokenConstraint)
-	{
-		PxD6Joint* Joint = InUnbrokenConstraint.ConstraintData;
-
-		if(Joint)
-		{
-			Joint->setDrivePosition(PxTransform(Joint->getDrivePosition().p, U2PQuat(InOrientation)));
-		}
-	});
-}
-
-void FPhysicsInterface_PhysX::SetDriveLinearVelocity(const FPhysicsConstraintHandle_PhysX& InHandle, const FVector& InLinVelocity)
-{
-	ExecuteOnUnbrokenConstraintReadWrite(InHandle, [&InLinVelocity](const FPhysicsConstraintHandle_PhysX& InUnbrokenConstraint)
-	{
-		PxD6Joint* Joint = InUnbrokenConstraint.ConstraintData;
-
-		if (Joint)
-		{
-			PxVec3 VelLin;
-			PxVec3 VelAng;
-			Joint->getDriveVelocity(VelLin, VelAng);
-			Joint->setDriveVelocity(U2PVector(InLinVelocity), VelAng);
-		}
-	});
-}
-
-void FPhysicsInterface_PhysX::SetDriveAngularVelocity(const FPhysicsConstraintHandle_PhysX& InHandle, const FVector& InAngVelocity)
-{
-	ExecuteOnUnbrokenConstraintReadWrite(InHandle, [&InAngVelocity](const FPhysicsConstraintHandle_PhysX& InUnbrokenConstraint)
-	{
-		PxD6Joint* Joint = InUnbrokenConstraint.ConstraintData;
-
-		if (Joint)
-		{
-			PxVec3 VelLin;
-			PxVec3 VelAng;
-			Joint->getDriveVelocity(VelLin, VelAng);
-			Joint->setDriveVelocity(VelLin, U2PVector(InAngVelocity));
-		}
-	});
-}
-
-void FPhysicsInterface_PhysX::SetTwistLimit(const FPhysicsConstraintHandle_PhysX& InHandle, float InLowerLimit, float InUpperLimit, float InContactDistance)
-{
-	PxD6Joint* Joint = InHandle.ConstraintData;
-	if(Joint)
-	{
-		Joint->setTwistLimit(PxJointAngularLimitPair(InLowerLimit, InUpperLimit, InContactDistance));
-	}
-}
-
-void FPhysicsInterface_PhysX::SetSwingLimit(const FPhysicsConstraintHandle_PhysX& InHandle, float InYLimit, float InZLimit, float InContactDistance)
-{
-	PxD6Joint* Joint = InHandle.ConstraintData;
-	if(Joint)
-	{
-		Joint->setSwingLimit(PxJointLimitCone(InYLimit, InZLimit, InContactDistance));
-	}
-}
-
-void FPhysicsInterface_PhysX::SetLinearLimit(const FPhysicsConstraintHandle_PhysX& InHandle, float InLimit)
-{
-	PxD6Joint* Joint = InHandle.ConstraintData;
-	if(Joint)
-	{
-		PxReal LimitContactDistance = 1.f * (PI / 180.0f);
-		const PxTolerancesScale& ToleranceScale = GPhysXSDK->getTolerancesScale();
-		Joint->setLinearLimit(PxJointLinearLimit(ToleranceScale, InLimit, LimitContactDistance * ToleranceScale.length)); // LOC_MOD33 need to scale the contactDistance if not using its default value
-	}
-}
-
-bool FPhysicsInterface_PhysX::IsBroken(const FPhysicsConstraintHandle_PhysX& InHandle)
-{
-	PxD6Joint* Joint = InHandle.ConstraintData;
-
-	if(Joint)
-	{
-		SCOPED_SCENE_READ_LOCK(Joint->getScene());
-		return Joint->getConstraintFlags() & PxConstraintFlag::eBROKEN;
-	}
-	return false;
-}
-
-bool FPhysicsInterface_PhysX::ExecuteOnUnbrokenConstraintReadOnly(const FPhysicsConstraintHandle_PhysX& InHandle, TFunctionRef<void(const FPhysicsConstraintHandle_PhysX&)> Func)
-{
-	PxD6Joint* Joint = InHandle.ConstraintData;
-	if(Joint)
-	{
-		SCOPED_SCENE_READ_LOCK(Joint->getScene());
-
-		if(!(Joint->getConstraintFlags()&PxConstraintFlag::eBROKEN))
-		{
-			Func(InHandle);
-			return true;
-		}
-	}
-
-	return false;
-}
-
-bool FPhysicsInterface_PhysX::ExecuteOnUnbrokenConstraintReadWrite(const FPhysicsConstraintHandle_PhysX& InHandle, TFunctionRef<void(const FPhysicsConstraintHandle_PhysX&)> Func)
-{
-	PxD6Joint* Joint = InHandle.ConstraintData;
-	if(Joint)
-	{
-		SCOPED_SCENE_WRITE_LOCK(Joint->getScene());
-
-		if(!(Joint->getConstraintFlags()&PxConstraintFlag::eBROKEN))
-		{
-			Func(InHandle);
-			return true;
-		}
-	}
-
-	return false;
-}
-
-bool FPhysicsInterface_PhysX::LineTrace_Geom(FHitResult& OutHit, const FBodyInstance* InInstance, const FVector& InStart, const FVector& InEnd, bool bTraceComplex, bool bExtractPhysMaterial /*= false*/)
-{
-	// Need an instandce to trace against
-	check(InInstance);
-
-	OutHit.TraceStart = InStart;
-	OutHit.TraceEnd = InEnd;
-
-	bool bHitSomething = false;
-
-	const FVector Delta = InEnd - InStart;
-	const float DeltaMag = Delta.Size();
-	if(DeltaMag > KINDA_SMALL_NUMBER)
-	{
-		{
-			// #PHYS2 Really need a concept for "multi" locks here - as we're locking ActorRef but not TargetInstance->ActorRef
-			FPhysicsCommand::ExecuteRead(InInstance->ActorHandle, [&](const FPhysicsActorHandle& Actor)
-			{
-				// If we're welded then the target instance is actually our parent
-				const FBodyInstance* TargetInstance = InInstance->WeldParent ? InInstance->WeldParent : InInstance;
-
-				const PxRigidActor* RigidBody = FPhysicsInterface::GetPxRigidActor_AssumesLocked(TargetInstance->ActorHandle);
-
-				if((RigidBody != NULL) && (RigidBody->getNbShapes() != 0))
-				{
-					// Create filter data used to filter collisions, should always return eTOUCH for LineTraceComponent		
-					PxHitFlags PHitFlags = PxHitFlag::ePOSITION | PxHitFlag::eNORMAL | PxHitFlag::eDISTANCE | PxHitFlag::eFACE_INDEX;
-
-					PxRaycastHit BestHit;
-					float BestHitDistance = BIG_NUMBER;
-
-					// Get all the shapes from the actor
-					FInlineShapeArray PShapes;
-					const int32 NumShapes = FillInlineShapeArray_AssumesLocked(PShapes, Actor);
-
-					// Iterate over each shape
-					for(int32 ShapeIdx = 0; ShapeIdx < NumShapes; ShapeIdx++)
-					{
-						// #PHYS2 - SHAPES - Resolve this single cast case
-						FPhysicsShapeHandle_PhysX& ShapeRef = PShapes[ShapeIdx];
-						PxShape* PShape = ShapeRef.Shape;
-						check(PShape);
-
-						if(TargetInstance->IsShapeBoundToBody(ShapeRef) == false)
-						{
-							continue;
-						}
-
-						const PxU32 HitBufferSize = 1;
-						PxRaycastHit PHits[HitBufferSize];
-
-						// Filter so we trace against the right kind of collision
-						PxFilterData ShapeFilter = PShape->getQueryFilterData();
-						const bool bShapeIsComplex = (ShapeFilter.word3 & EPDF_ComplexCollision) != 0;
-						const bool bShapeIsSimple = (ShapeFilter.word3 & EPDF_SimpleCollision) != 0;
-						if((bTraceComplex && bShapeIsComplex) || (!bTraceComplex && bShapeIsSimple))
-						{
-							const int32 ArraySize = UE_ARRAY_COUNT(PHits);
-							
-							PxTransform ShapeTransform = U2PTransform(GetTransform(ShapeRef));
-							const PxI32 NumHits = PxGeometryQuery::raycast(U2PVector(InStart), U2PVector(Delta / DeltaMag), PShape->getGeometry().any(), ShapeTransform, DeltaMag, PHitFlags, ArraySize, PHits);
-
-							if(ensure(NumHits <= ArraySize))
-							{
-								for(int HitIndex = 0; HitIndex < NumHits; HitIndex++)
-								{
-									PxRaycastHit& Hit = PHits[HitIndex];
-									if(Hit.distance < BestHitDistance)
-									{
-										BestHitDistance = PHits[HitIndex].distance;
-										BestHit = PHits[HitIndex];
-
-										// we don't get Shape information when we access via PShape, so I filled it up
-										BestHit.shape = PShape;
-
-										BestHit.actor = PShape->getActor();
-									}
-								}
-							}
-						}
-					}
-
-					if(BestHitDistance < BIG_NUMBER)
-					{
-						// we just like to make sure if the hit is made, set to test touch
-						FCollisionFilterData QueryFilter;
-						QueryFilter.Word2 = 0xFFFFF;
-
-						FTransform StartTM(InStart);
-						const UPrimitiveComponent* OwnerComponentInst = InInstance->OwnerComponent.Get();
-						ConvertQueryImpactHit(OwnerComponentInst ? OwnerComponentInst->GetWorld() : nullptr, BestHit, OutHit, DeltaMag, QueryFilter, InStart, InEnd, nullptr, StartTM, true, bExtractPhysMaterial);
-						bHitSomething = true;
-					}
-				}
-			});
-		}
-	}
-
-	return bHitSomething;
-}
-
-bool FPhysicsInterface_PhysX::Sweep_Geom(FHitResult& OutHit, const FBodyInstance* InInstance, const FVector& InStart, const FVector& InEnd, const FQuat& InShapeRotation, const FCollisionShape& InShape, bool bSweepComplex)
-{
-	bool bSweepHit = false;
-
-	if(InShape.IsNearlyZero())
-	{
-		bSweepHit = LineTrace_Geom(OutHit, InInstance, InStart, InEnd, bSweepComplex);
-	}
-	else
-	{
-		OutHit.TraceStart = InStart;
-		OutHit.TraceEnd = InEnd;
-
-		const FBodyInstance* TargetInstance = InInstance->WeldParent ? InInstance->WeldParent : InInstance;
-
-		FPhysicsCommand::ExecuteRead(TargetInstance->ActorHandle, [&](const FPhysicsActorHandle& Actor)
-		{
-			const PxRigidActor* RigidBody = FPhysicsInterface::GetPxRigidActor_AssumesLocked(Actor);
-
-			if((RigidBody != nullptr) && (RigidBody->getNbShapes() != 0) && (InInstance->OwnerComponent != nullptr))
-			{
-				FPhysXShapeAdapter ShapeAdapter(InShapeRotation, InShape);
-				
-				const FVector Delta = InEnd - InStart;
-				const float DeltaMag = Delta.Size();
-				if(DeltaMag > KINDA_SMALL_NUMBER)
-				{
-					PxHitFlags POutputFlags = PxHitFlag::ePOSITION | PxHitFlag::eNORMAL | PxHitFlag::eDISTANCE | PxHitFlag::eFACE_INDEX | PxHitFlag::eMTD;
-
-					UPrimitiveComponent* OwnerComponentInst = InInstance->OwnerComponent.Get();
-					PxTransform PStartTM(U2PVector(InStart), U2PQuat(ShapeAdapter.GetGeomOrientation()));
-					
-					PxVec3 PDir = U2PVector(Delta / DeltaMag);
-
-					PxSweepHit PHit;
-
-					// Get all the shapes from the actor
-					FInlineShapeArray PShapes;
-					// #PHYS2 - SHAPES - Resolve this function to not use px stuff
-					const int32 NumShapes = FillInlineShapeArray_AssumesLocked(PShapes, Actor); // #PHYS2 - Need a lock/execute here?
-
-					// Iterate over each shape
-					for(int32 ShapeIdx = 0; ShapeIdx < NumShapes; ShapeIdx++)
-					{
-						FPhysicsShapeHandle_PhysX& ShapeRef = PShapes[ShapeIdx];
-						PxShape* PShape = ShapeRef.Shape;
-						check(PShape);
-
-						// Skip shapes not bound to this instance
-						// We check InInstance here because we want the shapes that were created for that instance
-						// and not just the shapes created for the weld parent (if InInstance is welded)
-						if(!InInstance->IsShapeBoundToBody(ShapeRef))
-						{
-							continue;
-						}
-
-						// Filter so we trace against the right kind of collision
-						FCollisionFilterData ShapeFilter = P2UFilterData(PShape->getQueryFilterData());
-						const bool bShapeIsComplex = (ShapeFilter.Word3 & EPDF_ComplexCollision) != 0;
-						const bool bShapeIsSimple = (ShapeFilter.Word3 & EPDF_SimpleCollision) != 0;
-						if((bSweepComplex && bShapeIsComplex) || (!bSweepComplex && bShapeIsSimple))
-						{
-							PxTransform PGlobalPose = RigidBody->getGlobalPose().transform(PShape->getLocalPose());
-							const PxGeometry& Geometry = ShapeAdapter.GetGeometry();
-							if(PxGeometryQuery::sweep(PDir, DeltaMag, Geometry, PStartTM, PShape->getGeometry().any(), PGlobalPose, PHit, POutputFlags))
-							{
-								// we just like to make sure if the hit is made
-								FCollisionFilterData QueryFilter;
-								QueryFilter.Word2 = 0xFFFFF;
-
-								// we don't get Shape information when we access via PShape, so I filled it up
-								PHit.shape = PShape;
-								PHit.actor = PShape->getActor();
-
-								FTransform StartTransform(InStart);
-								PHit.faceIndex = FindFaceIndex(PHit, P2UVector(PDir));
-								ConvertQueryImpactHit(OwnerComponentInst->GetWorld(), PHit, OutHit, DeltaMag, QueryFilter, InStart, InEnd, nullptr, StartTransform, false, false);
-								bSweepHit = true;
-							}
-						}
-					}
-				}
-			}
-		});
-	}
-
-	return bSweepHit;
-}
-
-bool Overlap_GeomInternal(const FBodyInstance* InInstance, const PxGeometry& InPxGeom, const FTransform& InShapeTransform, FMTDResult* OutOptResult)
-{
-	PxTransform ShapePose = U2PTransform(InShapeTransform);
-	const FBodyInstance* TargetInstance = InInstance->WeldParent ? InInstance->WeldParent : InInstance;
-	PxRigidActor* RigidBody = FPhysicsInterface::GetPxRigidActor_AssumesLocked(TargetInstance->ActorHandle);
-
-	if(RigidBody == NULL || RigidBody->getNbShapes() == 0)
-	{
-		return false;
-	}
-
-	// Get all the shapes from the actor
-	FInlineShapeArray PShapes;
-	const int32 NumShapes = FillInlineShapeArray_AssumesLocked(PShapes, TargetInstance->ActorHandle);
-
-	// Iterate over each shape
-	for(int32 ShapeIdx = 0; ShapeIdx < NumShapes; ++ShapeIdx)
-	{
-		FPhysicsShapeHandle_PhysX& ShapeRef = PShapes[ShapeIdx];
-		const PxShape* PShape = ShapeRef.Shape;
-		check(PShape);
-
-		if(TargetInstance->IsShapeBoundToBody(ShapeRef))
-		{
-			PxVec3 POutDirection;
-			float OutDistance;
-
-			if(OutOptResult)
-			{
-				PxTransform PTransform = U2PTransform(FPhysicsInterface::GetTransform(ShapeRef));
-				if(PxGeometryQuery::computePenetration(POutDirection, OutDistance, InPxGeom, ShapePose, PShape->getGeometry().any(), PTransform))
-				{
-					//TODO: there are some edge cases that give us nan results. In these cases we skip
-					if(!POutDirection.isFinite())
-					{
-						POutDirection.x = 0.f;
-						POutDirection.y = 0.f;
-						POutDirection.z = 0.f;
-					}
-
-					OutOptResult->Direction = P2UVector(POutDirection);
-					OutOptResult->Distance = FMath::Abs(OutDistance);
-
-					return true;
-				}
-			}
-			else
-			{
-				PxTransform PTransform = U2PTransform(FPhysicsInterface::GetTransform(ShapeRef));
-				if(PxGeometryQuery::overlap(InPxGeom, ShapePose, PShape->getGeometry().any(), PTransform))
-				{
-					return true;
-				}
-			}
-
-		}
-	}
-
-	return false;
-}
-
-bool FPhysicsInterface_PhysX::Overlap_Geom(const FBodyInstance* InInstance, const FPhysicsGeometryCollection& InGeometry, const FTransform& InShapeTransform, FMTDResult* OutOptResult)
-{
-	const PxGeometry&  PGeom = InGeometry.GetGeometry();
-
-	return Overlap_GeomInternal(InInstance, PGeom, InShapeTransform, OutOptResult);
-}
-
-bool FPhysicsInterface_PhysX::Overlap_Geom(const FBodyInstance* InInstance, const FCollisionShape& InCollisionShape, const FQuat& InShapeRotation, const FTransform& InShapeTransform, FMTDResult* OutOptResult /*= nullptr*/)
-{
-	FPhysXShapeAdapter Adaptor(InShapeRotation, InCollisionShape);
-
-	return Overlap_GeomInternal(InInstance, Adaptor.GetGeometry(), Adaptor.GetGeomPose(InShapeTransform.GetTranslation()), OutOptResult);
-}
-
-bool FPhysicsInterface_PhysX::GetSquaredDistanceToBody(const FBodyInstance* InInstance, const FVector& InPoint, float& OutDistanceSquared, FVector* OutOptPointOnBody)
-{
-	if(OutOptPointOnBody)
-	{
-		*OutOptPointOnBody = InPoint;
-	}
-
-	float ReturnDistance = -1.f;
-	float MinDistanceSqr = BIG_NUMBER;
-	bool bFoundValidBody = false;
-	bool bEarlyOut = true;
-
-	const FBodyInstance* UseBI = InInstance->WeldParent ? InInstance->WeldParent : InInstance;
-
-	FPhysicsCommand::ExecuteRead(UseBI->ActorHandle, [&](const FPhysicsActorHandle& Actor)
-	{
-		const int32 NumShapes = FPhysicsInterface::GetNumShapes(Actor);
-
-		if(NumShapes == 0 || UseBI->OwnerComponent == NULL)
-		{
-			return;
-		}
-
-		bEarlyOut = false;
-
-		// Get all the shapes from the actor
-		FInlineShapeArray PShapes;
-		const int32 NumTotalShapes = FillInlineShapeArray_AssumesLocked(PShapes, Actor);
-
-		const PxVec3 PPoint = U2PVector(InPoint);
-
-		// Iterate over each shape
-		for(int32 ShapeIdx = 0; ShapeIdx < NumTotalShapes; ShapeIdx++)
-		{
-			// #PHYS2 - resolve px stuff here
-			FPhysicsShapeHandle_PhysX& ShapeRef = PShapes[ShapeIdx];
-			PxShape* PShape = ShapeRef.Shape;
-			check(PShape);
-
-			if(UseBI->IsShapeBoundToBody(ShapeRef) == false)	//skip welded shapes that do not belong to us
-			{
-				continue;
-			}
-
-			FPhysicsGeometryCollection GeoCollection = FPhysicsInterface::GetGeometryCollection(ShapeRef);
-
-			//PxTransform PGlobalPose = GetTransform_AssumesLocked(ShapeRef, Actor);
-			PxTransform PGlobalPose = U2PTransform(FPhysicsInterface::GetTransform(ShapeRef));
-
-			ECollisionShapeType GeomType = FPhysicsInterface::GetShapeType(ShapeRef);
-
-			if(GeomType == ECollisionShapeType::Trimesh)
-			{
-				// Type unsupported for this function, but some other shapes will probably work. 
-				continue;
-			}
-
-			bFoundValidBody = true;
-
-			PxVec3 PClosestPoint;
-			float SqrDistance = PxGeometryQuery::pointDistance(PPoint, GeoCollection.GetGeometry(), PGlobalPose, &PClosestPoint);
-			// distance has valid data and smaller than mindistance
-			if(SqrDistance > 0.f && MinDistanceSqr > SqrDistance)
-			{
-				MinDistanceSqr = SqrDistance;
-
-				if(OutOptPointOnBody)
-				{
-					*OutOptPointOnBody = P2UVector(PClosestPoint);
-				}
-			}
-			else if(SqrDistance == 0.f)
-			{
-				MinDistanceSqr = 0.f;
-				break;
-			}
-		}
-	});
-
-	if(!bFoundValidBody && !bEarlyOut)
-	{
-		UE_LOG(LogPhysics, Verbose, TEXT("GetDistanceToBody: Component (%s) has no simple collision and cannot be queried for closest point."), InInstance->OwnerComponent.Get() ? *(InInstance->OwnerComponent->GetPathName()) : TEXT("NONE"));
-	}
-
-	if(bFoundValidBody)
-	{
-		OutDistanceSquared = MinDistanceSqr;
-	}
-	return bFoundValidBody;
-}
-
-// #PHYS2 Want this gone eventually - need a better solution for mass properties
-void FPhysicsInterface_PhysX::CalculateMassPropertiesFromShapeCollection(PxMassProperties& OutProperties, const TArray<FPhysicsShapeHandle>& InShapes, float InDensityKGPerCM)
-{
-	TArray<PxShape*> PShapes;
-	PShapes.Reserve(InShapes.Num());
-
-	for(const FPhysicsShapeHandle& Shape : InShapes)
-	{
-		PShapes.Add(Shape.IsValid() ? Shape.Shape : nullptr);
-	}
-
-	OutProperties = PxRigidBodyExt::computeMassPropertiesFromShapes(PShapes.GetData(), PShapes.Num()) * InDensityKGPerCM;
-}
-
-FPhysicsGeometryCollection_PhysX::~FPhysicsGeometryCollection_PhysX() = default;
-FPhysicsGeometryCollection_PhysX::FPhysicsGeometryCollection_PhysX(FPhysicsGeometryCollection_PhysX&& Steal) = default;
-FPhysicsGeometryCollection_PhysX& FPhysicsGeometryCollection_PhysX::operator=(FPhysicsGeometryCollection_PhysX&& Steal) = default;
-
-ECollisionShapeType FPhysicsGeometryCollection_PhysX::GetType() const
-{
-	return P2UCollisionShapeType(GeomHolder->getType());
-}
-
-const physx::PxGeometry& FPhysicsGeometryCollection_PhysX::GetGeometry() const
-{
-	return GeomHolder->any();
-}
-
-physx::PxGeometry& FPhysicsGeometryCollection_PhysX::GetGeometry()
-{
-	return GeomHolder->any();
-}
-
-const physx::PxBoxGeometry& FPhysicsGeometryCollection_PhysX::GetBoxGeometry() const
-{
-	return GeomHolder->box();
-}
-
-physx::PxBoxGeometry& FPhysicsGeometryCollection_PhysX::GetBoxGeometry()
-{
-	return GeomHolder->box();
-}
-
-const physx::PxSphereGeometry& FPhysicsGeometryCollection_PhysX::GetSphereGeometry() const
-{
-	return GeomHolder->sphere();
-}
-
-physx::PxSphereGeometry& FPhysicsGeometryCollection_PhysX::GetSphereGeometry()
-{
-	return GeomHolder->sphere();
-}
-
-const physx::PxCapsuleGeometry& FPhysicsGeometryCollection_PhysX::GetCapsuleGeometry() const
-{
-	return GeomHolder->capsule();
-}
-
-physx::PxCapsuleGeometry& FPhysicsGeometryCollection_PhysX::GetCapsuleGeometry()
-{
-	return GeomHolder->capsule();
-}
-
-const physx::PxConvexMeshGeometry& FPhysicsGeometryCollection_PhysX::GetConvexGeometry() const
-{
-	return GeomHolder->convexMesh();
-}
-
-physx::PxConvexMeshGeometry& FPhysicsGeometryCollection_PhysX::GetConvexGeometry()
-{
-	return GeomHolder->convexMesh();
-}
-
-const physx::PxTriangleMeshGeometry& FPhysicsGeometryCollection_PhysX::GetTriMeshGeometry() const
-{
-	return GeomHolder->triangleMesh();
-}
-
-physx::PxTriangleMeshGeometry& FPhysicsGeometryCollection_PhysX::GetTriMeshGeometry()
-{
-	return GeomHolder->triangleMesh();
-}
-
-FPhysicsGeometryCollection_PhysX::FPhysicsGeometryCollection_PhysX(const FPhysicsShapeHandle_PhysX& ShapeRef)
-{
-	GeomHolder = MakeUnique<PxGeometryHolder>(ShapeRef.Shape->getGeometry());
-}
-
-#undef LOCTEXT_NAMESPACE
-
-#endif
-
-#endif
diff --git a/Engine/Source/Runtime/Engine/Private/PhysicsEngine/PhysicsReplication.cpp b/Engine/Source/Runtime/Engine/Private/PhysicsEngine/PhysicsReplication.cpp
--- a/Engine/Source/Runtime/Engine/Private/PhysicsEngine/PhysicsReplication.cpp
+++ b/Engine/Source/Runtime/Engine/Private/PhysicsEngine/PhysicsReplication.cpp
@@ -20,17 +20,11 @@
 #include "Components/SkeletalMeshComponent.h"
 #include "Engine/DemoNetDriver.h"
 
-#if WITH_CHAOS
 #include "Chaos/ChaosMarshallingManager.h"
 #include "PhysicsProxy/SingleParticlePhysicsProxy.h"
 #include "PBDRigidsSolver.h"
 #include "Chaos/PBDRigidsEvolutionGBF.h"
-#endif
 
-#if WITH_PHYSX
-#include "PhysXPublic.h"
-#endif 
-
 namespace CharacterMovementCVars
 {
 	extern int32 NetShowCorrections;
@@ -98,7 +92,6 @@
 #endif
 }
 
-#if WITH_CHAOS
 struct FAsyncPhysicsRepCallbackData : public Chaos::FSimCallbackInput
 {
 	TArray<FAsyncPhysicsDesiredState> Buffer;
@@ -117,7 +110,6 @@
 		FPhysicsReplication::ApplyAsyncDesiredState(GetDeltaTime_Internal(), GetConsumerInput_Internal());
 	}
 };
-#endif
 
 void ComputeDeltas(const FVector& CurrentPos, const FQuat& CurrentQuat, const FVector& TargetPos, const FQuat& TargetQuat, FVector& OutLinDiff, float& OutLinDiffSize,
 	FVector& OutAngDiffAxis, float& OutAngDiff, float& OutAngDiffSize)
@@ -133,7 +125,6 @@
 
 FPhysicsReplication::~FPhysicsReplication()
 {
-#if WITH_CHAOS
 	if (AsyncCallback)
 	{
 		if (auto* Solver = PhysScene->GetSolver())
@@ -141,7 +132,6 @@
 			Solver->UnregisterAndFreeSimCallbackObject_External(AsyncCallback);
 		}
 	}
-#endif
 }
 
 bool FPhysicsReplication::ApplyRigidBodyState(float DeltaSeconds, FBodyInstance* BI, FReplicatedPhysicsTarget& PhysicsTarget, const FRigidBodyErrorCorrection& ErrorCorrection, const float InPingSecondsOneWay, int32 LocalFrame, int32 NumPredictedFrames)
@@ -397,9 +387,7 @@
 		else
 		{
 			// Small enough error to interpolate
-#if WITH_CHAOS
 			if (AsyncCallback == nullptr)	//sync case
-#endif
 			{
 				const FVector NewLinVel = FVector(NewState.LinVel) + (LinDiff * LinearVelocityCoefficient * DeltaSeconds);
 				const FVector NewAngVel = FVector(NewState.AngVel) + (AngDiffAxis * AngDiff * AngularVelocityCoefficient * DeltaSeconds);
@@ -411,7 +399,6 @@
 				BI->SetLinearVelocity(NewLinVel, false);
 				BI->SetAngularVelocityInRadians(FMath::DegreesToRadians(NewAngVel), false);
 			}
-#if WITH_CHAOS
 			else
 			{
 				//If async is used, enqueue for callback
@@ -428,7 +415,6 @@
 			}
 
 
-#endif
 		}
 
 		// Should we show the async part?
@@ -450,8 +436,6 @@
 	}
 
 	/////// SLEEP UPDATE ///////
-
-#if WITH_CHAOS
 	if (bShouldSleep)
 	{
 		// In the async case, we apply sleep state in ApplyAsyncDesiredState
@@ -460,7 +444,6 @@
 			BI->PutInstanceToSleep();
 		}
 	}
-#endif
 
 	PhysicsTarget.PrevPosTarget = TargetPos;
 	PhysicsTarget.PrevPos = FVector(CurrentState.Position);
@@ -524,6 +507,8 @@
 
 void FPhysicsReplication::OnTick(float DeltaSeconds, TMap<TWeakObjectPtr<UPrimitiveComponent>, FReplicatedPhysicsTarget>& ComponentsToTargets)
 {
+	using namespace Chaos;
+
 	int32 LocalFrameOffset = 0;		// LocalFrame = ServerFrame + LocalFrameOffset;
 	int32 NumPredictedFrames = 0;	// How many frames "ahead" of the server we are predicting. That is, how many frames are in flight between us and the server.
 
@@ -544,13 +529,10 @@
 	}
 
 	const FRigidBodyErrorCorrection& PhysicErrorCorrection = UPhysicsSettings::Get()->PhysicErrorCorrection;
-#if WITH_CHAOS
-	using namespace Chaos;
 	if(AsyncCallback)
 	{
 		PrepareAsyncData_External(PhysicErrorCorrection);
 	}
-#endif
 
 	// Get the ping between this PC & the server
 	const float LocalPing = GetLocalPing();
@@ -608,9 +590,7 @@
 		}
 	}
 
-#if WITH_CHAOS
 	CurAsyncData = nullptr;
-#endif
 }
 
 void FPhysicsReplication::Tick(float DeltaSeconds)
@@ -621,7 +601,6 @@
 FPhysicsReplication::FPhysicsReplication(FPhysScene* InPhysicsScene)
 : PhysScene(InPhysicsScene)
 {
-#if WITH_CHAOS
 	using namespace Chaos;
 	CurAsyncData = nullptr;
 	AsyncCallback = nullptr;
@@ -629,11 +608,8 @@
 	{
 		AsyncCallback = Solver->CreateAndRegisterSimCallbackObject_External<FPhysicsReplicationAsyncCallback>();
 	}
-#endif
 }
 
-#if WITH_CHAOS
-
 void FPhysicsReplication::PrepareAsyncData_External(const FRigidBodyErrorCorrection& ErrorCorrection)
 {
 	//todo move this logic into a common function?
@@ -719,7 +695,6 @@
 		}
 	}
 }
-#endif
 
 void FPhysicsReplication::SetReplicatedTarget(UPrimitiveComponent* Component, FName BoneName, const FRigidBodyState& ReplicatedTarget, int32 ServerFrame)
 {
diff --git a/Engine/Source/Runtime/Engine/Private/PhysicsEngine/ScopedSQHitchRepeater.h b/Engine/Source/Runtime/Engine/Private/PhysicsEngine/ScopedSQHitchRepeater.h
--- a/Engine/Source/Runtime/Engine/Private/PhysicsEngine/ScopedSQHitchRepeater.h
+++ b/Engine/Source/Runtime/Engine/Private/PhysicsEngine/ScopedSQHitchRepeater.h
@@ -11,10 +11,6 @@
 #include "CollisionQueryFilterCallback.h"
 #include "WorldCollision.h"
 
-#if PHYSICS_INTERFACE_PHYSX
-#include "PhysXInterfaceWrapper.h"
-#endif
-
 #if DETECT_SQ_HITCHES
 struct FSQHitchRepeaterCVars
 {
diff --git a/Engine/Source/Runtime/Engine/Private/PrimitiveComponentPhysics.cpp b/Engine/Source/Runtime/Engine/Private/PrimitiveComponentPhysics.cpp
--- a/Engine/Source/Runtime/Engine/Private/PrimitiveComponentPhysics.cpp
+++ b/Engine/Source/Runtime/Engine/Private/PrimitiveComponentPhysics.cpp
@@ -14,10 +14,7 @@
 #include "PhysicsReplication.h"
 #include "Physics/PhysicsInterfaceCore.h"
 #include "UObject/UObjectThreadContext.h"
-
-#if WITH_CHAOS
 #include "PhysicsProxy/SingleParticlePhysicsProxy.h"
-#endif
 
 //////////////// PRIMITIVECOMPONENT ///////////////
 
@@ -46,9 +43,7 @@ void UPrimitiveComponent::SetRigidBodyReplicatedTarget(FRigidBodyState& UpdatedS
 				if (BI && BI->IsValidBodyInstance())
 				{
 					PhysicsReplication->SetReplicatedTarget(this, BoneName, UpdatedState, ServerFrame);
-#if WITH_CHAOS
 					BI->GetPhysicsActorHandle();// ->GetGameThreadAPI().SetParticleID(Chaos::FParticleID{ ServerPhysicsHandle, INDEX_NONE });
-#endif
 				}
 			}
 		}
