Physics interface cleanup.

diff --git a/Engine/Source/Runtime/Engine/Private/SkeletalMesh.cpp b/Engine/Source/Runtime/Engine/Private/SkeletalMesh.cpp
--- a/Engine/Source/Runtime/Engine/Private/SkeletalMesh.cpp
+++ b/Engine/Source/Runtime/Engine/Private/SkeletalMesh.cpp
@@ -72,10 +72,6 @@
 #include "MeshUtilities.h"
 #include "Engine/SkeletalMeshEditorData.h"
 
-#if WITH_APEX_CLOTHING
-#include "ApexClothingUtils.h"
-#endif
-
 #include "IMeshReductionManagerModule.h"
 #include "SkeletalMeshReductionSettings.h"
 #include "Engine/RendererSettings.h"
@@ -87,10 +83,6 @@
 #include "Interfaces/ITargetPlatform.h"
 #include "Interfaces/ITargetPlatformManagerModule.h"
 
-#if WITH_APEX
-#include "PhysXIncludes.h"
-#endif// #if WITH_APEX
-
 #include "EditorFramework/AssetImportData.h"
 #include "Engine/SkeletalMeshSocket.h"
 #include "Components/BrushComponent.h"
@@ -160,78 +152,6 @@ static FAutoConsoleVariableRef CVarSkeletalMeshMinLodQualityLevel(
 	TEXT("The quality level for the Min stripping LOD. \n"),
 	FConsoleVariableDelegate::CreateStatic(&USkeletalMesh::OnLodStrippingQualityLevelChanged),
 	ECVF_Scalability);
-	
-#if WITH_APEX_CLOTHING
-/*-----------------------------------------------------------------------------
-	utility functions for apex clothing 
------------------------------------------------------------------------------*/
-
-static apex::ClothingAsset* LoadApexClothingAssetFromBlob(const TArray<uint8>& Buffer)
-{
-	// Wrap this blob with the APEX read stream class
-	physx::PxFileBuf* Stream = GApexSDK->createMemoryReadStream( Buffer.GetData(), Buffer.Num() );
-	// Create an NvParameterized serializer
-	NvParameterized::Serializer* Serializer = GApexSDK->createSerializer(NvParameterized::Serializer::NST_BINARY);
-	// Deserialize into a DeserializedData buffer
-	NvParameterized::Serializer::DeserializedData DeserializedData;
-	Serializer->deserialize( *Stream, DeserializedData );
-	apex::Asset* ApexAsset = NULL;
-	if( DeserializedData.size() > 0 )
-	{
-		// The DeserializedData has something in it, so create an APEX asset from it
-		ApexAsset = GApexSDK->createAsset( DeserializedData[0], NULL);
-		// Make sure it's a Clothing asset
-		if (ApexAsset 
-			&& ApexAsset->getObjTypeID() != GApexModuleClothing->getModuleID()
-			)
-		{
-			GPhysCommandHandler->DeferredRelease(ApexAsset);
-			ApexAsset = NULL;
-		}
-	}
-
-	apex::ClothingAsset* ApexClothingAsset = static_cast<apex::ClothingAsset*>(ApexAsset);
-	// Release our temporary objects
-	Serializer->release();
-	GApexSDK->releaseMemoryReadStream( *Stream );
-
-	return ApexClothingAsset;
-}
-
-static bool SaveApexClothingAssetToBlob(const apex::ClothingAsset *InAsset, TArray<uint8>& OutBuffer)
-{
-	bool bResult = false;
-	uint32 Size = 0;
-	// Get the NvParameterized data for our Clothing asset
-	if( InAsset != NULL )
-	{
-		// Create an APEX write stream
-		physx::PxFileBuf* Stream = GApexSDK->createMemoryWriteStream();
-		// Create an NvParameterized serializer
-		NvParameterized::Serializer* Serializer = GApexSDK->createSerializer(NvParameterized::Serializer::NST_BINARY);
-
-		const NvParameterized::Interface* AssetParameterized = InAsset->getAssetNvParameterized();
-		if( AssetParameterized != NULL )
-		{
-			// Serialize the data into the stream
-			Serializer->serialize( *Stream, &AssetParameterized, 1 );
-			// Read the stream data into our buffer for UE serialzation
-			Size = Stream->getFileLength();
-			OutBuffer.AddUninitialized( Size );
-			Stream->read( OutBuffer.GetData(), Size );
-			bResult = true;
-		}
-
-		// Release our temporary objects
-		Serializer->release();
-		Stream->release();
-	}
-
-	return bResult;
-}
-
-#endif//#if WITH_APEX_CLOTHING
-
 
 /*-----------------------------------------------------------------------------
 FGPUSkinVertexBase
@@ -324,7 +244,7 @@ void FreeSkeletalMeshBuffersSinkCallback()
 
 FArchive& operator<<(FArchive& Ar, FClothingAssetData_Legacy& A)
 {
-	// Serialization to load and save ApexClothingAsset
+	// Serialization to load and skip legacy clothing assets
 	if( Ar.IsLoading() )
 	{
 		uint32 AssetSize;
@@ -336,25 +256,11 @@ FArchive& operator<<(FArchive& Ar, FClothingAssetData_Legacy& A)
 			TArray<uint8> Buffer;
 			Buffer.AddUninitialized( AssetSize );
 			Ar.Serialize( Buffer.GetData(), AssetSize );
-#if WITH_APEX_CLOTHING
-			A.ApexClothingAsset = LoadApexClothingAssetFromBlob(Buffer);
-#endif //#if WITH_APEX_CLOTHING
 		}
 	}
 	else
 	if( Ar.IsSaving() )
 	{
-#if WITH_APEX_CLOTHING
-		if (A.ApexClothingAsset)
-		{
-			TArray<uint8> Buffer;
-			SaveApexClothingAssetToBlob(A.ApexClothingAsset, Buffer);
-			uint32 AssetSize = Buffer.Num();
-			Ar << AssetSize;
-			Ar.Serialize(Buffer.GetData(), AssetSize);
-		}
-		else
-#endif// #if WITH_APEX_CLOTHING
 		{
 			uint32 AssetSize = 0;
 			Ar << AssetSize;
@@ -1596,20 +1502,6 @@ void USkeletalMesh::BeginDestroy()
 		GetSkeleton()->RemoveLinkup(this);
 	}
 
-#if WITH_EDITORONLY_DATA
-#if WITH_APEX_CLOTHING
-	// release clothing assets
-	for (FClothingAssetData_Legacy& Data : ClothingAssets_DEPRECATED)
-	{
-		if (Data.ApexClothingAsset)
-		{
-			GPhysCommandHandler->DeferredRelease(Data.ApexClothingAsset);
-			Data.ApexClothingAsset = nullptr;
-		}
-	}
-#endif // #if WITH_APEX_CLOTHING
-#endif // WITH_EDITORONLY_DATA
-
 	// Release the mesh's render resources now if no pending streaming op.
 	if (!HasPendingInitOrStreaming())
 	{
@@ -2348,109 +2240,6 @@ void USkeletalMesh::CalculateRequiredBones(FSkeletalMeshLODModel& LODModel, cons
 	LODModel.RequiredBones.Shrink();	
 }
 
-#if WITH_APEX_CLOTHING
-
-void USkeletalMesh::UpgradeOldClothingAssets()
-{
-	// Can only do an old-> new clothing asset upgrade in the editor.
-	// And only if APEX clothing is available to upgrade from
-	if (ClothingAssets_DEPRECATED.Num() > 0)
-	{
-		TRACE_CPUPROFILER_EVENT_SCOPE(USkeletalMesh::UpgradeOldClothingAssets);
-
-		const bool bCallPostEditChange = false;
-		const bool bReregisterComponents = false;
-		FScopedSkeletalMeshPostEditChange ScopedSkeletalMeshPostEditChange(this, bCallPostEditChange, bReregisterComponents);
-		// Upgrade the old deprecated clothing assets in to new clothing assets
-		TMap<int32, TArray<int32>> OldLodMappings; // Map asset index to multiple lod indices
-		TMap<int32, TArray<int32>> OldSectionMappings; // Map asset index to a section per LOD
-		for (int32 AssetIdx = 0; AssetIdx < ClothingAssets_DEPRECATED.Num(); ++AssetIdx)
-		{
-			FClothingAssetData_Legacy& OldAssetData = ClothingAssets_DEPRECATED[AssetIdx];
-
-			OldLodMappings.Add(AssetIdx);
-			OldSectionMappings.Add(AssetIdx);
-
-			if (ImportedModel.IsValid())
-			{
-				int32 FoundLod = INDEX_NONE;
-				int32 FoundSection = INDEX_NONE;
-				for (int32 LodIdx = 0; LodIdx < GetImportedModel()->LODModels.Num(); ++LodIdx)
-				{
-					FSkeletalMeshLODModel& LodModel = GetImportedModel()->LODModels[LodIdx];
-
-					for (int32 SecIdx = 0; SecIdx < LodModel.Sections.Num(); ++SecIdx)
-					{
-						FSkelMeshSection& Section = LodModel.Sections[SecIdx];
-
-						if (Section.CorrespondClothSectionIndex_DEPRECATED != INDEX_NONE && Section.bLegacyClothingSection_DEPRECATED)
-						{
-							FSkelMeshSection& ClothSection = LodModel.Sections[Section.CorrespondClothSectionIndex_DEPRECATED];
-
-							if (ClothSection.CorrespondClothAssetIndex == AssetIdx)
-							{
-								FoundSection = SecIdx;
-								break;
-							}
-						}
-					}
-
-					if (FoundSection != INDEX_NONE)
-					{
-						OldLodMappings[AssetIdx].Add(LodIdx);
-						OldSectionMappings[AssetIdx].Add(FoundSection);
-
-						// Reset for next LOD
-						FoundSection = INDEX_NONE;
-					}
-				}
-			}
-
-			FClothingSystemEditorInterfaceModule& ClothingEditorModule = FModuleManager::Get().LoadModuleChecked<FClothingSystemEditorInterfaceModule>(TEXT("ClothingSystemEditorInterface"));
-			UClothingAssetFactoryBase* Factory = ClothingEditorModule.GetClothingAssetFactory();
-			if (Factory)
-			{
-				UClothingAssetBase* NewAsset = Factory->CreateFromApexAsset(OldAssetData.ApexClothingAsset, this, *FPaths::GetBaseFilename(OldAssetData.ApexFileName));
-				check(NewAsset);
-
-				// Pull the path across so reimports work as expected
-				NewAsset->ImportedFilePath = OldAssetData.ApexFileName;
-
-				GetMeshClothingAssets().Add(NewAsset);
-			}
-		}
-
-		// Go back over the old assets and remove them from the skeletal mesh so the indices are preserved while
-		// calculating the LOD and section mappings above.
-		for (int32 AssetIdx = ClothingAssets_DEPRECATED.Num() - 1; AssetIdx >= 0; --AssetIdx)
-		{
-			ApexClothingUtils::RemoveAssetFromSkeletalMesh(this, AssetIdx, false);
-		}
-
-		check(OldLodMappings.Num() == OldSectionMappings.Num());
-
-		for (int32 NewAssetIdx = 0; NewAssetIdx < GetMeshClothingAssets().Num(); ++NewAssetIdx)
-		{
-			UClothingAssetBase* CurrAsset = GetMeshClothingAssets()[NewAssetIdx];
-
-			for (int32 MappedLodIdx = 0; MappedLodIdx < OldLodMappings[NewAssetIdx].Num(); ++MappedLodIdx)
-			{
-				const int32 MappedLod = OldLodMappings[NewAssetIdx][MappedLodIdx];
-				const int32 MappedSection = OldSectionMappings[NewAssetIdx][MappedLodIdx];
-
-				// Previously Clothing LODs were required to match skeletal mesh LODs, which is why we pass
-				// MappedLod for both the mesh and clothing LODs here when doing an upgrade to the new
-				// system. This restriction is now lifted and any mapping can be selected in Persona
-				CurrAsset->BindToSkeletalMesh(this, MappedLod, MappedSection, MappedLod);
-			}
-		}
-
-		UE_LOG(LogSkeletalMesh, Warning, TEXT("Legacy clothing asset '%s' was upgraded - please resave this asset."), *GetName());
-	}
-}
-
-#endif // WITH_APEX_CLOTHING
-
 void USkeletalMesh::RemoveLegacyClothingSections()
 {
 	// Remove duplicate skeletal mesh sections previously used for clothing simulation
@@ -3497,10 +3286,6 @@ void USkeletalMesh::ExecutePostLoadInternal(FSkeletalMeshPostLoadContext& Contex
 
 	if (!GetOutermost()->bIsCookedForEditor)
 	{
-#if WITH_APEX_CLOTHING
-		UpgradeOldClothingAssets();
-#endif // WITH_APEX_CLOTHING
-
 		RemoveLegacyClothingSections();
 
 		UpdateGenerateUpToData();
diff --git a/Engine/Source/Runtime/Engine/Private/SkeletalMeshComponentPhysics.cpp b/Engine/Source/Runtime/Engine/Private/SkeletalMeshComponentPhysics.cpp
--- a/Engine/Source/Runtime/Engine/Private/SkeletalMeshComponentPhysics.cpp
+++ b/Engine/Source/Runtime/Engine/Private/SkeletalMeshComponentPhysics.cpp
@@ -33,21 +33,10 @@
 #include "Logging/MessageLog.h"
 #include "CollisionDebugDrawingPublic.h"
 
-#if WITH_PHYSX
-	#include "SceneManagement.h"
-	#include "PhysXPublic.h"
-	#include "PhysicsEngine/PhysXSupport.h"
-#endif
-
-#if WITH_APEX
-
-#if WITH_APEX_CLOTHING || WITH_CHAOS_CLOTHING
-	// for cloth morph target	
-	#include "Animation/MorphTarget.h"
+#include "SceneManagement.h"
+#include "PhysXPublic.h"
+#include "PhysicsEngine/PhysXSupport.h"
 
-#endif// #if WITH_APEX_CLOTHING || WITH_CHAOS_CLOTHING
-
-#endif//#if WITH_APEX
 #include "PhysicsEngine/ConstraintInstance.h"
 #include "PhysicsEngine/PhysicsConstraintTemplate.h"
 #include "PhysicsEngine/BodySetup.h"
@@ -58,7 +47,6 @@
 #include "SimulationEditorExtender.h"
 #endif
 
-#if !PHYSICS_INTERFACE_PHYSX && WITH_CHAOS  // For cloth environmental collision
 #include "Chaos/Capsule.h"
 #include "Chaos/Sphere.h"
 #include "Chaos/Box.h"
@@ -67,7 +55,6 @@
 #include "Chaos/ImplicitObject.h"
 #include "Chaos/ImplicitObjectScaled.h"
 #include "Chaos/ImplicitObjectTransformed.h"
-#endif  // #if !PHYSICS_INTERFACE_PHYSX && WITH_CHAOS
 
 #define LOCTEXT_NAMESPACE "SkeletalMeshComponentPhysics"
 
@@ -579,7 +566,6 @@ void USkeletalMeshComponent::OnPlasticDeformationWrapper(int32 ConstraintIndex)
 
 void USkeletalMeshComponent::InitCollisionRelationships()
 {
-#if WITH_CHAOS
 	if (UPhysicsAsset* const PhysicsAsset = GetPhysicsAsset())
 	{
 		int32 NumDisabledCollisions = PhysicsAsset->CollisionDisableTable.Num();
@@ -627,12 +613,10 @@ void USkeletalMeshComponent::InitCollisionRelationships()
 				});
 		}
 	}
-#endif
 }
 
 void USkeletalMeshComponent::TermCollisionRelationships()
 {
-#if WITH_CHAOS
 	if (UPhysicsAsset* const PhysicsAsset = GetPhysicsAsset())
 	{
 		int32 NumDisabledCollisions = PhysicsAsset->CollisionDisableTable.Num();
@@ -658,8 +642,6 @@ void USkeletalMeshComponent::TermCollisionRelationships()
 			});
 		}
 	}
-#endif
-
 }
 
 
@@ -734,9 +716,6 @@ void USkeletalMeshComponent::InitArticulated(FPhysScene* PhysScene)
 	}
 
 	// Set up the map from skelmeshcomp ID to collision disable table
-#if !WITH_CHAOS
-	PhysScene->DeferredAddCollisionDisableTable(GetUniqueID(), &PhysicsAsset->CollisionDisableTable);
-#endif
 	int32 NumShapes = 0;
 	const int32 NumBodies = PhysicsAsset->SkeletalBodySetups.Num();
 	for (int32 BodyIndex = 0; BodyIndex < NumBodies; ++BodyIndex)
@@ -764,13 +743,11 @@ void USkeletalMeshComponent::InitArticulated(FPhysScene* PhysScene)
 	InitCollisionRelationships();
 
 	// Update Flag
-#if WITH_APEX_CLOTHING || WITH_CHAOS_CLOTHING
 	PrevRootBoneMatrix = GetBoneMatrix(0); // save the root bone transform
 
 	// pre-compute cloth teleport thresholds for performance
 	ComputeTeleportDistanceThresholdInRadians();
 	ComputeTeleportRotationThresholdInRadians();
-#endif // #if WITH_APEX_CLOTHING || WITH_CHAOS_CLOTHING
 }
 
 TAutoConsoleVariable<int32> CVarEnableRagdollPhysics(TEXT("p.RagdollPhysics"), 1, TEXT("If 1, ragdoll physics will be used. Otherwise just root body is simulated"));
@@ -1011,17 +988,7 @@ void USkeletalMeshComponent::TermArticulated()
 {
 	ResetRootBodyIndex();
 
-#if WITH_CHAOS
 	TermCollisionRelationships();
-#else
-	uint32 SkelMeshCompID = GetUniqueID();
-	UWorld* MyWorld = GetWorld();
-	FPhysScene* PhysScene = (MyWorld ? MyWorld->GetPhysicsScene() : nullptr);
-	if (PhysScene)
-	{
-		PhysScene->DeferredRemoveCollisionDisableTable(SkelMeshCompID);
-	}
-#endif
 
 	FPhysicsCommand::ExecuteWrite(this, [&]()
 	{
@@ -1559,7 +1526,6 @@ void USkeletalMeshComponent::OnUpdateTransform(EUpdateTransformFlags UpdateTrans
 				}
 			}
 
-#if WITH_CHAOS
 			if(GEnableKinematicDeferralStartPhysicsCondition)
 			{
 				if (World && (World->TickGroup == ETickingGroup::TG_StartPhysics))
@@ -1567,7 +1533,6 @@ void USkeletalMeshComponent::OnUpdateTransform(EUpdateTransformFlags UpdateTrans
 					AllowDeferral = EAllowKinematicDeferral::AllowDeferral;
 				}
 			}
-#endif
 		}
 
 		UpdateKinematicBonesToAnim(GetComponentSpaceTransforms(), Teleport, false, AllowDeferral);
@@ -3117,142 +3082,6 @@ void USkeletalMeshComponent::ProcessClothCollisionWithEnvironment()
 						return;
 					}
 
-#if PHYSICS_INTERFACE_PHYSX
-					// Matrices required to transform shapes into sim space (component space)
-					// Transform of external component and matrix describing external component -> this component
-					FTransform Transform = Component->GetComponentTransform();
-					FMatrix TransformMatrix = Transform.ToMatrixWithScale();
-					FMatrix ComponentToClothMatrix = TransformMatrix * GetComponentTransform().ToMatrixWithScale().Inverse();
-
-					for(FPhysicsShapeHandle& Shape : AllShapes)
-					{
-						ECollisionShapeType GeoType = FPhysicsInterface::GetShapeType(Shape);
-						FPhysicsGeometryCollection GeoCollection = FPhysicsInterface::GetGeometryCollection(Shape);
-
-						// Pose of the shape in actor space
-						FMatrix ShapeLocalPose = FPhysicsInterface::GetLocalTransform(Shape).ToMatrixWithScale();
-
-						switch(GeoType)
-						{
-							default:
-							break;
-
-							case ECollisionShapeType::Sphere:
-							{
-								PxSphereGeometry& SphereGeo = GeoCollection.GetSphereGeometry();
-
-								NewCollisionData.Spheres.AddDefaulted();
-								FClothCollisionPrim_Sphere& NewSphere = NewCollisionData.Spheres.Last();
-
-								NewSphere.BoneIndex = INDEX_NONE; // No bone, just local space
-								NewSphere.LocalPosition = ComponentToClothMatrix.TransformPosition(ShapeLocalPose.GetOrigin());
-								NewSphere.Radius = SphereGeo.radius;
-							}
-							break;
-
-							case ECollisionShapeType::Capsule:
-							{
-								PxCapsuleGeometry& CapGeo = GeoCollection.GetCapsuleGeometry();
-
-								const int32 BaseSphereIndex = NewCollisionData.Spheres.Num();
-
-								NewCollisionData.Spheres.AddDefaulted(2);
-								FClothCollisionPrim_Sphere& Sphere0 = NewCollisionData.Spheres.Last(1);
-								FClothCollisionPrim_Sphere& Sphere1 = NewCollisionData.Spheres.Last(0);
-
-								NewCollisionData.SphereConnections.AddDefaulted();
-								FClothCollisionPrim_SphereConnection& Connection = NewCollisionData.SphereConnections.Last();
-
-								FVector ZAxis = ComponentToClothMatrix.TransformVector(ShapeLocalPose.GetUnitAxis(EAxis::X));
-
-								Sphere0.BoneIndex = INDEX_NONE;
-								Sphere0.LocalPosition = ComponentToClothMatrix.TransformPosition(ShapeLocalPose.GetOrigin()) + CapGeo.halfHeight * ZAxis;
-								Sphere0.Radius = CapGeo.radius;
-
-								Sphere1.BoneIndex = INDEX_NONE;
-								Sphere1.LocalPosition = ComponentToClothMatrix.TransformPosition(ShapeLocalPose.GetOrigin()) - CapGeo.halfHeight * ZAxis;
-								Sphere1.Radius = CapGeo.radius;
-
-								Connection.SphereIndices[0] = BaseSphereIndex;
-								Connection.SphereIndices[1] = BaseSphereIndex + 1;
-							}
-							break;
-
-							case ECollisionShapeType::Box:
-							{
-								PxBoxGeometry& BoxGeo = GeoCollection.GetBoxGeometry();
-
-								// We're building the box in local space, so to get to the cloth transform
-								// we need to go through local -> actor -> world -> cloth
-								FMatrix FullTransformMatrix = ShapeLocalPose * ComponentToClothMatrix;
-
-								NewCollisionData.Convexes.AddDefaulted();
-								FClothCollisionPrim_Convex& Convex = NewCollisionData.Convexes.Last();
-								Convex.Faces.SetNum(6);
-
-								// we need to inflate the hull to get nicer collisions (only particles collide)
-								BoxGeo.halfExtents += PxVec3(EnvironmentCollisionConfig.Thickness);
-				
-								FPlane UPlane1(1, 0, 0, BoxGeo.halfExtents.x);
-								UPlane1 = UPlane1.TransformBy(FullTransformMatrix);
-								Convex.Faces[0].Plane = UPlane1;
-
-								FPlane UPlane2(-1, 0, 0, BoxGeo.halfExtents.x);
-								UPlane2 = UPlane2.TransformBy(FullTransformMatrix);
-								Convex.Faces[1].Plane = UPlane2;
-
-								FPlane UPlane3(0, 1, 0, BoxGeo.halfExtents.y);
-								UPlane3 = UPlane3.TransformBy(FullTransformMatrix);
-								Convex.Faces[2].Plane = UPlane3;
-
-								FPlane UPlane4(0, -1, 0, BoxGeo.halfExtents.y);
-								UPlane4 = UPlane4.TransformBy(FullTransformMatrix);
-								Convex.Faces[3].Plane = UPlane4;
-
-								FPlane UPlane5(0, 0, 1, BoxGeo.halfExtents.z);
-								UPlane5 = UPlane5.TransformBy(FullTransformMatrix);
-								Convex.Faces[4].Plane = UPlane5;
-
-								FPlane UPlane6(0, 0, -1, BoxGeo.halfExtents.z);
-								UPlane6 = UPlane6.TransformBy(FullTransformMatrix);
-								Convex.Faces[5].Plane = UPlane6;
-
-								Convex.BoneIndex = INDEX_NONE;
-							}
-							break;
-
-							case ECollisionShapeType::Convex:
-							{
-								PxConvexMeshGeometry& MeshGeo = GeoCollection.GetConvexGeometry();
-
-								if(MeshGeo.convexMesh)
-								{
-									NewCollisionData.Convexes.AddDefaulted();
-									FClothCollisionPrim_Convex& NewConvex = NewCollisionData.Convexes.Last();
-
-									FMatrix FullTransformMatrix = ShapeLocalPose * ComponentToClothMatrix;
-
-									uint32 NumPolys = MeshGeo.convexMesh->getNbPolygons();
-									NewConvex.Faces.SetNum(NumPolys);
-
-									PxHullPolygon HullData;
-									for(uint32 PolyIndex = 0; PolyIndex < NumPolys; ++PolyIndex)
-									{
-										MeshGeo.convexMesh->getPolygonData(PolyIndex, HullData);
-										PxPlane PPlane(HullData.mPlane[0], HullData.mPlane[1], HullData.mPlane[2], HullData.mPlane[3]);
-										FPlane UPlane = P2UPlane(PPlane);
-										UPlane = UPlane.TransformBy(FullTransformMatrix);
-										
-										UPlane.W += EnvironmentCollisionConfig.Thickness;
-
-										NewConvex.Faces[PolyIndex].Plane = UPlane;
-									}
-								}	
-							}
-							break;
-						}
-					}
-#elif WITH_CHAOS
 					using namespace Chaos;
 
 					const FReal Thickness = (FReal)EnvironmentCollisionConfig.Thickness;
@@ -3499,9 +3328,6 @@ void USkeletalMeshComponent::ProcessClothCollisionWithEnvironment()
 							break;
 						}
 					}
-#else
-					static_assert(false, "A physics engine interface must be defined to build");
-#endif
 					bSuccessfulRead = true;
 				});
 			}
diff --git a/Engine/Source/Runtime/Engine/Private/SkeletalRenderGPUSkin.cpp b/Engine/Source/Runtime/Engine/Private/SkeletalRenderGPUSkin.cpp
--- a/Engine/Source/Runtime/Engine/Private/SkeletalRenderGPUSkin.cpp
+++ b/Engine/Source/Runtime/Engine/Private/SkeletalRenderGPUSkin.cpp
@@ -702,10 +702,9 @@ void FSkeletalMeshObjectGPUSkin::ProcessUpdatedDynamicData(EGPUSkinCacheEntryMod
 			bool bNeedFence = ShaderData.UpdateBoneData(RHICmdList, ReferenceToLocalMatrices, Section.BoneMap, RevisionNumber, false, FeatureLevel, bUseSkinCache);
 			ShaderData.UpdatedFrameNumber = FrameNumberToPrepare;
 
-#if WITH_APEX_CLOTHING || WITH_CHAOS_CLOTHING
 			// Update uniform buffer for APEX cloth simulation mesh positions and normals
 			if( bClothFactory )
-			{				
+			{
 				FGPUBaseSkinAPEXClothVertexFactory::ClothShaderType& ClothShaderData = VertexFactoryData.ClothVertexFactories[SectionIdx]->GetClothShaderData();
 				ClothShaderData.ClothBlendWeight = DynamicData->ClothBlendWeight;
 
@@ -713,7 +712,6 @@ void FSkeletalMeshObjectGPUSkin::ProcessUpdatedDynamicData(EGPUSkinCacheEntryMod
 				// Transform from cloth space to local space. Cloth space is relative to cloth root bone, local space is component space.
 				ClothShaderData.GetClothToLocalForWriting(FrameNumberToPrepare) = FMatrix44f(SimData->ComponentRelativeTransform.ToMatrixWithScale());
 			}
-#endif // WITH_APEX_CLOTHING || WITH_CHAOS_CLOTHING
 
 			// Try to use the GPU skinning cache if possible
 			if (bUseSkinCache)
diff --git a/Engine/Source/Runtime/Engine/Private/UnrealEngine.cpp b/Engine/Source/Runtime/Engine/Private/UnrealEngine.cpp
--- a/Engine/Source/Runtime/Engine/Private/UnrealEngine.cpp
+++ b/Engine/Source/Runtime/Engine/Private/UnrealEngine.cpp
@@ -237,10 +237,6 @@ UnrealEngine.cpp: Implements the UEngine class and helpers.
 #include "Streaming/StreamingManagerTexture.h"
 #endif
 
-#if WITH_PHYSX
-#include "PhysXPublic.h"
-#endif
-
 #include "HAL/FileManagerGeneric.h"
 #include "UObject/UObjectThreadContext.h"
 #include "UObject/ReferenceChainSearch.h"
@@ -6270,18 +6266,10 @@ bool UEngine::HandleListStaticMeshesCommand(const TCHAR* Cmd, FOutputDevice& Ar)
 		int32 VertexCountCollision = 0;
 		if(Mesh->GetBodySetup())
 		{
-#if PHYSICS_INTERFACE_PHYSX
-			// Count PhysX trimesh mem usage
-			for (physx::PxTriangleMesh* TriMesh : Mesh->GetBodySetup()->TriMeshes)
-			{
-				VertexCountCollision += TriMesh->getNbVertices();
-			}
-#elif WITH_CHAOS
 			for (auto& TriMesh : Mesh->GetBodySetup()->ChaosTriMeshes)
 			{
 				VertexCountCollision += TriMesh->Particles().Size();
 			}
-#endif
 		}
 
 		FString Name = Mesh->GetFullName();
@@ -6589,18 +6577,10 @@ bool UEngine::HandleListSkeletalMeshesCommand(const TCHAR* Cmd, FOutputDevice& A
 		const USkeletalMesh* MeshConst = Mesh;
 		if (MeshConst->GetBodySetup())
 		{
-#if PHYSICS_INTERFACE_PHYSX
-			// Count PhysX trimesh mem usage
-			for (physx::PxTriangleMesh* TriMesh : MeshConst->GetBodySetup()->TriMeshes)
-			{
-				VertexCountCollision += TriMesh->getNbVertices();
-			}
-#elif WITH_CHAOS
 			for (auto& TriMesh : MeshConst->GetBodySetup()->ChaosTriMeshes)
 			{
 				VertexCountCollision += TriMesh->Particles().Size();
 			}
-#endif
 		}
 
 		FString Name = Mesh->GetFullName();
diff --git a/Engine/Source/Runtime/Engine/Private/Volume.cpp b/Engine/Source/Runtime/Engine/Private/Volume.cpp
--- a/Engine/Source/Runtime/Engine/Private/Volume.cpp
+++ b/Engine/Source/Runtime/Engine/Private/Volume.cpp
@@ -75,7 +75,7 @@ bool AVolume::EncompassesPoint(FVector Point, float SphereRadius/*=0.f*/, float*
 {
 	if (GetBrushComponent())
 	{
-#if WITH_PHYSX
+#if 1
 		FVector ClosestPoint;
 		float DistanceSqr;
 
diff --git a/Engine/Source/Runtime/Engine/Private/World.cpp b/Engine/Source/Runtime/Engine/Private/World.cpp
--- a/Engine/Source/Runtime/Engine/Private/World.cpp
+++ b/Engine/Source/Runtime/Engine/Private/World.cpp
@@ -138,10 +138,7 @@
 #include "ObjectTrace.h"
 #include "ReplaySubsystem.h"
 #include "Net/NetPing.h"
-
-#if INCLUDE_CHAOS
 #include "ChaosSolversModule.h"
-#endif
 
 DEFINE_LOG_CATEGORY_STATIC(LogWorld, Log, All);
 DEFINE_LOG_CATEGORY(LogSpawn);
@@ -1014,13 +1011,11 @@ void UWorld::BeginDestroy()
 		}
 	}
 
-#if WITH_CHAOS
 	if (PhysicsScene != nullptr)
 	{
 		// Tell PhysicsScene to stop kicking off async work so we can cleanup after pending work is complete.
 		PhysicsScene->BeginDestroy();
 	}
-#endif
 
 	if (Scene)
 	{
@@ -1044,12 +1039,10 @@ void UWorld::ReleasePhysicsScene()
 		delete PhysicsScene;
 		PhysicsScene = NULL;
 
-#if WITH_PHYSX
 		if (GPhysCommandHandler)
 		{
 			GPhysCommandHandler->Flush();
 		}
-#endif // WITH_PHYSX
 	}
 }
 
@@ -1136,8 +1129,6 @@ void UWorld::FinishDestroy()
 
 bool UWorld::IsReadyForFinishDestroy()
 {
-#if WITH_CHAOS
-
 	// In single threaded, task will never complete unless we wait on it, allow FinishDestroy so we can wait on task, otherwise this will hang GC.
 	// In multi threaded, we cannot wait in FinishDestroy, as this may schedule another task that is unsafe during GC.
 	const bool bIsSingleThreadEnvironment = FPlatformProcess::SupportsMultithreading() == false;
@@ -1151,7 +1142,6 @@ bool UWorld::IsReadyForFinishDestroy()
 			}
 		}
 	}
-#endif
 
 	return Super::IsReadyForFinishDestroy();
 }
@@ -1177,9 +1167,6 @@ void UWorld::PostLoad()
 	RepairWorldSettings();
 	RepairStreamingLevels();
 #endif
-#if INCLUDE_CHAOS
-	//RepairChaosActors();
-#endif
 
 	for (auto It = StreamingLevels.CreateIterator(); It; ++It)
 	{
@@ -1471,7 +1458,6 @@ UAISystemBase* UWorld::CreateAISystem()
 	return AISystem; 
 }
 
-#if INCLUDE_CHAOS
 void UWorld::RepairChaosActors()
 {
 	if (!PhysicsScene_Chaos)
@@ -1540,7 +1526,6 @@ void UWorld::RepairChaosActors()
 	// make the current scene the default scene
 	DefaultPhysicsScene_Chaos = PhysicsScene_Chaos;
 }
-#endif
 
 void UWorld::RepairStreamingLevels()
 {
@@ -1827,10 +1812,6 @@ void UWorld::InitWorld(const InitializationValues IVS)
 	RepairWorldSettings();
 	RepairStreamingLevels();
 #endif
-#if INCLUDE_CHAOS
-	//RepairChaosActors();
-#endif
-
 
 	// initialize DefaultPhysicsVolume for the world
 	// Spawned on demand by this function.
@@ -2005,16 +1986,6 @@ void UWorld::InitializeNewWorld(const InitializationValues IVS, bool bInSkipInit
 	}
 #endif
 
-#if INCLUDE_CHAOS
-	/*FChaosSolversModule* ChaosModule = FModuleManager::Get().GetModulePtr<FChaosSolversModule>("ChaosSolvers");
-	check(ChaosModule);
-	FActorSpawnParameters ChaosSpawnInfo;
-	ChaosSpawnInfo.SpawnCollisionHandlingOverride = ESpawnActorCollisionHandlingMethod::AlwaysSpawn;
-	ChaosSpawnInfo.Name = TEXT("DefaultChaosActor");
-	SpawnActor(ChaosModule->GetSolverActorClass(), nullptr, nullptr, ChaosSpawnInfo);
-	check(PhysicsScene_Chaos);*/
-#endif
-
 	if (!bInSkipInitWorld)
 	{
 		// Initialize the world
@@ -4927,12 +4898,10 @@ void UWorld::BeginPlay()
 
 	OnWorldBeginPlay.Broadcast();
 
-#if WITH_CHAOS
 	if(PhysicsScene)
 	{
 		PhysicsScene->OnWorldBeginPlay();
 	}
-#endif
 }
 
 bool UWorld::IsNavigationRebuilt() const
@@ -4972,9 +4941,7 @@ void UWorld::CleanupWorldInternal(bool bSessionEnded, bool bCleanupResources, bo
 	if(FPhysScene* CurrPhysicsScene = GetPhysicsScene())
 	{
 		CurrPhysicsScene->WaitPhysScenes();
-#if WITH_CHAOS
 		CurrPhysicsScene->OnWorldEndPlay();
-#endif
 	}
 
 	FWorldDelegates::OnWorldCleanup.Broadcast(this, bSessionEnded, bCleanupResources);
@@ -5404,16 +5371,13 @@ bool UWorld::AreActorsInitialized() const
 
 void UWorld::CreatePhysicsScene(const AWorldSettings* Settings)
 {
-#if WITH_CHAOS && CHAOS_DEBUG_NAME
+#if CHAOS_DEBUG_NAME
 	const FName PhysicsName = IsNetMode(NM_DedicatedServer) ? TEXT("ServerPhysics") : TEXT("ClientPhysics");
 	FPhysScene* NewScene = new FPhysScene(nullptr, PhysicsName);
-#else
-#if PHYSICS_INTERFACE_PHYSX
-	FPhysScene* NewScene = new FPhysScene(Settings);
 #else
 	FPhysScene* NewScene = new FPhysScene(nullptr);
 #endif
-#endif
+
 	SetPhysicsScene(NewScene);
 }
 
diff --git a/Engine/Source/Runtime/Engine/Public/AI/NavigationSystemHelpers.h b/Engine/Source/Runtime/Engine/Public/AI/NavigationSystemHelpers.h
--- a/Engine/Source/Runtime/Engine/Public/AI/NavigationSystemHelpers.h
+++ b/Engine/Source/Runtime/Engine/Public/AI/NavigationSystemHelpers.h
@@ -16,40 +16,20 @@ struct FNavHeightfieldSamples;
 
 template<typename InElementType> class TNavStatArray;
 
-#if PHYSICS_INTERFACE_PHYSX
-namespace physx
-{
-	class PxTriangleMesh;
-	class PxConvexMesh;
-	class PxHeightField;
-}
-#endif // WITH_PHYSX
-
-#if WITH_CHAOS
 namespace Chaos
 {
 	class FHeightField;
-
 	class FTriangleMeshImplicitObject;
 }
-#endif
 
 struct FNavigableGeometryExport
 {
 	virtual ~FNavigableGeometryExport() {}
-#if PHYSICS_INTERFACE_PHYSX
-	virtual void ExportPxTriMesh16Bit(physx::PxTriangleMesh const * const TriMesh, const FTransform& LocalToWorld) = 0;
-	virtual void ExportPxTriMesh32Bit(physx::PxTriangleMesh const * const TriMesh, const FTransform& LocalToWorld) = 0;
-	virtual void ExportPxConvexMesh(physx::PxConvexMesh const * const ConvexMesh, const FTransform& LocalToWorld) = 0;
-	virtual void ExportPxHeightField(physx::PxHeightField const * const HeightField, const FTransform& LocalToWorld) = 0;
-	virtual void ExportPxHeightFieldSlice(const FNavHeightfieldSamples& PrefetchedHeightfieldSamples, const int32 NumRows, const int32 NumCols, const FTransform& LocalToWorld, const FBox& SliceBox) = 0;
-#endif // WITH_PHYSX
-#if WITH_CHAOS
+
 	virtual void ExportChaosTriMesh(const Chaos::FTriangleMeshImplicitObject* const TriMesh, const FTransform& LocalToWorld) = 0;
 	virtual void ExportChaosConvexMesh(const FKConvexElem* const Convex, const FTransform& LocalToWorld) = 0;
 	virtual void ExportChaosHeightField(const Chaos::FHeightField* const Heightfield, const FTransform& LocalToWorld) = 0;
 	virtual void ExportChaosHeightFieldSlice(const FNavHeightfieldSamples& PrefetchedHeightfieldSamples, const int32 NumRows, const int32 NumCols, const FTransform& LocalToWorld, const FBox& SliceBox) = 0;
-#endif
 	virtual void ExportRigidBodySetup(UBodySetup& BodySetup, const FTransform& LocalToWorld) = 0;
 	virtual void ExportCustomMesh(const FVector* VertexBuffer, int32 NumVerts, const int32* IndexBuffer, int32 NumIndices, const FTransform& LocalToWorld) = 0;
 	
diff --git a/Engine/Source/Runtime/Engine/Public/CustomPhysXPayload.h b/Engine/Source/Runtime/Engine/Public/CustomPhysXPayload.h
--- a/Engine/Source/Runtime/Engine/Public/CustomPhysXPayload.h
+++ b/Engine/Source/Runtime/Engine/Public/CustomPhysXPayload.h
@@ -5,8 +5,6 @@
 #include "EngineDefines.h"
 #include "Components/PrimitiveComponent.h"
 
-#if WITH_PHYSX
-
 namespace physx
 {
 	class PxRigidActor;
@@ -55,4 +53,3 @@ struct FCustomPhysXPayload
 	FCustomPhysXSyncActors* CustomSyncActors;
 };
 
-#endif // WITH_PHYSICS
diff --git a/Engine/Source/Runtime/Engine/Public/PhysXCookHelper.h b/Engine/Source/Runtime/Engine/Public/PhysXCookHelper.h
--- a/Engine/Source/Runtime/Engine/Public/PhysXCookHelper.h
+++ b/Engine/Source/Runtime/Engine/Public/PhysXCookHelper.h
@@ -1,53 +1,3 @@
 // Copyright Epic Games, Inc. All Rights Reserved.
 
 #pragma once
-
-#include "Async/TaskGraphInterfaces.h"
-#include "HAL/ThreadSafeBool.h"
-#include "PhysicsEngine/BodySetup.h"
-
-namespace physx
-{
-	class PxCooking;
-	class PxTriangleMesh;
-	class PxConvexMesh;
-	class PxHeightField;
-	class PxFoundation;
-	struct PxCookingParams;
-}
-
-class IPhysXCookingModule;
-
-#if WITH_PHYSX
-/** Helper for physics cooking */
-struct ENGINE_API FPhysXCookHelper
-{
-	FPhysXCookHelper(IPhysXCookingModule* InPhysXCookingModule);
-
-	/** Cooks based on CookInfo and saves the results into the output data */
-	bool CreatePhysicsMeshes_Concurrent();
-
-	/** Cooks based on CookInfo and saves the results into the output data. Calls back into the delegate on the game thread when done */
-	void CreatePhysicsMeshesAsync_Concurrent(FSimpleDelegateGraphTask::FDelegate FinishDelegate);
-
-	static bool HasSomethingToCook(const FCookBodySetupInfo& InCookInfo)
-	{
-		return InCookInfo.bCookTriMesh || InCookInfo.bCookNonMirroredConvex || InCookInfo.bCookMirroredConvex;
-	}
-
-	void Abort() { bShouldAbort.AtomicSet(true); }
-
-	FCookBodySetupInfo CookInfo;	//Use this with UBodySetup::GetCookInfo (must be called on game thread). If you already have the info just override it manually
-
-	//output
-	TArray<physx::PxConvexMesh*> OutNonMirroredConvexMeshes;
-	TArray<physx::PxConvexMesh*> OutMirroredConvexMeshes;
-	TArray<physx::PxTriangleMesh*> OutTriangleMeshes;
-	FBodySetupUVInfo OutUVInfo;
-private:
-	void CreateConvexElements_Concurrent(const TArray<TArray<FVector>>& Elements, TArray<physx::PxConvexMesh*>& OutConvexMeshes, bool bFlipped);
-
-	IPhysXCookingModule* PhysXCookingModule;
-	FThreadSafeBool bShouldAbort;	
-};
-#endif //WITH_PHYSX
diff --git a/Engine/Source/Runtime/Engine/Public/PhysXInterfaceWrapper.h b/Engine/Source/Runtime/Engine/Public/PhysXInterfaceWrapper.h
--- a/Engine/Source/Runtime/Engine/Public/PhysXInterfaceWrapper.h
+++ b/Engine/Source/Runtime/Engine/Public/PhysXInterfaceWrapper.h
@@ -1,72 +1,2 @@
 // Copyright Epic Games, Inc. All Rights Reserved.
 #pragma once
-
-#if PHYSICS_INTERFACE_PHYSX
-#include "PhysXInterfaceWrapperCore.h"
-#include "PhysXPublic.h"
-#include "Physics/PhysicsInterfaceUtils.h"
-#include "CustomPhysXPayload.h"
-#include "PhysicsInterfaceWrapperShared.h"
-
-namespace PhysXInterface
-{
-#if PHYSICS_INTERFACE_PHYSX
-/**
-* Helper to lock/unlock a scene that also makes sure to unlock everything when it goes out of scope.
-* Multiple locks on the same scene are NOT SAFE. You can't call LockRead() if already locked.
-* Multiple unlocks on the same scene are safe (repeated unlocks do nothing after the first successful unlock).
-*/
-struct FScopedSceneReadLock
-{
-	FScopedSceneReadLock(FPhysScene& Scene)
-		: SceneLock(Scene.GetPxScene())
-	{
-		SCENE_LOCK_READ(SceneLock);
-	}
-
-	~FScopedSceneReadLock()
-	{
-		SCENE_UNLOCK_READ(SceneLock);
-	}
-
-	PxScene* SceneLock;
-};
-#endif
-
-inline PxQueryFlags StaticDynamicQueryFlags(const FCollisionQueryParams& Params)
-{
-	switch (Params.MobilityType)
-	{
-	case EQueryMobilityType::Any: return  PxQueryFlag::eSTATIC | PxQueryFlag::eDYNAMIC;
-	case EQueryMobilityType::Static: return  PxQueryFlag::eSTATIC;
-	case EQueryMobilityType::Dynamic: return  PxQueryFlag::eDYNAMIC;
-	default: check(0);
-	}
-
-	check(0);
-	return PxQueryFlag::eSTATIC | PxQueryFlag::eDYNAMIC;
-}
-
-inline FQueryFilterData MakeQueryFilterData(const FCollisionFilterData& FilterData, EQueryFlags QueryFlags, const FCollisionQueryParams& Params)
-{
-	return PxQueryFilterData(U2PFilterData(FilterData), U2PQueryFlags(QueryFlags) | StaticDynamicQueryFlags(Params));
-}
-
-FORCEINLINE UPhysicalMaterial* GetUserData(const PxMaterial& Material)
-{
-	return FPhysxUserData::Get<UPhysicalMaterial>(Material.userData);
-}
-
-FORCEINLINE FBodyInstance* GetUserData(const PxActor& Actor)
-{
-	return FPhysxUserData::Get<FBodyInstance>(Actor.userData);
-}
-
-template <typename T>
-T* GetUserData(const PxShape& Shape)
-{
-	return FPhysxUserData::Get<T>(Shape.userData);
-}
-}
-
-#endif
\ No newline at end of file
diff --git a/Engine/Source/Runtime/Engine/Public/PhysXPublic.h b/Engine/Source/Runtime/Engine/Public/PhysXPublic.h
--- a/Engine/Source/Runtime/Engine/Public/PhysXPublic.h
+++ b/Engine/Source/Runtime/Engine/Public/PhysXPublic.h
@@ -7,12 +7,3 @@
 #pragma once
 
 #include "CoreMinimal.h"
-
-#if PHYSICS_INTERFACE_PHYSX
-#include "PhysXPublicCore.h"
-
-/** Calculates correct impulse at the body's center of mass and adds the impulse to the body. */
-ENGINE_API void AddRadialImpulseToPxRigidBody_AssumesLocked(PxRigidBody& PRigidBody, const FVector& Origin, float Radius, float Strength, uint8 Falloff, bool bVelChange);
-ENGINE_API void AddRadialForceToPxRigidBody_AssumesLocked(PxRigidBody& PRigidBody, const FVector& Origin, float Radius, float Strength, uint8 Falloff, bool bAccelChange);
-
-#endif
\ No newline at end of file
diff --git a/Engine/Source/Runtime/Engine/Public/Physics/Experimental/ChaosInterfaceWrapper.h b/Engine/Source/Runtime/Engine/Public/Physics/Experimental/ChaosInterfaceWrapper.h
--- a/Engine/Source/Runtime/Engine/Public/Physics/Experimental/ChaosInterfaceWrapper.h
+++ b/Engine/Source/Runtime/Engine/Public/Physics/Experimental/ChaosInterfaceWrapper.h
@@ -23,7 +23,6 @@ namespace ChaosInterface
 {
 
 // Needed by low level SQ calls.
-#if WITH_CHAOS
 struct FScopedSceneReadLock
 {
 	FScopedSceneReadLock(FPhysScene_Chaos& SceneIn);
@@ -31,22 +30,14 @@ struct FScopedSceneReadLock
 
 	Chaos::FPBDRigidsSolver* Solver;
 };
-#endif
 
 inline FQueryFilterData MakeQueryFilterData(const FCollisionFilterData& FilterData, EQueryFlags QueryFlags, const FCollisionQueryParams& Params)
 {
-#if PHYSICS_INTERFACE_PHYSX
-	return PxQueryFilterData(U2PFilterData(FilterData), U2PQueryFlags(QueryFlags));
-#else
 	return FChaosQueryFilterData(U2CFilterData(FilterData), U2CQueryFlags(QueryFlags));
-#endif
 }
 
 FBodyInstance* GetUserData(const Chaos::FGeometryParticle& Actor);
-
-#if WITH_CHAOS
 UPhysicalMaterial* GetUserData(const Chaos::FChaosPhysicsMaterial& Material);
-#endif
 
 }
 
diff --git a/Engine/Source/Runtime/Engine/Public/Physics/Experimental/PhysInterface_Chaos.h b/Engine/Source/Runtime/Engine/Public/Physics/Experimental/PhysInterface_Chaos.h
--- a/Engine/Source/Runtime/Engine/Public/Physics/Experimental/PhysInterface_Chaos.h
+++ b/Engine/Source/Runtime/Engine/Public/Physics/Experimental/PhysInterface_Chaos.h
@@ -2,8 +2,6 @@
 
 #pragma once
 
-#if WITH_CHAOS
-
 #include "ChaosInterfaceWrapper.h"
 #include "Engine/Engine.h"
 #include "Chaos/Declares.h"
@@ -182,4 +180,3 @@ ENGINE_API bool FGenericPhysicsInterface::GeomSweepMulti(const UWorld* World, co
 template<>
 ENGINE_API bool FGenericPhysicsInterface::GeomOverlapMulti(const UWorld* World, const FPhysicsGeometryCollection& InGeom, const FVector& InPosition, const FQuat& InRotation, TArray<FOverlapResult>& OutOverlaps, ECollisionChannel TraceChannel, const FCollisionQueryParams& Params, const FCollisionResponseParams& ResponseParams, const FCollisionObjectQueryParams& ObjectParams);
 
-#endif
diff --git a/Engine/Source/Runtime/Engine/Public/Physics/Experimental/PhysScene_Chaos.h b/Engine/Source/Runtime/Engine/Public/Physics/Experimental/PhysScene_Chaos.h
--- a/Engine/Source/Runtime/Engine/Public/Physics/Experimental/PhysScene_Chaos.h
+++ b/Engine/Source/Runtime/Engine/Public/Physics/Experimental/PhysScene_Chaos.h
@@ -2,7 +2,6 @@
 
 #pragma once
 
-#if WITH_CHAOS
 #include "CoreMinimal.h"
 #include "Tickable.h"
 #include "Physics/PhysScene.h"
@@ -99,19 +98,11 @@
 
 	using Super = FChaosScene;
 	
-#if !WITH_CHAOS_NEEDS_TO_BE_FIXED
-	FPhysScene_Chaos(AActor* InSolverActor
-#if CHAOS_DEBUG_NAME
-	, const FName& DebugName=NAME_None
-#endif
-);
-#else
 	FPhysScene_Chaos(AActor* InSolverActor=nullptr
 #if CHAOS_DEBUG_NAME
 	, const FName& DebugName=NAME_None
 #endif
 );
-#endif
 
 	virtual ~FPhysScene_Chaos();
 
@@ -208,9 +199,7 @@
 
 	static bool SupportsOriginShifting();
 	void ApplyWorldOffset(FVector InOffset);
-#if WITH_CHAOS
 	virtual float OnStartFrame(float InDeltaTime) override;
-#endif
 
 	bool HandleExecCommands(const TCHAR* Cmd, FOutputDevice* Ar);
 	void ListAwakeRigidBodies(bool bIncludeKinematic);
@@ -296,9 +285,7 @@
 	bool IsOwningWorldEditor() const;
 #endif
 
-#if WITH_CHAOS
 	virtual void OnSyncBodies(Chaos::FPhysicsSolverBase* Solver) override;
-#endif
 
 	void EnableAsyncPhysicsTickCallback();
 
@@ -369,9 +356,7 @@
 	/** The SolverActor that spawned and owns this scene */
 	TWeakObjectPtr<AActor> SolverActor;
 
-#if WITH_CHAOS
 	Chaos::FReal LastEventDispatchTime;
-#endif 
 	class FAsyncPhysicsTickCallback* AsyncPhysicsTickCallback = nullptr;
 
 #if WITH_EDITOR
@@ -388,4 +373,3 @@
 	friend struct FScopedSceneLock_Chaos;
 };
 
-#endif
diff --git a/Engine/Source/Runtime/Engine/Public/Physics/Experimental/PhysScene_ImmediatePhysX.h b/Engine/Source/Runtime/Engine/Public/Physics/Experimental/PhysScene_ImmediatePhysX.h
--- a/Engine/Source/Runtime/Engine/Public/Physics/Experimental/PhysScene_ImmediatePhysX.h
+++ b/Engine/Source/Runtime/Engine/Public/Physics/Experimental/PhysScene_ImmediatePhysX.h
@@ -1,753 +1,3 @@
 // Copyright Epic Games, Inc. All Rights Reserved.
 
 #pragma once
-
-#if WITH_IMMEDIATE_PHYSX
-
-#include "EngineGlobals.h"
-#include "PhysicsPublic.h"
-#include "PhysxUserData.h"
-#include "Physics/PhysicsInterfaceUtils.h"
-#include "PhysScene_PhysX.h"
-#include "PhysicsSettingsEnums.h"
-#include "GenericPhysicsInterface.h"
-#include "HAL/LowLevelMemTracker.h"
-
-struct FConstraintBrokenDelegateData;
-class IPhysicsReplicationFactory;
-
-namespace physx
-{
-	class PxMaterial;
-	class PxActor;
-	class PxSimulationEventCallback;
-}
-
-struct FContactModifyCallback;
-struct FPhysXMbpBroadphaseCallback;
-
-#if WITH_PHYSX
-/** Interface for the creation of customized simulation event callbacks. */
-class ISimEventCallbackFactory
-{
-public:
-	virtual physx::PxSimulationEventCallback* Create(FPhysScene_ImmediatePhysX* PhysScene, int32 SceneType) = 0;
-	virtual void Destroy(physx::PxSimulationEventCallback* Callback) = 0;
-};
-#endif // WITH PHYSX
-
-/** Interface for the creation of contact modify callbacks. */
-class IContactModifyCallbackFactory
-{
-public:
-	virtual FContactModifyCallback* Create(FPhysScene_ImmediatePhysX* PhysScene, int32 SceneType) = 0;
-	virtual void Destroy(FContactModifyCallback* Callback) = 0;
-};
-
-/** Container object for a physics engine 'scene'. */
-
-class FPhysScene_ImmediatePhysX : public FGenericPhysicsInterface
-{
-    friend struct FPhysicsInterface_ImmediatePhysX;
-    friend struct FPhysicsMaterialReference_ImmediatePhysX;
-    friend struct FPhysicsConstraintReference_ImmediatePhysX;
-    friend struct FPhysicsShapeReference_ImmediatePhysX;
-	
-    /** Holds shape data*/
-	struct FKinematicTarget
-	{
-	#if WITH_PHYSX
-		PxTransform BodyToWorld;
-	#endif
-		bool bTargetSet;
-	
-		FKinematicTarget()
-			: bTargetSet(false)
-		{
-		}
-	};
-
-	/** Contact pair generated between entities */
-	struct FContactPair
-	{
-		/** Index of the dynamic actor that we generated the contact pair for*/
-		int32 DynamicActorDataIndex;
-
-		/** Index of the other actor that we generated the contact pair for. This could be either dynamic or static */
-		uint32 OtherActorDataIndex;
-
-		/** Index into the first contact point associated with this pair*/
-		uint32 StartContactIndex;
-
-		/** Number of contacts associated with this pair. */
-		uint32 NumContacts;
-
-		/** Identifies the pair index from the original contact generation test */
-		int32 PairIdx;
-	};
-
-	struct D6JointData
-	{
-		D6JointData(PxD6Joint* Joint);
-
-		/** End solver API */
-
-
-		PxConstraintInvMassScale	invMassScale;
-		PxTransform					c2b[2];
-
-		PxU32					locked;		// bitmap of locked DOFs
-		PxU32					limited;	// bitmap of limited DOFs
-		PxU32					driving;	// bitmap of active drives (implies driven DOFs not locked)
-
-	
-		PxD6Motion::Enum		motion[6];
-		PxJointLinearLimit		linearLimit;
-		PxJointAngularLimitPair	twistLimit;
-		PxJointLimitCone		swingLimit;
-
-		PxD6JointDrive			drive[PxD6Drive::eCOUNT];
-
-		PxTransform				drivePosition;
-		PxVec3					driveLinearVelocity;
-		PxVec3					driveAngularVelocity;
-
-		// derived quantities
-
-	
-											// tan-half and tan-quarter angles
-
-		PxReal					thSwingY;
-		PxReal					thSwingZ;
-		PxReal					thSwingPad;
-
-		PxReal					tqSwingY;
-		PxReal					tqSwingZ;
-		PxReal					tqSwingPad;
-
-		PxReal					tqTwistLow;
-		PxReal					tqTwistHigh;
-		PxReal					tqTwistPad;
-
-		PxReal					linearMinDist;	// linear limit minimum distance to get a good direction
-
-												// projection quantities
-		//PxReal					projectionLinearTolerance;
-		//PxReal					projectionAngularTolerance;
-
-		FTransform					ActorToBody[2];
-
-		bool HasConstraints() const
-		{
-			return locked || limited || driving;
-		}
-	};
-
-	struct FLinearBlockAllocator
-	{
-		const static int PageBufferSize = 1024 * 64;
-
-		struct FPageStruct
-		{
-			//We assume FPageStruct is allocated with 16-byte alignment. Do not move Buffer. If we get full support for alignas(16) we could make this better
-			uint8 Buffer[PageBufferSize];
-
-			FPageStruct* NextPage;
-			FPageStruct* PrevPage;
-			int32 SeekPosition;
-
-			FPageStruct()
-				: NextPage(nullptr)
-				, PrevPage(nullptr)
-				, SeekPosition(0)
-			{
-			}
-		};
-
-		FPageStruct* FreePage;
-		FPageStruct* FirstPage;
-
-		FLinearBlockAllocator()
-		{
-			FreePage = AllocPage();
-			FirstPage = FreePage;
-		}
-	
-		FPageStruct* AllocPage()
-		{
-			LLM_SCOPE(ELLMTag::PhysX);
-			FPageStruct* ReturnPage = (FPageStruct*)FMemory::Malloc(sizeof(FPageStruct), 16);
-			new(ReturnPage) FPageStruct();
-		
-			FPlatformMisc::TagBuffer("ImmediatePhysicsSim", 0, (const void*)ReturnPage, sizeof(FPageStruct));
-			return ReturnPage;
-		}
-
-		void ReleasePage(FPageStruct* Page)
-		{
-			FMemory::Free(Page);
-		}
-
-		uint8* Alloc(int32 Bytes)
-		{
-			check(Bytes < PageBufferSize);	//Page size needs to be increased since we don't allow spillover
-			if (Bytes)
-			{
-				//Assumes 16 byte alignment
-				int32 BytesLeft = PageBufferSize - FreePage->SeekPosition;	//don't switch to uint because negative implies we're out of 16 byte aligned space
-				if (BytesLeft < Bytes)
-				{
-					//no space so allocate new page if needed
-					if (FreePage->NextPage)
-					{
-						FreePage = FreePage->NextPage;
-					}
-					else
-					{
-						FPageStruct* NewPage = AllocPage();
-						NewPage->PrevPage = FreePage;
-						FreePage->NextPage = NewPage;
-						FreePage = NewPage;
-					}
-				}
-
-				uint32 ReturnSlot = FreePage->SeekPosition;
-				FreePage->SeekPosition = (FreePage->SeekPosition + Bytes + 15)&(~15);
-
-				return &FreePage->Buffer[ReturnSlot];
-			}
-			else
-			{
-				return nullptr;
-			}
-		}
-
-		void Reset()
-		{
-			for (FPageStruct* Page = FirstPage; Page; Page = Page->NextPage)
-			{
-				Page->SeekPosition = 0;
-			}
-
-			FreePage = FirstPage;
-		}
-
-		void Empty()
-		{
-			FPageStruct* CurrentPage = FirstPage->NextPage;
-			while (CurrentPage)
-			{
-				FPageStruct* OldPage = CurrentPage;
-				CurrentPage = CurrentPage->NextPage;
-				ReleasePage(OldPage);
-			}
-
-			FirstPage->NextPage = nullptr;
-			FirstPage->SeekPosition = 0;
-			FreePage = FirstPage;
-			//Note we do not deallocate FirstPage until the allocator deallocates
-		}
-
-		~FLinearBlockAllocator()
-		{
-			Empty();
-			ReleasePage(FirstPage);
-		}
-
-	private:
-		//Don't copy these around
-		FLinearBlockAllocator(const FLinearBlockAllocator& Other);
-		const FLinearBlockAllocator& operator=(const FLinearBlockAllocator& Other);
-	};
-
-	/** TODO: Use a smarter memory allocator */
-	struct FCacheAllocator : public PxCacheAllocator
-	{
-		FCacheAllocator() : External(0){}
-		PxU8* allocateCacheData(const PxU32 ByteSize) override
-		{
-			return BlockAllocator[External].Alloc(ByteSize);
-		}
-
-		void Reset()
-		{
-	#if PERSISTENT_CONTACT_PAIRS
-			External = 1 - External;	//flip buffer so we maintain cache for 1 extra step
-	#endif
-			BlockAllocator[External].Reset();
-		}
-
-		FLinearBlockAllocator BlockAllocator[2];
-		int32 External;
-	};
-
-	/** TODO: Use a smarter memory allocator */
-	class FConstraintAllocator : public PxConstraintAllocator
-	{
-	public:
-		FConstraintAllocator() : External(0){}
-
-		PxU8* reserveConstraintData(const PxU32 ByteSize) override
-		{
-			return BlockAllocator[External].Alloc(ByteSize);
-
-		}
-
-		PxU8* reserveFrictionData(const PxU32 ByteSize)
-		{
-			return BlockAllocator[External].Alloc(ByteSize);
-		}
-
-		void Reset()
-		{
-	#if PERSISTENT_CONTACT_PAIRS
-			External = 1 - External;	//flip buffer so we maintain cache for 1 extra step
-	#endif
-			BlockAllocator[External].Reset();
-		}
-
-		FLinearBlockAllocator BlockAllocator[2];
-		int32 External;
-	};
-
-	struct FJoint
-	{
-        int32 ParentIndex, ChildIndex;
-        FTransform JointToParent, JointToChild;
-
-        FJoint(const int32 ParentIndexIn, const int32 ChildIndexIn, const FTransform JointToParentIn, const FTransform JointToChildIn)
-            : ParentIndex(ParentIndexIn), ChildIndex(ChildIndexIn), JointToParent(JointToParentIn), JointToChild(JointToChildIn)
-        {
-        }
-	};
-
-	struct FMaterial
-	{
-		FMaterial()
-			: StaticFriction(0.7f)
-			, DynamicFriction(0.7f)
-			, Restitution(0.3f)
-		{
-		}
-
-		FMaterial(physx::PxMaterial* InPxMaterial)
-		    : StaticFriction(InPxMaterial->getStaticFriction())
-		    , DynamicFriction(InPxMaterial->getDynamicFriction())
-		    , Restitution(InPxMaterial->getRestitution())
-		    , FrictionCombineMode((EFrictionCombineMode::Type)InPxMaterial->getFrictionCombineMode())
-		    , RestitutionCombineMode((EFrictionCombineMode::Type)InPxMaterial->getRestitutionCombineMode())
-	    {
-	    }
-
-        bool operator==(const FMaterial& Other) const
-        {
-            return StaticFriction == Other.StaticFriction && DynamicFriction == Other.DynamicFriction && Restitution == Other.Restitution;
-        }
-
-		float StaticFriction;
-		float DynamicFriction;
-		float Restitution;
-
-		EFrictionCombineMode::Type FrictionCombineMode;
-		EFrictionCombineMode::Type RestitutionCombineMode;
-	};
-
-	/** Holds shape data*/
-	struct FShape
-	{
-	#if WITH_PHYSX
-		const PxTransform LocalTM;
-		const FMaterial Material;
-		const PxGeometry* Geometry;
-		const PxVec3 BoundsOffset;
-		const float BoundsMagnitude;
-
-		FShape(const PxTransform& InLocalTM, const PxVec3& InBoundsOffset, const float InBoundsMagnitude, const PxGeometry* InGeometry, const FMaterial& InMaterial)
-			: LocalTM(InLocalTM)
-			, Material(InMaterial)
-			, Geometry(InGeometry)
-			, BoundsOffset(InBoundsOffset)
-			, BoundsMagnitude(InBoundsMagnitude)
-		{
-		}
-
-        FShape& operator=(const FShape& Other)
-        {
-            const_cast<PxTransform&>(LocalTM) = Other.LocalTM;
-            const_cast<FMaterial&>(Material) = Other.Material;
-            const_cast<PxGeometry*&>(Geometry) = const_cast<PxGeometry*>(Other.Geometry);
-            const_cast<PxVec3&>(BoundsOffset) = Other.BoundsOffset;
-            const_cast<float&>(BoundsMagnitude) = Other.BoundsMagnitude;
-            return *this;
-        }
-
-        bool operator==(const FShape& Other) const
-        {
-            return LocalTM == Other.LocalTM && Material == Other.Material && Geometry == Other.Geometry && BoundsOffset == Other.BoundsOffset && BoundsMagnitude == Other.BoundsMagnitude;
-        }
-	#endif
-	};
-
-	/** Holds geometry data*/
-	struct FActor
-	{
-		TArray<FShape> Shapes;
-
-	#if WITH_PHYSX
-		/** Create geometry data for the entity */
-		void CreateGeometry(PxRigidActor* RigidActor, const PxTransform& ActorToBodyTM);
-	#endif
-
-		/** Ensures all the geometry data has been properly freed */
-		void TerminateGeometry();
-	};
-
-	friend struct FPhysicsActorReference_ImmediatePhysX;
-	friend struct FPhysicsInterface_ImmediatePhysX;
-public:
-	ENGINE_API FPhysScene_ImmediatePhysX();
-	ENGINE_API ~FPhysScene_ImmediatePhysX();
-
-	void SwapActorData(uint32 Actor1DataIdx, uint32 Actor2DataIdx);
-	void ResizeActorData(uint32 ActorDataLen);
-
-	//////////////////////////////////////////////////////////////////////////
-	// PhysicsInterface
-
-	//// Actor creation/registration
-	//void ReleaseActor(FPhysicsActorReference& InActor);
-	void AddActorsToScene_AssumesLocked(const TArray<FPhysicsActorHandle>& InActors);
-	void AddAggregateToScene(const FPhysicsAggregateHandle& InAggregate, bool bUseAsyncScene); // #PHYS2 Remove bUseAsyncScene flag somehow
-
-	//Owning world is made private so that any code which depends on setting an owning world can update
-    void SetOwningWorld(UWorld* InOwningWorld) { OwningWorld = InOwningWorld; }
-	UWorld* GetOwningWorld() { return OwningWorld; }
-	const UWorld* GetOwningWorld() const { return OwningWorld; }
-	
-    FPhysicsReplication* GetPhysicsReplication() { return nullptr; }
-
-	/** Lets the scene update anything related to this BodyInstance as it's now being terminated */
-	void RemoveBodyInstanceFromPendingLists_AssumesLocked(FBodyInstance* BodyInstance, int32 SceneType);
-
-	/** Add a custom callback for next step that will be called on every substep */
-	void AddCustomPhysics_AssumesLocked(FBodyInstance* BodyInstance, FCalculateCustomPhysics& CalculateCustomPhysics);
-
-	/** Adds a force to a body - We need to go through scene to support substepping */
-	void AddForce_AssumesLocked(FBodyInstance* BodyInstance, const FVector& Force, bool bAllowSubstepping, bool bAccelChange);
-
-	/** Adds a force to a body at a specific position - We need to go through scene to support substepping */
-	void AddForceAtPosition_AssumesLocked(FBodyInstance* BodyInstance, const FVector& Force, const FVector& Position, bool bAllowSubstepping, bool bIsLocalForce = false);
-
-	/** Adds a radial force to a body - We need to go through scene to support substepping */
-	void AddRadialForceToBody_AssumesLocked(FBodyInstance* BodyInstance, const FVector& Origin, const float Radius, const float Strength, const uint8 Falloff, bool bAccelChange, bool bAllowSubstepping);
-
-	/** Clears currently accumulated forces on a specified body instance */
-	void ClearForces_AssumesLocked(FBodyInstance* BodyInstance, bool bAllowSubstepping);
-
-	/** Adds torque to a body - We need to go through scene to support substepping */
-	void AddTorque_AssumesLocked(FBodyInstance* BodyInstance, const FVector& Torque, bool bAllowSubstepping, bool bAccelChange);
-
-	/** Clears currently accumulated torques on a specified body instance */
-	void ClearTorques_AssumesLocked(FBodyInstance* BodyInstance, bool bAllowSubstepping);
-
-	/** Sets a Kinematic actor's target position - We need to do this here to support substepping*/
-	void SetKinematicTarget_AssumesLocked(FBodyInstance* BodyInstance, const FTransform& TargetTM, bool bAllowSubstepping);
-
-	/** Gets a Kinematic actor's target position - We need to do this here to support substepping
-	* Returns true if kinematic target has been set. If false the OutTM is invalid */
-	bool GetKinematicTarget_AssumesLocked(const FBodyInstance* BodyInstance, FTransform& OutTM) const;
-
-	/** Gets the collision disable table */
-	const TMap<uint32, TMap<struct FRigidBodyIndexPair, bool> *> & GetCollisionDisableTableLookup()
-	{
-		return CollisionDisableTableLookup;
-	}
-
-	/** Adds to queue of skelmesh we want to add to collision disable table */
-	ENGINE_API void DeferredAddCollisionDisableTable(uint32 SkelMeshCompID, TMap<struct FRigidBodyIndexPair, bool> * CollisionDisableTable);
-
-	/** Adds to queue of skelmesh we want to remove from collision disable table */
-	ENGINE_API void DeferredRemoveCollisionDisableTable(uint32 SkelMeshCompID);
-
-	/** Add this SkeletalMeshComponent to the list needing kinematic bodies updated before simulating physics */
-	bool MarkForPreSimKinematicUpdate(USkeletalMeshComponent* InSkelComp, ETeleportType InTeleport, bool bNeedsSkinning);
-
-	/** Remove this SkeletalMeshComponent from set needing kinematic update before simulating physics*/
-	void ClearPreSimKinematicUpdate(USkeletalMeshComponent* InSkelComp);
-
-	/** Pending constraint break events */
-	void AddPendingOnConstraintBreak(FConstraintInstance* ConstraintInstance, int32 SceneType);
-	/** Pending wake/sleep events */
-	void AddPendingSleepingEvent(FBodyInstance* BI, ESleepEvent SleepEventType, int32 SceneType);
-
-	/** Gets the array of collision notifications, pending execution at the end of the physics engine run. */
-	TArray<FCollisionNotifyInfo>& GetPendingCollisionNotifies(int32 SceneType) { return PendingCollisionData.PendingCollisionNotifies; }
-
-	/** @return Whether physics scene supports scene origin shifting */
-	static bool SupportsOriginShifting() { return true; }
-
-	/** Shifts physics scene origin by specified offset */
-	void ApplyWorldOffset(FVector InOffset);
-
-	/** Set the gravity and timing of all physics scenes */
-	ENGINE_API void SetUpForFrame(const FVector* NewGrav, float InDeltaSeconds = 0.0f, float InMaxPhysicsDeltaTime = 0.0f, float InMaxSubstepDeltaTime = 0.f, int32 InMaxSubsteps = 1);
-
-	/** Starts a frame */
-	ENGINE_API void StartFrame();
-
-	/** Ends a frame */
-	ENGINE_API void EndFrame(ULineBatchComponent* InLineBatcher);
-
-	/** Waits for all physics scenes to complete */
-	ENGINE_API void WaitPhysScenes();
-
-	/** returns the completion event for a frame */
-	FGraphEventRef GetCompletionEvent()
-	{
-		return PhysicsSceneCompletion;
-	}
-
-	/** Handle exec commands related to scene (PXVIS and APEXVIS) */
-	bool HandleExecCommands(const TCHAR* Cmd, FOutputDevice* Ar);
-
-	void ListAwakeRigidBodies(bool bIncludeKinematic);
-
-	ENGINE_API int32 GetNumAwakeBodies();
-
-	/** Static factory used to override the simulation contact modify callback from other modules.*/
-	ENGINE_API static TSharedPtr<IContactModifyCallbackFactory> ContactModifyCallbackFactory;
-
-	/** Static factory used to override the physics replication manager from other modules. This is useful for custom game logic.
-	If not set it defaults to using FPhysicsReplication. */
-	ENGINE_API static TSharedPtr<IPhysicsReplicationFactory> PhysicsReplicationFactory;
-
-	//////////////////////////////////////////////////////////////////////////
-	// PhysScene_PhysX interface
-
-	/** Utility for looking up the PxScene of the given EPhysicsSceneType associated with this FPhysScene.  SceneType must be in the range [0,PST_MAX). */
-    ENGINE_API physx::PxScene* GetPxScene(uint32 SceneType) const { return nullptr; }
-
-#if WITH_APEX
-	/** Utility for looking up the ApexScene of the given EPhysicsSceneType associated with this FPhysScene.  SceneType must be in the range [0,PST_MAX). */
-    ENGINE_API nvidia::apex::Scene*	GetApexScene(uint32 SceneType) const { return nullptr; }
-#endif
-
-	/** Starts cloth Simulation*/
-    ENGINE_API void StartAsync() {}
-
-	/** Returns whether an async scene is setup and can be used. This depends on the console variable "p.EnableAsyncScene". */
-	ENGINE_API bool HasAsyncScene() const { return false; }
-
-	/** Ensures that the collision tree is built. */
-    ENGINE_API void EnsureCollisionTreeIsBuilt(UWorld* World) {}
-
-	/** The number of frames it takes to rebuild the PhysX scene query AABB tree. The bigger the number, the smaller fetchResults takes per frame, but the more the tree deteriorates until a new tree is built */
-    void SetPhysXTreeRebuildRate(int32 RebuildRate) { PhysXTreeRebuildRate = RebuildRate; }
-
-	/** Kill the visual debugger */
-    ENGINE_API void KillVisualDebugger();
-
-#if WITH_PHYSX
-	/** Static factory used to override the simulation event callback from other modules.
-	If not set it defaults to using FPhysXSimEventCallback. */
-	ENGINE_API static TSharedPtr<ISimEventCallbackFactory> SimEventCallbackFactory;
-
-#endif	// WITH_PHYSX
-
-	//////////////////////////////////////////////////////////////////////////
-
-	DECLARE_MULTICAST_DELEGATE_ThreeParams(FOnPhysScenePreTick, FPhysScene_ImmediatePhysX*, uint32 /*SceneType*/, float /*DeltaSeconds*/);
-	FOnPhysScenePreTick OnPhysScenePreTick;
-
-	DECLARE_MULTICAST_DELEGATE_ThreeParams(FOnPhysSceneStep, FPhysScene_ImmediatePhysX*, uint32 /*SceneType*/, float /*DeltaSeconds*/);
-	FOnPhysSceneStep OnPhysSceneStep;
-
-	const TArray<immediate::PxRigidBodyData>& GetRigidBodiesData() { return RigidBodiesData; }
-
-private:
-	/** Indicates whether the scene is using substepping */
-	bool							bSubstepping;
-
-	/** World that owns this physics scene */
-	UWorld*							OwningWorld;
-
-	/** DeltaSeconds from UWorld. */
-	float										DeltaSeconds;
-	/** DeltaSeconds from the WorldSettings. */
-	float										MaxPhysicsDeltaTime;
-	/** DeltaSeconds used by the last synchronous scene tick.  This may be used for the async scene tick. */
-	float										SyncDeltaSeconds;
-	/** LineBatcher from UWorld. */
-	ULineBatchComponent*						LineBatcher;
-
-	/** Completion events (task) for the physics scenes	(both apex and non-apex). This is a "join" of the above. */
-	FGraphEventRef PhysicsSceneCompletion;
-
-#if WITH_PHYSX
-	struct FPendingCollisionData
-	{
-		/** Array of collision notifications, pending execution at the end of the physics engine run. */
-		TArray<FCollisionNotifyInfo>	PendingCollisionNotifies;
-	};
-
-	FPendingCollisionData PendingCollisionData;
-
-	struct FPendingConstraintData
-	{
-		/** Array of constraint broken notifications, pending execution at the end of the physics engine run. */
-		TArray<FConstraintBrokenDelegateData> PendingConstraintBroken;
-	};
-
-	FPendingConstraintData PendingConstraintData;
-#endif	// WITH_PHYSX
-
-	/** Sync components in the scene to physics bodies that changed */
-	void SyncComponentsToBodies_AssumesLocked(uint32 SceneType);
-
-	/** Call after WaitPhysScene on the synchronous scene to make deferred OnRigidBodyCollision calls.  */
-	ENGINE_API void DispatchPhysNotifications_AssumesLocked();
-
-	/** Add any debug lines from the physics scene of the given type to the supplied line batcher. */
-	ENGINE_API void AddDebugLines(uint32 SceneType, class ULineBatchComponent* LineBatcherToUse);
-
-	/** Helper function for determining which scene a dyanmic body is in*/
-	EPhysicsSceneType SceneType_AssumesLocked(const FBodyInstance* BodyInstance) const;
-
-	/** Set whether we're doing a static load and want to stall, or are during gameplay and want to distribute over many frames */
-	void SetIsStaticLoading(bool bStaticLoading);
-
-#if WITH_PHYSX
-	/** User data wrapper passed to physx */
-	FPhysxUserData PhysxUserData;
-
-#endif
-
-	class FPhysSubstepTask * PhysSubSteppers[PST_MAX];
-
-	struct FPendingCollisionDisableTable
-	{
-		uint32 SkelMeshCompID;
-		TMap<struct FRigidBodyIndexPair, bool>* CollisionDisableTable;
-	};
-
-	/** Updates CollisionDisableTableLookup with the deferred insertion and deletion */
-	void FlushDeferredCollisionDisableTableQueue();
-
-	bool ExecPxVis(uint32 SceneType, const TCHAR* Cmd, FOutputDevice* Ar);
-
-	bool ExecApexVis(uint32 SceneType, const TCHAR* Cmd, FOutputDevice* Ar);
-
-
-	/** Queue of deferred collision table insertion and deletion */
-	TArray<FPendingCollisionDisableTable> DeferredCollisionDisableTableQueue;
-
-	/** Map from SkeletalMeshComponent UniqueID to a pointer to the collision disable table inside its PhysicsAsset */
-	TMap< uint32, TMap<struct FRigidBodyIndexPair, bool>* >		CollisionDisableTableLookup;
-
-#if WITH_PHYSX
-	TMap<FBodyInstance*, ESleepEvent> PendingSleepEvents;
-#endif
-
-    TArray<FBodyInstance*> BodyInstances;
-
-	FDelegateHandle PreGarbageCollectDelegateHandle;
-
-	int32 PhysXTreeRebuildRate;
-
-	// TAKEN FROM IMMEDIATEPHYSICS
-	/** Entities holding loose data. NOTE: for performance reasons we don't automatically cleanup on destructor (needed for tarray swaps etc...) it's very important that Terminate is called */
-	TArray<FActor> Actors;
-
-	TArray<FJoint> Joints;
-
-	/** Workspace memory that we use for per frame allocations */
-	FLinearBlockAllocator Workspace;
-
-#if WITH_PHYSX
-	/** Low level rigid body data */
-	TArray<immediate::PxRigidBodyData> RigidBodiesData;
-	
-	/** Low level solver bodies data */
-	TArray<PxSolverBodyData> SolverBodiesData;
-	
-	/** Kinematic targets used to implicitly compute the velocity of moving kinematic actors */
-	TArray<FKinematicTarget> KinematicTargets;
-
-	TArray<PxVec3> PendingAcceleration;
-	TArray<PxVec3> PendingVelocityChange;
-	TArray<PxVec3> PendingAngularAcceleration;
-	TArray<PxVec3> PendingAngularVelocityChange;
-
-	/** Low level contact points generated for this frame. Points are grouped together by pairs */
-	TArray<Gu::ContactPoint> ContactPoints;
-
-	/** Shapes used in the entire simulation. Shapes are sorted in the same order as actors. Note that an actor can have multiple shapes which will be adjacent*/
-	struct FShapeSOA
-	{
-		TArray<PxTransform> LocalTMs;
-		TArray<FMaterial> Materials;
-		TArray<const PxGeometry*> Geometries;
-		TArray<float> Bounds;
-		TArray<PxVec3> BoundsOffsets;
-		TArray<int32> OwningActors;
-#if PERSISTENT_CONTACT_PAIRS
-		TArray<FPersistentContactPairData> ContactPairData;
-#endif
-	} ShapeSOA;
-	
-	/** Low level solver bodies */
-	PxSolverBody* SolverBodies;
-
-	/** Low level constraint descriptors.*/
-	TArray<PxSolverConstraintDesc> OrderedDescriptors;
-	TArray<PxConstraintBatchHeader> BatchHeaders;
-
-	/** JointData as passed in from physics constraint template */
-	TArray<D6JointData> JointData;
-
-	/** When new joints are created we have to update the processing order */
-	bool bDirtyJointData;
-
-	PxU32 NumContactHeaders;
-	PxU32 NumJointHeaders;
-	uint32 NumActiveJoints;
-#endif // WITH_PHYSX
-
-	/** Contact pairs generated for this frame */
-	TArray<FContactPair> ContactPairs;
-
-	/** Number of dynamic bodies associated with the simulation */
-	uint32 NumSimulatedBodies;
-
-	/** Number of dynamic bodies that are actually active */
-	uint32 NumActiveSimulatedBodies;
-
-	/** Number of kinematic bodies (dynamic but not simulated) associated with the simulation */
-	uint32 NumKinematicBodies;
-
-	/** Total number of simulated shapes in the scene */
-	uint32 NumSimulatedShapesWithCollision;
-
-	/** Number of position iterations used by solver */
-	uint32 NumPositionIterations;
-
-	/** Number of velocity iterations used by solver */
-	uint32 NumVelocityIterations;
-
-	/** Count of how many times we've ticked. Useful for cache invalidation */
-	uint32 SimCount;
-
-	/** This cache is used to record which generate contact iteration we can skip. This assumes the iteration order has not changed (add/remove/swap actors must invalidate this) */
-	bool bRecreateIterationCache;
-
-	TArray<int32> SkipCollisionCache;	//Holds the iteration count that we should skip due to ignore filtering
-
-	friend struct FContactPointRecorder;
-
-#if WITH_PHYSX
-	FCacheAllocator CacheAllocator;
-	FConstraintAllocator ConstraintAllocator;
-#endif // WITH_PHYSX
-};
-
-#endif
diff --git a/Engine/Source/Runtime/Engine/Public/Physics/Experimental/PhysicsInterfaceImmediatePhysX.h b/Engine/Source/Runtime/Engine/Public/Physics/Experimental/PhysicsInterfaceImmediatePhysX.h
--- a/Engine/Source/Runtime/Engine/Public/Physics/Experimental/PhysicsInterfaceImmediatePhysX.h
+++ b/Engine/Source/Runtime/Engine/Public/Physics/Experimental/PhysicsInterfaceImmediatePhysX.h
@@ -1,450 +1,3 @@
 // Copyright Epic Games, Inc. All Rights Reserved.
 
 #pragma once
-
-#if WITH_IMMEDIATE_PHYSX
-
-#include "EngineGlobals.h"
-#include "Engine/EngineTypes.h"
-#include "CollisionQueryParams.h"
-#include "PhysicsEngine/ConstraintTypes.h"
-#include "PhysicsInterfaceTypes.h"
-#include "PhysScene_ImmediatePhysX.h"
-
-struct FLinearDriveConstraint;
-struct FAngularDriveConstraint;
-struct FConstraintDrive;
-struct FPhysxUserData;
-struct FKAggregateGeom;
-
-struct FCollisionQueryParams;
-struct FCollisionResponseParams;
-struct FCollisionObjectQueryParams;
-struct FCollisionShape;
-class UWorld;
-
-namespace physx
-{
-	class PxRigidActor;
-	class PxRigidStatic;
-	class PxRigidDynamic;
-	class PxRigidBody;
-	class PxAggregate;
-	class PxShape;
-	class PxD6Joint;
-	class PxGeometry;
-	class PxQuat;
-	class PxTransform;
-	class PxScene;
-}
-
-struct ENGINE_API FPhysicsActorReference_ImmediatePhysX
-{
-	FPhysicsActorReference_ImmediatePhysX();
-
-	bool IsValid() const;
-	bool Equals(const FPhysicsActorReference_ImmediatePhysX& Other) const;
-	bool operator==(const FPhysicsActorReference_ImmediatePhysX& Other) const { return Equals(Other); }
-
-	FPhysScene_ImmediatePhysX* Scene;
-	uint32 Index;
-};
-
-struct ENGINE_API FPhysicsConstraintReference_ImmediatePhysX
-{
-	FPhysicsConstraintReference_ImmediatePhysX();
-
-	bool IsValid() const;
-	bool Equals(const FPhysicsConstraintReference_ImmediatePhysX& Other) const;
-
-	FPhysScene_ImmediatePhysX* Scene;
-	uint32 Index;
-};
-
-struct ENGINE_API FPhysicsAggregateReference_ImmediatePhysX
-{
-	FPhysicsAggregateReference_ImmediatePhysX();
-
-	bool IsValid() const;
-
-	FPhysScene_ImmediatePhysX* Scene;
-	TArray<uint32> Indices;
-};
-
-struct ENGINE_API FPhysicsShapeReference_ImmediatePhysX
-{
-	FPhysicsShapeReference_ImmediatePhysX(FPhysScene_ImmediatePhysX::FShape&& Shape);
-
-	bool IsValid() const;
-    bool operator==(const FPhysicsShapeReference_ImmediatePhysX& Other) const { return Actor == Other.Actor && Index == Other.Index && Shape == Other.Shape; }
-
-    FPhysScene_ImmediatePhysX::FShape Shape;
-    FPhysicsActorReference_ImmediatePhysX* Actor;
-    int32 Index;
-};
-
-// Dummy holder for materials
-struct ENGINE_API FPhysicsMaterialReference_ImmediatePhysX
-{
-    FPhysScene_ImmediatePhysX::FMaterial Material;
-
-    bool IsValid() const { return true; }
-};
-
-struct ENGINE_API FPhysicsGeometryCollection_ImmediatePhysX
-{
-    FPhysicsGeometryCollection_ImmediatePhysX() {}
-    
-    bool IsValid() const { return Geometry != nullptr; }
-
-    physx::PxGeometry* Geometry;
-
-    ECollisionShapeType GetType() const
-    {
-        if (!IsValid())
-        {
-            return ECollisionShapeType::None;
-        }
-        if (Geometry->getType() == PxGeometryType::Enum::eBOX)
-        {
-            return ECollisionShapeType::Box;
-        }
-        if (Geometry->getType() == PxGeometryType::Enum::eCAPSULE)
-        {
-            return ECollisionShapeType::TCapsule;
-        }
-        if (Geometry->getType() == PxGeometryType::Enum::eCONVEXMESH)
-        {
-            return ECollisionShapeType::Convex;
-        }
-        if (Geometry->getType() == PxGeometryType::Enum::eHEIGHTFIELD)
-        {
-            return ECollisionShapeType::Heightfield;
-        }
-        if (Geometry->getType() == PxGeometryType::Enum::eSPHERE)
-        {
-            return ECollisionShapeType::Sphere;
-        }
-        if (Geometry->getType() == PxGeometryType::Enum::eTRIANGLEMESH)
-        {
-            return ECollisionShapeType::Trimesh;
-        }
-        return ECollisionShapeType::None;
-    }
-#if 0	//Not sure if this is needed, refactoring this code for Chaos usage so leaving as is for now
-    physx::PxGeometry& GetGeometry() const { return *Geometry; }
-    bool GetBoxGeometry(physx::PxBoxGeometry& OutGeom) const
-    {
-        return GPhysXSDK->createShape(GetGeometry(), nullptr, 0)->getBoxGeometry(OutGeom);
-    }
-    bool GetSphereGeometry(physx::PxSphereGeometry& OutGeom) const
-    {
-        return GPhysXSDK->createShape(GetGeometry(), nullptr, 0)->getSphereGeometry(OutGeom);
-    }
-    bool GetCapsuleGeometry(physx::PxCapsuleGeometry& OutGeom) const
-    {
-        return GPhysXSDK->createShape(GetGeometry(), nullptr, 0)->getCapsuleGeometry(OutGeom);
-    }
-    bool GetConvexGeometry(physx::PxConvexMeshGeometry& OutGeom) const
-    {
-        return GPhysXSDK->createShape(GetGeometry(), nullptr, 0)->getConvexMeshGeometry(OutGeom);
-    }
-    bool GetTriMeshGeometry(physx::PxTriangleMeshGeometry& OutGeom) const
-    {
-        return GPhysXSDK->createShape(GetGeometry(), nullptr, 0)->getTriangleMeshGeometry(OutGeom);
-    }
-#endif
-};
-
-// @todo(mlentine): Once we finalize the format of this we need to redo this
-FORCEINLINE uint32 GetTypeHash(const FPhysicsShapeReference_ImmediatePhysX& InShapeReference)
-{
-    return GetTypeHash(InShapeReference.Index);
-}
-
-/**
-* API to access the physics interface. All calls to FPhysicsInterface functions should be inside an Execute* callable.
-* This is to ensure correct lock semantics and command buffering if the specific API supports deferred commands.
-*/
-struct ENGINE_API FPhysicsCommand_ImmediatePhysX
-{
-	// Executes with appropriate read locking, return true if execution took place (actor was valid)
-	static bool ExecuteRead(const FPhysicsActorReference_ImmediatePhysX& InActorReference, TFunctionRef<void(const FPhysicsActorReference_ImmediatePhysX& Actor)> InCallable);
-	static bool ExecuteRead(USkeletalMeshComponent* InMeshComponent, TFunctionRef<void()> InCallable);
-	static bool ExecuteRead(const FPhysicsActorReference_ImmediatePhysX& InActorReferenceA, const FPhysicsActorReference_ImmediatePhysX& InActorReferenceB, TFunctionRef<void(const FPhysicsActorReference_ImmediatePhysX& ActorA, const FPhysicsActorReference_ImmediatePhysX& ActorB)> InCallable);
-	static bool ExecuteRead(const FPhysicsConstraintReference_ImmediatePhysX& InConstraintRef, TFunctionRef<void(const FPhysicsConstraintReference_ImmediatePhysX& Constraint)> InCallable);
-	static bool ExecuteRead(FPhysScene* InScene, TFunctionRef<void()> InCallable);
-
-	// Executes with appropriate write locking, return true if execution took place (actor was valid)
-	static bool ExecuteWrite(const FPhysicsActorReference_ImmediatePhysX& InActorReference, TFunctionRef<void(const FPhysicsActorReference_ImmediatePhysX& Actor)> InCallable);
-	static bool ExecuteWrite(USkeletalMeshComponent* InMeshComponent, TFunctionRef<void()> InCallable);
-	static bool ExecuteWrite(const FPhysicsActorReference_ImmediatePhysX& InActorReferenceA, const FPhysicsActorReference_ImmediatePhysX& InActorReferenceB, TFunctionRef<void(const FPhysicsActorReference_ImmediatePhysX& ActorA, const FPhysicsActorReference_ImmediatePhysX& ActorB)> InCallable);
-	static bool ExecuteWrite(const FPhysicsConstraintReference_ImmediatePhysX& InConstraintRef, TFunctionRef<void(const FPhysicsConstraintReference_ImmediatePhysX& Constraint)> InCallable);
-	static bool ExecuteWrite(FPhysScene* InScene, TFunctionRef<void()> InCallable);
-	
-    // Executes function on a shape, handling shared shapes
-	static void ExecuteShapeWrite(FBodyInstance* InInstance, FPhysicsShapeHandle& InShape, TFunctionRef<void(const FPhysicsShapeHandle& InShape)> InCallable);
-};
-
-struct ENGINE_API FPhysicsInterface_ImmediatePhysX
-{
-	// PhysX Only functions, not related to wider physics interface
-	// To be used only in code that handles PhysX
-    static physx::PxRigidActor* GetPxRigidActorFromScene_AssumesLocked(const FPhysicsActorReference_ImmediatePhysX& InActorRef, int32 SceneType = -1) { return nullptr; }
-    static physx::PxRigidActor* GetPxRigidActor_AssumesLocked(const FPhysicsActorReference_ImmediatePhysX& InRef) { return nullptr; }
-	//////////////////////////////////////////////////////////////////////////
-
-	// Aggregate interface functions
-	static FPhysicsAggregateReference_ImmediatePhysX CreateAggregate(int32 MaxBodies);
-	static void ReleaseAggregate(FPhysicsAggregateReference_ImmediatePhysX& InAggregate);
-	static int32 GetNumActorsInAggregate(const FPhysicsAggregateReference_ImmediatePhysX& InAggregate);
-	static void AddActorToAggregate_AssumesLocked(const FPhysicsAggregateReference_ImmediatePhysX& InAggregate, const FPhysicsActorReference_ImmediatePhysX& InActor);
-
-	// Shape interface functions
-	static FPhysicsShapeReference_ImmediatePhysX CreateShape(physx::PxGeometry* InGeom, bool bSimulation = true, bool bQuery = true, UPhysicalMaterial* InSimpleMaterial = nullptr, TArray<UPhysicalMaterial*>* InComplexMaterials = nullptr, bool bShared = false);
-
-	static void AddGeometry(const FPhysicsActorHandle& InActor, const FGeometryAddParams& InParams, TArray<FPhysicsShapeHandle>* OutOptShapes = nullptr);
-	static FPhysicsShapeReference_ImmediatePhysX CloneShape(const FPhysicsShapeReference_ImmediatePhysX& InShape);
-
-	static bool IsSimulationShape(const FPhysicsShapeReference_ImmediatePhysX& InShape);
-	static bool IsQueryShape(const FPhysicsShapeReference_ImmediatePhysX& InShape);
-    // @todo(mlentine): We don't keep track of what is shared but anything can be
-    static bool IsShared(const FPhysicsShapeHandle& InShape) { return true; }
-	static ECollisionShapeType GetShapeType(const FPhysicsShapeReference_ImmediatePhysX& InShape);
-	static FPhysicsGeometryCollection GetGeometryCollection(const FPhysicsShapeReference_ImmediatePhysX& InShape);
-	static FTransform GetLocalTransform(const FPhysicsShapeReference_ImmediatePhysX& InShape);
-	static void* GetUserData(const FPhysicsShapeReference_ImmediatePhysX& InShape);
-
-	// Set the mask filter of a shape, which is an extra level of filtering during collision detection / query for extra channels like "Blue Team" and "Red Team"
-	static void SetMaskFilter(const FPhysicsShapeReference_ImmediatePhysX& InShape, FMaskFilter InFilter);
-	static void SetSimulationFilter(const FPhysicsShapeReference_ImmediatePhysX& InShape, const FCollisionFilterData& InFilter);
-	static void SetQueryFilter(const FPhysicsShapeReference_ImmediatePhysX& InShape, const FCollisionFilterData& InFilter);
-	static void SetIsSimulationShape(const FPhysicsShapeReference_ImmediatePhysX& InShape, bool bIsSimShape);
-	static void SetIsQueryShape(const FPhysicsShapeReference_ImmediatePhysX& InShape, bool bIsQueryShape);
-	static void SetUserData(const FPhysicsShapeReference_ImmediatePhysX& InShape, void* InUserData);
-	static void SetGeometry(const FPhysicsShapeReference_ImmediatePhysX& InShape, physx::PxGeometry& InGeom);
-	static void SetLocalTransform(const FPhysicsShapeReference_ImmediatePhysX& InShape, const FTransform& NewLocalTransform);
-	static void SetMaterials(const FPhysicsShapeReference_ImmediatePhysX& InShape, const TArrayView<UPhysicalMaterial*>InMaterials);
-
-	// Material interface functions
-	static FPhysicsMaterialHandle CreateMaterial(const UPhysicalMaterial* InMaterial);
-	static void ReleaseMaterial(FPhysicsMaterialHandle& InHandle);
-	static void UpdateMaterial(const FPhysicsMaterialHandle& InHandle, UPhysicalMaterial* InMaterial);
-	static void SetUserData(const FPhysicsMaterialHandle& InHandle, void* InUserData);
-
-	// Actor interface functions
-
-	// #PHYS2 - These should be on the scene, but immediate mode stops us for now, eventually that should spawn its own minimal IM scene and these should move
-	static void CreateActor(const FActorCreationParams& Params, FPhysicsActorHandle& Handle);
-	static void ReleaseActor(FPhysicsActorReference_ImmediatePhysX& InActorReference, FPhysScene* InScene = nullptr);
-	//////////////////////////////////////////////////////////////////////////
-
-	template<typename AllocatorType>
-	static int32 GetAllShapes_AssumedLocked(const FPhysicsActorReference_ImmediatePhysX& InHandle, TArray<FPhysicsShapeHandle, AllocatorType>& OutShapes, EPhysicsSceneType InSceneType = PST_MAX);
-	static void GetNumShapes(const FPhysicsActorHandle& InHandle, int32& OutNumSyncShapes, int32& OutNumAsyncShapes);
-	static void ReleaseShape(const FPhysicsShapeHandle& InShape);
-
-	static void AttachShape(const FPhysicsActorHandle& InActor, const FPhysicsShapeHandle& InNewShape);
-	static void AttachShape(const FPhysicsActorHandle& InActor, const FPhysicsShapeHandle& InNewShape, EPhysicsSceneType SceneType);
-	static void DetachShape(const FPhysicsActorHandle& InActor, FPhysicsShapeHandle& InShape, bool bWakeTouching = true);
-
-	static void SetActorUserData_AssumesLocked(const FPhysicsActorReference_ImmediatePhysX& InActorReference, FPhysxUserData* InUserData);
-
-	static bool IsRigidBody(const FPhysicsActorReference_ImmediatePhysX& InActorReference);
-	static bool IsDynamic(const FPhysicsActorReference_ImmediatePhysX& InActorReference);
-	static bool IsStatic(const FPhysicsActorReference_ImmediatePhysX& InActorReference);
-	static bool IsKinematic_AssumesLocked(const FPhysicsActorReference_ImmediatePhysX& InActorReference);
-	static bool IsSleeping(const FPhysicsActorReference_ImmediatePhysX& InActorReference);
-	static bool IsCcdEnabled(const FPhysicsActorReference_ImmediatePhysX& InActorReference);
-	static bool IsInScene(const FPhysicsActorReference_ImmediatePhysX& InActorReference);
-    static bool HasSyncSceneData(const FPhysicsActorReference_ImmediatePhysX& InHandle) { return false; }
-    static bool HasAsyncSceneData(const FPhysicsActorReference_ImmediatePhysX& InHandle) { return false; }
-	static FPhysScene* GetCurrentScene(const FPhysicsActorReference_ImmediatePhysX& InActorReference);
-	static bool CanSimulate_AssumesLocked(const FPhysicsActorReference_ImmediatePhysX& InActorReference);
-	static float GetMass_AssumesLocked(const FPhysicsActorReference_ImmediatePhysX& InActorReference);
-
-	static void SetSendsSleepNotifies_AssumesLocked(const FPhysicsActorReference_ImmediatePhysX& InActorReference, bool bSendSleepNotifies);
-	static void PutToSleep_AssumesLocked(const FPhysicsActorReference_ImmediatePhysX& InActorReference);
-	static void WakeUp_AssumesLocked(const FPhysicsActorReference_ImmediatePhysX& InActorReference);
-
-	static void SetIsKinematic_AssumesLocked(const FPhysicsActorReference_ImmediatePhysX& InActorReference, bool bIsKinematic);
-	static void SetCcdEnabled_AssumesLocked(const FPhysicsActorReference_ImmediatePhysX& InActorReference, bool bIsCcdEnabled);
-
-	static FTransform GetGlobalPose_AssumesLocked(const FPhysicsActorReference_ImmediatePhysX& InActorReference);
-	static void SetGlobalPose_AssumesLocked(const FPhysicsActorReference_ImmediatePhysX& InActorReference, const FTransform& InNewPose, bool bAutoWake = true);
-
-	static FTransform GetTransform_AssumesLocked(const FPhysicsActorHandle& InRef, bool bForceGlobalPose = false);
-
-	static bool HasKinematicTarget_AssumesLocked(const FPhysicsActorReference_ImmediatePhysX& InActorReference);
-	static FTransform GetKinematicTarget_AssumesLocked(const FPhysicsActorReference_ImmediatePhysX& InActorReference);
-	static void SetKinematicTarget_AssumesLocked(const FPhysicsActorReference_ImmediatePhysX& InActorReference, const FTransform& InNewTarget);
-
-	static FVector GetLinearVelocity_AssumesLocked(const FPhysicsActorReference_ImmediatePhysX& InActorReference);
-	static void SetLinearVelocity_AssumesLocked(const FPhysicsActorReference_ImmediatePhysX& InActorReference, const FVector& InNewVelocity, bool bAutoWake = true);
-
-	static FVector GetAngularVelocity_AssumesLocked(const FPhysicsActorReference_ImmediatePhysX& InActorReference);
-	static void SetAngularVelocity_AssumesLocked(const FPhysicsActorReference_ImmediatePhysX& InActorReference, const FVector& InNewVelocity, bool bAutoWake = true);
-	static float GetMaxAngularVelocity_AssumesLocked(const FPhysicsActorReference_ImmediatePhysX& InActorReference);
-	static void SetMaxAngularVelocity_AssumesLocked(const FPhysicsActorReference_ImmediatePhysX& InActorReference, float InMaxAngularVelocity);
-
-	static float GetMaxDepenetrationVelocity_AssumesLocked(const FPhysicsActorReference_ImmediatePhysX& InActorReference);
-	static void SetMaxDepenetrationVelocity_AssumesLocked(const FPhysicsActorReference_ImmediatePhysX& InActorReference, float InMaxDepenetrationVelocity);
-
-	static FVector GetWorldVelocityAtPoint_AssumesLocked(const FPhysicsActorReference_ImmediatePhysX& InActorReference, const FVector& InPoint);
-
-	static FTransform GetComTransform_AssumesLocked(const FPhysicsActorReference_ImmediatePhysX& InActorReference);
-
-	static FVector GetLocalInertiaTensor_AssumesLocked(const FPhysicsActorReference_ImmediatePhysX& InActorReference);
-	static FBox GetBounds_AssumesLocked(const FPhysicsActorReference_ImmediatePhysX& InActorReference);
-
-	static void SetLinearDamping_AssumesLocked(const FPhysicsActorReference_ImmediatePhysX& InActorReference, float InDamping);
-	static void SetAngularDamping_AssumesLocked(const FPhysicsActorReference_ImmediatePhysX& InActorReference, float InDamping);
-
-	static void AddForce_AssumesLocked(const FPhysicsActorReference_ImmediatePhysX& InActorReference, const FVector& InForce);
-	static void AddTorque_AssumesLocked(const FPhysicsActorReference_ImmediatePhysX& InActorReference, const FVector& InTorque);
-	static void AddForceMassIndependent_AssumesLocked(const FPhysicsActorReference_ImmediatePhysX& InActorReference, const FVector& InForce);
-	static void AddTorqueMassIndependent_AssumesLocked(const FPhysicsActorReference_ImmediatePhysX& InActorReference, const FVector& InTorque);
-	static void AddImpulseAtLocation_AssumesLocked(const FPhysicsActorReference_ImmediatePhysX& InActorReference, const FVector& InImpulse, const FVector& InLocation);
-	static void AddRadialImpulse_AssumesLocked(const FPhysicsActorReference_ImmediatePhysX& InActorReference, const FVector& InOrigin, float InRadius, float InStrength, ERadialImpulseFalloff InFalloff, bool bInVelChange);
-
-	static bool IsGravityEnabled_AssumesLocked(const FPhysicsActorReference_ImmediatePhysX& InActorReference);
-	static void SetGravityEnabled_AssumesLocked(const FPhysicsActorReference_ImmediatePhysX& InActorReference, bool bEnabled);
-
-	static float GetSleepEnergyThreshold_AssumesLocked(const FPhysicsActorReference_ImmediatePhysX& InActorReference);
-	static void SetSleepEnergyThreshold_AssumesLocked(const FPhysicsActorReference_ImmediatePhysX& InActorReference, float InEnergyThreshold);
-
-	static void SetMass_AssumesLocked(const FPhysicsActorReference_ImmediatePhysX& InHandle, float InMass);
-	static void SetMassSpaceInertiaTensor_AssumesLocked(const FPhysicsActorReference_ImmediatePhysX& InHandle, const FVector& InTensor);
-	static void SetComLocalPose_AssumesLocked(const FPhysicsActorReference_ImmediatePhysX& InHandle, const FTransform& InComLocalPose);
-
-	static float GetStabilizationEnergyThreshold_AssumesLocked(const FPhysicsActorReference_ImmediatePhysX& InHandle);
-	static void SetStabilizationEnergyThreshold_AssumesLocked(const FPhysicsActorReference_ImmediatePhysX& InHandle, float InThreshold);
-	static uint32 GetSolverPositionIterationCount_AssumesLocked(const FPhysicsActorReference_ImmediatePhysX& InHandle);
-	static void SetSolverPositionIterationCount_AssumesLocked(const FPhysicsActorReference_ImmediatePhysX& InHandle, uint32 InSolverIterationCount);
-	static uint32 GetSolverVelocityIterationCount_AssumesLocked(const FPhysicsActorReference_ImmediatePhysX& InHandle);
-	static void SetSolverVelocityIterationCount_AssumesLocked(const FPhysicsActorReference_ImmediatePhysX& InHandle, uint32 InSolverIterationCount);
-	static float GetWakeCounter_AssumesLocked(const FPhysicsActorReference_ImmediatePhysX& InHandle);
-	static void SetWakeCounter_AssumesLocked(const FPhysicsActorReference_ImmediatePhysX& InHandle, float InWakeCounter);
-
-	static SIZE_T GetResourceSizeEx(const FPhysicsActorReference_ImmediatePhysX& InActorRef);
-
-	//////////////////////////////////////////////////////////////////////////
-
-	// Constraint interface
-
-	// Functions
-	static FPhysicsConstraintReference_ImmediatePhysX CreateConstraint(const FPhysicsActorReference_ImmediatePhysX& InActorRef1, const FPhysicsActorReference_ImmediatePhysX& InActorRef2, const FTransform& InLocalFrame1, const FTransform& InLocalFrame2);
-	static void SetConstraintUserData(const FPhysicsConstraintReference_ImmediatePhysX& InConstraintRef, void* InUserData);
-	static void ReleaseConstraint(FPhysicsConstraintReference_ImmediatePhysX& InConstraintRef);
-
-	static FTransform GetLocalPose(const FPhysicsConstraintReference_ImmediatePhysX& InConstraintRef, EConstraintFrame::Type InFrame);
-	static FTransform GetGlobalPose(const FPhysicsConstraintReference_ImmediatePhysX& InConstraintRef, EConstraintFrame::Type InFrame);
-	static FVector GetLocation(const FPhysicsConstraintReference_ImmediatePhysX& InConstraintRef);
-	static void GetForce(const FPhysicsConstraintReference_ImmediatePhysX& InConstraintRef, FVector& OutLinForce, FVector& OutAngForce);
-	static void GetDriveLinearVelocity(const FPhysicsConstraintReference_ImmediatePhysX& InConstraintRef, FVector& OutLinVelocity);
-	static void GetDriveAngularVelocity(const FPhysicsConstraintReference_ImmediatePhysX& InConstraintRef, FVector& OutAngVelocity);
-
-	static float GetCurrentSwing1(const FPhysicsConstraintReference_ImmediatePhysX& InConstraintRef);
-	static float GetCurrentSwing2(const FPhysicsConstraintReference_ImmediatePhysX& InConstraintRef);
-	static float GetCurrentTwist(const FPhysicsConstraintReference_ImmediatePhysX& InConstraintRef);
-
-	static void SetCanVisualize(const FPhysicsConstraintReference_ImmediatePhysX& InConstraintRef, bool bInCanVisualize);
-	static void SetCollisionEnabled(const FPhysicsConstraintReference_ImmediatePhysX& InConstraintRef, bool bInCollisionEnabled);
-	static void SetProjectionEnabled_AssumesLocked(const FPhysicsConstraintReference_ImmediatePhysX& InConstraintRef, bool bInProjectionEnabled, float InLinearTolerance = 0.0f, float InAngularToleranceDegrees = 0.0f);
-	static void SetParentDominates_AssumesLocked(const FPhysicsConstraintReference_ImmediatePhysX& InConstraintRef, bool bInParentDominates);
-	static void SetBreakForces_AssumesLocked(const FPhysicsConstraintReference_ImmediatePhysX& InConstraintRef, float InLinearBreakForce, float InAngularBreakForce);
-	static void SetPlasticityLimits_AssumesLocked(const FPhysicsConstraintHandle& InConstraintRef, float InLinearBreakForce, float InAngularBreakTorque);
-	static void SetLocalPose(const FPhysicsConstraintReference_ImmediatePhysX& InConstraintRef, const FTransform& InPose, EConstraintFrame::Type InFrame);
-
-	static void SetLinearMotionLimitType_AssumesLocked(const FPhysicsConstraintReference_ImmediatePhysX& InConstraintRef, PhysicsInterfaceTypes::ELimitAxis InAxis, ELinearConstraintMotion InMotion);
-	static void SetAngularMotionLimitType_AssumesLocked(const FPhysicsConstraintReference_ImmediatePhysX& InConstraintRef, PhysicsInterfaceTypes::ELimitAxis InAxis, EAngularConstraintMotion InMotion);
-
-	static void UpdateLinearLimitParams_AssumesLocked(const FPhysicsConstraintReference_ImmediatePhysX& InConstraintRef, float InLimit, float InAverageMass, const FLinearConstraint& InParams);
-	static void UpdateConeLimitParams_AssumesLocked(const FPhysicsConstraintReference_ImmediatePhysX& InConstraintRef, float InAverageMass, const FConeConstraint& InParams);
-	static void UpdateTwistLimitParams_AssumesLocked(const FPhysicsConstraintReference_ImmediatePhysX& InConstraintRef, float InAverageMass, const FTwistConstraint& InParams);
-	static void UpdateLinearDrive_AssumesLocked(const FPhysicsConstraintReference_ImmediatePhysX& InConstraintRef, const FLinearDriveConstraint& InDriveParams);
-	static void UpdateAngularDrive_AssumesLocked(const FPhysicsConstraintReference_ImmediatePhysX& InConstraintRef, const FAngularDriveConstraint& InDriveParams);
-	static void UpdateDriveTarget_AssumesLocked(const FPhysicsConstraintReference_ImmediatePhysX& InConstraintRef, const FLinearDriveConstraint& InLinDrive, const FAngularDriveConstraint& InAngDrive);
-	static void SetDrivePosition(const FPhysicsConstraintReference_ImmediatePhysX& InConstraintRef, const FVector& InPosition);
-	static void SetDriveOrientation(const FPhysicsConstraintReference_ImmediatePhysX& InConstraintRef, const FQuat& InOrientation);
-	static void SetDriveLinearVelocity(const FPhysicsConstraintReference_ImmediatePhysX& InConstraintRef, const FVector& InLinVelocity);
-	static void SetDriveAngularVelocity(const FPhysicsConstraintReference_ImmediatePhysX& InConstraintRef, const FVector& InAngVelocity);
-
-	static void SetTwistLimit(const FPhysicsConstraintReference_ImmediatePhysX& InConstraintRef, float InLowerLimit, float InUpperLimit, float InContactDistance);
-	static void SetSwingLimit(const FPhysicsConstraintReference_ImmediatePhysX& InConstraintRef, float InYLimit, float InZLimit, float InContactDistance);
-	static void SetLinearLimit(const FPhysicsConstraintReference_ImmediatePhysX& InConstraintRef, float InLimit);
-
-	static bool IsBroken(const FPhysicsConstraintReference_ImmediatePhysX& InConstraintRef);
-
-	static bool ExecuteOnUnbrokenConstraintReadOnly(const FPhysicsConstraintReference_ImmediatePhysX& InConstraintRef, TFunctionRef<void(const FPhysicsConstraintReference_ImmediatePhysX&)> Func);
-	static bool ExecuteOnUnbrokenConstraintReadWrite(const FPhysicsConstraintReference_ImmediatePhysX& InConstraintRef, TFunctionRef<void(const FPhysicsConstraintReference_ImmediatePhysX&)> Func);
-
-	//////////////////////////////////////////////////////////////////////////
-
-	// Scene query interface functions
-
-	/** Trace a ray against the world and return if a blocking hit is found */
-	static bool RaycastTest(const UWorld* World, const FVector Start, const FVector End, ECollisionChannel TraceChannel, const FCollisionQueryParams& Params, const FCollisionResponseParams& ResponseParams, const FCollisionObjectQueryParams& ObjectParams = FCollisionObjectQueryParams::DefaultObjectQueryParam);
-
-	/** Trace a ray against the world and return the first blocking hit */
-	static bool RaycastSingle(const UWorld* World, struct FHitResult& OutHit, const FVector Start, const FVector End, ECollisionChannel TraceChannel, const FCollisionQueryParams& Params, const FCollisionResponseParams& ResponseParams, const FCollisionObjectQueryParams& ObjectParams = FCollisionObjectQueryParams::DefaultObjectQueryParam);
-
-	/**
-	*  Trace a ray against the world and return touching hits and then first blocking hit
-	*  Results are sorted, so a blocking hit (if found) will be the last element of the array
-	*  Only the single closest blocking result will be generated, no tests will be done after that
-	*/
-	static bool RaycastMulti(const UWorld* World, TArray<struct FHitResult>& OutHits, const FVector& Start, const FVector& End, ECollisionChannel TraceChannel, const FCollisionQueryParams& Params, const FCollisionResponseParams& ResponseParams, const FCollisionObjectQueryParams& ObjectParams = FCollisionObjectQueryParams::DefaultObjectQueryParam);
-
-	// GEOM OVERLAP
-
-	/** Function for testing overlaps between a supplied PxGeometry and the world. Returns true if at least one overlapping shape is blocking*/
-	static bool GeomOverlapBlockingTest(const UWorld* World, const FCollisionShape& CollisionShape, const FVector& Pos, const FQuat& Rot, ECollisionChannel TraceChannel, const FCollisionQueryParams& Params, const FCollisionResponseParams& ResponseParams, const FCollisionObjectQueryParams& ObjectParams = FCollisionObjectQueryParams::DefaultObjectQueryParam);
-
-	/** Function for testing overlaps between a supplied PxGeometry and the world. Returns true if anything is overlapping (blocking or touching)*/
-	static bool GeomOverlapAnyTest(const UWorld* World, const FCollisionShape& CollisionShape, const FVector& Pos, const FQuat& Rot, ECollisionChannel TraceChannel, const FCollisionQueryParams& Params, const FCollisionResponseParams& ResponseParams, const FCollisionObjectQueryParams& ObjectParams = FCollisionObjectQueryParams::DefaultObjectQueryParam);
-
-	// GEOM SWEEP
-
-	/** Function used for sweeping a supplied PxGeometry against the world as a test */
-	static bool GeomSweepTest(const UWorld* World, const FCollisionShape& CollisionShape, const FQuat& Rot, FVector Start, FVector End, ECollisionChannel TraceChannel, const FCollisionQueryParams& Params, const FCollisionResponseParams& ResponseParams, const FCollisionObjectQueryParams& ObjectParams = FCollisionObjectQueryParams::DefaultObjectQueryParam);
-
-	/** Function for sweeping a supplied PxGeometry against the world */
-	static bool GeomSweepSingle(const UWorld* World, const FCollisionShape& CollisionShape, const FQuat& Rot, FHitResult& OutHit, FVector Start, FVector End, ECollisionChannel TraceChannel, const FCollisionQueryParams& Params, const FCollisionResponseParams& ResponseParams, const FCollisionObjectQueryParams& ObjectParams = FCollisionObjectQueryParams::DefaultObjectQueryParam);
-
-	template<typename GeomType>
-	static bool GeomSweepMulti(const UWorld* World, const GeomType& InGeom, const FQuat& InGeomRot, TArray<FHitResult>& OutHits, FVector Start, FVector End, ECollisionChannel TraceChannel, const FCollisionQueryParams& Params, const FCollisionResponseParams& ResponseParams, const FCollisionObjectQueryParams& ObjectParams = FCollisionObjectQueryParams::DefaultObjectQueryParam);
-	template<typename GeomType>
-	static bool GeomOverlapMulti(const UWorld* World, const GeomType& InGeom, const FVector& InPosition, const FQuat& InRotation, TArray<FOverlapResult>& OutOverlaps, ECollisionChannel TraceChannel, const FCollisionQueryParams& Params, const FCollisionResponseParams& ResponseParams, const FCollisionObjectQueryParams& ObjectParams);
-
-	//////////////////////////////////////////////////////////////////////////
-	// Trace functions for testing specific geometry (not against a world)
-	//////////////////////////////////////////////////////////////////////////
-	static bool LineTrace_Geom(FHitResult& OutHit, const FBodyInstance* InInstance, const FVector& InStart, const FVector& InEnd, bool bTraceComplex, bool bExtractPhysMaterial = false);
-	static bool Sweep_Geom(FHitResult& OutHit, const FBodyInstance* InInstance, const FVector& InStart, const FVector& InEnd, const FQuat& InShapeRotation, const FCollisionShape& InShape, bool bSweepComplex);
-
-	static bool Overlap_Geom(const FBodyInstance* InBodyInstance, const FPhysicsGeometryCollection& InGeometry, const FTransform& InShapeTransform, FMTDResult* OutOptResult = nullptr);
-	static bool Overlap_Geom(const FBodyInstance* InBodyInstance, const FCollisionShape& InCollisionShape, const FQuat& InShapeRotation, const FTransform& InShapeTransform, FMTDResult* OutOptResult = nullptr);
-
-	static bool GetSquaredDistanceToBody(const FBodyInstance* InInstance, const FVector& InPoint, float& OutDistanceSquared, FVector* OutOptPointOnBody = nullptr);
-
-	//////////////////////////////////////////////////////////////////////////
-
-	// Misc
-
-	static bool ExecPhysCommands(const TCHAR* Cmd, FOutputDevice* Ar, UWorld* InWorld);
-    
-    static void CalculateMassPropertiesFromShapeCollection(physx::PxMassProperties& OutProperties, const TArray<FPhysicsShapeHandle>& InShapes, float InDensityKGPerCM);
-};
-
-template<>
-bool FPhysicsInterface_ImmediatePhysX::GeomSweepMulti(const UWorld* World, const FPhysicsGeometryCollection& InGeom, const FQuat& InGeomRot, TArray<FHitResult>& OutHits, FVector Start, FVector End, ECollisionChannel TraceChannel, const FCollisionQueryParams& Params, const FCollisionResponseParams& ResponseParams, const FCollisionObjectQueryParams& ObjectParams /*= FCollisionObjectQueryParams::DefaultObjectQueryParam*/);
-template<>
-bool FPhysicsInterface_ImmediatePhysX::GeomSweepMulti(const UWorld* World, const FCollisionShape& InGeom, const FQuat& InGeomRot, TArray<FHitResult>& OutHits, FVector Start, FVector End, ECollisionChannel TraceChannel, const FCollisionQueryParams& Params, const FCollisionResponseParams& ResponseParams, const FCollisionObjectQueryParams& ObjectParams /*= FCollisionObjectQueryParams::DefaultObjectQueryParam*/);
-
-template<>
-bool FPhysicsInterface_ImmediatePhysX::GeomOverlapMulti(const UWorld* World, const FPhysicsGeometryCollection& InGeom, const FVector& InPosition, const FQuat& InRotation, TArray<FOverlapResult>& OutOverlaps, ECollisionChannel TraceChannel, const FCollisionQueryParams& Params, const FCollisionResponseParams& ResponseParams, const FCollisionObjectQueryParams& ObjectParams);
-template<>
-bool FPhysicsInterface_ImmediatePhysX::GeomOverlapMulti(const UWorld* World, const FCollisionShape& InGeom, const FVector& InPosition, const FQuat& InRotation, TArray<FOverlapResult>& OutOverlaps, ECollisionChannel TraceChannel, const FCollisionQueryParams& Params, const FCollisionResponseParams& ResponseParams, const FCollisionObjectQueryParams& ObjectParams);
-
-#endif
