Physics interface cleanup.

diff --git a/Engine/Source/Runtime/Engine/Public/Physics/ImmediatePhysics/ImmediatePhysicsDeclares.h b/Engine/Source/Runtime/Engine/Public/Physics/ImmediatePhysics/ImmediatePhysicsDeclares.h
--- a/Engine/Source/Runtime/Engine/Public/Physics/ImmediatePhysics/ImmediatePhysicsDeclares.h
+++ b/Engine/Source/Runtime/Engine/Public/Physics/ImmediatePhysics/ImmediatePhysicsDeclares.h
@@ -2,38 +2,14 @@
 
 #pragma once
 
-#include "Physics/ImmediatePhysics/ImmediatePhysicsPhysX/ImmediatePhysicsDeclares_PhysX.h"
 #include "Physics/ImmediatePhysics/ImmediatePhysicsChaos/ImmediatePhysicsDeclares_Chaos.h"
 #include "Physics/ImmediatePhysics/ImmediatePhysicsShared/ImmediatePhysicsCore.h"
 
-
-/**
- * @file ImmediatePhysicsDeclares.h
- * 
- * Move either ImmediatePhysics_PhysX or ImmediatePhysics_Chaos Immediate Physics into the ImmediatePhysics namespace
- *
- * We can use PhysX ImmediatePhysics only if PhysX is also providing the global physics interface.
- * because ImmediatePhysics_PhysX relies on physx types being instantiated by BodyInstance, ConstraintInstance, etc.
- * 
- * You can currently run PhysX and Chaos ImmediatePhysics simulations in the same build if PhysX is
- * providing the global physics interface. This capability will likely disappear once Chaos is established
- * at which point you will only be able to use Chaos immediate physics if you are also using Chaos
- * for global physics.
- */
 namespace ImmediatePhysics
 {
 	using EActorType = ImmediatePhysics_Shared::EActorType;
 	using EForceType = ImmediatePhysics_Shared::EForceType;
-
-#if WITH_PHYSX && PHYSICS_INTERFACE_PHYSX
-	using FSimulation = ImmediatePhysics_PhysX::FSimulation;
-	using FActorHandle = ImmediatePhysics_PhysX::FActorHandle;
-	using FJointHandle = ImmediatePhysics_PhysX::FJointHandle;
-#elif WITH_CHAOS
 	using FSimulation = ImmediatePhysics_Chaos::FSimulation;
 	using FActorHandle = ImmediatePhysics_Chaos::FActorHandle;
 	using FJointHandle = ImmediatePhysics_Chaos::FJointHandle;
-#else
-#error Global ImmediatePhysics is not defined
-#endif
 }
\ No newline at end of file
diff --git a/Engine/Source/Runtime/Engine/Public/Physics/ImmediatePhysics/ImmediatePhysicsPhysX/ImmediatePhysicsActor_PhysX.h b/Engine/Source/Runtime/Engine/Public/Physics/ImmediatePhysics/ImmediatePhysicsPhysX/ImmediatePhysicsActor_PhysX.h
--- a/Engine/Source/Runtime/Engine/Public/Physics/ImmediatePhysics/ImmediatePhysicsPhysX/ImmediatePhysicsActor_PhysX.h
+++ b/Engine/Source/Runtime/Engine/Public/Physics/ImmediatePhysics/ImmediatePhysicsPhysX/ImmediatePhysicsActor_PhysX.h
@@ -1,34 +1,3 @@
 // Copyright Epic Games, Inc. All Rights Reserved.
 
-#pragma once
-
-#if WITH_PHYSX && PHYSICS_INTERFACE_PHYSX
-#include "PhysXPublic.h"
-#include "Physics/ImmediatePhysics/ImmediatePhysicsPhysX/ImmediatePhysicsShape_PhysX.h"
-
-namespace ImmediatePhysics_PhysX
-{
-
-	/** Holds geometry data*/
-	struct FActor
-	{
-		FActor()
-			: UserData(nullptr)
-		{}
-
-		~FActor();
-
-		TArray<FShape> Shapes;
-		void* UserData;
-
-		/** Create geometry data for the entity */
-		void CreateGeometry(PxRigidActor* RigidActor, const PxTransform& ActorToBodyTM);
-		bool AddShape(PxShape* InShape);
-
-		/** Ensures all the geometry data has been properly freed */
-		void TerminateGeometry();
-	};
-
-}
-
-#endif // WITH_PHYSX
\ No newline at end of file
+#pragma once
\ No newline at end of file
diff --git a/Engine/Source/Runtime/Engine/Public/Physics/ImmediatePhysics/ImmediatePhysicsPhysX/ImmediatePhysicsCacheAllocator_PhysX.h b/Engine/Source/Runtime/Engine/Public/Physics/ImmediatePhysics/ImmediatePhysicsPhysX/ImmediatePhysicsCacheAllocator_PhysX.h
--- a/Engine/Source/Runtime/Engine/Public/Physics/ImmediatePhysics/ImmediatePhysicsPhysX/ImmediatePhysicsCacheAllocator_PhysX.h
+++ b/Engine/Source/Runtime/Engine/Public/Physics/ImmediatePhysics/ImmediatePhysicsPhysX/ImmediatePhysicsCacheAllocator_PhysX.h
@@ -1,35 +1,3 @@
 // Copyright Epic Games, Inc. All Rights Reserved.
 
 #pragma once
-
-#if WITH_PHYSX && PHYSICS_INTERFACE_PHYSX
-#include "PhysXPublic.h"
-#include "Physics/ImmediatePhysics/ImmediatePhysicsPhysX/ImmediatePhysicsLinearBlockAllocator_PhysX.h"
-
-namespace ImmediatePhysics_PhysX
-{
-
-	/** TODO: Use a smarter memory allocator */
-	struct FCacheAllocator : public PxCacheAllocator
-	{
-		FCacheAllocator() : External(0) {}
-		PxU8* allocateCacheData(const PxU32 ByteSize) override
-		{
-			return BlockAllocator[External].Alloc(ByteSize);
-		}
-
-		void Reset()
-		{
-#if PERSISTENT_CONTACT_PAIRS
-			External = 1 - External;	//flip buffer so we maintain cache for 1 extra step
-#endif
-			BlockAllocator[External].Reset();
-		}
-
-		FLinearBlockAllocator BlockAllocator[2];
-		int32 External;
-	};
-
-}
-
-#endif // WITH_PHYSX
\ No newline at end of file
diff --git a/Engine/Source/Runtime/Engine/Public/Physics/ImmediatePhysics/ImmediatePhysicsPhysX/ImmediatePhysicsConstraintAllocator_PhysX.h b/Engine/Source/Runtime/Engine/Public/Physics/ImmediatePhysics/ImmediatePhysicsPhysX/ImmediatePhysicsConstraintAllocator_PhysX.h
--- a/Engine/Source/Runtime/Engine/Public/Physics/ImmediatePhysics/ImmediatePhysicsPhysX/ImmediatePhysicsConstraintAllocator_PhysX.h
+++ b/Engine/Source/Runtime/Engine/Public/Physics/ImmediatePhysics/ImmediatePhysicsPhysX/ImmediatePhysicsConstraintAllocator_PhysX.h
@@ -1,43 +1,3 @@
 // Copyright Epic Games, Inc. All Rights Reserved.
 
 #pragma once
-
-#if WITH_PHYSX && PHYSICS_INTERFACE_PHYSX
-#include "PhysXPublic.h"
-#include "Physics/ImmediatePhysics/ImmediatePhysicsPhysX/ImmediatePhysicsLinearBlockAllocator_PhysX.h"
-
-namespace ImmediatePhysics_PhysX
-{
-
-	/** TODO: Use a smarter memory allocator */
-	class FConstraintAllocator : public PxConstraintAllocator
-	{
-	public:
-		FConstraintAllocator() : External(0) {}
-
-		PxU8* reserveConstraintData(const PxU32 ByteSize) override
-		{
-			return BlockAllocator[External].Alloc(ByteSize);
-
-		}
-
-		PxU8* reserveFrictionData(const PxU32 ByteSize)
-		{
-			return BlockAllocator[External].Alloc(ByteSize);
-		}
-
-		void Reset()
-		{
-#if PERSISTENT_CONTACT_PAIRS
-			External = 1 - External;	//flip buffer so we maintain cache for 1 extra step
-#endif
-			BlockAllocator[External].Reset();
-		}
-
-		FLinearBlockAllocator BlockAllocator[2];
-		int32 External;
-	};
-
-}
-
-#endif // WITH_PHYSX
\ No newline at end of file
diff --git a/Engine/Source/Runtime/Engine/Public/Physics/ImmediatePhysics/ImmediatePhysicsPhysX/ImmediatePhysicsContactPair_PhysX.h b/Engine/Source/Runtime/Engine/Public/Physics/ImmediatePhysics/ImmediatePhysicsPhysX/ImmediatePhysicsContactPair_PhysX.h
--- a/Engine/Source/Runtime/Engine/Public/Physics/ImmediatePhysics/ImmediatePhysicsPhysX/ImmediatePhysicsContactPair_PhysX.h
+++ b/Engine/Source/Runtime/Engine/Public/Physics/ImmediatePhysics/ImmediatePhysicsPhysX/ImmediatePhysicsContactPair_PhysX.h
@@ -1,32 +1,3 @@
 // Copyright Epic Games, Inc. All Rights Reserved.
 
 #pragma once
-
-#if WITH_PHYSX && PHYSICS_INTERFACE_PHYSX
-#include "CoreMinimal.h"
-
-namespace ImmediatePhysics_PhysX
-{
-
-	/** Contact pair generated between entities */
-	struct FContactPair
-	{
-		/** Index of the dynamic actor that we generated the contact pair for*/
-		int32 DynamicActorDataIndex;
-
-		/** Index of the other actor that we generated the contact pair for. This could be either dynamic or static */
-		uint32 OtherActorDataIndex;
-
-		/** Index into the first contact point associated with this pair*/
-		uint32 StartContactIndex;
-
-		/** Number of contacts associated with this pair. */
-		uint32 NumContacts;
-
-		/** Identifies the pair index from the original contact generation test */
-		int32 PairIdx;
-	};
-
-}
-
-#endif // WITH_PHYSX
\ No newline at end of file
diff --git a/Engine/Source/Runtime/Engine/Public/Physics/ImmediatePhysics/ImmediatePhysicsPhysX/ImmediatePhysicsContactPointRecorder_PhysX.h b/Engine/Source/Runtime/Engine/Public/Physics/ImmediatePhysics/ImmediatePhysicsPhysX/ImmediatePhysicsContactPointRecorder_PhysX.h
--- a/Engine/Source/Runtime/Engine/Public/Physics/ImmediatePhysics/ImmediatePhysicsPhysX/ImmediatePhysicsContactPointRecorder_PhysX.h
+++ b/Engine/Source/Runtime/Engine/Public/Physics/ImmediatePhysics/ImmediatePhysicsPhysX/ImmediatePhysicsContactPointRecorder_PhysX.h
@@ -1,37 +1,3 @@
 // Copyright Epic Games, Inc. All Rights Reserved.
 
 #pragma once
-
-#if WITH_PHYSX && PHYSICS_INTERFACE_PHYSX
-#include "PhysXPublic.h"
-
-namespace ImmediatePhysics_PhysX
-{
-	struct FSimulation;
-	struct FMaterial;
-
-	struct FContactPointRecorder : public immediate::PxContactRecorder
-	{
-		FContactPointRecorder(struct FSimulation& InSimulation, int32 InDynamicActorDataIndex, int32 InOtherActorDataIndex, int32 InPairIdx, const FMaterial& InSimulatedShapeMaterial, const FMaterial& InOtherShapeMaterial)
-			: Simulation(InSimulation)
-			, SimulatedShapeMaterial(InSimulatedShapeMaterial)
-			, OtherShapeMaterial(InOtherShapeMaterial)
-			, DynamicActorDataIndex(InDynamicActorDataIndex)
-			, OtherActorDataIndex(InOtherActorDataIndex)
-			, PairIdx(InPairIdx)
-		{
-		}
-
-		bool recordContacts(const Gu::ContactPoint* ContactPoints, const PxU32 NumContacts, const PxU32 Index) override;
-
-		FSimulation& Simulation;
-		const FMaterial& SimulatedShapeMaterial;
-		const FMaterial& OtherShapeMaterial;
-		int32 DynamicActorDataIndex;
-		int32 OtherActorDataIndex;
-		int32 PairIdx;
-	};
-
-}
-
-#endif // WITH_PHYSX
\ No newline at end of file
diff --git a/Engine/Source/Runtime/Engine/Public/Physics/ImmediatePhysics/ImmediatePhysicsPhysX/ImmediatePhysicsCore_PhysX.h b/Engine/Source/Runtime/Engine/Public/Physics/ImmediatePhysics/ImmediatePhysicsPhysX/ImmediatePhysicsCore_PhysX.h
--- a/Engine/Source/Runtime/Engine/Public/Physics/ImmediatePhysics/ImmediatePhysicsPhysX/ImmediatePhysicsCore_PhysX.h
+++ b/Engine/Source/Runtime/Engine/Public/Physics/ImmediatePhysics/ImmediatePhysicsPhysX/ImmediatePhysicsCore_PhysX.h
@@ -1,18 +1,3 @@
 // Copyright Epic Games, Inc. All Rights Reserved.
 
 #pragma once
-
-#include "Physics/ImmediatePhysics/ImmediatePhysicsShared/ImmediatePhysicsCore.h"
-
-#if WITH_PHYSX && PHYSICS_INTERFACE_PHYSX
-
-namespace ImmediatePhysics_PhysX
-{
-	using EActorType = ImmediatePhysics_Shared::EActorType;
-	using EForceType = ImmediatePhysics_Shared::EForceType;
-}
-
-struct FBodyInstance;
-struct FConstraintInstance;
-
-#endif
diff --git a/Engine/Source/Runtime/Engine/Public/Physics/ImmediatePhysics/ImmediatePhysicsPhysX/ImmediatePhysicsD6JointData_PhysX.h b/Engine/Source/Runtime/Engine/Public/Physics/ImmediatePhysics/ImmediatePhysicsPhysX/ImmediatePhysicsD6JointData_PhysX.h
--- a/Engine/Source/Runtime/Engine/Public/Physics/ImmediatePhysics/ImmediatePhysicsPhysX/ImmediatePhysicsD6JointData_PhysX.h
+++ b/Engine/Source/Runtime/Engine/Public/Physics/ImmediatePhysics/ImmediatePhysicsPhysX/ImmediatePhysicsD6JointData_PhysX.h
@@ -1,71 +1,3 @@
 // Copyright Epic Games, Inc. All Rights Reserved.
 
 #pragma once
-
-#if WITH_PHYSX && PHYSICS_INTERFACE_PHYSX
-#include "PhysXPublic.h"
-
-//TODO: this is just straight copied from physx
-namespace ImmediatePhysics_PhysX
-{
-
-	struct D6JointData
-	{
-		D6JointData(PxD6Joint* Joint);
-
-		PxConstraintInvMassScale	invMassScale;
-		PxTransform					c2b[2];
-
-		PxU32					locked;		// bitmap of locked DOFs
-		PxU32					limited;	// bitmap of limited DOFs
-		PxU32					driving;	// bitmap of active drives (implies driven DOFs not locked)
-
-
-		PxD6Motion::Enum		motion[6];
-		PxJointLinearLimit		linearLimit;
-		PxJointAngularLimitPair	twistLimit;
-		PxJointLimitCone		swingLimit;
-
-		PxD6JointDrive			drive[PxD6Drive::eCOUNT];
-
-		PxTransform				drivePosition;
-		PxVec3					driveLinearVelocity;
-		PxVec3					driveAngularVelocity;
-
-		// derived quantities
-
-
-		// tan-half and tan-quarter angles
-
-		PxReal					thSwingY;
-		PxReal					thSwingZ;
-		PxReal					thSwingPad;
-
-		PxReal					tqSwingY;
-		PxReal					tqSwingZ;
-		PxReal					tqSwingPad;
-
-		PxReal					tqTwistLow;
-		PxReal					tqTwistHigh;
-		PxReal					tqTwistPad;
-
-		PxReal					linearMinDist;	// linear limit minimum distance to get a good direction
-
-		// projection quantities
-		//PxReal					projectionLinearTolerance;
-		//PxReal					projectionAngularTolerance;
-
-		FTransform					ActorToBody[2];
-
-		bool HasConstraints() const
-		{
-			return locked || limited || driving;
-		}
-	};
-
-	void PrepareJointData(D6JointData& JointData);
-	PxU32 D6JointSolverPrep(Px1DConstraint* constraints, PxVec3& body0WorldOffset, PxU32 maxConstraints, PxConstraintInvMassScale& invMassScale, const void* constantBlock, const PxTransform& bA2w, const PxTransform& bB2w);
-
-}
-
-#endif // WITH_PHYSX
diff --git a/Engine/Source/Runtime/Engine/Public/Physics/ImmediatePhysics/ImmediatePhysicsPhysX/ImmediatePhysicsDeclares_PhysX.h b/Engine/Source/Runtime/Engine/Public/Physics/ImmediatePhysics/ImmediatePhysicsPhysX/ImmediatePhysicsDeclares_PhysX.h
--- a/Engine/Source/Runtime/Engine/Public/Physics/ImmediatePhysics/ImmediatePhysicsPhysX/ImmediatePhysicsDeclares_PhysX.h
+++ b/Engine/Source/Runtime/Engine/Public/Physics/ImmediatePhysics/ImmediatePhysicsPhysX/ImmediatePhysicsDeclares_PhysX.h
@@ -1,12 +1,3 @@
 // Copyright Epic Games, Inc. All Rights Reserved.
 
 #pragma once
-
-#if WITH_PHYSX && PHYSICS_INTERFACE_PHYSX
-namespace ImmediatePhysics_PhysX
-{
-	struct FSimulation;
-	struct FActorHandle;
-	struct FJointHandle;
-}
-#endif
diff --git a/Engine/Source/Runtime/Engine/Public/Physics/ImmediatePhysics/ImmediatePhysicsPhysX/ImmediatePhysicsJointHandle_PhysX.h b/Engine/Source/Runtime/Engine/Public/Physics/ImmediatePhysics/ImmediatePhysicsPhysX/ImmediatePhysicsJointHandle_PhysX.h
--- a/Engine/Source/Runtime/Engine/Public/Physics/ImmediatePhysics/ImmediatePhysicsPhysX/ImmediatePhysicsJointHandle_PhysX.h
+++ b/Engine/Source/Runtime/Engine/Public/Physics/ImmediatePhysics/ImmediatePhysicsPhysX/ImmediatePhysicsJointHandle_PhysX.h
@@ -1,36 +1,3 @@
 // Copyright Epic Games, Inc. All Rights Reserved.
 
 #pragma once
-
-#if WITH_PHYSX && PHYSICS_INTERFACE_PHYSX
-#include "PhysXPublic.h"
-
-namespace ImmediatePhysics_PhysX
-{
-	struct FSimulation;
-
-	/** handle associated with a physics joint. This is the proper way to read/write to the physics simulation */
-	struct ENGINE_API FJointHandle
-	{
-	private:
-		FSimulation& OwningSimulation;
-		int32 JointDataIndex;
-
-		friend FSimulation;
-		FJointHandle(FSimulation& InOwningSimulation, int32 InJointDataIndex)
-			: OwningSimulation(InOwningSimulation)
-			, JointDataIndex(InJointDataIndex)
-		{
-		}
-
-		~FJointHandle()
-		{
-		}
-
-		FJointHandle(const FJointHandle&);	//Ensure no copying of handles
-
-	};
-
-}
-
-#endif // WITH_PHYSX
diff --git a/Engine/Source/Runtime/Engine/Public/Physics/ImmediatePhysics/ImmediatePhysicsPhysX/ImmediatePhysicsJoint_PhysX.h b/Engine/Source/Runtime/Engine/Public/Physics/ImmediatePhysics/ImmediatePhysicsPhysX/ImmediatePhysicsJoint_PhysX.h
--- a/Engine/Source/Runtime/Engine/Public/Physics/ImmediatePhysics/ImmediatePhysicsPhysX/ImmediatePhysicsJoint_PhysX.h
+++ b/Engine/Source/Runtime/Engine/Public/Physics/ImmediatePhysics/ImmediatePhysicsPhysX/ImmediatePhysicsJoint_PhysX.h
@@ -1,20 +1,3 @@
 // Copyright Epic Games, Inc. All Rights Reserved.
 
 #pragma once
-
-#if WITH_PHYSX && PHYSICS_INTERFACE_PHYSX
-
-namespace ImmediatePhysics_PhysX
-{
-
-	struct FActorHandle;
-
-	struct FJoint
-	{
-		FActorHandle* DynamicActor;
-		FActorHandle* OtherActor;
-	};
-
-}
-
-#endif // WITH_PHYSX
\ No newline at end of file
diff --git a/Engine/Source/Runtime/Engine/Public/Physics/ImmediatePhysics/ImmediatePhysicsPhysX/ImmediatePhysicsKinematicTarget_PhysX.h b/Engine/Source/Runtime/Engine/Public/Physics/ImmediatePhysics/ImmediatePhysicsPhysX/ImmediatePhysicsKinematicTarget_PhysX.h
--- a/Engine/Source/Runtime/Engine/Public/Physics/ImmediatePhysics/ImmediatePhysicsPhysX/ImmediatePhysicsKinematicTarget_PhysX.h
+++ b/Engine/Source/Runtime/Engine/Public/Physics/ImmediatePhysics/ImmediatePhysicsPhysX/ImmediatePhysicsKinematicTarget_PhysX.h
@@ -1,24 +1,3 @@
 // Copyright Epic Games, Inc. All Rights Reserved.
 
 #pragma once
-
-#if WITH_PHYSX && PHYSICS_INTERFACE_PHYSX
-#include "PhysXPublic.h"
-
-namespace ImmediatePhysics_PhysX
-{
-	/** Holds shape data*/
-	struct FImmediateKinematicTarget
-	{
-		PxTransform BodyToWorld;
-		bool bTargetSet;
-
-		FImmediateKinematicTarget()
-			: bTargetSet(false)
-		{
-		}
-	};
-
-}
-
-#endif // WITH_PHYSX
diff --git a/Engine/Source/Runtime/Engine/Public/Physics/ImmediatePhysics/ImmediatePhysicsPhysX/ImmediatePhysicsLinearBlockAllocator_PhysX.h b/Engine/Source/Runtime/Engine/Public/Physics/ImmediatePhysics/ImmediatePhysicsPhysX/ImmediatePhysicsLinearBlockAllocator_PhysX.h
--- a/Engine/Source/Runtime/Engine/Public/Physics/ImmediatePhysics/ImmediatePhysicsPhysX/ImmediatePhysicsLinearBlockAllocator_PhysX.h
+++ b/Engine/Source/Runtime/Engine/Public/Physics/ImmediatePhysics/ImmediatePhysicsPhysX/ImmediatePhysicsLinearBlockAllocator_PhysX.h
@@ -1,129 +1,3 @@
 // Copyright Epic Games, Inc. All Rights Reserved.
 
 #pragma once
-
-#if WITH_PHYSX && PHYSICS_INTERFACE_PHYSX
-#include "HAL/LowLevelMemTracker.h"
-
-namespace ImmediatePhysics_PhysX
-{
-	const int PageBufferSize = 1024 * 64;	//This has to be here because of VS2013
-
-	struct FLinearBlockAllocator
-	{
-		struct FPageStruct
-		{
-			//We assume FPageStruct is allocated with 16-byte alignment. Do not move Buffer. If we get full support for alignas(16) we could make this better
-			uint8 Buffer[PageBufferSize];
-
-			FPageStruct* NextPage;
-			FPageStruct* PrevPage;
-			int32 SeekPosition;
-
-			FPageStruct()
-				: NextPage(nullptr)
-				, PrevPage(nullptr)
-				, SeekPosition(0)
-			{
-			}
-		};
-
-		FPageStruct* FreePage;
-		FPageStruct* FirstPage;
-
-		FLinearBlockAllocator()
-		{
-			FreePage = AllocPage();
-			FirstPage = FreePage;
-		}
-
-		FPageStruct* AllocPage()
-		{
-			LLM_SCOPE(ELLMTag::PhysX);
-			FPageStruct* ReturnPage = (FPageStruct*)FMemory::Malloc(sizeof(FPageStruct), 16);
-			new(ReturnPage) FPageStruct();
-
-			FPlatformMisc::TagBuffer("ImmediatePhysicsSim", 0, (const void*)ReturnPage, sizeof(FPageStruct));
-			return ReturnPage;
-		}
-
-		void ReleasePage(FPageStruct* Page)
-		{
-			FMemory::Free(Page);
-		}
-
-		uint8* Alloc(int32 Bytes)
-		{
-			check(Bytes < PageBufferSize);	//Page size needs to be increased since we don't allow spillover
-			if (Bytes)
-			{
-				//Assumes 16 byte alignment
-				int32 BytesLeft = PageBufferSize - FreePage->SeekPosition;	//don't switch to uint because negative implies we're out of 16 byte aligned space
-				if (BytesLeft < Bytes)
-				{
-					//no space so allocate new page if needed
-					if (FreePage->NextPage)
-					{
-						FreePage = FreePage->NextPage;
-					}
-					else
-					{
-						FPageStruct* NewPage = AllocPage();
-						NewPage->PrevPage = FreePage;
-						FreePage->NextPage = NewPage;
-						FreePage = NewPage;
-					}
-				}
-
-				uint32 ReturnSlot = FreePage->SeekPosition;
-				FreePage->SeekPosition = (FreePage->SeekPosition + Bytes + 15)&(~15);
-
-				return &FreePage->Buffer[ReturnSlot];
-			}
-			else
-			{
-				return nullptr;
-			}
-		}
-
-		void Reset()
-		{
-			for (FPageStruct* Page = FirstPage; Page; Page = Page->NextPage)
-			{
-				Page->SeekPosition = 0;
-			}
-
-			FreePage = FirstPage;
-		}
-
-		void Empty()
-		{
-			FPageStruct* CurrentPage = FirstPage->NextPage;
-			while (CurrentPage)
-			{
-				FPageStruct* OldPage = CurrentPage;
-				CurrentPage = CurrentPage->NextPage;
-				ReleasePage(OldPage);
-			}
-
-			FirstPage->NextPage = nullptr;
-			FirstPage->SeekPosition = 0;
-			FreePage = FirstPage;
-			//Note we do not deallocate FirstPage until the allocator deallocates
-		}
-
-		~FLinearBlockAllocator()
-		{
-			Empty();
-			ReleasePage(FirstPage);
-		}
-
-	private:
-		//Don't copy these around
-		FLinearBlockAllocator(const FLinearBlockAllocator& Other);
-		const FLinearBlockAllocator& operator=(const FLinearBlockAllocator& Other);
-	};
-
-}
-
-#endif // WITH_PHYSX
\ No newline at end of file
diff --git a/Engine/Source/Runtime/Engine/Public/Physics/ImmediatePhysics/ImmediatePhysicsPhysX/ImmediatePhysicsMaterialHandle_PhysX.h b/Engine/Source/Runtime/Engine/Public/Physics/ImmediatePhysics/ImmediatePhysicsPhysX/ImmediatePhysicsMaterialHandle_PhysX.h
--- a/Engine/Source/Runtime/Engine/Public/Physics/ImmediatePhysics/ImmediatePhysicsPhysX/ImmediatePhysicsMaterialHandle_PhysX.h
+++ b/Engine/Source/Runtime/Engine/Public/Physics/ImmediatePhysics/ImmediatePhysicsPhysX/ImmediatePhysicsMaterialHandle_PhysX.h
@@ -1,47 +1,3 @@
 // Copyright Epic Games, Inc. All Rights Reserved.
 
 #pragma once
-
-#if WITH_PHYSX && PHYSICS_INTERFACE_PHYSX
-
-class UPhysicalMaterial;
-
-namespace ImmediatePhysics_PhysX
-{
-	struct ENGINE_API FMaterialHandle
-	{
-
-	public:
-
-		FMaterial& Get();
-		void SetStaticFriction(float InFriction);
-		void SetDynamicFriction(float InFriction);
-		void SetRestitution(float InRestitution);
-		void SetFrictionCombineMode(EFrictionCombineMode InCombineMode);
-		void SetRestitutionCombineMode(EFrictionCombineMode InCombineMode);
-		void SetFromPhysicalMaterial(UPhysicalMaterial* InMaterial);
-
-	private:
-
-		FSimulation & OwningSimulation;
-		int32 DataIndex;
-
-		friend struct FSimulation;
-
-		FMaterialHandle(FSimulation& InOwningSimulation, int32 InDataIndex)
-			: OwningSimulation(InOwningSimulation)
-			, DataIndex(InDataIndex)
-		{}
-
-		~FMaterialHandle() 
-		{}
-
-		// Removed copy/assign
-		FMaterialHandle(const FMaterialHandle& Other) = delete;
-		FMaterialHandle(FMaterialHandle&& Other) = delete;
-		FMaterialHandle& operator=(const FMaterialHandle& Other) = delete;
-		FMaterialHandle& operator=(FMaterialHandle&& Other) = delete;
-	};
-}
-
-#endif // WITH_PHYSX
\ No newline at end of file
diff --git a/Engine/Source/Runtime/Engine/Public/Physics/ImmediatePhysics/ImmediatePhysicsPhysX/ImmediatePhysicsMaterial_PhysX.h b/Engine/Source/Runtime/Engine/Public/Physics/ImmediatePhysics/ImmediatePhysicsPhysX/ImmediatePhysicsMaterial_PhysX.h
--- a/Engine/Source/Runtime/Engine/Public/Physics/ImmediatePhysics/ImmediatePhysicsPhysX/ImmediatePhysicsMaterial_PhysX.h
+++ b/Engine/Source/Runtime/Engine/Public/Physics/ImmediatePhysics/ImmediatePhysicsPhysX/ImmediatePhysicsMaterial_PhysX.h
@@ -1,37 +1,3 @@
 // Copyright Epic Games, Inc. All Rights Reserved.
 
 #pragma once
-
-#if WITH_PHYSX && PHYSICS_INTERFACE_PHYSX
-#include "PhysicsSettingsEnums.h"
-
-namespace physx
-{
-	class PxMaterial;
-}
-
-namespace ImmediatePhysics_PhysX
-{
-	struct FMaterial
-	{
-		FMaterial()
-			: StaticFriction(0.7f)
-			, DynamicFriction(0.7f)
-			, Restitution(0.3f)
-		{
-		}
-
-		FMaterial(physx::PxMaterial* InPxMaterial);
-
-		float StaticFriction;
-		float DynamicFriction;
-		float Restitution;
-
-		EFrictionCombineMode::Type FrictionCombineMode;
-		EFrictionCombineMode::Type RestitutionCombineMode;
-
-		static FMaterial Default;
-	};
-}
-
-#endif // WITH_PHYSX
\ No newline at end of file
diff --git a/Engine/Source/Runtime/Engine/Public/Physics/ImmediatePhysics/ImmediatePhysicsPhysX/ImmediatePhysicsPersistentContactPairData_PhysX.h b/Engine/Source/Runtime/Engine/Public/Physics/ImmediatePhysics/ImmediatePhysicsPhysX/ImmediatePhysicsPersistentContactPairData_PhysX.h
--- a/Engine/Source/Runtime/Engine/Public/Physics/ImmediatePhysics/ImmediatePhysicsPhysX/ImmediatePhysicsPersistentContactPairData_PhysX.h
+++ b/Engine/Source/Runtime/Engine/Public/Physics/ImmediatePhysics/ImmediatePhysicsPhysX/ImmediatePhysicsPersistentContactPairData_PhysX.h
@@ -1,32 +1,3 @@
 // Copyright Epic Games, Inc. All Rights Reserved.
 
 #pragma once
-
-#if WITH_PHYSX && PHYSICS_INTERFACE_PHYSX
-
-#ifndef PERSISTENT_CONTACT_PAIRS
-#define PERSISTENT_CONTACT_PAIRS 1
-#endif
-
-namespace ImmediatePhysics_PhysX
-{
-
-#if PERSISTENT_CONTACT_PAIRS
-	/** Persistent data like contact manifold, friction patches, etc... */
-	struct FPersistentContactPairData
-	{
-		PxCache Cache;
-		PxU8* Frictions;
-		PxU32 NumFrictions;
-		uint32 SimCount;
-
-		void Clear()
-		{
-			FPlatformMemory::Memzero(this, sizeof(FPersistentContactPairData));
-		}
-	};
-#endif
-
-}
-
-#endif // WITH_PHYSX
\ No newline at end of file
diff --git a/Engine/Source/Runtime/Engine/Public/Physics/ImmediatePhysics/ImmediatePhysicsPhysX/ImmediatePhysicsShape_PhysX.h b/Engine/Source/Runtime/Engine/Public/Physics/ImmediatePhysics/ImmediatePhysicsPhysX/ImmediatePhysicsShape_PhysX.h
--- a/Engine/Source/Runtime/Engine/Public/Physics/ImmediatePhysics/ImmediatePhysicsPhysX/ImmediatePhysicsShape_PhysX.h
+++ b/Engine/Source/Runtime/Engine/Public/Physics/ImmediatePhysics/ImmediatePhysicsPhysX/ImmediatePhysicsShape_PhysX.h
@@ -1,60 +1,3 @@
 // Copyright Epic Games, Inc. All Rights Reserved.
 
 #pragma once
-
-#if WITH_PHYSX && PHYSICS_INTERFACE_PHYSX
-#include "PhysXPublic.h"
-#include "Physics/ImmediatePhysics/ImmediatePhysicsPhysX/ImmediatePhysicsMaterial_PhysX.h"
-
-namespace ImmediatePhysics_PhysX
-{
-	struct FMaterialHandle;
-}
-
-namespace ImmediatePhysics_PhysX
-{
-	/** Holds shape data*/
-	struct FShape
-	{
-		PxTransform LocalTM;
-		FMaterial* Material;
-		FMaterial InternalMaterial;
-		PxGeometry* Geometry;
-		PxVec3 BoundsOffset;
-		float BoundsMagnitude;
-		void* UserData;
-
-		FShape()
-			: LocalTM(PxTransform(PxIDENTITY::PxIdentity))
-			, Material(nullptr)
-			, Geometry(nullptr)
-			, BoundsOffset(PxVec3(PxIDENTITY::PxIdentity))
-			, BoundsMagnitude(0.0f)
-			, UserData(nullptr)
-		{}
-
-		FShape(const PxTransform& InLocalTM, const PxVec3& InBoundsOffset, const float InBoundsMagnitude, PxGeometry* InGeometry, FMaterial* InMaterial = nullptr)
-			: LocalTM(InLocalTM)
-			, Material(InMaterial)
-			, Geometry(InGeometry)
-			, BoundsOffset(InBoundsOffset)
-			, BoundsMagnitude(InBoundsMagnitude)
-			, UserData(nullptr)
-		{
-		}
-
-		FShape(const PxTransform& InLocalTM, const PxVec3& InBoundsOffset, const float InBoundsMagnitude, PxGeometry* InGeometry, FMaterial InMaterial)
-			: LocalTM(InLocalTM)
-			, InternalMaterial(InMaterial)
-			, Geometry(InGeometry)
-			, BoundsOffset(InBoundsOffset)
-			, BoundsMagnitude(InBoundsMagnitude)
-			, UserData(nullptr)
-		{
-			Material = &InternalMaterial;
-		}
-	};
-
-}
-
-#endif // WITH_PHYSX
\ No newline at end of file
diff --git a/Engine/Source/Runtime/Engine/Public/Physics/ImmediatePhysics/ImmediatePhysicsPhysX/ImmediatePhysicsSimulation_PhysX.h b/Engine/Source/Runtime/Engine/Public/Physics/ImmediatePhysics/ImmediatePhysicsPhysX/ImmediatePhysicsSimulation_PhysX.h
--- a/Engine/Source/Runtime/Engine/Public/Physics/ImmediatePhysics/ImmediatePhysicsPhysX/ImmediatePhysicsSimulation_PhysX.h
+++ b/Engine/Source/Runtime/Engine/Public/Physics/ImmediatePhysics/ImmediatePhysicsPhysX/ImmediatePhysicsSimulation_PhysX.h
@@ -1,278 +1,3 @@
 // Copyright Epic Games, Inc. All Rights Reserved.
 
 #pragma once
-
-#if WITH_PHYSX && PHYSICS_INTERFACE_PHYSX
-#include "PhysXPublic.h"
-
-#include "Physics/PhysicsInterfaceTypes.h"
-
-#include "Engine/EngineTypes.h"
-#include "Physics/ImmediatePhysics/ImmediatePhysicsPhysX/ImmediatePhysicsCore_PhysX.h"
-#include "Physics/ImmediatePhysics/ImmediatePhysicsPhysX/ImmediatePhysicsActor_PhysX.h"
-#include "Physics/ImmediatePhysics/ImmediatePhysicsPhysX/ImmediatePhysicsJoint_PhysX.h"
-#include "Physics/ImmediatePhysics/ImmediatePhysicsPhysX/ImmediatePhysicsLinearBlockAllocator_PhysX.h"
-#include "Physics/ImmediatePhysics/ImmediatePhysicsPhysX/ImmediatePhysicsContactPair_PhysX.h"
-#include "Physics/ImmediatePhysics/ImmediatePhysicsPhysX/ImmediatePhysicsPersistentContactPairData_PhysX.h"
-#include "Physics/ImmediatePhysics/ImmediatePhysicsPhysX/ImmediatePhysicsCacheAllocator_PhysX.h"
-#include "Physics/ImmediatePhysics/ImmediatePhysicsPhysX/ImmediatePhysicsConstraintAllocator_PhysX.h"
-#include "Physics/ImmediatePhysics/ImmediatePhysicsPhysX/ImmediatePhysicsKinematicTarget_PhysX.h"
-#include "Physics/ImmediatePhysics/ImmediatePhysicsPhysX/ImmediatePhysicsD6JointData_PhysX.h"
-
-#include "Chaos/ChaosEngineInterface.h"	//temp to help get code out of engine
-
-class UBodySetup;
-class UPhysicsConstraintTemplate;
-struct FBodyInstance;
-struct FConstraintInstance;
-
-namespace ImmediatePhysics_PhysX
-{
-	struct FActorHandle;
-	struct FJointHandle;
-    struct FContactPointRecorder;
-}
-
-namespace ImmediatePhysics_PhysX
-{
-	struct FActorData
-	{
-		immediate::PxRigidBodyData RigidBodyData;
-		FTransform InitialTransform;
-
-		bool bStatic;
-		bool bKinematic;
-	};
-
-	FActorData CreateActorData(const FActorCreationParams& InParams);
-
-	/** Owns all the data associated with the simulation. Can be considered a single scene or world*/
-	struct ENGINE_API FSimulation
-	{
-	public:
-
-		FActorHandle* InsertActorData(const FActor& InActor, const FActorData& InData);
-
-		void RemoveActor(FActorHandle* InHandle);
-
-		/** Create a kinematic body with the same setup as the body instance and add it to the simulation */
-		FActorHandle* CreateKinematicActor(FBodyInstance* BodyInstance, const FTransform& TM);
-
-		/** Create a dynamic body with the same setup as the body instance and add it to the simulation */
-		FActorHandle* CreateDynamicActor(FBodyInstance* BodyInstance, const FTransform& TM);
-
-		/** Create a static body with the same setup as the body instance and add it to the simulation */
-		FActorHandle* CreateStaticActor(FBodyInstance* BodyInstance);
-
-		/** Create a body of the specified type with the same setup as the body instance and add it to the simulation */
-		FActorHandle* CreateActor(EActorType ActorType, FBodyInstance* BodyInstance, const FTransform& Transform);
-
-		/** Create a physical joint with the same setup as the constraint instance and add it to the simulation */
-		FJointHandle* CreateJoint(FConstraintInstance* ConstraintInstance, FActorHandle* Body1, FActorHandle* Body2);
-
-		const immediate::PxRigidBodyData& GetLowLevelBody(int32 ActorDataIndex) const
-		{
-			return RigidBodiesData[ActorDataIndex];
-		}
-
-		immediate::PxRigidBodyData& GetLowLevelBody(int32 ActorDataIndex)
-		{
-			return RigidBodiesData[ActorDataIndex];
-		}
-
-		/** Sets the number of active bodies. This number is reset any time a new simulated body is created */
-		void SetNumActiveBodies(int32 NumActiveBodies);
-
-		/** An array of actors to ignore. */
-		struct FIgnorePair
-		{
-			FActorHandle* A;
-			FActorHandle* B;
-		};
-
-		/** Set pair of bodies to ignore collision for */
-		void SetIgnoreCollisionPairTable(const TArray<FIgnorePair>& InIgnoreCollisionPairTable);
-
-		/** Set bodies that require no collision */
-		void SetIgnoreCollisionActors(const TArray<FActorHandle*>& InIgnoreCollision);
-
-		/** Get/Set whether the body is kinematic or not, kinematics do not simulate and move where they're told. They also act as if they have infinite mass */
-		bool GetIsKinematic(int32 ActorDataIndex);
-		void SetIsKinematic(int32 ActorDataIndex, bool bKinematic);
-
-		/** Advance the simulation by DeltaTime */
-		void Simulate(float DeltaTime, const FVector& Gravity);
-		void Simulate_AssumesLocked(float DeltaTime, const FVector& Gravity);
-
-		/** Whether or not an entity is simulated */
-		bool IsSimulated(int32 ActorDataIndex) const
-		{
-			return ActorDataIndex < NumSimulatedBodies;
-		}
-
-		/** Add a radial impulse to the given actor */
-		void AddRadialForce(int32 ActorDataIndex, const FVector& Origin, float Strength, float Radius, ERadialImpulseFalloff Falloff, EForceType ForceType);
-
-		/** Add a force to the given actor */
-		void AddForce(int32 ActorDataIndex, const FVector& Force);
-
-		FSimulation();
-
-		~FSimulation();
-
-		int32 NumActors() const { return ActorHandles.Num(); }
-
-		FActorHandle* GetActorHandle(int ActorHandleIndex) { return ActorHandles[ActorHandleIndex]; }
-		const FActorHandle* GetActorHandle(int ActorHandleIndex) const { return ActorHandles[ActorHandleIndex]; }
-
-		/* Set solver iteration counts per step */
-		void SetPositionIterationCount(int32 InIterationCount);
-		void SetVelocityIterationCount(int32 InIterationCount);
-
-	private:
-		friend FActorHandle;
-
-		const FImmediateKinematicTarget& GetKinematicTarget(int32 ActorDataIndex) const
-		{
-			return KinematicTargets[ActorDataIndex];
-		}
-
-		FImmediateKinematicTarget& GetKinematicTarget(int32 ActorDataIndex)
-		{
-			return KinematicTargets[ActorDataIndex];
-		}
-
-		template <EActorType ActorType>
-		int32 CreateActor(PxRigidActor* RigidActor, const FTransform& TM);
-
-		/** Swap actor data - that is move all data associated with the two actors in the various arrays*/
-		void SwapActorData(int32 Entity1Idx, int32 Entity2Idx);
-
-		/** Swap joint data - that is move all data associated with the two joints in the various arrays*/
-		void SwapJointData(int32 Joint1Idx, int32 Joint2Idx);
-
-		/** Ensure arrays are valid */
-		void ValidateArrays() const;
-
-		/** Constructs solver bodies */
-		void ConstructSolverBodies(float DeltaTime, const FVector& Gravity);
-
-		/** Generate contacts*/
-		void GenerateContacts();
-
-		/** Batch constraints and re-order them for optimal processing */
-		void BatchConstraints();
-
-		/** Prepares the various constraints (contact,joints) for the solver */
-		void PrepareConstraints(float DeltaTime);
-
-		/** Solve constraints and integrate velocities */
-		void SolveAndIntegrate(float DeltaTime);
-
-		/** Prepares iteration cache for generating contacts */
-		void PrepareIterationCache();
-
-		//void EvictCache();
-
-	private:
-
-		/** Mapping from entity index to handle */
-		TArray<FActorHandle*>	ActorHandles;
-
-		/** Mapping from constraint index to handle */
-		TArray<FJointHandle*> JointHandles;
-
-		/** Entities holding loose data. NOTE: for performance reasons we don't automatically cleanup on destructor (needed for tarray swaps etc...) it's very important that Terminate is called */
-		TArray<FActor> Actors;
-		TArray<FJoint> Joints;
-
-		/** Workspace memory that we use for per frame allocations */
-		FLinearBlockAllocator Workspace;
-
-		/** Low level rigid body data */
-		TArray<immediate::PxRigidBodyData> RigidBodiesData;
-
-		/** Low level solver bodies data */
-		TArray<PxSolverBodyData> SolverBodiesData;
-
-		/** Kinematic targets used to implicitly compute the velocity of moving kinematic actors */
-		TArray<FImmediateKinematicTarget> KinematicTargets;
-
-		TArray<PxVec3> PendingAcceleration;
-
-		/** Low level contact points generated for this frame. Points are grouped together by pairs */
-		TArray<Gu::ContactPoint> ContactPoints;
-
-		/** Shapes used in the entire simulation. Shapes are sorted in the same order as actors. Note that an actor can have multiple shapes which will be adjacent*/
-		struct FShapeSOA
-		{
-			TArray<PxTransform> LocalTMs;
-			TArray<FMaterial> Materials;
-			TArray<const PxGeometry*> Geometries;
-			TArray<float> Bounds;
-			TArray<PxVec3> BoundsOffsets;
-			TArray<int32> OwningActors;
-#if PERSISTENT_CONTACT_PAIRS
-			TArray<FPersistentContactPairData> ContactPairData;
-#endif
-		} ShapeSOA;
-
-		/** Low level solver bodies */
-		PxSolverBody* SolverBodies;
-
-		/** Low level constraint descriptors.*/
-		TArray<PxSolverConstraintDesc> OrderedDescriptors;
-		TArray<PxConstraintBatchHeader> BatchHeaders;
-
-		/** JointData as passed in from physics constraint template */
-		TArray<D6JointData> JointData;
-
-		/** When new joints are created we have to update the processing order */
-		bool bDirtyJointData;
-
-		PxU32 NumContactHeaders;
-		PxU32 NumJointHeaders;
-		int32 NumActiveJoints;
-
-		/** Contact pairs generated for this frame */
-		TArray<FContactPair> ContactPairs;
-
-		/** Number of dynamic bodies associated with the simulation */
-		int32 NumSimulatedBodies;
-
-		/** Number of dynamic bodies that are actually active */
-		int32 NumActiveSimulatedBodies;
-
-		/** Number of kinematic bodies (dynamic but not simulated) associated with the simulation */
-		int32 NumKinematicBodies;
-
-		/** Total number of simulated shapes in the scene */
-		int32 NumSimulatedShapesWithCollision;
-
-		/** Number of position iterations used by solver */
-		int32 NumPositionIterations;
-
-		/** Number of velocity iterations used by solver */
-		int32 NumVelocityIterations;
-
-		/** Count of how many times we've ticked. Useful for cache invalidation */
-		int32 SimCount;
-
-		/** Both of these are slow to access. Make sure to use iteration cache when possible */
-		TMap<FActorHandle*, TSet<FActorHandle*>> IgnoreCollisionPairTable;
-		TSet<FActorHandle*> IgnoreCollisionActors;
-
-		/** This cache is used to record which generate contact iteration we can skip. This assumes the iteration order has not changed (add/remove/swap actors must invalidate this) */
-		bool bRecreateIterationCache;
-
-		TArray<int32> SkipCollisionCache;	//Holds the iteration count that we should skip due to ignore filtering
-
-		friend FContactPointRecorder;
-
-		FCacheAllocator CacheAllocator;
-		FConstraintAllocator ConstraintAllocator;
-	};
-
-}
-
-#endif // WITH_PHYSX
diff --git a/Engine/Source/Runtime/Engine/Public/Physics/ImmediatePhysics/ImmediatePhysicsPhysX/ResourceManager_PhysX.h b/Engine/Source/Runtime/Engine/Public/Physics/ImmediatePhysics/ImmediatePhysicsPhysX/ResourceManager_PhysX.h
--- a/Engine/Source/Runtime/Engine/Public/Physics/ImmediatePhysics/ImmediatePhysicsPhysX/ResourceManager_PhysX.h
+++ b/Engine/Source/Runtime/Engine/Public/Physics/ImmediatePhysics/ImmediatePhysicsPhysX/ResourceManager_PhysX.h
@@ -1,192 +1,3 @@
 // Copyright Epic Games, Inc. All Rights Reserved.
 
 #pragma once
-
-#if WITH_PHYSX && PHYSICS_INTERFACE_PHYSX
-
-#include "CoreMinimal.h"
-#include "CoreTypes.h"
-#include "Containers/SparseArray.h"
-#include "Physics/ImmediatePhysics/ImmediatePhysicsPhysX/ImmediatePhysicsMaterial_PhysX.h"
-
-namespace ImmediatePhysics_PhysX
-{
-	struct FMaterial;
-}
-
-namespace ImmediatePhysics_PhysX
-{
-	/* List of available resource types managed here */
-	enum class EResourceType : uint8
-	{
-		Material = 0,
-		NumResources,
-		Invalid
-	};
-
-	/* Handle to a specific shared resource */
-	struct FResourceHandle
-	{
-		FResourceHandle()
-			: Type(EResourceType::Invalid)
-			, ResourceIndex(INDEX_NONE)
-			, ResourceId(INDEX_NONE)
-		{}
-
-		FResourceHandle(EResourceType InType, int32 InIndex, int32 InId)
-			: Type(InType)
-			, ResourceIndex(InIndex)
-			, ResourceId(InId)
-		{}
-
-		/** Get the type of the resource, see the list of supported types above */
-		EResourceType GetType() const { return Type; }
-
-		/** Get the storage index of this resource, note this is not enough to identify a resource, the comparison id is also required */
-		int32 GetIndex() const { return ResourceIndex; }
-
-		/** ID for this resource, used to tell whether a resource is what this handle describes, or just reuses memory */
-		int32 GetId() const { return ResourceId; }
-
-		/** Invalidate the handle */
-		void Invalidate()
-		{
-			Type = EResourceType::Invalid;
-			ResourceIndex = INDEX_NONE;
-			ResourceId = INDEX_NONE;
-		}
-
-	private:
-
-		EResourceType Type;
-		int32 ResourceIndex;
-		int32 ResourceId;
-	};
-
-
-	/** Simple pair for holding resources and IDs, with a per-type ID counter for comparisons */
-	template<typename ResourceType>
-	class TResourceWithId
-	{
-		static int32 ResourceIdCounter;
-
-	public:
-
-		TResourceWithId()
-		{
-			Id = ResourceIdCounter++;
-		}
-
-		ResourceType Resource;
-		int32 Id;
-	};
-
-	template<typename ResourceType>
-	int32 TResourceWithId<ResourceType>::ResourceIdCounter = 0;
-
-	/** Responsible for holding and providing shared resources for the physics system under immediate mode.
-	 *  Any resource that cannot be owned entirely by another physics object should be managed here by this
-	 *  shared resource manager.
-	 */
-	class FSharedResourceManager
-	{
-	public:
-
-		friend struct FScopedSharedResourceReadLock;
-
-		static FSharedResourceManager& Get()
-		{
-			return Instance;
-		}
-
-		/** Material methods */
-		FResourceHandle CreateMaterial();
-		void ReleaseMaterial(int32 InIndex);
-		int32 GetMaterialId(int32 InIndex);
-		FMaterial* GetMaterial(int32 InIndex);
-
-		/** Returns the locking primitive used by this manager, usually the scoped resource lock should suffice */
-		FRWLock& GetLockObject();
-
-	private:
-
-		FSharedResourceManager() {}
-		~FSharedResourceManager() {}
-		FSharedResourceManager(const FSharedResourceManager& Other) = delete;
-		FSharedResourceManager(FSharedResourceManager&& Other) = delete;
-		FSharedResourceManager& operator=(const FSharedResourceManager&& Other) = delete;
-		FSharedResourceManager& operator=(FSharedResourceManager&& Other) = delete;
-
-		/** Storage for managed resources */
-		TSparseArray<TResourceWithId<FMaterial>> Materials;
-
-		/** Locking primitive. Should be used in the desired mode whenever manipulated. Simulations will read lock this
-		 *  so user code cannot write to shared resources while a simulation is in flight
-		 */
-		FRWLock ResourceLock;
-
-		/** Singleton instance of the manager */
-		static FSharedResourceManager Instance;
-	};
-
-	/** Modes for scoped resource locking */
-	enum class LockMode : uint8
-	{
-		Read,
-		Write
-	};
-
-	/** Scoped locking object for physics resources, to be used whenever manipulating shared objects */
-	template<LockMode LockType>
-	struct FScopedSharedResourceLock
-	{
-		FScopedSharedResourceLock()
-		{
-			DoLock();
-		}
-
-		~FScopedSharedResourceLock()
-		{
-			DoUnlock();
-		}
-
-		void DoLock();
-		void DoUnlock();
-
-		FScopedSharedResourceLock(const FScopedSharedResourceLock& Other) = delete;
-		FScopedSharedResourceLock(FScopedSharedResourceLock&& Other) = delete;
-		FScopedSharedResourceLock& operator=(const FScopedSharedResourceLock& Other) = delete;
-		FScopedSharedResourceLock& operator=(FScopedSharedResourceLock&& Other) = delete;
-	};
-
-	template<>
-	inline void FScopedSharedResourceLock<LockMode::Read>::DoLock()
-	{
-		FSharedResourceManager& Manager = FSharedResourceManager::Get();
-		Manager.GetLockObject().ReadLock();
-	}
-
-	template<>
-	inline void FScopedSharedResourceLock<LockMode::Read>::DoUnlock()
-	{
-		FSharedResourceManager& Manager = FSharedResourceManager::Get();
-		Manager.GetLockObject().ReadUnlock();
-	}
-
-	template<>
-	inline void FScopedSharedResourceLock<LockMode::Write>::DoLock()
-	{
-		FSharedResourceManager& Manager = FSharedResourceManager::Get();
-		Manager.GetLockObject().WriteLock();
-	}
-
-	template<>
-	inline void FScopedSharedResourceLock<LockMode::Write>::DoUnlock()
-	{
-		FSharedResourceManager& Manager = FSharedResourceManager::Get();
-		Manager.GetLockObject().WriteUnlock();
-	}
-
-}
-
-#endif // WITH_PHYSX
diff --git a/Engine/Source/Runtime/Engine/Public/Physics/PhysScene_PhysX.h b/Engine/Source/Runtime/Engine/Public/Physics/PhysScene_PhysX.h
--- a/Engine/Source/Runtime/Engine/Public/Physics/PhysScene_PhysX.h
+++ b/Engine/Source/Runtime/Engine/Public/Physics/PhysScene_PhysX.h
@@ -8,16 +8,10 @@
 #include "PhysicsEngine/RigidBodyIndexPair.h"
 #include "Physics/PhysicsInterfaceTypes.h"
 
-class ISQAccelerator;
-class FSQAccelerator;
-class FSQAcceleratorUnion;
-class FSQAcceleratorEntry;
-class FPhysXSQAccelerator;
-
 /** Buffers used as scratch space for PhysX to avoid allocations during simulation */
-struct FSimulationScratchBuffer
+struct FSimulationScratchBuffer2
 {
-	FSimulationScratchBuffer()
+	FSimulationScratchBuffer2()
 		: Buffer(nullptr)
 		, BufferSize(0)
 	{}
@@ -28,406 +22,3 @@ struct FSimulationScratchBuffer
 	// Allocated size of the buffer
 	int32 BufferSize;
 };
-
-#if !WITH_CHAOS && !WITH_IMMEDIATE_PHYSX
-
-class FPhysicsReplication;
-class IPhysicsReplicationFactory;
-struct FConstraintBrokenDelegateData;
-
-namespace physx
-{
-	class PxActor;
-	class PxScene;
-	class PxSimulationEventCallback;
-}
-
-struct FContactModifyCallback;
-struct FPhysXMbpBroadphaseCallback;
-
-#if WITH_APEX
-namespace nvidia
-{
-	namespace apex
-	{
-		class Scene;
-	}
-}
-#endif // WITH_APEX
-
-#if WITH_PHYSX
-/** Interface for the creation of customized simulation event callbacks. */
-class ISimEventCallbackFactory
-{
-public:
-	virtual physx::PxSimulationEventCallback* Create(FPhysScene_PhysX* PhysScene) = 0;
-	virtual void Destroy(physx::PxSimulationEventCallback* Callback) = 0;
-};
-#endif // WITH PHYSX
-
-/** Interface for the creation of contact modify callbacks. */
-class IContactModifyCallbackFactory
-{
-public:
-	virtual FContactModifyCallback* Create(FPhysScene_PhysX* PhysScene) = 0;
-	virtual void Destroy(FContactModifyCallback* Callback) = 0;
-};
-
-/** Interface for the creation of ccd contact modify callbacks. */
-class ICCDContactModifyCallbackFactory
-{
-public:
-	virtual FCCDContactModifyCallback* Create(FPhysScene_PhysX* PhysScene) = 0;
-	virtual void Destroy(FCCDContactModifyCallback* Callback) = 0;
-};
-
-/** Container object for a physics engine 'scene'. */
-
-class FPhysScene_PhysX
-{
-public:
-	ENGINE_API FPhysScene_PhysX(const AWorldSettings* Settings = nullptr);
-	ENGINE_API ~FPhysScene_PhysX();
-
-	//////////////////////////////////////////////////////////////////////////
-	// PhysicsInterface
-
-	//// Actor creation/registration
-	//void ReleaseActor(FPhysicsActorReference& InActor);
-	void AddActorsToScene_AssumesLocked(const TArray<FPhysicsActorHandle>& InActors);
-	void AddAggregateToScene(const FPhysicsAggregateHandle& InAggregate);
-
-	//Owning world is made private so that any code which depends on setting an owning world can update
-	void SetOwningWorld(UWorld* InOwningWorld);
-	UWorld* GetOwningWorld() { return OwningWorld; }
-	const UWorld* GetOwningWorld() const { return OwningWorld; }
-
-	FPhysicsReplication* GetPhysicsReplication() { return PhysicsReplication; }
-
-	/** Lets the scene update anything related to this BodyInstance as it's now being terminated */
-	void RemoveBodyInstanceFromPendingLists_AssumesLocked(FBodyInstance* BodyInstance);
-
-	/** Add a custom callback for next step that will be called on every substep */
-	void AddCustomPhysics_AssumesLocked(FBodyInstance* BodyInstance, FCalculateCustomPhysics& CalculateCustomPhysics);
-
-	/** Adds a force to a body - We need to go through scene to support substepping */
-	void AddForce_AssumesLocked(FBodyInstance* BodyInstance, const FVector& Force, bool bAllowSubstepping, bool bAccelChange);
-
-	/** Adds a force to a body at a specific position - We need to go through scene to support substepping */
-	void AddForceAtPosition_AssumesLocked(FBodyInstance* BodyInstance, const FVector& Force, const FVector& Position, bool bAllowSubstepping, bool bIsLocalForce = false);
-
-	/** Adds a radial force to a body - We need to go through scene to support substepping */
-	void AddRadialForceToBody_AssumesLocked(FBodyInstance* BodyInstance, const FVector& Origin, const float Radius, const float Strength, const uint8 Falloff, bool bAccelChange, bool bAllowSubstepping);
-
-	/** Clears currently accumulated forces on a specified body instance */
-	void ClearForces_AssumesLocked(FBodyInstance* BodyInstance, bool bAllowSubstepping);
-
-	/** Adds torque to a body - We need to go through scene to support substepping */
-	void AddTorque_AssumesLocked(FBodyInstance* BodyInstance, const FVector& Torque, bool bAllowSubstepping, bool bAccelChange);
-
-	/** Clears currently accumulated torques on a specified body instance */
-	void ClearTorques_AssumesLocked(FBodyInstance* BodyInstance, bool bAllowSubstepping);
-
-	/** Sets a Kinematic actor's target position - We need to do this here to support substepping*/
-	void SetKinematicTarget_AssumesLocked(FBodyInstance* BodyInstance, const FTransform& TargetTM, bool bAllowSubstepping);
-
-	/** Gets a Kinematic actor's target position - We need to do this here to support substepping
-	* Returns true if kinematic target has been set. If false the OutTM is invalid */
-	bool GetKinematicTarget_AssumesLocked(const FBodyInstance* BodyInstance, FTransform& OutTM) const;
-
-	/** Gets the collision disable table */
-	const TMap<uint32, TMap<struct FRigidBodyIndexPair, bool> *> & GetCollisionDisableTableLookup()
-	{
-		return CollisionDisableTableLookup;
-	}
-
-	/** Adds to queue of skelmesh we want to add to collision disable table */
-	ENGINE_API void DeferredAddCollisionDisableTable(uint32 SkelMeshCompID, TMap<struct FRigidBodyIndexPair, bool> * CollisionDisableTable);
-
-	/** Adds to queue of skelmesh we want to remove from collision disable table */
-	ENGINE_API void DeferredRemoveCollisionDisableTable(uint32 SkelMeshCompID);
-
-	/** Add this SkeletalMeshComponent to the list needing kinematic bodies updated before simulating physics */
-	bool MarkForPreSimKinematicUpdate(USkeletalMeshComponent* InSkelComp, ETeleportType InTeleport, bool bNeedsSkinning);
-
-	/** Remove this SkeletalMeshComponent from set needing kinematic update before simulating physics*/
-	void ClearPreSimKinematicUpdate(USkeletalMeshComponent* InSkelComp);
-
-	/** Pending constraint break events */
-	ENGINE_API void AddPendingOnConstraintBreak(FConstraintInstance* ConstraintInstance);
-	/** Pending wake/sleep events */
-	ENGINE_API void AddPendingSleepingEvent(FBodyInstance* BI, ESleepEvent SleepEventType);
-
-	/** Gets the array of collision notifications, pending execution at the end of the physics engine run. */
-	TArray<FCollisionNotifyInfo>& GetPendingCollisionNotifies() { return PendingCollisionData.PendingCollisionNotifies; }
-
-	/** @return Whether physics scene supports scene origin shifting */
-	static bool SupportsOriginShifting() { return true; }
-
-	/** Shifts physics scene origin by specified offset */
-	void ApplyWorldOffset(FVector InOffset);
-
-	/** Set the gravity and timing of all physics scenes */
-	ENGINE_API void SetUpForFrame(const FVector* NewGrav, float InDeltaSeconds = 0.0f, float UnusedMinPhysicsDeltaTime = 0.0f, float InMaxPhysicsDeltaTime = 0.0f, float InMaxSubstepDeltaTime = 0.f, int32 InMaxSubsteps = 1, bool bUnused = false);
-
-	/** Starts a frame */
-	ENGINE_API void StartFrame();
-
-	/** Ends a frame */
-	ENGINE_API void EndFrame(ULineBatchComponent* InLineBatcher);
-
-	/** Waits for all physics scenes to complete */
-	ENGINE_API void WaitPhysScenes();
-
-	/** returns the completion event for a frame */
-	FGraphEventRef GetCompletionEvent()
-	{
-		return PhysicsSceneCompletion;
-	}
-
-	FGraphEventArray GetCompletionEvents();
-	bool IsCompletionEventComplete() const;
-
-	/** Handle exec commands related to scene (PXVIS and APEXVIS) */
-	bool HandleExecCommands(const TCHAR* Cmd, FOutputDevice* Ar);
-
-	void ListAwakeRigidBodies(bool bIncludeKinematic);
-
-	ENGINE_API int32 GetNumAwakeBodies();
-
-	/** Static factories used to override the simulation contact modify callback from other modules.*/
-	ENGINE_API static TSharedPtr<IContactModifyCallbackFactory> ContactModifyCallbackFactory;
-	ENGINE_API static TSharedPtr<ICCDContactModifyCallbackFactory> CCDContactModifyCallbackFactory;
-
-	/** Static factory used to override the physics replication manager from other modules. This is useful for custom game logic.
-	If not set it defaults to using FPhysicsReplication. */
-	ENGINE_API static TSharedPtr<IPhysicsReplicationFactory> PhysicsReplicationFactory;
-
-	ENGINE_API void SerializeForTesting(FArchive& Ar);
-
-	//////////////////////////////////////////////////////////////////////////
-	// PhysScene_PhysX interface
-
-	/** Get the PxScene */
-	ENGINE_API physx::PxScene* GetPxScene() const;
-
-	ENGINE_API ISQAccelerator* GetSQAccelerator() const;
-
-	ENGINE_API FSQAcceleratorUnion* GetSQAcceleratorUnion() const;
-
-#if WITH_APEX
-	/** Get the ApexScene*/
-	ENGINE_API nvidia::apex::Scene*	GetApexScene() const { return PhysXScene; }
-#endif
-	
-	/** The number of frames it takes to rebuild the PhysX scene query AABB tree. The bigger the number, the smaller fetchResults takes per frame, but the more the tree deteriorates until a new tree is built */
-	void SetPhysXTreeRebuildRate(int32 RebuildRate);
-
-	/** Ensures that the collision tree is built. */
-	ENGINE_API void EnsureCollisionTreeIsBuilt(UWorld* World);
-
-	/** Kill the visual debugger */
-	ENGINE_API void KillVisualDebugger();
-
-#if WITH_PHYSX
-	/** Static factory used to override the simulation event callback from other modules.
-	If not set it defaults to using FPhysXSimEventCallback. */
-	ENGINE_API static TSharedPtr<ISimEventCallbackFactory> SimEventCallbackFactory;
-
-#endif	// WITH_PHYSX
-
-	//////////////////////////////////////////////////////////////////////////
-
-	DECLARE_MULTICAST_DELEGATE_TwoParams(FOnPhysScenePreTick, FPhysScene_PhysX*, float /*DeltaSeconds*/);
-	FOnPhysScenePreTick OnPhysScenePreTick;
-
-	DECLARE_MULTICAST_DELEGATE_TwoParams(FOnPhysSceneStep, FPhysScene_PhysX*, float /*DeltaSeconds*/);
-	FOnPhysSceneStep OnPhysSceneStep;
-
-	DECLARE_MULTICAST_DELEGATE_OneParam(FOnPhysScenePostTick, FPhysScene*);
-	FOnPhysScenePostTick OnPhysScenePostTick;
-
-private:
-	/** Indicates whether the scene is using substepping */
-	bool							bSubstepping;
-
-	/** World that owns this physics scene */
-	UWorld*							OwningWorld;
-
-	/** Replication manager that updates physics bodies towards replicated physics state */
-	FPhysicsReplication*			PhysicsReplication;
-
-#if WITH_CUSTOM_SQ_STRUCTURE
-	TMap<physx::PxRigidActor*, FSQAcceleratorEntry*> RigidActorToSQEntries;
-#endif
-
-#if WITH_APEX
-	nvidia::apex::Scene* PhysXScene;
-#else // #if WITH_APEX
-	PxScene* PhysXScene;
-#endif // #if WITH_APEX
-
-	/** Whether or not the given scene is between its execute and sync point. */
-	bool							bPhysXSceneExecuting;
-
-	/** Frame time, weighted with current frame time. */
-	float							AveragedFrameTime;
-
-	/**
-	* Weight for averaged frame time.  Value should be in the range [0.0f, 1.0f].
-	* Weight = 0.0f => no averaging; current frame time always used.
-	* Weight = 1.0f => current frame time ignored; initial value of AveragedFrameTime[i] is always used.
-	*/
-	float							FrameTimeSmoothingFactor;
-
-	/** DeltaSeconds from UWorld. */
-	float										DeltaSeconds;
-	/** DeltaSeconds from the WorldSettings. */
-	float										MaxPhysicsDeltaTime;
-
-	/** LineBatcher from UWorld. */
-	ULineBatchComponent*						LineBatcher;
-
-	/** Completion event (not tasks) for the physics scenes these are fired by the physics system when it is done; prerequisites for the below */
-	FGraphEventRef PhysicsSubsceneCompletion;
-	/** Completion events (not tasks) for the frame lagged physics scenes these are fired by the physics system when it is done; prerequisites for the below */
-	FGraphEventRef FrameLaggedPhysicsSubsceneCompletion;
-	/** Completion events (task) for the physics scenes	(both apex and non-apex). This is a "join" of the above. */
-	FGraphEventRef PhysicsSceneCompletion;
-
-	// Data for scene scratch buffers, these will be allocated once on FPhysScene construction and used
-	// for the calls to PxScene::simulate to save it calling into the OS to allocate during simulation
-	FSimulationScratchBuffer SimScratchBuffer;
-
-	// Boundary value for PhysX scratch buffers (currently PhysX requires the buffer length be a multiple of 16K)
-	static const int32 SimScratchBufferBoundary = 16 * 1024;
-
-#if WITH_CUSTOM_SQ_STRUCTURE
-	FSQAcceleratorUnion* SQAcceleratorUnion;
-	FSQAccelerator* SQAccelerator;
-#if WITH_PHYSX && !WITH_CHAOS
-	FPhysXSQAccelerator* PhysXSQAccelerator;
-#endif
-#endif
-
-#if WITH_PHYSX
-
-	bool bIsSceneSimulating = false;
-
-	/** Dispatcher for CPU tasks */
-	class PxCpuDispatcher*			CPUDispatcher;
-	/** Simulation event callback object */
-	physx::PxSimulationEventCallback*			SimEventCallback;
-	FContactModifyCallback*			ContactModifyCallback;
-	FCCDContactModifyCallback*			CCDContactModifyCallback;
-	FPhysXMbpBroadphaseCallback* MbpBroadphaseCallback;
-
-	struct FPendingCollisionData
-	{
-		/** Array of collision notifications, pending execution at the end of the physics engine run. */
-		TArray<FCollisionNotifyInfo>	PendingCollisionNotifies;
-	};
-
-	FPendingCollisionData PendingCollisionData;
-
-	struct FPendingConstraintData
-	{
-		/** Array of constraint broken notifications, pending execution at the end of the physics engine run. */
-		TArray<FConstraintBrokenDelegateData> PendingConstraintBroken;
-	};
-
-	FPendingConstraintData PendingConstraintData;
-
-#endif	// WITH_PHYSX
-
-	/** Start simulation on the physics scene of the given type */
-	ENGINE_API void TickPhysScene(FGraphEventRef& InOutCompletionEvent);
-
-	/** Fetches results, fires events, and adds debug lines */
-	void ProcessPhysScene();
-
-	/** Sync components in the scene to physics bodies that changed */
-	void SyncComponentsToBodies_AssumesLocked();
-
-	/** Call after WaitPhysScene on the synchronous scene to make deferred OnRigidBodyCollision calls.  */
-	ENGINE_API void DispatchPhysNotifications_AssumesLocked();
-
-	/** Add any debug lines from the physics scene of the given type to the supplied line batcher. */
-	ENGINE_API void AddDebugLines(class ULineBatchComponent* LineBatcherToUse);
-
-	void AddActorsToPhysXScene_AssumesLocked(const TArray<FPhysicsActorHandle>& InActors);
-
-	/** @return Whether physics scene is using substepping */
-	bool IsSubstepping() const;
-
-	/** Initialize a scene of the given type.  Must only be called once for each scene type. */
-	void InitPhysScene(const AWorldSettings* Settings = nullptr);
-
-	/** Terminate a scene of the given type.  Must only be called once for each scene type. */
-	void TermPhysScene();
-
-	/** Called when all subscenes of a given scene are complete, calls  ProcessPhysScene*/
-	void SceneCompletionTask(ENamedThreads::Type CurrentThread, const FGraphEventRef& MyCompletionGraphEvent);
-	
-	/** Process kinematic updates on any deferred skeletal meshes */
-	void UpdateKinematicsOnDeferredSkelMeshes();
-
-	/** Task created from TickPhysScene so we can substep without blocking */
-	bool SubstepSimulation(FGraphEventRef& InOutCompletionEvent);
-
-	/** Set whether we're doing a static load and want to stall, or are during gameplay and want to distribute over many frames */
-	void SetIsStaticLoading(bool bStaticLoading);
-
-	/** The number of frames it takes to rebuild the PhysX scene query AABB tree. The bigger the number, the smaller fetchResults takes per frame, but the more the tree deteriorates until a new tree is built */
-	void SetPhysXTreeRebuildRateImp(int32 RebuildRate);
-
-#if WITH_PHYSX
-	/** User data wrapper passed to physx */
-	FPhysxUserData PhysxUserData;
-
-#endif
-
-	class FPhysSubstepTask * PhysSubStepper;
-
-	struct FPendingCollisionDisableTable
-	{
-		uint32 SkelMeshCompID;
-		TMap<struct FRigidBodyIndexPair, bool>* CollisionDisableTable;
-	};
-
-	/** Updates CollisionDisableTableLookup with the deferred insertion and deletion */
-	void FlushDeferredCollisionDisableTableQueue();
-
-	bool ExecPxVis(const TCHAR* Cmd, FOutputDevice* Ar);
-
-	bool ExecApexVis(const TCHAR* Cmd, FOutputDevice* Ar);
-
-
-	/** Queue of deferred collision table insertion and deletion */
-	TArray<FPendingCollisionDisableTable> DeferredCollisionDisableTableQueue;
-
-	/** Map from SkeletalMeshComponent UniqueID to a pointer to the collision disable table inside its PhysicsAsset */
-	TMap< uint32, TMap<struct FRigidBodyIndexPair, bool>* >		CollisionDisableTableLookup;
-
-#if WITH_PHYSX
-	TMap<FBodyInstance*, ESleepEvent> PendingSleepEvents;
-#endif
-
-	/** Information about how to perform kinematic update before physics */
-	struct FDeferredKinematicUpdateInfo
-	{
-		/** Whether to teleport physics bodies or not */
-		ETeleportType	TeleportType;
-		/** Whether to update skinning info */
-		bool			bNeedsSkinning;
-	};
-
-	/** Map of SkeletalMeshComponents that need their bone transforms sent to the physics engine before simulation. */
-	TArray<TPair<USkeletalMeshComponent*, FDeferredKinematicUpdateInfo>>	DeferredKinematicUpdateSkelMeshes;
-
-	FDelegateHandle PreGarbageCollectDelegateHandle;
-
-	int32 PhysXTreeRebuildRate;
-};
-
-#endif
\ No newline at end of file
diff --git a/Engine/Source/Runtime/Engine/Public/Physics/PhysicsFiltering.h b/Engine/Source/Runtime/Engine/Public/Physics/PhysicsFiltering.h
--- a/Engine/Source/Runtime/Engine/Public/Physics/PhysicsFiltering.h
+++ b/Engine/Source/Runtime/Engine/Public/Physics/PhysicsFiltering.h
@@ -143,13 +143,3 @@ inline void UpdateMaskFilter(uint32& Word3, FMaskFilter NewMaskFilter)
 	Word3 &= (0xFFFFFFFFu >> NumExtraFilterBits);	//we drop the top NumExtraFilterBits bits because that's where the new mask filter is going
 	Word3 |= uint32(NewMaskFilter) << (32 - NumExtraFilterBits);
 }
-
-#if PHYSICS_INTERFACE_PHYSX
-	extern ENGINE_API PxSimulationFilterShader GSimulationFilterShader;
-#endif
-
-/** Utility for creating a PhysX PxFilterData for performing a query (trace) against the scene */
-// PxFilterData ZZ_CreateQueryFilterData(const uint8 MyChannel, const bool bTraceComplex,
-// 	const FCollisionResponseContainer& InCollisionResponseContainer,
-// 	const struct FCollisionObjectQueryParams & ObjectParam,
-// 	const bool bMultitrace);
diff --git a/Engine/Source/Runtime/Engine/Public/Physics/PhysicsGeometry.h b/Engine/Source/Runtime/Engine/Public/Physics/PhysicsGeometry.h
--- a/Engine/Source/Runtime/Engine/Public/Physics/PhysicsGeometry.h
+++ b/Engine/Source/Runtime/Engine/Public/Physics/PhysicsGeometry.h
@@ -7,46 +7,13 @@
 
 struct FKSphereElem;
 
-namespace physx
-{
-	class PxSphereGeometry;
-	class PxTransform;
-
-	class PxSphereGeometry;
-	class PxBoxGeometry;
-	class PxCapsuleGeometry;
-	class PxConvexMeshGeometry;
-	class PxTriangleMesh;
-	class PxTriangleMeshGeometry;
-}
-
 class UBodySetup;
 
 /** Helper struct for iterating over shapes in a body setup.*/
 struct ENGINE_API FBodySetupShapeIterator
 {
-	// Only valid for PhysX and Chaos - will fail to compile when used otherwise
-	using TTransform = TChooseClass<WITH_PHYSX && PHYSICS_INTERFACE_PHYSX, physx::PxTransform, TChooseClass<WITH_CHAOS, FTransform, void>::Result>::Result;
-
 	FBodySetupShapeIterator();
 
-#if WITH_PHYSX && PHYSICS_INTERFACE_PHYSX
-	FBodySetupShapeIterator(const FVector& InScale3D, const FTransform& InRelativeTM, bool bDoubleSidedTrimeshes);
-	/** Iterates over the elements array and creates the needed geometry and local pose. Note that this memory is on the stack so it's illegal to use it by reference outside the lambda */
-	template <typename ElemType, typename GeomType>
-	void ForEachShape(const TArrayView<ElemType>& Elements, TFunctionRef<void(const ElemType& Elem, const GeomType& Geom, const TTransform& LocalPose, float ContactOffset, float RestOffset)> VisitorFunc) const;
-
-	/** Helper function to determine contact offset params */
-	static void GetContactOffsetParams(float& InOutContactOffsetFactor, float& InOutMinContactOffset, float& InOutMaxContactOffset);
-
-private:
-
-	template <typename ElemType, typename GeomType> bool PopulateGeometryAndTransform(const ElemType& Elem, GeomType& Geom, TTransform& OutTM) const;
-	template <typename GeomType> float ComputeContactOffset(const GeomType& Geom) const;
-	template <typename ElemType> float ComputeRestOffset(const ElemType& Geom) const;
-	template <typename ElemType> FString GetDebugName() const;
-#endif
-
 private:
 	FVector Scale3D;
 	const FTransform& RelativeTM;
@@ -62,18 +29,3 @@ private:
 
 	bool bDoubleSidedTriMeshGeo;
 };
-
-#if WITH_PHYSX && PHYSICS_INTERFACE_PHYSX
-
-/// @cond DOXYGEN_WARNINGS
-
-//Explicit export of template instantiation 
-extern template ENGINE_API void FBodySetupShapeIterator::ForEachShape(const TArrayView<FKSphereElem>&, TFunctionRef<void(const FKSphereElem&, const physx::PxSphereGeometry&, const FBodySetupShapeIterator::TTransform&, float, float)>) const;
-extern template ENGINE_API void FBodySetupShapeIterator::ForEachShape(const TArrayView<FKBoxElem>&, TFunctionRef<void(const FKBoxElem&, const physx::PxBoxGeometry&, const FBodySetupShapeIterator::TTransform&, float, float)>) const;
-extern template ENGINE_API void FBodySetupShapeIterator::ForEachShape(const TArrayView<FKSphylElem>&, TFunctionRef<void(const FKSphylElem&, const physx::PxCapsuleGeometry&, const FBodySetupShapeIterator::TTransform&, float, float)>) const;
-extern template ENGINE_API void FBodySetupShapeIterator::ForEachShape(const TArrayView<FKConvexElem>&, TFunctionRef<void(const FKConvexElem&, const physx::PxConvexMeshGeometry&, const FBodySetupShapeIterator::TTransform&, float, float)>) const;
-extern template ENGINE_API void FBodySetupShapeIterator::ForEachShape(const TArrayView<physx::PxTriangleMesh*>&, TFunctionRef<void(physx::PxTriangleMesh* const &, const physx::PxTriangleMeshGeometry&, const FBodySetupShapeIterator::TTransform&, float, float)>) const;
-
-/// @endcond
-
-#endif //WITH_PHYSX
\ No newline at end of file
diff --git a/Engine/Source/Runtime/Engine/Public/Physics/PhysicsInterfaceCore.h b/Engine/Source/Runtime/Engine/Public/Physics/PhysicsInterfaceCore.h
--- a/Engine/Source/Runtime/Engine/Public/Physics/PhysicsInterfaceCore.h
+++ b/Engine/Source/Runtime/Engine/Public/Physics/PhysicsInterfaceCore.h
@@ -4,22 +4,5 @@
 
 #include "PhysicsInterfaceDeclares.h"
 #include "PhysicsCore.h"
-
-
-#if PHYSICS_INTERFACE_PHYSX
-
-#include "Physics/PhysicsInterfacePhysX.h"
-#include "Physics/PhysScene_PhysX.h"
-#include "Physics/Experimental/PhysScene_ImmediatePhysX.h"
-#include "Physics/Experimental/PhysicsInterfaceImmediatePhysX.h"
-
-#elif WITH_CHAOS
-
 #include "Physics/Experimental/PhysInterface_Chaos.h"
 #include "Physics/Experimental/PhysScene_Chaos.h"
-
-#else
-
-static_assert(false, "A physics engine interface must be defined to build");
-
-#endif
diff --git a/Engine/Source/Runtime/Engine/Public/Physics/PhysicsInterfacePhysX.h b/Engine/Source/Runtime/Engine/Public/Physics/PhysicsInterfacePhysX.h
--- a/Engine/Source/Runtime/Engine/Public/Physics/PhysicsInterfacePhysX.h
+++ b/Engine/Source/Runtime/Engine/Public/Physics/PhysicsInterfacePhysX.h
@@ -1,423 +1,3 @@
 // Copyright Epic Games, Inc. All Rights Reserved.
 
 #pragma once
-
-#if !WITH_CHAOS && !WITH_IMMEDIATE_PHYSX
-
-#include "EngineGlobals.h"
-#include "Engine/EngineTypes.h"
-#include "CollisionQueryParams.h"
-#include "PhysicsEngine/ConstraintTypes.h"
-#include "PhysicsInterfaceTypes.h"
-#include "Containers/ContainerAllocationPolicies.h"
-#include "GenericPhysicsInterface.h"
-#include "Chaos/ChaosEngineInterface.h"	//temp to help get code out of engine
-
-struct FLinearDriveConstraint;
-struct FAngularDriveConstraint;
-struct FConstraintDrive;
-struct FPhysxUserData;
-struct FKAggregateGeom;
-
-struct FCollisionQueryParams;
-struct FCollisionResponseParams;
-struct FCollisionObjectQueryParams;
-struct FCollisionShape;
-class UWorld;
-struct FBodyInstance;
-class UPhysicalMaterial;
-
-namespace physx
-{
-	class PxRigidActor;
-	class PxRigidStatic;
-	class PxRigidDynamic;
-	class PxRigidBody;
-	class PxAggregate;
-	class PxShape;
-	class PxD6Joint;
-	class PxGeometry;
-	class PxQuat;
-	class PxTransform;
-	class PxScene;
-	class PxBoxGeometry;
-	class PxCapsuleGeometry;
-	class PxSphereGeometry;
-	class PxConvexMeshGeometry;
-	class PxTriangleMeshGeometry;
-	class PxMassProperties;
-	class PxGeometryHolder;
-	class PxMaterial;
-}
-
-struct ENGINE_API FPhysicsActorHandle_PhysX
-{
-	FPhysicsActorHandle_PhysX();
-
-	bool IsValid() const;
-	bool Equals(const FPhysicsActorHandle_PhysX& Other) const;
-	bool operator==(const FPhysicsActorHandle_PhysX& Other) const { return Equals(Other); }
-
-	physx::PxRigidActor* SyncActor;
-
-private:
-
-	struct FPxActorContainer
-	{
-		physx::PxRigidActor* Actor;
-		physx::PxRigidStatic* Static;
-		physx::PxRigidDynamic* Dynamic;
-		physx::PxRigidBody* Body;
-	};
-};
-
-struct ENGINE_API FPhysicsConstraintHandle_PhysX
-{
-	FPhysicsConstraintHandle_PhysX();
-
-	bool IsValid() const;
-	bool Equals(const FPhysicsConstraintHandle_PhysX& Other) const;
-
-	physx::PxD6Joint* ConstraintData;
-};
-
-struct ENGINE_API FPhysicsAggregateHandle_PhysX
-{
-	FPhysicsAggregateHandle_PhysX();
-
-	bool IsValid() const;
-
-	physx::PxAggregate* Aggregate;
-};
-
-struct ENGINE_API FPhysicsShapeHandle_PhysX
-{
-	FPhysicsShapeHandle_PhysX() : Shape(nullptr) {}
-	explicit FPhysicsShapeHandle_PhysX(physx::PxShape* InShape) : Shape(InShape) {}
-
-	bool IsValid() const { return Shape != nullptr; }
-	bool operator==(const FPhysicsShapeHandle_PhysX Rhs) const { return Shape == Rhs.Shape; }
-	physx::PxShape* Shape;
-};
-
-inline uint32 GetTypeHash(const FPhysicsShapeHandle_PhysX& InHandle)
-{
-	return GetTypeHash(InHandle.Shape);
-}
-
-/**
- * This object is essentially a one-stop container for any geometry a shape can have and is necessary because of the
- * PxGeometryHolder type. This needs to have a longer lifetime than any usage of the geometry types it returns.
- * Because we want to have that inside the interface this container is required to manage the lifetime of the holder
- */
-struct ENGINE_API FPhysicsGeometryCollection_PhysX
-{
-	// Delete default constructor, want only construction by interface (private constructor below)
-	FPhysicsGeometryCollection_PhysX() = delete;
-	// No copying or assignment, move construction only, these are defaulted in the source file as they need
-	// to be able to delete physx::PxGeometryHolder which is incomplete here
-	FPhysicsGeometryCollection_PhysX(const FPhysicsGeometryCollection_PhysX& Copy) = delete;
-	FPhysicsGeometryCollection_PhysX& operator=(const FPhysicsGeometryCollection_PhysX& Copy) = delete;
-	FPhysicsGeometryCollection_PhysX(FPhysicsGeometryCollection_PhysX&& Steal);
-	FPhysicsGeometryCollection_PhysX& operator=(FPhysicsGeometryCollection_PhysX&& Steal);
-	~FPhysicsGeometryCollection_PhysX();
-
-	ECollisionShapeType GetType() const;
-	const physx::PxGeometry& GetGeometry() const;
-	physx::PxGeometry& GetGeometry();
-	const physx::PxBoxGeometry& GetBoxGeometry() const;
-	physx::PxBoxGeometry& GetBoxGeometry();
-	const physx::PxSphereGeometry& GetSphereGeometry() const;
-	physx::PxSphereGeometry& GetSphereGeometry();
-	const physx::PxCapsuleGeometry& GetCapsuleGeometry() const;
-	physx::PxCapsuleGeometry& GetCapsuleGeometry();
-	const physx::PxConvexMeshGeometry& GetConvexGeometry() const;
-	physx::PxConvexMeshGeometry& GetConvexGeometry();
-	const physx::PxTriangleMeshGeometry& GetTriMeshGeometry() const;
-	physx::PxTriangleMeshGeometry& GetTriMeshGeometry();
-
-private:
-	friend struct FPhysicsInterface_PhysX;
-	explicit FPhysicsGeometryCollection_PhysX(const FPhysicsShapeHandle_PhysX& InShape);
-
-	// PhysX geom holder, needs to exist longer than the uses of any geometry it returns
-	TUniquePtr<physx::PxGeometryHolder> GeomHolder;
-};
-
-/**
-* API to access the physics interface. All calls to FPhysicsInterface functions should be inside an Execute* callable.
-* This is to ensure correct lock semantics and command buffering if the specific API supports deferred commands.
-*/
-struct ENGINE_API FPhysicsCommand_PhysX
-{
-	// Executes with appropriate read locking, return true if execution took place (actor was valid)
-	static bool ExecuteRead(const FPhysicsActorHandle_PhysX& InHandle, TFunctionRef<void(const FPhysicsActorHandle_PhysX& Actor)> InCallable);
-	static bool ExecuteRead(USkeletalMeshComponent* InMeshComponent, TFunctionRef<void()> InCallable);
-	static bool ExecuteRead(const FPhysicsActorHandle_PhysX& InHandleA, const FPhysicsActorHandle_PhysX& InHandleB, TFunctionRef<void(const FPhysicsActorHandle_PhysX& ActorA, const FPhysicsActorHandle_PhysX& ActorB)> InCallable);
-	static bool ExecuteRead(const FPhysicsConstraintHandle_PhysX& InHandle, TFunctionRef<void(const FPhysicsConstraintHandle_PhysX& Constraint)> InCallable);
-	static bool ExecuteRead(FPhysScene_PhysX* InScene, TFunctionRef<void()> InCallable);
-
-	// Executes with appropriate write locking, return true if execution took place (actor was valid)
-	static bool ExecuteWrite(const FPhysicsActorHandle_PhysX& InHandle, TFunctionRef<void(const FPhysicsActorHandle_PhysX& Actor)> InCallable);
-	static bool ExecuteWrite(FPhysicsActorHandle_PhysX& InHandle, TFunctionRef<void(FPhysicsActorHandle_PhysX& Actor)> InCallable);
-	static bool ExecuteWrite(USkeletalMeshComponent* InMeshComponent, TFunctionRef<void()> InCallable);
-	static bool ExecuteWrite(const FPhysicsActorHandle_PhysX& InHandleA, const FPhysicsActorHandle_PhysX& InHandleB, TFunctionRef<void(const FPhysicsActorHandle_PhysX& ActorA, const FPhysicsActorHandle_PhysX& ActorB)> InCallable);
-	static bool ExecuteWrite(const FPhysicsConstraintHandle_PhysX& InHandle, TFunctionRef<void(const FPhysicsConstraintHandle_PhysX& Constraint)> InCallable);
-	static bool ExecuteWrite(FPhysScene_PhysX* InScene, TFunctionRef<void()> InCallable);
-
-	// Executes function on a shape
-	static void ExecuteShapeWrite(FBodyInstance* InInstance, FPhysicsShapeHandle_PhysX& InShape, TFunctionRef<void(FPhysicsShapeHandle_PhysX& InShape)> InCallable);
-};
-
-struct ENGINE_API FPhysicsInterface_PhysX : public FGenericPhysicsInterface, public FChaosEngineInterface	//temp to help get code out of engine
-{
-	// Describe the interface to identify it to the caller
-	static FString GetInterfaceDescription() { return TEXT("PhysX"); }
-
-	// PhysX Only functions, not related to wider physics interface
-	// To be used only in code that handles PhysX
-	static physx::PxRigidActor* GetPxRigidActorFromScene_AssumesLocked(const FPhysicsActorHandle_PhysX& InActorHandle);
-	static physx::PxRigidActor* GetPxRigidActor_AssumesLocked(const FPhysicsActorHandle_PhysX& InHandle);
-	static physx::PxRigidDynamic* GetPxRigidDynamic_AssumesLocked(const FPhysicsActorHandle_PhysX& InHandle);
-	static physx::PxRigidBody* GetPxRigidBody_AssumesLocked(const FPhysicsActorHandle_PhysX& InHandle);
-	static const FBodyInstance* ShapeToOriginalBodyInstance(const FBodyInstance* InCurrentInstance, const physx::PxShape* InShape);
-	//////////////////////////////////////////////////////////////////////////
-
-	// Aggregate interface functions
-	static FPhysicsAggregateHandle_PhysX CreateAggregate(int32 MaxBodies);
-	static void ReleaseAggregate(FPhysicsAggregateHandle_PhysX& InAggregate);
-	static int32 GetNumActorsInAggregate(const FPhysicsAggregateHandle_PhysX& InAggregate);
-	static void AddActorToAggregate_AssumesLocked(const FPhysicsAggregateHandle_PhysX& InAggregate, const FPhysicsActorHandle_PhysX& InActor);
-
-	//////////////////////////////////////////////////////////////////////////
-	// Shape interface functions
-	//////////////////////////////////////////////////////////////////////////
-	static FPhysicsShapeHandle_PhysX CreateShape(physx::PxGeometry* InGeom, bool bSimulation = true, bool bQuery = true, UPhysicalMaterial* InSimpleMaterial = nullptr, TArray<UPhysicalMaterial*>* InComplexMaterials = nullptr);
-	
-	static void AddGeometry(const FPhysicsActorHandle& InActor, const FGeometryAddParams& InParams, TArray<FPhysicsShapeHandle_PhysX>* OutOptShapes = nullptr);
-	static FPhysicsShapeHandle_PhysX CloneShape(const FPhysicsShapeHandle_PhysX& InShape);
-
-	static FCollisionFilterData GetSimulationFilter(const FPhysicsShapeHandle_PhysX& InShape);
-	static FCollisionFilterData GetQueryFilter(const FPhysicsShapeHandle_PhysX& InShape);
-	static bool IsSimulationShape(const FPhysicsShapeHandle_PhysX& InShape);
-	static bool IsQueryShape(const FPhysicsShapeHandle_PhysX& InShape);
-	static ECollisionShapeType GetShapeType(const FPhysicsShapeHandle_PhysX& InShape);
-	static FPhysicsGeometryCollection_PhysX GetGeometryCollection(const FPhysicsShapeHandle_PhysX& InShape);
-	static FTransform GetLocalTransform(const FPhysicsShapeHandle_PhysX& InShape);
-	static FTransform GetTransform(const FPhysicsShapeHandle_PhysX& InShape);
-	static void* GetUserData(const FPhysicsShapeHandle_PhysX& InShape);
-
-	// Set the mask filter of a shape, which is an extra level of filtering during collision detection / query for extra channels like "Blue Team" and "Red Team"
-	static void SetMaskFilter(const FPhysicsShapeHandle_PhysX& InShape, FMaskFilter InFilter);
-	static void SetSimulationFilter(const FPhysicsShapeHandle_PhysX& InShape, const FCollisionFilterData& InFilter);
-	static void SetQueryFilter(const FPhysicsShapeHandle_PhysX& InShape, const FCollisionFilterData& InFilter);
-	static void SetIsSimulationShape(const FPhysicsShapeHandle_PhysX& InShape, bool bIsSimShape);
-	static void SetIsQueryShape(const FPhysicsShapeHandle_PhysX& InShape, bool bIsQueryShape);
-	static void SetUserData(const FPhysicsShapeHandle_PhysX& InShape, void* InUserData);
-	static void SetGeometry(const FPhysicsShapeHandle_PhysX& InShape, physx::PxGeometry& InGeom);
-	static void SetLocalTransform(const FPhysicsShapeHandle_PhysX& InShape, const FTransform& NewLocalTransform);
-	static void SetMaterials(const FPhysicsShapeHandle_PhysX& InShape, const TArrayView<UPhysicalMaterial*>InMaterials);
-
-	//////////////////////////////////////////////////////////////////////////
-	// Material interface functions
-	//////////////////////////////////////////////////////////////////////////
-	static FPhysicsMaterialHandle CreateMaterial(const UPhysicalMaterial* InMaterial);
-	static void ReleaseMaterial(FPhysicsMaterialHandle_PhysX& InHandle);
-	static void UpdateMaterial(const FPhysicsMaterialHandle_PhysX& InHandle, UPhysicalMaterial* InMaterial);
-	static void SetUserData(const FPhysicsMaterialHandle_PhysX& InHandle, void* InUserData);
-
-	//////////////////////////////////////////////////////////////////////////
-	// Actor interface functions
-	//////////////////////////////////////////////////////////////////////////
-
-	// #PHYS2 - These should be on the scene, but immediate mode stops us for now, eventually that should spawn its own minimal IM scene and these should move
-	static void CreateActor(const FActorCreationParams& Params, FPhysicsActorHandle& Handle);
-	static void ReleaseActor(FPhysicsActorHandle_PhysX& InHandle, FPhysScene* InScene = nullptr, bool bNeverDeferRelease = false);
-	//////////////////////////////////////////////////////////////////////////
-
-	// NOTE: Adding this temporarily, until we phase out Handle.IsValid().
-	static bool IsValid(const FPhysicsActorHandle& Handle) { return Handle.IsValid(); }
-
-	static bool IsInScene(const FPhysicsActorHandle& InActorReference);
-
-	template<typename AllocatorType>
-	static int32 GetAllShapes_AssumedLocked(const FPhysicsActorHandle_PhysX& InHandle, TArray<FPhysicsShapeHandle_PhysX, AllocatorType>& OutShapes);
-	static int32 GetNumShapes(const FPhysicsActorHandle_PhysX& InHandle);
-
-	static void ReleaseShape(const FPhysicsShapeHandle_PhysX& InShape);
-
-	static void AttachShape(const FPhysicsActorHandle_PhysX& InActor, const FPhysicsShapeHandle_PhysX& InNewShape);
-	static void DetachShape(const FPhysicsActorHandle_PhysX& InActor, FPhysicsShapeHandle_PhysX& InShape, bool bWakeTouching = true);
-
-	static void SetActorUserData_AssumesLocked(const FPhysicsActorHandle_PhysX& InHandle, FPhysxUserData* InUserData);
-
-	static bool IsRigidBody(const FPhysicsActorHandle_PhysX& InHandle);
-	static bool IsDynamic(const FPhysicsActorHandle_PhysX& InHandle);
-	static bool IsStatic(const FPhysicsActorHandle_PhysX& InHandle);
-	static bool IsKinematic_AssumesLocked(const FPhysicsActorHandle_PhysX& InHandle);
-	static bool IsSleeping(const FPhysicsActorHandle_PhysX& InHandle);
-	static bool IsCcdEnabled(const FPhysicsActorHandle_PhysX& InHandle);
-	static bool IsDisabled(const FPhysicsActorHandle_PhysX& InHandle);
-	static void SetDisabled(const FPhysicsActorHandle_PhysX& InHandle, bool bSetDisabled);
-	static FPhysScene* GetCurrentScene(const FPhysicsActorHandle_PhysX& InHandle);
-	static bool CanSimulate_AssumesLocked(const FPhysicsActorHandle_PhysX& InHandle);
-	static float GetMass_AssumesLocked(const FPhysicsActorHandle_PhysX& InHandle);
-
-	static void SetSendsSleepNotifies_AssumesLocked(const FPhysicsActorHandle_PhysX& InHandle, bool bSendSleepNotifies);
-	static void PutToSleep_AssumesLocked(const FPhysicsActorHandle_PhysX& InHandle);
-	static void WakeUp_AssumesLocked(const FPhysicsActorHandle_PhysX& InHandle);
-
-	static void SetIsKinematic_AssumesLocked(const FPhysicsActorHandle_PhysX& InHandle, bool bIsKinematic);
-	static void SetCcdEnabled_AssumesLocked(const FPhysicsActorHandle_PhysX& InHandle, bool bIsCcdEnabled);
-
-	static FTransform GetGlobalPose_AssumesLocked(const FPhysicsActorHandle_PhysX& InHandle);
-	static void SetGlobalPose_AssumesLocked(const FPhysicsActorHandle_PhysX& InHandle, const FTransform& InNewPose, bool bAutoWake = true);
-
-	static FTransform GetTransform_AssumesLocked(const FPhysicsActorHandle& InRef, bool bForceGlobalPose = false);
-
-	static bool HasKinematicTarget_AssumesLocked(const FPhysicsActorHandle_PhysX& InHandle);
-	static FTransform GetKinematicTarget_AssumesLocked(const FPhysicsActorHandle_PhysX& InHandle);
-	static void SetKinematicTarget_AssumesLocked(const FPhysicsActorHandle_PhysX& InHandle, const FTransform& InNewTarget);
-
-	static FVector GetLinearVelocity_AssumesLocked(const FPhysicsActorHandle_PhysX& InHandle);
-	static void SetLinearVelocity_AssumesLocked(const FPhysicsActorHandle_PhysX& InHandle, const FVector& InNewVelocity, bool bAutoWake = true);
-
-	static FVector GetAngularVelocity_AssumesLocked(const FPhysicsActorHandle_PhysX& InHandle);
-	static void SetAngularVelocity_AssumesLocked(const FPhysicsActorHandle_PhysX& InHandle, const FVector& InNewVelocity, bool bAutoWake = true);
-	static float GetMaxAngularVelocity_AssumesLocked(const FPhysicsActorHandle_PhysX& InHandle);
-	static void SetMaxAngularVelocity_AssumesLocked(const FPhysicsActorHandle_PhysX& InHandle, float InMaxAngularVelocity);
-
-	static float GetMaxDepenetrationVelocity_AssumesLocked(const FPhysicsActorHandle_PhysX& InHandle);
-	static void SetMaxDepenetrationVelocity_AssumesLocked(const FPhysicsActorHandle_PhysX& InHandle, float InMaxDepenetrationVelocity);
-
-	static FVector GetWorldVelocityAtPoint_AssumesLocked(const FPhysicsActorHandle_PhysX& InHandle, const FVector& InPoint);
-
-	static FTransform GetComTransform_AssumesLocked(const FPhysicsActorHandle_PhysX& InHandle);
-	static FTransform GetComTransformLocal_AssumesLocked(const FPhysicsActorHandle_PhysX& InHandle);
-
-	static FVector GetLocalInertiaTensor_AssumesLocked(const FPhysicsActorHandle_PhysX& InHandle);
-	static FBox GetBounds_AssumesLocked(const FPhysicsActorHandle_PhysX& InHandle);
-
-	static void SetLinearDamping_AssumesLocked(const FPhysicsActorHandle_PhysX& InHandle, float InDamping);
-	static void SetAngularDamping_AssumesLocked(const FPhysicsActorHandle_PhysX& InHandle, float InDamping);
-
-	static void AddImpulse_AssumesLocked(const FPhysicsActorHandle_PhysX& InHandle, const FVector& InForce);
-	static void AddAngularImpulseInRadians_AssumesLocked(const FPhysicsActorHandle_PhysX& InHandle, const FVector& InTorque);
-	static void AddVelocity_AssumesLocked(const FPhysicsActorHandle_PhysX& InHandle, const FVector& InForce);
-	static void AddAngularVelocityInRadians_AssumesLocked(const FPhysicsActorHandle_PhysX& InHandle, const FVector& InTorque);
-	static void AddImpulseAtLocation_AssumesLocked(const FPhysicsActorHandle_PhysX& InHandle, const FVector& InImpulse, const FVector& InLocation);
-	static void AddRadialImpulse_AssumesLocked(const FPhysicsActorHandle_PhysX& InHandle, const FVector& InOrigin, float InRadius, float InStrength, ERadialImpulseFalloff InFalloff, bool bInVelChange);
-
-	static bool IsGravityEnabled_AssumesLocked(const FPhysicsActorHandle_PhysX& InHandle);
-	static void SetGravityEnabled_AssumesLocked(const FPhysicsActorHandle_PhysX& InHandle, bool bEnabled);
-
-	static float GetSleepEnergyThreshold_AssumesLocked(const FPhysicsActorHandle_PhysX& InHandle);
-	static void SetSleepEnergyThreshold_AssumesLocked(const FPhysicsActorHandle_PhysX& InHandle, float InEnergyThreshold);
-
-	static void SetMass_AssumesLocked(const FPhysicsActorHandle_PhysX& InHandle, float InMass);
-	static void SetMassSpaceInertiaTensor_AssumesLocked(const FPhysicsActorHandle_PhysX& InHandle, const FVector& InTensor);
-	static void SetComLocalPose_AssumesLocked(const FPhysicsActorHandle_PhysX& InHandle, const FTransform& InComLocalPose);
-
-	static float GetStabilizationEnergyThreshold_AssumesLocked(const FPhysicsActorHandle_PhysX& InHandle);
-	static void SetStabilizationEnergyThreshold_AssumesLocked(const FPhysicsActorHandle_PhysX& InHandle, float InThreshold);
-	static uint32 GetSolverPositionIterationCount_AssumesLocked(const FPhysicsActorHandle_PhysX& InHandle);
-	static void SetSolverPositionIterationCount_AssumesLocked(const FPhysicsActorHandle_PhysX& InHandle, uint32 InSolverIterationCount);
-	static uint32 GetSolverVelocityIterationCount_AssumesLocked(const FPhysicsActorHandle_PhysX& InHandle);
-	static void SetSolverVelocityIterationCount_AssumesLocked(const FPhysicsActorHandle_PhysX& InHandle, uint32 InSolverIterationCount);
-	static float GetWakeCounter_AssumesLocked(const FPhysicsActorHandle_PhysX& InHandle);
-	static void SetWakeCounter_AssumesLocked(const FPhysicsActorHandle_PhysX& InHandle, float InWakeCounter);
-
-	static SIZE_T GetResourceSizeEx(const FPhysicsActorHandle_PhysX& InActorHandle);
-
-	//////////////////////////////////////////////////////////////////////////
-
-	// Constraint interface
-
-	// Functions
-	static FPhysicsConstraintHandle_PhysX CreateConstraint(const FPhysicsActorHandle_PhysX& InActorHandle1, const FPhysicsActorHandle_PhysX& InActorHandle2, const FTransform& InLocalFrame1, const FTransform& InLocalFrame2);
-	static void SetConstraintUserData(const FPhysicsConstraintHandle_PhysX& InHandle, void* InUserData);
-	static void ReleaseConstraint(FPhysicsConstraintHandle_PhysX& InHandle);
-
-	static FTransform GetLocalPose(const FPhysicsConstraintHandle_PhysX& InHandle, EConstraintFrame::Type InFrame);
-	static FTransform GetGlobalPose(const FPhysicsConstraintHandle_PhysX& InHandle, EConstraintFrame::Type InFrame);
-	static FVector GetLocation(const FPhysicsConstraintHandle_PhysX& InHandle);
-	static void GetForce(const FPhysicsConstraintHandle_PhysX& InHandle, FVector& OutLinForce, FVector& OutAngForce);
-	static void GetDriveLinearVelocity(const FPhysicsConstraintHandle_PhysX& InHandle, FVector& OutLinVelocity);
-	static void GetDriveAngularVelocity(const FPhysicsConstraintHandle_PhysX& InHandle, FVector& OutAngVelocity);
-
-	static float GetCurrentSwing1(const FPhysicsConstraintHandle_PhysX& InHandle);
-	static float GetCurrentSwing2(const FPhysicsConstraintHandle_PhysX& InHandle);
-	static float GetCurrentTwist(const FPhysicsConstraintHandle_PhysX& InHandle);
-
-	static void SetCanVisualize(const FPhysicsConstraintHandle_PhysX& InHandle, bool bInCanVisualize);
-	static void SetCollisionEnabled(const FPhysicsConstraintHandle_PhysX& InHandle, bool bInCollisionEnabled);
-	static void SetProjectionEnabled_AssumesLocked(const FPhysicsConstraintHandle_PhysX& InHandle, bool bInProjectionEnabled, float InLinearTolerance = 0.0f, float InAngularToleranceDegrees = 0.0f);
-	static void SetParentDominates_AssumesLocked(const FPhysicsConstraintHandle_PhysX& InHandle, bool bInParentDominates);
-	static void SetBreakForces_AssumesLocked(const FPhysicsConstraintHandle_PhysX& InHandle, float InLinearBreakForce, float InAngularBreakForce);
-	static void SetLocalPose(const FPhysicsConstraintHandle_PhysX& InHandle, const FTransform& InPose, EConstraintFrame::Type InFrame);
-
-	static void SetLinearMotionLimitType_AssumesLocked(const FPhysicsConstraintHandle_PhysX& InHandle, PhysicsInterfaceTypes::ELimitAxis InAxis, ELinearConstraintMotion InMotion);
-	static void SetAngularMotionLimitType_AssumesLocked(const FPhysicsConstraintHandle_PhysX& InHandle, PhysicsInterfaceTypes::ELimitAxis InAxis, EAngularConstraintMotion InMotion);
-
-	static void UpdateLinearLimitParams_AssumesLocked(const FPhysicsConstraintHandle_PhysX& InHandle, float InLimit, float InAverageMass, const FLinearConstraint& InParams);
-	static void UpdateConeLimitParams_AssumesLocked(const FPhysicsConstraintHandle_PhysX& InHandle, float InAverageMass, const FConeConstraint& InParams);
-	static void UpdateTwistLimitParams_AssumesLocked(const FPhysicsConstraintHandle_PhysX& InHandle, float InAverageMass, const FTwistConstraint& InParams);
-	static void UpdateLinearDrive_AssumesLocked(const FPhysicsConstraintHandle_PhysX& InHandle, const FLinearDriveConstraint& InDriveParams);
-	static void UpdateAngularDrive_AssumesLocked(const FPhysicsConstraintHandle_PhysX& InHandle, const FAngularDriveConstraint& InDriveParams);
-	static void UpdateDriveTarget_AssumesLocked(const FPhysicsConstraintHandle_PhysX& InHandle, const FLinearDriveConstraint& InLinDrive, const FAngularDriveConstraint& InAngDrive, bool bInitialized_IGNORED = true);
-	static void SetDrivePosition(const FPhysicsConstraintHandle_PhysX& InHandle, const FVector& InPosition);
-	static void SetDriveOrientation(const FPhysicsConstraintHandle_PhysX& InHandle, const FQuat& InOrientation);
-	static void SetDriveLinearVelocity(const FPhysicsConstraintHandle_PhysX& InHandle, const FVector& InLinVelocity);
-	static void SetDriveAngularVelocity(const FPhysicsConstraintHandle_PhysX& InHandle, const FVector& InAngVelocity);
-
-	static void SetTwistLimit(const FPhysicsConstraintHandle_PhysX& InHandle, float InLowerLimit, float InUpperLimit, float InContactDistance);
-	static void SetSwingLimit(const FPhysicsConstraintHandle_PhysX& InHandle, float InYLimit, float InZLimit, float InContactDistance);
-	static void SetLinearLimit(const FPhysicsConstraintHandle_PhysX& InHandle, float InLimit);
-
-	static bool IsBroken(const FPhysicsConstraintHandle_PhysX& InHandle);
-
-	static bool ExecuteOnUnbrokenConstraintReadOnly(const FPhysicsConstraintHandle_PhysX& InHandle, TFunctionRef<void(const FPhysicsConstraintHandle_PhysX&)> Func);
-	static bool ExecuteOnUnbrokenConstraintReadWrite(const FPhysicsConstraintHandle_PhysX& InHandle, TFunctionRef<void(const FPhysicsConstraintHandle_PhysX&)> Func);
-
-	//////////////////////////////////////////////////////////////////////////
-
-	// Scene query interface functions
-
-	
-
-	// GEOM SWEEP
-
-	//////////////////////////////////////////////////////////////////////////
-	// Trace functions for testing specific geometry (not against a world)
-	//////////////////////////////////////////////////////////////////////////
-	static bool LineTrace_Geom(FHitResult& OutHit, const FBodyInstance* InInstance, const FVector& InStart, const FVector& InEnd, bool bTraceComplex, bool bExtractPhysMaterial = false);
-	static bool Sweep_Geom(FHitResult& OutHit, const FBodyInstance* InInstance, const FVector& InStart, const FVector& InEnd, const FQuat& InShapeRotation, const FCollisionShape& InShape, bool bSweepComplex);
-
-	static bool Overlap_Geom(const FBodyInstance* InBodyInstance, const FPhysicsGeometryCollection& InGeometry, const FTransform& InShapeTransform, FMTDResult* OutOptResult = nullptr);
-	static bool Overlap_Geom(const FBodyInstance* InBodyInstance, const FCollisionShape& InCollisionShape, const FQuat& InShapeRotation, const FTransform& InShapeTransform, FMTDResult* OutOptResult = nullptr);
-
-	static bool GetSquaredDistanceToBody(const FBodyInstance* InInstance, const FVector& InPoint, float& OutDistanceSquared, FVector* OutOptPointOnBody = nullptr);
-
-	//////////////////////////////////////////////////////////////////////////
-	// Misc
-	//////////////////////////////////////////////////////////////////////////
-
-	static bool ExecPhysCommands(const TCHAR* Cmd, FOutputDevice* Ar, UWorld* InWorld);
-
-	static void CalculateMassPropertiesFromShapeCollection(physx::PxMassProperties& OutProperties, const TArray<FPhysicsShapeHandle>& InShapes, float InDensityKGPerCM);
-
-};
-
-template <>
-ENGINE_API int32 FPhysicsInterface_PhysX::GetAllShapes_AssumedLocked(const FPhysicsActorHandle_PhysX& InActorHandle, TArray<FPhysicsShapeHandle_PhysX, FDefaultAllocator>& OutShapes);
-template <>
-ENGINE_API int32 FPhysicsInterface_PhysX::GetAllShapes_AssumedLocked(const FPhysicsActorHandle_PhysX& InActorHandle, PhysicsInterfaceTypes::FInlineShapeArray& OutShapes);
-
-template<>
-ENGINE_API bool FGenericPhysicsInterface::GeomSweepMulti(const UWorld* World, const FPhysicsGeometryCollection& InGeom, const FQuat& InGeomRot, TArray<FHitResult>& OutHits, FVector Start, FVector End, ECollisionChannel TraceChannel, const FCollisionQueryParams& Params, const FCollisionResponseParams& ResponseParams, const FCollisionObjectQueryParams& ObjectParams /*= FCollisionObjectQueryParams::DefaultObjectQueryParam*/);
-
-template<>
-ENGINE_API bool FGenericPhysicsInterface::GeomOverlapMulti(const UWorld* World, const FPhysicsGeometryCollection& InGeom, const FVector& InPosition, const FQuat& InRotation, TArray<FOverlapResult>& OutOverlaps, ECollisionChannel TraceChannel, const FCollisionQueryParams& Params, const FCollisionResponseParams& ResponseParams, const FCollisionObjectQueryParams& ObjectParams);
-
-
-#endif
\ No newline at end of file
diff --git a/Engine/Source/Runtime/Engine/Public/Physics/PhysicsInterfaceTypes.h b/Engine/Source/Runtime/Engine/Public/Physics/PhysicsInterfaceTypes.h
--- a/Engine/Source/Runtime/Engine/Public/Physics/PhysicsInterfaceTypes.h
+++ b/Engine/Source/Runtime/Engine/Public/Physics/PhysicsInterfaceTypes.h
@@ -7,19 +7,14 @@
 #include "PhysicsInterfaceDeclares.h"
 #include "BodySetupEnums.h"
 #include "PhysicsInterfaceTypesCore.h"
-
-#if WITH_CHAOS
 #include "Chaos/Serializable.h"
 
 namespace Chaos
 {
 	class FImplicitObjectUnion;
-
 	class FImplicitObject;
-
 	class FTriangleMeshImplicitObject;
 }
-#endif
 
 class UPhysicalMaterialMask;
 class UMaterialInterface;
@@ -35,12 +30,6 @@ class UMaterialInterface;
 
 struct FKAggregateGeom;
 
-namespace physx
-{
-	class PxShape;
-	class PxTriangleMesh;
-}
-
 struct FPhysicalMaterialMaskParams
 {
 	/** Physical materials mask */
@@ -58,17 +47,9 @@ struct FGeometryAddParams
 	FVector Scale;
 	UPhysicalMaterial* SimpleMaterial;
 	TArrayView<UPhysicalMaterial*> ComplexMaterials;
-#if WITH_CHAOS
 	TArrayView<FPhysicalMaterialMaskParams> ComplexMaterialMasks;
-#endif
 	FTransform LocalTransform;
 	FTransform WorldTransform;
 	FKAggregateGeom* Geometry;
-	// FPhysicsInterfaceTriMesh - Per implementation
-#if WITH_PHYSX
-	TArrayView<physx::PxTriangleMesh*> TriMeshes;
-#endif
-#if WITH_CHAOS
 	TArrayView<TSharedPtr<Chaos::FTriangleMeshImplicitObject, ESPMode::ThreadSafe>> ChaosTriMeshes;
-#endif
 };
\ No newline at end of file
diff --git a/Engine/Source/Runtime/Engine/Public/Physics/PhysicsInterfaceUtils.h b/Engine/Source/Runtime/Engine/Public/Physics/PhysicsInterfaceUtils.h
--- a/Engine/Source/Runtime/Engine/Public/Physics/PhysicsInterfaceUtils.h
+++ b/Engine/Source/Runtime/Engine/Public/Physics/PhysicsInterfaceUtils.h
@@ -12,11 +12,9 @@ struct FConstraintInstance;
 
 // FILTER DATA
 
-/** Utility for creating a PhysX PxFilterData for performing a query (trace) against the scene */
+/** Utility for creating a filter data object for performing a query (trace) against the scene */
 FCollisionFilterData CreateQueryFilterData(const uint8 MyChannel, const bool bTraceComplex, const FCollisionResponseContainer& InCollisionResponseContainer, const struct FCollisionQueryParams& QueryParam, const struct FCollisionObjectQueryParams & ObjectParam, const bool bMultitrace);
 
-#if WITH_PHYSX
-
 struct FConstraintBrokenDelegateData
 {
 	FConstraintBrokenDelegateData(FConstraintInstance* ConstraintInstance);
@@ -30,8 +28,6 @@ struct FConstraintBrokenDelegateData
 	int32 ConstraintIndex;
 };
 
-#endif // WITH_PHYSX
-
 class FPhysicsReplication;
 
 /** Interface for the creation of customized physics replication.*/
diff --git a/Engine/Source/Runtime/Engine/Public/PhysicsPublic.h b/Engine/Source/Runtime/Engine/Public/PhysicsPublic.h
--- a/Engine/Source/Runtime/Engine/Public/PhysicsPublic.h
+++ b/Engine/Source/Runtime/Engine/Public/PhysicsPublic.h
@@ -59,71 +59,9 @@ DECLARE_CYCLE_STAT_EXTERN(TEXT("FetchAndStart Time (all)"), STAT_TotalPhysicsTim
 DECLARE_DWORD_COUNTER_STAT_EXTERN(TEXT("Cloth Actor Count"), STAT_NumCloths, STATGROUP_Physics, ENGINE_API);
 DECLARE_DWORD_COUNTER_STAT_EXTERN(TEXT("Simulated Cloth Verts"), STAT_NumClothVerts, STATGROUP_Physics, ENGINE_API);
 
-#define WITH_PHYSX_VEHICLES WITH_PHYSX && PHYSICS_INTERFACE_PHYSX
-
 /** Pointer to PhysX Command Handler */
 extern ENGINE_API class FPhysCommandHandler* GPhysCommandHandler;
 
-#if PHYSICS_INTERFACE_PHYSX
-
-namespace physx
-{
-	class PxScene;
-	class PxConvexMesh;
-	class PxTriangleMesh;
-	class PxCooking;
-	class PxPhysics;
-	class PxVec3;
-	class PxJoint;
-	class PxMat44;
-	class PxCpuDispatcher;
-	class PxSimulationEventCallback;
-	struct PxActiveTransform;
-	class PxActor;
-	class PxRigidActor;
-}
-
-#if WITH_APEX
-namespace nvidia
-{
-	namespace apex
-	{
-		class DestructibleAsset;
-		class Scene;
-		struct DamageEventReportData;
-		class ApexSDK;
-		class ModuleDestructible;
-		class DestructibleActor;
-		class ModuleClothing;
-		class Module;
-		class ClothingActor;
-		class ClothingAsset;
-		class ApexInterface;
-	}
-}
-#endif // WITH_APEX
-
-struct FConstraintInstance;
-struct FContactModifyCallback;
-struct FCCDContactModifyCallback;
-struct FPhysXMbpBroadphaseCallback;
-class UPhysicsAsset;
-
-using namespace physx;
-#if WITH_APEX
-using namespace nvidia;
-#endif	//WITH_APEX
-
-/** Pointer to PhysX cooking object */
-extern ENGINE_API PxCooking*			GPhysXCooking;
-
-namespace NvParameterized
-{
-	class Interface;
-}
-
-#endif // WITH_PHYSX
-
 /** Information about a specific object involved in a rigid body collision */
 struct ENGINE_API FRigidBodyCollisionInfo
 {
@@ -203,40 +141,14 @@ public:
 	void ENGINE_API Flush();
 	bool ENGINE_API HasPendingCommands();
 
-#if WITH_APEX
-	/** enqueues a command to release destructible actor once apex has finished simulating */
-	void ENGINE_API DeferredRelease(apex::ApexInterface* ApexInterface);
-#endif
-
-#if PHYSICS_INTERFACE_PHYSX
-	void ENGINE_API DeferredRelease(physx::PxScene * PScene);
-	void ENGINE_API DeferredDeleteSimEventCallback(physx::PxSimulationEventCallback* SimEventCallback);
-	void ENGINE_API DeferredDeleteContactModifyCallback(FContactModifyCallback* ContactModifyCallback);
-	void ENGINE_API DeferredDeleteCCDContactModifyCallback(FCCDContactModifyCallback* CCDContactModifyCallback);
-	void ENGINE_API DeferredDeleteMbpBroadphaseCallback(FPhysXMbpBroadphaseCallback* MbpCallback);
-	void ENGINE_API DeferredDeleteCPUDispathcer(physx::PxCpuDispatcher * CPUDispatcher);
-#endif
-	
 private:
 
 	/** Command to execute when physics simulation is done */
 	struct FPhysPendingCommand
 	{
-		union
-		{
-#if WITH_APEX
-			apex::ApexInterface * ApexInterface;
-			apex::DestructibleActor * DestructibleActor;
-#endif
-#if PHYSICS_INTERFACE_PHYSX
-			physx::PxScene* PScene;
-			physx::PxCpuDispatcher* CPUDispatcher;
-			physx::PxSimulationEventCallback* SimEventCallback;
-			FContactModifyCallback* ContactModifyCallback;
-			FCCDContactModifyCallback* CCDContactModifyCallback;
-			FPhysXMbpBroadphaseCallback* MbpCallback;
-#endif
-		} Pointer;
+		//union
+		//{
+		//} Pointer;
 
 		PhysCommand::Type CommandType;
 	};
@@ -270,16 +182,6 @@ FORCEINLINE bool PhysSingleThreadedMode()
 	return false;
 }
 
-#if PHYSICS_INTERFACE_PHYSX
-/** Struct used for passing info to the PhysX shader */
-
-struct FPhysSceneShaderInfo
-{
-	FPhysScene * PhysScene;
-};
-
-#endif
-
 // Only used for legacy serialization (ver < VER_UE4_REMOVE_PHYS_SCALED_GEOM_CACHES)
 class FKCachedConvexDataElement
 {
@@ -331,8 +233,6 @@ public:
 	bool HasValidGeometry();
 };
 
-
-
 ENGINE_API bool	InitGamePhys();
 ENGINE_API void	TermGamePhys();
 
@@ -340,8 +240,6 @@ ENGINE_API void	TermGamePhys();
 ENGINE_API void DeferredPhysResourceCleanup();
 
 
-
-
 FTransform FindBodyTransform(AActor* Actor, FName BoneName);
 FBox	FindBodyBox(AActor* Actor, FName BoneName);
 
