Add an on screen hint for console command 'DumpUnbuiltLightInteractions'

diff --git a/Engine/Source/Runtime/Engine/Private/UnrealEngine.cpp b/Engine/Source/Runtime/Engine/Private/UnrealEngine.cpp
--- a/Engine/Source/Runtime/Engine/Private/UnrealEngine.cpp
+++ b/Engine/Source/Runtime/Engine/Private/UnrealEngine.cpp
@@ -11394,6 +11394,8 @@ float DrawMapWarnings(UWorld* World, FViewport* Viewport, FCanvas* Canvas, UCanv
 			}
 		}
 
+		bool bShouldDisplayCmdHint = false;
+		
 		if (NumLightingScenariosEnabled > 1)
 		{
 			SmallTextItem.Text = LOCTEXT("MULTIPLE_LIGHTING_SCENARIO_LEVELS_ENABLED", "MULTIPLE LIGHTING SCENARIO LEVELS ENABLED");		
@@ -11401,11 +11403,20 @@ float DrawMapWarnings(UWorld* World, FViewport* Viewport, FCanvas* Canvas, UCanv
 		else
 		{
 			// Use 'DumpUnbuiltLightInteractions' to investigate, if lighting is still unbuilt after a lighting build
-			SmallTextItem.Text = FText::Format(LOCTEXT("LIGHTING_NEEDS_TO_BE_REBUILT_FMT", "LIGHTING NEEDS TO BE REBUILT ({0} unbuilt {0}|plural(one=object,other=objects))"), World->NumLightingUnbuiltObjects);		
+			SmallTextItem.Text = FText::Format(LOCTEXT("LIGHTING_NEEDS_TO_BE_REBUILT_FMT", "LIGHTING NEEDS TO BE REBUILT ({0} unbuilt {0}|plural(one=object,other=objects))"), World->NumLightingUnbuiltObjects);
+			bShouldDisplayCmdHint = true;
 		}
 
 		Canvas->DrawItem(SmallTextItem, FVector2D(MessageX, MessageY));
 		MessageY += FontSizeY;
+
+		if (bShouldDisplayCmdHint)
+		{
+			SmallTextItem.SetColor(FLinearColor::White);
+			SmallTextItem.Text = LOCTEXT("UNBUILT_CONSOLE_CMD_HINT", "Run console command 'DumpUnbuiltLightInteractions' to see what is unbuilt");
+			Canvas->DrawItem(SmallTextItem, FVector2D(MessageX, MessageY));
+			MessageY += FontSizeY;
+		}
 	}
 
 	if (World->NumUnbuiltReflectionCaptures > 0)
