ANY_PACKAGE removal: Add UWorld::PostLoadAssetRegistryTags because it needs to handle postloading the tags from its LevelBlueprint that are added in UWorld::GetAssetRegistryTags.
#rnx

diff --git a/Engine/Source/Runtime/Engine/Classes/Engine/Blueprint.h b/Engine/Source/Runtime/Engine/Classes/Engine/Blueprint.h
--- a/Engine/Source/Runtime/Engine/Classes/Engine/Blueprint.h
+++ b/Engine/Source/Runtime/Engine/Classes/Engine/Blueprint.h
@@ -961,6 +961,7 @@ public:
 	virtual bool Modify(bool bAlwaysMarkDirty = true) override;
 	virtual void GetAssetRegistryTags(TArray<FAssetRegistryTag>& OutTags) const override;
 	virtual void PostLoadAssetRegistryTags(const FAssetData& InAssetData, TArray<FAssetRegistryTag>& OutTagsAndValuesToUpdate) const;
+	static void PostLoadBlueprintAssetRegistryTags(const FAssetData& InAssetData, TArray<FAssetRegistryTag>& OutTagsAndValuesToUpdate);
 	virtual FPrimaryAssetId GetPrimaryAssetId() const override;
 	virtual void BeginCacheForCookedPlatformData(const ITargetPlatform *TargetPlatform) override;
 	virtual bool IsCachedCookedPlatformDataLoaded(const ITargetPlatform* TargetPlatform) override;
diff --git a/Engine/Source/Runtime/Engine/Classes/Engine/World.h b/Engine/Source/Runtime/Engine/Classes/Engine/World.h
--- a/Engine/Source/Runtime/Engine/Classes/Engine/World.h
+++ b/Engine/Source/Runtime/Engine/Classes/Engine/World.h
@@ -2874,7 +2874,8 @@ public:
 #if WITH_EDITOR
 	virtual bool Rename(const TCHAR* NewName = NULL, UObject* NewOuter = NULL, ERenameFlags Flags = REN_None) override;
 	virtual void GetAssetRegistryTags(TArray<FAssetRegistryTag>& OutTags) const override;
-	virtual bool IsNameStableForNetworking() const override;	
+	virtual void PostLoadAssetRegistryTags(const FAssetData& InAssetData, TArray<FAssetRegistryTag>& OutTagsAndValuesToUpdate) const;
+	virtual bool IsNameStableForNetworking() const override;
 #endif
 	virtual bool ResolveSubobject(const TCHAR* SubObjectPath, UObject*& OutObject, bool bLoadIfExists) override;
 	virtual void PostDuplicate(bool bDuplicateForPIE) override;
diff --git a/Engine/Source/Runtime/Engine/Private/Blueprint.cpp b/Engine/Source/Runtime/Engine/Private/Blueprint.cpp
--- a/Engine/Source/Runtime/Engine/Private/Blueprint.cpp
+++ b/Engine/Source/Runtime/Engine/Private/Blueprint.cpp
@@ -1015,7 +1015,11 @@ void UBlueprint::GetAssetRegistryTags(TArray<FAssetRegistryTag>& OutTags) const
 void UBlueprint::PostLoadAssetRegistryTags(const FAssetData& InAssetData, TArray<FAssetRegistryTag>& OutTagsAndValuesToUpdate) const
 {
 	Super::PostLoadAssetRegistryTags(InAssetData, OutTagsAndValuesToUpdate);
+	PostLoadBlueprintAssetRegistryTags(InAssetData, OutTagsAndValuesToUpdate);
+}
 
+void UBlueprint::PostLoadBlueprintAssetRegistryTags(const FAssetData& InAssetData, TArray<FAssetRegistryTag>& OutTagsAndValuesToUpdate)
+{
 	auto FixTagValueShortClassName = [&InAssetData, &OutTagsAndValuesToUpdate](FName TagName, FAssetRegistryTag::ETagType TagType)
 	{
 		FString TagValue = InAssetData.GetTagValueRef<FString>(TagName);
diff --git a/Engine/Source/Runtime/Engine/Private/World.cpp b/Engine/Source/Runtime/Engine/Private/World.cpp
--- a/Engine/Source/Runtime/Engine/Private/World.cpp
+++ b/Engine/Source/Runtime/Engine/Private/World.cpp
@@ -8460,6 +8460,14 @@ void UWorld::GetAssetRegistryTags(TArray<FAssetRegistryTag>& OutTags) const
 	FWorldDelegates::GetAssetTags.Broadcast(this, OutTags);
 }
 
+void UWorld::PostLoadAssetRegistryTags(const FAssetData& InAssetData, TArray<FAssetRegistryTag>& OutTagsAndValuesToUpdate) const
+{
+	Super::PostLoadAssetRegistryTags(InAssetData, OutTagsAndValuesToUpdate);
+
+	// GetAssetRegistryTags appends the LevelBlueprint tags to the World's tags, so we also have to run the Blueprint PostLoadAssetRegistryTags
+	UBlueprint::PostLoadBlueprintAssetRegistryTags(InAssetData, OutTagsAndValuesToUpdate);
+}
+
 bool UWorld::IsNameStableForNetworking() const
 {
 	return bIsNameStableForNetworking || Super::IsNameStableForNetworking();
