Clarify the units in a couple of constraint function comments (no functional changes)

diff --git a/Engine/Source/Runtime/Engine/Classes/PhysicsEngine/ConstraintDrives.h b/Engine/Source/Runtime/Engine/Classes/PhysicsEngine/ConstraintDrives.h
--- a/Engine/Source/Runtime/Engine/Classes/PhysicsEngine/ConstraintDrives.h
+++ b/Engine/Source/Runtime/Engine/Classes/PhysicsEngine/ConstraintDrives.h
@@ -132,7 +132,7 @@ struct ENGINE_API FAngularDriveConstraint
 	UPROPERTY(EditAnywhere, Category = AngularMotor)
 	FRotator OrientationTarget;
 
-	/** Target angular velocity relative to the body reference frame. */
+	/** Target angular velocity relative to the body reference frame in revolutions per second. */
 	UPROPERTY(EditAnywhere, Category = AngularMotor)
 	FVector AngularVelocityTarget;
 
diff --git a/Engine/Source/Runtime/Engine/Classes/PhysicsEngine/ConstraintInstance.h b/Engine/Source/Runtime/Engine/Classes/PhysicsEngine/ConstraintInstance.h
--- a/Engine/Source/Runtime/Engine/Classes/PhysicsEngine/ConstraintInstance.h
+++ b/Engine/Source/Runtime/Engine/Classes/PhysicsEngine/ConstraintInstance.h
@@ -801,7 +801,7 @@ public:
 		return ProfileInstance.AngularDrive.IsVelocityDriveEnabled();
 	}
 	
-	/** Set the angular drive's angular velocity target*/
+	/** Set the angular drive's angular velocity target (in revolutions per second)*/
 	void SetAngularVelocityTarget(const FVector& InVelTarget);
 
 	/** Get the angular drive's angular velocity target*/
