Fixed some variables that were not following naming conventions

diff --git a/Engine/Source/Runtime/Engine/Classes/GameFramework/GameUserSettings.h b/Engine/Source/Runtime/Engine/Classes/GameFramework/GameUserSettings.h
--- a/Engine/Source/Runtime/Engine/Classes/GameFramework/GameUserSettings.h
+++ b/Engine/Source/Runtime/Engine/Classes/GameFramework/GameUserSettings.h
@@ -562,7 +562,7 @@ protected:
 
 	void UpdateResolutionQuality();
 
-	void EnableHDRDisplayOutputInternal(bool bEnable, int32 DisplayNits, bool FromUserSettings);
+	void EnableHDRDisplayOutputInternal(bool bEnable, int32 DisplayNits, bool bFromUserSettings);
 
 private:
 
diff --git a/Engine/Source/Runtime/Engine/Private/GameUserSettings.cpp b/Engine/Source/Runtime/Engine/Private/GameUserSettings.cpp
--- a/Engine/Source/Runtime/Engine/Private/GameUserSettings.cpp
+++ b/Engine/Source/Runtime/Engine/Private/GameUserSettings.cpp
@@ -942,7 +942,7 @@ void UGameUserSettings::EnableHDRDisplayOutput(bool bEnable, int32 DisplayNits /
 	EnableHDRDisplayOutputInternal(bEnable, DisplayNits, false);
 }
 
-void UGameUserSettings::EnableHDRDisplayOutputInternal(bool bEnable, int32 DisplayNits, bool FromUserSettings)
+void UGameUserSettings::EnableHDRDisplayOutputInternal(bool bEnable, int32 DisplayNits, bool bFromUserSettings)
 {
 	static IConsoleVariable* CVarHDROutputEnabled = IConsoleManager::Get().FindConsoleVariable(TEXT("r.HDR.EnableHDROutput"));
 
@@ -970,13 +970,13 @@ void UGameUserSettings::EnableHDRDisplayOutputInternal(bool bEnable, int32 Displ
 				RequestUIUpdate();
 			}
 #endif
-			CVarHDROutputEnabled->Set(1, FromUserSettings ? ECVF_SetByGameSetting : ECVF_SetByCode);
+			CVarHDROutputEnabled->Set(1, bFromUserSettings ? ECVF_SetByGameSetting : ECVF_SetByCode);
 		}
 
 		// Always test this branch as can be used to flush errors
 		if (!bEnable)
 		{
-			CVarHDROutputEnabled->Set(0, FromUserSettings ? ECVF_SetByGameSetting : ECVF_SetByCode);
+			CVarHDROutputEnabled->Set(0, bFromUserSettings ? ECVF_SetByGameSetting : ECVF_SetByCode);
 		}
 
 		// Update final requested state for saved config
