Added new ObjectIsA method (blueprint pure).

#UE-42897
#3376

diff --git a/Engine/Source/Runtime/Engine/Classes/Kismet/GameplayStatics.h b/Engine/Source/Runtime/Engine/Classes/Kismet/GameplayStatics.h
--- a/Engine/Source/Runtime/Engine/Classes/Kismet/GameplayStatics.h
+++ b/Engine/Source/Runtime/Engine/Classes/Kismet/GameplayStatics.h
@@ -335,6 +335,13 @@ class ENGINE_API UGameplayStatics : public UBlueprintFunctionLibrary
 	UFUNCTION(BlueprintPure, meta=(DisplayName = "Get Class", DeterminesOutputType = "Object"), Category="Utilities")
 	static class UClass *GetObjectClass(const UObject *Object);
 
+	/**
+	 * Returns whether or not the object passed in is of (or inherits from) the class type.
+	 * @return True if the object is of (or inherits from) the class type.
+	 */
+	UFUNCTION(BlueprintPure, meta = (DisplayName = "IsA"), Category = "Utilities")
+	static bool ObjectIsA(const UObject* Object, TSubclassOf<UObject> ObjectClass);
+
 	/**
 	 * Gets the current global time dilation.
 	 * @return Current time dilation.
diff --git a/Engine/Source/Runtime/Engine/Private/GameplayStatics.cpp b/Engine/Source/Runtime/Engine/Private/GameplayStatics.cpp
--- a/Engine/Source/Runtime/Engine/Private/GameplayStatics.cpp
+++ b/Engine/Source/Runtime/Engine/Private/GameplayStatics.cpp
@@ -471,6 +471,11 @@ class UClass* UGameplayStatics::GetObjectClass(const UObject* Object)
 	return Object ? Object->GetClass() : nullptr;
 }
 
+bool UGameplayStatics::ObjectIsA(const UObject* Object, TSubclassOf<UObject> ObjectClass)
+{
+	return (Object && ObjectClass) ? Object->IsA(ObjectClass) : false;
+}
+
 float UGameplayStatics::GetGlobalTimeDilation(const UObject* WorldContextObject)
 {
 	UWorld* const World = GEngine->GetWorldFromContextObject(WorldContextObject, EGetWorldErrorMode::LogAndReturnNull);
