Do not mark texture uasset as changed when editing "Editor Show Final Encode" or "Editor Defer Compression" properties.

diff --git a/Engine/Source/Runtime/Engine/Classes/Engine/Texture.h b/Engine/Source/Runtime/Engine/Classes/Engine/Texture.h
--- a/Engine/Source/Runtime/Engine/Classes/Engine/Texture.h
+++ b/Engine/Source/Runtime/Engine/Classes/Engine/Texture.h
@@ -1155,11 +1155,11 @@ public:
 	uint32 CompressionNone:1;
 	
 	/** If enabled, compress with Final quality during this Editor session. */
-	UPROPERTY(EditAnywhere, BlueprintReadWrite, Transient, Category=Compression, meta=(NoResetToDefault), meta=(DisplayName="Editor Show Final Encode"))
+	UPROPERTY(EditAnywhere, BlueprintReadWrite, Transient, SkipSerialization, Category=Compression, meta=(NoResetToDefault), meta=(DisplayName="Editor Show Final Encode"))
 	uint32 CompressFinal:1;
 
 	/** If enabled, defer compression of the texture until save or manually compressed in the texture editor. */
-	UPROPERTY(EditAnywhere, BlueprintReadWrite, Transient, Category=Compression, meta=(NoResetToDefault), meta=(DisplayName="Editor Defer Compression"))
+	UPROPERTY(EditAnywhere, BlueprintReadWrite, Transient, SkipSerialization, Category=Compression, meta=(NoResetToDefault), meta=(DisplayName="Editor Defer Compression"))
 	uint32 DeferCompression:1;
 	
 	/** How aggressively should any relevant lossy compression be applied. For compressors that support EncodeSpeed (i.e. Oodle), this is only
