Fix Exposure Bias functionality in the texture editor.

diff --git a/Engine/Source/Editor/TextureEditor/Private/Models/TextureEditorViewportClient.cpp b/Engine/Source/Editor/TextureEditor/Private/Models/TextureEditorViewportClient.cpp
--- a/Engine/Source/Editor/TextureEditor/Private/Models/TextureEditorViewportClient.cpp
+++ b/Engine/Source/Editor/TextureEditor/Private/Models/TextureEditorViewportClient.cpp
@@ -177,7 +177,7 @@ void FTextureEditorViewportClient::Draw(FViewport* Viewport, FCanvas* Canvas)
 		}
 	}
 
-	float Exposure = FMath::Pow(2.0f, (float)TextureEditorPinned->GetExposureBias());
+	float Exposure = Texture->GetRunningPlatformData() && IsHDR((*Texture->GetRunningPlatformData())->PixelFormat) ? FMath::Pow(2.0f, (float)TextureEditorPinned->GetExposureBias()) : 1.0f;
 
 	if ( Texture->GetResource() != nullptr )
 	{
diff --git a/Engine/Source/Editor/TextureEditor/Private/TextureEditorToolkit.cpp b/Engine/Source/Editor/TextureEditor/Private/TextureEditorToolkit.cpp
--- a/Engine/Source/Editor/TextureEditor/Private/TextureEditorToolkit.cpp
+++ b/Engine/Source/Editor/TextureEditor/Private/TextureEditorToolkit.cpp
@@ -2295,7 +2295,7 @@ bool FTextureEditorToolkit::HandleTextureBorderActionIsChecked( ) const
 
 EVisibility FTextureEditorToolkit::HandleExposureBiasWidgetVisibility() const
 {
-	if ((Texture != nullptr) && (Texture->CompressionSettings == TC_HDR || Texture->CompressionSettings == TC_HDR_Compressed))
+	if (Texture && Texture->GetRunningPlatformData() && IsHDR((*Texture->GetRunningPlatformData())->PixelFormat))
 	{
 		return EVisibility::Visible;
 	}
diff --git a/Engine/Source/Runtime/Core/Public/PixelFormat.h b/Engine/Source/Runtime/Core/Public/PixelFormat.h
--- a/Engine/Source/Runtime/Core/Public/PixelFormat.h
+++ b/Engine/Source/Runtime/Core/Public/PixelFormat.h
@@ -205,3 +205,8 @@ ENUM_CLASS_FLAGS(EPixelFormatChannelFlags);
 // should be updated to take an EPixelFormat instead, but in the interim this allows fixing
 // type conversion warnings
 #define UE_PIXELFORMAT_TO_UINT8(argument) static_cast<uint8>(argument)
+
+FORCEINLINE bool IsHDR(EPixelFormat PixelFormat)
+{
+	return PixelFormat == PF_FloatRGBA || PixelFormat == PF_BC6H || PixelFormat == PF_R16F || PixelFormat == PF_R32_FLOAT || PixelFormat == PF_A32B32G32R32F;
+}
diff --git a/Engine/Source/Runtime/Engine/Private/Rendering/StreamableTextureResource.cpp b/Engine/Source/Runtime/Engine/Private/Rendering/StreamableTextureResource.cpp
--- a/Engine/Source/Runtime/Engine/Private/Rendering/StreamableTextureResource.cpp
+++ b/Engine/Source/Runtime/Engine/Private/Rendering/StreamableTextureResource.cpp
@@ -94,7 +94,7 @@ FStreamableTextureResource::FStreamableTextureResource(UTexture* InOwner, const
 	, PixelFormat(InPlatformData->PixelFormat)
 {
 	// HDR images are stored in linear but still require gamma correction to display correctly.
-	bIgnoreGammaConversions = !InOwner->SRGB && ! UE::TextureDefines::IsHDR(InOwner->CompressionSettings);
+	bIgnoreGammaConversions = !InOwner->SRGB && !IsHDR(PixelFormat);
 	bSRGB = InOwner->SRGB;
 	bGreyScaleFormat = (PixelFormat == PF_G8) || (PixelFormat == PF_BC4);
 
