World Partition: include actor label when dumping actor descriptor to string.

diff --git a/Engine/Source/Runtime/Engine/Private/WorldPartition/WorldPartitionActorDesc.cpp b/Engine/Source/Runtime/Engine/Private/WorldPartition/WorldPartitionActorDesc.cpp
--- a/Engine/Source/Runtime/Engine/Private/WorldPartition/WorldPartitionActorDesc.cpp
+++ b/Engine/Source/Runtime/Engine/Private/WorldPartition/WorldPartitionActorDesc.cpp
@@ -240,11 +240,12 @@ FString FWorldPartitionActorDesc::ToString() const
 	};
 
 	return FString::Printf(
-		TEXT("Guid:%s BaseClass:%s NativeClass:%s Name:%s SpatiallyLoaded:%s Bounds:%s RuntimeGrid:%s EditorOnly:%s LevelBoundsRelevant:%s HLODRelevant:%s FolderPath:%s FolderGuid:%s Parent:%s"), 
+		TEXT("Guid:%s BaseClass:%s NativeClass:%s Name:%s Label:%s SpatiallyLoaded:%s Bounds:%s RuntimeGrid:%s EditorOnly:%s LevelBoundsRelevant:%s HLODRelevant:%s FolderPath:%s FolderGuid:%s Parent:%s"), 
 		*Guid.ToString(), 
 		*BaseClass.ToString(), 
 		*NativeClass.ToString(), 
 		*GetActorName().ToString(),
+		*GetActorLabel().ToString(),
 		GetBoolStr(bIsSpatiallyLoaded),
 		*GetBounds().ToString(),
 		*RuntimeGrid.ToString(),
