World Partition: filter out editor-only references before building actor clusters during streaming generation.

diff --git a/Engine/Source/Runtime/Engine/Private/WorldPartition/WorldPartitionActorDescView.cpp b/Engine/Source/Runtime/Engine/Private/WorldPartition/WorldPartitionActorDescView.cpp
--- a/Engine/Source/Runtime/Engine/Private/WorldPartition/WorldPartitionActorDescView.cpp
+++ b/Engine/Source/Runtime/Engine/Private/WorldPartition/WorldPartitionActorDescView.cpp
@@ -132,7 +132,7 @@ FBox FWorldPartitionActorDescView::GetBounds() const
 
 const TArray<FGuid>& FWorldPartitionActorDescView::GetReferences() const
 {
-	return ActorDesc->GetReferences();
+	return RuntimeReferences.IsSet() ? RuntimeReferences.GetValue() : ActorDesc->GetReferences();
 }
 
 FString FWorldPartitionActorDescView::ToString() const
@@ -165,7 +165,7 @@ void FWorldPartitionActorDescView::SetForcedNonSpatiallyLoaded()
 	if (!bIsForcedNonSpatiallyLoaded)
 	{
 		bIsForcedNonSpatiallyLoaded = true;
-		UE_LOG(LogWorldPartition, Verbose, TEXT("Actor '%s' forced to be non-spatially loaded"), *GetActorLabel().ToString());
+		UE_LOG(LogWorldPartition, Verbose, TEXT("Actor '%s' forced to be non-spatially loaded"), *GetActorLabelOrName().ToString());
 	}
 }
 
@@ -179,7 +179,7 @@ void FWorldPartitionActorDescView::SetInvalidDataLayers()
 	if (!bInvalidDataLayers)
 	{
 		bInvalidDataLayers = true;
-		UE_LOG(LogWorldPartition, Verbose, TEXT("Actor '%s' data layers invalidated"), *GetActorLabel().ToString());
+		UE_LOG(LogWorldPartition, Verbose, TEXT("Actor '%s' data layers invalidated"), *GetActorLabelOrName().ToString());
 	}
 }
 
@@ -188,4 +188,22 @@ bool FWorldPartitionActorDescView::IsResaveNeeded() const
 	return ActorDesc->IsResaveNeeded();
 }
 
+void FWorldPartitionActorDescView::ResolveRuntimeReferences(const UActorDescContainer* InContainer)
+{
+	RuntimeReferences.Emplace();
+	for (const FGuid& ReferenceGuid : ActorDesc->GetReferences())
+	{
+		if (const FWorldPartitionActorDesc* ReferenceDesc = InContainer->GetActorDesc(ReferenceGuid))
+		{
+			if (ReferenceDesc->IsLoaded() ? !ReferenceDesc->GetActor()->IsEditorOnly() : !ReferenceDesc->GetActorIsEditorOnly())
+			{
+				RuntimeReferences->Add(ReferenceGuid);
+			}
+			else
+			{
+				UE_LOG(LogWorldPartition, Verbose, TEXT("Actor reference '%s' from '%s' filtered-out (editor-only)"), *ReferenceDesc->GetActorLabelOrName().ToString(), *GetActorLabelOrName().ToString());
+			}
+		}
+	}
+}
 #endif
\ No newline at end of file
diff --git a/Engine/Source/Runtime/Engine/Private/WorldPartition/WorldPartitionStreamingGeneration.cpp b/Engine/Source/Runtime/Engine/Private/WorldPartition/WorldPartitionStreamingGeneration.cpp
--- a/Engine/Source/Runtime/Engine/Private/WorldPartition/WorldPartitionStreamingGeneration.cpp
+++ b/Engine/Source/Runtime/Engine/Private/WorldPartition/WorldPartitionStreamingGeneration.cpp
@@ -63,6 +63,8 @@ class FWorldPartitionStreamingGenerator
 				ActorDescView.ResolveRuntimeDataLayers(InContainer);
 			}
 
+			ActorDescView.ResolveRuntimeReferences(InContainer);
+
 			OutActorDescViewMap.Emplace(ActorGuid, ActorDescView);
 		};
 		
@@ -368,7 +370,7 @@ class FWorldPartitionStreamingGenerator
 							}
 							// Do not increment NbErrorsDetected since it won't be fixed and thus will always occur
 						}
-					}				
+					}
 				}		
 
 				NbValidationPasses++;
diff --git a/Engine/Source/Runtime/Engine/Public/WorldPartition/WorldPartitionActorDescView.h b/Engine/Source/Runtime/Engine/Public/WorldPartition/WorldPartitionActorDescView.h
--- a/Engine/Source/Runtime/Engine/Public/WorldPartition/WorldPartitionActorDescView.h
+++ b/Engine/Source/Runtime/Engine/Public/WorldPartition/WorldPartitionActorDescView.h
@@ -35,7 +35,6 @@ public:
 	FName GetHLODLayer() const;
 	const TArray<FName>& GetDataLayers() const;
 	const TArray<FName>& GetRuntimeDataLayers() const;
-	void ResolveRuntimeDataLayers(const UActorDescContainer* InContainer = nullptr);
 	FName GetActorPackage() const;
 	FName GetActorPath() const;
 	FName GetActorLabel() const;
@@ -58,6 +57,9 @@ public:
 
 	bool IsResaveNeeded() const;
 
+	void ResolveRuntimeDataLayers(const UActorDescContainer* InContainer = nullptr);
+	void ResolveRuntimeReferences(const UActorDescContainer* InContainer);
+
 	bool operator==(const FWorldPartitionActorDescView& Other) const
 	{
 		return GetGuid() == Other.GetGuid();
@@ -74,5 +76,6 @@ protected:
 	bool bInvalidDataLayers;
 	bool bInvalidRuntimeGrid;
 	TOptional<TArray<FName>> RuntimeDataLayers;
+	TOptional<TArray<FGuid>> RuntimeReferences;
 };
 #endif
\ No newline at end of file
