World Partition: don't resolve runtime data layers when creating actor desc views for unsaved actors, let the streaming generator handle this call instead.

diff --git a/Engine/Source/Runtime/Engine/Private/WorldPartition/DataLayer/DataLayerUtils.cpp b/Engine/Source/Runtime/Engine/Private/WorldPartition/DataLayer/DataLayerUtils.cpp
--- a/Engine/Source/Runtime/Engine/Private/WorldPartition/DataLayer/DataLayerUtils.cpp
+++ b/Engine/Source/Runtime/Engine/Private/WorldPartition/DataLayer/DataLayerUtils.cpp
@@ -116,6 +116,8 @@ TArray<FName> FDataLayerUtils::ResolvedDataLayerInstanceNames(const FWorldPartit
 // For performance reasons, this function assumes that InActorDesc's DataLayerInstanceNames was already resolved.
 TArray<FName> FDataLayerUtils::ResolveRuntimeDataLayerInstanceNames(const FWorldPartitionActorDesc* InActorDesc, const UActorDescContainer* InContainer, bool* bOutIsResultValid)
 {
+	check(InContainer);
+
 	bool bLocalIsSuccess = false;
 	bool& bIsSuccess = bOutIsResultValid ? *bOutIsResultValid : bLocalIsSuccess;
 	bIsSuccess = false;
@@ -127,34 +129,30 @@ TArray<FName> FDataLayerUtils::ResolveRuntimeDataLayerInstanceNames(const FWorld
 		return Result;
 	}
 
-	const UActorDescContainer* Container = InContainer ? InContainer : InActorDesc->GetContainer();
-	if (Container)
+	UWorld* World = InContainer->GetWorld();
+	const UDataLayerSubsystem* DataLayerSubsystem = World ? World->GetSubsystem<UDataLayerSubsystem>() : nullptr;
+	if (DataLayerSubsystem && DataLayerSubsystem->CanResolveDataLayers())
 	{
-		UWorld* World = Container->GetWorld();
-		const UDataLayerSubsystem* DataLayerSubsystem = World ? World->GetSubsystem<UDataLayerSubsystem>() : nullptr;
-		if (DataLayerSubsystem && DataLayerSubsystem->CanResolveDataLayers())
+		bIsSuccess = true;
+		for (FName DataLayerInstanceName : InActorDesc->GetDataLayerInstanceNames())
 		{
-			bIsSuccess = true;
-			for (FName DataLayerInstanceName : InActorDesc->GetDataLayerInstanceNames())
+			const UDataLayerInstance* DataLayerInstance = DataLayerSubsystem->GetDataLayerInstance(DataLayerInstanceName);
+			if (DataLayerInstance && DataLayerInstance->IsRuntime())
 			{
-				const UDataLayerInstance* DataLayerInstance = DataLayerSubsystem->GetDataLayerInstance(DataLayerInstanceName);
-				if (DataLayerInstance && DataLayerInstance->IsRuntime())
-				{
-					Result.Add(DataLayerInstanceName);
-				}
+				Result.Add(DataLayerInstanceName);
 			}
 		}
-		// Fallback on FWorldDataLayersActorDesc
-		else if (FWorldDataLayersActorDesc* WorldDataLayersActorDesc = FDataLayerUtils::GetWorldDataLayersActorDesc(Container))
+	}
+	// Fallback on FWorldDataLayersActorDesc
+	else if (FWorldDataLayersActorDesc* WorldDataLayersActorDesc = FDataLayerUtils::GetWorldDataLayersActorDesc(InContainer))
+	{
+		bIsSuccess = true;
+		for (FName DataLayerInstanceName : InActorDesc->GetDataLayerInstanceNames())
 		{
-			bIsSuccess = true;
-			for (FName DataLayerInstanceName : InActorDesc->GetDataLayerInstanceNames())
+			const FDataLayerInstanceDesc* DataLayerInstanceDesc = WorldDataLayersActorDesc->GetDataLayerInstanceFromInstanceName(DataLayerInstanceName);
+			if (DataLayerInstanceDesc && (DataLayerInstanceDesc->GetDataLayerType() == EDataLayerType::Runtime))
 			{
-				const FDataLayerInstanceDesc* DataLayerInstanceDesc = WorldDataLayersActorDesc->GetDataLayerInstanceFromInstanceName(DataLayerInstanceName);
-				if (DataLayerInstanceDesc && (DataLayerInstanceDesc->GetDataLayerType() == EDataLayerType::Runtime))
-				{
-					Result.Add(DataLayerInstanceName);
-				}
+				Result.Add(DataLayerInstanceName);
 			}
 		}
 	}
diff --git a/Engine/Source/Runtime/Engine/Private/WorldPartition/WorldPartitionActorDescView.cpp b/Engine/Source/Runtime/Engine/Private/WorldPartition/WorldPartitionActorDescView.cpp
--- a/Engine/Source/Runtime/Engine/Private/WorldPartition/WorldPartitionActorDescView.cpp
+++ b/Engine/Source/Runtime/Engine/Private/WorldPartition/WorldPartitionActorDescView.cpp
@@ -17,11 +17,8 @@ FWorldPartitionActorDescView::FWorldPartitionActorDescView(const FWorldPartition
 	, bIsForcedNonSpatiallyLoaded(false)
 	, bInvalidDataLayers(false)
 	, bInvalidRuntimeGrid(false)
-{
-	ResolveRuntimeDataLayers();
-}
+{}
 
-//The only case where we need to call ResolveRuntimeDataLayers and pass a Container is for the "unsaved actors" case of FWorldPartitionStreamingGenerator
 void FWorldPartitionActorDescView::ResolveRuntimeDataLayers(const UActorDescContainer* InContainer)
 {
 	bool bSuccess = true;
diff --git a/Engine/Source/Runtime/Engine/Public/WorldPartition/WorldPartitionActorDescView.h b/Engine/Source/Runtime/Engine/Public/WorldPartition/WorldPartitionActorDescView.h
--- a/Engine/Source/Runtime/Engine/Public/WorldPartition/WorldPartitionActorDescView.h
+++ b/Engine/Source/Runtime/Engine/Public/WorldPartition/WorldPartitionActorDescView.h
@@ -57,7 +57,7 @@ public:
 
 	bool IsResaveNeeded() const;
 
-	void ResolveRuntimeDataLayers(const UActorDescContainer* InContainer = nullptr);
+	void ResolveRuntimeDataLayers(const UActorDescContainer* InContainer);
 	void ResolveRuntimeReferences(const UActorDescContainer* InContainer);
 
 	bool operator==(const FWorldPartitionActorDescView& Other) const
