World Partition: added null check until we figure out a good repro.

diff --git a/Engine/Source/Runtime/Engine/Private/WorldPartition/WorldPartitionRuntimeSpatialHash.cpp b/Engine/Source/Runtime/Engine/Private/WorldPartition/WorldPartitionRuntimeSpatialHash.cpp
--- a/Engine/Source/Runtime/Engine/Private/WorldPartition/WorldPartitionRuntimeSpatialHash.cpp
+++ b/Engine/Source/Runtime/Engine/Private/WorldPartition/WorldPartitionRuntimeSpatialHash.cpp
@@ -1022,19 +1022,20 @@ bool UWorldPartitionRuntimeSpatialHash::CreateStreamingGrid(const FSpatialHashRu
 								if (ActorInstance.ContainerInstance->Container == WorldPartition)
 								{
 									// This will load the actor if it isn't already loaded
-									FWorldPartitionReference Reference(WorldPartition, ActorInstance.Actor);
-									AActor* AlwaysLoadedActor = FindObject<AActor>(nullptr, *ActorDescView.GetActorPath().ToString());
-									AlwaysLoadedActorsForPIE.Emplace(Reference, AlwaysLoadedActor);
-
-									// Handle child actors
-									AlwaysLoadedActor->ForEachComponent<UChildActorComponent>(true, [this, &Reference](UChildActorComponent* ChildActorComponent)
+									FWorldPartitionReference Reference(WorldPartition, ActorInstance.Actor);									
+									if (AActor* AlwaysLoadedActor = FindObject<AActor>(nullptr, *ActorDescView.GetActorPath().ToString()))
 									{
-										if (AActor* ChildActor = ChildActorComponent->GetChildActor())
-										{
-											AlwaysLoadedActorsForPIE.Emplace(Reference, ChildActor);
-										}
-									});
+										AlwaysLoadedActorsForPIE.Emplace(Reference, AlwaysLoadedActor);
 
+										// Handle child actors
+										AlwaysLoadedActor->ForEachComponent<UChildActorComponent>(true, [this, &Reference](UChildActorComponent* ChildActorComponent)
+										{
+											if (AActor* ChildActor = ChildActorComponent->GetChildActor())
+											{
+												AlwaysLoadedActorsForPIE.Emplace(Reference, ChildActor);
+											}
+										});
+									}
 									continue;
 								}
 							}
