WorldPartition: CVar allows ignoring replay streaming source data + add information in debug display when streaming source is from the replay

diff --git a/Engine/Source/Runtime/Engine/Private/WorldPartition/WorldPartitionReplay.cpp b/Engine/Source/Runtime/Engine/Private/WorldPartition/WorldPartitionReplay.cpp
--- a/Engine/Source/Runtime/Engine/Private/WorldPartition/WorldPartitionReplay.cpp
+++ b/Engine/Source/Runtime/Engine/Private/WorldPartition/WorldPartitionReplay.cpp
@@ -11,6 +11,12 @@
 #include "Engine/DemoNetDriver.h"
 #include "ReplaySubsystem.h"
 
+static int32 GUseReplayStreamingSources = 1;
+static FAutoConsoleVariableRef CVarUseReplayStreamingSources(
+	TEXT("wp.Runtime.UseReplayStreamingSources"),
+	GUseReplayStreamingSources,
+	TEXT("Set to 1 to use the recorded streaming sources when playing replay."));
+
 FArchive& operator<<(FArchive& Ar, FWorldPartitionReplayStreamingSource& StreamingSource)
 {
 	Ar << StreamingSource.Location;
@@ -52,10 +58,15 @@ FArchive& operator<<(FArchive& Ar, FWorldPartitionReplaySample& ReplaySample)
 		for (int32 i = 0; i < ReplaySample.StreamingSources.Num(); ++i)
 		{
 			int32 NameIndex = ReplaySample.StreamingSourceNameIndices[i];
-			if (ensure(ReplaySample.Replay->StreamingSourceNames.IsValidIndex(NameIndex)))
+			// @todo_ow: It sometimes happens at start of replay that StreamingSourceNames haven't been replicated yet. need to investigate.
+			if (ReplaySample.Replay->StreamingSourceNames.IsValidIndex(NameIndex))
 			{
 				ReplaySample.StreamingSources[i].Name = ReplaySample.Replay->StreamingSourceNames[NameIndex];
 			}
+			else
+			{
+				ReplaySample.StreamingSources[i].Name = FName(FString::Printf(TEXT("Source {%i}"), NameIndex));
+			}
 		}
 	}
 	return Ar;
@@ -151,6 +162,11 @@ void AWorldPartitionReplay::PreReplication(IRepChangedPropertyTracker& ChangedPr
 
 bool AWorldPartitionReplay::GetReplayStreamingSources(TArray<FWorldPartitionStreamingSource>& OutStreamingSources)
 {
+	if (!GUseReplayStreamingSources)
+	{
+		return false;
+	}
+
 	UWorld* World = GetWorld();
 	UDemoNetDriver* DemoNetDriver = World->GetDemoNetDriver();
 	if (DemoNetDriver && DemoNetDriver->IsPlaying())
@@ -197,8 +213,8 @@ bool AWorldPartitionReplay::GetReplayStreamingSources(TArray<FWorldPartitionStre
 					int32 SourceIndex2 = ReplaySample2.StreamingSources.IndexOfByPredicate([&](const FWorldPartitionStreamingSource& StreamingSource) { return StreamingSource.Name == ReplaySample1.StreamingSources[SourceIndex].Name; });
 					if (SourceIndex2 != INDEX_NONE)
 					{
-						const FWorldPartitionStreamingSource& Source1 = ReplaySample1.StreamingSources[SourceIndex];
-						const FWorldPartitionStreamingSource& Source2 = ReplaySample2.StreamingSources[SourceIndex2];
+						const FWorldPartitionReplayStreamingSource& Source1 = ReplaySample1.StreamingSources[SourceIndex];
+						const FWorldPartitionReplayStreamingSource& Source2 = ReplaySample2.StreamingSources[SourceIndex2];
 						
 						const float EPSILON = UE_SMALL_NUMBER;
 						const float Delta = ReplaySample2.TimeSeconds - ReplaySample1.TimeSeconds;
@@ -207,7 +223,7 @@ bool AWorldPartitionReplay::GetReplayStreamingSources(TArray<FWorldPartitionStre
 						const FVector Location = FMath::Lerp(Source1.Location, Source2.Location, LerpPercent);
 						const FQuat Rotation = FQuat::FastLerp(FQuat(Source1.Rotation), FQuat(Source2.Rotation), LerpPercent).GetNormalized();
 						const float Velocity = FMath::Lerp(Source1.Velocity, Source2.Velocity, LerpPercent);
-						OutStreamingSources.Add(FWorldPartitionStreamingSource(Source1.Name, Location, Rotation.Rotator(), Source1.TargetState, Source1.bBlockOnSlowLoading, Source1.Priority, Velocity));
+						OutStreamingSources.Add(FWorldPartitionReplayStreamingSource(Source1.Name, Location, Rotation.Rotator(), Source1.TargetState, Source1.bBlockOnSlowLoading, Source1.Priority, Velocity));
 					}
 					else
 					{
diff --git a/Engine/Source/Runtime/Engine/Private/WorldPartition/WorldPartitionSubsystem.cpp b/Engine/Source/Runtime/Engine/Private/WorldPartition/WorldPartitionSubsystem.cpp
--- a/Engine/Source/Runtime/Engine/Private/WorldPartition/WorldPartitionSubsystem.cpp
+++ b/Engine/Source/Runtime/Engine/Private/WorldPartition/WorldPartitionSubsystem.cpp
@@ -387,6 +387,10 @@ void UWorldPartitionSubsystem::Draw(UCanvas* Canvas, class APlayerController* PC
 			for (const FWorldPartitionStreamingSource& StreamingSource : StreamingSources)
 			{
 				FString StreamingSourceDisplay = StreamingSource.Name.ToString();
+				if (StreamingSource.bReplay)
+				{
+					StreamingSourceDisplay += TEXT(" (Replay)");
+				}
 				FWorldPartitionDebugHelper::DrawText(Canvas, StreamingSourceDisplay, GEngine->GetSmallFont(), StreamingSource.GetDebugColor(), Pos, &MaxTextWidth);
 			}
 			Pos = CurrentOffset + FVector2D(MaxTextWidth + 10, 0.f);
diff --git a/Engine/Source/Runtime/Engine/Public/WorldPartition/WorldPartitionReplay.h b/Engine/Source/Runtime/Engine/Public/WorldPartition/WorldPartitionReplay.h
--- a/Engine/Source/Runtime/Engine/Public/WorldPartition/WorldPartitionReplay.h
+++ b/Engine/Source/Runtime/Engine/Public/WorldPartition/WorldPartitionReplay.h
@@ -16,10 +16,16 @@ struct FWorldPartitionReplayStreamingSource : public FWorldPartitionStreamingSou
 	}
 
 	FWorldPartitionReplayStreamingSource(const FWorldPartitionStreamingSource& InStreamingSource)
-		: FWorldPartitionStreamingSource(InStreamingSource.Name, InStreamingSource.Location, InStreamingSource.Rotation, InStreamingSource.TargetState, InStreamingSource.bBlockOnSlowLoading, InStreamingSource.Priority, InStreamingSource.Velocity)
+		: FWorldPartitionReplayStreamingSource(InStreamingSource.Name, InStreamingSource.Location, InStreamingSource.Rotation, InStreamingSource.TargetState, InStreamingSource.bBlockOnSlowLoading, InStreamingSource.Priority, InStreamingSource.Velocity)
 	{
 	}
 
+	FWorldPartitionReplayStreamingSource(FName InName, const FVector& InLocation, const FRotator& InRotation, EStreamingSourceTargetState InTargetState, bool bInBlockOnSlowLoading, EStreamingSourcePriority InPriority, float InVelocity)
+		: FWorldPartitionStreamingSource(InName, InLocation, InRotation, InTargetState, bInBlockOnSlowLoading, InPriority, InVelocity)
+	{
+		bReplay = true;
+	}
+
 	friend FArchive& operator<<(FArchive& Ar, FWorldPartitionReplayStreamingSource& StreamingSource);
 };
 
diff --git a/Engine/Source/Runtime/Engine/Public/WorldPartition/WorldPartitionStreamingSource.h b/Engine/Source/Runtime/Engine/Public/WorldPartition/WorldPartitionStreamingSource.h
--- a/Engine/Source/Runtime/Engine/Public/WorldPartition/WorldPartitionStreamingSource.h
+++ b/Engine/Source/Runtime/Engine/Public/WorldPartition/WorldPartitionStreamingSource.h
@@ -324,6 +324,7 @@ struct ENGINE_API FWorldPartitionStreamingSource
 		, Velocity(0.f)
 		, DebugColor(ForceInit)
 		, TargetGrid(NAME_None)
+		, bReplay(false)
 	{}
 
 	FWorldPartitionStreamingSource(FName InName, const FVector& InLocation, const FRotator& InRotation, EStreamingSourceTargetState InTargetState, bool bInBlockOnSlowLoading, EStreamingSourcePriority InPriority, float InVelocity = 0.f)
@@ -336,6 +337,7 @@ struct ENGINE_API FWorldPartitionStreamingSource
 		, Velocity(InVelocity)
 		, DebugColor(ForceInit)
 		, TargetGrid(NAME_None)
+		, bReplay(false)
 	{}
 
 	FColor GetDebugColor() const
@@ -381,6 +383,9 @@ struct ENGINE_API FWorldPartitionStreamingSource
 	/** Source internal shapes. When none are provided, a sphere is automatically used. It's radius is equal to grid's loading range and center equals source's location. */
 	TArray<FStreamingSourceShape> Shapes;
 
+	/** If true, this streaming source is from a replay recording */
+	bool bReplay;
+
 	/** Returns a box encapsulating all shapes. */
 	FORCEINLINE FBox CalcBounds(float InGridLoadingRange, FName InGridName, const FSoftObjectPath& InGridHLODLayer, bool bCalcIn2D = false) const
 	{
