[KillPendingKill] Release pointers to components in FParticleSystemWorldManager when they are marked as garbage.

diff --git a/Engine/Source/Runtime/Engine/Private/Particles/ParticleSystemManager.cpp b/Engine/Source/Runtime/Engine/Private/Particles/ParticleSystemManager.cpp
--- a/Engine/Source/Runtime/Engine/Private/Particles/ParticleSystemManager.cpp
+++ b/Engine/Source/Runtime/Engine/Private/Particles/ParticleSystemManager.cpp
@@ -256,18 +256,38 @@ void FParticleSystemWorldManager::AddReferencedObjects(FReferenceCollector& Coll
 
 		for (int32 PSCIndex = 0; PSCIndex < ManagedPSCs.Num(); ++PSCIndex)
 		{
-			//UE_LOG(LogParticles, Warning, TEXT("| Add Ref %d - 0x%p |"), PSCIndex, ManagedPSCs[PSCIndex]);
-			Collector.AddReferencedObject(ManagedPSCs[PSCIndex]);
-			if (PSCTickData[PSCIndex].PrereqComponent)
+			// If a managed component is streamed out or destroyed, drop references to it and its prerequisite			
+			if (IsValid(ManagedPSCs[PSCIndex]))
+			{
+				Collector.AddReferencedObject(ManagedPSCs[PSCIndex]);
+			}
+			else
+			{
+				ManagedPSCs[PSCIndex] = nullptr; // Null entries will be cleaned up after GC
+			}
+
+			// If prerequisite has been marked for deletion forget it
+			if (IsValid(PSCTickData[PSCIndex].PrereqComponent))
 			{
 				Collector.AddReferencedObject(PSCTickData[PSCIndex].PrereqComponent);
 			}
+			else
+			{
+				PSCTickData[PSCIndex].PrereqComponent = nullptr; 
+			}
 		}
 
 		for (int32 PSCIndex = 0; PSCIndex < PendingRegisterPSCs.Num(); ++PSCIndex)
 		{
-			//UE_LOG(LogParticles, Warning, TEXT("| Add Pending PSC Ref %d - 0x%p |"), PSCIndex, PendingRegisterPSCs[PSCIndex]);
-			Collector.AddReferencedObject(PendingRegisterPSCs[PSCIndex]);
+			if (IsValid(PendingRegisterPSCs[PSCIndex]))
+			{
+				//UE_LOG(LogParticles, Warning, TEXT("| Add Pending PSC Ref %d - 0x%p |"), PSCIndex, PendingRegisterPSCs[PSCIndex]);
+				Collector.AddReferencedObject(PendingRegisterPSCs[PSCIndex]);
+			}
+			else
+			{
+				PendingRegisterPSCs[PSCIndex] = nullptr; // Array will be emptied next time we handled pending entries
+			}
 		}
 	}
 	else
@@ -411,7 +431,7 @@ void FParticleSystemWorldManager::UnregisterComponent(UParticleSystemComponent*
 
 void FParticleSystemWorldManager::AddPSC(UParticleSystemComponent* PSC)
 {
-	if (PSC)
+	if (IsValid(PSC))  // Don't add PSC if it has been marked for deletion
 	{
 		int32 Handle = ManagedPSCs.Add(PSC);
 		PSCTickData.AddDefaulted();
