Change ULocalPlayer::IsPrimaryPlayer to use GetLocalPlayerIndex so that the behavior is overrideable.

diff --git a/Engine/Source/Runtime/Engine/Private/LocalPlayer.cpp b/Engine/Source/Runtime/Engine/Private/LocalPlayer.cpp
--- a/Engine/Source/Runtime/Engine/Private/LocalPlayer.cpp
+++ b/Engine/Source/Runtime/Engine/Private/LocalPlayer.cpp
@@ -1721,12 +1721,7 @@ void ULocalPlayer::AddReferencedObjects(UObject* InThis, FReferenceCollector& Co
 
 bool ULocalPlayer::IsPrimaryPlayer() const
 {
-	if (UWorld* World = GetWorld())
-	{
-		ULocalPlayer* const PrimaryPlayer = GetOuterUEngine()->GetFirstGamePlayer(World);
-		return (this == PrimaryPlayer);
-	}
-	return false;
+	return GetLocalPlayerIndex() == 0;
 }
 
 void ULocalPlayer::CleanupViewState()
