World Partition: allow spatially loaded actors to reference non-spatially loaded actors when the target actors is not in any data layers.

#rnx

diff --git a/Engine/Source/Runtime/Engine/Private/WorldPartition/WorldPartitionActorCluster.cpp b/Engine/Source/Runtime/Engine/Private/WorldPartition/WorldPartitionActorCluster.cpp
--- a/Engine/Source/Runtime/Engine/Private/WorldPartition/WorldPartitionActorCluster.cpp
+++ b/Engine/Source/Runtime/Engine/Private/WorldPartition/WorldPartitionActorCluster.cpp
@@ -131,6 +131,12 @@ void CreateActorCluster(const FWorldPartitionActorDescView& ActorDescView, TMap<
 	{
 		if (const FWorldPartitionActorDescView* ReferenceActorDescView = ActorDescViewMap.Find(ReferenceGuid))
 		{
+			// Ignore references from spatially loaded actors to non-spatially loaded ones (at this point, they are in the same data layers).
+			if (ActorDescView.GetIsSpatiallyLoaded() && !ReferenceActorDescView->GetIsSpatiallyLoaded())
+			{
+				continue;
+			}
+
 			FActorCluster* ReferenceCluster = ActorToActorCluster.FindRef(ReferenceGuid);
 			if (ReferenceCluster)
 			{
diff --git a/Engine/Source/Runtime/Engine/Private/WorldPartition/WorldPartitionStreamingGeneration.cpp b/Engine/Source/Runtime/Engine/Private/WorldPartition/WorldPartitionStreamingGeneration.cpp
--- a/Engine/Source/Runtime/Engine/Private/WorldPartition/WorldPartitionStreamingGeneration.cpp
+++ b/Engine/Source/Runtime/Engine/Private/WorldPartition/WorldPartitionStreamingGeneration.cpp
@@ -183,19 +183,41 @@ class FWorldPartitionStreamingGenerator
 	void ValidateActorDescriptorViews()
 	{
 		// Validate data layers
-		auto IsReferenceDataLayersValid = [](const FWorldPartitionActorDescView& ActorDescView, const FWorldPartitionActorDescView& ReferenceActorDescView)
+		auto IsReferenceGridPlacementValid = [](const FWorldPartitionActorDescView& RefererActorDescView, const FWorldPartitionActorDescView& ReferenceActorDescView)
 		{
-			if (ActorDescView.GetRuntimeDataLayers().Num() == ReferenceActorDescView.GetRuntimeDataLayers().Num())
+			const bool bIsActorDescSpatiallyLoaded = RefererActorDescView.GetIsSpatiallyLoaded();
+			const bool bIsActorDescRefSpatiallyLoaded = ReferenceActorDescView.GetIsSpatiallyLoaded();
+
+			// The only case we support right now is spatially loaded actors referencing non-spatially loaded actors, when target is not in data layers.
+			// For this to work with data layers, we need to implement dependency logic support in the content cooker splitter.
+			if (bIsActorDescSpatiallyLoaded && !bIsActorDescRefSpatiallyLoaded && ReferenceActorDescView.GetDataLayers().IsEmpty())
 			{
-				const TSet<FName> ActorDescDataLayers(ActorDescView.GetRuntimeDataLayers());
+				return true;
+			}
+
+			return bIsActorDescSpatiallyLoaded == bIsActorDescRefSpatiallyLoaded;
+		};
+
+		// Validate grid placement
+		auto IsReferenceDataLayersValid = [](const FWorldPartitionActorDescView& RefererActorDescView, const FWorldPartitionActorDescView& ReferenceActorDescView)
+		{
+			if (RefererActorDescView.GetRuntimeDataLayers().Num() == ReferenceActorDescView.GetRuntimeDataLayers().Num())
+			{
+				const TSet<FName> RefererActorDescDataLayers(RefererActorDescView.GetRuntimeDataLayers());
 				const TSet<FName> ReferenceActorDescDataLayers(ReferenceActorDescView.GetRuntimeDataLayers());
 
-				return ActorDescDataLayers.Includes(ReferenceActorDescDataLayers);
+				return RefererActorDescDataLayers.Includes(ReferenceActorDescDataLayers);
 			}
 
 			return false;
 		};
 
+		// Validate runtime grid
+		auto IsReferenceRuntimeGridValid = [](const FWorldPartitionActorDescView& RefererActorDescView, const FWorldPartitionActorDescView& ReferenceActorDescView)
+		{
+			return RefererActorDescView.GetRuntimeGrid() == ReferenceActorDescView.GetRuntimeGrid();
+		};
+
 		for (auto It = ContainerDescriptorsMap.CreateIterator(); It; ++It)
 		{
 			const FActorContainerID& ContainerID = It.Key();
@@ -313,10 +335,7 @@ class FWorldPartitionStreamingGenerator
 						if (ReferenceActorDescView)
 						{
 							// Validate grid placement
-							const bool bIsActorDescSpatiallyLoaded = RefererActorDescView->GetIsSpatiallyLoaded();
-							const bool bIsActorDescRefSpatiallyLoaded = ReferenceActorDescView->GetIsSpatiallyLoaded();
-
-							if (bIsActorDescSpatiallyLoaded != bIsActorDescRefSpatiallyLoaded)
+							if (!IsReferenceGridPlacementValid(*RefererActorDescView, *ReferenceActorDescView))
 							{
 								if (!NbValidationPasses)
 								{
@@ -346,7 +365,7 @@ class FWorldPartitionStreamingGenerator
 								}
 							}
 
-							if (ReferenceActorDescView->GetRuntimeGrid() != RefererActorDescView->GetRuntimeGrid())
+							if (!IsReferenceRuntimeGridValid(*RefererActorDescView, *ReferenceActorDescView))
 							{
 								if (!NbValidationPasses)
 								{
