Fix for editor crash when removing the last texture source. FTextureSource::Init tries to calculate the size of the layers, which asserts because the size of TSF_Unknown is not determinable.

diff --git a/Engine/Source/Runtime/Engine/Classes/Engine/Texture2DArray.cpp b/Engine/Source/Runtime/Engine/Classes/Engine/Texture2DArray.cpp
--- a/Engine/Source/Runtime/Engine/Classes/Engine/Texture2DArray.cpp
+++ b/Engine/Source/Runtime/Engine/Classes/Engine/Texture2DArray.cpp
@@ -239,7 +239,9 @@ ENGINE_API void UTexture2DArray::InvadiateTextureSource()
 		PlatformData = NULL;
 	}
 
-	Source.Init(0, 0, 0, 0, TSF_Invalid, nullptr);
+	Source = FTextureSource();
+	Source.SetOwner(this);
+
 	UpdateResource();
 }
 #endif
diff --git a/Engine/Source/Runtime/Engine/Private/Texture.cpp b/Engine/Source/Runtime/Engine/Private/Texture.cpp
--- a/Engine/Source/Runtime/Engine/Private/Texture.cpp
+++ b/Engine/Source/Runtime/Engine/Private/Texture.cpp
@@ -376,101 +376,104 @@ void UTexture::ValidateSettingsAfterImportOrEdit(bool * pRequiresNotifyMaterials
 		
 #if WITH_EDITORONLY_DATA
 
-	if ( MipGenSettings == TMGS_LeaveExistingMips && PowerOfTwoMode != ETexturePowerOfTwoSetting::None )
+	if (Source.IsValid()) // we can have an empty source if the last source in a texture2d array is removed via the editor.
 	{
-		// power of 2 pads not allowed with LeaveExistingMips
-		UE_LOG(LogTexture, Display, TEXT("Power of 2 padding cannot be used with LeaveExistingMips, disabled. (%s)"), *GetName());
+		if ( MipGenSettings == TMGS_LeaveExistingMips && PowerOfTwoMode != ETexturePowerOfTwoSetting::None )
+		{
+			// power of 2 pads not allowed with LeaveExistingMips
+			UE_LOG(LogTexture, Display, TEXT("Power of 2 padding cannot be used with LeaveExistingMips, disabled. (%s)"), *GetName());
 
-		PowerOfTwoMode = ETexturePowerOfTwoSetting::None;
-	}
+			PowerOfTwoMode = ETexturePowerOfTwoSetting::None;
+		}
 
-	// IsPowerOfTwo only checks XY :
-	bool bIsPowerOfTwo = Source.IsPowerOfTwo();
-	if ( ! FMath::IsPowerOfTwo(Source.GetVolumeSizeZ()) )
-	{
-		bIsPowerOfTwo = false;
-	}
-	if ( PowerOfTwoMode != ETexturePowerOfTwoSetting::None )
-	{
-		bIsPowerOfTwo = true;
-	}
+		// IsPowerOfTwo only checks XY :
+		bool bIsPowerOfTwo = Source.IsPowerOfTwo();
+		if ( ! FMath::IsPowerOfTwo(Source.GetVolumeSizeZ()) )
+		{
+			bIsPowerOfTwo = false;
+		}
+		if ( PowerOfTwoMode != ETexturePowerOfTwoSetting::None )
+		{
+			bIsPowerOfTwo = true;
+		}
 
-	if ( ! bIsPowerOfTwo )
-	{
-		// streaming only supports power of 2 mips
-		// due to failure to compensate for the GPU row pitch
-		// it only works for mips that naturally have the required 256 pitch
-		// so mip levels >= 256 and power of 2 only
-		// (this used to be in Texture2D.cpp)
-		// this could be fixed and then streaming nonpow2 could be allowed
-		// see WarnRequiresTightPackedMip
-		NeverStream = true;	
-	}
+		if ( ! bIsPowerOfTwo )
+		{
+			// streaming only supports power of 2 mips
+			// due to failure to compensate for the GPU row pitch
+			// it only works for mips that naturally have the required 256 pitch
+			// so mip levels >= 256 and power of 2 only
+			// (this used to be in Texture2D.cpp)
+			// this could be fixed and then streaming nonpow2 could be allowed
+			// see WarnRequiresTightPackedMip
+			NeverStream = true;	
+		}
 	
-	int32 MaxDimension = FMath::Max( Source.GetSizeX() , Source.GetSizeY() );
-	bool bLargeTextureMustBeVT = MaxDimension > GetMaximumDimensionOfNonVT();
+		int32 MaxDimension = FMath::Max( Source.GetSizeX() , Source.GetSizeY() );
+		bool bLargeTextureMustBeVT = MaxDimension > GetMaximumDimensionOfNonVT();
 
-	if ( bLargeTextureMustBeVT && ! VirtualTextureStreaming && MaxTextureSize == 0 )
-	{
-		static const auto CVarVirtualTexturesEnabled = IConsoleManager::Get().FindTConsoleVariableDataInt(TEXT("r.VirtualTextures"));
-		check( CVarVirtualTexturesEnabled != nullptr );
-
-		if ( CVarVirtualTexturesEnabled->GetValueOnAnyThread() )
+		if ( bLargeTextureMustBeVT && ! VirtualTextureStreaming && MaxTextureSize == 0 )
 		{
-			if ( GetTextureClass() == ETextureClass::TwoD )
+			static const auto CVarVirtualTexturesEnabled = IConsoleManager::Get().FindTConsoleVariableDataInt(TEXT("r.VirtualTextures"));
+			check( CVarVirtualTexturesEnabled != nullptr );
+
+			if ( CVarVirtualTexturesEnabled->GetValueOnAnyThread() )
 			{
-				UE_LOG(LogTexture, Display, TEXT("Large Texture %s Dimension=%d changed to VT"), *GetName(),MaxDimension);
-				VirtualTextureStreaming = true;
-				bRequiresNotifyMaterials = true;
+				if ( GetTextureClass() == ETextureClass::TwoD )
+				{
+					UE_LOG(LogTexture, Display, TEXT("Large Texture %s Dimension=%d changed to VT"), *GetName(),MaxDimension);
+					VirtualTextureStreaming = true;
+					bRequiresNotifyMaterials = true;
+				}
+				else
+				{
+					UE_LOG(LogTexture, Warning, TEXT("Large Texture %s Dimension=%d needs to be VT but is not 2d, changing MaxTextureSize"), *GetName(),MaxDimension);
+				
+					// GetMaximumDimension is the max size for this texture type on the current RHI
+					MaxTextureSize = GetMaximumDimension();
+				}
 			}
 			else
 			{
-				UE_LOG(LogTexture, Warning, TEXT("Large Texture %s Dimension=%d needs to be VT but is not 2d, changing MaxTextureSize"), *GetName(),MaxDimension);
-				
+				UE_LOG(LogTexture, Warning, TEXT("Large Texture %s Dimension=%d must be VT but VirtualTextures are disabled, changing MaxTextureSize"), *GetName(),MaxDimension);
+
 				// GetMaximumDimension is the max size for this texture type on the current RHI
 				MaxTextureSize = GetMaximumDimension();
 			}
 		}
-		else
-		{
-			UE_LOG(LogTexture, Warning, TEXT("Large Texture %s Dimension=%d must be VT but VirtualTextures are disabled, changing MaxTextureSize"), *GetName(),MaxDimension);
-
-			// GetMaximumDimension is the max size for this texture type on the current RHI
-			MaxTextureSize = GetMaximumDimension();
-		}
-	}
 	
-	if (VirtualTextureStreaming)
-	{
-		if (!bIsPowerOfTwo)
+		if (VirtualTextureStreaming)
 		{
-			if ( bLargeTextureMustBeVT )
-			{
-				UE_LOG(LogTexture, Warning, TEXT("Large VT \"%s\", must be padded to power-of-2 for VT support (%dx%d)"), *GetName(), Source.GetSizeX(),Source.GetSizeY());
-				// VT nonpow2 will fail to build
-				// force it into a state that will succeed? or just let it fail?
-				// you can either pad to pow2 or set MaxTextureSize and turn off VT
-				//PowerOfTwoMode = ETexturePowerOfTwoSetting::PadToPowerOfTwo;
-			}
-			else
+			if (!bIsPowerOfTwo)
 			{
-				UE_LOG(LogTexture, Warning, TEXT("VirtualTextureStreaming not supported for \"%s\", texture size is not a power-of-2"), *GetName());
-				VirtualTextureStreaming = false;
-				bRequiresNotifyMaterials = true;
+				if ( bLargeTextureMustBeVT )
+				{
+					UE_LOG(LogTexture, Warning, TEXT("Large VT \"%s\", must be padded to power-of-2 for VT support (%dx%d)"), *GetName(), Source.GetSizeX(),Source.GetSizeY());
+					// VT nonpow2 will fail to build
+					// force it into a state that will succeed? or just let it fail?
+					// you can either pad to pow2 or set MaxTextureSize and turn off VT
+					//PowerOfTwoMode = ETexturePowerOfTwoSetting::PadToPowerOfTwo;
+				}
+				else
+				{
+					UE_LOG(LogTexture, Warning, TEXT("VirtualTextureStreaming not supported for \"%s\", texture size is not a power-of-2"), *GetName());
+					VirtualTextureStreaming = false;
+					bRequiresNotifyMaterials = true;
+				}
 			}
 		}
-	}
 
-	// Make sure settings are correct for LUT textures.
-	if(LODGroup == TEXTUREGROUP_ColorLookupTable)
-	{
-		if ( MipGenSettings != TMGS_NoMipmaps || SRGB != false )
+		// Make sure settings are correct for LUT textures.
+		if(LODGroup == TEXTUREGROUP_ColorLookupTable)
 		{
-			MipGenSettings = TMGS_NoMipmaps;
-			SRGB = false;
-			bRequiresNotifyMaterials = true;
+			if ( MipGenSettings != TMGS_NoMipmaps || SRGB != false )
+			{
+				MipGenSettings = TMGS_NoMipmaps;
+				SRGB = false;
+				bRequiresNotifyMaterials = true;
+			}
 		}
-	}
+	} // end if valid source
 #endif // #if WITH_EDITORONLY_DATA
 
 	// check TC_ CompressionSettings that should have SRGB off
diff --git a/Engine/Source/Runtime/Engine/Private/TextureCubeArray.cpp b/Engine/Source/Runtime/Engine/Private/TextureCubeArray.cpp
--- a/Engine/Source/Runtime/Engine/Private/TextureCubeArray.cpp
+++ b/Engine/Source/Runtime/Engine/Private/TextureCubeArray.cpp
@@ -511,7 +511,8 @@ ENGINE_API void UTextureCubeArray::InvadiateTextureSource()
 		PlatformData = NULL;
 	}
 
-	Source.Init(0, 0, 0, 0, TSF_Invalid, nullptr);
+	Source = FTextureSource();
+	Source.SetOwner(this);
 	UpdateResource();
 }
 #endif
diff --git a/Engine/Source/Runtime/Engine/Private/VolumeTexture.cpp b/Engine/Source/Runtime/Engine/Private/VolumeTexture.cpp
--- a/Engine/Source/Runtime/Engine/Private/VolumeTexture.cpp
+++ b/Engine/Source/Runtime/Engine/Private/VolumeTexture.cpp
@@ -92,7 +92,8 @@ bool UVolumeTexture::UpdateSourceFromSourceTexture()
 	}
 	else
 	{
-		Source.Init(0, 0, 0, 0, TSF_Invalid, nullptr);
+		Source = FTextureSource();
+		Source.SetOwner(this);
 		SourceLightingGuid.Invalidate();
 
 		if (PlatformData)
