Rename garbage levels before GC so we can reuse their names.
Also break up references to actors to help ensure partial GC when there are outstanding refernces.

diff --git a/Engine/Source/Runtime/Engine/Private/World.cpp b/Engine/Source/Runtime/Engine/Private/World.cpp
--- a/Engine/Source/Runtime/Engine/Private/World.cpp
+++ b/Engine/Source/Runtime/Engine/Private/World.cpp
@@ -3304,6 +3304,24 @@ void FLevelStreamingGCHelper::PrepareStreamedOutLevelsForGC()
 				}
 			}, true, RF_NoFlags, EInternalObjectFlags::Garbage);
 
+			if (!UObjectBaseUtility::IsPendingKillEnabled())
+			{
+				// Rename the packages that we are streaming out so that we can possibly reload another copy of them
+				for (UPackage* Package : Packages)
+				{
+					FName NewName = MakeUniqueObjectName(nullptr, UPackage::StaticClass(), Package->GetFName());
+					Package->Rename(*NewName.ToString(), nullptr, REN_ForceNoResetLoaders | REN_DontCreateRedirectors | REN_NonTransactional);
+				}
+
+#if !WITH_EDITOR
+				// Clear the level actors array to maximize the memory savings from GC if there are outstanding references to some actors.
+				// Don't do this outside the editor now til we can validate it works properly with external packages.
+				Level->Actors.Empty();
+				Level->ActorsForGC.Empty();
+				Level->ActorCluster = nullptr; // The actor cluster has been marked garbage so will be dissolved, also drop a reference to it here since it has an internal array of pointers
+#endif
+			}
+
 			Level->CleanupReferences();
 		}
 	}
