Fixed BeginPlay not called on actors that are part of an Always Loaded sub-level that is loaded through a Level Streaming Volume.

diff --git a/Engine/Source/Runtime/Engine/Private/World.cpp b/Engine/Source/Runtime/Engine/Private/World.cpp
--- a/Engine/Source/Runtime/Engine/Private/World.cpp
+++ b/Engine/Source/Runtime/Engine/Private/World.cpp
@@ -2679,6 +2679,16 @@ private:
 // Cumulated time passed in UWorld::AddToWorld since last call to UWorld::UpdateLevelStreaming.
 static double GAddToWorldTimeCumul = 0.0;
 
+// Only called internally by UWorld::AddToWorld and UWorld::InitializeActorsForPlay
+static void ResetLevelFlagsOnLevelAddedToWorld(ULevel* Level)
+{
+	Level->bAlreadyShiftedActors = false;
+	Level->bAlreadyUpdatedComponents = false;
+	Level->bAlreadySortedActorList = false;
+	Level->bAlreadyClearedActorsSeamlessTravelFlag = false;
+	Level->ResetRouteActorInitializationState();
+}
+
 void UWorld::AddToWorld( ULevel* Level, const FTransform& LevelTransform, bool bConsiderTimeLimit, FNetLevelVisibilityTransactionId TransactionId)
 {
 	TRACE_CPUPROFILER_EVENT_SCOPE(AddToWorld);
@@ -2932,11 +2942,7 @@ void UWorld::AddToWorld( ULevel* Level, const FTransform& LevelTransform, bool b
 	{
 		SCOPE_TIME_TO_VAR(&PerformLastStepTime);
 		
-		Level->bAlreadyShiftedActors					= false;
-		Level->bAlreadyUpdatedComponents				= false;
-		Level->bAlreadySortedActorList					= false;
-		Level->bAlreadyClearedActorsSeamlessTravelFlag	= false;
-		Level->ResetRouteActorInitializationState();
+		ResetLevelFlagsOnLevelAddedToWorld(Level);
 
 		// Finished making level visible - allow other levels to be added to the world.
 		CurrentLevelPendingVisibility = nullptr;
@@ -4850,6 +4856,18 @@ void UWorld::InitializeActorsForPlay(const FURL& InURL, bool bResetTime, FRegist
 		Level->SortActorList();
 	}
 
+	// Make sure to reset level flags necessary for a future AddToWorld to work properly.
+	// 
+	// One use case where this is important is for AlwaysLoaded Streaming Levels : 
+	//  - Streaming Level Volumes can unload the level
+	//  - If we don't reset these flags, the next time the level gets added, it will skip steps like RouteActorInitialize.
+	//
+	// Skip PersistentLevel (as it's not required for it)
+	for (int32 LevelIndex = 1; LevelIndex < Levels.Num(); LevelIndex++)
+	{
+		ResetLevelFlagsOnLevelAddedToWorld(Levels[LevelIndex]);
+	} 
+
 	// update the auto-complete list for the console
 	UConsole* ViewportConsole = (GEngine->GameViewport != nullptr) ? GEngine->GameViewport->ViewportConsole : nullptr;
 	if (ViewportConsole != nullptr)
