Extend logging of preprocessed shader for build machines.

#rnx

diff --git a/Engine/Source/Runtime/Engine/Private/ShaderCompiler/ShaderCompiler.cpp b/Engine/Source/Runtime/Engine/Private/ShaderCompiler/ShaderCompiler.cpp
--- a/Engine/Source/Runtime/Engine/Private/ShaderCompiler/ShaderCompiler.cpp
+++ b/Engine/Source/Runtime/Engine/Private/ShaderCompiler/ShaderCompiler.cpp
@@ -171,7 +171,13 @@ static FAutoConsoleVariableRef CVarShaderCompilerAllowDistributedCompilation(
 );
 
 /** Maximum number of preprocessed shaders to dump to the log on a crash. Replace with CVar if needed. */
-static constexpr uint32 GMaxNumDumpedShaderSources = 10;
+int32 GMaxNumDumpedShaderSources = 1;
+static FAutoConsoleVariableRef CVarShaderCompilerMaxDumpedShaderSources(
+	TEXT("r.ShaderCompiler.MaxDumpedShaderSources"),
+	GMaxNumDumpedShaderSources,
+	TEXT("Maximum number of preprocessed shader sources to dump to the log on shader compile errors. By default 1."),
+	ECVF_Default
+);
 
 int32 GSShaderCheckLevel = 1;
 static FAutoConsoleVariableRef CVarGSShaderCheckLevel(
@@ -264,6 +270,21 @@ void DumpShaderDDCKeyToFile(const EShaderPlatform InPlatform, bool bWithEditor,
 	DumpAr->Serialize(const_cast<TCHAR*>(*DDCKey), DDCKey.Len() * sizeof(TCHAR));
 }
 
+static void DumpShaderSourceToLog(const FString& DumpedSource)
+{
+	// Log dumped shader source line by line as message lengths in UE_LOG are limited.
+	TArray<FString> DumpedSourceLines;
+	DumpedSource.ParseIntoArrayLines(DumpedSourceLines, /*bCullEmpty:*/ false);
+
+	TGuardValue<ELogTimes::Type> DisableLogTimes(GPrintLogTimes, ELogTimes::None);
+	GLog->Serialize(TEXT("\n======================= DUMPED SHADER BEGIN =======================\n"), ELogVerbosity::Log, NAME_None);
+	for (const FString& Line : DumpedSourceLines)
+	{
+		GLog->Serialize(*Line, ELogVerbosity::Log, NAME_None);
+	}
+	GLog->Serialize(TEXT("\n======================= DUMPED SHADER END =======================\n"), ELogVerbosity::Log, NAME_None);
+}
+
 namespace ShaderCompiler
 {
 	bool IsJobCacheEnabled()
@@ -1472,7 +1493,7 @@ bool FShaderCompileUtilities::DoWriteTasks(const TArray<FShaderCommonCompileJobP
 	return InTransferFile.Close();
 }
 
-static void ProcessErrors(const FShaderCompileJob& CurrentJob, TArray<FString>& UniqueErrors, FString& ErrorString)
+static void ProcessErrors(const FShaderCompileJob& CurrentJob, TArray<FString>& UniqueErrors, FString& ErrorString, int32& NumDumpedShaderSources)
 {
 	bool bReportedDebugInfo = false;
 
@@ -1550,6 +1571,14 @@ static void ProcessErrors(const FShaderCompileJob& CurrentJob, TArray<FString>&
 			{
 				// Format everything on one line, and with the correct verbosity, so we can display proper errors in the failure logs.
 				UE_LOG(LogShaderCompilers, Error, TEXT("%s%s"), *UniqueErrorPrefix.Replace(TEXT("\n"), TEXT("")), *CurrentError.GetErrorStringWithLineMarker());
+
+				const FString& DumpedSource = CurrentJob.Output.OptionalPreprocessedShaderSource;
+				if (!DumpedSource.IsEmpty() && NumDumpedShaderSources < GMaxNumDumpedShaderSources)
+				{
+					// Limit number of preprocessed shaders to dump to the log as they are quite large
+					DumpShaderSourceToLog(DumpedSource);
+					++NumDumpedShaderSources;
+				}
 			}
 			else if (FPlatformMisc::IsDebuggerPresent())
 			{
@@ -4419,19 +4448,8 @@ void FShaderCompilingManager::ProcessCompiledShaderMaps(
 
 					if (!DumpedSource.IsEmpty() && NumDumpedShaderSources < GMaxNumDumpedShaderSources)
 					{
-						// Log dumped shader source line by line as message lengths in UE_LOG are limited.
-						TArray<FString> DumpedSourceLines;
-						DumpedSource.ParseIntoArrayLines(DumpedSourceLines, /*bCullEmpty:*/ false);
-
-						TGuardValue<ELogTimes::Type> DisableLogTimes(GPrintLogTimes, ELogTimes::None);
-						GLog->Serialize(TEXT("\n======================= DUMPED SHADER BEGIN =======================\n"), ELogVerbosity::Log, NAME_None);
-						for (const FString& Line : DumpedSourceLines)
-						{
-							GLog->Serialize(*Line, ELogVerbosity::Log, NAME_None);
-						}
-						GLog->Serialize(TEXT("\n======================= DUMPED SHADER END =======================\n"), ELogVerbosity::Log, NAME_None);
-
 						// Limit number of preprocessed shaders to dump to the log as they are quite large
+						DumpShaderSourceToLog(DumpedSource);
 						++NumDumpedShaderSources;
 					}
 
@@ -4759,7 +4777,7 @@ bool FShaderCompilingManager::HandlePotentialRetryOnError(TMap<int32, FShaderMap
 						const auto* SingleJob = CurrentJob.GetSingleShaderJob();
 						if (SingleJob)
 						{
-							ProcessErrors(*SingleJob, UniqueErrors, ErrorString);
+							ProcessErrors(*SingleJob, UniqueErrors, ErrorString, NumDumpedShaderSources);
 						}
 						else
 						{
@@ -4767,7 +4785,7 @@ bool FShaderCompilingManager::HandlePotentialRetryOnError(TMap<int32, FShaderMap
 							check(PipelineJob);
 							for (auto CommonJob : PipelineJob->StageJobs)
 							{
-								ProcessErrors(*CommonJob, UniqueErrors, ErrorString);
+								ProcessErrors(*CommonJob, UniqueErrors, ErrorString, NumDumpedShaderSources);
 							}
 						}
 					}
diff --git a/Engine/Source/Runtime/Engine/Public/ShaderCompiler.h b/Engine/Source/Runtime/Engine/Public/ShaderCompiler.h
--- a/Engine/Source/Runtime/Engine/Public/ShaderCompiler.h
+++ b/Engine/Source/Runtime/Engine/Public/ShaderCompiler.h
@@ -879,7 +879,7 @@ private:
 	/** Number of runs through single-threaded compiling before we can retry to compile through workers. -1 if not used. */
 	int32 NumSingleThreadedRunsBeforeRetry;
 	/** Number of preprocessed shader sources that are dumped due to a crash of the shader compiler. Only used for D3DCompile exceptions at the moment. */
-	uint32 NumDumpedShaderSources;
+	int32 NumDumpedShaderSources;
 	/** Process Id of UE. */
 	uint32 ProcessId;
 	/** Whether to allow compiling shaders through the worker application, which allows multiple cores to be used. */
