World Partition: cleanup invalid handles for pinned actors and also pin actors on save if there's zero reference to them.

#rnx

diff --git a/Engine/Source/Runtime/Engine/Private/WorldPartition/WorldPartition.cpp b/Engine/Source/Runtime/Engine/Private/WorldPartition/WorldPartition.cpp
--- a/Engine/Source/Runtime/Engine/Private/WorldPartition/WorldPartition.cpp
+++ b/Engine/Source/Runtime/Engine/Private/WorldPartition/WorldPartition.cpp
@@ -243,7 +243,7 @@ protected:
 	//~ Begin IWorldPartitionActorLoaderInterface::ILoaderAdapterList interface
 	virtual bool ShouldActorBeLoaded(const FWorldPartitionHandle& ActorHandle) const override
 	{
-		return !ActorsToRemove.Contains(ActorHandle) && PassDataLayersFilter(ActorHandle);
+		return ActorHandle.IsValid() && !ActorsToRemove.Contains(ActorHandle) && PassDataLayersFilter(ActorHandle);
 	}
 	//~ End IWorldPartitionActorLoaderInterface::ILoaderAdapterList interface
 
@@ -328,10 +328,10 @@ void UWorldPartition::OnPackageDirtyStateChanged(UPackage* Package)
 			{
 				DirtyActors.Add(ActorHandle.ToReference(), Actor);
 			}
-			else if (ActorHandle->GetHardRefCount())
+			else
 			{
-				// If we hold the last reference to that actor, pin it to avoid unloading
-				if (PinnedActors && (ActorHandle->GetHardRefCount() == 1))
+				// If we hold the last reference to that actor (or no reference are held at all), pin it to avoid unloading
+				if (PinnedActors && (ActorHandle->GetHardRefCount() <= 1))
 				{
 					PinnedActors->AddActors({ ActorHandle });
 				}
