Add support for dynamically spawned Level Instance by replicating the LI Guid from the server to the clients
#rnx

diff --git a/Engine/Source/Runtime/Engine/Private/LevelInstance/LevelInstanceActor.cpp b/Engine/Source/Runtime/Engine/Private/LevelInstance/LevelInstanceActor.cpp
--- a/Engine/Source/Runtime/Engine/Private/LevelInstance/LevelInstanceActor.cpp
+++ b/Engine/Source/Runtime/Engine/Private/LevelInstance/LevelInstanceActor.cpp
@@ -8,6 +8,7 @@
 #include "Engine/World.h"
 #include "Engine/Engine.h"
 #include "LevelInstancePrivate.h"
+#include "Net/UnrealNetwork.h"
 
 #if WITH_EDITOR
 #include "WorldPartition/WorldPartitionActorDesc.h"
@@ -46,16 +47,41 @@
 #endif
 }
 
+void ALevelInstance::GetLifetimeReplicatedProps(TArray<FLifetimeProperty>& OutLifetimeProps) const
+{
+	Super::GetLifetimeReplicatedProps(OutLifetimeProps);
+	FDoRepLifetimeParams Params;
+	Params.Condition = COND_InitialOnly;
+	DOREPLIFETIME_WITH_PARAMS_FAST(ALevelInstance, LevelInstanceSpawnGuid, Params);
+}
+
+void ALevelInstance::OnRep_LevelInstanceSpawnGuid()
+{
+	if (GetLocalRole() != ENetRole::ROLE_Authority && !LevelInstanceActorGuid.IsValid())
+	{
+		check(LevelInstanceSpawnGuid.IsValid());
+		LevelInstanceActorGuid.ActorGuid = LevelInstanceSpawnGuid;
+		LevelInstanceActorImpl.RegisterLevelInstance();
+	}
+}
+
 void ALevelInstance::PostRegisterAllComponents()
 {
 	Super::PostRegisterAllComponents();
-
+	
+	if (GetLocalRole() == ENetRole::ROLE_Authority)
+	{
 #if !WITH_EDITOR
-	// If the level instance was spawned, not loaded
-	LevelInstanceActorGuid.AssignIfInvalid();
+		// If the level instance was spawned, not loaded
+		LevelInstanceActorGuid.AssignIfInvalid();
 #endif
+		LevelInstanceSpawnGuid = LevelInstanceActorGuid.GetGuid();
+	}
 
-	LevelInstanceActorImpl.RegisterLevelInstance();
+	if (LevelInstanceActorGuid.IsValid())
+	{
+		LevelInstanceActorImpl.RegisterLevelInstance();
+	}
 }
 
 void ALevelInstance::PostUnregisterAllComponents()
diff --git a/Engine/Source/Runtime/Engine/Private/LevelInstance/LevelInstanceActorGuid.cpp b/Engine/Source/Runtime/Engine/Private/LevelInstance/LevelInstanceActorGuid.cpp
--- a/Engine/Source/Runtime/Engine/Private/LevelInstance/LevelInstanceActorGuid.cpp
+++ b/Engine/Source/Runtime/Engine/Private/LevelInstance/LevelInstanceActorGuid.cpp
@@ -13,15 +13,26 @@ void FLevelInstanceActorGuid::AssignIfInvalid()
 }
 #endif
 
+bool FLevelInstanceActorGuid::IsValid() const
+{
+	return GetGuid_Internal().IsValid();
+}
+
 const FGuid& FLevelInstanceActorGuid::GetGuid() const
+{
+	const FGuid& Guid = GetGuid_Internal();
+	check(Actor->IsTemplate() || Guid.IsValid());
+	return Guid;
+}
+
+const FGuid& FLevelInstanceActorGuid::GetGuid_Internal() const
 {
 	check(Actor);
 #if WITH_EDITOR
-	const FGuid& Guid = Actor->GetActorGuid();
+	const FGuid& Guid = Actor->GetLocalRole() == ENetRole::ROLE_Authority ? Actor->GetActorGuid() : ActorGuid;
 #else
 	const FGuid& Guid = ActorGuid;
 #endif
-	check(Actor->IsTemplate() || Guid.IsValid());
 	return Guid;
 }
 
diff --git a/Engine/Source/Runtime/Engine/Public/LevelInstance/LevelInstanceActor.h b/Engine/Source/Runtime/Engine/Public/LevelInstance/LevelInstanceActor.h
--- a/Engine/Source/Runtime/Engine/Public/LevelInstance/LevelInstanceActor.h
+++ b/Engine/Source/Runtime/Engine/Public/LevelInstance/LevelInstanceActor.h
@@ -31,6 +31,9 @@
 	UPROPERTY()
 	TSoftObjectPtr<UWorld> CookedWorldAsset;
 
+	UPROPERTY(Transient, ReplicatedUsing=OnRep_LevelInstanceSpawnGuid)
+	FGuid LevelInstanceSpawnGuid;
+
 private:
 	FLevelInstanceActorGuid LevelInstanceActorGuid;
 	FLevelInstanceActorImpl LevelInstanceActorImpl;
@@ -54,7 +57,11 @@
 	virtual void PostRegisterAllComponents() override;
 	virtual void PostUnregisterAllComponents() override;
 	virtual void Serialize(FArchive& Ar) override;
+	virtual void GetLifetimeReplicatedProps(TArray<FLifetimeProperty>& OutLifetimeProps) const override;
 		
+	UFUNCTION()
+	void OnRep_LevelInstanceSpawnGuid();
+
 #if WITH_EDITOR
 	// Begin ILevelInstanceInterface
 	virtual ULevelInstanceComponent* GetLevelInstanceComponent() const override;
diff --git a/Engine/Source/Runtime/Engine/Public/LevelInstance/LevelInstanceActorGuid.h b/Engine/Source/Runtime/Engine/Public/LevelInstance/LevelInstanceActorGuid.h
--- a/Engine/Source/Runtime/Engine/Public/LevelInstance/LevelInstanceActorGuid.h
+++ b/Engine/Source/Runtime/Engine/Public/LevelInstance/LevelInstanceActorGuid.h
@@ -22,10 +22,14 @@ struct ENGINE_API FLevelInstanceActorGuid
 	void AssignIfInvalid();
 #endif
 
+	bool IsValid() const;
 	const FGuid& GetGuid() const;
 
 	TObjectPtr<AActor> Actor = nullptr;
 	FGuid ActorGuid;
 
 	ENGINE_API friend FArchive& operator<<(FArchive& Ar, FLevelInstanceActorGuid& LevelInstanceActorGuid);
+
+private:
+	const FGuid& GetGuid_Internal() const;
 };
\ No newline at end of file
