Texture2dArrays now stream properly, correctly handling LODBias and packed mip tails.

diff --git a/Engine/Source/Runtime/Engine/Classes/Engine/Texture2DArray.cpp b/Engine/Source/Runtime/Engine/Classes/Engine/Texture2DArray.cpp
--- a/Engine/Source/Runtime/Engine/Classes/Engine/Texture2DArray.cpp
+++ b/Engine/Source/Runtime/Engine/Classes/Engine/Texture2DArray.cpp
@@ -5,6 +5,8 @@
 =============================================================================*/
 
 #include "Engine/Texture2DArray.h"
+
+#include "AsyncCompilationHelpers.h"
 #include "RenderUtils.h"
 #include "TextureResource.h"
 #include "EngineUtils.h"
@@ -16,6 +18,9 @@
 #include "Streaming/TextureStreamIn.h"
 #include "Streaming/TextureStreamOut.h"
 #include "Engine/TextureMipDataProviderFactory.h"
+#include "Misc/ScopedSlowTask.h"
+
+#define LOCTEXT_NAMESPACE "UTexture2DArray"
 
 // Master switch to control whether streaming is enabled for texture2darray. 
 bool GSupportsTexture2DArrayStreaming = true;
@@ -67,6 +72,30 @@ FTextureResource* UTexture2DArray::CreateResource()
 	return nullptr;
 }
 
+
+// Any direct access to GetPlatformDataOrWait() will stall until the structure
+// is safe to use. It is advisable to replace those use case with
+// async aware code to avoid stalls where possible.
+FTexturePlatformData* UTexture2DArray::GetPlatformDataOrWait()
+{
+#if WITH_EDITOR
+	if (PlatformData && !PlatformData->IsAsyncWorkComplete())
+	{
+		TRACE_CPUPROFILER_EVENT_SCOPE(UTexture2D::GetPlatformDataStall);
+		UE_LOG(LogTexture, Log, TEXT("Call to GetPlatformDataOrWait() is forcing a wait on data that is not yet ready."));
+
+		FText Msg = FText::Format(LOCTEXT("WaitOnTextureCompilation", "Waiting on texture array compilation {0} ..."), FText::FromString(GetName()));
+		FScopedSlowTask Progress(1.f, Msg, true);
+		Progress.MakeDialog(true);
+		uint64 StartTime = FPlatformTime::Cycles64();
+		PlatformData->FinishCache();
+		AsyncCompilationHelpers::SaveStallStack(FPlatformTime::Cycles64() - StartTime);
+	}
+#endif
+
+	return PlatformData;
+}
+
 void UTexture2DArray::UpdateResource()
 {
 #if WITH_EDITOR
@@ -411,3 +440,5 @@ bool UTexture2DArray::StreamIn(int32 NewMipCount, bool bHighPrio)
 	}
 	return false;
 }
+
+#undef LOCTEXT_NAMESPACE
\ No newline at end of file
diff --git a/Engine/Source/Runtime/Engine/Classes/Engine/Texture2DArray.h b/Engine/Source/Runtime/Engine/Classes/Engine/Texture2DArray.h
--- a/Engine/Source/Runtime/Engine/Classes/Engine/Texture2DArray.h
+++ b/Engine/Source/Runtime/Engine/Classes/Engine/Texture2DArray.h
@@ -47,6 +47,21 @@ public:
 		return PlatformData ? PlatformData->PixelFormat : PF_Unknown;
 	}
 
+	/**
+	*  Makes a copy of the mip chain for the texture array starting with InFirstMipToLoad.
+	* 
+	*  In the editor this can potentially wait for texture builds to complete.
+	*/
+	bool GetMipData(int32 InFirstMipToLoad, TArray<FUniqueBuffer, TInlineAllocator<MAX_TEXTURE_MIP_COUNT>>& OutMipData);
+
+	/** 
+	*  Returns the platform data, potentially waiting on texture build if not ready.
+	*  Generally waits should be avoided by structuring callers to expect async behavior.
+	* 
+	*  Waits only occur in the editor (as builds only occur in the editor).
+	*/
+	FTexturePlatformData* GetPlatformDataOrWait();
+
 	//~ Begin UTexture Interface
 	virtual ETextureClass GetTextureClass() const override { return ETextureClass::Array; }
 	virtual void Serialize(FArchive& Ar) override;
diff --git a/Engine/Source/Runtime/Engine/Private/Rendering/Texture2DArrayResource.cpp b/Engine/Source/Runtime/Engine/Private/Rendering/Texture2DArrayResource.cpp
--- a/Engine/Source/Runtime/Engine/Private/Rendering/Texture2DArrayResource.cpp
+++ b/Engine/Source/Runtime/Engine/Private/Rendering/Texture2DArrayResource.cpp
@@ -20,58 +20,32 @@ FTexture2DArrayResource::FTexture2DArrayResource(UTexture2DArray* InOwner, const
 	AddressU = InOwner->AddressX == TA_Wrap ? AM_Wrap : (InOwner->AddressX == TA_Clamp ? AM_Clamp : AM_Mirror);
 	AddressV = InOwner->AddressY == TA_Wrap ? AM_Wrap : (InOwner->AddressY == TA_Clamp ? AM_Clamp : AM_Mirror);
 	AddressW = InOwner->AddressZ == TA_Wrap ? AM_Wrap : (InOwner->AddressZ == TA_Clamp ? AM_Clamp : AM_Mirror);
-
-	const int32 MaxLoadableMipIndex = State.MaxNumLODs - FMath::Max(1, PlatformData->GetNumMipsInTail()) + 1;
-
-	TArray<uint64> MipOffsets;
-	MipOffsets.AddZeroed(State.MaxNumLODs);
-
-	uint64 InitialMipDataSize = 0;
-	TArrayView<const FTexture2DMipMap*> MipsView = GetPlatformMipsView();
-	for (int32 MipIdx = State.RequestedFirstLODIdx(); MipIdx < MaxLoadableMipIndex; ++MipIdx)
-	{
-		const FTexture2DMipMap& Mip = *MipsView[MipIdx];
-		MipOffsets[MipIdx] = InitialMipDataSize;
-		InitialMipDataSize += Mip.BulkData.GetBulkDataSize();
-	}
-	InitialMipData.Reset(new uint8[InitialMipDataSize]);
-
-	// Resize each structure correctly per slice, since we access it per mip after.
-	SliceMipDataViews.AddDefaulted(SizeZ);
-	for (uint32 SliceIdx = 0; SliceIdx < SizeZ; ++SliceIdx)
+	
+	//
+	// All resource requests assume LOD bias is baked in - however GetMipData doesn't, and we don't want to grab
+	// that data if we don't have to. LODCountToAssetFirstLODIdx() is the index with LODBias, and RequestedFirstLODIdx()
+	// isn't. (LODCountToAssetFirstLODIdx = RequestedFirstLODIdx + AssetLODBias).
+	//
+	if (InOwner->GetMipData(State.LODCountToAssetFirstLODIdx(State.NumRequestedLODs), AllMipsData) == false)
 	{
-		SliceMipDataViews[SliceIdx].AddDefaulted(MaxLoadableMipIndex);
+		// This is fatal as we will crash trying to upload the data below, this way we crash at the cause.
+		UE_LOG(LogTexture, Fatal, TEXT("Corrupt texture [%s]! Unable to load mips (bulk data missing)"), *TextureName.ToString());
+		return;
 	}
-
-	for (int32 MipIdx = State.RequestedFirstLODIdx(); MipIdx < MaxLoadableMipIndex; ++MipIdx)
+	
+	// AllMipsData has NumRequestedLODs - (GetNumMipsInTail() - 1)
 	{
-		FTexture2DMipMap& Mip = const_cast<FTexture2DMipMap&>(*MipsView[MipIdx]);
-		if (Mip.BulkData.GetBulkDataSize() > 0)
-		{
-			const uint64 SliceMipDataSize = Mip.BulkData.GetBulkDataSize() / SizeZ;
-			const uint8* SrcData = (const uint8*)Mip.BulkData.Lock(LOCK_READ_ONLY);
-
-			for (uint32 SliceIdx = 0; SliceIdx < SizeZ; ++SliceIdx)
-			{
-				TArrayView<uint8>& SliceMipDataView = SliceMipDataViews[SliceIdx][MipIdx];
-				SliceMipDataView = TArrayView<uint8>(InitialMipData.Get() + MipOffsets[MipIdx] + SliceIdx * SliceMipDataSize, SliceMipDataSize);
-				FMemory::Memcpy(SliceMipDataView.GetData(), SrcData + SliceIdx * SliceMipDataSize, SliceMipDataSize);
-			}
-
-			Mip.BulkData.Unlock();
-		}
-		else
-		{
-			UE_LOG(LogTexture, Error, TEXT("Corrupt texture [%s]! Missing bulk data for MipIndex=%d"), *TextureName.ToString(), MipIdx);
-		}
+		// If zero, then it's unpacked which means there's 1 mip in the tail.
+		const int32 MipsInTail = PlatformData->GetNumMipsInTail() ? PlatformData->GetNumMipsInTail() : 1;
+		const int32 MipCountLostDueToPacking = MipsInTail - 1;
+		check(AllMipsData.Num() == State.NumRequestedLODs - MipCountLostDueToPacking);
 	}
 }
 
 void FTexture2DArrayResource::CreateTexture()
 {
-	const int32 RequestedFirstLODIdx = State.RequestedFirstLODIdx();
-	TArrayView<const FTexture2DMipMap*> MipsView = GetPlatformMipsView();
-	const FTexture2DMipMap& FirstMip = *MipsView[RequestedFirstLODIdx];
+	TArrayView<const FTexture2DMipMap*> MipsViewPostLODBias = GetPlatformMipsView();
+	const FTexture2DMipMap& FirstMip = *MipsViewPostLODBias[State.RequestedFirstLODIdx()];
 
 	const FRHITextureCreateDesc Desc =
 		FRHITextureCreateDesc::Create2DArray(TEXT("FTexture2DArrayResource"), FirstMip.SizeX, FirstMip.SizeY, FirstMip.SizeZ, PixelFormat)
@@ -81,25 +55,26 @@ void FTexture2DArrayResource::CreateTexture()
 
 	TextureRHI = RHICreateTexture(Desc);
 
-	// Read the initial cached mip levels into the RHI texture.
-	const int32 NumLoadableMips = State.NumRequestedLODs - FMath::Max(1, PlatformData->GetNumMipsInTail()) + 1;
-	for (int32 RHIMipIdx = 0; RHIMipIdx < NumLoadableMips; ++RHIMipIdx)
-	{
-		const int32 MipIdx = RHIMipIdx + RequestedFirstLODIdx;
-		for (uint32 SliceIdx = 0; SliceIdx < SizeZ; ++SliceIdx)
+	// Copy the mip data in to the texture.
+	for (int32 MipIndex = 0; MipIndex < AllMipsData.Num(); MipIndex++)
+	{		
+		for (uint32 ArrayIndex = 0; ArrayIndex < SizeZ; ++ArrayIndex)
 		{
 			uint32 DestStride = 0;
-			void* DestData = RHILockTexture2DArray(TextureRHI, SliceIdx, RHIMipIdx, RLM_WriteOnly, DestStride, false);
+			void* DestData = RHILockTexture2DArray(TextureRHI, ArrayIndex, MipIndex, RLM_WriteOnly, DestStride, false);
 			if (DestData)
 			{
-				GetData(SliceIdx, MipIdx, DestData, DestStride);
+				GetData(FirstMip.SizeX, FirstMip.SizeY, ArrayIndex, MipIndex, DestData, DestStride);
+			}
+			else
+			{
+				UE_LOG(LogTexture, Warning, TEXT("Failed to lock texture 2d array mip/slice %d / %d (%s)"), MipIndex, ArrayIndex, *TextureName.ToString());
 			}
-			RHIUnlockTexture2DArray(TextureRHI, SliceIdx, RHIMipIdx, false);
+			RHIUnlockTexture2DArray(TextureRHI, ArrayIndex, MipIndex, false);
 		}
 	}
-		
-	SliceMipDataViews.Empty();
-	InitialMipData.Reset();
+
+	AllMipsData.Empty();
 }
 
 void FTexture2DArrayResource::CreatePartiallyResidentTexture()
@@ -122,15 +97,22 @@ uint64 FTexture2DArrayResource::GetPlatformMipsSize(uint32 NumMips) const
 	}
 }
 
-void FTexture2DArrayResource::GetData(uint32 SliceIndex, uint32 MipIndex, void* Dest, uint32 DestPitch)
+void FTexture2DArrayResource::GetData(int32 BaseRHIMipSizeX, int32 BaseRHIMipSizeY, uint32 ArrayIndex, uint32 MipIndex, void* Dest, uint32 DestPitch) const
 {
-	const TArrayView<uint8>& SliceMipDataView = SliceMipDataViews[SliceIndex][MipIndex];
+	const uint32 ArrayCount = SizeZ;
+	const FMemoryView MipView = AllMipsData[MipIndex].GetView();
+	const uint64 ArraySliceDataSize = MipView.GetSize() / ArrayCount;
+	const FMemoryView& SliceMipDataView = MipView.Mid(ArrayIndex * ArraySliceDataSize, ArraySliceDataSize);
+
+	
+	// This check works for normal mips but not packed mips.
+	//check(ArraySliceDataSize == CalcTextureMipMapSize(BaseRHIMipSizeX, BaseRHIMipSizeX, PixelFormat, MipIndex));
 
 	// for platforms that returned 0 pitch from Lock, we need to just use the bulk data directly, never do 
 	// runtime block size checking, conversion, or the like
 	if (DestPitch == 0)
 	{
-		FMemory::Memcpy(Dest, SliceMipDataView.GetData(), SliceMipDataView.Num());
+		FMemory::Memcpy(Dest, SliceMipDataView.GetData(), SliceMipDataView.GetSize());
 	}
 	else
 	{
@@ -138,8 +120,8 @@ void FTexture2DArrayResource::GetData(uint32 SliceIndex, uint32 MipIndex, void*
 		const int32 BlockSizeY = GPixelFormats[PixelFormat].BlockSizeY;
 		const int32 BlockBytes = GPixelFormats[PixelFormat].BlockBytes;
 
-		const uint32 MipSizeX = FMath::Max<int32>(SizeX >> MipIndex, 1);
-		const uint32 MipSizeY = FMath::Max<int32>(SizeY >> MipIndex, 1);
+		const uint32 MipSizeX = FMath::Max<int32>(BaseRHIMipSizeX >> MipIndex, 1);
+		const uint32 MipSizeY = FMath::Max<int32>(BaseRHIMipSizeY >> MipIndex, 1);
 
 		uint32 NumColumns = FMath::DivideAndRoundUp<int32>(MipSizeX, BlockSizeX);
 		uint32 NumRows = FMath::DivideAndRoundUp<int32>(MipSizeY, BlockSizeY);
diff --git a/Engine/Source/Runtime/Engine/Private/Rendering/Texture2DArrayResource.h b/Engine/Source/Runtime/Engine/Private/Rendering/Texture2DArrayResource.h
--- a/Engine/Source/Runtime/Engine/Private/Rendering/Texture2DArrayResource.h
+++ b/Engine/Source/Runtime/Engine/Private/Rendering/Texture2DArrayResource.h
@@ -9,6 +9,7 @@
 #include "CoreMinimal.h"
 #include "Rendering/StreamableTextureResource.h"
 #include "Containers/ResourceArray.h"
+#include "Memory/SharedBuffer.h"
 
 class UTexture2DArray;
 
@@ -31,12 +32,8 @@ protected:
 	void CreateTexture() final override;
 	void CreatePartiallyResidentTexture() final override;
 
-	void GetData(uint32 SliceIndex, uint32 MipIndex, void* Dest, uint32 DestPitch);
+	void GetData(int32 BaseRHIMipSizeX, int32 BaseRHIMipSizeY, uint32 ArrayIndex, uint32 MipIndex, void* Dest, uint32 DestPitch) const;
 
-	/** The initial data for all mips of a single slice. */
-	typedef TArray<TArrayView<uint8>, TInlineAllocator<MAX_TEXTURE_MIP_COUNT> > FSingleSliceMipDataView;
-	/** All slices initial data. */
-	TArray<FSingleSliceMipDataView> SliceMipDataViews;
-	/** The single allocation holding initial data. */
-	TUniquePtr<uint8[]> InitialMipData;
+	// Each mip has all array slices. This will be [State.NumRequestedLODs] long, less any packed mips.
+	TArray<FUniqueBuffer, TInlineAllocator<MAX_TEXTURE_MIP_COUNT>> AllMipsData;
 };
diff --git a/Engine/Source/Runtime/Engine/Private/TextureDerivedData.cpp b/Engine/Source/Runtime/Engine/Private/TextureDerivedData.cpp
--- a/Engine/Source/Runtime/Engine/Private/TextureDerivedData.cpp
+++ b/Engine/Source/Runtime/Engine/Private/TextureDerivedData.cpp
@@ -2575,6 +2575,44 @@
 	Texture derived data interface.
 ------------------------------------------------------------------------------*/
 
+bool UTexture2DArray::GetMipData(int32 InFirstMipToLoad, TArray<FUniqueBuffer, TInlineAllocator<MAX_TEXTURE_MIP_COUNT>>& OutMipData)
+{
+	FTexturePlatformData* LocalPlatformData = GetPlatformDataOrWait();
+	const int32 ReadableMipCount = LocalPlatformData->Mips.Num() - (LocalPlatformData->GetNumMipsInTail() > 0 ? LocalPlatformData->GetNumMipsInTail() - 1 : 0);
+
+	int32 OutputMipCount = ReadableMipCount - InFirstMipToLoad;
+
+	check(OutputMipCount <= MAX_TEXTURE_MIP_COUNT);
+
+	void* MipData[MAX_TEXTURE_MIP_COUNT] = {};
+	int64 MipSizes[MAX_TEXTURE_MIP_COUNT];
+	if (LocalPlatformData->TryLoadMipsWithSizes(InFirstMipToLoad, MipData, MipSizes, GetPathName()) == false)
+	{
+		// Unable to load mips from the cache. Rebuild the texture and try again.
+		UE_LOG(LogTexture, Warning, TEXT("GetMipData failed for %s (%s)"),
+			*GetPathName(), GPixelFormats[GetPixelFormat()].Name);
+#if WITH_EDITOR
+		if (!GetOutermost()->bIsCookedForEditor)
+		{
+			ForceRebuildPlatformData();
+			if (LocalPlatformData->TryLoadMipsWithSizes(InFirstMipToLoad, MipData, MipSizes, GetPathName()) == false)
+			{
+				UE_LOG(LogTexture, Error, TEXT("TryLoadMipsWithSizes still failed after ForceRebuildPlatformData %s."), *GetPathName());
+				return false;
+			}
+		}
+#else // #if WITH_EDITOR
+		return false;
+#endif // WITH_EDITOR
+	}
+
+	for (int32 MipIndex = 0; MipIndex < OutputMipCount; MipIndex++)
+	{
+		OutMipData.Emplace(FUniqueBuffer::TakeOwnership(MipData[MipIndex], MipSizes[MipIndex], [](void* Ptr) {FMemory::Free(Ptr);} ));
+	}
+	return true;
+}
+
 void UTexture2D::GetMipData(int32 FirstMipToLoad, void** OutMipData)
 {
 	if (GetPlatformData()->TryLoadMips(FirstMipToLoad, OutMipData, GetPathName()) == false)
