Add a FInputDeviceId option to the FInputKeyParams. This will allow us to filter down input based on what input device it came from later down the line. This is a small part of the refactor for multiple gamepad support.

Add the bFilterInputByPlatformUser flag to the UInputSettings to enable or disable this filtering on the player controller.

diff --git a/Engine/Source/Runtime/Engine/Classes/GameFramework/InputSettings.h b/Engine/Source/Runtime/Engine/Classes/GameFramework/InputSettings.h
--- a/Engine/Source/Runtime/Engine/Classes/GameFramework/InputSettings.h
+++ b/Engine/Source/Runtime/Engine/Classes/GameFramework/InputSettings.h
@@ -66,6 +66,13 @@ class ENGINE_API UInputSettings
 	 */
 	UPROPERTY(config, EditAnywhere, Category = "Input")
 	uint8 bEnableMotionControls:1;
+
+	/**
+	 * If true, then the PlayerController::InputKey function will only process an input event if it
+	 * came from an input device that is owned by the PlayerController's Platform User.
+	 */
+	UPROPERTY(config, EditAnywhere, Category = "Input")
+	uint8 bFilterInputByPlatformUser:1;
 	
 	/** Should the touch input interface be shown always, or only when the platform has a touch screen? */
 	UPROPERTY(config, EditAnywhere, Category="Mobile")
diff --git a/Engine/Source/Runtime/Engine/Classes/GameFramework/PlayerInput.h b/Engine/Source/Runtime/Engine/Classes/GameFramework/PlayerInput.h
--- a/Engine/Source/Runtime/Engine/Classes/GameFramework/PlayerInput.h
+++ b/Engine/Source/Runtime/Engine/Classes/GameFramework/PlayerInput.h
@@ -329,22 +329,25 @@ struct ENGINE_API FInputKeyParams
 {
 	FInputKeyParams() = default;
 	
-	FInputKeyParams(FKey InKey, enum EInputEvent InEvent, FVector InDelta, bool bGamepadOverride = false)
+	FInputKeyParams(FKey InKey, enum EInputEvent InEvent, FVector InDelta, bool bGamepadOverride = false, FInputDeviceId InInputDevice = INPUTDEVICEID_NONE)
 		: Key(InKey)
+		, InputDevice(InInputDevice)
 		, Event(InEvent)
 		, Delta(InDelta)
 		, bIsGamepadOverride(bGamepadOverride)
 	{};
 
-	FInputKeyParams(FKey InKey, enum EInputEvent InEvent, double InDelta, bool bGamepadOverride = false)
+	FInputKeyParams(FKey InKey, enum EInputEvent InEvent, double InDelta, bool bGamepadOverride = false, FInputDeviceId InInputDevice = INPUTDEVICEID_NONE)
 		: Key(InKey)
+		, InputDevice(InInputDevice)
 		, Event(InEvent)
 		, Delta(FVector(InDelta, 0.0, 0.0))
 		, bIsGamepadOverride(bGamepadOverride)
 	{};
 	
-	FInputKeyParams(FKey InKey, double InDelta, float InDeltaTime, int32 InNumSamples, bool bGamepadOverride = false)
+	FInputKeyParams(FKey InKey, double InDelta, float InDeltaTime, int32 InNumSamples, bool bGamepadOverride = false, FInputDeviceId InInputDevice = INPUTDEVICEID_NONE)
 		: Key(InKey)
+		, InputDevice(InInputDevice)
 		, NumSamples(InNumSamples)
 		, DeltaTime(InDeltaTime)
 		, Delta(FVector(InDelta, 0.0, 0.0))
@@ -354,6 +357,9 @@ struct ENGINE_API FInputKeyParams
 	/** The key that has been pressed */
 	FKey Key = EKeys::Invalid;
 
+	/** The input device that has triggered this input */
+	FInputDeviceId InputDevice = INPUTDEVICEID_NONE;
+
 	/** The event that has caused a Button key to be considered */
 	enum EInputEvent Event = EInputEvent::IE_Pressed;
 
diff --git a/Engine/Source/Runtime/Engine/Private/PlayerController.cpp b/Engine/Source/Runtime/Engine/Private/PlayerController.cpp
--- a/Engine/Source/Runtime/Engine/Private/PlayerController.cpp
+++ b/Engine/Source/Runtime/Engine/Private/PlayerController.cpp
@@ -70,6 +70,7 @@
 #include "LevelUtils.h"
 #include "WorldPartition/WorldPartitionStreamingSource.h"
 #include "Physics/AsyncPhysicsInputComponent.h"
+#include "GenericPlatform/GenericPlatformInputDeviceMapper.h"
 
 DEFINE_LOG_CATEGORY(LogPlayerController);
 
@@ -2335,6 +2336,13 @@
 	// Any analog values can simply be passed to the UPlayerInput
 	if(Params.Key.IsAnalog())
 	{
+		// Only process the given input if it came from an input device that is owned by our owning local player
+		if (GetDefault<UInputSettings>()->bFilterInputByPlatformUser &&
+			IPlatformInputDeviceMapper::Get().GetUserForInputDevice(Params.InputDevice) != GetPlatformUserId())
+		{
+			return false;
+		}
+		
 		if(PlayerInput)
 		{
 			bResult = PlayerInput->InputKey(Params);
diff --git a/Engine/Source/Runtime/Engine/Private/UserInterface/InputSettings.cpp b/Engine/Source/Runtime/Engine/Private/UserInterface/InputSettings.cpp
--- a/Engine/Source/Runtime/Engine/Private/UserInterface/InputSettings.cpp
+++ b/Engine/Source/Runtime/Engine/Private/UserInterface/InputSettings.cpp
@@ -24,6 +24,7 @@ UInputSettings::UInputSettings(const FObjectInitializer& ObjectInitializer)
 	, bDefaultViewportMouseLock_DEPRECATED(false)
 	, bEnableLegacyInputScales(true)
 	, bEnableMotionControls(true)
+	, bFilterInputByPlatformUser(false)
 	, DefaultViewportMouseCaptureMode(EMouseCaptureMode::CapturePermanently_IncludingInitialMouseDown)
 	, DefaultViewportMouseLockMode(EMouseLockMode::LockOnCapture)
 	, DefaultPlayerInputClass(UPlayerInput::StaticClass())
