For a given cell, make sur to load the base level before loading any levels for data layers.

diff --git a/Engine/Source/Runtime/Engine/Private/WorldPartition/WorldPartitionRuntimeCell.cpp b/Engine/Source/Runtime/Engine/Private/WorldPartition/WorldPartitionRuntimeCell.cpp
--- a/Engine/Source/Runtime/Engine/Private/WorldPartition/WorldPartitionRuntimeCell.cpp
+++ b/Engine/Source/Runtime/Engine/Private/WorldPartition/WorldPartitionRuntimeCell.cpp
@@ -103,9 +103,21 @@ bool UWorldPartitionRuntimeCell::CacheStreamingSourceInfo(const UWorldPartitionR
 int32 UWorldPartitionRuntimeCell::SortCompare(const UWorldPartitionRuntimeCell* Other) const
 {
 	// Source priority (lower value is higher prio)
-	const int32 Comparison = (int32)CachedMinSourcePriority - (int32)Other->CachedMinSourcePriority;
+	int32 Comparison = (int32)CachedMinSourcePriority - (int32)Other->CachedMinSourcePriority;
+	
 	// Cell priority (lower value is higher prio)
-	return (Comparison != 0) ? Comparison : (Priority - Other->Priority);
+	if (Comparison == 0)
+	{
+		Comparison = Priority - Other->Priority;
+	}
+
+	// Data Layer priority (the more data layer a cell has, the later it should be streamed.)
+	if (Comparison == 0)
+	{
+		Comparison = DataLayers.Num() - Other->DataLayers.Num();
+	}
+
+	return Comparison;
 }
 
 bool UWorldPartitionRuntimeCell::IsDebugShown() const
