Always ResolveRuntimeDataDatalayer when registering an ActorDescView for streaming generation.
Fix non-dirty actors missing their datalayers when generating the streaming grid

diff --git a/Engine/Source/Runtime/Engine/Private/WorldPartition/WorldPartitionStreamingGeneration.cpp b/Engine/Source/Runtime/Engine/Private/WorldPartition/WorldPartitionStreamingGeneration.cpp
--- a/Engine/Source/Runtime/Engine/Private/WorldPartition/WorldPartitionStreamingGeneration.cpp
+++ b/Engine/Source/Runtime/Engine/Private/WorldPartition/WorldPartitionStreamingGeneration.cpp
@@ -51,17 +51,14 @@ class FWorldPartitionStreamingGenerator
 		auto GetModifiedActorDesc = [this](AActor* InActor) { return ModifiedActorsDescList->AddActor(InActor); };
 
 		// Adjust and register the actor descriptor view
-		auto RegisterActorDescView = [this, InContainer, &OutActorDescViewMap](const FGuid& ActorGuid, FWorldPartitionActorDescView& ActorDescView, bool bResolveRuntimeDataLayers)
+		auto RegisterActorDescView = [this, InContainer, &OutActorDescViewMap](const FGuid& ActorGuid, FWorldPartitionActorDescView& ActorDescView)
 		{
 			if (!bEnableStreaming)
 			{
 				ActorDescView.SetForcedNonSpatiallyLoaded();
 			}
 
-			if (bResolveRuntimeDataLayers)
-			{
-				ActorDescView.ResolveRuntimeDataLayers(InContainer);
-			}
+			ActorDescView.ResolveRuntimeDataLayers(InContainer);
 
 			ActorDescView.ResolveRuntimeReferences(InContainer);
 
@@ -87,14 +84,14 @@ class FWorldPartitionStreamingGenerator
 					{
 						// Dirty, unsaved actor for PIE
 						FWorldPartitionActorDescView ModifiedActorDescView = GetModifiedActorDesc(Actor);
-						RegisterActorDescView(ActorDescIt->GetGuid(), ModifiedActorDescView, true);
+						RegisterActorDescView(ActorDescIt->GetGuid(), ModifiedActorDescView);
 						continue;
 					}
 				}
 
 				// Non-dirty actor
 				FWorldPartitionActorDescView ActorDescView(*ActorDescIt);
-				RegisterActorDescView(ActorDescIt->GetGuid(), ActorDescView, false);
+				RegisterActorDescView(ActorDescIt->GetGuid(), ActorDescView);
 			}
 		}
 
@@ -106,7 +103,7 @@ class FWorldPartitionStreamingGenerator
 				if (IsValid(Actor) && Actor->IsPackageExternal() && Actor->IsMainPackageActor() && !Actor->IsEditorOnly() && !InContainer->GetActorDesc(Actor->GetActorGuid()))
 				{
 					FWorldPartitionActorDescView ModifiedActorDescView = GetModifiedActorDesc(Actor);
-					RegisterActorDescView(Actor->GetActorGuid(), ModifiedActorDescView, true);
+					RegisterActorDescView(Actor->GetActorGuid(), ModifiedActorDescView);
 				}
 			}
 		}
