WorldPartitonReplay:
- Allow recording on server
- Split Recording/Playback conditions
- Gate StreamingSourceName replication to Replay for clarity (already NetDriverName = NAME_DemoNetDriver)
- Enable in shipping build

diff --git a/Engine/Source/Runtime/Engine/Private/WorldPartition/WorldPartitionReplay.cpp b/Engine/Source/Runtime/Engine/Private/WorldPartition/WorldPartitionReplay.cpp
--- a/Engine/Source/Runtime/Engine/Private/WorldPartition/WorldPartitionReplay.cpp
+++ b/Engine/Source/Runtime/Engine/Private/WorldPartition/WorldPartitionReplay.cpp
@@ -17,6 +17,12 @@ static FAutoConsoleVariableRef CVarUseReplayStreamingSources(
 	GUseReplayStreamingSources,
 	TEXT("Set to 1 to use the recorded streaming sources when playing replay."));
 
+static int32 GRecordReplayStreamingSources = 1;
+static FAutoConsoleVariableRef CVarRecordReplayStreamingSources(
+	TEXT("wp.Runtime.RecordReplayStreamingSources"),
+	GRecordReplayStreamingSources,
+	TEXT("Set to 1 to record streaming sources when recording replay."));
+
 FArchive& operator<<(FArchive& Ar, FWorldPartitionReplayStreamingSource& StreamingSource)
 {
 	Ar << StreamingSource.Location;
@@ -93,7 +99,6 @@ void AWorldPartitionReplay::RewindForReplay()
 
 void AWorldPartitionReplay::Initialize(UWorld* World)
 {
-#if !UE_BUILD_SHIPPING
 	if (World->IsGameWorld())
 	{
 		UWorldPartition* WorldPartition = World->GetWorldPartition();
@@ -103,39 +108,39 @@ void AWorldPartitionReplay::Initialize(UWorld* World)
 		SpawnParams.Name = AWorldPartitionReplay::StaticClass()->GetFName();
 		WorldPartition->Replay = World->SpawnActor<AWorldPartitionReplay>(SpawnParams);
 	}
-#endif
 }
 
-bool AWorldPartitionReplay::IsEnabled(UWorld* World)
+bool AWorldPartitionReplay::IsPlaybackEnabled(UWorld* World)
 {
-#if !UE_BUILD_SHIPPING
-	if (UWorldPartition* WorldPartition = World->GetWorldPartition())
+	if (GUseReplayStreamingSources && World->IsPlayingReplay())
 	{
-		return WorldPartition->Replay && WorldPartition->Replay->IsEnabled();
+		if (UWorldPartition* WorldPartition = World->GetWorldPartition())
+		{
+			return !!WorldPartition->Replay;
+		}
 	}
-#endif
-
+	
 	return false;
 }
 
-void AWorldPartitionReplay::BeginPlay()
+bool AWorldPartitionReplay::IsRecordingEnabled(UWorld* World)
 {
-	Super::BeginPlay();
-	
-	UWorld* World = GetWorld();
-	check(World->IsGameWorld());
-	check(World->IsPartitionedWorld());
-	ENetMode NetMode = World->GetNetMode();
-	
-	// Recording will happen only in NM_Standalone/NM_Client. Replay can happen on any net mode.
-	bEnabled = (NetMode == NM_Standalone || NetMode == NM_Client || World->IsPlayingReplay());
+	if (GRecordReplayStreamingSources && World->IsRecordingReplay())
+	{
+		if (UWorldPartition* WorldPartition = World->GetWorldPartition())
+		{
+			return !!WorldPartition->Replay;
+		}
+	}
+		
+	return false;
 }
 
 void AWorldPartitionReplay::GetLifetimeReplicatedProps(TArray<FLifetimeProperty>& OutLifetimeProps) const
 {
 	Super::GetLifetimeReplicatedProps(OutLifetimeProps);
 
-	DOREPLIFETIME(AWorldPartitionReplay, StreamingSourceNames);
+	DOREPLIFETIME_CONDITION(AWorldPartitionReplay, StreamingSourceNames, COND_ReplayOnly);
 }
 
 void AWorldPartitionReplay::PreReplication(IRepChangedPropertyTracker& ChangedPropertyTracker)
@@ -162,12 +167,9 @@ void AWorldPartitionReplay::PreReplication(IRepChangedPropertyTracker& ChangedPr
 
 bool AWorldPartitionReplay::GetReplayStreamingSources(TArray<FWorldPartitionStreamingSource>& OutStreamingSources)
 {
-	if (!GUseReplayStreamingSources)
-	{
-		return false;
-	}
-
 	UWorld* World = GetWorld();
+	verify(IsPlaybackEnabled(World));
+
 	UDemoNetDriver* DemoNetDriver = World->GetDemoNetDriver();
 	if (DemoNetDriver && DemoNetDriver->IsPlaying())
 	{
diff --git a/Engine/Source/Runtime/Engine/Private/WorldPartition/WorldPartitionStreamingPolicy.cpp b/Engine/Source/Runtime/Engine/Private/WorldPartition/WorldPartitionStreamingPolicy.cpp
--- a/Engine/Source/Runtime/Engine/Private/WorldPartition/WorldPartitionStreamingPolicy.cpp
+++ b/Engine/Source/Runtime/Engine/Private/WorldPartition/WorldPartitionStreamingPolicy.cpp
@@ -94,16 +94,15 @@ void UWorldPartitionStreamingPolicy::UpdateStreamingSources()
 	}
 
 	const FTransform WorldToLocal = WorldPartition->GetInstanceTransform().Inverse();
-
-#if !UE_BUILD_SHIPPING
+	UWorld* World = WorldPartition->GetWorld();
 	bool bUseReplaySources = false;
-	if (AWorldPartitionReplay* Replay = WorldPartition->Replay)
+	if (AWorldPartitionReplay::IsPlaybackEnabled(World))
 	{
-		bUseReplaySources = Replay->GetReplayStreamingSources(StreamingSources);
+		bUseReplaySources = WorldPartition->Replay->GetReplayStreamingSources(StreamingSources);
 	}
 
-	if (!bUseReplaySources)
-#endif
+	const ENetMode NetMode = World->GetNetMode();
+	if (!bUseReplaySources && (NetMode == NM_Standalone || NetMode == NM_Client || AWorldPartitionReplay::IsRecordingEnabled(World)))
 	{
 #if WITH_EDITOR
 		// We are in the SIE
@@ -117,32 +116,27 @@ void UWorldPartitionStreamingPolicy::UpdateStreamingSources()
 		else
 #endif
 		{
-			UWorld* World = WorldPartition->GetWorld();
-			const ENetMode NetMode = World->GetNetMode();
-			if (NetMode == NM_Standalone || NetMode == NM_Client || AWorldPartitionReplay::IsEnabled(World))
+			const int32 NumPlayers = GEngine->GetNumGamePlayers(World);
+			for (int32 PlayerIndex = 0; PlayerIndex < NumPlayers; ++PlayerIndex)
 			{
-				const int32 NumPlayers = GEngine->GetNumGamePlayers(World);
-				for (int32 PlayerIndex = 0; PlayerIndex < NumPlayers; ++PlayerIndex)
+				ULocalPlayer* Player = GEngine->GetGamePlayer(World, PlayerIndex);
+				if (Player && Player->PlayerController && Player->PlayerController->IsStreamingSourceEnabled())
 				{
-					ULocalPlayer* Player = GEngine->GetGamePlayer(World, PlayerIndex);
-					if (Player && Player->PlayerController && Player->PlayerController->IsStreamingSourceEnabled())
-					{
-						FVector ViewLocation;
-						FRotator ViewRotation;
-						Player->PlayerController->GetPlayerViewPoint(ViewLocation, ViewRotation);
-
-						// Transform to Local
-						ViewLocation = WorldToLocal.TransformPosition(ViewLocation);
-						ViewRotation = WorldToLocal.TransformRotation(ViewRotation.Quaternion()).Rotator();
-						const EStreamingSourceTargetState TargetState = Player->PlayerController->StreamingSourceShouldActivate() ? EStreamingSourceTargetState::Activated : EStreamingSourceTargetState::Loaded;
-						const bool bBlockOnSlowLoading = Player->PlayerController->StreamingSourceShouldBlockOnSlowStreaming();
-						const EStreamingSourcePriority StreamingSourcePriority = Player->PlayerController->GetStreamingSourcePriority();
-						StreamingSources.Add(FWorldPartitionStreamingSource(Player->PlayerController->GetFName(), ViewLocation, ViewRotation, TargetState, bBlockOnSlowLoading, StreamingSourcePriority));
-					}
+					FVector ViewLocation;
+					FRotator ViewRotation;
+					Player->PlayerController->GetPlayerViewPoint(ViewLocation, ViewRotation);
+
+					// Transform to Local
+					ViewLocation = WorldToLocal.TransformPosition(ViewLocation);
+					ViewRotation = WorldToLocal.TransformRotation(ViewRotation.Quaternion()).Rotator();
+					const EStreamingSourceTargetState TargetState = Player->PlayerController->StreamingSourceShouldActivate() ? EStreamingSourceTargetState::Activated : EStreamingSourceTargetState::Loaded;
+					const bool bBlockOnSlowLoading = Player->PlayerController->StreamingSourceShouldBlockOnSlowStreaming();
+					const EStreamingSourcePriority StreamingSourcePriority = Player->PlayerController->GetStreamingSourcePriority();
+					StreamingSources.Add(FWorldPartitionStreamingSource(Player->PlayerController->GetFName(), ViewLocation, ViewRotation, TargetState, bBlockOnSlowLoading, StreamingSourcePriority));
 				}
 			}
 		}
-	
+
 		UWorldPartitionSubsystem* WorldPartitionSubsystem = WorldPartition->GetWorld()->GetSubsystem<UWorldPartitionSubsystem>();
 		check(WorldPartitionSubsystem);
 		for (IWorldPartitionStreamingSourceProvider* StreamingSourceProvider : WorldPartitionSubsystem->GetStreamingSourceProviders())
@@ -240,7 +234,7 @@ void UWorldPartitionStreamingPolicy::UpdateStreamingState()
 
 	const ENetMode NetMode = World->GetNetMode();
 	
-	bool bClient = NetMode == NM_Standalone || NetMode == NM_Client || AWorldPartitionReplay::IsEnabled(World);
+	bool bClient = NetMode == NM_Standalone || NetMode == NM_Client || AWorldPartitionReplay::IsPlaybackEnabled(World);
 	if (bClient)
 	{
 		// When world partition can't stream, all cells must be unloaded
diff --git a/Engine/Source/Runtime/Engine/Public/WorldPartition/WorldPartitionReplay.h b/Engine/Source/Runtime/Engine/Public/WorldPartition/WorldPartitionReplay.h
--- a/Engine/Source/Runtime/Engine/Public/WorldPartition/WorldPartitionReplay.h
+++ b/Engine/Source/Runtime/Engine/Public/WorldPartition/WorldPartitionReplay.h
@@ -51,13 +51,12 @@ class ENGINE_API AWorldPartitionReplay : public AActor
 
 public:
 	static void Initialize(UWorld* World);
-	static bool IsEnabled(UWorld* World);
+	static bool IsPlaybackEnabled(UWorld* World);
+	static bool IsRecordingEnabled(UWorld* World);
 
-	virtual void BeginPlay() override;
 	virtual void RewindForReplay() override;
 	virtual void PreReplication(IRepChangedPropertyTracker& ChangedPropertyTracker) override;
 		
-	bool IsEnabled() const { return bEnabled; }
 	bool GetReplayStreamingSources(TArray<FWorldPartitionStreamingSource>& OutStreamingSources);
 		
 private:
@@ -69,6 +68,5 @@ private:
 	TArray<FName> StreamingSourceNames;
 
 	TArray<FWorldPartitionReplaySample> ReplaySamples;
-	bool bEnabled = false;
 };
 
