Fix Shipping build

diff --git a/Engine/Source/Runtime/Engine/Private/WorldPartition/WorldPartition.cpp b/Engine/Source/Runtime/Engine/Private/WorldPartition/WorldPartition.cpp
--- a/Engine/Source/Runtime/Engine/Private/WorldPartition/WorldPartition.cpp
+++ b/Engine/Source/Runtime/Engine/Private/WorldPartition/WorldPartition.cpp
@@ -270,9 +270,7 @@ UWorldPartition::UWorldPartition(const FObjectInitializer& ObjectInitializer)
 #endif
 	, InitState(EWorldPartitionInitState::Uninitialized)
 	, StreamingPolicy(nullptr)
-#if !UE_BUILD_SHIPPING
 	, Replay(nullptr)
-#endif
 {
 #if WITH_EDITOR
 	WorldPartitionStreamingPolicyClass = UWorldPartitionLevelStreamingPolicy::StaticClass();
