Fix for Zen being started inappropriately to report stats.

diff --git a/Engine/Source/Runtime/Engine/Private/Analytics/StudioAnalytics.cpp b/Engine/Source/Runtime/Engine/Private/Analytics/StudioAnalytics.cpp
--- a/Engine/Source/Runtime/Engine/Private/Analytics/StudioAnalytics.cpp
+++ b/Engine/Source/Runtime/Engine/Private/Analytics/StudioAnalytics.cpp
@@ -295,6 +295,7 @@ void FStudioAnalytics::FireEvent_Loading(const FString& LoadingName, double Seco
 		}
 
 #if UE_WITH_ZEN
+		if (UE::Zen::IsDefaultServicePresent())
 		{
 			// Grab the Zen summary stats
 			UE::Zen::FZenStats ZenStats;
