Use proper validation for the Texture::GetRunningPlatformData() return values.

diff --git a/Engine/Source/Editor/TextureEditor/Private/Models/TextureEditorViewportClient.cpp b/Engine/Source/Editor/TextureEditor/Private/Models/TextureEditorViewportClient.cpp
--- a/Engine/Source/Editor/TextureEditor/Private/Models/TextureEditorViewportClient.cpp
+++ b/Engine/Source/Editor/TextureEditor/Private/Models/TextureEditorViewportClient.cpp
@@ -177,7 +177,8 @@ void FTextureEditorViewportClient::Draw(FViewport* Viewport, FCanvas* Canvas)
 		}
 	}
 
-	float Exposure = Texture->GetRunningPlatformData() && IsHDR((*Texture->GetRunningPlatformData())->PixelFormat) ? FMath::Pow(2.0f, (float)TextureEditorPinned->GetExposureBias()) : 1.0f;
+	FTexturePlatformData** RunningPlatformDataPtr = Texture->GetRunningPlatformData();
+	float Exposure = RunningPlatformDataPtr && *RunningPlatformDataPtr && IsHDR((*RunningPlatformDataPtr)->PixelFormat) ? FMath::Pow(2.0f, (float)TextureEditorPinned->GetExposureBias()) : 1.0f;
 
 	if ( Texture->GetResource() != nullptr )
 	{
diff --git a/Engine/Source/Editor/TextureEditor/Private/TextureEditorToolkit.cpp b/Engine/Source/Editor/TextureEditor/Private/TextureEditorToolkit.cpp
--- a/Engine/Source/Editor/TextureEditor/Private/TextureEditorToolkit.cpp
+++ b/Engine/Source/Editor/TextureEditor/Private/TextureEditorToolkit.cpp
@@ -2295,11 +2295,14 @@ bool FTextureEditorToolkit::HandleTextureBorderActionIsChecked( ) const
 
 EVisibility FTextureEditorToolkit::HandleExposureBiasWidgetVisibility() const
 {
-	if (Texture && Texture->GetRunningPlatformData() && IsHDR((*Texture->GetRunningPlatformData())->PixelFormat))
+	if (Texture)
 	{
-		return EVisibility::Visible;
+		FTexturePlatformData** RunningPlatformDataPtr = Texture->GetRunningPlatformData();
+		if (RunningPlatformDataPtr && *RunningPlatformDataPtr && IsHDR((*RunningPlatformDataPtr)->PixelFormat))
+		{
+			return EVisibility::Visible;
+		}
 	}
-
 	return EVisibility::Collapsed;
 }
 
diff --git a/Engine/Source/Runtime/Engine/Private/Streaming/StreamingManagerTexture.cpp b/Engine/Source/Runtime/Engine/Private/Streaming/StreamingManagerTexture.cpp
--- a/Engine/Source/Runtime/Engine/Private/Streaming/StreamingManagerTexture.cpp
+++ b/Engine/Source/Runtime/Engine/Private/Streaming/StreamingManagerTexture.cpp
@@ -2062,7 +2062,7 @@ bool FRenderAssetStreamingManager::HandleListStreamingRenderAssetsCommand( const
 		const int32 MaxAllowedMipIndex = ResourceState.LODCountToAssetFirstLODIdx(StreamingRenderAsset.MaxAllowedMips);
 
 		FTexturePlatformData** TexturePlatformData = Cast<UTexture>(RenderAsset) ? const_cast<UTexture*>(Cast<UTexture>(RenderAsset))->GetRunningPlatformData() : nullptr;
-		if (AssetType == EStreamableRenderAssetType::Texture && TexturePlatformData)
+		if (AssetType == EStreamableRenderAssetType::Texture && TexturePlatformData && *TexturePlatformData)
 		{
 			const UTexture2D* Texture2D = Cast<UTexture2D>(RenderAsset);
 			const UVolumeTexture* VolumeTexture = Cast<UVolumeTexture>(RenderAsset);
@@ -2443,7 +2443,7 @@ bool FRenderAssetStreamingManager::HandleLODGroupsCommand( const TCHAR* Cmd, FOu
 
 		const EPixelFormat PixelFormat = [&]()->EPixelFormat
 		{
-			if (Texture->GetRunningPlatformData())
+			if (Texture->GetRunningPlatformData() && *Texture->GetRunningPlatformData())
 			{
 				return (*Texture->GetRunningPlatformData())->PixelFormat;
 			}
@@ -2568,7 +2568,7 @@ bool FRenderAssetStreamingManager::HandleInvestigateRenderAssetCommand(const TCH
 				{
 					UE_LOG(LogContentStreaming, Log, TEXT("  LOD group:       %s [Unkown Bucket]"), *LODGroupName);
 				}
-				if (Texture && Texture->GetRunningPlatformData())
+				if (Texture && Texture->GetRunningPlatformData() && *Texture->GetRunningPlatformData())
 				{
 					UE_LOG(LogContentStreaming, Log, TEXT("  Format:          %s"), GPixelFormats[(*Texture->GetRunningPlatformData())->PixelFormat].Name);
 				}
@@ -2594,7 +2594,7 @@ bool FRenderAssetStreamingManager::HandleInvestigateRenderAssetCommand(const TCH
 					UE_LOG(LogContentStreaming, Log, TEXT("  Force all mips:  No mip-maps"));
 				}
 				
-				if (Texture && Texture->GetRunningPlatformData())
+				if (Texture && Texture->GetRunningPlatformData() && *Texture->GetRunningPlatformData())
 				{
 					const TIndirectArray<struct FTexture2DMipMap>& TextureMips = (*Texture->GetRunningPlatformData())->Mips;
 					UTexture2D* Texture2D = Cast<UTexture2D>(Texture);
