Respect current property when traversing a material expression chain through a MakeMaterialExpression node.
Fixes issue where VT texture not connected opacity mask is still preventing material from compiling.

diff --git a/Engine/Source/Runtime/Engine/Classes/Materials/Material.h b/Engine/Source/Runtime/Engine/Classes/Materials/Material.h
--- a/Engine/Source/Runtime/Engine/Classes/Materials/Material.h
+++ b/Engine/Source/Runtime/Engine/Classes/Materials/Material.h
@@ -1694,7 +1694,8 @@ protected:
 		ERHIFeatureLevel::Type InFeatureLevel = ERHIFeatureLevel::Num,
 		EMaterialQualityLevel::Type InQuality = EMaterialQualityLevel::Num,
 		ERHIShadingPath::Type InShadingPath = ERHIShadingPath::Num,
-		EShaderFrequency InShaderFrequency = SF_NumFrequencies);
+		EShaderFrequency InShaderFrequency = SF_NumFrequencies, 
+		EMaterialProperty InProperty = MP_MAX);
 
 	/**
 	*	Recursively update the bRealtimePreview for each expression based on whether it is connected to something that is time-varying.
diff --git a/Engine/Source/Runtime/Engine/Classes/Materials/MaterialExpressionMakeMaterialAttributes.h b/Engine/Source/Runtime/Engine/Classes/Materials/MaterialExpressionMakeMaterialAttributes.h
--- a/Engine/Source/Runtime/Engine/Classes/Materials/MaterialExpressionMakeMaterialAttributes.h
+++ b/Engine/Source/Runtime/Engine/Classes/Materials/MaterialExpressionMakeMaterialAttributes.h
@@ -71,6 +71,9 @@ class UMaterialExpressionMakeMaterialAttributes : public UMaterialExpression
 	UPROPERTY()
 	FExpressionInput ShadingModel;
 
+	/** Get the input for a material property. Returns nullptr if the property isn't supported. */
+	FExpressionInput* GetExpressionInput(EMaterialProperty InProperty);
+
 	//~ Begin UObject Interface
 	virtual void Serialize(FStructuredArchive::FRecord Record) override;
 	//~ End UObject Interface
diff --git a/Engine/Source/Runtime/Engine/Private/Materials/Material.cpp b/Engine/Source/Runtime/Engine/Private/Materials/Material.cpp
--- a/Engine/Source/Runtime/Engine/Private/Materials/Material.cpp
+++ b/Engine/Source/Runtime/Engine/Private/Materials/Material.cpp
@@ -5371,7 +5371,7 @@ bool UMaterial::GetExpressionsInPropertyChain(EMaterialProperty InProperty,
 			ShaderFrequency = FMaterialAttributeDefinitionMap::GetShaderFrequency(InProperty);
 		}
 
-		RecursiveGetExpressionChain(StartingExpression->Expression, ProcessedInputs, OutExpressions, InStaticParameterSet, InFeatureLevel, InQuality, InShadingPath, ShaderFrequency);
+		RecursiveGetExpressionChain(StartingExpression->Expression, ProcessedInputs, OutExpressions, InStaticParameterSet, InFeatureLevel, InQuality, InShadingPath, ShaderFrequency, InProperty);
 	}
 	return true;
 }
@@ -5436,9 +5436,10 @@ bool UMaterial::GetTexturesInPropertyChain(EMaterialProperty InProperty, TArray<
 	return false;
 }
 
-bool UMaterial::RecursiveGetExpressionChain(UMaterialExpression* InExpression, TArray<FExpressionInput*>& InOutProcessedInputs, 
+bool UMaterial::RecursiveGetExpressionChain(
+	UMaterialExpression* InExpression, TArray<FExpressionInput*>& InOutProcessedInputs, 
 	TArray<UMaterialExpression*>& OutExpressions, struct FStaticParameterSet* InStaticParameterSet,
-	ERHIFeatureLevel::Type InFeatureLevel, EMaterialQualityLevel::Type InQuality, ERHIShadingPath::Type InShadingPath, EShaderFrequency InShaderFrequency)
+	ERHIFeatureLevel::Type InFeatureLevel, EMaterialQualityLevel::Type InQuality, ERHIShadingPath::Type InShadingPath, EShaderFrequency InShaderFrequency, EMaterialProperty InProperty)
 {
 	OutExpressions.AddUnique(InExpression);
 	TArray<FExpressionInput*> Inputs;
@@ -5502,19 +5503,53 @@ bool UMaterial::RecursiveGetExpressionChain(UMaterialExpression* InExpression, T
 			InputsFrequency.Add(InShaderFrequency);
 		}
 	}
-	else if (InFeatureLevel <= ERHIFeatureLevel::ES3_1 && (MakeMaterialAttributesExp = Cast<UMaterialExpressionMakeMaterialAttributes>(InExpression)) != nullptr)
+	else if ((MakeMaterialAttributesExp = Cast<UMaterialExpressionMakeMaterialAttributes>(InExpression)) != nullptr)
 	{
-		// Follow only mobile-relevant inputs
-		Inputs.Add(&MakeMaterialAttributesExp->EmissiveColor);
-		InputsFrequency.Add(SF_Pixel);
-		Inputs.Add(&MakeMaterialAttributesExp->OpacityMask);
-		InputsFrequency.Add(SF_Pixel);
-		Inputs.Add(&MakeMaterialAttributesExp->BaseColor);
-		InputsFrequency.Add(SF_Pixel);
-		Inputs.Add(&MakeMaterialAttributesExp->Normal);
-		InputsFrequency.Add(SF_Pixel);
-		Inputs.Add(&MakeMaterialAttributesExp->WorldPositionOffset);
-		InputsFrequency.Add(SF_Vertex);
+		// Filtered list of properties that we follow on mobile platforms.
+		struct FMobileProperty 
+		{
+			EMaterialProperty Property;
+			EShaderFrequency Frequency;
+		};
+		static const FMobileProperty MobileProperties[] =
+		{
+			{ MP_WorldPositionOffset, SF_Vertex },
+			{ MP_EmissiveColor, SF_Pixel },
+			{ MP_BaseColor, SF_Pixel },
+			{ MP_Normal, SF_Pixel },
+			{ MP_OpacityMask, SF_Pixel },
+		};
+
+		const bool bMobileOnly = InFeatureLevel <= ERHIFeatureLevel::ES3_1;
+		
+		if (bMobileOnly)
+		{
+			// Follow mobile properties (and only InProperty if that is specified).
+			for (FMobileProperty MobileProperty : MobileProperties)
+			{
+				if (InProperty == MP_MAX || InProperty == MobileProperty.Property)
+				{
+					Inputs.Add(MakeMaterialAttributesExp->GetExpressionInput(MobileProperty.Property));
+					InputsFrequency.Add(MobileProperty.Frequency);
+				}
+			}
+		}
+		else if (InProperty != MP_MAX)
+		{
+			// Only follow the specified InProperty.
+			FExpressionInput* Input = MakeMaterialAttributesExp->GetExpressionInput(InProperty);
+			if (Input != nullptr)
+			{
+				Inputs.Add(Input);
+				InputsFrequency.Add(InShaderFrequency);
+			}
+		}
+		else
+		{
+			// Follow all properties.
+			Inputs = InExpression->GetInputs();
+			InputsFrequency.Init(InShaderFrequency, Inputs.Num());
+		}
 	}
 	else
 	{
@@ -5583,7 +5618,7 @@ bool UMaterial::RecursiveGetExpressionChain(UMaterialExpression* InExpression, T
 					if (bProcessInput == true)
 					{
 						InOutProcessedInputs.Add(InnerInput);
-						RecursiveGetExpressionChain(InnerInput->Expression, InOutProcessedInputs, OutExpressions, InStaticParameterSet, InFeatureLevel, InQuality, InShadingPath, InputsFrequency[InputIdx]);
+						RecursiveGetExpressionChain(InnerInput->Expression, InOutProcessedInputs, OutExpressions, InStaticParameterSet, InFeatureLevel, InQuality, InShadingPath, InputsFrequency[InputIdx], InProperty);
 					}
 				}
 			}
