Fix static analysis warning (A = B > C is evaluated as A = (B > C), which was not the intent).

diff --git a/Engine/Source/Runtime/Engine/Private/AudioDerivedData.cpp b/Engine/Source/Runtime/Engine/Private/AudioDerivedData.cpp
--- a/Engine/Source/Runtime/Engine/Private/AudioDerivedData.cpp
+++ b/Engine/Source/Runtime/Engine/Private/AudioDerivedData.cpp
@@ -1352,7 +1352,8 @@ static void CookSimpleWave(const FAudioCookInputs& Inputs, TArray<uint8>& Output
 	}
 
 	// Clip Normalize
-	if (const float MaxValue = Audio::ArrayMaxAbsValue(InputFloatBuffer) > 1.0f)
+	const float MaxValue = Audio::ArrayMaxAbsValue(InputFloatBuffer);
+	if (MaxValue > 1.0f)
 	{
 		UE_LOG(LogAudioDerivedData, Display, TEXT("Audio clipped during cook: This asset will be normalized by a factor of 1/%f. Consider attenuating the above asset."), MaxValue);
 
@@ -1588,7 +1589,8 @@ static void CookSurroundWave(const FAudioCookInputs& Inputs,  TArray<uint8>& Out
 		}
 
 		// clip normalize
-		if (const float MaxValue = Audio::ArrayMaxAbsValue(InterleavedFloatBuffer) > 1.0f)
+		const float MaxValue = Audio::ArrayMaxAbsValue(InterleavedFloatBuffer);
+		if (MaxValue > 1.0f)
 		{
 			UE_LOG(LogAudioDerivedData, Display, TEXT("Audio clipped during cook: This asset will be normalized by a factor of 1/%f. Consider attenuating the above asset."), MaxValue);
 
