Remove null checks causing false positives in static analysis. EditorOnly is already checked for null above in the function, and earlier code already tries to dereference it.

diff --git a/Engine/Source/Runtime/Engine/Private/Materials/Material.cpp b/Engine/Source/Runtime/Engine/Private/Materials/Material.cpp
--- a/Engine/Source/Runtime/Engine/Private/Materials/Material.cpp
+++ b/Engine/Source/Runtime/Engine/Private/Materials/Material.cpp
@@ -3421,31 +3421,31 @@ void UMaterial::PostLoad()
 	DoMaterialAttributeReorder(&EditorOnly->ShadingModelFromMaterialExpression, UEVer, RenderObjVer, UE5MainVer);
 	DoMaterialAttributeReorder(&EditorOnly->FrontMaterial, UEVer, RenderObjVer, UE5MainVer);
 
-	if (EditorOnly && ParameterGroupData_DEPRECATED.Num() > 0)
+	if (ParameterGroupData_DEPRECATED.Num() > 0)
 	{
 		ensure(EditorOnly->ParameterGroupData.Num() == 0);
 		EditorOnly->ParameterGroupData = MoveTemp(ParameterGroupData_DEPRECATED);
 	}
 
-	if (EditorOnly && Expressions_DEPRECATED.Num() > 0)
+	if (Expressions_DEPRECATED.Num() > 0)
 	{
 		ensure(EditorOnly->ExpressionCollection.Expressions.Num() == 0);
 		EditorOnly->ExpressionCollection.Expressions = MoveTemp(Expressions_DEPRECATED);
 	}
 
-	if (EditorOnly && EditorComments_DEPRECATED.Num() > 0)
+	if (EditorComments_DEPRECATED.Num() > 0)
 	{
 		ensure(EditorOnly->ExpressionCollection.EditorComments.Num() == 0);
 		EditorOnly->ExpressionCollection.EditorComments = MoveTemp(EditorComments_DEPRECATED);
 	}
 
-	if (EditorOnly && ExpressionExecBegin_DEPRECATED)
+	if (ExpressionExecBegin_DEPRECATED)
 	{
 		ensure(!EditorOnly->ExpressionCollection.ExpressionExecBegin);
 		EditorOnly->ExpressionCollection.ExpressionExecBegin = MoveTemp(ExpressionExecBegin_DEPRECATED);
 	}
 
-	if (EditorOnly && ExpressionExecEnd_DEPRECATED)
+	if (ExpressionExecEnd_DEPRECATED)
 	{
 		ensure(!EditorOnly->ExpressionCollection.ExpressionExecEnd);
 		EditorOnly->ExpressionCollection.ExpressionExecEnd = MoveTemp(ExpressionExecEnd_DEPRECATED);
