Level Instance: Check Local role only in Game worlds

diff --git a/Engine/Source/Runtime/Engine/Private/LevelInstance/LevelInstanceActor.cpp b/Engine/Source/Runtime/Engine/Private/LevelInstance/LevelInstanceActor.cpp
--- a/Engine/Source/Runtime/Engine/Private/LevelInstance/LevelInstanceActor.cpp
+++ b/Engine/Source/Runtime/Engine/Private/LevelInstance/LevelInstanceActor.cpp
@@ -69,7 +69,7 @@ void ALevelInstance::PostRegisterAllComponents()
 {
 	Super::PostRegisterAllComponents();
 	
-	if (GetLocalRole() == ENetRole::ROLE_Authority)
+	if (GetLocalRole() == ENetRole::ROLE_Authority && GetWorld()->IsGameWorld())
 	{
 #if !WITH_EDITOR
 		// If the level instance was spawned, not loaded
diff --git a/Engine/Source/Runtime/Engine/Private/LevelInstance/LevelInstanceActorGuid.cpp b/Engine/Source/Runtime/Engine/Private/LevelInstance/LevelInstanceActorGuid.cpp
--- a/Engine/Source/Runtime/Engine/Private/LevelInstance/LevelInstanceActorGuid.cpp
+++ b/Engine/Source/Runtime/Engine/Private/LevelInstance/LevelInstanceActorGuid.cpp
@@ -29,7 +29,7 @@ const FGuid& FLevelInstanceActorGuid::GetGuid_Internal() const
 {
 	check(Actor);
 #if WITH_EDITOR
-	const FGuid& Guid = Actor->GetLocalRole() == ENetRole::ROLE_Authority ? Actor->GetActorGuid() : ActorGuid;
+	const FGuid& Guid = (!Actor->GetWorld()->IsGameWorld() || Actor->GetLocalRole() == ENetRole::ROLE_Authority) ? Actor->GetActorGuid() : ActorGuid;
 #else
 	const FGuid& Guid = ActorGuid;
 #endif
