Fix non-unity error.

diff --git a/Engine/Source/Runtime/Engine/Classes/EdGraph/EdGraphSchema.h b/Engine/Source/Runtime/Engine/Classes/EdGraph/EdGraphSchema.h
--- a/Engine/Source/Runtime/Engine/Classes/EdGraph/EdGraphSchema.h
+++ b/Engine/Source/Runtime/Engine/Classes/EdGraph/EdGraphSchema.h
@@ -16,6 +16,7 @@
 #include "EdGraphSchema.generated.h"
 
 class FSlateRect;
+struct FSlateBrush;
 class UEdGraph;
 struct FBPVariableDescription;
 
