- Allow LevelInstanceSubsystem to be created for inactive worlds, to make it available during cook (LevelInstanceSubsystem::Tick is disabled when worldType is inactive)
- Fixes crash during cook when LevelInstanceActorDesc::RegisterContainerInstance tries to access the LevelInstanceSubsystem.

diff --git a/Engine/Source/Runtime/Engine/Private/LevelInstance/LevelInstanceSubsystem.cpp b/Engine/Source/Runtime/Engine/Private/LevelInstance/LevelInstanceSubsystem.cpp
--- a/Engine/Source/Runtime/Engine/Private/LevelInstance/LevelInstanceSubsystem.cpp
+++ b/Engine/Source/Runtime/Engine/Private/LevelInstance/LevelInstanceSubsystem.cpp
@@ -94,7 +94,7 @@ void ULevelInstanceSubsystem::Initialize(FSubsystemCollectionBase& Collection)
 
 bool ULevelInstanceSubsystem::DoesSupportWorldType(EWorldType::Type WorldType) const
 {
-	return Super::DoesSupportWorldType(WorldType) || WorldType == EWorldType::EditorPreview;
+	return Super::DoesSupportWorldType(WorldType) || WorldType == EWorldType::EditorPreview || WorldType == EWorldType::Inactive;
 }
 
 ILevelInstanceInterface* ULevelInstanceSubsystem::GetLevelInstance(FLevelInstanceID LevelInstanceID) const
@@ -466,6 +466,11 @@ ILevelInstanceInterface* ULevelInstanceSubsystem::GetOwningLevelInstance(const U
 
 void ULevelInstanceSubsystem::Tick()
 {
+	if (GetWorld()->WorldType == EWorldType::Inactive)
+	{
+		return;
+	}
+
 	// For non-game world, Tick is responsible of processing LevelInstances to update/load/unload
 	if (!GetWorld()->IsGameWorld())
 	{
