- Fix incorrect LOD calculation in GetFirstValidLODIdx

#9245

diff --git a/Engine/Source/Runtime/Engine/Private/StaticMesh.cpp b/Engine/Source/Runtime/Engine/Private/StaticMesh.cpp
--- a/Engine/Source/Runtime/Engine/Private/StaticMesh.cpp
+++ b/Engine/Source/Runtime/Engine/Private/StaticMesh.cpp
@@ -1896,21 +1896,10 @@ int32 FStaticMeshRenderData::GetFirstValidLODIdx(int32 MinIdx) const
 	if (LODCount > 0)
 	{
 		LODIndex = FMath::Clamp<int32>(MinIdx, 0, LODCount - 1);
-		if (FPlatformProperties::RequiresCookedData())
-		{
-			// When cooked, BuffersSize can be used reliably to know if there is valid data, even though num vertices could be non zero.
-			// This happens because UStaticMesh::MinLOD is platform specific, where as num vertices is the same for all platforms.
-			while (LODIndex < LODCount && !LODResources[LODIndex].GetNumVertices() && LODResources[LODIndex].BuffersSize)
-			{
-				++LODIndex;
-			}
-		}
-		else
+
+		while (LODIndex < LODCount && (LODResources[LODIndex].GetNumVertices() == 0 || LODResources[LODIndex].BuffersSize == 0))
 		{
-			while (LODIndex < LODCount && !LODResources[LODIndex].GetNumVertices())
-			{
-				++LODIndex;
-			}
+			++LODIndex;
 		}
 
 		if (LODIndex >= LODCount)
