UX - When enabling streaming for the 1st time on Template projects, open the documentation for the user (if the user wants) instead of having to copy/paste, trim the info. This also works with locale/territories.

#rnx

diff --git a/Engine/Source/Editor/WorldPartitionEditor/Private/WorldPartition/Customizations/WorldPartitionDetailsCustomization.cpp b/Engine/Source/Editor/WorldPartitionEditor/Private/WorldPartition/Customizations/WorldPartitionDetailsCustomization.cpp
--- a/Engine/Source/Editor/WorldPartitionEditor/Private/WorldPartition/Customizations/WorldPartitionDetailsCustomization.cpp
+++ b/Engine/Source/Editor/WorldPartitionEditor/Private/WorldPartition/Customizations/WorldPartitionDetailsCustomization.cpp
@@ -13,6 +13,7 @@
 #include "Widgets/Input/SButton.h"
 #include "Widgets/Input/SCheckBox.h"
 #include "Widgets/Input/SNumericEntryBox.h"
+#include "IDocumentation.h"
 
 #define LOCTEXT_NAMESPACE "FWorldPartitionDetails"
 
@@ -73,8 +74,12 @@ void FWorldPartitionDetails::HandleWorldPartitionEnableStreamingChanged(ECheckBo
 			{
 				return;
 			}
-			
-			FMessageDialog::Open(EAppMsgType::Ok, LOCTEXT("WorldPartitionEnableStreamingDialolg", "Please refer to https://docs.unrealengine.com/5.0/en-US/world-partition-in-unreal-engine/ for how to set up streaming."));
+
+			if (FMessageDialog::Open(EAppMsgType::YesNo, LOCTEXT("WorldPartitionEnableStreamingDialolg", "Please refer to our documentation for how to set up streaming.\n\nWould you like to open it now? ")) == EAppReturnType::Yes)
+			{
+				IDocumentation::Get()->Open(TEXT("world-partition-in-unreal-engine"), FDocumentationSourceInfo(TEXT("worldpartition")));
+			}
+
 			WorldPartition->bStreamingWasEnabled = true;
 		}
 
diff --git a/Engine/Source/Runtime/Engine/Private/WorldPartition/WorldPartition.cpp b/Engine/Source/Runtime/Engine/Private/WorldPartition/WorldPartition.cpp
--- a/Engine/Source/Runtime/Engine/Private/WorldPartition/WorldPartition.cpp
+++ b/Engine/Source/Runtime/Engine/Private/WorldPartition/WorldPartition.cpp
@@ -51,6 +51,7 @@
 #include "WorldPartition/HLOD/HLODLayer.h"
 #include "Modules/ModuleManager.h"
 #include "GameDelegates.h"
+#include "IDocumentation.h"
 
 static int32 GLoadingRangeBugItGo = 12800;
 static FAutoConsoleVariableRef CVarLoadingRangeBugItGo(
@@ -1153,7 +1154,10 @@ void UWorldPartition::Tick(float DeltaSeconds)
 						{
 							SetEnableStreaming(true);
 							bStreamingWasEnabled = true;
-							FMessageDialog::Open(EAppMsgType::Ok, LOCTEXT("WorldPartitionEnableStreamingDialolg", "Please refer to https://docs.unrealengine.com/5.0/en-US/world-partition-in-unreal-engine/ for how to set up streaming."));			
+							if (FMessageDialog::Open(EAppMsgType::YesNo, LOCTEXT("WorldPartitionEnableStreamingDialolg", "Please refer to our documentation for how to set up streaming.\n\nWould you like to open it now? ")) == EAppReturnType::Yes)
+							{
+								IDocumentation::Get()->Open(TEXT("world-partition-in-unreal-engine"), FDocumentationSourceInfo(TEXT("worldpartition")));
+							}
 						}
 
 						break;
@@ -1458,4 +1462,4 @@ FBox UWorldPartition::GetEditorWorldBounds() const
 }
 #endif
 
-#undef LOCTEXT_NAMESPACE
+#undef LOCTEXT_NAMESPACE
\ No newline at end of file
