World Partition: filter out invalid actor handles for loader list interfaces.

#rnx

diff --git a/Engine/Source/Runtime/Engine/Private/WorldPartition/LoaderAdapter/LoaderAdapterList.cpp b/Engine/Source/Runtime/Engine/Private/WorldPartition/LoaderAdapter/LoaderAdapterList.cpp
--- a/Engine/Source/Runtime/Engine/Private/WorldPartition/LoaderAdapter/LoaderAdapterList.cpp
+++ b/Engine/Source/Runtime/Engine/Private/WorldPartition/LoaderAdapter/LoaderAdapterList.cpp
@@ -11,7 +11,10 @@ void FLoaderAdapterList::ForEachActor(TFunctionRef<void(const FWorldPartitionHan
 {
 	for (const FWorldPartitionHandle& Actor : Actors)
 	{
-		InOperation(Actor);
+		if (Actor.IsValid())
+		{
+			InOperation(Actor);
+		}
 	}
 }
 #endif
