World Partition Streaming Generation:
- Added console command `wp.Editor.DumpStreamingGenerationLog` to dump the streaming generation log
- Fixed runtime data layers and references resolving to happen on the actor descriptor views instead of the container to include unsaved actors.
- After actor descriptor views map is created, there should be no resolving happening using the source actor descriptor containers.

#rnx

diff --git a/Engine/Source/Runtime/Engine/Private/WorldPartition/DataLayer/DataLayerSubsystem.cpp b/Engine/Source/Runtime/Engine/Private/WorldPartition/DataLayer/DataLayerSubsystem.cpp
--- a/Engine/Source/Runtime/Engine/Private/WorldPartition/DataLayer/DataLayerSubsystem.cpp
+++ b/Engine/Source/Runtime/Engine/Private/WorldPartition/DataLayer/DataLayerSubsystem.cpp
@@ -2,6 +2,7 @@
 
 #include "WorldPartition/DataLayer/DataLayerSubsystem.h"
 #include "WorldPartition/DataLayer/WorldDataLayers.h"
+#include "WorldPartition/DataLayer/WorldDataLayersActorDesc.h"
 #include "WorldPartition/DataLayer/DataLayerAsset.h"
 #include "WorldPartition/WorldPartitionDebugHelper.h"
 #include "WorldPartition/WorldPartitionSubsystem.h"
@@ -640,7 +641,26 @@ void UDataLayerSubsystem::OnActorDescContainerInitialized(UActorDescContainer* I
 
 	UE_SCOPED_TIMER(*FString::Printf(TEXT("Resolving Data Layer Instance Names for %s"), *InActorDescContainer->GetContainerPackage().ToString()), LogWorldPartition, Display);
 
-	const FWorldDataLayersActorDesc* WorldDataLayersActorDesc = FDataLayerUtils::GetWorldDataLayersActorDesc(InActorDescContainer);
+	auto GetWorldDataLayersActorDesc = [](const UActorDescContainer* InContainer) -> FWorldDataLayersActorDesc*
+	{
+		if (InContainer)
+		{
+			for (FActorDescList::TIterator<AWorldDataLayers> Iterator(const_cast<UActorDescContainer*>(InContainer)); Iterator; ++Iterator)
+			{
+				// IsValid will return false for maps that don't have a valid WorldDataLayers actors
+				if (Iterator->IsValid())
+				{
+					FWorldDataLayersActorDesc* WorldDataLayersActorDesc = *Iterator;
+					return WorldDataLayersActorDesc;
+				}
+				// No need to iterate (we assume that there's only 1 AWorldDataLayer for now)
+				return nullptr;
+			}
+		}
+		return nullptr;
+	};
+
+	const FWorldDataLayersActorDesc* WorldDataLayersActorDesc = GetWorldDataLayersActorDesc(InActorDescContainer);
 	for (FActorDescList::TIterator<> Iterator(InActorDescContainer); Iterator; ++Iterator)
 	{
 		FWorldPartitionActorDesc* ActorDesc = *Iterator;
diff --git a/Engine/Source/Runtime/Engine/Private/WorldPartition/DataLayer/DataLayerUtils.cpp b/Engine/Source/Runtime/Engine/Private/WorldPartition/DataLayer/DataLayerUtils.cpp
--- a/Engine/Source/Runtime/Engine/Private/WorldPartition/DataLayer/DataLayerUtils.cpp
+++ b/Engine/Source/Runtime/Engine/Private/WorldPartition/DataLayer/DataLayerUtils.cpp
@@ -11,24 +11,6 @@
 #include "WorldPartition/ActorDescContainer.h"
 #include "WorldPartition/DataLayer/DataLayerSubsystem.h"
 
-FWorldDataLayersActorDesc* FDataLayerUtils::GetWorldDataLayersActorDesc(const UActorDescContainer* InContainer, bool bInCheckValid)
-{
-	if (InContainer)
-	{
-		for (FActorDescList::TIterator<AWorldDataLayers> Iterator(const_cast<UActorDescContainer*>(InContainer)); Iterator; ++Iterator)
-		{
-			if (!bInCheckValid || Iterator->IsValid())
-			{
-				FWorldDataLayersActorDesc* WorldDataLayersActorDesc = *Iterator;
-				return WorldDataLayersActorDesc;
-			}
-			// No need to iterate (we assume that there's only 1 AWorldDataLayer for now)
-			return nullptr;
-		}
-	}
-	return nullptr;
-}
-
 TArray<FName> FDataLayerUtils::ResolvedDataLayerInstanceNames(const FWorldPartitionActorDesc* InActorDesc, const FWorldDataLayersActorDesc* InWorldDataLayersActorDesc, UWorld* InWorld, bool* bOutIsResultValid)
 {
 	bool bLocalIsSuccess = true;
@@ -114,49 +96,35 @@ TArray<FName> FDataLayerUtils::ResolvedDataLayerInstanceNames(const FWorldPartit
 }
 
 // For performance reasons, this function assumes that InActorDesc's DataLayerInstanceNames was already resolved.
-TArray<FName> FDataLayerUtils::ResolveRuntimeDataLayerInstanceNames(const FWorldPartitionActorDesc* InActorDesc, const UActorDescContainer* InContainer, bool* bOutIsResultValid)
+bool FDataLayerUtils::ResolveRuntimeDataLayerInstanceNames(const FWorldPartitionActorDescView& InActorDescView, const TMap<FGuid, FWorldPartitionActorDescView>& ActorDescViews, TArray<FName>& OutRuntimeDataLayerInstanceNames)
 {
-	check(InContainer);
-
-	bool bLocalIsSuccess = false;
-	bool& bIsSuccess = bOutIsResultValid ? *bOutIsResultValid : bLocalIsSuccess;
-	bIsSuccess = false;
-
-	TArray<FName> Result;
-	bIsSuccess = InActorDesc->GetDataLayerInstanceNames().IsEmpty();
-	if (bIsSuccess)
+	if (InActorDescView.GetDataLayerInstanceNames().IsEmpty())
 	{
-		return Result;
+		return true;
 	}
 
-	UWorld* World = InContainer->GetWorld();
-	const UDataLayerSubsystem* DataLayerSubsystem = World ? World->GetSubsystem<UDataLayerSubsystem>() : nullptr;
-	if (DataLayerSubsystem && DataLayerSubsystem->CanResolveDataLayers())
+	for (auto& [ActorGuid, ActorDescView] : ActorDescViews)
 	{
-		bIsSuccess = true;
-		for (FName DataLayerInstanceName : InActorDesc->GetDataLayerInstanceNames())
+		if (ActorDescView.GetActorNativeClass()->IsChildOf<AWorldDataLayers>())
 		{
-			const UDataLayerInstance* DataLayerInstance = DataLayerSubsystem->GetDataLayerInstance(DataLayerInstanceName);
-			if (DataLayerInstance && DataLayerInstance->IsRuntime())
-			{
-				Result.Add(DataLayerInstanceName);
-			}
-		}
-	}
-	// Fallback on FWorldDataLayersActorDesc
-	else if (FWorldDataLayersActorDesc* WorldDataLayersActorDesc = FDataLayerUtils::GetWorldDataLayersActorDesc(InContainer))
-	{
-		bIsSuccess = true;
-		for (FName DataLayerInstanceName : InActorDesc->GetDataLayerInstanceNames())
-		{
-			const FDataLayerInstanceDesc* DataLayerInstanceDesc = WorldDataLayersActorDesc->GetDataLayerInstanceFromInstanceName(DataLayerInstanceName);
-			if (DataLayerInstanceDesc && (DataLayerInstanceDesc->GetDataLayerType() == EDataLayerType::Runtime))
+			FWorldDataLayersActorDesc* WorldDataLayersActorDesc = (FWorldDataLayersActorDesc*)ActorDescView.GetActorDesc();
+
+			for (FName DataLayerInstanceName : InActorDescView.GetDataLayerInstanceNames())
 			{
-				Result.Add(DataLayerInstanceName);
+				if (const FDataLayerInstanceDesc* DataLayerInstanceDesc = WorldDataLayersActorDesc->GetDataLayerInstanceFromInstanceName(DataLayerInstanceName))
+				{
+					if (DataLayerInstanceDesc->GetDataLayerType() == EDataLayerType::Runtime)
+					{
+						OutRuntimeDataLayerInstanceNames.Add(DataLayerInstanceName);
+					}
+				}
 			}
+
+			return true;
 		}
 	}
-	return Result;
+
+	return false;
 }
 
 #endif
\ No newline at end of file
diff --git a/Engine/Source/Runtime/Engine/Private/WorldPartition/WorldPartition.cpp b/Engine/Source/Runtime/Engine/Private/WorldPartition/WorldPartition.cpp
--- a/Engine/Source/Runtime/Engine/Private/WorldPartition/WorldPartition.cpp
+++ b/Engine/Source/Runtime/Engine/Private/WorldPartition/WorldPartition.cpp
@@ -371,10 +371,8 @@ void UWorldPartition::OnCancelPIE()
 
 void UWorldPartition::OnEndPlay()
 {
-	RuntimeHash->FlushStreaming();
+	FlushStreaming();
 	RuntimeHash->OnEndPlay();
-
-	StreamingPolicy = nullptr;
 }
 
 FName UWorldPartition::GetWorldPartitionEditorName() const
diff --git a/Engine/Source/Runtime/Engine/Private/WorldPartition/WorldPartitionActorDescView.cpp b/Engine/Source/Runtime/Engine/Private/WorldPartition/WorldPartitionActorDescView.cpp
--- a/Engine/Source/Runtime/Engine/Private/WorldPartition/WorldPartitionActorDescView.cpp
+++ b/Engine/Source/Runtime/Engine/Private/WorldPartition/WorldPartitionActorDescView.cpp
@@ -6,7 +6,6 @@
 #include "WorldPartition/WorldPartitionLog.h"
 #include "WorldPartition/WorldPartitionActorDesc.h"
 #include "WorldPartition/ActorDescContainer.h"
-#include "WorldPartition/DataLayer/DataLayerUtils.h"
 
 FWorldPartitionActorDescView::FWorldPartitionActorDescView()
 	: FWorldPartitionActorDescView(nullptr)
@@ -19,16 +18,6 @@ FWorldPartitionActorDescView::FWorldPartitionActorDescView(const FWorldPartition
 	, bInvalidRuntimeGrid(false)
 {}
 
-void FWorldPartitionActorDescView::ResolveRuntimeDataLayers(const UActorDescContainer* InContainer)
-{
-	bool bSuccess = true;
-	RuntimeDataLayers = FDataLayerUtils::ResolveRuntimeDataLayerInstanceNames(ActorDesc, InContainer, &bSuccess);
-	if (!bSuccess)
-	{
-		RuntimeDataLayers.Reset();
-	}
-}
-
 const FGuid& FWorldPartitionActorDescView::GetGuid() const
 {
 	return ActorDesc->GetGuid();
@@ -95,12 +84,21 @@ const TArray<FName>& FWorldPartitionActorDescView::GetDataLayers() const
 	return bInvalidDataLayers ? EmptyDataLayers : ActorDesc->GetDataLayerInstanceNames();
 }
 
+const TArray<FName>& FWorldPartitionActorDescView::GetDataLayerInstanceNames() const
+{
+	return ActorDesc->GetDataLayerInstanceNames();
+}
 const TArray<FName>& FWorldPartitionActorDescView::GetRuntimeDataLayers() const
 {
 	static TArray<FName> EmptyDataLayers;
 	return (bInvalidDataLayers || !RuntimeDataLayers.IsSet()) ? EmptyDataLayers : RuntimeDataLayers.GetValue();
 }
 
+const TArray<FName>& FWorldPartitionActorDescView::GetTags() const
+{
+	return ActorDesc->GetTags();
+}
+
 FName FWorldPartitionActorDescView::GetActorPackage() const
 {
 	return ActorDesc->GetActorPackage();
@@ -180,27 +178,18 @@ void FWorldPartitionActorDescView::SetInvalidDataLayers()
 	}
 }
 
-bool FWorldPartitionActorDescView::IsResaveNeeded() const
+void FWorldPartitionActorDescView::SetRuntimeDataLayers(TArray<FName>& InRuntimeDataLayers)
 {
-	return ActorDesc->IsResaveNeeded();
+	RuntimeDataLayers = InRuntimeDataLayers;
 }
 
-void FWorldPartitionActorDescView::ResolveRuntimeReferences(const UActorDescContainer* InContainer)
+void FWorldPartitionActorDescView::SetRuntimeReferences(TArray<FGuid>& InRuntimeReferences)
 {
-	RuntimeReferences.Emplace();
-	for (const FGuid& ReferenceGuid : ActorDesc->GetReferences())
-	{
-		if (const FWorldPartitionActorDesc* ReferenceDesc = InContainer->GetActorDesc(ReferenceGuid))
-		{
-			if (ReferenceDesc->IsLoaded() ? !ReferenceDesc->GetActor()->IsEditorOnly() : !ReferenceDesc->GetActorIsEditorOnly())
-			{
-				RuntimeReferences->Add(ReferenceGuid);
-			}
-			else
-			{
-				UE_LOG(LogWorldPartition, Verbose, TEXT("Actor reference '%s' from '%s' filtered-out (editor-only)"), *ReferenceDesc->GetActorLabelOrName().ToString(), *GetActorLabelOrName().ToString());
-			}
-		}
-	}
+	RuntimeReferences = InRuntimeReferences;
+}
+
+bool FWorldPartitionActorDescView::IsResaveNeeded() const
+{
+	return ActorDesc->IsResaveNeeded();
 }
 #endif
\ No newline at end of file
diff --git a/Engine/Source/Runtime/Engine/Private/WorldPartition/WorldPartitionStreamingGeneration.cpp b/Engine/Source/Runtime/Engine/Private/WorldPartition/WorldPartitionStreamingGeneration.cpp
--- a/Engine/Source/Runtime/Engine/Private/WorldPartition/WorldPartitionStreamingGeneration.cpp
+++ b/Engine/Source/Runtime/Engine/Private/WorldPartition/WorldPartitionStreamingGeneration.cpp
@@ -3,6 +3,7 @@
 #include "WorldPartition/WorldPartition.h"
 
 #if WITH_EDITOR
+#include "Editor.h"
 #include "Engine/World.h"
 #include "Engine/LevelScriptBlueprint.h"
 #include "ActorReferencesUtils.h"
@@ -10,6 +11,7 @@
 #include "WorldPartition/WorldPartitionStreamingPolicy.h"
 #include "WorldPartition/WorldPartitionActorCluster.h"
 #include "WorldPartition/DataLayer/DataLayerSubsystem.h"
+#include "WorldPartition/DataLayer/DataLayerUtils.h"
 #include "WorldPartition/ErrorHandling/WorldPartitionStreamingGenerationNullErrorHandler.h"
 #include "WorldPartition/ErrorHandling/WorldPartitionStreamingGenerationLogErrorHandler.h"
 #include "WorldPartition/ErrorHandling/WorldPartitionStreamingGenerationMapCheckErrorHandler.h"
@@ -18,6 +20,22 @@
 
 #define LOCTEXT_NAMESPACE "WorldPartition"
 
+static FAutoConsoleCommand DumpStreamingGenerationLog(
+	TEXT("wp.Editor.DumpStreamingGenerationLog"),
+	TEXT("Dump the streaming generation log."),
+	FConsoleCommandWithArgsDelegate::CreateLambda([](const TArray<FString>& Args)
+	{
+		if (UWorld* World = GEditor->GetEditorWorldContext().World(); World && !World->IsGameWorld())
+		{
+			if (UWorldPartition* WorldPartition = World->GetWorldPartition())
+			{
+				WorldPartition->GenerateStreaming();
+				WorldPartition->FlushStreaming();
+			}
+		}
+	})
+);
+
 /*
 	Preparation Phase
 		Actor Descriptor Views Creation
@@ -36,6 +54,37 @@
 
 class FWorldPartitionStreamingGenerator
 {
+	void ResolveRuntimeDataLayers(FWorldPartitionActorDescView& ActorDescView, const TMap<FGuid, FWorldPartitionActorDescView>& ActorDescViews)
+	{
+		TArray<FName> RuntimeDataLayerInstanceNames;
+		RuntimeDataLayerInstanceNames.Reserve(ActorDescView.GetDataLayers().Num());
+
+		if (FDataLayerUtils::ResolveRuntimeDataLayerInstanceNames(ActorDescView, ActorDescViews, RuntimeDataLayerInstanceNames))
+		{
+			ActorDescView.SetRuntimeDataLayers(RuntimeDataLayerInstanceNames);
+		}
+	}
+
+	void ResolveRuntimeReferences(FWorldPartitionActorDescView& ActorDescView, const TMap<FGuid, FWorldPartitionActorDescView>& ActorDescViews)
+	{
+		TArray<FGuid> RuntimeReferences;
+		RuntimeReferences.Reserve(ActorDescView.GetReferences().Num());
+
+		for (const FGuid& ReferenceGuid : ActorDescView.GetReferences())
+		{
+			if (const FWorldPartitionActorDescView* ReferenceDescView = ActorDescViews.Find(ReferenceGuid))
+			{
+				check(!ReferenceDescView->GetActorIsEditorOnly());
+				RuntimeReferences.Add(ReferenceGuid);
+			}
+		}
+
+		if (RuntimeReferences.Num() != ActorDescView.GetReferences().Num())
+		{
+			ActorDescView.SetRuntimeReferences(RuntimeReferences);
+		}
+	}
+
 	void CreateActorDescViewMap(const UActorDescContainer* InContainer, TMap<FGuid, FWorldPartitionActorDescView>& OutActorDescViewMap, const FActorContainerID& InContainerID)
 	{
 		// should we handle unsaved or newly created actors?
@@ -45,30 +94,38 @@ class FWorldPartitionStreamingGenerator
 		const bool bIsTempContainerPackage = FPackageName::IsTempPackage(InContainer->GetPackage()->GetName());
 		
 		// Test whether an actor is editor only. Will fallback to the actor descriptor only if the actor is not loaded
-		auto IsActorEditorOnly = [](const FWorldPartitionActorDesc* ActorDesc) { return ActorDesc->IsLoaded() ? ActorDesc->GetActor()->IsEditorOnly() : ActorDesc->GetActorIsEditorOnly(); };
+		auto IsActorEditorOnly = [](const FWorldPartitionActorDesc* ActorDesc, const FActorContainerID& ContainerID)
+		{
+			if (ActorDesc->IsRuntimeRelevant(ContainerID))
+			{
+				if (ActorDesc->IsLoaded())
+				{
+					return ActorDesc->GetActor()->IsEditorOnly();
+				}
+				else
+				{
+					return ActorDesc->GetActorIsEditorOnly();
+				}
+			}
+			return false;
+		};
 
 		// Create an actor descriptor view for the specified actor (modified or unsaved actors)
-		auto GetModifiedActorDesc = [this](AActor* InActor) { return ModifiedActorsDescList->AddActor(InActor); };
+		auto GetModifiedActorDesc = [this](AActor* InActor)
+		{
+			return ModifiedActorsDescList->AddActor(InActor);
+		};
 
 		// Adjust and register the actor descriptor view
 		auto RegisterActorDescView = [this, InContainer, &OutActorDescViewMap](const FGuid& ActorGuid, FWorldPartitionActorDescView& ActorDescView)
 		{
-			if (!bEnableStreaming)
-			{
-				ActorDescView.SetForcedNonSpatiallyLoaded();
-			}
-
-			ActorDescView.ResolveRuntimeDataLayers(InContainer);
-
-			ActorDescView.ResolveRuntimeReferences(InContainer);
-
 			OutActorDescViewMap.Emplace(ActorGuid, ActorDescView);
 		};
 		
 		TMap<FGuid, FGuid> ContainerGuidsRemap;
 		for (FActorDescList::TConstIterator<> ActorDescIt(InContainer); ActorDescIt; ++ActorDescIt)
 		{
-			if (!IsActorEditorOnly(*ActorDescIt) && ActorDescIt->IsRuntimeRelevant(InContainerID))
+			if (!IsActorEditorOnly(*ActorDescIt, InContainerID))
 			{
 				// Handle unsaved actors
 				if (AActor* Actor = ActorDescIt->GetActor())
@@ -171,6 +228,26 @@ class FWorldPartitionStreamingGenerator
 	void CreateActorDescriptorViews(const UActorDescContainer* InContainer)
 	{
 		CreateActorDescriptorViewsRecursive(InContainer, FTransform::Identity, TSet<FName>(), FActorContainerID(), FActorContainerID(), EContainerClusterMode::Partitioned, TEXT("MainContainer"));
+
+		// Once all views are created
+		for (auto It = ContainerDescriptorsMap.CreateIterator(); It; ++It)
+		{
+			const FActorContainerID& ContainerID = It.Key();
+			FContainerDescriptor& ContainerDescriptor = It.Value();
+
+			for (auto ActorDescIt = ContainerDescriptor.ActorDescViewMap.CreateIterator(); ActorDescIt; ++ActorDescIt)
+			{
+				FWorldPartitionActorDescView& ActorDescView = ActorDescIt.Value();
+
+				if (!bEnableStreaming)
+				{
+					ActorDescView.SetForcedNonSpatiallyLoaded();
+				}
+
+				ResolveRuntimeDataLayers(ActorDescView, ContainerDescriptor.ActorDescViewMap);
+				ResolveRuntimeReferences(ActorDescView, ContainerDescriptor.ActorDescViewMap);
+			}
+		}
 	}
 
 	/** 
@@ -615,7 +692,7 @@ bool UWorldPartition::GenerateStreaming(TArray<FString>* OutPackagesToGenerate)
 	}
 
 	// Dump state log
-	const TCHAR* StateLogSuffix = bIsPIE ? TEXT("PIE") : (IsRunningGame() ? TEXT("Game") : (IsRunningCookCommandlet() ? TEXT("Cook") : TEXT("Unknown")));
+	const TCHAR* StateLogSuffix = bIsPIE ? TEXT("PIE") : (IsRunningGame() ? TEXT("Game") : (IsRunningCookCommandlet() ? TEXT("Cook") : TEXT("Manual")));
 	TUniquePtr<FArchive> LogFileAr = FWorldPartitionStreamingGenerator::CreateDumpStateLogArchive(StateLogSuffix);
 	FHierarchicalLogArchive HierarchicalLogAr(*LogFileAr);
 
@@ -651,6 +728,12 @@ bool UWorldPartition::GenerateStreaming(TArray<FString>* OutPackagesToGenerate)
 	return false;
 }
 
+void UWorldPartition::FlushStreaming()
+{
+	RuntimeHash->FlushStreaming();
+	StreamingPolicy = nullptr;
+}
+
 void UWorldPartition::GenerateHLOD(ISourceControlHelper* SourceControlHelper, bool bCreateActorsOnly)
 {
 	FStreamingGenerationLogErrorHandler LogErrorHandler;
@@ -686,4 +769,4 @@ void UWorldPartition::CheckForErrors(IStreamingGenerationErrorHandler* ErrorHand
 }
 #endif // WITH_EDITOR
 
-#undef LOCTEXT_NAMESPACE
+#undef LOCTEXT_NAMESPACE
\ No newline at end of file
diff --git a/Engine/Source/Runtime/Engine/Public/WorldPartition/DataLayer/DataLayerUtils.h b/Engine/Source/Runtime/Engine/Public/WorldPartition/DataLayer/DataLayerUtils.h
--- a/Engine/Source/Runtime/Engine/Public/WorldPartition/DataLayer/DataLayerUtils.h
+++ b/Engine/Source/Runtime/Engine/Public/WorldPartition/DataLayer/DataLayerUtils.h
@@ -7,6 +7,7 @@
 #include "UObject/UObjectGlobals.h"
 #include "WorldPartition/DataLayer/DataLayerType.h"
 #include "WorldPartition/DataLayer/DataLayerInstance.h"
+#include "WorldPartition/WorldPartitionActorDescView.h"
 
 class FWorldDataLayersActorDesc;
 class FWorldPartitionActorDesc;
@@ -22,9 +23,10 @@ public:
 		static constexpr const TCHAR* IconNameByType[static_cast<int>(EDataLayerType::Size)] = { TEXT("DataLayer.Runtime") , TEXT("DataLayer.Editor"), TEXT("") };
 		return IconNameByType[static_cast<uint32>(DataLayerType)];
 	}
-	static FWorldDataLayersActorDesc* GetWorldDataLayersActorDesc(const UActorDescContainer* InContainer, bool bInCheckValid = true);
+
 	static TArray<FName> ResolvedDataLayerInstanceNames(const FWorldPartitionActorDesc* InActorDesc, const FWorldDataLayersActorDesc* InWorldDataLayersActorDesc = nullptr, UWorld* InWorld = nullptr, bool* bOutIsResultValid = nullptr);
-	static TArray<FName> ResolveRuntimeDataLayerInstanceNames(const FWorldPartitionActorDesc* InActorDesc, const UActorDescContainer* InContainer, bool* bOutIsResultValid = nullptr);
+	
+	static bool ResolveRuntimeDataLayerInstanceNames(const FWorldPartitionActorDescView& InActorDescView, const TMap<FGuid, FWorldPartitionActorDescView>& ActorDescViews, TArray<FName>& OutRuntimeDataLayerInstanceNames);
 #endif
 
 #if DATALAYER_TO_INSTANCE_RUNTIME_CONVERSION_ENABLED
diff --git a/Engine/Source/Runtime/Engine/Public/WorldPartition/WorldPartition.h b/Engine/Source/Runtime/Engine/Public/WorldPartition/WorldPartition.h
--- a/Engine/Source/Runtime/Engine/Public/WorldPartition/WorldPartition.h
+++ b/Engine/Source/Runtime/Engine/Public/WorldPartition/WorldPartition.h
@@ -115,7 +115,7 @@ private:
 	void OnGCPostReachabilityAnalysis();
 	void OnPackageDirtyStateChanged(UPackage* Package);
 
-	// PIE/Game Methods
+	// PIE/Game
 	void OnPreBeginPIE(bool bStartSimulate);
 	void OnPrePIEEnded(bool bWasSimulatingInEditor);
 	void OnCancelPIE();
@@ -145,11 +145,13 @@ public:
 #if WITH_EDITOR
 	FName GetWorldPartitionEditorName() const;
 
-	// PIE/Game/Cook Methods
+	// Streaming generation
 	bool GenerateStreaming(TArray<FString>* OutPackagesToGenerate = nullptr);
+	void FlushStreaming();
+
 	void RemapSoftObjectPath(FSoftObjectPath& ObjectPath);
 
-	// Cook Methods
+	// Cooking
 	bool PopulateGeneratedPackageForCook(UPackage* InPackage, const FString& InPackageRelativePath);
 	bool FinalizeGeneratorPackageForCook(const TArray<ICookPackageSplitter::FGeneratedPackageForPreSave>& InGeneratedPackages);
 
@@ -160,7 +162,7 @@ public:
 	void SetDefaultHLODLayer(UHLODLayer* InDefaultHLODLayer) { DefaultHLODLayer = InDefaultHLODLayer; }
 	void GenerateHLOD(ISourceControlHelper* SourceControlHelper, bool bCreateActorsOnly);
 
-	// Debugging Methods
+	// Debugging
 	void DrawRuntimeHashPreview();
 	void DumpActorDescs(const FString& Path);
 
@@ -195,7 +197,7 @@ public:
 
 	const TArray<FWorldPartitionStreamingSource>& GetStreamingSources() const;
 
-	// Debugging Methods
+	// Debugging
 	bool CanDrawRuntimeHash() const;
 	FVector2D GetDrawRuntimeHash2DDesiredFootprint(const FVector2D& CanvasSize);
 	void DrawRuntimeHash2D(UCanvas* Canvas, const FVector2D& PartitionCanvasSize, const FVector2D& Offset);
diff --git a/Engine/Source/Runtime/Engine/Public/WorldPartition/WorldPartitionActorDescView.h b/Engine/Source/Runtime/Engine/Public/WorldPartition/WorldPartitionActorDescView.h
--- a/Engine/Source/Runtime/Engine/Public/WorldPartition/WorldPartitionActorDescView.h
+++ b/Engine/Source/Runtime/Engine/Public/WorldPartition/WorldPartitionActorDescView.h
@@ -34,7 +34,9 @@ public:
 	bool GetActorIsHLODRelevant() const;
 	FName GetHLODLayer() const;
 	const TArray<FName>& GetDataLayers() const;
+	const TArray<FName>& GetDataLayerInstanceNames() const;
 	const TArray<FName>& GetRuntimeDataLayers() const;
+	const TArray<FName>& GetTags() const;
 	FName GetActorPackage() const;
 	FName GetActorPath() const;
 	FName GetActorLabel() const;
@@ -50,15 +52,14 @@ public:
 	FName GetActorLabelOrName() const;
 
 	void SetForcedNonSpatiallyLoaded();
-
 	void SetInvalidRuntimeGrid();
-
 	void SetInvalidDataLayers();
+	void SetRuntimeDataLayers(TArray<FName>& InRuntimeDataLayers);
+	void SetRuntimeReferences(TArray<FGuid>& InRuntimeReferences);
 
 	bool IsResaveNeeded() const;
 
 	void ResolveRuntimeDataLayers(const UActorDescContainer* InContainer);
-	void ResolveRuntimeReferences(const UActorDescContainer* InContainer);
 
 	bool operator==(const FWorldPartitionActorDescView& Other) const
 	{
@@ -70,6 +71,8 @@ public:
 		return GetTypeHash(Key.GetGuid());
 	}
 
+	const FWorldPartitionActorDesc* GetActorDesc() const { return ActorDesc; }
+
 protected:
 	const FWorldPartitionActorDesc* ActorDesc;
 	bool bIsForcedNonSpatiallyLoaded;
