Avoid physics replication relying on stale API for num predicted frames

diff --git a/Engine/Source/Runtime/Engine/Private/PhysicsEngine/PhysicsReplication.cpp b/Engine/Source/Runtime/Engine/Private/PhysicsEngine/PhysicsReplication.cpp
--- a/Engine/Source/Runtime/Engine/Private/PhysicsEngine/PhysicsReplication.cpp
+++ b/Engine/Source/Runtime/Engine/Private/PhysicsEngine/PhysicsReplication.cpp
@@ -521,9 +521,9 @@ void FPhysicsReplication::OnTick(float DeltaSeconds, TMap<TWeakObjectPtr<UPrimit
 				Chaos::FPhysicsSolver* Solver = PhysScene->GetSolver();
 				check(Solver);
 
-				const APlayerController::FClientFrameInfo& ClientFrameInfo = PlayerController->GetClientFrameInfo();
-				LocalFrameOffset = ClientFrameInfo.GetLocalFrameOffset();
-				NumPredictedFrames = Solver->GetCurrentFrame() - ClientFrameInfo.LastProcessedInputFrame;
+				//LocalFrameOffset = PlayerController->GetLocalToServerAsyncPhysicsTickOffset();
+				//TODO: as we send physics updates down we need to record latest seen
+				//NumPredictedFrames = Solver->GetCurrentFrame() - ClientFrameInfo.LastProcessedInputFrame;
 			}
 		}
 	}
