Location Volume: exposed loading, unloading and loaded state query functions to blueprints and in detail panels.

#rnx

diff --git a/Engine/Source/Editor/WorldPartitionEditor/Private/WorldPartition/SWorldPartitionEditorGrid2D.cpp b/Engine/Source/Editor/WorldPartitionEditor/Private/WorldPartition/SWorldPartitionEditorGrid2D.cpp
--- a/Engine/Source/Editor/WorldPartitionEditor/Private/WorldPartition/SWorldPartitionEditorGrid2D.cpp
+++ b/Engine/Source/Editor/WorldPartitionEditor/Private/WorldPartition/SWorldPartitionEditorGrid2D.cpp
@@ -318,7 +318,9 @@ void SWorldPartitionEditorGrid2D::UnloadSelectedRegions()
 
 	ForEachSelectedLoaderAdapters(WorldPartition, SelectBox, [&LoaderAdaptersToRelease, &RegisteredLoaderAdapters](IWorldPartitionActorLoaderInterface::ILoaderAdapter* LoaderAdapter)
 	{
-		if (LoaderAdapter->Unload() && RegisteredLoaderAdapters.Contains(LoaderAdapter))
+		LoaderAdapter->Unload();
+
+		if (RegisteredLoaderAdapters.Contains(LoaderAdapter))
 		{
 			LoaderAdaptersToRelease.Add(LoaderAdapter);
 		}
diff --git a/Engine/Source/Runtime/Engine/Private/LocationVolume.cpp b/Engine/Source/Runtime/Engine/Private/LocationVolume.cpp
--- a/Engine/Source/Runtime/Engine/Private/LocationVolume.cpp
+++ b/Engine/Source/Runtime/Engine/Private/LocationVolume.cpp
@@ -66,4 +66,37 @@ IWorldPartitionActorLoaderInterface::ILoaderAdapter* ALocationVolume::GetLoaderA
 {
 	return WorldPartitionActorLoader;
 }
+
+void ALocationVolume::Load()
+{
+#if WITH_EDITOR
+	if (WorldPartitionActorLoader)
+	{
+		WorldPartitionActorLoader->Load();
+	}
+#endif
+}
+
+void ALocationVolume::Unload()
+{
+#if WITH_EDITOR
+	if (WorldPartitionActorLoader)
+	{
+		WorldPartitionActorLoader->Unload();
+	}
+#endif
+}
+
+bool ALocationVolume::IsLoaded() const
+{
+#if WITH_EDITOR
+	if (WorldPartitionActorLoader)
+	{
+		return WorldPartitionActorLoader->IsLoaded();
+	}
+#endif
+
+	return false;
+}
+
 #endif
\ No newline at end of file
diff --git a/Engine/Source/Runtime/Engine/Private/WorldPartition/WorldPartitionActorLoaderInterface.cpp b/Engine/Source/Runtime/Engine/Private/WorldPartition/WorldPartitionActorLoaderInterface.cpp
--- a/Engine/Source/Runtime/Engine/Private/WorldPartition/WorldPartitionActorLoaderInterface.cpp
+++ b/Engine/Source/Runtime/Engine/Private/WorldPartition/WorldPartitionActorLoaderInterface.cpp
@@ -35,7 +35,7 @@ IWorldPartitionActorLoaderInterface::ILoaderAdapter::~ILoaderAdapter()
 	UnregisterDelegates();
 }
 
-bool IWorldPartitionActorLoaderInterface::ILoaderAdapter::Load()
+void IWorldPartitionActorLoaderInterface::ILoaderAdapter::Load()
 {
 	if (!bLoaded)
 	{
@@ -43,11 +43,9 @@ bool IWorldPartitionActorLoaderInterface::ILoaderAdapter::Load()
 		RefreshLoadedState();
 		RegisterDelegates();
 	}
-
-	return true;
 }
 
-bool IWorldPartitionActorLoaderInterface::ILoaderAdapter::Unload()
+void IWorldPartitionActorLoaderInterface::ILoaderAdapter::Unload()
 {
 	if (bLoaded && !IsEngineExitRequested())
 	{
@@ -75,8 +73,6 @@ bool IWorldPartitionActorLoaderInterface::ILoaderAdapter::Unload()
 			PostLoadedStateChanged(true);
 		}
 	}
-
-	return true;
 }
 
 bool IWorldPartitionActorLoaderInterface::ILoaderAdapter::IsLoaded() const
diff --git a/Engine/Source/Runtime/Engine/Public/LocationVolume.h b/Engine/Source/Runtime/Engine/Public/LocationVolume.h
--- a/Engine/Source/Runtime/Engine/Public/LocationVolume.h
+++ b/Engine/Source/Runtime/Engine/Public/LocationVolume.h
@@ -13,7 +13,7 @@ class FLoaderAdapterActor;
 /**
  * A volume representing a location in the world
  */
-UCLASS()
+UCLASS(BlueprintType, Blueprintable)
 class ENGINE_API ALocationVolume : public AVolume, public IWorldPartitionActorLoaderInterface
 {
 	GENERATED_UCLASS_BODY()
@@ -31,11 +31,23 @@ public:
 	virtual void PostRegisterAllComponents();
 	virtual bool CanChangeIsSpatiallyLoadedFlag() const override { return false; }
 	//~ End AActor Interface
+#endif
 
 	//~ Begin IWorldPartitionActorLoaderInterface interface
 	virtual ILoaderAdapter* GetLoaderAdapter() override;
 	//~ End IWorldPartitionActorLoaderInterface interface
-#endif
+
+	/** Load this location volume */
+	UFUNCTION(BlueprintCallable, CallInEditor, Category=WorldPartition)
+	void Load();
+
+	/** Unload this location volume */
+	UFUNCTION(BlueprintCallable, CallInEditor, Category=WorldPartition)
+	void Unload();
+
+	/** Return if this location volume is loaded */
+	UFUNCTION(BlueprintCallable, Category=WorldPartition)
+	bool IsLoaded() const;
 
 	UPROPERTY(EditAnywhere, Category=LocationVolume)
 	FColor DebugColor;
diff --git a/Engine/Source/Runtime/Engine/Public/WorldPartition/WorldPartitionActorLoaderInterface.h b/Engine/Source/Runtime/Engine/Public/WorldPartition/WorldPartitionActorLoaderInterface.h
--- a/Engine/Source/Runtime/Engine/Public/WorldPartition/WorldPartitionActorLoaderInterface.h
+++ b/Engine/Source/Runtime/Engine/Public/WorldPartition/WorldPartitionActorLoaderInterface.h
@@ -27,8 +27,8 @@
 		ILoaderAdapter(UWorld* InWorld);
 		virtual ~ILoaderAdapter();
 
-		bool Load();
-		bool Unload();
+		void Load();
+		void Unload();
 		bool IsLoaded() const;
 
 		UWorld* GetWorld() const { return World; }
