Correct location of #if WITH_EDITOR
#rnx

diff --git a/Engine/Source/Runtime/Engine/Private/LocationVolume.cpp b/Engine/Source/Runtime/Engine/Private/LocationVolume.cpp
--- a/Engine/Source/Runtime/Engine/Private/LocationVolume.cpp
+++ b/Engine/Source/Runtime/Engine/Private/LocationVolume.cpp
@@ -66,6 +66,7 @@ IWorldPartitionActorLoaderInterface::ILoaderAdapter* ALocationVolume::GetLoaderA
 {
 	return WorldPartitionActorLoader;
 }
+#endif
 
 void ALocationVolume::Load()
 {
@@ -97,6 +98,4 @@ bool ALocationVolume::IsLoaded() const
 #endif
 
 	return false;
-}
-
-#endif
\ No newline at end of file
+}
\ No newline at end of file
