Fix non-unity
#rnx

diff --git a/Engine/Source/Runtime/Engine/Private/LevelInstance/LevelInstanceActorGuid.cpp b/Engine/Source/Runtime/Engine/Private/LevelInstance/LevelInstanceActorGuid.cpp
--- a/Engine/Source/Runtime/Engine/Private/LevelInstance/LevelInstanceActorGuid.cpp
+++ b/Engine/Source/Runtime/Engine/Private/LevelInstance/LevelInstanceActorGuid.cpp
@@ -2,6 +2,7 @@
 
 #include "LevelInstance/LevelInstanceActorGuid.h"
 #include "GameFramework/Actor.h"
+#include "Engine/World.h"
 
 #if !WITH_EDITOR
 void FLevelInstanceActorGuid::AssignIfInvalid()
