Fix bug where Replicated Arrays may not receive RepNotifies until their size changed.

diff --git a/Engine/Source/Runtime/Engine/Private/RepLayout.cpp b/Engine/Source/Runtime/Engine/Private/RepLayout.cpp
--- a/Engine/Source/Runtime/Engine/Private/RepLayout.cpp
+++ b/Engine/Source/Runtime/Engine/Private/RepLayout.cpp
@@ -77,6 +77,9 @@ static FAutoConsoleVariableRef CVarShareInitialCompareState(TEXT("net.ShareIniti
 bool GbTrackNetSerializeObjectReferences = false;
 static FAutoConsoleVariableRef CVarTrackNetSerializeObjectReferences(TEXT("net.TrackNetSerializeObjectReferences"), GbTrackNetSerializeObjectReferences, TEXT("If true, we will create small layouts for Net Serialize Structs if they have Object Properties. This can prevent some Shadow State GC crashes."));
 
+bool GbWithArrayOnRepFix = true;
+static FAutoConsoleVariableRef CVarWithArrayOnRepFix(TEXT("net.WithArrayOnRepFix"), GbWithArrayOnRepFix, TEXT("If true, attempt to prevent issues with Arrays not receiving OnRep calls until their size changes if their Archetypes have different values from instances in levels."));
+
 #if WITH_PUSH_VALIDATION_SUPPORT
 
 static bool GbPushModelValidateProperties = false;
@@ -3355,6 +3358,41 @@ static FGuidReferencesMap* PrepReceivedArray(
 			UE_CLOG(LogSkippedRepNotifies> 0, LogRep, Display, TEXT("1 FReceivedPropertiesStackState Skipping RepNotify for property %s because local value has not changed."), *Cmd.Property->GetName());
 		}
 		check(ShadowArray != nullptr);
+
+		// This should only happen if
+		//	A. Clients have modified the array locally.
+		//	B. The Archetype has different values in its array than the actual object.
+		//		This can happen if we have an instance of an Actor / Object in a level,
+		//		the array property is EditAnywhere,
+		//		and items have been added to or removed from the array. 
+		//
+		//	In the case of b, it's possible that our ShadowArray has never allocated its dynamic
+		//	memory, and that will cause us problems detecting changes to items in the data array.
+		//	As a stopgap, we will resize the shadow array to match the size of the current,
+		//	pre-replicated size of the array.
+		//
+		//	We don't need to resize the ShadowArray to the newly replicated size of the array,
+		//	because if we detect changes (either in the size, like above, or in any elements),
+		//	then during CallRepNotifies, we will copy the entire state of the DataArray (live array) to the
+		//	ShadowArray.
+		//
+		//	This still leaves us with a bit of a problem, however.
+		//	If the Archetype's array is the same size as the instance's array, but the values are different,
+		//	AND the user has defined an OnRep function that takes the Old Array as a parameter,
+		//	when the OnRep is called, it will have correct "old" values for any of the properties that
+		//	were replicated (because we will copy them prior to applying the replicated data),
+		//	but any properties that were NOT replicated will instead have the Archetype's values.
+		//
+		//	This would only happen the first time the OnRep was called, however, because we will
+		//	copy the complete value of the array after every OnRep (see CallRepNotifies).
+		//
+		// Note, this is behind a CVar because it changes behavior, and that could cause subtle bugs
+		// in game logic.
+		if (GbWithArrayOnRepFix && ShadowArray->Num() != DataArray->Num())
+		{
+			Parent.Property->CopyCompleteValue((uint8*)ShadowArray, (uint8*)DataArray);
+		}
+
 		*OutShadowBaseData = ShadowArray->GetData();
 	}
 	else
