Quartz deprecations for 5.1

diff --git a/Engine/Source/Runtime/Engine/Classes/Components/AudioComponent.h b/Engine/Source/Runtime/Engine/Classes/Components/AudioComponent.h
--- a/Engine/Source/Runtime/Engine/Classes/Components/AudioComponent.h
+++ b/Engine/Source/Runtime/Engine/Classes/Components/AudioComponent.h
@@ -748,12 +748,6 @@ public:
 	virtual void BeginDestroy() override;
 	//~ End UObject Interface.
 
-	// QuartzTodo: delete this after fix-up
-	static void AddReferencedObjects(UObject* InThis, FReferenceCollector& Collector)
-	{
-		CastChecked<UAudioComponent>(InThis)->QuartzAddReferencedObjects(Collector);
-	}
-
 	//~ Begin USceneComponent Interface
 	virtual void Activate(bool bReset=false) override;
 	virtual void Deactivate() override;
diff --git a/Engine/Source/Runtime/Engine/Classes/Sound/QuartzSubscription.h b/Engine/Source/Runtime/Engine/Classes/Sound/QuartzSubscription.h
--- a/Engine/Source/Runtime/Engine/Classes/Sound/QuartzSubscription.h
+++ b/Engine/Source/Runtime/Engine/Classes/Sound/QuartzSubscription.h
@@ -150,15 +150,9 @@ public:
 
 	bool QuartzIsTickable() const;
 
-	// QuartzTodo: uncomment post-cleanup
-	// UE_DEPRECATED(5.1, "Derived classes should have access to their own UWorld, this function will always return null")
+	UE_DEPRECATED(5.1, "Derived classes should have access to their own UWorld, this function will always return null")
 	UWorld* QuartzGetWorld() const { return nullptr; }
 
-	// QuartzTodo: delete this static function post-cleanup
-	void QuartzAddReferencedObjects(FReferenceCollector& Collector)
-	{
-	}
-
 	void AddMetronomeBpDelegate(EQuartzCommandQuantization InQuantizationBoundary, const FOnQuartzMetronomeEventBP& OnQuantizationEvent);
 
 
@@ -171,9 +165,7 @@ public:
 	UE_DEPRECATED(5.1, "use FQuartzTickableObject::AddCommandDelegate(const FOnQuartzCommandEventBP& InDelegate) insead")
 	int32 AddCommandDelegate(const FOnQuartzCommandEventBP& InDelegate, TArray<FQuartzGameThreadCommandQueuePtr>& TargetSubscriberArray){ return -1; }
 
-
-	// QuartzTodo: uncomment  after fix-up
-	// UE_DEPRECATED(5.1, "This object no longer holds any UObject references. Caller should have their own UWorld* and use static 'UQuartzSubsystem::Get()' instead.")
+	UE_DEPRECATED(5.1, "This object no longer holds any UObject references. Caller should have their own UWorld* and use static 'UQuartzSubsystem::Get()' instead.")
 	UQuartzSubsystem* GetQuartzSubsystem() const;
 
 	// required by TQuartzShareableCommandQueue template
@@ -212,8 +204,7 @@ private:
 
 public:
 	// deprecate public access
-	// QuartzTodo: uncomment after fix-up
-	// UE_DEPRECATED(5.1, "FQuartzTickableObject::CommandQueuePtr should no longer be accessed directly. (This member will be private in a future engine version).  Use GetQuartzSubscriber() instead")
+	UE_DEPRECATED(5.1, "FQuartzTickableObject::CommandQueuePtr should no longer be accessed directly. (This member will be private in a future engine version).  Use GetQuartzSubscriber() instead")
 	TSharedPtr<Audio::TQuartzShareableCommandQueue<FQuartzTickableObject>, ESPMode::ThreadSafe> CommandQueuePtr;
 
 private:
