Add a function to clear the bindings for a UObject from the input component. Now that we support any UObject binding, we can't assume that the bindings will be cleaned up by the components shutting down.

#rnx

diff --git a/Engine/Source/Runtime/Engine/Classes/Components/InputComponent.h b/Engine/Source/Runtime/Engine/Classes/Components/InputComponent.h
--- a/Engine/Source/Runtime/Engine/Classes/Components/InputComponent.h
+++ b/Engine/Source/Runtime/Engine/Classes/Components/InputComponent.h
@@ -764,6 +764,9 @@ public:
 	/** Whether any components lower on the input stack should be allowed to receive input. */
 	uint8 bBlockInput:1;
 
+	/** Clears any inut callback delegates from the given UObject */
+	void ClearBindingsForObject(UObject* InOwner);
+	
 	void ConditionalBuildKeyMap(UPlayerInput* PlayerInput);
 
 	/**
diff --git a/Engine/Source/Runtime/Engine/Private/Components/InputComponent.cpp b/Engine/Source/Runtime/Engine/Private/Components/InputComponent.cpp
--- a/Engine/Source/Runtime/Engine/Private/Components/InputComponent.cpp
+++ b/Engine/Source/Runtime/Engine/Private/Components/InputComponent.cpp
@@ -103,6 +103,19 @@ void UInputComponent::OnInputOwnerEndPlayed(AActor* InOwner, EEndPlayReason::Typ
 	}
 }
 
+void UInputComponent::ClearBindingsForObject(UObject* InOwner)
+{
+	for (int32 Index = CachedKeyToActionInfo.Num() - 1; Index >= 0; --Index)
+	{
+		FCachedKeyToActionInfo& CachedInfo = CachedKeyToActionInfo[Index];
+		const UPlayerInput* CachedInput = CachedInfo.PlayerInput.Get();
+		if (CachedInput && CachedInput->GetTypedOuter<UObject>() == InOwner)
+		{
+			CachedKeyToActionInfo.RemoveAtSwap(Index, 1, false);
+		}
+	}
+}
+
 void UInputComponent::GetActionsBoundToKey(UPlayerInput* PlayerInput, const FKey Key, TArray<TSharedPtr<FInputActionBinding>>& Actions) const
 {
 	for (const FCachedKeyToActionInfo& CachedInfo : CachedKeyToActionInfo)
