Add a cvar csv.AlwaysShowFrameCount to enable showing CSV frame counts even if screen messages are disabled

diff --git a/Engine/Source/Runtime/Engine/Private/UnrealEngine.cpp b/Engine/Source/Runtime/Engine/Private/UnrealEngine.cpp
--- a/Engine/Source/Runtime/Engine/Private/UnrealEngine.cpp
+++ b/Engine/Source/Runtime/Engine/Private/UnrealEngine.cpp
@@ -426,6 +426,15 @@ static FAutoConsoleVariableRef GDelayTrimMemoryDuringMapLoadModeCVar(
 	ECVF_Default
 );
 
+#if !UE_BUILD_SHIPPING
+TAutoConsoleVariable<bool> CVarCsvAlwaysShowFrameCount(
+	TEXT("csv.AlwaysShowFrameCount"),
+	false,
+	TEXT("If enabled, we show the frame count in non-shipping builds, even if screen messages are disabled")
+);
+#endif
+
+
 #if !UE_BUILD_SHIPPING
 class FDisplayCVarListExecHelper : public FSelfRegisteringExec
 {
@@ -11707,20 +11716,19 @@ void DrawStatsHUD( UWorld* World, FViewport* Viewport, FCanvas* Canvas, UCanvas*
 	const FVector2D FontScale(1, 1);
 	const int32 FontSizeY = 20;
 #if !UE_BUILD_SHIPPING
-	if (!GIsHighResScreenshot && !GIsDumpingMovie && GAreScreenMessagesEnabled)
-	{
-		const int32 MessageX = 40;
-
-		FCanvasTextItem SmallTextItem(FVector2D(0, 0), FText::GetEmpty(), GEngine->GetSmallFont(), FLinearColor::White);
-		SmallTextItem.Scale = FontScale;
-		SmallTextItem.EnableShadow(FLinearColor::Black);
+	const int32 MessageX = 40;
+	FCanvasTextItem SmallTextItem(FVector2D(0, 0), FText::GetEmpty(), GEngine->GetSmallFont(), FLinearColor::White);
+	SmallTextItem.Scale = FontScale;
+	SmallTextItem.EnableShadow(FLinearColor::Black);
 
+	bool bShowScreenMessages = !GIsHighResScreenshot && !GIsDumpingMovie && GAreScreenMessagesEnabled;
+	if (bShowScreenMessages)
+	{
 		// Draw map warnings?
 		if (!GEngine->bSuppressMapWarnings)
 		{
 			MessageY = DrawMapWarnings(World, Viewport, Canvas, CanvasObject, MessageX, MessageY);
 		}
-
 #if ENABLE_AUDIO_DEBUG
 		MessageY = Audio::FAudioDebugger::DrawDebugStats(*World, Viewport, Canvas, MessageY);
 #endif // ENABLE_AUDIO_DEBUG
@@ -11798,35 +11806,6 @@ void DrawStatsHUD( UWorld* World, FViewport* Viewport, FCanvas* Canvas, UCanvas*
 		}
 #endif // FRAMEPRO_ENABLED
 
-#if CSV_PROFILER
-		if (FCsvProfiler::Get()->IsCapturing() || FCsvProfiler::Get()->IsWritingFile())
-		{
-			if (FCsvProfiler::Get()->IsWritingFile())
-			{
-				SmallTextItem.SetColor(FLinearColor(1.0f, 0.0f, 0.0f, 1.0f));
-				SmallTextItem.Text = FText::Format(LOCTEXT("CsvProfilerWriteFmt", "CsvProfiler, please wait, writing file... Total frames: {0}"), FCsvProfiler::Get()->GetCaptureFrameNumber());
-			}
-			else
-			{
-				SmallTextItem.SetColor(FLinearColor(0.0f, 1.0f, 0.0f, 1.0f));
-				SmallTextItem.Text = FText::Format(LOCTEXT("CsvProfilerFrameFmt", "CsvProfiler frame: {0}"), FCsvProfiler::Get()->GetCaptureFrameNumber());
-			}
-
-			MessageY += 250.0f;
-
-			// Swap the shadow for a background to improve readability
-			static FVector2D PrevDrawnSize(0.f); // one frame behind, but much faster than calling GetTextSize
-			const FLinearColor BackgroundColor = FLinearColor(0.02f, 0.02f, 0.02f, 0.88f); // color used by stat group tables
-			SmallTextItem.DisableShadow();
-			Canvas->DrawTile(MessageX - 1.0f, MessageY - 1.0f, PrevDrawnSize.X + 2.0f, PrevDrawnSize.Y + 2.0f, 0.0f, 0.0f, 0.0f, 0.0f, BackgroundColor);
-
-			Canvas->DrawItem(SmallTextItem, FVector2D(MessageX, MessageY));
-			SmallTextItem.EnableShadow(FLinearColor::Black); // reset
-			PrevDrawnSize = SmallTextItem.DrawnSize;
-			MessageY += FontSizeY;
-		}
-#endif
-
 #if TRACING_PROFILER
 PRAGMA_DISABLE_DEPRECATION_WARNINGS
 		if (FTracingProfiler::Get()->IsCapturing())
@@ -11884,6 +11863,36 @@ PRAGMA_ENABLE_DEPRECATION_WARNINGS
 			}
 		}
 	}
+
+#if CSV_PROFILER
+	if ( ( bShowScreenMessages || CVarCsvAlwaysShowFrameCount.GetValueOnGameThread() ) && ( FCsvProfiler::Get()->IsCapturing() || FCsvProfiler::Get()->IsWritingFile() ) )
+	{
+		if (FCsvProfiler::Get()->IsWritingFile())
+		{
+			SmallTextItem.SetColor(FLinearColor(1.0f, 0.0f, 0.0f, 1.0f));
+			SmallTextItem.Text = FText::Format(LOCTEXT("CsvProfilerWriteFmt", "CsvProfiler, please wait, writing file... Total frames: {0}"), FCsvProfiler::Get()->GetCaptureFrameNumber());
+		}
+		else
+		{
+			SmallTextItem.SetColor(FLinearColor(0.0f, 1.0f, 0.0f, 1.0f));
+			SmallTextItem.Text = FText::Format(LOCTEXT("CsvProfilerFrameFmt", "CsvProfiler frame: {0}"), FCsvProfiler::Get()->GetCaptureFrameNumber());
+		}
+
+		MessageY += 250.0f;
+
+		// Swap the shadow for a background to improve readability
+		static FVector2D PrevDrawnSize(0.f); // one frame behind, but much faster than calling GetTextSize
+		const FLinearColor BackgroundColor = FLinearColor(0.02f, 0.02f, 0.02f, 0.88f); // color used by stat group tables
+		SmallTextItem.DisableShadow();
+		Canvas->DrawTile(MessageX - 1.0f, MessageY - 1.0f, PrevDrawnSize.X + 2.0f, PrevDrawnSize.Y + 2.0f, 0.0f, 0.0f, 0.0f, 0.0f, BackgroundColor);
+
+		Canvas->DrawItem(SmallTextItem, FVector2D(MessageX, MessageY));
+		SmallTextItem.EnableShadow(FLinearColor::Black); // reset
+		PrevDrawnSize = SmallTextItem.DrawnSize;
+		MessageY += FontSizeY;
+	}
+#endif // CSV_PROFILER
+
 #endif // UE_BUILD_SHIPPING 
 
 	{
