Changelist Validation: don't register level instance containers when the world is not loaded.

#rnx

diff --git a/Engine/Source/Runtime/Engine/Private/WorldPartition/LevelInstance/LevelInstanceActorDesc.cpp b/Engine/Source/Runtime/Engine/Private/WorldPartition/LevelInstance/LevelInstanceActorDesc.cpp
--- a/Engine/Source/Runtime/Engine/Private/WorldPartition/LevelInstance/LevelInstanceActorDesc.cpp
+++ b/Engine/Source/Runtime/Engine/Private/WorldPartition/LevelInstance/LevelInstanceActorDesc.cpp
@@ -65,23 +65,25 @@ bool FLevelInstanceActorDesc::Equals(const FWorldPartitionActorDesc* Other) cons
 
 void FLevelInstanceActorDesc::RegisterContainerInstance(UWorld* InWorld)
 {
-	check(InWorld);
-	check(!LevelInstanceContainer.IsValid());
-	if (DesiredRuntimeBehavior == ELevelInstanceRuntimeBehavior::Partitioned && !GLevelInstanceDebugForceLevelStreaming)
+	if (InWorld)
 	{
-		if (!LevelPackage.IsNone() && ULevel::GetIsLevelUsingExternalActorsFromPackage(LevelPackage) && ULevelInstanceSubsystem::CanUsePackage(LevelPackage))
+		check(!LevelInstanceContainer.IsValid());
+		if (DesiredRuntimeBehavior == ELevelInstanceRuntimeBehavior::Partitioned && !GLevelInstanceDebugForceLevelStreaming)
 		{
-			ULevelInstanceSubsystem* LevelInstanceSubsystem = UWorld::GetSubsystem<ULevelInstanceSubsystem>(InWorld);
-			check(LevelInstanceSubsystem);
+			if (!LevelPackage.IsNone() && ULevel::GetIsLevelUsingExternalActorsFromPackage(LevelPackage) && ULevelInstanceSubsystem::CanUsePackage(LevelPackage))
+			{
+				ULevelInstanceSubsystem* LevelInstanceSubsystem = UWorld::GetSubsystem<ULevelInstanceSubsystem>(InWorld);
+				check(LevelInstanceSubsystem);
 
-			LevelInstanceContainer = LevelInstanceSubsystem->RegisterContainer(LevelPackage);
-			check(LevelInstanceContainer.IsValid());
+				LevelInstanceContainer = LevelInstanceSubsystem->RegisterContainer(LevelPackage);
+				check(LevelInstanceContainer.IsValid());
 
-			FTransform LevelInstancePivotOffsetTransform = FTransform(ULevel::GetLevelInstancePivotOffsetFromPackage(LevelPackage));
-			FTransform FinalLevelTransform = LevelInstancePivotOffsetTransform * LevelInstanceTransform;
-			FBox ContainerBounds = LevelInstanceSubsystem->GetContainerBounds(LevelPackage).TransformBy(FinalLevelTransform);
+				FTransform LevelInstancePivotOffsetTransform = FTransform(ULevel::GetLevelInstancePivotOffsetFromPackage(LevelPackage));
+				FTransform FinalLevelTransform = LevelInstancePivotOffsetTransform * LevelInstanceTransform;
+				FBox ContainerBounds = LevelInstanceSubsystem->GetContainerBounds(LevelPackage).TransformBy(FinalLevelTransform);
 
-			ContainerBounds.GetCenterAndExtents(BoundsLocation, BoundsExtent);
+				ContainerBounds.GetCenterAndExtents(BoundsLocation, BoundsExtent);
+			}
 		}
 	}
 }
@@ -188,4 +190,4 @@ void FLevelInstanceActorDesc::Serialize(FArchive& Ar)
 	}
 }
 
-#endif
\ No newline at end of file
+#endif
