SubObject replication
* Added netdrivername and beginplay status to ensure when regression is detected
* Compile out of shipping the validation code

diff --git a/Engine/Source/Runtime/Engine/Private/DataChannel.cpp b/Engine/Source/Runtime/Engine/Private/DataChannel.cpp
--- a/Engine/Source/Runtime/Engine/Private/DataChannel.cpp
+++ b/Engine/Source/Runtime/Engine/Private/DataChannel.cpp
@@ -52,6 +52,11 @@ LLM_DEFINE_TAG(NetChannel, NAME_None, TEXT("Networking"), GET_STATFNAME(STAT_Net
 
 CSV_DECLARE_CATEGORY_EXTERN(Replication);
 
+// Enable the code allowing to validate replicate subobjects when converting from the legacy method to the explicit registration list
+#ifndef SUBOBJECT_TRANSITION_VALIDATION
+	#define SUBOBJECT_TRANSITION_VALIDATION !UE_BUILD_SHIPPING
+#endif
+
 extern int32 GDoReplicationContextString;
 extern int32 GNetDormancyValidate;
 extern bool GbNetReuseReplicatorsForDormantObjects;
@@ -120,6 +125,7 @@ namespace UE::Net
 		QueuedBunchTimeoutSeconds,
 		TEXT("Time in seconds to wait for queued bunches on a channel to flush before logging a warning."));
 
+#if SUBOBJECT_TRANSITION_VALIDATION
 	static int32 GCVarCompareSubObjectsReplicated = 0;
 	static FAutoConsoleVariableRef CVarCompareSubObjectsReplicated(
 		TEXT("net.SubObjects.CompareWithLegacy"),
@@ -130,6 +136,7 @@ namespace UE::Net
 		TEXT("net.SubObjects.LogAllComparisonErrors"),
 		0,
 		TEXT("If enabled log all the errors detected by the CompareWithLegacy cheat. Otherwise only the first ensure triggered gets logged."));
+#endif
 
 }; // namespace UE::Net
 
@@ -163,6 +170,7 @@ namespace DataChannelInternal
 	*/
 	bool bTestingLegacyMethodForComparison = false;
 
+#if SUBOBJECT_TRANSITION_VALIDATION
     struct FSubObjectReplicatedInfo
 	{
 		// The replicated subobject
@@ -175,13 +183,18 @@ namespace DataChannelInternal
 		bool operator==(const FSubObjectReplicatedInfo& rhs) const { return SubObject == rhs.SubObject; }
 
 		FString Describe(AActor* ReplicatedActor) const
-		{ 
-			return FString::Printf(TEXT("Subobject %s::%s (0x%p) [owner %s]"), *ReplicatedActor->GetName(), *SubObject->GetName(), SubObject, *GetNameSafe(SubObjectOwner));
+		{
+			return FString::Printf(TEXT("Subobject %s::%s (%s) [owner %s]"), 
+				*ReplicatedActor->GetName(), 
+				*SubObject->GetName(), 
+				ReplicatedActor->HasActorBegunPlay() ? TEXT("BegunPlay") : 
+					ReplicatedActor->IsActorBeginningPlay() ? TEXT("BeginningPlay") :
+					TEXT("HasNotBegunPlay"),
+				*GetNameSafe(SubObjectOwner));
 		}
 	};
 
 	/**
-	
 	* List of replicated subobjects collected via Actor::ReplicateSubObjects.
 	* This list is static to prevent memory overhead on every actor channel.
 	* It is emptied at the end of every call to UActorChannel::ReplicateActor()
@@ -194,6 +207,7 @@ namespace DataChannelInternal
     * It is emptied at the end of every call to UActorChannel::ReplicateActor()
 	*/
 	TArray<FSubObjectReplicatedInfo> ReplicatedSubObjectsTracker;
+#endif
 
 	/** 
 	* Temporary pointer to the NetworkSubsystem that gets set and released in ReplicateActor
@@ -3525,12 +3539,14 @@ bool UActorChannel::DoSubObjectReplication(FOutBunch& Bunch, FReplicationFlags&
 		// If the actor replicates it's subobjects and those of it's components via the list.
 		bWroteSomethingImportant |= ReplicateRegisteredSubObjects(Bunch, OutRepFlags);
 
+#if SUBOBJECT_TRANSITION_VALIDATION
 		if (UE::Net::GCVarCompareSubObjectsReplicated)
 		{
 			ValidateReplicatedSubObjects();
 
 			DataChannelInternal::bTrackReplicatedSubObjects = false;
 		}
+#endif
 	}
 	else
 	{
@@ -3580,6 +3596,7 @@ bool UActorChannel::ReplicateRegisteredSubObjects(FOutBunch& Bunch, FReplication
 
 	check(Actor->IsUsingRegisteredSubObjectList());
 
+#if SUBOBJECT_TRANSITION_VALIDATION
 	if (GCVarCompareSubObjectsReplicated)
 	{
 		TestLegacyReplicateSubObjects(Bunch, RepFlags);
@@ -3587,6 +3604,7 @@ bool UActorChannel::ReplicateRegisteredSubObjects(FOutBunch& Bunch, FReplication
 		// From here track all subobjects that write into the bunch.
 		DataChannelInternal::bTrackReplicatedSubObjects = true;
 	}
+#endif
 
 	const TStaticBitArray<COND_Max> ConditionMap = FSendingRepState::BuildConditionMapFromRepFlags(RepFlags);
 
@@ -3655,11 +3673,13 @@ bool UActorChannel::ReplicateSubobject(UObject* SubObj, FOutBunch& Bunch, FRepli
 		return false;
 	}
 
+#if SUBOBJECT_TRANSITION_VALIDATION
 	if (DataChannelInternal::bTestingLegacyMethodForComparison)
 	{
 		DataChannelInternal::LegacySubObjectsCollected.Emplace(DataChannelInternal::FSubObjectReplicatedInfo(SubObj));
 		return false;
 	}
+#endif
 
 	bool bWroteSomethingImportant = false;
 
@@ -3669,9 +3689,9 @@ bool UActorChannel::ReplicateSubobject(UObject* SubObj, FOutBunch& Bunch, FRepli
 	}
 	else
 	{
-#if !UE_BUILD_SHIPPING
 		using namespace UE::Net;
 
+#if SUBOBJECT_TRANSITION_VALIDATION
 		// Make sure the SubObjectOwner actually has the SubObject in his list and will replicate it later.
 		UActorComponent* Component = CastChecked<UActorComponent>(DataChannelInternal::CurrentSubObjectOwner);
 		const bool bIsInRegistry = FSubObjectRegistryGetter::IsSubObjectInRegistry(Actor, Component, SubObj);
@@ -3685,6 +3705,7 @@ bool UActorChannel::ReplicateSubobject(UObject* SubObj, FOutBunch& Bunch, FRepli
 
 void UActorChannel::TestLegacyReplicateSubObjects(FOutBunch& Bunch, FReplicationFlags RepFlags)
 {
+#if SUBOBJECT_TRANSITION_VALIDATION
 	// Call the legacy method but simply collect the subobjects it would have replicated
 	DataChannelInternal::bTestingLegacyMethodForComparison = true;
 
@@ -3692,10 +3713,12 @@ void UActorChannel::TestLegacyReplicateSubObjects(FOutBunch& Bunch, FReplication
 	Actor->ReplicateSubobjects(this, &Bunch, &RepFlags);
 
 	DataChannelInternal::bTestingLegacyMethodForComparison = false;
+#endif
 }
 
 void UActorChannel::TestLegacyReplicateSubObjects(UActorComponent* ReplicatedComponent, FOutBunch& Bunch, FReplicationFlags RepFlags)
 {
+#if SUBOBJECT_TRANSITION_VALIDATION
 	// Call the legacy method but simply collect the subobjects it would have replicated
 	DataChannelInternal::bTestingLegacyMethodForComparison = true;
 
@@ -3703,6 +3726,7 @@ void UActorChannel::TestLegacyReplicateSubObjects(UActorComponent* ReplicatedCom
 	ReplicatedComponent->ReplicateSubobjects(this, &Bunch, &RepFlags);
 
 	DataChannelInternal::bTestingLegacyMethodForComparison = false;
+#endif
 }
 
 
@@ -3716,6 +3740,7 @@ bool UActorChannel::ReplicateSubobject(UActorComponent* ReplicatedComponent, FOu
 	// This traps actor components using the registration list even if their owning Actor isn't using the list itself.
 	if (ReplicatedComponent->IsUsingRegisteredSubObjectList() && !DataChannelInternal::bTestingLegacyMethodForComparison)
 	{
+#if SUBOBJECT_TRANSITION_VALIDATION
 		if (UE::Net::GCVarCompareSubObjectsReplicated)
 		{
 			TestLegacyReplicateSubObjects(ReplicatedComponent, Bunch, RepFlags);
@@ -3723,6 +3748,7 @@ bool UActorChannel::ReplicateSubobject(UActorComponent* ReplicatedComponent, FOu
 			// From here track all subobjects that write into the bunch.
 			DataChannelInternal::bTrackReplicatedSubObjects = true;
 		}
+#endif
 
 		bool bWroteSomethingImportant = false;
 
@@ -3734,12 +3760,14 @@ bool UActorChannel::ReplicateSubobject(UActorComponent* ReplicatedComponent, FOu
 		// Write all the subobjects of the component
 		bWroteSomethingImportant |= WriteSubObjects(ReplicatedComponent, Bunch, RepFlags, ConditionMap);
 
+#if SUBOBJECT_TRANSITION_VALIDATION
 		if (UE::Net::GCVarCompareSubObjectsReplicated)
 		{
 			ValidateReplicatedSubObjects();
 
 			DataChannelInternal::bTrackReplicatedSubObjects = false;
 		}
+#endif
 
         // Finally write the component itself
 		SetCurrentSubObjectOwner(Actor);
@@ -3782,6 +3810,7 @@ bool UActorChannel::WriteSubObjects(UActorComponent* ReplicatedComponent, const
 
 bool UActorChannel::ValidateReplicatedSubObjects()
 {
+#if SUBOBJECT_TRANSITION_VALIDATION
 	bool bErrorDetected = false;
 	const bool bLogAllErrors = UE::Net::CVarLogAllSubObjectComparisonErrors.GetValueOnAnyThread() != 0;
 
@@ -3792,7 +3821,8 @@ bool UActorChannel::ValidateReplicatedSubObjects()
 		constexpr bool bNoShrinking = false;
 		const bool bWasFound = DataChannelInternal::ReplicatedSubObjectsTracker.RemoveSingleSwap(Info, bNoShrinking) != 0;
 
-		ensureMsgf(bWasFound, TEXT("%s was not replicated by the subobject list. Only by the legacy ReplicateSubObjects method."), *Info.Describe(Actor));
+		ensureMsgf(bWasFound, TEXT("%s was not replicated by the subobject list in %s. Only by the legacy ReplicateSubObjects method."),
+			*Info.Describe(Actor), *Connection->GetDriver()->NetDriverName.ToString());
 		UE_CLOG(!bWasFound && bLogAllErrors, LogNetSubObject, Warning, TEXT("%s was not replicated by the subobject list."), *Info.Describe(Actor));
 
 		bErrorDetected |= !bWasFound;
@@ -3803,7 +3833,8 @@ bool UActorChannel::ValidateReplicatedSubObjects()
 	{
 		const DataChannelInternal::FSubObjectReplicatedInfo& Info = DataChannelInternal::ReplicatedSubObjectsTracker[i];
 
-		ensureMsgf(false, TEXT("%s was replicated only in the registered subobject list. Not in the legacy path."), *Info.Describe(Actor));
+		ensureMsgf(false, TEXT("%s was replicated only in the registered subobject list in %s. Not in the legacy path."), 
+			*Info.Describe(Actor), *Info.Describe(Actor), *Connection->GetDriver()->NetDriverName.ToString());
 		UE_CLOG(bLogAllErrors, LogNet, Warning, TEXT("%s was replicated only in the registered subobject list. Not in the legacy path."), *Info.Describe(Actor));
 	}
 
@@ -3813,6 +3844,9 @@ bool UActorChannel::ValidateReplicatedSubObjects()
 	DataChannelInternal::ReplicatedSubObjectsTracker.Reset();
 
 	return bErrorDetected;
+#else
+	return false;
+#endif
 }
 
 FString UActorChannel::Describe()
@@ -4822,10 +4856,12 @@ bool UActorChannel::WriteSubObjectInBunch(UObject* Obj, FOutBunch& Bunch, FRepli
 		return ReplicateSubobjectCustom(Obj, Bunch, RepFlags);
 	}
 
+#if SUBOBJECT_TRANSITION_VALIDATION
 	if (DataChannelInternal::bTrackReplicatedSubObjects)
 	{
 		DataChannelInternal::ReplicatedSubObjectsTracker.Emplace(DataChannelInternal::FSubObjectReplicatedInfo(Obj));
 	}
+#endif
 
 	bool bFoundInvalidReplicator = false;
 	TSharedRef<FObjectReplicator>* FoundReplicator = FindReplicator(Obj, &bFoundInvalidReplicator);
