For dynamic streaming instances, separate component scaling from uv density when uploading render asset infos for static and skinned mesh components. Component scales are stored as a part of Bounds4 and get updated regularly. The scales of dynamic components can change at runtime so it is invalid to bake them into texel factors which aren't expected to change at runtime.

diff --git a/Engine/Source/Runtime/Engine/Classes/Components/PrimitiveComponent.h b/Engine/Source/Runtime/Engine/Classes/Components/PrimitiveComponent.h
--- a/Engine/Source/Runtime/Engine/Classes/Components/PrimitiveComponent.h
+++ b/Engine/Source/Runtime/Engine/Classes/Components/PrimitiveComponent.h
@@ -1790,6 +1790,8 @@ public:
 	/** Return true if primitive can skip getting texture streaming render asset info. */
 	bool CanSkipGetTextureStreamingRenderAssetInfo() const;
 
+	virtual float GetStreamingScale() const { return 1.f; }
+
 	/** Returns true if the component is static and has the right static mesh setup to support lightmaps. */
 	virtual bool HasValidSettingsForStaticLighting(bool bOverlookInvalidComponents) const 
 	{
diff --git a/Engine/Source/Runtime/Engine/Classes/Components/SkinnedMeshComponent.h b/Engine/Source/Runtime/Engine/Classes/Components/SkinnedMeshComponent.h
--- a/Engine/Source/Runtime/Engine/Classes/Components/SkinnedMeshComponent.h
+++ b/Engine/Source/Runtime/Engine/Classes/Components/SkinnedMeshComponent.h
@@ -896,6 +896,7 @@ public:
 	virtual bool GetMaterialStreamingData(int32 MaterialIndex, FPrimitiveMaterialInfo& MaterialData) const override;
 	virtual void GetStreamingRenderAssetInfo(FStreamingTextureLevelContext& LevelContext, TArray<FStreamingRenderAssetPrimitiveInfo>& OutStreamingRenderAssets) const override;
 	virtual int32 GetNumMaterials() const override;
+	virtual float GetStreamingScale() const override { return GetComponentTransform().GetMaximumAxisScale(); }
 	//~ End UPrimitiveComponent Interface
 
 	//~ Begin UMeshComponent Interface
diff --git a/Engine/Source/Runtime/Engine/Classes/Components/StaticMeshComponent.h b/Engine/Source/Runtime/Engine/Classes/Components/StaticMeshComponent.h
--- a/Engine/Source/Runtime/Engine/Classes/Components/StaticMeshComponent.h
+++ b/Engine/Source/Runtime/Engine/Classes/Components/StaticMeshComponent.h
@@ -604,6 +604,7 @@ public:
 	virtual bool ComponentIsTouchingSelectionBox(const FBox& InSelBBox, const bool bConsiderOnlyBSP, const bool bMustEncompassEntireComponent) const override;
 	virtual bool ComponentIsTouchingSelectionFrustum(const FConvexVolume& InSelBBox, const bool bConsiderOnlyBSP, const bool bMustEncompassEntireComponent) const override;
 #endif
+	virtual float GetStreamingScale() const override { return GetComponentTransform().GetMaximumAxisScale(); }
 	//~ End UPrimitiveComponent Interface.
 
 	//~ Begin INavRelevantInterface Interface.
diff --git a/Engine/Source/Runtime/Engine/Classes/Engine/TextureStreamingTypes.h b/Engine/Source/Runtime/Engine/Classes/Engine/TextureStreamingTypes.h
--- a/Engine/Source/Runtime/Engine/Classes/Engine/TextureStreamingTypes.h
+++ b/Engine/Source/Runtime/Engine/Classes/Engine/TextureStreamingTypes.h
@@ -65,12 +65,17 @@ struct FStreamingRenderAssetPrimitiveInfo
 	UPROPERTY(Transient)
 	uint32 bAllowInvalidTexelFactorWhenUnregistered : 1;
 
+	/** Mesh texel factors aren't uv density and shouldn't be affected by component scales */
+	UPROPERTY(Transient)
+	uint32 bAffectedByComponentScale : 1;
+
 	FStreamingRenderAssetPrimitiveInfo() : 
 		RenderAsset(nullptr),
 		Bounds(ForceInit), 
 		TexelFactor(1.0f),
 		PackedRelativeBox(0),
-		bAllowInvalidTexelFactorWhenUnregistered(false)
+		bAllowInvalidTexelFactorWhenUnregistered(false),
+		bAffectedByComponentScale(true)
 	{
 	}
 
@@ -79,12 +84,14 @@ struct FStreamingRenderAssetPrimitiveInfo
 		const FBoxSphereBounds& InBounds,
 		float InTexelFactor,
 		uint32 InPackedRelativeBox = 0,
-		bool bInAllowInvalidTexelFactorWhenUnregistered = false) :
+		bool bInAllowInvalidTexelFactorWhenUnregistered = false,
+		bool bInAffectedByComponentScale = true) :
 		RenderAsset(InAsset),
 		Bounds(InBounds), 
 		TexelFactor(InTexelFactor),
 		PackedRelativeBox(InPackedRelativeBox),
-		bAllowInvalidTexelFactorWhenUnregistered(bInAllowInvalidTexelFactorWhenUnregistered)
+		bAllowInvalidTexelFactorWhenUnregistered(bInAllowInvalidTexelFactorWhenUnregistered),
+		bAffectedByComponentScale(bInAffectedByComponentScale)
 	{
 	}
 
diff --git a/Engine/Source/Runtime/Engine/Private/Components/SkinnedMeshComponent.cpp b/Engine/Source/Runtime/Engine/Private/Components/SkinnedMeshComponent.cpp
--- a/Engine/Source/Runtime/Engine/Private/Components/SkinnedMeshComponent.cpp
+++ b/Engine/Source/Runtime/Engine/Private/Components/SkinnedMeshComponent.cpp
@@ -1109,13 +1109,13 @@ bool USkinnedMeshComponent::GetMaterialStreamingData(int32 MaterialIndex, FPrimi
 
 void USkinnedMeshComponent::GetStreamingRenderAssetInfo(FStreamingTextureLevelContext& LevelContext, TArray<FStreamingRenderAssetPrimitiveInfo>& OutStreamingRenderAssets) const
 {
-	GetStreamingTextureInfoInner(LevelContext, nullptr, GetComponentTransform().GetMaximumAxisScale() * StreamingDistanceMultiplier, OutStreamingRenderAssets);
+	GetStreamingTextureInfoInner(LevelContext, nullptr, StreamingDistanceMultiplier, OutStreamingRenderAssets);
 
 	if (SkeletalMesh && SkeletalMesh->IsStreamable())
 	{
 		const int32 LocalForcedLodModel = GetForcedLOD();
 		const float TexelFactor = LocalForcedLodModel > 0 ? -(SkeletalMesh->GetLODNum() - LocalForcedLodModel + 1) : Bounds.SphereRadius * 2.f;
-		new (OutStreamingRenderAssets) FStreamingRenderAssetPrimitiveInfo(SkeletalMesh, Bounds, TexelFactor, PackedRelativeBox_Identity);
+		new (OutStreamingRenderAssets) FStreamingRenderAssetPrimitiveInfo(SkeletalMesh, Bounds, TexelFactor, PackedRelativeBox_Identity, false, false);
 	}
 }
 
diff --git a/Engine/Source/Runtime/Engine/Private/Components/StaticMeshComponent.cpp b/Engine/Source/Runtime/Engine/Private/Components/StaticMeshComponent.cpp
--- a/Engine/Source/Runtime/Engine/Private/Components/StaticMeshComponent.cpp
+++ b/Engine/Source/Runtime/Engine/Private/Components/StaticMeshComponent.cpp
@@ -1059,7 +1059,7 @@ void UStaticMeshComponent::GetStreamingRenderAssetInfo(FStreamingTextureLevelCon
 	
 	if (!CanSkipGetTextureStreamingRenderAssetInfo())
 	{
-		GetStreamingTextureInfoInner(LevelContext, Mobility == EComponentMobility::Static ? &StreamingTextureData : nullptr, GetTransformScale() * StreamingDistanceMultiplier, OutStreamingRenderAssets);
+		GetStreamingTextureInfoInner(LevelContext, Mobility == EComponentMobility::Static ? &StreamingTextureData : nullptr, StreamingDistanceMultiplier, OutStreamingRenderAssets);
 	}
 
 	// Process the lightmaps and shadowmaps entries.
@@ -1074,7 +1074,7 @@ void UStaticMeshComponent::GetStreamingRenderAssetInfo(FStreamingTextureLevelCon
 			const FVector2D& Scale = Lightmap->GetCoordinateScale();
 			if (Scale.X > UE_SMALL_NUMBER && Scale.Y > UE_SMALL_NUMBER)
 			{
-				const float TexelFactor = GetStaticMesh()->GetLightmapUVDensity() * GetTransformScale() / FMath::Min(Scale.X, Scale.Y);
+				const float TexelFactor = GetStaticMesh()->GetLightmapUVDensity() / FMath::Min(Scale.X, Scale.Y);
 				new (OutStreamingRenderAssets) FStreamingRenderAssetPrimitiveInfo(Lightmap->GetTexture(LightmapIndex), Bounds, TexelFactor, PackedRelativeBox_Identity);
 				new (OutStreamingRenderAssets) FStreamingRenderAssetPrimitiveInfo(Lightmap->GetAOMaterialMaskTexture(), Bounds, TexelFactor, PackedRelativeBox_Identity);
 				new (OutStreamingRenderAssets) FStreamingRenderAssetPrimitiveInfo(Lightmap->GetSkyOcclusionTexture(), Bounds, TexelFactor, PackedRelativeBox_Identity);
@@ -1087,7 +1087,7 @@ void UStaticMeshComponent::GetStreamingRenderAssetInfo(FStreamingTextureLevelCon
 			const FVector2D& Scale = Shadowmap->GetCoordinateScale();
 			if (Scale.X > UE_SMALL_NUMBER && Scale.Y > UE_SMALL_NUMBER)
 			{
-				const float TexelFactor = GetStaticMesh()->GetLightmapUVDensity() * GetTransformScale() / FMath::Min(Scale.X, Scale.Y);
+				const float TexelFactor = GetStaticMesh()->GetLightmapUVDensity() / FMath::Min(Scale.X, Scale.Y);
 				new (OutStreamingRenderAssets) FStreamingRenderAssetPrimitiveInfo(Shadowmap->GetTexture(), Bounds, TexelFactor, PackedRelativeBox_Identity);
 			}
 		}
@@ -1098,7 +1098,7 @@ void UStaticMeshComponent::GetStreamingRenderAssetInfo(FStreamingTextureLevelCon
 		const float TexelFactor = ForcedLodModel > 0 ?
 			-(GetStaticMesh()->GetRenderData()->LODResources.Num() - ForcedLodModel + 1) :
 			(IsRegistered() ? Bounds.SphereRadius * 2.f : 0.f);
-		new (OutStreamingRenderAssets) FStreamingRenderAssetPrimitiveInfo(GetStaticMesh(), Bounds, TexelFactor, PackedRelativeBox_Identity, true);
+		new (OutStreamingRenderAssets) FStreamingRenderAssetPrimitiveInfo(GetStaticMesh(), Bounds, TexelFactor, PackedRelativeBox_Identity, true, false);
 	}
 }
 
diff --git a/Engine/Source/Runtime/Engine/Private/Streaming/DynamicTextureInstanceManager.cpp b/Engine/Source/Runtime/Engine/Private/Streaming/DynamicTextureInstanceManager.cpp
--- a/Engine/Source/Runtime/Engine/Private/Streaming/DynamicTextureInstanceManager.cpp
+++ b/Engine/Source/Runtime/Engine/Private/Streaming/DynamicTextureInstanceManager.cpp
@@ -18,7 +18,8 @@ void FDynamicRenderAssetInstanceManager::FTasks::SyncResults()
 }
 
 FDynamicRenderAssetInstanceManager::FDynamicRenderAssetInstanceManager()
-	: DirtyIndex(0)
+	: StateSync(true)
+	, DirtyIndex(0)
 	, PendingDefragSrcBoundIndex(INDEX_NONE)
 	, PendingDefragDstBoundIndex(INDEX_NONE)
 {
diff --git a/Engine/Source/Runtime/Engine/Private/Streaming/StaticTextureInstanceManager.cpp b/Engine/Source/Runtime/Engine/Private/Streaming/StaticTextureInstanceManager.cpp
--- a/Engine/Source/Runtime/Engine/Private/Streaming/StaticTextureInstanceManager.cpp
+++ b/Engine/Source/Runtime/Engine/Private/Streaming/StaticTextureInstanceManager.cpp
@@ -18,7 +18,8 @@ void FStaticRenderAssetInstanceManager::FTasks::SyncResults()
 }
 
 FStaticRenderAssetInstanceManager::FStaticRenderAssetInstanceManager(RenderAssetInstanceTask::FDoWorkTask& AsyncTask)
-	: DirtyIndex(0)
+	: StateSync(false)
+	, DirtyIndex(0)
 {
 	FTasks& Tasks = StateSync.GetTasks();
 
diff --git a/Engine/Source/Runtime/Engine/Private/Streaming/TextureInstanceState.cpp b/Engine/Source/Runtime/Engine/Private/Streaming/TextureInstanceState.cpp
--- a/Engine/Source/Runtime/Engine/Private/Streaming/TextureInstanceState.cpp
+++ b/Engine/Source/Runtime/Engine/Private/Streaming/TextureInstanceState.cpp
@@ -14,9 +14,19 @@
 #include "UObject/UObjectHash.h"
 #include "Templates/RefCounting.h"
 
+FRenderAssetInstanceState::FRenderAssetInstanceState(bool bForDynamicInstances)
+#if DO_CHECK
+	: bIsDynamicInstanceState(bForDynamicInstances)
+#endif
+{
+}
+
 int32 FRenderAssetInstanceState::AddBounds(const UPrimitiveComponent* Component)
 {
-	return AddBounds(Component->Bounds, PackedRelativeBox_Identity, Component, Component->GetLastRenderTimeOnScreen(), Component->Bounds.Origin, 0, 0, FLT_MAX);
+	checkf(bIsDynamicInstanceState, TEXT("This version of AddBounds should only be called by the dynamic instance manager."));
+	FBoxSphereBounds Bounds = Component->Bounds;
+	Bounds.SphereRadius = Component->GetStreamingScale();
+	return AddBounds(Bounds, PackedRelativeBox_Identity, Component, Component->GetLastRenderTimeOnScreen(), Component->Bounds.Origin, 0, 0, FLT_MAX);
 }
 
 int32 FRenderAssetInstanceState::AddBounds(const FBoxSphereBounds& Bounds, uint32 PackedRelativeBox, const UPrimitiveComponent* InComponent, float LastRenderTime, const FVector4& RangeOrigin, float MinDistanceSq, float MinRangeSq, float MaxRangeSq)
@@ -295,9 +305,20 @@ void FRenderAssetInstanceState::AddRenderAssetElements(const UPrimitiveComponent
 EAddComponentResult FRenderAssetInstanceState::AddComponent(const UPrimitiveComponent* Component, FStreamingTextureLevelContext& LevelContext, float MaxAllowedUIDensity)
 {
 	check(Component);
+	checkf(!bIsDynamicInstanceState, TEXT("Error: trying to add component to dynamic instance manager as static"));
 
 	TArray<FStreamingRenderAssetPrimitiveInfo> RenderAssetInstanceInfos;
 	Component->GetStreamingRenderAssetInfoWithNULLRemoval(LevelContext, RenderAssetInstanceInfos);
+
+	const float ComponentScale = Component->GetStreamingScale();
+	if (ComponentScale != 1.f)
+	{
+		for (FStreamingRenderAssetPrimitiveInfo& Info : RenderAssetInstanceInfos)
+		{
+			Info.TexelFactor *= Info.bAffectedByComponentScale ? ComponentScale : 1.f;
+		}
+	}
+
 	// Texture entries are guarantied to be relevant here, except for bounds if the component is not registered.
 	if (!RenderAssetInstanceInfos.Num())
 	{
@@ -410,6 +431,7 @@ EAddComponentResult FRenderAssetInstanceState::AddComponent(const UPrimitiveComp
 EAddComponentResult FRenderAssetInstanceState::AddComponentIgnoreBounds(const UPrimitiveComponent* Component, FStreamingTextureLevelContext& LevelContext)
 {
 	check(Component->IsRegistered()); // Must be registered otherwise bounds are invalid.
+	checkf(bIsDynamicInstanceState, TEXT("Error: trying to add component to static instance manager as dynamic"));
 
 	TArray<FStreamingRenderAssetPrimitiveInfo> RenderAssetInstanceInfos;
 	Component->GetStreamingRenderAssetInfoWithNULLRemoval(LevelContext, RenderAssetInstanceInfos);
@@ -507,16 +529,20 @@ void FRenderAssetInstanceState::GetReferencedComponents(TArray<const UPrimitiveC
 
 void FRenderAssetInstanceState::UpdateBounds(const UPrimitiveComponent* Component)
 {
+	checkf(bIsDynamicInstanceState, TEXT("Bounds shouldn't be updated after creation unless the instances are dynamic"));
+
 	int32* ComponentLink = ComponentMap.Find(Component);
 	if (ComponentLink)
 	{
+		FBoxSphereBounds Bounds = Component->Bounds;
+		Bounds.SphereRadius = Component->GetStreamingScale();
 		int32 ElementIndex = *ComponentLink;
 		while (ElementIndex != INDEX_NONE)
 		{
 			const FElement& Element = Elements[ElementIndex];
 			if (Element.BoundsIndex != INDEX_NONE)
 			{
-				Bounds4[Element.BoundsIndex / 4].FullUpdate(Element.BoundsIndex % 4, Component->Bounds, Component->GetLastRenderTimeOnScreen());
+				Bounds4[Element.BoundsIndex / 4].FullUpdate(Element.BoundsIndex % 4, Bounds, Component->GetLastRenderTimeOnScreen());
 			}
 			ElementIndex = Element.NextComponentLink;
 		}
@@ -525,10 +551,14 @@ void FRenderAssetInstanceState::UpdateBounds(const UPrimitiveComponent* Componen
 
 bool FRenderAssetInstanceState::UpdateBounds(int32 BoundIndex)
 {
+	checkf(bIsDynamicInstanceState, TEXT("Bounds shouldn't be updated after creation unless the instances are dynamic"));
+
 	const UPrimitiveComponent* Component = ensure(Bounds4Components.IsValidIndex(BoundIndex)) ? Bounds4Components[BoundIndex] : nullptr;
 	if (Component)
 	{
-		Bounds4[BoundIndex / 4].FullUpdate(BoundIndex % 4, Component->Bounds, Component->GetLastRenderTimeOnScreen());
+		FBoxSphereBounds Bounds = Component->Bounds;
+		Bounds.SphereRadius = Component->GetStreamingScale();
+		Bounds4[BoundIndex / 4].FullUpdate(BoundIndex % 4, Bounds, Component->GetLastRenderTimeOnScreen());
 		return true;
 	}
 	else
@@ -539,12 +569,14 @@ bool FRenderAssetInstanceState::UpdateBounds(int32 BoundIndex)
 
 bool FRenderAssetInstanceState::ConditionalUpdateBounds(int32 BoundIndex)
 {
+	checkf(bIsDynamicInstanceState, TEXT("Bounds shouldn't be updated after creation unless the instances are dynamic"));
+
 	const UPrimitiveComponent* Component = ensure(Bounds4Components.IsValidIndex(BoundIndex)) ? Bounds4Components[BoundIndex] : nullptr;
 	if (Component)
 	{
 		if (Component->Mobility != EComponentMobility::Static)
 		{
-			const FBoxSphereBounds Bounds = Component->Bounds;
+			FBoxSphereBounds Bounds = Component->Bounds;
 
 			// Check if the bound is coherent as it could be updated while we read it (from async task).
 			// We don't have to check the position, as if it was partially updated, this should be ok (interp)
@@ -555,13 +587,16 @@ bool FRenderAssetInstanceState::ConditionalUpdateBounds(int32 BoundIndex)
 
 			if (0.5f * FMath::Min3<float>(XSquared, YSquared, ZSquared) <= RadiusSquared && RadiusSquared <= 2.f * (XSquared + YSquared + ZSquared))
 			{
+				Bounds.SphereRadius = Component->GetStreamingScale();
 				Bounds4[BoundIndex / 4].FullUpdate(BoundIndex % 4, Bounds, Component->GetLastRenderTimeOnScreen());
 				return true;
 			}
 		}
 		else // Otherwise we assume it is guarantied to be good.
 		{
-			Bounds4[BoundIndex / 4].FullUpdate(BoundIndex % 4, Component->Bounds, Component->GetLastRenderTimeOnScreen());
+			FBoxSphereBounds Bounds = Component->Bounds;
+			Bounds.SphereRadius = Component->GetStreamingScale();
+			Bounds4[BoundIndex / 4].FullUpdate(BoundIndex % 4, Bounds, Component->GetLastRenderTimeOnScreen());
 			return true;
 		}
 	}
diff --git a/Engine/Source/Runtime/Engine/Private/Streaming/TextureInstanceState.h b/Engine/Source/Runtime/Engine/Private/Streaming/TextureInstanceState.h
--- a/Engine/Source/Runtime/Engine/Private/Streaming/TextureInstanceState.h
+++ b/Engine/Source/Runtime/Engine/Private/Streaming/TextureInstanceState.h
@@ -27,7 +27,7 @@ enum class EAddComponentResult : uint8
 class FRenderAssetInstanceState : public FRenderAssetInstanceView
 {
 public:
-
+	FRenderAssetInstanceState(bool bForDynamicInstances);
 
 	// Will also remove bounds
 	EAddComponentResult AddComponent(const UPrimitiveComponent* Component, FStreamingTextureLevelContext& LevelContext, float MaxAllowedUIDensity);
@@ -92,6 +92,10 @@ private:
 
 	TMap<const UPrimitiveComponent*, int32> ComponentMap;
 
+#if DO_CHECK
+	bool bIsDynamicInstanceState;
+#endif
+
 	friend class FRenderAssetLinkIterator;
 	friend class FRenderAssetIterator;
 };
@@ -101,7 +105,7 @@ class FRenderAssetInstanceStateTaskSync
 {
 public:
 
-	FRenderAssetInstanceStateTaskSync() : State(new FRenderAssetInstanceState()) {}
+	FRenderAssetInstanceStateTaskSync(bool bForDynamicInstances) : State(new FRenderAssetInstanceState(bForDynamicInstances)) {}
 
 	FORCEINLINE void Sync()
 	{
diff --git a/Engine/Source/Runtime/Engine/Private/Streaming/TextureInstanceView.cpp b/Engine/Source/Runtime/Engine/Private/Streaming/TextureInstanceView.cpp
--- a/Engine/Source/Runtime/Engine/Private/Streaming/TextureInstanceView.cpp
+++ b/Engine/Source/Runtime/Engine/Private/Streaming/TextureInstanceView.cpp
@@ -24,7 +24,7 @@ void FRenderAssetInstanceView::FBounds4::Set(int32 Index, const FBoxSphereBounds
 	ExtentX.Component(Index) = Bounds.BoxExtent.X;
 	ExtentY.Component(Index) = Bounds.BoxExtent.Y;
 	ExtentZ.Component(Index) = Bounds.BoxExtent.Z;
-	Radius.Component(Index) = Bounds.SphereRadius;
+	RadiusOrComponentScale.Component(Index) = Bounds.SphereRadius;
 	PackedRelativeBox[Index] = InPackedRelativeBox;
 	MinDistanceSq.Component(Index) = InMinDistanceSq;
 	MinRangeSq.Component(Index) = InMinRangeSq;
@@ -55,7 +55,7 @@ void FRenderAssetInstanceView::FBounds4::UnpackBounds(int32 Index, const UPrimit
 		ExtentX.Component(Index) = SubBounds.BoxExtent.X;
 		ExtentY.Component(Index) = SubBounds.BoxExtent.Y;
 		ExtentZ.Component(Index) = SubBounds.BoxExtent.Z;
-		Radius.Component(Index) = SubBounds.SphereRadius;
+		RadiusOrComponentScale.Component(Index) = SubBounds.SphereRadius;
 		PackedRelativeBox[Index] = PackedRelativeBox_Identity;
 		MinDistanceSq.Component(Index) = MinDistance2;
 		MinRangeSq.Component(Index) = MinRange2;
@@ -77,7 +77,7 @@ void FRenderAssetInstanceView::FBounds4::FullUpdate(int32 Index, const FBoxSpher
 	ExtentX.Component(Index) = Bounds.BoxExtent.X;
 	ExtentY.Component(Index) = Bounds.BoxExtent.Y;
 	ExtentZ.Component(Index) = Bounds.BoxExtent.Z;
-	Radius.Component(Index) = Bounds.SphereRadius;
+	RadiusOrComponentScale.Component(Index) = Bounds.SphereRadius;
 	PackedRelativeBox[Index] = PackedRelativeBox_Identity;
 	MinDistanceSq.Component(Index) = 0;
 	MinRangeSq.Component(Index) = 0;
@@ -114,7 +114,7 @@ FBoxSphereBounds FRenderAssetInstanceView::FRenderAssetLinkConstIterator::GetBou
 		Bounds.BoxExtent.Y = TheBounds4.ExtentY[Index];
 		Bounds.BoxExtent.Z = TheBounds4.ExtentZ[Index];
 
-		Bounds.SphereRadius = TheBounds4.Radius[Index];
+		Bounds.SphereRadius = Bounds.BoxExtent.Length();
 	}
 	return Bounds;
 }
@@ -368,7 +368,7 @@ void FRenderAssetInstanceAsyncView::UpdateBoundSizes_Async(
 		const VectorRegister ExtentX = VectorLoadAligned( &CurrentBounds4.ExtentX );
 		const VectorRegister ExtentY = VectorLoadAligned( &CurrentBounds4.ExtentY );
 		const VectorRegister ExtentZ = VectorLoadAligned( &CurrentBounds4.ExtentZ );
-		const VectorRegister Radius = VectorLoadAligned( &CurrentBounds4.Radius );
+		const VectorRegister ComponentScale = VectorLoadAligned( &CurrentBounds4.RadiusOrComponentScale );
 		const VectorRegister PackedRelativeBox = VectorLoadAligned( reinterpret_cast<const FVector4f*>(&CurrentBounds4.PackedRelativeBox) );
 		const VectorRegister MinDistanceSq = VectorLoadAligned( &CurrentBounds4.MinDistanceSq );
 		const VectorRegister MinRangeSq = VectorLoadAligned( &CurrentBounds4.MinRangeSq );
@@ -420,7 +420,9 @@ void FRenderAssetInstanceAsyncView::UpdateBoundSizes_Async(
 				Temp = VectorSubtract( ViewOriginZ, OriginZ );
 				DistSq = VectorMultiplyAdd( Temp, Temp, DistSq );
 
-				DistSqMinusRadiusSq = VectorNegateMultiplyAdd( Radius, Radius, DistSq );
+				DistSqMinusRadiusSq = VectorNegateMultiplyAdd( ExtentX, ExtentX, DistSq );
+				DistSqMinusRadiusSq = VectorNegateMultiplyAdd( ExtentY, ExtentY, DistSq );
+				DistSqMinusRadiusSq = VectorNegateMultiplyAdd( ExtentZ, ExtentZ, DistSq );
 				// This can be negative here!!!
 			}
 
@@ -466,10 +468,13 @@ void FRenderAssetInstanceAsyncView::UpdateBoundSizes_Async(
 		VectorStoreAligned(MaxNormalizedSize, MaxNormalizedSizeScalar);
 		MS_ALIGN(16) float MaxNormalizedSize_VisibleOnlyScalar[4] GCC_ALIGN(16);
 		VectorStoreAligned(MaxNormalizedSize_VisibleOnly, MaxNormalizedSize_VisibleOnlyScalar);
+		MS_ALIGN(16) float ComponentScaleScalar[4] GCC_ALIGN(16);
+		VectorStoreAligned(ComponentScale, ComponentScaleScalar);
 		for (int32 SubIndex = 0; SubIndex < 4; ++SubIndex)
 		{
 			BoundsVieWInfo[SubIndex].MaxNormalizedSize = MaxNormalizedSizeScalar[SubIndex];
 			BoundsVieWInfo[SubIndex].MaxNormalizedSize_VisibleOnly = MaxNormalizedSize_VisibleOnlyScalar[SubIndex];
+			BoundsVieWInfo[SubIndex].ComponentScale = ComponentScaleScalar[SubIndex];
 		}
 	}
 
@@ -572,7 +577,7 @@ void FRenderAssetInstanceAsyncView::GetRenderAssetScreenSize(
 							&& ensure(CompiledElement.BoundsIndex < View->NumBounds4() * 4))
 						{
 							FRenderAssetInstanceView::FCompiledElement* MutableCompiledElement = const_cast<FRenderAssetInstanceView::FCompiledElement*>(&CompiledElement);
-							MutableCompiledElement->TexelFactor = View->GetBounds4(CompiledElement.BoundsIndex / 4).Radius.Component(CompiledElement.BoundsIndex % 4) * 2.f;
+							MutableCompiledElement->TexelFactor = View->GetBounds4(CompiledElement.BoundsIndex / 4).RadiusOrComponentScale.Component(CompiledElement.BoundsIndex % 4) * 2.f;
 						}
 
 						ProcessElement(
@@ -605,7 +610,7 @@ void FRenderAssetInstanceAsyncView::GetRenderAssetScreenSize(
 				if (ensure(BoundsViewInfo.IsValidIndex(It.GetBoundsIndex())))
 				{
 					const FBoundsViewInfo& BoundsVieWInfo = BoundsViewInfo[It.GetBoundsIndex()];
-					ProcessElement(AssetType, BoundsVieWInfo, It.GetTexelFactor(), It.GetForceLoad(), MaxSize, MaxSize_VisibleOnly, MaxNumForcedLODs);
+					ProcessElement(AssetType, BoundsVieWInfo, AssetType != EStreamableRenderAssetType::Texture ? It.GetTexelFactor() : It.GetTexelFactor() * BoundsVieWInfo.ComponentScale, It.GetForceLoad(), MaxSize, MaxSize_VisibleOnly, MaxNumForcedLODs);
 					if (LogPrefix)
 					{
 						It.OutputToLog(BoundsVieWInfo.MaxNormalizedSize, BoundsVieWInfo.MaxNormalizedSize_VisibleOnly, LogPrefix);
diff --git a/Engine/Source/Runtime/Engine/Private/Streaming/TextureInstanceView.h b/Engine/Source/Runtime/Engine/Private/Streaming/TextureInstanceView.h
--- a/Engine/Source/Runtime/Engine/Private/Streaming/TextureInstanceView.h
+++ b/Engine/Source/Runtime/Engine/Private/Streaming/TextureInstanceView.h
@@ -53,8 +53,8 @@ public:
 		/** Z size of the bounds box extent of 4 texture/mesh instances */
 		FVector4f ExtentZ;
 
-		/** Sphere radii for the bounding sphere of 4 texture/mesh instances */
-		FVector4f Radius;
+		/** Sphere radii for the bounding sphere of 4 texture/mesh static instances or component scale for dynamic instances */
+		FVector4f RadiusOrComponentScale;
 
 		/** The relative box the bound was computed with. Aligned to be interpreted as FVector4  */
 		MS_ALIGN(16) FUintVector4 PackedRelativeBox;
@@ -311,6 +311,8 @@ private:
 		 * Visible instances are the one that are in range and also that have been seen recently.
 		 */
 		float MaxNormalizedSize_VisibleOnly;
+		/** A custom per component scale applied to the texel factor. */
+		float ComponentScale;
 	};
 
 	// Normalized Texel Factors for each bounds and view. This is the data built by ComputeBoundsViewInfos
diff --git a/Engine/Source/Runtime/Engine/Private/Streaming/TextureInstanceView.inl b/Engine/Source/Runtime/Engine/Private/Streaming/TextureInstanceView.inl
--- a/Engine/Source/Runtime/Engine/Private/Streaming/TextureInstanceView.inl
+++ b/Engine/Source/Runtime/Engine/Private/Streaming/TextureInstanceView.inl
@@ -16,7 +16,7 @@ FRenderAssetInstanceView::FBounds4::FBounds4()
 ,	ExtentX( 0, 0, 0, 0 )
 ,	ExtentY( 0, 0, 0, 0 )
 ,	ExtentZ( 0, 0, 0, 0 )
-,	Radius( 0, 0, 0, 0 )
+,	RadiusOrComponentScale( 0, 0, 0, 0 )
 ,	PackedRelativeBox(0, 0, 0, 0)
 ,	MinDistanceSq( 0, 0, 0, 0 )
 ,	MinRangeSq( 0, 0, 0, 0 )
@@ -38,7 +38,7 @@ void FRenderAssetInstanceView::FBounds4::Clear(int32 Index)
 	ExtentX.Component(Index) = 0;
 	ExtentY.Component(Index) = 0;
 	ExtentZ.Component(Index) = 0;
-	Radius.Component(Index) = 0;
+	RadiusOrComponentScale.Component(Index) = 0;
 	PackedRelativeBox[Index] = 0;
 	MinDistanceSq.Component(Index) = 0; 
 	MinRangeSq.Component(Index) = 0;
