LevelInstance: Fix invalid guids on PIE duplicate

diff --git a/Engine/Source/Runtime/Engine/Private/LevelInstance/LevelInstanceActorGuid.cpp b/Engine/Source/Runtime/Engine/Private/LevelInstance/LevelInstanceActorGuid.cpp
--- a/Engine/Source/Runtime/Engine/Private/LevelInstance/LevelInstanceActorGuid.cpp
+++ b/Engine/Source/Runtime/Engine/Private/LevelInstance/LevelInstanceActorGuid.cpp
@@ -41,10 +41,25 @@ FArchive& operator<<(FArchive& Ar, FLevelInstanceActorGuid& LevelInstanceActorGu
 {
 	check(LevelInstanceActorGuid.Actor);
 #if WITH_EDITOR
-	if (Ar.IsSaving() && Ar.IsCooking() && !LevelInstanceActorGuid.Actor->IsTemplate())
+	if (!LevelInstanceActorGuid.Actor->IsTemplate())
 	{
-		FGuid Guid = LevelInstanceActorGuid.GetGuid();
-		Ar << Guid;
+		if (Ar.IsSaving() && Ar.IsCooking())
+		{
+			FGuid Guid = LevelInstanceActorGuid.GetGuid();
+			Ar << Guid;
+		}
+		else if (Ar.GetPortFlags() & PPF_DuplicateForPIE)
+		{
+			if (Ar.IsLoading())
+			{
+				Ar << LevelInstanceActorGuid.ActorGuid;
+			}
+			else if (Ar.IsSaving())
+			{
+				FGuid Guid = LevelInstanceActorGuid.GetGuid();
+				Ar << Guid;
+			}
+		}
 	}
 #else
 	if (Ar.IsLoading())
@@ -65,4 +80,4 @@ FArchive& operator<<(FArchive& Ar, FLevelInstanceActorGuid& LevelInstanceActorGu
 	}
 #endif
 	return Ar;
-}
\ No newline at end of file
+}
