LevelInstance: Fix crash caused by missing AddReferencedObjects on temp level instance edit

diff --git a/Engine/Source/Runtime/Engine/Private/LevelInstance/LevelInstanceSubsystem.cpp b/Engine/Source/Runtime/Engine/Private/LevelInstance/LevelInstanceSubsystem.cpp
--- a/Engine/Source/Runtime/Engine/Private/LevelInstance/LevelInstanceSubsystem.cpp
+++ b/Engine/Source/Runtime/Engine/Private/LevelInstance/LevelInstanceSubsystem.cpp
@@ -44,6 +44,7 @@
 #include "EdGraph/EdGraph.h"
 #include "UObject/ObjectSaveContext.h"
 #include "UObject/UObjectGlobals.h"
+#include "UObject/GCObject.h"
 #include "EditorActorFolders.h"
 #include "Misc/MessageDialog.h"
 #include "Subsystems/ActorEditorContextSubsystem.h"
@@ -1082,9 +1083,32 @@ ILevelInstanceInterface* ULevelInstanceSubsystem::CreateLevelInstanceFrom(const
 	
 	// New level instance
 	const FLevelInstanceID NewLevelInstanceID = NewLevelInstance->GetLevelInstanceID();
-	TUniquePtr<FLevelInstanceEdit> TempLevelInstanceEdit = MakeUnique<FLevelInstanceEdit>(LevelStreaming, NewLevelInstance->GetLevelInstanceID());
+
+	struct FStackLevelInstanceEdit : public FGCObject
+	{
+		TUniquePtr<FLevelInstanceEdit> LevelInstanceEdit;
+
+		virtual ~FStackLevelInstanceEdit() {}
+
+		//~ FGCObject
+		virtual void AddReferencedObjects(FReferenceCollector& Collector) override
+		{
+			if (LevelInstanceEdit)
+			{
+				LevelInstanceEdit->AddReferencedObjects(Collector);
+			}
+		}
+		virtual FString GetReferencerName() const override
+		{
+			return TEXT("FStackLevelInstanceEdit");
+		}
+		//~ FGCObject
+	};
+
+	FStackLevelInstanceEdit StackLevelInstanceEdit;
+	StackLevelInstanceEdit.LevelInstanceEdit = MakeUnique<FLevelInstanceEdit>(LevelStreaming, NewLevelInstance->GetLevelInstanceID());
 	// Force mark it as changed
-	TempLevelInstanceEdit->MarkCommittedChanges();
+	StackLevelInstanceEdit.LevelInstanceEdit->MarkCommittedChanges();
 
 	GetWorld()->SetCurrentLevel(LoadedLevel);
 
@@ -1096,9 +1120,9 @@ ILevelInstanceInterface* ULevelInstanceSubsystem::CreateLevelInstanceFrom(const
 	
 	// Commit will always pop the actor editor context, make sure to push one here
 	UActorEditorContextSubsystem::Get()->PushContext();
-	bool bCommitted = CommitLevelInstanceInternal(TempLevelInstanceEdit, /*bDiscardEdits=*/false, /*bDiscardOnFailure=*/true, &DirtyPackages);
+	bool bCommitted = CommitLevelInstanceInternal(StackLevelInstanceEdit.LevelInstanceEdit, /*bDiscardEdits=*/false, /*bDiscardOnFailure=*/true, &DirtyPackages);
 	check(bCommitted);
-	check(!TempLevelInstanceEdit);
+	check(!StackLevelInstanceEdit.LevelInstanceEdit);
 
 	// EditorLevelUtils::CreateNewStreamingLevelForWorld deactivates all modes. Re-activate if needed
 	if (LevelInstanceEdit)
