MaterialShader loading from DDC: Fix deadlock caused by cancelling within a DDC callback.
#rnx

diff --git a/Engine/Source/Runtime/Engine/Private/Materials/MaterialShader.cpp b/Engine/Source/Runtime/Engine/Private/Materials/MaterialShader.cpp
--- a/Engine/Source/Runtime/Engine/Private/Materials/MaterialShader.cpp
+++ b/Engine/Source/Runtime/Engine/Private/Materials/MaterialShader.cpp
@@ -1611,6 +1611,9 @@ TSharedRef<FMaterialShaderMap::FAsyncLoadContext> FMaterialShaderMap::BeginLoadF
 				GetCache().GetValue({Request}, *Result->RequestOwner, [Result](FCacheGetValueResponse&& Response)
 				{
 					Result->CachedData = Response.Value.GetData().Decompress();
+					// This callback might hold the last reference to Result, which owns RequestOwner, so
+					// we must not cancel in the Owner's destructor; Canceling in a callback will deadlock.
+					Result->RequestOwner->KeepAlive();
 				});
 			}
 		}
