Fixed issue caused by different actor folders coming from different level instances (but with the same resulting path) added to the same runtime level and generating duplicate folders.
Disable actor folder feature on generated runtime levels (for PIE), since there's no lost functionality.
#okfrogithub public

diff --git a/Engine/Source/Editor/SceneOutliner/Private/ActorFolderHierarchy.cpp b/Engine/Source/Editor/SceneOutliner/Private/ActorFolderHierarchy.cpp
--- a/Engine/Source/Editor/SceneOutliner/Private/ActorFolderHierarchy.cpp
+++ b/Engine/Source/Editor/SceneOutliner/Private/ActorFolderHierarchy.cpp
@@ -8,6 +8,7 @@
 #include "ActorFolderTreeItem.h"
 #include "EditorActorFolders.h"
 #include "EditorFolderUtils.h"
+#include "ActorMode.h"
 #include "LevelInstance/LevelInstanceInterface.h"
 
 FActorFolderHierarchy::FActorFolderHierarchy(ISceneOutlinerMode* InMode, const TWeakObjectPtr<UWorld>& World, const FFolder::FRootObject& InRootObject)
@@ -102,8 +103,7 @@ void FActorFolderHierarchy::CreateWorldChildren(UWorld* World, TArray<FSceneOutl
 		}
 		else if (ULevel* RootLevel = Cast<ULevel>(RootObjectPtr))
 		{
-			// Don't create a LevelTreeItems for the PersistentLevel
-			if (!RootLevel->IsPersistentLevel())
+			if (FActorMode::IsActorLevelDisplayable(RootLevel))
 			{
 				if (FSceneOutlinerTreeItemPtr ActorItem = Mode->CreateItemFor<FLevelTreeItem>(RootLevel, true))
 				{
diff --git a/Engine/Source/Editor/SceneOutliner/Private/ActorHierarchy.cpp b/Engine/Source/Editor/SceneOutliner/Private/ActorHierarchy.cpp
--- a/Engine/Source/Editor/SceneOutliner/Private/ActorHierarchy.cpp
+++ b/Engine/Source/Editor/SceneOutliner/Private/ActorHierarchy.cpp
@@ -25,6 +25,7 @@
 #include "WorldPartition/IWorldPartitionEditorModule.h"
 #include "WorldPartition/WorldPartitionSubsystem.h"
 #include "ActorFolder.h"
+#include "ActorMode.h"
 
 TUniquePtr<FActorHierarchy> FActorHierarchy::Create(ISceneOutlinerMode* Mode, const TWeakObjectPtr<UWorld>& World)
 {
@@ -204,7 +205,7 @@ FSceneOutlinerTreeItemPtr FActorHierarchy::FindOrCreateParentItem(const ISceneOu
 			// Parent Level Using Actor Folders
 			ULevel* OwningLevel = Cast<ULevel>(Folder.GetRootObjectPtr());
 			// For the persistent level, fallback on the world
-			if (OwningLevel && !OwningLevel->IsPersistentLevel())
+			if (FActorMode::IsActorLevelDisplayable(OwningLevel))
 			{
 				if (const FSceneOutlinerTreeItemPtr* ParentItem = Items.Find(OwningLevel))
 				{
@@ -255,8 +256,7 @@ FSceneOutlinerTreeItemPtr FActorHierarchy::FindOrCreateParentItem(const ISceneOu
 		// Parent Level Using Actor Folders
 		else if (ULevel* OwningLevel = Cast<ULevel>(ParentPath.GetRootObjectPtr()))
 		{
-			// For the persistent level, fallback on the world
-			if (!OwningLevel->IsPersistentLevel())
+			if (FActorMode::IsActorLevelDisplayable(OwningLevel))
 			{
 				if (const FSceneOutlinerTreeItemPtr* ParentItem = Items.Find(OwningLevel))
 				{
diff --git a/Engine/Source/Editor/SceneOutliner/Private/ActorMode.cpp b/Engine/Source/Editor/SceneOutliner/Private/ActorMode.cpp
--- a/Engine/Source/Editor/SceneOutliner/Private/ActorMode.cpp
+++ b/Engine/Source/Editor/SceneOutliner/Private/ActorMode.cpp
@@ -343,6 +343,12 @@ bool FActorMode::IsActorDisplayable(const SSceneOutliner* SceneOutliner, const A
 		FLevelUtils::IsLevelVisible(Actor->GetLevel());											// Only show Actors whose level is visible
 }
 
+bool FActorMode::IsActorLevelDisplayable(ULevel* InLevel)
+{
+	// Don't show level tree item for the persistent level
+	return (InLevel && !InLevel->IsPersistentLevel());
+}
+
 void FActorMode::OnFilterTextChanged(const FText& InFilterText)
 {
 	// Scroll last item (if it passes the filter) into view - this means if we are multi-selecting, we show newest selection that passes the filter
diff --git a/Engine/Source/Editor/SceneOutliner/Public/ActorMode.h b/Engine/Source/Editor/SceneOutliner/Public/ActorMode.h
--- a/Engine/Source/Editor/SceneOutliner/Public/ActorMode.h
+++ b/Engine/Source/Editor/SceneOutliner/Public/ActorMode.h
@@ -61,6 +61,7 @@ public:
 	virtual int32 GetTypeSortPriority(const ISceneOutlinerTreeItem& Item) const override;
 
 	static bool IsActorDisplayable(const SSceneOutliner* SceneOutliner, const AActor* Actor);
+	static bool IsActorLevelDisplayable(ULevel* InLevel);
 
 	virtual FFolder::FRootObject GetRootObject() const override;
 	virtual FFolder::FRootObject GetPasteTargetRootObject() const override;
diff --git a/Engine/Source/Runtime/Engine/Classes/GameFramework/Actor.h b/Engine/Source/Runtime/Engine/Classes/GameFramework/Actor.h
--- a/Engine/Source/Runtime/Engine/Classes/GameFramework/Actor.h
+++ b/Engine/Source/Runtime/Engine/Classes/GameFramework/Actor.h
@@ -2410,8 +2410,11 @@ public:
 	UFUNCTION(BlueprintCallable, Category = "Editor Scripting | Actor Editing")
 	FName GetFolderPath() const;
 
-	/** Returns actor folder guid. If level is not using actor folder objects, returns an invalid guid. */
-	FGuid GetFolderGuid() const;
+	/*
+	 * Returns actor folder guid. If level is not using actor folder objects, returns an invalid guid. 
+	 * @param bDirectAccess If true, returns the raw value without testing if level uses Actor Folders.
+	 */
+	FGuid GetFolderGuid(bool bDirectAccess = false) const;
 
 	/** Returns the actor's folder root object. Null, is interpreted as the actor's world. */
 	FFolder::FRootObject GetFolderRootObject() const;
diff --git a/Engine/Source/Runtime/Engine/Private/ActorEditor.cpp b/Engine/Source/Runtime/Engine/Private/ActorEditor.cpp
--- a/Engine/Source/Runtime/Engine/Private/ActorEditor.cpp
+++ b/Engine/Source/Runtime/Engine/Private/ActorEditor.cpp
@@ -1114,9 +1114,9 @@ void AActor::FixupActorFolder()
 	}
 }
 
-FGuid AActor::GetFolderGuid() const
+FGuid AActor::GetFolderGuid(bool bDirectAccess) const
 {
-	return IsUsingActorFolders(this) ? FolderGuid : FGuid();
+	return bDirectAccess || IsUsingActorFolders(this) ? FolderGuid : FGuid();
 }
 
 FName AActor::GetFolderPath() const
diff --git a/Engine/Source/Runtime/Engine/Private/WorldPartition/WorldPartitionLevelHelper.cpp b/Engine/Source/Runtime/Engine/Private/WorldPartition/WorldPartitionLevelHelper.cpp
--- a/Engine/Source/Runtime/Engine/Private/WorldPartition/WorldPartitionLevelHelper.cpp
+++ b/Engine/Source/Runtime/Engine/Private/WorldPartition/WorldPartitionLevelHelper.cpp
@@ -278,38 +278,6 @@ ULevel* FWorldPartitionLevelHelper::CreateEmptyLevelForRuntimeCell(const UWorldP
 	return NewLevel;
 }
 
-void FWorldPartitionLevelHelper::DuplicateActorFolderToRuntimeCell(ULevel* CellLevel, ULevel* SrcLevel, const FGuid& ActorFolderGuid)
-{
-	check(SrcLevel->IsUsingActorFolders() && CellLevel->IsUsingActorFolders());
-		
-	FGuid CurrentGuid = ActorFolderGuid;
-	while (CurrentGuid.IsValid() && !CellLevel->GetActorFolder(CurrentGuid))
-	{
-		const bool bSkipDeleted = false;
-		UActorFolder* ActorFolder = SrcLevel->GetActorFolder(CurrentGuid, bSkipDeleted);
-		CurrentGuid.Invalidate();
-		if (ActorFolder)
-		{
-			CurrentGuid.Invalidate();
-			FObjectDuplicationParameters DupParams = InitStaticDuplicateObjectParams(ActorFolder, CellLevel);
-			DupParams.ApplyFlags |= RF_Transient;
-			DupParams.FlagMask &= ~RF_Transactional;
-			DupParams.bAssignExternalPackages = false;
-			UActorFolder* DuplicatedFolder = Cast<UActorFolder>(StaticDuplicateObjectEx(DupParams));
-
-			const bool bShouldDirtyLevel = false;
-			const bool bShouldBroadcast = false;
-			FLevelActorFoldersHelper::AddActorFolder(CellLevel, DuplicatedFolder, bShouldDirtyLevel, bShouldBroadcast);
-
-			// Continue to parent
-			if (UActorFolder* ParentFolder = ActorFolder->GetParent(bSkipDeleted))
-			{
-				CurrentGuid = ParentFolder->GetGuid();
-			}
-		}
-	}
-}
-
 bool FWorldPartitionLevelHelper::LoadActors(UWorld* InOwningWorld, ULevel* InDestLevel, TArrayView<FWorldPartitionRuntimeCellObjectMapping> InActorPackages, FWorldPartitionLevelHelper::FPackageReferencer& InPackageReferencer, TFunction<void(bool)> InCompletionCallback, bool bInLoadAsync, FLinkerInstancingContext InstancingContext)
 {
 	UPackage* DestPackage = InDestLevel ? InDestLevel->GetPackage() : nullptr;
@@ -378,20 +346,25 @@ bool FWorldPartitionLevelHelper::LoadActors(UWorld* InOwningWorld, ULevel* InDes
 			{
 				const UWorld* ContainerWorld = PackageObjectMapping->ContainerID.IsMainContainer() ? InOwningWorld : Actor->GetTypedOuter<UWorld>();
 				
-				// Duplicate Folder if needed (this will recreate a transient folder structure in InDestLevel if it doesn't exist (only used in PIE)
-				if (FGuid ActorFolderGuid = Actor->GetFolderGuid(); InDestLevel && ActorFolderGuid.IsValid())
-				{
-					check(InDestLevel->IsUsingActorFolders());
-					InDestLevel->bFixupActorFoldersAtLoad = false;
+				TOptional<FName> SrcActorFolderPath;
 
-					// Make sure Source level actor folder fixup was called
+				// Make sure Source level actor folder fixup was called
+				if (ContainerWorld->PersistentLevel->IsUsingActorFolders())
+				{ 
 					if (!ContainerWorld->PersistentLevel->LoadedExternalActorFolders.IsEmpty())
 					{
 						ContainerWorld->PersistentLevel->bFixupActorFoldersAtLoad = false;
 						ContainerWorld->PersistentLevel->FixupActorFolders();
 					}
-					
-					FWorldPartitionLevelHelper::DuplicateActorFolderToRuntimeCell(InDestLevel, ContainerWorld->PersistentLevel, ActorFolderGuid);
+				
+					// Since actor's level doesn't necessarily uses actor folders, access Folder Guid directly
+					const bool bDirectAccess = true;
+					const FGuid ActorFolderGuid = Actor->GetFolderGuid(bDirectAccess);
+					// Resolve folder guid from source container level and resolve/backup the folder path
+					if (UActorFolder* SrcFolder = ContainerWorld->PersistentLevel->GetActorFolder(ActorFolderGuid))
+					{
+						SrcActorFolderPath = SrcFolder->GetPath();
+					}
 				}
 
 				if (!PackageObjectMapping->ContainerID.IsMainContainer())
@@ -427,6 +400,13 @@ bool FWorldPartitionLevelHelper::LoadActors(UWorld* InOwningWorld, ULevel* InDes
 
 				if (InDestLevel)
 				{
+					// Propagate resolved actor folder path
+					check(!InDestLevel->IsUsingActorFolders());
+					if (SrcActorFolderPath.IsSet())
+					{
+						Actor->SetFolderPath(*SrcActorFolderPath);
+					}
+
 					check(Actor->IsPackageExternal());
 					InDestLevel->Actors.Add(Actor);
 					checkf(Actor->GetLevel() == InDestLevel, TEXT("Levels mismatch, got : %s, expected: %s\nActor: %s\nActorFullName: %s\nActorPackage: %s"), *InDestLevel->GetFullName(), *Actor->GetLevel()->GetFullName(), *Actor->GetActorNameOrLabel(), *Actor->GetFullName(), *Actor->GetPackage()->GetFullName());
diff --git a/Engine/Source/Runtime/Engine/Private/WorldPartition/WorldPartitionLevelStreamingDynamic.cpp b/Engine/Source/Runtime/Engine/Private/WorldPartition/WorldPartitionLevelStreamingDynamic.cpp
--- a/Engine/Source/Runtime/Engine/Private/WorldPartition/WorldPartitionLevelStreamingDynamic.cpp
+++ b/Engine/Source/Runtime/Engine/Private/WorldPartition/WorldPartitionLevelStreamingDynamic.cpp
@@ -134,10 +134,9 @@ void UWorldPartitionLevelStreamingDynamic::CreateRuntimeLevel()
 	RuntimeLevel = FWorldPartitionLevelHelper::CreateEmptyLevelForRuntimeCell(StreamingCell.Get(), World, GetWorldAsset().ToString());
 	check(RuntimeLevel);
 
-	if (World->PersistentLevel->IsUsingActorFolders())
-	{
-		FLevelActorFoldersHelper::SetUseActorFolders(RuntimeLevel, true);
-	}
+	// Make sure Actor Folders is disabled on generated runtime levels to avoid any problems with duplicate folders that
+	// can be caused by level instances injecting their actors, which can cause duplicate folders (which only happens during PIE).
+	FLevelActorFoldersHelper::SetUseActorFolders(RuntimeLevel, false);
 
 	UPackage* RuntimeLevelPackage = RuntimeLevel->GetPackage();
 	check(RuntimeLevelPackage);
diff --git a/Engine/Source/Runtime/Engine/Public/WorldPartition/WorldPartitionLevelHelper.h b/Engine/Source/Runtime/Engine/Public/WorldPartition/WorldPartitionLevelHelper.h
--- a/Engine/Source/Runtime/Engine/Public/WorldPartition/WorldPartitionLevelHelper.h
+++ b/Engine/Source/Runtime/Engine/Public/WorldPartition/WorldPartitionLevelHelper.h
@@ -32,7 +32,6 @@ public:
 	};
 
 	static ULevel* CreateEmptyLevelForRuntimeCell(const UWorldPartitionRuntimeCell* Cell, const UWorld* InWorld, const FString& InWorldAssetName, UPackage* DestPackage = nullptr);
-	static void DuplicateActorFolderToRuntimeCell(ULevel* CellLevel, ULevel* SrcLevel, const FGuid& ActorFolderGuid);
 	static void MoveExternalActorsToLevel(const TArray<FWorldPartitionRuntimeCellObjectMapping>& InChildPackages, ULevel* InLevel);
 	static void RemapLevelSoftObjectPaths(ULevel* InLevel, UWorldPartition* InWorldPartition);
 	
