Texture fix check on ForceAlpha
Add ForceAlpha bool and propagate it

diff --git a/Engine/Source/Runtime/Engine/Classes/Engine/Texture.h b/Engine/Source/Runtime/Engine/Classes/Engine/Texture.h
--- a/Engine/Source/Runtime/Engine/Classes/Engine/Texture.h
+++ b/Engine/Source/Runtime/Engine/Classes/Engine/Texture.h
@@ -1149,6 +1149,10 @@ public:
 	/** If enabled, the texture's alpha channel will be forced to opaque for any compressed texture output format.  Does not apply if output format is uncompressed RGBA. */
 	UPROPERTY(EditAnywhere, BlueprintReadWrite, Category=Compression, meta=(DisplayName="Compress Without Alpha"))
 	uint32 CompressionNoAlpha:1;
+	
+	/** If true, force alpha channel in output format when possible, eg. for AutoDXT BC1/BC3 choice **/
+	UPROPERTY()
+	uint32 CompressionForceAlpha:1;
 
 	/** If true, force the texture to be uncompressed no matter the format. */
 	UPROPERTY()
diff --git a/Engine/Source/Runtime/Engine/Private/Texture.cpp b/Engine/Source/Runtime/Engine/Private/Texture.cpp
--- a/Engine/Source/Runtime/Engine/Private/Texture.cpp
+++ b/Engine/Source/Runtime/Engine/Private/Texture.cpp
@@ -124,6 +124,7 @@ UTexture::UTexture(const FObjectInitializer& ObjectInitializer)
 	AdjustMinAlpha = 0.0f;
 	AdjustMaxAlpha = 1.0f;
 	CompressionNoAlpha = 0;
+	CompressionForceAlpha = 0;
 	CompressionNone = 0;
 	CompressFinal = 0;
 	DeferCompression = 0;
@@ -2680,8 +2681,7 @@ void UTexture::GetDefaultFormatSettings(FTextureFormatSettings& OutSettings) con
 	OutSettings.CompressionSettings = CompressionSettings;
 	OutSettings.CompressionNone = CompressionNone;
 	OutSettings.CompressionNoAlpha = CompressionNoAlpha;
-	// Texture does not have CompressionForceAlpha
-	OutSettings.CompressionForceAlpha = false;
+	OutSettings.CompressionForceAlpha = CompressionForceAlpha;
 	OutSettings.CompressionYCoCg = CompressionYCoCg;
 	OutSettings.SRGB = SRGB;
 }
@@ -2708,9 +2708,7 @@ void UTexture::SetLayerFormatSettings(int32 LayerIndex, const FTextureFormatSett
 		CompressionSettings = InSettings.CompressionSettings;
 		CompressionNone = InSettings.CompressionNone;
 		CompressionNoAlpha = InSettings.CompressionNoAlpha;
-		// CompressionForceAlpha not copied
-		//CompressionForceAlpha = InSettings.CompressionForceAlpha;
-		check( ! InSettings.CompressionForceAlpha );
+		CompressionForceAlpha = InSettings.CompressionForceAlpha;
 		CompressionYCoCg = InSettings.CompressionYCoCg;
 		SRGB = InSettings.SRGB;
 	}
@@ -3129,6 +3127,9 @@ FName GetDefaultTextureFormatName( const ITargetPlatform* TargetPlatform, const
 		}
 		else
 		{
+			// CompressionForceAlpha is applied later, where the bHasAlpha/DetectAlpha check is done and AutoDXT is resolved
+			// alternatively it could be mapped immediately to NameDXT5 here.
+
 			TextureFormatName = NameAutoDXT;
 		}
 	}
