Map Check: tread actors that needs to be resaved as warnings, not errors.

diff --git a/Engine/Source/Runtime/Engine/Private/WorldPartition/ErrorHandling/WorldPartitionStreamingGenerationTokenizedMessageErrorHandler.cpp b/Engine/Source/Runtime/Engine/Private/WorldPartition/ErrorHandling/WorldPartitionStreamingGenerationTokenizedMessageErrorHandler.cpp
--- a/Engine/Source/Runtime/Engine/Private/WorldPartition/ErrorHandling/WorldPartitionStreamingGenerationTokenizedMessageErrorHandler.cpp
+++ b/Engine/Source/Runtime/Engine/Private/WorldPartition/ErrorHandling/WorldPartitionStreamingGenerationTokenizedMessageErrorHandler.cpp
@@ -133,7 +133,7 @@ void ITokenizedMessageErrorHandler::OnDataLayerAssetConflict(const UDataLayerIns
 
 void ITokenizedMessageErrorHandler::OnActorNeedsResave(const FWorldPartitionActorDescView& ActorDescView)
 {
-	TSharedRef<FTokenizedMessage> Message = FTokenizedMessage::Create(EMessageSeverity::Error);
+	TSharedRef<FTokenizedMessage> Message = FTokenizedMessage::Create(EMessageSeverity::Warning);
 	Message->AddToken(FTextToken::Create(LOCTEXT("TokenMessage_WorldPartition_ActorNeedsResave", "Actor needs resave")))
 		->AddToken(FActorToken::Create(ActorDescView.GetActorPath().ToString(), ActorDescView.GetGuid(), FText::FromName(ActorDescView.GetActorLabelOrName())));
 
