Moving some generic Audio Settings out of the Quality Tab

diff --git a/Engine/Source/Runtime/Engine/Classes/Sound/AudioSettings.h b/Engine/Source/Runtime/Engine/Classes/Sound/AudioSettings.h
--- a/Engine/Source/Runtime/Engine/Classes/Sound/AudioSettings.h
+++ b/Engine/Source/Runtime/Engine/Classes/Sound/AudioSettings.h
@@ -190,34 +190,34 @@ class ENGINE_API UAudioSettings : public UDeveloperSettings
 	TArray<FAudioQualitySettings> QualityLevels;
 
 	/** Allows sounds to play at 0 volume. */
-	UPROPERTY(config, EditAnywhere, Category = "Quality", AdvancedDisplay)
+	UPROPERTY(config, EditAnywhere, Category = "Audio", AdvancedDisplay)
 	uint32 bAllowPlayWhenSilent:1;
 
 	/** Disables master EQ effect in the audio DSP graph. */
-	UPROPERTY(config, EditAnywhere, Category = "Quality", AdvancedDisplay)
+	UPROPERTY(config, EditAnywhere, Category = "Mix", AdvancedDisplay)
 	uint32 bDisableMasterEQ : 1;
 
 	/** Enables the surround sound spatialization calculations to include the center channel. */
-	UPROPERTY(config, EditAnywhere, Category = "Quality", AdvancedDisplay)
+	UPROPERTY(config, EditAnywhere, Category = "Panning", AdvancedDisplay)
 	uint32 bAllowCenterChannel3DPanning : 1;
 
 	/**
 	 * The max number of sources to reserve for "stopping" sounds. A "stopping" sound applies a fast fade in the DSP
 	 * render to prevent discontinuities when stopping sources.
 	 */
-	UPROPERTY(config, EditAnywhere, Category = "Quality", AdvancedDisplay)
+	UPROPERTY(config, EditAnywhere, Category = "Audio", AdvancedDisplay)
 	uint32 NumStoppingSources;
 
 	/**
 	* The method to use when doing non-binaural or object-based panning.
 	*/
-	UPROPERTY(config, EditAnywhere, Category = "Quality", AdvancedDisplay)
+	UPROPERTY(config, EditAnywhere, Category = "Panning", AdvancedDisplay)
 	EPanningMethod PanningMethod;
 
 	/**
 	* The upmixing method for mono sound sources. Defines how mono channels are up-mixed to stereo channels.
 	*/
-	UPROPERTY(config, EditAnywhere, Category = "Quality", AdvancedDisplay)
+	UPROPERTY(config, EditAnywhere, Category = "Mix", AdvancedDisplay)
 	EMonoChannelUpmixMethod MonoChannelUpmixMethod;
 
 	/**
