[KillPendingKill] Mark NetDriver and related objects as garbage as much as possible when destroying the net driver.
Release references to objects from open channels on clients during net driver shutdown.

diff --git a/Engine/Source/Runtime/Engine/Classes/Engine/ActorChannel.h b/Engine/Source/Runtime/Engine/Classes/Engine/ActorChannel.h
--- a/Engine/Source/Runtime/Engine/Classes/Engine/ActorChannel.h
+++ b/Engine/Source/Runtime/Engine/Classes/Engine/ActorChannel.h
@@ -155,6 +155,9 @@ public:
 	virtual int64 Close(EChannelCloseReason Reason) override;
 	virtual FString Describe() override;
 
+	/** Release any references this channel is holding to UObjects and object replicators and mark it as broken. */
+	void BreakAndReleaseReferences();
+
 	/** UActorChannel interface and accessors. */
 	AActor* GetActor() {return Actor;}
 
diff --git a/Engine/Source/Runtime/Engine/Private/DataChannel.cpp b/Engine/Source/Runtime/Engine/Private/DataChannel.cpp
--- a/Engine/Source/Runtime/Engine/Private/DataChannel.cpp
+++ b/Engine/Source/Runtime/Engine/Private/DataChannel.cpp
@@ -2156,6 +2156,14 @@ void UActorChannel::SetClosingFlag()
 	UChannel::SetClosingFlag();
 }
 
+void UActorChannel::BreakAndReleaseReferences()
+{
+	Actor = nullptr;
+	Broken = true;
+	CleanupReplicators();
+	CreateSubObjects.Empty();
+}
+
 int64 UActorChannel::Close(EChannelCloseReason Reason)
 {
 	FScopedRepContext RepContext(Connection, Actor);
diff --git a/Engine/Source/Runtime/Engine/Private/NetConnection.cpp b/Engine/Source/Runtime/Engine/Private/NetConnection.cpp
--- a/Engine/Source/Runtime/Engine/Private/NetConnection.cpp
+++ b/Engine/Source/Runtime/Engine/Private/NetConnection.cpp
@@ -1190,6 +1190,7 @@ void UNetConnection::CleanUp()
 	SetClientLoginState(EClientLoginState::CleanedUp);
 
 	Driver = nullptr;
+	MarkAsGarbage();
 
 #if UE_NET_TRACE_ENABLED	
 	UE_NET_TRACE_DESTROY_COLLECTOR(InTraceCollector);
diff --git a/Engine/Source/Runtime/Engine/Private/NetDriver.cpp b/Engine/Source/Runtime/Engine/Private/NetDriver.cpp
--- a/Engine/Source/Runtime/Engine/Private/NetDriver.cpp
+++ b/Engine/Source/Runtime/Engine/Private/NetDriver.cpp
@@ -1776,7 +1776,7 @@ void UNetDriver::Shutdown()
 				 {
 					 GuidsToLog.Append(ActorChannel->PendingGuidResolves);
 				 }
-				 ActorChannel->CleanupReplicators();
+				 ActorChannel->BreakAndReleaseReferences();
 			 }
 		}
 
@@ -1804,6 +1804,7 @@ void UNetDriver::Shutdown()
 		// Calls Channel[0]->Close to send a close bunch to server
 		ServerConnection->Close();
 		ServerConnection->FlushNet();
+		ServerConnection->MarkAsGarbage();
 	}
 
 	// Server closing connections with clients
@@ -1841,6 +1842,10 @@ void UNetDriver::Shutdown()
 	ReplicationChangeListMap.Empty();
 
 	// Clean up the actor channel pool
+	for (TObjectPtr<UChannel> Channel : ActorChannelPool)
+	{
+		Channel->MarkAsGarbage();
+	}
 	ActorChannelPool.Empty();
 
 	ConnectionlessHandler.Reset(nullptr);
@@ -3618,9 +3623,7 @@ void UNetDriver::NotifyActorLevelUnloaded( AActor* TheActor )
 		if (Channel != NULL)
 		{
 			ServerConnection->RemoveActorChannel(TheActor);
-			Channel->Actor = NULL;
-			Channel->Broken = true;
-			Channel->CleanupReplicators();
+			Channel->BreakAndReleaseReferences();
 		}
 	}
 }
diff --git a/Engine/Source/Runtime/Engine/Private/UnrealEngine.cpp b/Engine/Source/Runtime/Engine/Private/UnrealEngine.cpp
--- a/Engine/Source/Runtime/Engine/Private/UnrealEngine.cpp
+++ b/Engine/Source/Runtime/Engine/Private/UnrealEngine.cpp
@@ -13117,6 +13117,7 @@ void DestroyNamedNetDriver_Local(FWorldContext &Context, FName NetDriverName)
 			NetDriver->SetWorld(NULL);
 			NetDriver->Shutdown();
 			NetDriver->LowLevelDestroy();
+			NetDriver->MarkAsGarbage();
 
 			Context.World()->ClearNetDriver(NetDriver);
 
