Fixed missing files in my previous cl 20440722

diff --git a/Engine/Source/Runtime/Engine/Private/World.cpp b/Engine/Source/Runtime/Engine/Private/World.cpp
--- a/Engine/Source/Runtime/Engine/Private/World.cpp
+++ b/Engine/Source/Runtime/Engine/Private/World.cpp
@@ -5752,6 +5752,10 @@ void UWorld::WelcomePlayer(UNetConnection* Connection)
 		LevelName = CurrentLevel->GetOutermost()->GetName();
 		Connection->SetClientWorldPackageName(CurrentLevel->GetOutermost()->GetFName());
 	}
+	if (UGameInstance* GameInst = GetGameInstance())
+	{
+		GameInst->ModifyClientTravelLevelURL(LevelName);
+	}
 
 	FString GameName;
 	FString RedirectURL;
