World Partition: added helper functions to compare unsorted arrays.

#rnx

diff --git a/Engine/Source/Runtime/Engine/Private/WorldPartition/DataLayer/WorldDataLayersActorDesc.cpp b/Engine/Source/Runtime/Engine/Private/WorldPartition/DataLayer/WorldDataLayersActorDesc.cpp
--- a/Engine/Source/Runtime/Engine/Private/WorldPartition/DataLayer/WorldDataLayersActorDesc.cpp
+++ b/Engine/Source/Runtime/Engine/Private/WorldPartition/DataLayer/WorldDataLayersActorDesc.cpp
@@ -152,14 +152,7 @@ bool FWorldDataLayersActorDesc::Equals(const FWorldPartitionActorDesc* Other) co
 	if (FWorldPartitionActorDesc::Equals(Other))
 	{
 		const FWorldDataLayersActorDesc* OtherDesc = (FWorldDataLayersActorDesc*)Other;
-		if (DataLayerInstances.Num() == OtherDesc->DataLayerInstances.Num())
-		{
-			TArray<FDataLayerInstanceDesc> SortedDataLayerInstances(DataLayerInstances);
-			TArray<FDataLayerInstanceDesc> SortedDataLayerInstancesOther(OtherDesc->DataLayerInstances);
-			SortedDataLayerInstances.Sort([](const FDataLayerInstanceDesc& LHS, const FDataLayerInstanceDesc& RHS) { return LHS < RHS; });
-			SortedDataLayerInstancesOther.Sort([](const FDataLayerInstanceDesc& LHS, const FDataLayerInstanceDesc& RHS) { return LHS < RHS; });
-			return SortedDataLayerInstances == SortedDataLayerInstancesOther;
-		}
+		return CompareUnsortedArrays(DataLayerInstances, OtherDesc->DataLayerInstances);
 	}
 	return false;
 }
diff --git a/Engine/Source/Runtime/Engine/Private/WorldPartition/HLOD/HLODActorDesc.cpp b/Engine/Source/Runtime/Engine/Private/WorldPartition/HLOD/HLODActorDesc.cpp
--- a/Engine/Source/Runtime/Engine/Private/WorldPartition/HLOD/HLODActorDesc.cpp
+++ b/Engine/Source/Runtime/Engine/Private/WorldPartition/HLOD/HLODActorDesc.cpp
@@ -68,17 +68,8 @@ bool FHLODActorDesc::Equals(const FWorldPartitionActorDesc* Other) const
 	if (FWorldPartitionActorDesc::Equals(Other))
 	{
 		const FHLODActorDesc* HLODActorDesc = (FHLODActorDesc*)Other;
-
-		if (CellHash == HLODActorDesc->CellHash && HLODSubActors.Num() == HLODActorDesc->HLODSubActors.Num())
-		{
-			TArray<FHLODSubActorDesc> SortedSubActors(HLODSubActors);
-			TArray<FHLODSubActorDesc> SortedSubActorsOther(HLODActorDesc->HLODSubActors);
-			SortedSubActors.Sort();
-			SortedSubActorsOther.Sort();
-			return SortedSubActors == SortedSubActorsOther;
-		}
+		return (CellHash == HLODActorDesc->CellHash) && CompareUnsortedArrays(HLODSubActors, HLODActorDesc->HLODSubActors);
 	}
-
 	return false;
 }
 
diff --git a/Engine/Source/Runtime/Engine/Private/WorldPartition/WorldPartitionActorDesc.cpp b/Engine/Source/Runtime/Engine/Private/WorldPartition/WorldPartitionActorDesc.cpp
--- a/Engine/Source/Runtime/Engine/Private/WorldPartition/WorldPartitionActorDesc.cpp
+++ b/Engine/Source/Runtime/Engine/Private/WorldPartition/WorldPartitionActorDesc.cpp
@@ -150,7 +150,8 @@ void FWorldPartitionActorDesc::Init(const FWorldPartitionActorDescInitData& Desc
 
 bool FWorldPartitionActorDesc::Equals(const FWorldPartitionActorDesc* Other) const
 {
-	if (Guid == Other->Guid && 
+	return 
+		Guid == Other->Guid && 
 		BaseClass == Other->BaseClass && 
 		NativeClass == Other->NativeClass && 
 		ActorPackage == Other->ActorPackage && 
@@ -167,34 +168,9 @@ bool FWorldPartitionActorDesc::Equals(const FWorldPartitionActorDesc* Other) con
 		FolderPath == Other->FolderPath &&
 		FolderGuid == Other->FolderGuid &&
 		ParentActor == Other->ParentActor &&
-		DataLayers.Num() == Other->DataLayers.Num() &&
-		References.Num() == Other->References.Num())
-	{
-		TArray<FName> SortedDataLayers(DataLayers);
-		TArray<FName> SortedDataLayersOther(Other->DataLayers);
-		SortedDataLayers.Sort([](const FName& LHS, const FName& RHS) { return LHS.LexicalLess(RHS); });
-		SortedDataLayersOther.Sort([](const FName& LHS, const FName& RHS) { return LHS.LexicalLess(RHS); });
-
-		if (SortedDataLayers == SortedDataLayersOther)
-		{
-			TArray<FGuid> SortedReferences(References);
-			TArray<FGuid> SortedReferencesOther(Other->References);
-			SortedReferences.Sort();
-			SortedReferencesOther.Sort();
-
-			if (SortedReferences == SortedReferencesOther)
-			{
-				TArray<FName> SortedTags(Tags);
-				TArray<FName> SortedTagsOther(Other->Tags);
-				SortedTags.Sort([](const FName& LHS, const FName& RHS) { return LHS.LexicalLess(RHS); });
-				SortedTagsOther.Sort([](const FName& LHS, const FName& RHS) { return LHS.LexicalLess(RHS); });
-
-				return SortedTags == SortedTagsOther;
-			}
-		}
-	}
-
-	return false;
+		CompareUnsortedArrays(DataLayers, Other->DataLayers) &&
+		CompareUnsortedArrays(References, Other->References) &&
+		CompareUnsortedArrays(Tags, Other->Tags);
 }
 
 void FWorldPartitionActorDesc::SerializeTo(TArray<uint8>& OutData)
diff --git a/Engine/Source/Runtime/Engine/Public/WorldPartition/WorldPartitionActorDesc.h b/Engine/Source/Runtime/Engine/Public/WorldPartition/WorldPartitionActorDesc.h
--- a/Engine/Source/Runtime/Engine/Public/WorldPartition/WorldPartitionActorDesc.h
+++ b/Engine/Source/Runtime/Engine/Public/WorldPartition/WorldPartitionActorDesc.h
@@ -26,6 +26,32 @@ enum class EContainerClusterMode : uint8
 {
 	Partitioned, // Per Actor Partitioning
 };
+
+template <typename T, class F>
+inline bool CompareUnsortedArrays(const TArray<T>& Array1, const TArray<T>& Array2, F Func)
+{
+	if (Array1.Num() == Array2.Num())
+	{
+		TArray<T> SortedArray1(Array1);
+		TArray<T> SortedArray2(Array2);
+		SortedArray1.Sort(Func);
+		SortedArray2.Sort(Func);
+		return SortedArray1 == SortedArray2;
+	}
+	return false;
+}
+
+template <typename T>
+inline bool CompareUnsortedArrays(const TArray<T>& Array1, const TArray<T>& Array2)
+{
+	return CompareUnsortedArrays(Array1, Array2, [](const T& A, const T& B) { return A < B; });
+}
+
+template <>
+inline bool CompareUnsortedArrays(const TArray<FName>& Array1, const TArray<FName>& Array2)
+{
+	return CompareUnsortedArrays(Array1, Array2, [](const FName& A, const FName& B) { return A.LexicalLess(B); });
+}
 #endif
 
 /**
