make sure that StereoRenderTargetManager updates RTTSize whenever UpdateViewportRHI is called: all implementations compare GetRenderTargetTextureSizeXY against what's computed in CalculateRenderTargetSize

diff --git a/Engine/Source/Runtime/Engine/Private/Slate/SceneViewport.cpp b/Engine/Source/Runtime/Engine/Private/Slate/SceneViewport.cpp
--- a/Engine/Source/Runtime/Engine/Private/Slate/SceneViewport.cpp
+++ b/Engine/Source/Runtime/Engine/Private/Slate/SceneViewport.cpp
@@ -1711,7 +1711,18 @@
 				WindowRenderTargetUpdate(Renderer, Window.Get());
 				if (UseSeparateRenderTarget())
 				{
-					RTTSize = FIntPoint(SizeX, SizeY);
+					uint32 TexSizeX = SizeX, TexSizeY = SizeY;
+					{
+						IStereoRenderTargetManager* const StereoRenderTargetManager =
+							(IsStereoRenderingAllowed() && GEngine->StereoRenderingDevice.IsValid() && GEngine->StereoRenderingDevice->IsStereoEnabledOnNextFrame())
+							? GEngine->StereoRenderingDevice->GetRenderTargetManager()
+							: nullptr;
+						if (StereoRenderTargetManager)
+						{
+							StereoRenderTargetManager->CalculateRenderTargetSize(*this, TexSizeX, TexSizeY);
+						}
+					}
+					RTTSize = FIntPoint(TexSizeX, TexSizeY);
 				}
 			}
 
