just comment nop
bad dupe of behavior in FVirtualTextureDataBuilder::DetectAlphaChannel

diff --git a/Engine/Source/Runtime/Engine/Private/VT/VirtualTextureDataBuilder.cpp b/Engine/Source/Runtime/Engine/Private/VT/VirtualTextureDataBuilder.cpp
--- a/Engine/Source/Runtime/Engine/Private/VT/VirtualTextureDataBuilder.cpp
+++ b/Engine/Source/Runtime/Engine/Private/VT/VirtualTextureDataBuilder.cpp
@@ -1240,6 +1240,9 @@ void FVirtualTextureDataBuilder::BuildMipTails()
 
 bool FVirtualTextureDataBuilder::DetectAlphaChannel(const FImage &Image)
 {
+	// @todo Oodle : BAD DetectAlphaChannel slightly different than the same function in TextureCompressorModule
+	//	  factor them out and share
+
 	if (Image.Format == ERawImageFormat::BGRA8)
 	{
 		const FColor* SrcColors = (&Image.AsBGRA8()[0]);
