clean up handling of CompressionForceAlpha in VT Builder
make sure ForceAlpha and ForceNoAlpha are not both set

diff --git a/Engine/Source/Developer/TextureCompressor/Private/TextureCompressorModule.cpp b/Engine/Source/Developer/TextureCompressor/Private/TextureCompressorModule.cpp
--- a/Engine/Source/Developer/TextureCompressor/Private/TextureCompressorModule.cpp
+++ b/Engine/Source/Developer/TextureCompressor/Private/TextureCompressorModule.cpp
@@ -3148,9 +3148,8 @@ public:
 
 		
 		// DetectAlphaChannel on the top mip of the generated mip chain
-		//	it is now always Linear F32
 		//	BuildSettings could have programatically introduced alpha that was not in the source
-		// note the order of operations in bForceAlphaChannel and bForceNoAlphaChannel (ambiguity if both are on)
+		// note the order of operations in bForceAlphaChannel and bForceNoAlphaChannel ( ForceNo takes precedence )
 		const bool bImageHasAlphaChannel = !BuildSettings.bForceNoAlphaChannel  && (BuildSettings.bForceAlphaChannel || DetectAlphaChannel(IntermediateMipChain[0]));
 	
 		#if 0
diff --git a/Engine/Source/Runtime/Engine/Private/VT/VirtualTextureDataBuilder.cpp b/Engine/Source/Runtime/Engine/Private/VT/VirtualTextureDataBuilder.cpp
--- a/Engine/Source/Runtime/Engine/Private/VT/VirtualTextureDataBuilder.cpp
+++ b/Engine/Source/Runtime/Engine/Private/VT/VirtualTextureDataBuilder.cpp
@@ -579,8 +579,13 @@ void FVirtualTextureDataBuilder::BuildTiles(const TArray<FVTSourceTileEntry>& Ti
 		TBSettings.bSRGB = BuildSettingsForLayer.bSRGB;
 		TBSettings.bUseLegacyGamma = BuildSettingsForLayer.bUseLegacyGamma;
 		TBSettings.MipGenSettings = TMGS_NoMipmaps;
-		TBSettings.bForceAlphaChannel = BuildSettingsForLayer.bForceAlphaChannel || LayerData.bHasAlpha;
-		TBSettings.bForceNoAlphaChannel = BuildSettingsForLayer.bForceNoAlphaChannel || !LayerData.bHasAlpha;
+
+		// LayerData.bHasAlpha was set for the layer if there is alpha anywhere (or ForceAlpha and not ForceNoAlpha)
+		// force each tile to make the same choice about whether it has alpha or not, do not DetectAlpha per tile
+		// make sure either ForceAlpha or ForceNoAlpha is on for tiles
+		TBSettings.bForceAlphaChannel = LayerData.bHasAlpha;
+		TBSettings.bForceNoAlphaChannel = !LayerData.bHasAlpha;
+
 		TBSettings.bHDRSource = BuildSettingsForLayer.bHDRSource;
 		TBSettings.bVirtualStreamable = true;
 
@@ -911,7 +916,7 @@ void FVirtualTextureDataBuilder::BuildSourcePixels(const FTextureSourceData& Sou
 			
 			if ( LayerIndex != 0 )
 			{
-				// @todo Oodle : this looks like a mess
+				// @todo Oodle : this looks fragile
 				//	some of the processing options are copied from BuildSettingsForLayer
 				//	but some are NOT
 				//	it seems semi-random
@@ -1027,9 +1032,12 @@ void FVirtualTextureDataBuilder::BuildSourcePixels(const FTextureSourceData& Sou
 				Image->RawData = MoveTemp(CompressedMip.RawData);
 			}
 
-			if (!LayerData.bHasAlpha && DetectAlphaChannel(BlockData.MipsPerLayer[LayerIndex][0]))
+			// bHasAlpha was previously set to true if bForceAlphaChannel
+			// if it's false and not bForceNoAlphaChannel, we scan each block for alpha
+			// if alpha is in any tile of any block, it gets enabled for all so they have consistent pixel format
+			if (!LayerData.bHasAlpha && !BuildSettingsForLayer.bForceNoAlphaChannel )
 			{
-				LayerData.bHasAlpha = true;
+				LayerData.bHasAlpha = DetectAlphaChannel(BlockData.MipsPerLayer[LayerIndex][0]);
 			}
 		}
 	}
