First-frame notify states can potentially cause check() during cook shutdown
#fix when calling ::UninitializeAnimation and triggering queued anim notifies check whether or not the NotifyStateClass is still reachable

diff --git a/Engine/Source/Runtime/Engine/Private/Animation/AnimInstance.cpp b/Engine/Source/Runtime/Engine/Private/Animation/AnimInstance.cpp
--- a/Engine/Source/Runtime/Engine/Private/Animation/AnimInstance.cpp
+++ b/Engine/Source/Runtime/Engine/Private/Animation/AnimInstance.cpp
@@ -292,7 +292,11 @@ void UAnimInstance::UninitializeAnimation()
 			{
 				const FAnimNotifyEvent& AnimNotifyEvent = ActiveAnimNotifyState[Index];
 				const FAnimNotifyEventReference& EventReference = ActiveAnimNotifyEventReference[Index];
-				if (ShouldTriggerAnimNotifyState(AnimNotifyEvent.NotifyStateClass))
+				if (ShouldTriggerAnimNotifyState(AnimNotifyEvent.NotifyStateClass)
+#if WITH_EDITOR
+					&& !AnimNotifyEvent.NotifyStateClass->IsUnreachable()
+#endif
+				)
 				{
 #if WITH_EDITOR
 					// Prevent firing notifies in animation editors if requested 
