Chaos - validate solver type cvar before applying

diff --git a/Engine/Source/Runtime/Engine/Private/PhysicsEngine/ImmediatePhysics/ImmediatePhysicsChaos/ImmediatePhysicsSimulation_Chaos.cpp b/Engine/Source/Runtime/Engine/Private/PhysicsEngine/ImmediatePhysics/ImmediatePhysicsChaos/ImmediatePhysicsSimulation_Chaos.cpp
--- a/Engine/Source/Runtime/Engine/Private/PhysicsEngine/ImmediatePhysics/ImmediatePhysicsChaos/ImmediatePhysicsSimulation_Chaos.cpp
+++ b/Engine/Source/Runtime/Engine/Private/PhysicsEngine/ImmediatePhysics/ImmediatePhysicsChaos/ImmediatePhysicsSimulation_Chaos.cpp
@@ -836,7 +836,7 @@ namespace ImmediatePhysics_Chaos
 
 		// TEMP: overrides
 		{
-			const EConstraintSolverType SolverType = (EConstraintSolverType)ChaosImmediate_SolverType;
+			const EConstraintSolverType SolverType = (EConstraintSolverType)FMath::Clamp(ChaosImmediate_SolverType, int32(EConstraintSolverType::None), int32(EConstraintSolverType::QuasiPbd));
 			Implementation->Evolution.SetSolverType(SolverType);
 			Implementation->Collisions.SetSolverType(SolverType);
 			Implementation->Joints.SetSolverType(SolverType);
diff --git a/Engine/Source/Runtime/Experimental/Chaos/Private/PBDRigidsSolver.cpp b/Engine/Source/Runtime/Experimental/Chaos/Private/PBDRigidsSolver.cpp
--- a/Engine/Source/Runtime/Experimental/Chaos/Private/PBDRigidsSolver.cpp
+++ b/Engine/Source/Runtime/Experimental/Chaos/Private/PBDRigidsSolver.cpp
@@ -830,8 +830,9 @@ namespace Chaos
 
 		// Apply CVAR overrides if set
 		{
-			GetEvolution()->GetCollisionConstraints().SetSolverType((EConstraintSolverType)ChaosSolver_SolverType);
-			GetJointConstraints().SetSolverType((EConstraintSolverType)ChaosSolver_SolverType);
+			const EConstraintSolverType SolverType = (EConstraintSolverType)FMath::Clamp(ChaosSolver_SolverType, int32(EConstraintSolverType::None), int32(EConstraintSolverType::QuasiPbd));
+			GetEvolution()->GetCollisionConstraints().SetSolverType(SolverType);
+			GetJointConstraints().SetSolverType(SolverType);
 
 			if (ChaosSolverCollisionPositionFrictionIterations >= 0)
 			{
