WorldPartition: Add optional result param to ForEachActorWithLoading

diff --git a/Engine/Source/Runtime/Engine/Private/WorldPartition/WorldPartitionHelpers.cpp b/Engine/Source/Runtime/Engine/Private/WorldPartition/WorldPartitionHelpers.cpp
--- a/Engine/Source/Runtime/Engine/Private/WorldPartition/WorldPartitionHelpers.cpp
+++ b/Engine/Source/Runtime/Engine/Private/WorldPartition/WorldPartitionHelpers.cpp
@@ -91,6 +91,7 @@ namespace WorldPartitionHelpers
 
 FWorldPartitionHelpers::FForEachActorWithLoadingParams::FForEachActorWithLoadingParams()
 	: bGCPerActor(false)
+	, bKeepReferences(false)
 	, ActorClass(AActor::StaticClass())
 {}
 
@@ -113,16 +114,21 @@ void FWorldPartitionHelpers::ForEachActorWithLoading(UWorldPartition* WorldParti
 
 void FWorldPartitionHelpers::ForEachActorWithLoading(UWorldPartition* WorldPartition, TFunctionRef<bool(const FWorldPartitionActorDesc*)> Func, const FForEachActorWithLoadingParams& Params)
 {
-	TMap<FGuid, FWorldPartitionReference> ActorReferences;
+	FForEachActorWithLoadingResult Result;
+	ForEachActorWithLoading(WorldPartition, Func, Params, Result);
+}
 
-	auto CallGarbageCollect = [&Params, &ActorReferences]()
+void FWorldPartitionHelpers::ForEachActorWithLoading(UWorldPartition* WorldPartition, TFunctionRef<bool(const FWorldPartitionActorDesc*)> Func, const FForEachActorWithLoadingParams& Params, FForEachActorWithLoadingResult& Result)
+{
+	check(Result.ActorReferences.IsEmpty());
+	auto CallGarbageCollect = [&Params, &Result]()
 	{
+		check(!Params.bKeepReferences);
 		if (Params.OnPreGarbageCollect)
 		{
 			Params.OnPreGarbageCollect();
 		}
-
-		ActorReferences.Empty();
+		Result.ActorReferences.Empty();
 		DoCollectGarbage();
 	};
 
@@ -132,7 +138,7 @@ void FWorldPartitionHelpers::ForEachActorWithLoading(UWorldPartition* WorldParti
 		{
 			if (!Params.FilterActorDesc || Params.FilterActorDesc(*ActorDescIterator))
 			{
-				WorldPartitionHelpers::LoadReferences(WorldPartition, ActorDescIterator->GetGuid(), ActorReferences);
+				WorldPartitionHelpers::LoadReferences(WorldPartition, ActorDescIterator->GetGuid(), Result.ActorReferences);
 
 				FWorldPartitionReference ActorReference(WorldPartition, ActorDescIterator->GetGuid());
 				if (!Func(ActorReference.Get()))
@@ -140,7 +146,7 @@ void FWorldPartitionHelpers::ForEachActorWithLoading(UWorldPartition* WorldParti
 					break;
 				}
 
-				if (Params.bGCPerActor || FWorldPartitionHelpers::HasExceededMaxMemory())
+				if (!Params.bKeepReferences && (Params.bGCPerActor || FWorldPartitionHelpers::HasExceededMaxMemory()))
 				{
 					CallGarbageCollect();
 				}
@@ -148,7 +154,10 @@ void FWorldPartitionHelpers::ForEachActorWithLoading(UWorldPartition* WorldParti
 		}
 	}
 
-	CallGarbageCollect();
+	if (!Params.bKeepReferences)
+	{
+		CallGarbageCollect();
+	}
 }
 
 bool FWorldPartitionHelpers::HasExceededMaxMemory()
diff --git a/Engine/Source/Runtime/Engine/Public/WorldPartition/WorldPartitionHelpers.h b/Engine/Source/Runtime/Engine/Public/WorldPartition/WorldPartitionHelpers.h
--- a/Engine/Source/Runtime/Engine/Public/WorldPartition/WorldPartitionHelpers.h
+++ b/Engine/Source/Runtime/Engine/Public/WorldPartition/WorldPartitionHelpers.h
@@ -3,6 +3,7 @@
 
 #include "CoreMinimal.h"
 #include "GameFramework/Actor.h"
+#include "WorldPartition/WorldPartitionHandle.h"
 
 #if WITH_EDITOR
 
@@ -38,6 +39,9 @@ public:
 		/* Perform a garbage collection per-actor, useful to test if the caller properly handle GCs. */
 		bool bGCPerActor;
 
+		/* Prevent clearing of actor references. */
+		bool bKeepReferences;
+
 		/* The class used to filter actors loading. */
 		TSubclassOf<AActor> ActorClass;
 
@@ -48,6 +52,13 @@ public:
 		TUniqueFunction<void()> OnPreGarbageCollect;
 	};
 
+	/* Struct of optional output from foreach actordesc functions. */
+	struct ENGINE_API FForEachActorWithLoadingResult
+	{
+		/* Reference to all actors and actor references loaded by the foreach actordesc function */
+		TMap<FGuid, FWorldPartitionReference> ActorReferences;
+	};
+
 	UE_DEPRECATED(5.1, "ForEachActorWithLoading is deprecated, ForEachActorWithLoading with FForEachActorWithLoadingParams should be used instead.")
 	static void ForEachActorWithLoading(UWorldPartition* WorldPartition, TSubclassOf<AActor> ActorClass, TFunctionRef<bool(const FWorldPartitionActorDesc*)> Func, TFunctionRef<void()> OnReleasingActorReferences = [](){}, bool bGCPerActor = false);
 
@@ -55,6 +66,7 @@ public:
 	static void ForEachActorWithLoading(UWorldPartition* WorldPartition, const TArray<FGuid>& ActorGuids, TFunctionRef<bool(const FWorldPartitionActorDesc*)> Func, TFunctionRef<void()> OnReleasingActorReferences = [](){}, bool bGCPerActor = false);
 
 	static void ForEachActorWithLoading(UWorldPartition* WorldPartition, TFunctionRef<bool(const FWorldPartitionActorDesc*)> Func, const FForEachActorWithLoadingParams& Params = FForEachActorWithLoadingParams());
+	static void ForEachActorWithLoading(UWorldPartition* WorldPartition, TFunctionRef<bool(const FWorldPartitionActorDesc*)> Func, const FForEachActorWithLoadingParams& Params, FForEachActorWithLoadingResult& Result);
 	
 	static bool HasExceededMaxMemory();
 	static void DoCollectGarbage();
