Add a "SwapPlatformUserId" to compare to the old SwapControllerId function.

#rnx

diff --git a/Engine/Source/Runtime/Engine/Classes/Engine/Engine.h b/Engine/Source/Runtime/Engine/Classes/Engine/Engine.h
--- a/Engine/Source/Runtime/Engine/Classes/Engine/Engine.h
+++ b/Engine/Source/Runtime/Engine/Classes/Engine/Engine.h
@@ -2663,6 +2663,8 @@ public:
 	ULocalPlayer* GetLocalPlayerFromInputDevice(UWorld * InWorld, const FInputDeviceId InputDevice) const;
 
 	void SwapControllerId(ULocalPlayer *NewPlayer, const int32 CurrentControllerId, const int32 NewControllerID) const;
+	
+	void SwapPlatformUserId(ULocalPlayer *NewPlayer, const FPlatformUserId CurrentUserId, const FPlatformUserId NewUserID) const;
 
 	/** 
 	 * Find a Local Player Controller, which may not exist at all if this is a server.
diff --git a/Engine/Source/Runtime/Engine/Private/UnrealEngine.cpp b/Engine/Source/Runtime/Engine/Private/UnrealEngine.cpp
--- a/Engine/Source/Runtime/Engine/Private/UnrealEngine.cpp
+++ b/Engine/Source/Runtime/Engine/Private/UnrealEngine.cpp
@@ -3951,6 +3951,32 @@ void UEngine::SwapControllerId(ULocalPlayer *NewPlayer, const int32 CurrentContr
 	}
 }
 
+void UEngine::SwapPlatformUserId(ULocalPlayer *NewPlayer, const FPlatformUserId CurrentUserId, const FPlatformUserId NewUserID) const
+{
+	for (auto It = WorldList.CreateConstIterator(); It; ++It)
+	{
+		if (It->OwningGameInstance == nullptr)
+		{
+			continue;
+		}
+
+		const TArray<class ULocalPlayer*> & LocalPlayers = It->OwningGameInstance->GetLocalPlayers();
+
+		if (LocalPlayers.Contains(NewPlayer))
+		{
+			// This is the world context that NewPlayer belongs to, see if anyone is using their CurrentUserId
+			for (ULocalPlayer* LocalPlayer : LocalPlayers)
+			{
+				if (LocalPlayer && LocalPlayer != NewPlayer && LocalPlayer->GetPlatformUserId() == NewUserID)
+				{
+					LocalPlayer->SetPlatformUserId(CurrentUserId);
+					return;
+				}
+			}
+		}
+	}
+}
+
 APlayerController* UEngine::GetFirstLocalPlayerController(const UWorld* InWorld)
 {
 	const FWorldContext &Context = GetWorldContextFromWorldChecked(InWorld);
