WorldPartition: New TMap<FName,FName> in ActorDesc allows adding custom data without subclassing FWorldPartitionActorDesc.

Actor/Component classes can write to this map by overriding GetActorDescProperties

diff --git a/Engine/Source/Runtime/Core/Public/UObject/FortniteNCBranchObjectVersion.h b/Engine/Source/Runtime/Core/Public/UObject/FortniteNCBranchObjectVersion.h
--- a/Engine/Source/Runtime/Core/Public/UObject/FortniteNCBranchObjectVersion.h
+++ b/Engine/Source/Runtime/Core/Public/UObject/FortniteNCBranchObjectVersion.h
@@ -30,6 +30,9 @@ struct CORE_API FFortniteNCBranchObjectVersion
 		// Serialize tags for actor descriptors
 		WorldPartitionActorDescTagsSerialization,
 
+		// Serialize property map for actor descriptors
+		WorldPartitionActorDescPropertyMapSerialization,
+
 		// -----<new versions can be added above this line>-------------------------------------------------
 		VersionPlusOne,
 		LatestVersion = VersionPlusOne - 1
diff --git a/Engine/Source/Runtime/Engine/Classes/Components/ActorComponent.h b/Engine/Source/Runtime/Engine/Classes/Components/ActorComponent.h
--- a/Engine/Source/Runtime/Engine/Classes/Components/ActorComponent.h
+++ b/Engine/Source/Runtime/Engine/Classes/Components/ActorComponent.h
@@ -839,6 +839,11 @@ public:
 	 * externally generated HLODs.
 	 */
 	virtual TSubclassOf<class UHLODBuilder> GetCustomHLODBuilderClass() const { return nullptr; }
+
+	/**
+	 * Add properties to the component owner's actor desc.
+	 */
+	virtual void GetActorDescProperties(TMap<FName, FName>& Properties) const { }
 #endif // WITH_EDITOR
 
 	/**
diff --git a/Engine/Source/Runtime/Engine/Classes/GameFramework/Actor.h b/Engine/Source/Runtime/Engine/Classes/GameFramework/Actor.h
--- a/Engine/Source/Runtime/Engine/Classes/GameFramework/Actor.h
+++ b/Engine/Source/Runtime/Engine/Classes/GameFramework/Actor.h
@@ -933,6 +933,11 @@ public:
 	 */
 	TUniquePtr<class FWorldPartitionActorDesc> CreateActorDesc() const;
 
+	/**
+	 * Add properties to the actor desc.
+	 */
+	virtual void GetActorDescProperties(TMap<FName, FName>& Properties) const;
+
 	/**
 	 * Creates an uninitialized actor descriptor from a specific class.
 	 */
diff --git a/Engine/Source/Runtime/Engine/Private/ActorEditor.cpp b/Engine/Source/Runtime/Engine/Private/ActorEditor.cpp
--- a/Engine/Source/Runtime/Engine/Private/ActorEditor.cpp
+++ b/Engine/Source/Runtime/Engine/Private/ActorEditor.cpp
@@ -871,6 +871,14 @@ TUniquePtr<FWorldPartitionActorDesc> AActor::CreateActorDesc() const
 	return ActorDesc;
 }
 
+void AActor::GetActorDescProperties(TMap<FName, FName>& Properties) const
+{
+	ForEachComponent<UActorComponent>(false, [&Properties](UActorComponent* Component)
+	{
+		Component->GetActorDescProperties(Properties);
+	});
+}
+
 TUniquePtr<class FWorldPartitionActorDesc> AActor::StaticCreateClassActorDesc(const TSubclassOf<AActor>& ActorClass)
 {
 	return CastChecked<AActor>(ActorClass->GetDefaultObject())->CreateClassActorDesc();
diff --git a/Engine/Source/Runtime/Engine/Private/WorldPartition/WorldPartitionActorDesc.cpp b/Engine/Source/Runtime/Engine/Private/WorldPartition/WorldPartitionActorDesc.cpp
--- a/Engine/Source/Runtime/Engine/Private/WorldPartition/WorldPartitionActorDesc.cpp
+++ b/Engine/Source/Runtime/Engine/Private/WorldPartition/WorldPartitionActorDesc.cpp
@@ -99,6 +99,9 @@ void FWorldPartitionActorDesc::Init(const AActor* InActor)
 
 	Tags = InActor->Tags;
 
+	Properties.Empty();
+	InActor->GetActorDescProperties(Properties);
+
 	ActorPackage = InActor->GetPackage()->GetFName();
 	ActorPath = *InActor->GetPathName();
 	FolderPath = InActor->GetFolderPath();
@@ -150,27 +153,28 @@ void FWorldPartitionActorDesc::Init(const FWorldPartitionActorDescInitData& Desc
 
 bool FWorldPartitionActorDesc::Equals(const FWorldPartitionActorDesc* Other) const
 {
-	return 
-		Guid == Other->Guid && 
-		BaseClass == Other->BaseClass && 
-		NativeClass == Other->NativeClass && 
-		ActorPackage == Other->ActorPackage && 
-		ActorPath == Other->ActorPath && 
-		ActorLabel == Other->ActorLabel && 
+	return
+		Guid == Other->Guid &&
+		BaseClass == Other->BaseClass &&
+		NativeClass == Other->NativeClass &&
+		ActorPackage == Other->ActorPackage &&
+		ActorPath == Other->ActorPath &&
+		ActorLabel == Other->ActorLabel &&
 		BoundsLocation.Equals(Other->BoundsLocation, 0.1f) &&
-		BoundsExtent.Equals(Other->BoundsExtent, 0.1f) && 
-		RuntimeGrid == Other->RuntimeGrid && 
-		bActorIsEditorOnly == Other->bActorIsEditorOnly && 
-		bLevelBoundsRelevant == Other->bLevelBoundsRelevant && 
-		bActorIsHLODRelevant == Other->bActorIsHLODRelevant && 
+		BoundsExtent.Equals(Other->BoundsExtent, 0.1f) &&
+		RuntimeGrid == Other->RuntimeGrid &&
+		bActorIsEditorOnly == Other->bActorIsEditorOnly &&
+		bLevelBoundsRelevant == Other->bLevelBoundsRelevant &&
+		bActorIsHLODRelevant == Other->bActorIsHLODRelevant &&
 		bIsUsingDataLayerAsset == Other->bIsUsingDataLayerAsset &&
-		HLODLayer == Other->HLODLayer && 
+		HLODLayer == Other->HLODLayer &&
 		FolderPath == Other->FolderPath &&
 		FolderGuid == Other->FolderGuid &&
 		ParentActor == Other->ParentActor &&
 		CompareUnsortedArrays(DataLayers, Other->DataLayers) &&
 		CompareUnsortedArrays(References, Other->References) &&
-		CompareUnsortedArrays(Tags, Other->Tags);
+		CompareUnsortedArrays(Tags, Other->Tags) &&
+		Properties.OrderIndependentCompareEqual(Other->Properties);
 }
 
 void FWorldPartitionActorDesc::SerializeTo(TArray<uint8>& OutData)
@@ -335,6 +339,11 @@ void FWorldPartitionActorDesc::Serialize(FArchive& Ar)
 	{
 		Ar << FolderGuid;
 	}
+
+	if (Ar.CustomVer(FFortniteNCBranchObjectVersion::GUID) >= FFortniteNCBranchObjectVersion::WorldPartitionActorDescPropertyMapSerialization)
+	{
+		Ar << Properties;
+	}
 }
 
 FBox FWorldPartitionActorDesc::GetBounds() const
diff --git a/Engine/Source/Runtime/Engine/Public/WorldPartition/WorldPartitionActorDesc.h b/Engine/Source/Runtime/Engine/Public/WorldPartition/WorldPartitionActorDesc.h
--- a/Engine/Source/Runtime/Engine/Public/WorldPartition/WorldPartitionActorDesc.h
+++ b/Engine/Source/Runtime/Engine/Public/WorldPartition/WorldPartitionActorDesc.h
@@ -97,6 +97,8 @@ public:
 	FBox GetBounds() const;
 	inline const FGuid& GetParentActor() const { return ParentActor; }
 	inline bool IsUsingDataLayerAsset() const { return bIsUsingDataLayerAsset; }
+	inline FName GetProperty(FName PropertyName) const { return Properties.FindRef(PropertyName); }
+	inline bool HasProperty(FName PropertyName) const { return !GetProperty(PropertyName).IsNone(); }
 
 	FName GetActorName() const;
 	FName GetActorLabelOrName() const;
@@ -226,6 +228,7 @@ protected:
 	TArray<FName>					DataLayers;
 	TArray<FGuid>					References;
 	TArray<FName>					Tags;
+	TMap<FName, FName>				Properties;
 	FName							FolderPath;
 	FGuid							FolderGuid;
 
