Additional logging for actor channel failure

diff --git a/Engine/Source/Runtime/Engine/Private/DataChannel.cpp b/Engine/Source/Runtime/Engine/Private/DataChannel.cpp
--- a/Engine/Source/Runtime/Engine/Private/DataChannel.cpp
+++ b/Engine/Source/Runtime/Engine/Private/DataChannel.cpp
@@ -1681,7 +1681,7 @@ void UControlChannel::ReceivedBunch( FInBunch& Bunch )
 
 				if (FNetControlMessage<NMT_ActorChannelFailure>::Receive(Bunch, ChannelIndex))
 				{
-					UE_LOG(LogNet, Log, TEXT("Server connection received: %s %s"), FNetControlMessageInfo::GetName(MessageType), *Describe());
+					UE_LOG(LogNet, Log, TEXT("Server connection received: %s %d %s"), FNetControlMessageInfo::GetName(MessageType), ChannelIndex, *Describe());
 
 					// Check if Channel index provided by client is valid and within range of channel on server
 					if (ChannelIndex >= 0 && ChannelIndex < Connection->Channels.Num())
@@ -1689,6 +1689,8 @@ void UControlChannel::ReceivedBunch( FInBunch& Bunch )
 						// Get the actor channel that the client provided as having failed
 						UActorChannel* ActorChan = Cast<UActorChannel>(Connection->Channels[ChannelIndex]);
 
+						UE_CLOG(ActorChan != nullptr, LogNet, Log, TEXT("Actor channel failed: %s"), *ActorChan->Describe());
+
 						// The channel and the actor attached to the channel exists on the server
 						if (ActorChan != nullptr && ActorChan->Actor != nullptr)
 						{
