Add a way to upgrade replays used for perf testing

diff --git a/Engine/Source/Runtime/Engine/Public/ReplayTypes.h b/Engine/Source/Runtime/Engine/Public/ReplayTypes.h
--- a/Engine/Source/Runtime/Engine/Public/ReplayTypes.h
+++ b/Engine/Source/Runtime/Engine/Public/ReplayTypes.h
@@ -13,7 +13,7 @@
 #include "Serialization/BitReader.h"
 #include "ReplayTypes.generated.h"
 
-DECLARE_LOG_CATEGORY_EXTERN(LogDemo, Log, All);
+ENGINE_API DECLARE_LOG_CATEGORY_EXTERN(LogDemo, Log, All);
 
 class UNetConnection;
 
@@ -153,6 +153,12 @@
 		EngineNetworkProtocolVersion(FNetworkVersion::GetEngineNetworkProtocolVersion()),
 		GameNetworkProtocolVersion(FNetworkVersion::GetGameNetworkProtocolVersion()),
 		Guid(),
+		MinRecordHz(0.0f),
+		MaxRecordHz(0.0f),
+		FrameLimitInMS(0.0f),
+		CheckpointLimitInMS(0.0f),
+		BuildConfig(EBuildConfiguration::Unknown),
+		BuildTarget(EBuildTargetType::Unknown),
 		EngineVersion(FEngineVersion::Current()),
 		HeaderFlags(EReplayHeaderFlags::None),
 		PackageVersionUE(GPackageFileUEVersion),
