Chaos - re-expose iteration counts for non-linear joint solver

diff --git a/Engine/Source/Runtime/Engine/Classes/PhysicsEngine/PhysicsAsset.h b/Engine/Source/Runtime/Engine/Classes/PhysicsEngine/PhysicsAsset.h
--- a/Engine/Source/Runtime/Engine/Classes/PhysicsEngine/PhysicsAsset.h
+++ b/Engine/Source/Runtime/Engine/Classes/PhysicsEngine/PhysicsAsset.h
@@ -201,11 +201,11 @@ public:
 	FPhysicsAssetSolverSettings SolverSettings;
 
 	/** 
-	 * Solver settings for RBAN simulations with the legacy RBAN solver.
-	 * These are only used for the legacy solver and are hidden by default.
+	 * Solver settings for RBAN simulations with the non-linear RBAN solver.
+	 * These are only used for the non-linear solver which can be enabled with a config variable.
 	 * They will eventually be deprecated and removed.
 	*/
-	UPROPERTY(Config)	// (EditAnywhere, Category = LegacySolverSettings, Config)
+	UPROPERTY(EditAnywhere, Category = LegacySolverSettings, Config)
 	FSolverIterations SolverIterations;
 
 	/** 
