Fixed crash when starting PIE with a modified level instance

diff --git a/Engine/Source/Runtime/Engine/Private/LevelInstance/LevelInstanceActorGuid.cpp b/Engine/Source/Runtime/Engine/Private/LevelInstance/LevelInstanceActorGuid.cpp
--- a/Engine/Source/Runtime/Engine/Private/LevelInstance/LevelInstanceActorGuid.cpp
+++ b/Engine/Source/Runtime/Engine/Private/LevelInstance/LevelInstanceActorGuid.cpp
@@ -30,7 +30,8 @@ const FGuid& FLevelInstanceActorGuid::GetGuid_Internal() const
 {
 	check(Actor);
 #if WITH_EDITOR
-	const FGuid& Guid = (!Actor->GetWorld()->IsGameWorld() || Actor->GetLocalRole() == ENetRole::ROLE_Authority) ? Actor->GetActorGuid() : ActorGuid;
+	// It's possible for an actor not to have its world when duplicating unsaved actor when starting PIE (actor is outered to its UActorContainer)
+	const FGuid& Guid = (!Actor->GetWorld() || !Actor->GetWorld()->IsGameWorld() || Actor->GetLocalRole() == ENetRole::ROLE_Authority) ? Actor->GetActorGuid() : ActorGuid;
 #else
 	const FGuid& Guid = ActorGuid;
 #endif
