Adapted Data Layer stack of Actor Editor Context to pass a context ID when pushing/popping context. This garantees that pop context will always match.

diff --git a/Engine/Source/Runtime/Engine/Private/WorldPartition/DataLayer/DataLayerSubsystem.cpp b/Engine/Source/Runtime/Engine/Private/WorldPartition/DataLayer/DataLayerSubsystem.cpp
--- a/Engine/Source/Runtime/Engine/Private/WorldPartition/DataLayer/DataLayerSubsystem.cpp
+++ b/Engine/Source/Runtime/Engine/Private/WorldPartition/DataLayer/DataLayerSubsystem.cpp
@@ -56,7 +56,11 @@ void FDataLayersEditorBroadcast::StaticOnActorDataLayersEditorLoadingStateChange
 #endif
 
 UDataLayerSubsystem::UDataLayerSubsystem()
-{}
+{
+#if WITH_EDITOR
+	DataLayerActorEditorContextID = 0;
+#endif
+}
 
 void UDataLayerSubsystem::Initialize(FSubsystemCollectionBase& Collection)
 {
@@ -601,18 +605,21 @@ void UDataLayerSubsystem::GetDataLayerDebugColors(TMap<FName, FColor>& OutMappin
 
 void UDataLayerSubsystem::PushActorEditorContext() const
 {
+	++DataLayerActorEditorContextID;
 	if (AWorldDataLayers* WorldDataLayers = GetWorld()->GetWorldDataLayers())
 	{
-		WorldDataLayers->PushActorEditorContext();
+		WorldDataLayers->PushActorEditorContext(DataLayerActorEditorContextID);
 	}
 }
 
 void UDataLayerSubsystem::PopActorEditorContext() const
 {
+	check(DataLayerActorEditorContextID > 0);
 	if (AWorldDataLayers* WorldDataLayers = GetWorld()->GetWorldDataLayers())
 	{
-		WorldDataLayers->PopActorEditorContext();
+		WorldDataLayers->PopActorEditorContext(DataLayerActorEditorContextID);
 	}
+	--DataLayerActorEditorContextID;
 }
 
 TArray<UDataLayerInstance*> UDataLayerSubsystem::GetActorEditorContextDataLayers() const
diff --git a/Engine/Source/Runtime/Engine/Private/WorldPartition/DataLayer/WorldDataLayers.cpp b/Engine/Source/Runtime/Engine/Private/WorldPartition/DataLayer/WorldDataLayers.cpp
--- a/Engine/Source/Runtime/Engine/Private/WorldPartition/DataLayer/WorldDataLayers.cpp
+++ b/Engine/Source/Runtime/Engine/Private/WorldPartition/DataLayer/WorldDataLayers.cpp
@@ -16,6 +16,7 @@
 #if WITH_EDITOR
 #include "WorldPartition/WorldPartitionEditorPerProjectUserSettings.h"
 #include "WorldPartition/DataLayer/WorldDataLayersActorDesc.h"
+#include "Algo/Find.h"
 #endif
 
 #define LOCTEXT_NAMESPACE "WorldDataLayers"
@@ -576,18 +577,28 @@ bool AWorldDataLayers::RemoveFromActorEditorContext(UDataLayerInstance* InDataLa
 	return false;
 }
 
-void AWorldDataLayers::PushActorEditorContext()
+void AWorldDataLayers::PushActorEditorContext(int32 InContextID)
 {
 	Modify(/*bDirty*/false);
+	CurrentDataLayers.ContextID = InContextID;
 	CurrentDataLayersStack.Push(CurrentDataLayers);
 	CurrentDataLayers.Reset();
 }
 
-void AWorldDataLayers::PopActorEditorContext()
+void AWorldDataLayers::PopActorEditorContext(int32 InContextID)
 {
-	check(!CurrentDataLayersStack.IsEmpty());
-	Modify(/*bDirty*/false);
-	CurrentDataLayers = CurrentDataLayersStack.Pop();
+	if (Algo::FindByPredicate(CurrentDataLayersStack, [InContextID](FActorPlacementDataLayers& Element) { return Element.ContextID == InContextID; }))
+	{
+		Modify(/*bDirty*/false);
+		while (!CurrentDataLayersStack.IsEmpty())
+		{
+			CurrentDataLayers = CurrentDataLayersStack.Pop();
+			if (CurrentDataLayers.ContextID == InContextID)
+			{
+				break;
+			}
+		}
+	}
 }
 
 TArray<UDataLayerInstance*> AWorldDataLayers::GetActorEditorContextDataLayers() const
diff --git a/Engine/Source/Runtime/Engine/Public/WorldPartition/DataLayer/DataLayerSubsystem.h b/Engine/Source/Runtime/Engine/Public/WorldPartition/DataLayer/DataLayerSubsystem.h
--- a/Engine/Source/Runtime/Engine/Public/WorldPartition/DataLayer/DataLayerSubsystem.h
+++ b/Engine/Source/Runtime/Engine/Public/WorldPartition/DataLayer/DataLayerSubsystem.h
@@ -200,6 +200,8 @@ public:
 private:
 #if WITH_EDITOR
 	void OnActorDescContainerInitialized(UActorDescContainer* InActorDescContainer);
+
+	mutable int32 DataLayerActorEditorContextID;
 #endif
 
 	/** Console command used to toggle activation of a DataLayer */
diff --git a/Engine/Source/Runtime/Engine/Public/WorldPartition/DataLayer/WorldDataLayers.h b/Engine/Source/Runtime/Engine/Public/WorldPartition/DataLayer/WorldDataLayers.h
--- a/Engine/Source/Runtime/Engine/Public/WorldPartition/DataLayer/WorldDataLayers.h
+++ b/Engine/Source/Runtime/Engine/Public/WorldPartition/DataLayer/WorldDataLayers.h
@@ -24,6 +24,9 @@ struct FActorPlacementDataLayers
 	UPROPERTY()
 	TArray<FName> DataLayerInstanceNames;
 
+	UPROPERTY()
+	int32 ContextID = INT32_MAX;
+
 	void Reset()
 	{
 		DataLayerInstanceNames.Reset();
@@ -68,8 +71,8 @@ public:
 	bool IsInActorEditorContext(const UDataLayerInstance* InDataLayerInstance) const;
 	bool AddToActorEditorContext(UDataLayerInstance* InDataLayerInstance);
 	bool RemoveFromActorEditorContext(UDataLayerInstance* InDataLayerInstance);
-	void PushActorEditorContext();
-	void PopActorEditorContext();
+	void PushActorEditorContext(int32 InContextID);
+	void PopActorEditorContext(int32 InContextID);
 	TArray<UDataLayerInstance*> GetActorEditorContextDataLayers() const;
 
 	//~ Begin Helper Functions
