Add functions on the game instance and local player to create players based on a platform user id instead of the old int32 controllerId

diff --git a/Engine/Source/Runtime/Engine/Classes/Engine/GameInstance.h b/Engine/Source/Runtime/Engine/Classes/Engine/GameInstance.h
--- a/Engine/Source/Runtime/Engine/Classes/Engine/GameInstance.h
+++ b/Engine/Source/Runtime/Engine/Classes/Engine/GameInstance.h
@@ -337,14 +337,32 @@ public:
 	 */
 	ULocalPlayer*			CreateLocalPlayer(int32 ControllerId, FString& OutError, bool bSpawnPlayerController);
 
+	/**
+	 * Adds a new player.
+	 * @param UserId - The platform user id the player should accept input from
+	 * @param OutError - If no player is returned, OutError will contain a string describing the reason.
+	 * @param bSpawnPlayerController - True if a player controller should be spawned immediately for the new player.
+	 * @return The player which was created.
+	 */
+	ULocalPlayer* CreateLocalPlayer(FPlatformUserId UserId, FString& OutError, bool bSpawnPlayerController);
+
 	/**
 	 * Adds a LocalPlayer to the local and global list of Players.
 	 *
 	 * @param	NewPlayer	the player to add
 	 * @param	ControllerId id of the controller associated with the player
 	 */
+	UE_DEPRECATED(5.1, "This version of AddLocalPlayer has been deprecated, pleasse use the version that takes a FPlatformUserId instead.")
 	virtual int32			AddLocalPlayer(ULocalPlayer* NewPlayer, int32 ControllerId);
 
+	/**
+	 * Adds a LocalPlayer to the local and global list of Players.
+	 *
+	 * @param	NewPlayer	The player to add
+	 * @param	UserId		Id of the platform user associated with the player
+	 */
+	virtual int32 AddLocalPlayer(ULocalPlayer* NewPlayer, FPlatformUserId UserId);
+
 	/**
 	 * Removes a player.
 	 * @param Player - The player to remove.
@@ -366,6 +384,9 @@ public:
 
 	/** Returns the local player assigned to a physical Controller Id, or null if not found */
 	ULocalPlayer*			FindLocalPlayerFromControllerId(const int32 ControllerId) const;
+	
+	/** Returns the local player assigned to this platform user id, or null if not found */
+	ULocalPlayer* FindLocalPlayerFromPlatformUserId(const FPlatformUserId UserId) const;
 
 	/** Returns the local player that has been assigned the specific unique net id */
 	ULocalPlayer*			FindLocalPlayerFromUniqueNetId(FUniqueNetIdPtr UniqueNetId) const;
diff --git a/Engine/Source/Runtime/Engine/Classes/Engine/LocalPlayer.h b/Engine/Source/Runtime/Engine/Classes/Engine/LocalPlayer.h
--- a/Engine/Source/Runtime/Engine/Classes/Engine/LocalPlayer.h
+++ b/Engine/Source/Runtime/Engine/Classes/Engine/LocalPlayer.h
@@ -377,6 +377,11 @@ public:
 	 */
 	virtual void PlayerAdded(class UGameViewportClient* InViewportClient, int32 InControllerID);
 
+	/**
+	 * Called at creation time for internal setup
+	 */
+	virtual void PlayerAdded(class UGameViewportClient* InViewportClient, FPlatformUserId InUserId);
+
 	/**
 	 * Called to initialize the online delegates
 	 */
diff --git a/Engine/Source/Runtime/Engine/Private/GameInstance.cpp b/Engine/Source/Runtime/Engine/Private/GameInstance.cpp
--- a/Engine/Source/Runtime/Engine/Private/GameInstance.cpp
+++ b/Engine/Source/Runtime/Engine/Private/GameInstance.cpp
@@ -724,10 +724,19 @@ bool UGameInstance::Exec(UWorld* InWorld, const TCHAR* Cmd, FOutputDevice& Ar)
 
 ULocalPlayer* UGameInstance::CreateInitialPlayer(FString& OutError)
 {
-	return CreateLocalPlayer( 0, OutError, false );
+	return CreateLocalPlayer(IPlatformInputDeviceMapper::Get().GetPrimaryPlatformUser(), OutError, false);
 }
 
 ULocalPlayer* UGameInstance::CreateLocalPlayer(int32 ControllerId, FString& OutError, bool bSpawnPlayerController)
+{
+	// A compatibility call that will map the old int32 ControllerId to the new platform user
+	FPlatformUserId UserId = FGenericPlatformMisc::GetPlatformUserForUserIndex(ControllerId);
+	FInputDeviceId DummyInputDevice = INPUTDEVICEID_NONE;
+	IPlatformInputDeviceMapper::Get().RemapControllerIdToPlatformUserAndDevice(ControllerId, UserId, DummyInputDevice);
+	return CreateLocalPlayer(UserId, OutError, bSpawnPlayerController);
+}
+
+ULocalPlayer* UGameInstance::CreateLocalPlayer(FPlatformUserId UserId, FString& OutError, bool bSpawnPlayerController)
 {
 	check(GetEngine()->LocalPlayerClass != NULL);
 
@@ -736,31 +745,31 @@ ULocalPlayer* UGameInstance::CreateLocalPlayer(int32 ControllerId, FString& OutE
 
 	const int32 MaxSplitscreenPlayers = (GetGameViewportClient() != NULL) ? GetGameViewportClient()->MaxSplitscreenPlayers : 1;
 
-	if (FindLocalPlayerFromControllerId( ControllerId ) != NULL)
+	if (FindLocalPlayerFromPlatformUserId(UserId) != NULL)
 	{
-		OutError = FString::Printf(TEXT("A local player already exists for controller ID %d,"), ControllerId);
+		OutError = FString::Printf(TEXT("A local player already exists for PlatformUserId %d,"), UserId.GetInternalId());
 	}
 	else if (LocalPlayers.Num() < MaxSplitscreenPlayers)
 	{
 		// If the controller ID is not specified then find the first available
-		if (ControllerId < 0)
+		if (!UserId.IsValid())
 		{
-			for (ControllerId = 0; ControllerId < MaxSplitscreenPlayers; ++ControllerId)
+			for (int32 Id = 0; Id < MaxSplitscreenPlayers; ++Id)
 			{
-				if (FindLocalPlayerFromControllerId( ControllerId ) == NULL)
+				if (FindLocalPlayerFromControllerId(Id) == nullptr)
 				{
 					break;
 				}
 			}
-			check(ControllerId < MaxSplitscreenPlayers);
+			check(UserId.GetInternalId() < MaxSplitscreenPlayers);
 		}
-		else if (ControllerId >= MaxSplitscreenPlayers)
+		else if (UserId.GetInternalId() >= MaxSplitscreenPlayers)
 		{
-			UE_LOG(LogPlayerManagement, Warning, TEXT("Controller ID (%d) is unlikely to map to any physical device, so this player will not receive input"), ControllerId);
+			UE_LOG(LogPlayerManagement, Warning, TEXT("Controller ID (%d) is unlikely to map to any physical device, so this player will not receive input"), UserId.GetInternalId());
 		}
 
 		NewPlayer = NewObject<ULocalPlayer>(GetEngine(), GetEngine()->LocalPlayerClass);
-		InsertIndex = AddLocalPlayer(NewPlayer, ControllerId);
+		InsertIndex = AddLocalPlayer(NewPlayer, UserId);
 		UWorld* CurrentWorld = GetWorld();
 		if (bSpawnPlayerController && InsertIndex != INDEX_NONE && CurrentWorld != nullptr)
 		{
@@ -807,6 +816,14 @@ ULocalPlayer* UGameInstance::CreateLocalPlayer(int32 ControllerId, FString& OutE
 }
 
 int32 UGameInstance::AddLocalPlayer(ULocalPlayer* NewLocalPlayer, int32 ControllerId)
+{
+	FPlatformUserId UserId = FGenericPlatformMisc::GetPlatformUserForUserIndex(ControllerId);
+	FInputDeviceId DummyInputDevice = INPUTDEVICEID_NONE;
+	IPlatformInputDeviceMapper::Get().RemapControllerIdToPlatformUserAndDevice(ControllerId, UserId, DummyInputDevice);
+	return AddLocalPlayer(NewLocalPlayer, UserId);
+}
+
+int32 UGameInstance::AddLocalPlayer(ULocalPlayer* NewLocalPlayer, FPlatformUserId UserId)
 {
 	if (NewLocalPlayer == nullptr)
 	{
@@ -816,8 +833,8 @@ int32 UGameInstance::AddLocalPlayer(ULocalPlayer* NewLocalPlayer, int32 Controll
 	// Add to list
 	const int32 InsertIndex = LocalPlayers.AddUnique(NewLocalPlayer);
 
-	// Notify the player he/she was added
-	NewLocalPlayer->PlayerAdded(GetGameViewportClient(), ControllerId);
+	// Notify the player they were added
+	NewLocalPlayer->PlayerAdded(GetGameViewportClient(), UserId);
 
 	// Notify the viewport that we added a player (so it can update splitscreen settings, etc)
 	if ( GetGameViewportClient() != nullptr)
@@ -825,7 +842,7 @@ int32 UGameInstance::AddLocalPlayer(ULocalPlayer* NewLocalPlayer, int32 Controll
 		GetGameViewportClient()->NotifyPlayerAdded(InsertIndex, NewLocalPlayer);
 	}
 
-	UE_LOG(LogPlayerManagement, Log, TEXT("UGameInstance::AddLocalPlayer: Added player %s with ControllerId %d at index %d (%d remaining players)"), *NewLocalPlayer->GetName(), NewLocalPlayer->GetControllerId(), InsertIndex, LocalPlayers.Num());
+	UE_LOG(LogPlayerManagement, Log, TEXT("UGameInstance::AddLocalPlayer: Added player %s with PlatformUserId %d at index %d (%d remaining players)"), *NewLocalPlayer->GetName(), NewLocalPlayer->GetPlatformUserId(), InsertIndex, LocalPlayers.Num());
 
 	OnLocalPlayerAddedEvent.Broadcast(NewLocalPlayer);
 
@@ -1029,6 +1046,19 @@ ULocalPlayer* UGameInstance::FindLocalPlayerFromControllerId(const int32 Control
 	return nullptr;
 }
 
+ULocalPlayer* UGameInstance::FindLocalPlayerFromPlatformUserId(const FPlatformUserId UserId) const
+{
+	for (ULocalPlayer* LP : LocalPlayers)
+	{
+		if (LP && (LP->GetPlatformUserId() == UserId))
+		{
+			return LP;
+		}
+	}
+
+	return nullptr;
+}
+
 ULocalPlayer* UGameInstance::FindLocalPlayerFromUniqueNetId(const FUniqueNetId& UniqueNetId) const
 {
 	for (ULocalPlayer* Player : LocalPlayers)
diff --git a/Engine/Source/Runtime/Engine/Private/LocalPlayer.cpp b/Engine/Source/Runtime/Engine/Private/LocalPlayer.cpp
--- a/Engine/Source/Runtime/Engine/Private/LocalPlayer.cpp
+++ b/Engine/Source/Runtime/Engine/Private/LocalPlayer.cpp
@@ -268,6 +268,14 @@ void ULocalPlayer::PlayerAdded(UGameViewportClient* InViewportClient, int32 InCo
 	SubsystemCollection.Initialize(this);
 }
 
+void ULocalPlayer::PlayerAdded(UGameViewportClient* InViewportClient, FPlatformUserId InUserId)
+{
+	ViewportClient = InViewportClient;
+	SetPlatformUserId(InUserId);
+
+	SubsystemCollection.Initialize(this);
+}
+
 void ULocalPlayer::InitOnlineSession()
 {
 	// FIXME: This may be obsolete, still here to support a few straggler cases that do stuff in child classes
@@ -1564,14 +1572,21 @@ void ULocalPlayer::SetControllerId( int32 NewControllerId )
 	}
 }
 
-void ULocalPlayer::SetPlatformUserId(FPlatformUserId InPlatformUserId)
+void ULocalPlayer::SetPlatformUserId(FPlatformUserId NewPlatformUserId)
 {
-	if (InPlatformUserId != PlatformUserId)
+	if (NewPlatformUserId != PlatformUserId)
 	{
 		const FPlatformUserId CurrentPlatformUserId = PlatformUserId;
 
-		PlatformUserId = InPlatformUserId;
-		OnPlatformUserIdChanged().Broadcast(InPlatformUserId, CurrentPlatformUserId);
+		// set this player's CurrentPlatformUserId to PLATFORMUSERID_NONE so that if we need to swap
+		// platform users with another player we don't re-enter the function for this player.
+		PlatformUserId = PLATFORMUSERID_NONE;
+
+		// see if another player is already using this PlatformUserID; if so, swap PlatformUserIDs with them
+		GEngine->SwapPlatformUserId(this, CurrentPlatformUserId, NewPlatformUserId);
+		PlatformUserId = NewPlatformUserId;
+		
+		OnPlatformUserIdChanged().Broadcast(NewPlatformUserId, CurrentPlatformUserId);
 
 		if (GEngine->IsControllerIdUsingPlatformUserId())
 		{
