Add FORCENOINLINE all the animtrace functions
Attempt to prevent stack memory usage bloat caused by inlining

diff --git a/Engine/Source/Runtime/Engine/Public/Animation/AnimTrace.h b/Engine/Source/Runtime/Engine/Public/Animation/AnimTrace.h
--- a/Engine/Source/Runtime/Engine/Public/Animation/AnimTrace.h
+++ b/Engine/Source/Runtime/Engine/Public/Animation/AnimTrace.h
@@ -115,30 +115,30 @@ struct FAnimTrace
 	};
 
 	/** Helper function to output a tick record */
-	ENGINE_API static void OutputAnimTickRecord(const FAnimationBaseContext& InContext, const FAnimTickRecord& InTickRecord);
+	ENGINE_API FORCENOINLINE static void OutputAnimTickRecord(const FAnimationBaseContext& InContext, const FAnimTickRecord& InTickRecord);
 
 	/** Helper function to output a skeletal mesh */
-	ENGINE_API static void OutputSkeletalMesh(const USkeletalMesh* InMesh);
+	ENGINE_API FORCENOINLINE static void OutputSkeletalMesh(const USkeletalMesh* InMesh);
 
 	/** Helper function to output a skeletal mesh pose, curves etc. */
-	ENGINE_API static void OutputSkeletalMeshComponent(const USkeletalMeshComponent* InComponent);
+	ENGINE_API FORCENOINLINE static void OutputSkeletalMeshComponent(const USkeletalMeshComponent* InComponent);
 
 	/** Helper function to output a skeletal mesh frame marker */
-	ENGINE_API static void OutputSkeletalMeshFrame(const USkeletalMeshComponent* InComponent);
+	ENGINE_API FORCENOINLINE static void OutputSkeletalMeshFrame(const USkeletalMeshComponent* InComponent);
 
 	/** Helper function to output an anim graph's execution event */
-	ENGINE_API static void OutputAnimGraph(const FAnimationBaseContext& InContext, uint64 InStartCycle, uint64 InEndCycle, uint8 InPhase);
+	ENGINE_API FORCENOINLINE static void OutputAnimGraph(const FAnimationBaseContext& InContext, uint64 InStartCycle, uint64 InEndCycle, uint8 InPhase);
 
 	/** Helper function to output an anim node's execution event */
-	ENGINE_API static void OutputAnimNodeStart(const FAnimationBaseContext& InContext, uint64 InStartCycle, int32 InPreviousNodeId, int32 InNodeId, float InBlendWeight, float InRootMotionWeight, uint8 InPhase);
-	ENGINE_API static void OutputAnimNodeEnd(const FAnimationBaseContext& InContext, uint64 InEndCycle);
+	ENGINE_API FORCENOINLINE static void OutputAnimNodeStart(const FAnimationBaseContext& InContext, uint64 InStartCycle, int32 InPreviousNodeId, int32 InNodeId, float InBlendWeight, float InRootMotionWeight, uint8 InPhase);
+	ENGINE_API FORCENOINLINE static void OutputAnimNodeEnd(const FAnimationBaseContext& InContext, uint64 InEndCycle);
 
 	/** Output the current set of attributes in the supplied context */
-	ENGINE_API static void OutputAnimNodeAttribute(const FAnimInstanceProxy& InTargetProxy, const FAnimInstanceProxy& InSourceProxy, int32 InTargetNodeId, int32 InSourceNodeId, FName InAttribute);
+	ENGINE_API FORCENOINLINE static void OutputAnimNodeAttribute(const FAnimInstanceProxy& InTargetProxy, const FAnimInstanceProxy& InSourceProxy, int32 InTargetNodeId, int32 InSourceNodeId, FName InAttribute);
 
 	/** Output the current set of attributes in the supplied context */
 	template<typename ContextType>
-	static void OutputAnimNodeBlendableAttributes(const ContextType& InContext, int32 InTargetNodeId, int32 InSourceNodeId)
+	FORCENOINLINE static void OutputAnimNodeBlendableAttributes(const ContextType& InContext, int32 InTargetNodeId, int32 InSourceNodeId)
 	{
 		if(InContext.CustomAttributes.ContainsData())
 		{
@@ -152,16 +152,16 @@ struct FAnimTrace
 	}
 
 	/** Helper function to output a tracked value for an anim node */
-	ENGINE_API static void OutputAnimNodeValue(const FAnimationBaseContext& InContext, uint32 NodeIndex, const TCHAR* InKey, bool InValue);
-	ENGINE_API static void OutputAnimNodeValue(const FAnimationBaseContext& InContext, uint32 NodeIndex, const TCHAR* InKey, int32 InValue);
-	ENGINE_API static void OutputAnimNodeValue(const FAnimationBaseContext& InContext, uint32 NodeIndex, const TCHAR* InKey, float InValue);
-	ENGINE_API static void OutputAnimNodeValue(const FAnimationBaseContext& InContext, uint32 NodeIndex, const TCHAR* InKey, const FVector2D& InValue);
-	ENGINE_API static void OutputAnimNodeValue(const FAnimationBaseContext& InContext, uint32 NodeIndex, const TCHAR* InKey, const FVector& InValue);
-	ENGINE_API static void OutputAnimNodeValue(const FAnimationBaseContext& InContext, uint32 NodeIndex, const TCHAR* InKey, const FRotator& InValue);
-	ENGINE_API static void OutputAnimNodeValue(const FAnimationBaseContext& InContext, uint32 NodeIndex, const TCHAR* InKey, const FName& InValue);
-	ENGINE_API static void OutputAnimNodeValue(const FAnimationBaseContext& InContext, uint32 NodeIndex, const TCHAR* InKey, const TCHAR* InValue);
-	ENGINE_API static void OutputAnimNodeValue(const FAnimationBaseContext& InContext, uint32 NodeIndex, const TCHAR* InKey, const UClass* InValue);
-	ENGINE_API static void OutputAnimNodeValue(const FAnimationBaseContext& InContext, uint32 NodeIndex, const TCHAR* InKey, const UObject* InValue);
+	ENGINE_API FORCENOINLINE static void OutputAnimNodeValue(const FAnimationBaseContext& InContext, uint32 NodeIndex, const TCHAR* InKey, bool InValue);
+	ENGINE_API FORCENOINLINE static void OutputAnimNodeValue(const FAnimationBaseContext& InContext, uint32 NodeIndex, const TCHAR* InKey, int32 InValue);
+	ENGINE_API FORCENOINLINE static void OutputAnimNodeValue(const FAnimationBaseContext& InContext, uint32 NodeIndex, const TCHAR* InKey, float InValue);
+	ENGINE_API FORCENOINLINE static void OutputAnimNodeValue(const FAnimationBaseContext& InContext, uint32 NodeIndex, const TCHAR* InKey, const FVector2D& InValue);
+	ENGINE_API FORCENOINLINE static void OutputAnimNodeValue(const FAnimationBaseContext& InContext, uint32 NodeIndex, const TCHAR* InKey, const FVector& InValue);
+	ENGINE_API FORCENOINLINE static void OutputAnimNodeValue(const FAnimationBaseContext& InContext, uint32 NodeIndex, const TCHAR* InKey, const FRotator& InValue);
+	ENGINE_API FORCENOINLINE static void OutputAnimNodeValue(const FAnimationBaseContext& InContext, uint32 NodeIndex, const TCHAR* InKey, const FName& InValue);
+	ENGINE_API FORCENOINLINE static void OutputAnimNodeValue(const FAnimationBaseContext& InContext, uint32 NodeIndex, const TCHAR* InKey, const TCHAR* InValue);
+	ENGINE_API FORCENOINLINE static void OutputAnimNodeValue(const FAnimationBaseContext& InContext, uint32 NodeIndex, const TCHAR* InKey, const UClass* InValue);
+	ENGINE_API FORCENOINLINE static void OutputAnimNodeValue(const FAnimationBaseContext& InContext, uint32 NodeIndex, const TCHAR* InKey, const UObject* InValue);
 
 	/** Helper function to output debug info for sequence player nodes */
 	ENGINE_API static void OutputAnimSequencePlayer(const FAnimationBaseContext& InContext, const FAnimNode_SequencePlayerBase& InNode);
