Added extra debug information to track down crash around mismatched anim BP class data

diff --git a/Engine/Source/Runtime/Engine/Private/Animation/AnimBlueprintGeneratedClass.cpp b/Engine/Source/Runtime/Engine/Private/Animation/AnimBlueprintGeneratedClass.cpp
--- a/Engine/Source/Runtime/Engine/Private/Animation/AnimBlueprintGeneratedClass.cpp
+++ b/Engine/Source/Runtime/Engine/Private/Animation/AnimBlueprintGeneratedClass.cpp
@@ -388,6 +388,32 @@ void UAnimBlueprintGeneratedClass::ForEachSubsystem(TFunctionRef<EAnimSubsystemE
 
 void UAnimBlueprintGeneratedClass::ForEachSubsystem(UObject* InObject, TFunctionRef<EAnimSubsystemEnumeration(const FAnimSubsystemInstanceContext&)> InFunction) const
 {
+#if DO_CHECK && WITH_EDITOR
+	// Tracking for UE-152551
+	if(ConstantSubsystemProperties.Num() != MutableSubsystemProperties.Num())
+	{
+		UE_LOG(LogAnimation, Error, TEXT("Mismatched constant/mutable subsystems in class %s (ConstantSubsystemProperties == %d, MutableSubsystemProperties == %d)"), *GetName(), ConstantSubsystemProperties.Num(), MutableSubsystemProperties.Num());
+		UE_LOG(LogAnimation, Error, TEXT("Class Info:\n"));
+		UE_LOG(LogAnimation, Error, TEXT("  Parent: %s (%s)\n"), *GetSuperClass()->GetName(), GetSuperClass()->HasAnyClassFlags(CLASS_Native) ? TEXT("Native") : TEXT("Blueprint"));
+		UE_LOG(LogAnimation, Error, TEXT("  Sparse class data struct (%s):"), (GetSparseClassDataStruct() != nullptr) ? *GetSparseClassDataStruct()->GetName() : TEXT("None"));
+		if(GetSparseClassDataStruct() != nullptr)
+		{
+			for(TFieldIterator<FProperty> It(GetSparseClassDataStruct()); It; ++It)
+			{
+				const FProperty* Property = *It;
+				UE_LOG(LogAnimation, Error, TEXT("    %s (%s)"), *Property->GetName(), *Property->GetClass()->GetName()); 
+			}
+		}
+
+		UE_LOG(LogAnimation, Error, TEXT("  Properties:"));
+		for (TFieldIterator<FProperty> It(this); It; ++It)
+		{
+			const FProperty* Property = *It;
+			UE_LOG(LogAnimation, Error, TEXT("    %s (%s)"), *Property->GetName(), *Property->GetClass()->GetName()); 
+		}
+	}
+#endif
+	
 	check(ConstantSubsystemProperties.Num() == MutableSubsystemProperties.Num());
 	
 	for(int32 SubsystemIndex = 0; SubsystemIndex < ConstantSubsystemProperties.Num(); ++SubsystemIndex)
