Deprecate ARFilter struct make/break nodes in Blueprint graphs to preserve backwards compatibility after removal of short package name support.

diff --git a/Engine/Source/Runtime/CoreUObject/Public/UObject/NoExportTypes.h b/Engine/Source/Runtime/CoreUObject/Public/UObject/NoExportTypes.h
--- a/Engine/Source/Runtime/CoreUObject/Public/UObject/NoExportTypes.h
+++ b/Engine/Source/Runtime/CoreUObject/Public/UObject/NoExportTypes.h
@@ -2131,7 +2131,7 @@ enum class EDataValidationResult : uint8
 	NotValidated
 };
 
-USTRUCT(noexport, BlueprintType, IsAlwaysAccessible, HasDefaults)
+USTRUCT(noexport, BlueprintType, IsAlwaysAccessible, HasDefaults, meta = (HasNativeMake = "/Script/Engine.KismetSystemLibrary.MakeARFilter", HasNativeBreak = "/Script/Engine.KismetSystemLibrary.BreakARFilter"))
 struct FARFilter
 {
 	/** The filter component for package names */
@@ -2146,8 +2146,8 @@ struct FARFilter
 	UPROPERTY(EditDefaultsOnly, BlueprintReadWrite, Category = AssetRegistry)
 	TArray<FName> ObjectPaths;
 
-	/** The filter component for class names. Instances of the specified classes, but not subclasses (by default), will be included. Derived classes will be included only if bRecursiveClasses is true. Deprecated. */
-	UPROPERTY(EditDefaultsOnly, BlueprintReadWrite, Category = AssetRegistry)
+	/** [DEPRECATED] - Class names are now represented by path names. Please use ClassPaths instead. */
+	UPROPERTY(BlueprintReadWrite, Category = AssetRegistry)
 	TArray<FName> ClassNames;
 
 	/** The filter component for class path names. Instances of the specified classes, but not subclasses (by default), will be included. Derived classes will be included only if bRecursiveClasses is true. */
@@ -2157,8 +2157,8 @@ struct FARFilter
 	/** The filter component for properties marked with the AssetRegistrySearchable flag */
 	TMultiMap<FName, TOptional<FString>> TagsAndValues;
 
-	/** Deprecated. Only if bRecursiveClasses is true, the results will exclude classes (and subclasses) in this list */
-	UPROPERTY()
+	/** [DEPRECATED] - Class names are now represented by path names. Please use RecursiveClassPathsExclusionSet instead. */
+	UPROPERTY(BlueprintReadWrite, Category = AssetRegistry)
 	TSet<FName> RecursiveClassesExclusionSet;
 
 	/** Only if bRecursiveClasses is true, the results will exclude classes (and subclasses) in this list */
diff --git a/Engine/Source/Runtime/Engine/Classes/Kismet/KismetSystemLibrary.h b/Engine/Source/Runtime/Engine/Classes/Kismet/KismetSystemLibrary.h
--- a/Engine/Source/Runtime/Engine/Classes/Kismet/KismetSystemLibrary.h
+++ b/Engine/Source/Runtime/Engine/Classes/Kismet/KismetSystemLibrary.h
@@ -19,6 +19,7 @@
 #include "Engine/LatentActionManager.h"
 #include "Kismet/BlueprintFunctionLibrary.h"
 #include "Engine/CollisionProfile.h"
+#include "AssetRegistry/ARFilter.h"
 #include "KismetSystemLibrary.generated.h"
 
 class AActor;
@@ -2056,6 +2057,25 @@ class ENGINE_API UKismetSystemLibrary : public UBlueprintFunctionLibrary
 	 */
 	UFUNCTION(BlueprintCallable, Category = "AssetManager", meta=(AutoCreateRefTerm = "ExcludedBundles, ValidTypes"))
 	static void GetPrimaryAssetsWithBundleState(const TArray<FName>& RequiredBundles, const TArray<FName>& ExcludedBundles, const TArray<FPrimaryAssetType>& ValidTypes, bool bForceCurrentState, TArray<FPrimaryAssetId>& OutPrimaryAssetIdList);
+
+	/**
+	 * Builds an ARFilter struct. You should be using ClassPaths and RecursiveClassPathsExclusionSet, ClassNames and RecursiveClassesExclusionSet are deprecated.
+	 * 
+	 * @param ClassNames [DEPRECATED] - Class names are now represented by path names. If non-empty, this input will result in a runtime warning. Please use the ClassPaths input instead.
+	 * @param RecursiveClassesExclusionSet [DEPRECATED] - Class names are now represented by path names. If non-empty, this input will result in a runtime warning. Please use the RecursiveClassPathsExclusionSet input instead.
+	 */
+	UFUNCTION(BlueprintPure, Category = "Utilities", meta = (Keywords = "construct build", NativeMakeFunc, BlueprintThreadSafe, AdvancedDisplay = "5", AutoCreateRefTerm = "PackageNames, PackagePaths, ObjectPaths, ClassNames, ClassPaths, RecursiveClassPathsExclusionSet, RecursiveClassesExclusionSet"))
+	static FARFilter MakeARFilter(const TArray<FName>& PackageNames, const TArray<FName>& PackagePaths, const TArray<FName>& ObjectPaths, const TArray<FTopLevelAssetPath>& ClassPaths, const TSet<FTopLevelAssetPath>& RecursiveClassPathsExclusionSet, const TArray<FName>& ClassNames, const TSet<FName>& RecursiveClassesExclusionSet, const bool bRecursivePaths = false, const bool bRecursiveClasses = false, const bool bIncludeOnlyOnDiskAssets = false);
+
+	/**
+	 * Breaks an ARFilter struct into its component pieces. You should be using ClassPaths and RecursiveClassPathsExclusionSet from this node, ClassNames and RecursiveClassesExclusionSet are deprecated.
+	 *
+	 * @param ClassNames [DEPRECATED] - Class names are now represented by path names. Please use the ClassPaths output instead.
+	 * @param RecursiveClassesExclusionSet [DEPRECATED] - Class names are now represented by path names. Please use the RecursiveClassPathsExclusionSet output instead.
+	 */
+	UFUNCTION(BlueprintPure, Category = "Utilities", meta = (NativeBreakFunc, BlueprintThreadSafe, AdvancedDisplay = "6"))
+	static void BreakARFilter(FARFilter InARFilter, TArray<FName>& PackageNames, TArray<FName>& PackagePaths, TArray<FName>& ObjectPaths, TArray<FTopLevelAssetPath>& ClassPaths, TSet<FTopLevelAssetPath>& RecursiveClassPathsExclusionSet, TArray<FName>& ClassNames, TSet<FName>& RecursiveClassesExclusionSet, bool& bRecursivePaths, bool& bRecursiveClasses, bool& bIncludeOnlyOnDiskAssets);
+
 };
 
 
diff --git a/Engine/Source/Runtime/Engine/Private/KismetSystemLibrary.cpp b/Engine/Source/Runtime/Engine/Private/KismetSystemLibrary.cpp
--- a/Engine/Source/Runtime/Engine/Private/KismetSystemLibrary.cpp
+++ b/Engine/Source/Runtime/Engine/Private/KismetSystemLibrary.cpp
@@ -3269,4 +3269,69 @@ void UKismetSystemLibrary::GetPrimaryAssetsWithBundleState(const TArray<FName>&
 		Manager->GetPrimaryAssetsWithBundleState(OutPrimaryAssetIdList, ValidTypes, RequiredBundles, ExcludedBundles, bForceCurrentState);
 	}
 }
+
+FARFilter UKismetSystemLibrary::MakeARFilter(const TArray<FName>& PackageNames, const TArray<FName>& PackagePaths, const TArray<FName>& ObjectPaths, const TArray<FTopLevelAssetPath>& ClassPaths, const TSet<FTopLevelAssetPath>& RecursiveClassPathsExclusionSet, const TArray<FName>& ClassNames, const TSet<FName>& RecursiveClassesExclusionSet, const bool bRecursivePaths, const bool bRecursiveClasses, const bool bIncludeOnlyOnDiskAssets)
+{
+	FARFilter NewFilter;
+	NewFilter.PackageNames = PackageNames;
+	NewFilter.PackagePaths = PackagePaths;
+	NewFilter.ObjectPaths = ObjectPaths;
+	NewFilter.bRecursivePaths = bRecursivePaths;
+	NewFilter.bRecursiveClasses = bRecursiveClasses;
+	NewFilter.bIncludeOnlyOnDiskAssets = bIncludeOnlyOnDiskAssets;
+
+	NewFilter.ClassPaths = ClassPaths;
+	NewFilter.RecursiveClassPathsExclusionSet = RecursiveClassPathsExclusionSet;
+
+	// Fixup to move to FTopLevelAssetPath
+	for (const FName& ClassName : ClassNames)
+	{
+		FTopLevelAssetPath ClassPathName;
+		if (!ClassName.IsNone())
+		{
+			FString ShortClassName = ClassName.ToString();
+			ClassPathName = UClass::TryConvertShortTypeNameToPathName<UStruct>(*ShortClassName, ELogVerbosity::Warning, TEXT("MakeARFilter should use ClassPaths, ClassNames is deprecated."));
+			UE_CLOG(ClassPathName.IsNull(), LogClass, Error, TEXT("Failed to convert short class name %s to class path name."), *ShortClassName);
+		}
+		
+		NewFilter.ClassPaths.Add(ClassPathName);
+	}
+	for (const FName& RecursiveClassToExclude : RecursiveClassesExclusionSet)
+	{
+		FTopLevelAssetPath ClassPathName;
+		if (!RecursiveClassToExclude.IsNone())
+		{
+			FString ShortClassName = RecursiveClassToExclude.ToString();
+			ClassPathName = UClass::TryConvertShortTypeNameToPathName<UStruct>(*ShortClassName, ELogVerbosity::Warning, TEXT("MakeARFilter should use RecursiveClassPathsExclusionSet, RecursiveClassesExclusionSet is deprecated."));
+			UE_CLOG(ClassPathName.IsNull(), LogClass, Error, TEXT("Failed to convert short class name %s to class path name."), *ShortClassName);
+		}
+		
+		NewFilter.RecursiveClassPathsExclusionSet.Add(ClassPathName);
+	}
+
+	return NewFilter;
+}
+
+void UKismetSystemLibrary::BreakARFilter(FARFilter InARFilter, TArray<FName>& PackageNames, TArray<FName>& PackagePaths, TArray<FName>& ObjectPaths, TArray<FTopLevelAssetPath>& ClassPaths, TSet<FTopLevelAssetPath>& RecursiveClassPathsExclusionSet, TArray<FName>& ClassNames, TSet<FName>& RecursiveClassesExclusionSet, bool& bRecursivePaths, bool& bRecursiveClasses, bool& bIncludeOnlyOnDiskAssets)
+{
+	PackageNames = InARFilter.PackageNames;
+	PackagePaths = InARFilter.PackagePaths;
+	ObjectPaths = InARFilter.ObjectPaths;
+	ClassPaths = InARFilter.ClassPaths;
+	RecursiveClassPathsExclusionSet = InARFilter.RecursiveClassPathsExclusionSet;
+	bRecursivePaths = InARFilter.bRecursivePaths;
+	bRecursiveClasses = InARFilter.bRecursiveClasses;
+	bIncludeOnlyOnDiskAssets = InARFilter.bIncludeOnlyOnDiskAssets;
+
+	// Fixup to move from FTopLevelAssetPath to legacy types
+	for (const FTopLevelAssetPath& ClassPath : ClassPaths)
+	{
+		ClassNames.Add(ClassPath.GetAssetName());
+	}
+	for (const FTopLevelAssetPath& RecursiveClassPathToExclude : RecursiveClassPathsExclusionSet)
+	{
+		RecursiveClassesExclusionSet.Add(RecursiveClassPathToExclude.GetAssetName());
+	}
+}
+
 #undef LOCTEXT_NAMESPACE
