World Partition: added FActorDescViewMap to encapsulate the map of actor desc views so we can cache views by class, etc...

#rnx

diff --git a/Engine/Source/Editor/UnrealEd/Private/WorldPartition/WorldPartitionRenameDuplicateBuilder.cpp b/Engine/Source/Editor/UnrealEd/Private/WorldPartition/WorldPartitionRenameDuplicateBuilder.cpp
--- a/Engine/Source/Editor/UnrealEd/Private/WorldPartition/WorldPartitionRenameDuplicateBuilder.cpp
+++ b/Engine/Source/Editor/UnrealEd/Private/WorldPartition/WorldPartitionRenameDuplicateBuilder.cpp
@@ -1,7 +1,7 @@
 // Copyright Epic Games, Inc. All Rights Reserved.
 
 #include "WorldPartition/WorldPartitionRenameDuplicateBuilder.h"
-
+#include "WorldPartition/WorldPartitionStreamingGeneration.h"
 #include "WorldPartition/WorldPartition.h"
 #include "WorldPartition/WorldPartitionHelpers.h"
 #include "WorldPartition/WorldPartitionActorDescView.h"
@@ -133,21 +133,21 @@ bool UWorldPartitionRenameDuplicateBuilder::RunInternal(UWorld* World, const FCe
 		UE_LOG(LogWorldPartitionRenameDuplicateBuilder, Error, TEXT("Failed to retrieve WorldPartition."));
 		return false;
 	}
-		
-	TMap<FGuid, FWorldPartitionActorDescView> ActorDescViewMap;
+
+	FActorDescViewMap ActorDescViewMap;
 	for (UActorDescContainer::TConstIterator<> ActorDescIterator(WorldPartition); ActorDescIterator; ++ActorDescIterator)
 	{
 		const FWorldPartitionActorDesc* ActorDesc = *ActorDescIterator;
 		FWorldPartitionActorDescView ActorDescView(ActorDesc);
 		// Invalidate data layers to avoid clustering errors. It doesn't matter here as we need clustering only to gather references
 		ActorDescView.SetInvalidDataLayers();
-		ActorDescViewMap.Add(ActorDesc->GetGuid(), ActorDescView);
+		ActorDescViewMap.Emplace(ActorDescView.GetGuid(), ActorDescView);
 	}
 
 	TArray<FActorCluster> ActorClusters;
 	{
 		UE_SCOPED_TIMER(TEXT("Create actor clusters"), LogWorldPartitionRenameDuplicateBuilder, Display);
-		FActorClusterContext::CreateActorClusters(World, ActorDescViewMap, ActorClusters);
+		FActorClusterContext::CreateActorClusters(World, ActorDescViewMap.GetActorDescViewsByGuid(), ActorClusters);
 	}
 
 	UPackage* OriginalPackage = World->GetPackage();
@@ -342,8 +342,10 @@ bool UWorldPartitionRenameDuplicateBuilder::RunInternal(UWorld* World, const FCe
 					// If we are renaming add package to delete
 					if (bRename)
 					{
-						FWorldPartitionActorDescView& ActorDescView = ActorDescViewMap.FindChecked(ActorGuid);
-						PackagesToDelete.Add(ActorDescView.GetActorPackage().ToString());
+						const FWorldPartitionActorDescView* ActorDescView = const_cast<const FActorDescViewMap&>(ActorDescViewMap).FindByGuid(ActorGuid);
+						check(ActorDescView);
+
+						PackagesToDelete.Add(ActorDescView->GetActorPackage().ToString());
 					}
 				}
 						
diff --git a/Engine/Source/Runtime/Core/Public/Misc/HierarchicalLogArchive.h b/Engine/Source/Runtime/Core/Public/Misc/HierarchicalLogArchive.h
--- a/Engine/Source/Runtime/Core/Public/Misc/HierarchicalLogArchive.h
+++ b/Engine/Source/Runtime/Core/Public/Misc/HierarchicalLogArchive.h
@@ -17,6 +17,17 @@ public:
 		FHierarchicalLogArchive* Ar;
 	};
 
+	void Print(const TCHAR* InLine)
+	{
+		WriteLine(InLine, false);
+	}
+
+	UE_NODISCARD FIndentScope PrintIndent(const TCHAR* InLine)
+	{
+		WriteLine(InLine, true);
+		return FIndentScope(this);
+	}
+
 	template <typename FmtType, typename... Types>
 	void Printf(const FmtType& Fmt, Types... Args)
 	{
diff --git a/Engine/Source/Runtime/Engine/Private/WorldPartition/DataLayer/DataLayerUtils.cpp b/Engine/Source/Runtime/Engine/Private/WorldPartition/DataLayer/DataLayerUtils.cpp
--- a/Engine/Source/Runtime/Engine/Private/WorldPartition/DataLayer/DataLayerUtils.cpp
+++ b/Engine/Source/Runtime/Engine/Private/WorldPartition/DataLayer/DataLayerUtils.cpp
@@ -96,32 +96,29 @@ TArray<FName> FDataLayerUtils::ResolvedDataLayerInstanceNames(const FWorldPartit
 }
 
 // For performance reasons, this function assumes that InActorDesc's DataLayerInstanceNames was already resolved.
-bool FDataLayerUtils::ResolveRuntimeDataLayerInstanceNames(const FWorldPartitionActorDescView& InActorDescView, const TMap<FGuid, FWorldPartitionActorDescView>& ActorDescViews, TArray<FName>& OutRuntimeDataLayerInstanceNames)
+bool FDataLayerUtils::ResolveRuntimeDataLayerInstanceNames(const FWorldPartitionActorDescView& InActorDescView, const FActorDescViewMap& ActorDescViewMap, TArray<FName>& OutRuntimeDataLayerInstanceNames)
 {
-	if (InActorDescView.GetDataLayerInstanceNames().IsEmpty())
+	TArray<const FWorldPartitionActorDescView*> WorldDataLayerViews = ActorDescViewMap.FindByExactNativeClass<AWorldDataLayers>();
+	
+	if (WorldDataLayerViews.Num())
 	{
-		return true;
-	}
+		check(WorldDataLayerViews.Num() == 1);
+		const FWorldPartitionActorDescView* WorldDataLayersActorDescView = WorldDataLayerViews[0];
+	
+		FWorldDataLayersActorDesc* WorldDataLayersActorDesc = (FWorldDataLayersActorDesc*)WorldDataLayersActorDescView->GetActorDesc();
 
-	for (auto& [ActorGuid, ActorDescView] : ActorDescViews)
-	{
-		if (ActorDescView.GetActorNativeClass()->IsChildOf<AWorldDataLayers>())
+		for (FName DataLayerInstanceName : InActorDescView.GetDataLayerInstanceNames())
 		{
-			FWorldDataLayersActorDesc* WorldDataLayersActorDesc = (FWorldDataLayersActorDesc*)ActorDescView.GetActorDesc();
-
-			for (FName DataLayerInstanceName : InActorDescView.GetDataLayerInstanceNames())
+			if (const FDataLayerInstanceDesc* DataLayerInstanceDesc = WorldDataLayersActorDesc->GetDataLayerInstanceFromInstanceName(DataLayerInstanceName))
 			{
-				if (const FDataLayerInstanceDesc* DataLayerInstanceDesc = WorldDataLayersActorDesc->GetDataLayerInstanceFromInstanceName(DataLayerInstanceName))
+				if (DataLayerInstanceDesc->GetDataLayerType() == EDataLayerType::Runtime)
 				{
-					if (DataLayerInstanceDesc->GetDataLayerType() == EDataLayerType::Runtime)
-					{
-						OutRuntimeDataLayerInstanceNames.Add(DataLayerInstanceName);
-					}
+					OutRuntimeDataLayerInstanceNames.Add(DataLayerInstanceName);
 				}
 			}
-
-			return true;
 		}
+
+		return true;
 	}
 
 	return false;
diff --git a/Engine/Source/Runtime/Engine/Private/WorldPartition/LevelInstance/LevelInstanceActorDesc.cpp b/Engine/Source/Runtime/Engine/Private/WorldPartition/LevelInstance/LevelInstanceActorDesc.cpp
--- a/Engine/Source/Runtime/Engine/Private/WorldPartition/LevelInstance/LevelInstanceActorDesc.cpp
+++ b/Engine/Source/Runtime/Engine/Private/WorldPartition/LevelInstance/LevelInstanceActorDesc.cpp
@@ -114,6 +114,11 @@ void FLevelInstanceActorDesc::SetContainer(UActorDescContainer* InContainer)
 	}
 }
 
+bool FLevelInstanceActorDesc::IsContainerInstance() const
+{
+	return DesiredRuntimeBehavior == ELevelInstanceRuntimeBehavior::Partitioned;
+}
+
 bool FLevelInstanceActorDesc::GetContainerInstance(const UActorDescContainer*& OutLevelContainer, FTransform& OutLevelTransform, EContainerClusterMode& OutClusterMode) const
 {
 	if (!LevelInstanceContainer.IsValid() && IsLoaded())
@@ -121,13 +126,16 @@ bool FLevelInstanceActorDesc::GetContainerInstance(const UActorDescContainer*& O
 		// Lazy initialization of LevelInstanceContainer (used by ModifiedActorsDescList)
 		const_cast<FLevelInstanceActorDesc*>(this)->RegisterContainerInstance(GetActor()->GetWorld());
 	}
+
 	if (LevelInstanceContainer.IsValid())
 	{
 		OutLevelContainer = LevelInstanceContainer.Get();
+		OutClusterMode = EContainerClusterMode::Partitioned;
+
 		// Apply level instance pivot offset
 		FTransform LevelInstancePivotOffsetTransform = FTransform(ULevel::GetLevelInstancePivotOffsetFromPackage(LevelInstanceContainer->GetContainerPackage()));
 		OutLevelTransform = LevelInstancePivotOffsetTransform * LevelInstanceTransform;
-		OutClusterMode = EContainerClusterMode::Partitioned;
+
 		return true;
 	}
 
diff --git a/Engine/Source/Runtime/Engine/Private/WorldPartition/RuntimeSpatialHash/RuntimeSpatialHashHLOD.cpp b/Engine/Source/Runtime/Engine/Private/WorldPartition/RuntimeSpatialHash/RuntimeSpatialHashHLOD.cpp
--- a/Engine/Source/Runtime/Engine/Private/WorldPartition/RuntimeSpatialHash/RuntimeSpatialHashHLOD.cpp
+++ b/Engine/Source/Runtime/Engine/Private/WorldPartition/RuntimeSpatialHash/RuntimeSpatialHashHLOD.cpp
@@ -469,8 +469,10 @@ bool UWorldPartitionRuntimeSpatialHash::GenerateHLOD(ISourceControlHelper* Sourc
 
 	// Keep track of all valid HLOD actors, along with which runtime grid they live in
 	TMap<FName, TArray<FGuid>> GridsHLODActors;
+	TArray<TUniquePtr<FWorldPartitionActorDescView>> ActorDescViews;
 
-	auto GenerateHLODs = [&GridsHLODActors, &MainContainerInstance, &ActorClusterContext, WorldPartition, &Context, SourceControlHelper, bCreateActorsOnly](const FSpatialHashRuntimeGrid& RuntimeGrid, uint32 HLODLevel, const TArray<const FActorClusterInstance*>& ActorClusterInstances)
+	auto GenerateHLODs = [&GridsHLODActors, &ActorDescViews, &MainContainerInstance, &ActorClusterContext, WorldPartition, &Context, SourceControlHelper, bCreateActorsOnly]
+	(const FSpatialHashRuntimeGrid& RuntimeGrid, uint32 HLODLevel, const TArray<const FActorClusterInstance*>& ActorClusterInstances)
 	{
 		// Generate HLODs for this grid
 		TArray<FGuid> HLODActors = GenerateHLODsForGrid(WorldPartition, ActorClusterContext, RuntimeGrid, HLODLevel, Context, SourceControlHelper, bCreateActorsOnly, ActorClusterInstances);
@@ -480,7 +482,8 @@ bool UWorldPartitionRuntimeSpatialHash::GenerateHLOD(ISourceControlHelper* Sourc
 			FWorldPartitionActorDesc* HLODActorDesc = WorldPartition->GetActorDesc(HLODActorGuid);
 			check(HLODActorDesc);
 
-			MainContainerInstance.ActorDescViewMap.Emplace(HLODActorGuid, HLODActorDesc);
+			FWorldPartitionActorDescView* ActorDescView = ActorDescViews.Emplace_GetRef(MakeUnique<FWorldPartitionActorDescView>(HLODActorDesc)).Get();
+			MainContainerInstance.ActorDescViewMap.Emplace(HLODActorGuid, ActorDescView);
 
 			const FWorldPartitionActorDescView& HLODActorDescView = MainContainerInstance.GetActorDescView(HLODActorGuid);
 			GridsHLODActors.FindOrAdd(HLODActorDescView.GetRuntimeGrid()).Add(HLODActorGuid);
diff --git a/Engine/Source/Runtime/Engine/Private/WorldPartition/WorldPartitionActorCluster.cpp b/Engine/Source/Runtime/Engine/Private/WorldPartition/WorldPartitionActorCluster.cpp
--- a/Engine/Source/Runtime/Engine/Private/WorldPartition/WorldPartitionActorCluster.cpp
+++ b/Engine/Source/Runtime/Engine/Private/WorldPartition/WorldPartitionActorCluster.cpp
@@ -42,7 +42,7 @@ FActorCluster::FActorCluster(UWorld* InWorld, const FWorldPartitionActorDescView
 	DataLayersID = FDataLayersID(DataLayers.Array());
 }
 
-void FActorCluster::Add(const FActorCluster& InActorCluster, const TMap<FGuid, FWorldPartitionActorDescView>& InActorDescViewMap)
+void FActorCluster::Add(const FActorCluster& InActorCluster)
 {
 	check(RuntimeGrid == InActorCluster.RuntimeGrid);
 	check(DataLayersID == InActorCluster.DataLayersID);
@@ -83,7 +83,7 @@ FActorClusterContext::FActorClusterContext(TArray<FActorContainerInstance>&& InC
 	}
 }
 
-FActorContainerInstance::FActorContainerInstance(const FActorContainerID& InID, const FTransform& InTransform, const FBox& InBounds, const TSet<FName>& InRuntimeDataLayers, EContainerClusterMode InClusterMode, const UActorDescContainer* InContainer, TMap<FGuid, FWorldPartitionActorDescView> InActorDescViewMap)
+FActorContainerInstance::FActorContainerInstance(const FActorContainerID& InID, const FTransform& InTransform, const FBox& InBounds, const TSet<FName>& InRuntimeDataLayers, EContainerClusterMode InClusterMode, const UActorDescContainer* InContainer, TMap<FGuid, FWorldPartitionActorDescView*> InActorDescViewMap)
 	: ID(InID)
 	, Transform(InTransform)
 	, Bounds(InBounds)
@@ -96,7 +96,7 @@ FActorContainerInstance::FActorContainerInstance(const FActorContainerID& InID,
 
 const FWorldPartitionActorDescView& FActorContainerInstance::GetActorDescView(const FGuid& InGuid) const
 {
-	return ActorDescViewMap.FindChecked(InGuid);
+	return *ActorDescViewMap.FindChecked(InGuid);
 }
 
 FActorInstance::FActorInstance()
@@ -115,7 +115,7 @@ const FWorldPartitionActorDescView& FActorInstance::GetActorDescView() const
 	return ContainerInstance->GetActorDescView(Actor);
 }
 
-void CreateActorCluster(const FWorldPartitionActorDescView& ActorDescView, TMap<FGuid, FActorCluster*>& ActorToActorCluster, TSet<FActorCluster*>& ActorClustersSet, UWorld* World, const TMap<FGuid, FWorldPartitionActorDescView>& ActorDescViewMap)
+void CreateActorCluster(const FWorldPartitionActorDescView& ActorDescView, TMap<FGuid, FActorCluster*>& ActorToActorCluster, TSet<FActorCluster*>& ActorClustersSet, UWorld* World, const TMap<FGuid, FWorldPartitionActorDescView*>& ActorDescViewMap)
 {
 	const FGuid& ActorGuid = ActorDescView.GetGuid();
 
@@ -129,8 +129,10 @@ void CreateActorCluster(const FWorldPartitionActorDescView& ActorDescView, TMap<
 
 	for (const FGuid& ReferenceGuid : ActorDescView.GetReferences())
 	{
-		if (const FWorldPartitionActorDescView* ReferenceActorDescView = ActorDescViewMap.Find(ReferenceGuid))
+		if (const FWorldPartitionActorDescView* const* ReferenceActorDescViewPtr = ActorDescViewMap.Find(ReferenceGuid))
 		{
+			const FWorldPartitionActorDescView* ReferenceActorDescView = *ReferenceActorDescViewPtr;
+
 			// Ignore references from spatially loaded actors to non-spatially loaded ones (at this point, they are in the same data layers).
 			if (ActorDescView.GetIsSpatiallyLoaded() && !ReferenceActorDescView->GetIsSpatiallyLoaded())
 			{
@@ -143,7 +145,7 @@ void CreateActorCluster(const FWorldPartitionActorDescView& ActorDescView, TMap<
 				if (ReferenceCluster != ActorCluster)
 				{
 					// Merge reference cluster in Actor's cluster
-					ActorCluster->Add(*ReferenceCluster, ActorDescViewMap);
+					ActorCluster->Add(*ReferenceCluster);
 					for (const FGuid& ReferenceClusterActorGuid : ReferenceCluster->Actors)
 					{
 						ActorToActorCluster[ReferenceClusterActorGuid] = ActorCluster;
@@ -155,7 +157,7 @@ void CreateActorCluster(const FWorldPartitionActorDescView& ActorDescView, TMap<
 			else
 			{
 				// Put Reference in Actor's cluster
-				ActorCluster->Add(FActorCluster(World, *ReferenceActorDescView), ActorDescViewMap);
+				ActorCluster->Add(FActorCluster(World, *ReferenceActorDescView));
 			}
 
 			// Map its cluster
@@ -203,14 +205,14 @@ const FActorContainerInstance* FActorClusterContext::GetClusterInstance(const UA
 	return nullptr;
 }
 
-void FActorClusterContext::CreateActorClusters(UWorld* World, const TMap<FGuid, FWorldPartitionActorDescView>& ActorDescViewMap, TArray<FActorCluster>& OutActorClusters, FActorClusterContext::FFilterActorDescViewFunc InFilterActorDescViewFunc)
+void FActorClusterContext::CreateActorClusters(UWorld* World, const TMap<FGuid, FWorldPartitionActorDescView*>& ActorDescViewMap, TArray<FActorCluster>& OutActorClusters, FActorClusterContext::FFilterActorDescViewFunc InFilterActorDescViewFunc)
 {
 	TMap<FGuid, FActorCluster*> ActorToActorCluster;
 	TSet<FActorCluster*> ActorClustersSet;
 
 	for (auto& ActorDescViewPair : ActorDescViewMap)
 	{
-		const FWorldPartitionActorDescView& ActorDescView = ActorDescViewPair.Value;
+		const FWorldPartitionActorDescView& ActorDescView = *ActorDescViewPair.Value;
 
 		if (!InFilterActorDescViewFunc || InFilterActorDescViewFunc(ActorDescView))
 		{
@@ -223,7 +225,7 @@ void FActorClusterContext::CreateActorClusters(UWorld* World, const TMap<FGuid,
 	for (FActorCluster* ActorCluster : ActorClustersSet) { delete ActorCluster; }
 }
 
-void FActorClusterContext::CreateActorClusters(UWorld* World, const TMap<FGuid, FWorldPartitionActorDescView>& ActorDescViewMap, TArray<FActorCluster>& OutActorClusters)
+void FActorClusterContext::CreateActorClusters(UWorld* World, const TMap<FGuid, FWorldPartitionActorDescView*>& ActorDescViewMap, TArray<FActorCluster>& OutActorClusters)
 {
 	CreateActorClusters(World, ActorDescViewMap, OutActorClusters, nullptr);
 }
diff --git a/Engine/Source/Runtime/Engine/Private/WorldPartition/WorldPartitionActorDescView.cpp b/Engine/Source/Runtime/Engine/Private/WorldPartition/WorldPartitionActorDescView.cpp
--- a/Engine/Source/Runtime/Engine/Private/WorldPartition/WorldPartitionActorDescView.cpp
+++ b/Engine/Source/Runtime/Engine/Private/WorldPartition/WorldPartitionActorDescView.cpp
@@ -145,6 +145,11 @@ const FGuid& FWorldPartitionActorDescView::GetFolderGuid() const
 	return ActorDesc->GetFolderGuid();
 }
 
+bool FWorldPartitionActorDescView::IsContainerInstance() const
+{
+	return ActorDesc->IsContainerInstance();
+}
+
 bool FWorldPartitionActorDescView::GetContainerInstance(const UActorDescContainer*& OutLevelContainer, FTransform& OutLevelTransform, EContainerClusterMode& OutClusterMode) const
 {
 	return ActorDesc->GetContainerInstance(OutLevelContainer, OutLevelTransform, OutClusterMode);
diff --git a/Engine/Source/Runtime/Engine/Private/WorldPartition/WorldPartitionRuntimeSpatialHash.cpp b/Engine/Source/Runtime/Engine/Private/WorldPartition/WorldPartitionRuntimeSpatialHash.cpp
--- a/Engine/Source/Runtime/Engine/Private/WorldPartition/WorldPartitionRuntimeSpatialHash.cpp
+++ b/Engine/Source/Runtime/Engine/Private/WorldPartition/WorldPartitionRuntimeSpatialHash.cpp
@@ -802,7 +802,7 @@ bool UWorldPartitionRuntimeSpatialHash::GenerateStreaming(UWorldPartitionStreami
 
 	for (auto& ActorDescViewPair : ContainerInstance->ActorDescViewMap)
 	{
-		const FWorldPartitionActorDescView& ActorDescView = ActorDescViewPair.Value;
+		const FWorldPartitionActorDescView& ActorDescView = *ActorDescViewPair.Value;
 		if (ActorDescView.GetActorNativeClass()->IsChildOf<ASpatialHashRuntimeGridInfo>())
 		{
 			FWorldPartitionReference Ref(WorldPartition, ActorDescView.GetGuid());
diff --git a/Engine/Source/Runtime/Engine/Private/WorldPartition/WorldPartitionStreamingGeneration.cpp b/Engine/Source/Runtime/Engine/Private/WorldPartition/WorldPartitionStreamingGeneration.cpp
--- a/Engine/Source/Runtime/Engine/Private/WorldPartition/WorldPartitionStreamingGeneration.cpp
+++ b/Engine/Source/Runtime/Engine/Private/WorldPartition/WorldPartitionStreamingGeneration.cpp
@@ -1,17 +1,19 @@
 // Copyright Epic Games, Inc. All Rights Reserved.
 
-#include "WorldPartition/WorldPartition.h"
-
 #if WITH_EDITOR
+#include "WorldPartition/WorldPartitionStreamingGeneration.h"
+
 #include "Editor.h"
 #include "Engine/World.h"
 #include "Engine/LevelScriptBlueprint.h"
 #include "ActorReferencesUtils.h"
+#include "WorldPartition/WorldPartition.h"
 #include "WorldPartition/WorldPartitionRuntimeHash.h"
 #include "WorldPartition/WorldPartitionStreamingPolicy.h"
 #include "WorldPartition/WorldPartitionActorCluster.h"
 #include "WorldPartition/DataLayer/DataLayerSubsystem.h"
 #include "WorldPartition/DataLayer/DataLayerUtils.h"
+#include "WorldPartition/DataLayer/WorldDataLayers.h"
 #include "WorldPartition/ErrorHandling/WorldPartitionStreamingGenerationNullErrorHandler.h"
 #include "WorldPartition/ErrorHandling/WorldPartitionStreamingGenerationLogErrorHandler.h"
 #include "WorldPartition/ErrorHandling/WorldPartitionStreamingGenerationMapCheckErrorHandler.h"
@@ -36,6 +38,31 @@ static FAutoConsoleCommand DumpStreamingGenerationLog(
 	})
 );
 
+FActorDescViewMap::FActorDescViewMap()
+{}
+
+TArray<const FWorldPartitionActorDescView*> FActorDescViewMap::FindByExactNativeClass(UClass* InExactNativeClass) const
+{
+	check(InExactNativeClass->IsNative());
+	const FName NativeClassName = InExactNativeClass->GetFName();
+	TArray<const FWorldPartitionActorDescView*> Result;
+	ActorDescViewsByClass.MultiFind(NativeClassName, Result);
+	return Result;
+}
+
+FWorldPartitionActorDescView* FActorDescViewMap::Emplace(const FGuid& InGuid, const FWorldPartitionActorDescView& InActorDescView)
+{
+	FWorldPartitionActorDescView* NewActorDescView = ActorDescViewList.Emplace_GetRef(MakeUnique<FWorldPartitionActorDescView>(InActorDescView)).Get();
+	
+	const UClass* NativeClass = NewActorDescView->GetActorNativeClass();
+	const FName NativeClassName = NativeClass->GetFName();
+
+	ActorDescViewsByGuid.Emplace(InGuid, NewActorDescView);
+	ActorDescViewsByClass.Add(NativeClassName, NewActorDescView);
+
+	return NewActorDescView;
+}
+
 /*
 	Preparation Phase
 		Actor Descriptor Views Creation
@@ -54,25 +81,25 @@ static FAutoConsoleCommand DumpStreamingGenerationLog(
 
 class FWorldPartitionStreamingGenerator
 {
-	void ResolveRuntimeDataLayers(FWorldPartitionActorDescView& ActorDescView, const TMap<FGuid, FWorldPartitionActorDescView>& ActorDescViews)
+	void ResolveRuntimeDataLayers(FWorldPartitionActorDescView& ActorDescView, const FActorDescViewMap& ActorDescViewMap)
 	{
 		TArray<FName> RuntimeDataLayerInstanceNames;
 		RuntimeDataLayerInstanceNames.Reserve(ActorDescView.GetDataLayers().Num());
 
-		if (FDataLayerUtils::ResolveRuntimeDataLayerInstanceNames(ActorDescView, ActorDescViews, RuntimeDataLayerInstanceNames))
+		if (FDataLayerUtils::ResolveRuntimeDataLayerInstanceNames(ActorDescView, ActorDescViewMap, RuntimeDataLayerInstanceNames))
 		{
 			ActorDescView.SetRuntimeDataLayers(RuntimeDataLayerInstanceNames);
 		}
 	}
 
-	void ResolveRuntimeReferences(FWorldPartitionActorDescView& ActorDescView, const TMap<FGuid, FWorldPartitionActorDescView>& ActorDescViews)
+	void ResolveRuntimeReferences(FWorldPartitionActorDescView& ActorDescView, const FActorDescViewMap& ActorDescViewMap)
 	{
 		TArray<FGuid> RuntimeReferences;
 		RuntimeReferences.Reserve(ActorDescView.GetReferences().Num());
 
 		for (const FGuid& ReferenceGuid : ActorDescView.GetReferences())
 		{
-			if (const FWorldPartitionActorDescView* ReferenceDescView = ActorDescViews.Find(ReferenceGuid))
+			if (const FWorldPartitionActorDescView* ReferenceDescView = ActorDescViewMap.FindByGuid(ReferenceGuid))
 			{
 				check(!ReferenceDescView->GetActorIsEditorOnly());
 				RuntimeReferences.Add(ReferenceGuid);
@@ -85,7 +112,7 @@ class FWorldPartitionStreamingGenerator
 		}
 	}
 
-	void CreateActorDescViewMap(const UActorDescContainer* InContainer, TMap<FGuid, FWorldPartitionActorDescView>& OutActorDescViewMap, const FActorContainerID& InContainerID)
+	void CreateActorDescViewMap(const UActorDescContainer* InContainer, FActorDescViewMap& OutActorDescViewMap, const FActorContainerID& InContainerID, TArray<FWorldPartitionActorDescView>& OutContainerInstances)
 	{
 		// should we handle unsaved or newly created actors?
 		const bool bHandleUnsavedActors = ModifiedActorsDescList && InContainerID.IsMainContainer();
@@ -116,10 +143,17 @@ class FWorldPartitionStreamingGenerator
 			return ModifiedActorsDescList->AddActor(InActor);
 		};
 
-		// Adjust and register the actor descriptor view
-		auto RegisterActorDescView = [this, InContainer, &OutActorDescViewMap](const FGuid& ActorGuid, FWorldPartitionActorDescView& ActorDescView)
+		// Register the actor descriptor view
+		auto RegisterActorDescView = [this, InContainer, &OutActorDescViewMap, &OutContainerInstances](const FGuid& ActorGuid, FWorldPartitionActorDescView& ActorDescView)
 		{
-			OutActorDescViewMap.Emplace(ActorGuid, ActorDescView);
+			if (ActorDescView.IsContainerInstance())
+			{
+				OutContainerInstances.Add(ActorDescView);
+			}
+			else
+			{
+				OutActorDescViewMap.Emplace(ActorGuid, ActorDescView);
+			}
 		};
 		
 		TMap<FGuid, FGuid> ContainerGuidsRemap;
@@ -168,40 +202,36 @@ class FWorldPartitionStreamingGenerator
 
 	void CreateActorDescriptorViewsRecursive(const UActorDescContainer* InContainer, const FTransform& InTransform, const TSet<FName>& InRuntimeDataLayers, const FActorContainerID& InContainerID, const FActorContainerID& InParentContainerID, EContainerClusterMode InClusterMode, const TCHAR* OwnerName)
 	{
-		TMap<FGuid, FWorldPartitionActorDescView> ActorDescViewMap;
+		FActorDescViewMap ActorDescViewMap;
+		TArray<FWorldPartitionActorDescView> ContainerInstanceViews;
 		
 		// Gather actor descriptor views for this container
-		CreateActorDescViewMap(InContainer, ActorDescViewMap, InContainerID);
+		CreateActorDescViewMap(InContainer, ActorDescViewMap, InContainerID, ContainerInstanceViews);
 
-		// Parse actor descriptors
-		for (auto It = ActorDescViewMap.CreateIterator(); It; ++It)
+		// Parse actor containers
+		for (const FWorldPartitionActorDescView& ContainerInstanceView : ContainerInstanceViews)
 		{
-			FWorldPartitionActorDescView& ActorDescView = It.Value();
-
 			const UActorDescContainer* SubContainer;
 			EContainerClusterMode SubClusterMode;
 			FTransform SubTransform;
-			if (ActorDescView.GetContainerInstance(SubContainer, SubTransform, SubClusterMode))
-			{
-				check(SubContainer);
 
-				const FGuid ActorGuid = ActorDescView.GetGuid();
-				const FActorContainerID SubContainerID(InContainerID, ActorGuid);
+			if (!ContainerInstanceView.GetContainerInstance(SubContainer, SubTransform, SubClusterMode))
+			{
+				//@todo_ow: make a specific error for missing container instance sublevel?
+				ErrorHandler->OnInvalidReference(ContainerInstanceView, FGuid());
+				continue;
+			}
 
-				// Combine actor runtime Data Layers with parent container runtime Data Layers
-				TSet<FName> CombinedRuntimeDataLayers = InRuntimeDataLayers;
-				CombinedRuntimeDataLayers.Append(ActorDescView.GetRuntimeDataLayers());
+			check(SubContainer);
 
-				CreateActorDescriptorViewsRecursive(SubContainer, SubTransform * InTransform, CombinedRuntimeDataLayers, SubContainerID, InContainerID, SubClusterMode, *ActorDescView.GetActorLabelOrName().ToString());
+			const FGuid ActorGuid = ContainerInstanceView.GetGuid();
+			const FActorContainerID SubContainerID(InContainerID, ActorGuid);
 
-				// The container actor can be removed now that its contained actors has been registered
-				It.RemoveCurrent();
-			}
+			// Combine actor runtime Data Layers with parent container runtime Data Layers
+			TSet<FName> CombinedRuntimeDataLayers = InRuntimeDataLayers;
+			CombinedRuntimeDataLayers.Append(ContainerInstanceView.GetRuntimeDataLayers());
 
-			if (!InContainerID.IsMainContainer())
-			{
-				ActorDescViewsContainersMap.Add(ActorDescView.GetGuid(), InContainerID);
-			}
+			CreateActorDescriptorViewsRecursive(SubContainer, SubTransform * InTransform, CombinedRuntimeDataLayers, SubContainerID, InContainerID, SubClusterMode, *ContainerInstanceView.GetActorLabelOrName().ToString());
 		}
 
 		// Create container descriptor
@@ -229,16 +259,14 @@ class FWorldPartitionStreamingGenerator
 	{
 		CreateActorDescriptorViewsRecursive(InContainer, FTransform::Identity, TSet<FName>(), FActorContainerID(), FActorContainerID(), EContainerClusterMode::Partitioned, TEXT("MainContainer"));
 
-		// Once all views are created
+		// Resolve actor descriptor views once all views are created
 		for (auto It = ContainerDescriptorsMap.CreateIterator(); It; ++It)
 		{
 			const FActorContainerID& ContainerID = It.Key();
 			FContainerDescriptor& ContainerDescriptor = It.Value();
 
-			for (auto ActorDescIt = ContainerDescriptor.ActorDescViewMap.CreateIterator(); ActorDescIt; ++ActorDescIt)
+			ContainerDescriptor.ActorDescViewMap.ForEachActorDescView([this, &ContainerDescriptor](FWorldPartitionActorDescView& ActorDescView)
 			{
-				FWorldPartitionActorDescView& ActorDescView = ActorDescIt.Value();
-
 				if (!bEnableStreaming)
 				{
 					ActorDescView.SetForcedNonSpatiallyLoaded();
@@ -246,7 +274,7 @@ class FWorldPartitionStreamingGenerator
 
 				ResolveRuntimeDataLayers(ActorDescView, ContainerDescriptor.ActorDescViewMap);
 				ResolveRuntimeReferences(ActorDescView, ContainerDescriptor.ActorDescViewMap);
-			}
+			});
 		}
 	}
 
@@ -256,42 +284,6 @@ class FWorldPartitionStreamingGenerator
 	 */
 	void ValidateActorDescriptorViews()
 	{
-		// Validate data layers
-		auto IsReferenceGridPlacementValid = [](const FWorldPartitionActorDescView& RefererActorDescView, const FWorldPartitionActorDescView& ReferenceActorDescView)
-		{
-			const bool bIsActorDescSpatiallyLoaded = RefererActorDescView.GetIsSpatiallyLoaded();
-			const bool bIsActorDescRefSpatiallyLoaded = ReferenceActorDescView.GetIsSpatiallyLoaded();
-
-			// The only case we support right now is spatially loaded actors referencing non-spatially loaded actors, when target is not in data layers.
-			// For this to work with data layers, we need to implement dependency logic support in the content cooker splitter.
-			if (bIsActorDescSpatiallyLoaded && !bIsActorDescRefSpatiallyLoaded && ReferenceActorDescView.GetDataLayers().IsEmpty())
-			{
-				return true;
-			}
-
-			return bIsActorDescSpatiallyLoaded == bIsActorDescRefSpatiallyLoaded;
-		};
-
-		// Validate grid placement
-		auto IsReferenceDataLayersValid = [](const FWorldPartitionActorDescView& RefererActorDescView, const FWorldPartitionActorDescView& ReferenceActorDescView)
-		{
-			if (RefererActorDescView.GetRuntimeDataLayers().Num() == ReferenceActorDescView.GetRuntimeDataLayers().Num())
-			{
-				const TSet<FName> RefererActorDescDataLayers(RefererActorDescView.GetRuntimeDataLayers());
-				const TSet<FName> ReferenceActorDescDataLayers(ReferenceActorDescView.GetRuntimeDataLayers());
-
-				return RefererActorDescDataLayers.Includes(ReferenceActorDescDataLayers);
-			}
-
-			return false;
-		};
-
-		// Validate runtime grid
-		auto IsReferenceRuntimeGridValid = [](const FWorldPartitionActorDescView& RefererActorDescView, const FWorldPartitionActorDescView& ReferenceActorDescView)
-		{
-			return RefererActorDescView.GetRuntimeGrid() == ReferenceActorDescView.GetRuntimeGrid();
-		};
-
 		for (auto It = ContainerDescriptorsMap.CreateIterator(); It; ++It)
 		{
 			const FActorContainerID& ContainerID = It.Key();
@@ -306,7 +298,7 @@ class FWorldPartitionStreamingGenerator
 
 					for (AActor* Actor : LevelScriptExternalActorReferences)
 					{
-						if (FWorldPartitionActorDescView* ActorDescView = ContainerDescriptor.ActorDescViewMap.Find(Actor->GetActorGuid()))
+						if (FWorldPartitionActorDescView* ActorDescView = ContainerDescriptor.ActorDescViewMap.FindByGuid(Actor->GetActorGuid()))
 						{
 							if (ActorDescView->GetIsSpatiallyLoaded())
 							{
@@ -330,28 +322,59 @@ class FWorldPartitionStreamingGenerator
 			// Since the references can form cycles/long chains in the data fixes might need to be propagated in multiple passes.
 			// 
 			// This works because fixes are deterministic and always apply the same way to both Actors being modified, so there's no ordering issues possible
-			uint32 NbValidationPasses = 0;
-			uint32 NbErrorsDetected;
-
-			do
+			int32 NbErrorsDetected = INDEX_NONE;
+			for(uint32 NbValidationPasses = 0; NbErrorsDetected; NbValidationPasses++)
 			{
 				NbErrorsDetected = 0;
 
-				for (auto ActorDescIt = ContainerDescriptor.ActorDescViewMap.CreateIterator(); ActorDescIt; ++ActorDescIt)
+				ContainerDescriptor.ActorDescViewMap.ForEachActorDescView([this, &ContainerDescriptor, &NbErrorsDetected, &NbValidationPasses](FWorldPartitionActorDescView& ActorDescView)
 				{
-					FWorldPartitionActorDescView& ActorDescView = ActorDescIt.Value();
+					// Validate data layers
+					auto IsReferenceGridPlacementValid = [](const FWorldPartitionActorDescView& RefererActorDescView, const FWorldPartitionActorDescView& ReferenceActorDescView)
+					{
+						const bool bIsActorDescSpatiallyLoaded = RefererActorDescView.GetIsSpatiallyLoaded();
+						const bool bIsActorDescRefSpatiallyLoaded = ReferenceActorDescView.GetIsSpatiallyLoaded();
+
+						// The only case we support right now is spatially loaded actors referencing non-spatially loaded actors, when target is not in data layers.
+						// For this to work with data layers, we need to implement dependency logic support in the content cooker splitter.
+						if (bIsActorDescSpatiallyLoaded && !bIsActorDescRefSpatiallyLoaded && ReferenceActorDescView.GetDataLayers().IsEmpty())
+						{
+							return true;
+						}
 
-					struct ActorReferenceInfo
+						return bIsActorDescSpatiallyLoaded == bIsActorDescRefSpatiallyLoaded;
+					};
+
+					// Validate grid placement
+					auto IsReferenceDataLayersValid = [](const FWorldPartitionActorDescView& RefererActorDescView, const FWorldPartitionActorDescView& ReferenceActorDescView)
 					{
-						FGuid								ActorGuid;
-						FWorldPartitionActorDescView*		ActorDesc;
-						FGuid								ReferenceGuid;
-						FWorldPartitionActorDescView*		ReferenceActorDesc;
+						if (RefererActorDescView.GetRuntimeDataLayers().Num() == ReferenceActorDescView.GetRuntimeDataLayers().Num())
+						{
+							const TSet<FName> RefererActorDescDataLayers(RefererActorDescView.GetRuntimeDataLayers());
+							const TSet<FName> ReferenceActorDescDataLayers(ReferenceActorDescView.GetRuntimeDataLayers());
+
+							return RefererActorDescDataLayers.Includes(ReferenceActorDescDataLayers);
+						}
+
+						return false;
+					};
+
+					// Validate runtime grid
+					auto IsReferenceRuntimeGridValid = [](const FWorldPartitionActorDescView& RefererActorDescView, const FWorldPartitionActorDescView& ReferenceActorDescView)
+					{
+						return RefererActorDescView.GetRuntimeGrid() == ReferenceActorDescView.GetRuntimeGrid();
+					};
+
+					struct FActorReferenceInfo
+					{
+						FGuid ActorGuid;
+						FWorldPartitionActorDescView* ActorDesc;
+						FGuid ReferenceGuid;
+						FWorldPartitionActorDescView* ReferenceActorDesc;
 					};
 
-					TArray<ActorReferenceInfo, TInlineAllocator<8>> References;
-					
 					// Build references List
+					TArray<FActorReferenceInfo> References;
 
 					// Add normal actor references
 					for (const FGuid& ReferenceGuid : ActorDescView.GetReferences())
@@ -359,13 +382,13 @@ class FWorldPartitionStreamingGenerator
 						if (ReferenceGuid != ActorDescView.GetParentActor()) // References to the parent are inversed in their handling 
 						{
 							// Filter out parent back references
-							FWorldPartitionActorDescView* ReferenceActorDesc = ContainerDescriptor.ActorDescViewMap.Find(ReferenceGuid);
+							FWorldPartitionActorDescView* ReferenceActorDesc = ContainerDescriptor.ActorDescViewMap.FindByGuid(ReferenceGuid);
 							if (ReferenceActorDesc && ReferenceActorDesc->GetParentActor() == ActorDescView.GetGuid())
 							{
 								continue;
 							}
 
-							References.Emplace(ActorReferenceInfo{ ActorDescView.GetGuid(), &ActorDescView, ReferenceGuid, ReferenceActorDesc });
+							References.Emplace(FActorReferenceInfo{ ActorDescView.GetGuid(), &ActorDescView, ReferenceGuid, ReferenceActorDesc });
 						}
 					}
 
@@ -376,7 +399,7 @@ class FWorldPartitionStreamingGenerator
 
 					while (ParentGuid.IsValid())
 					{
-						FWorldPartitionActorDescView* ParentDescView = ContainerDescriptor.ActorDescViewMap.Find(ParentGuid);
+						FWorldPartitionActorDescView* ParentDescView = ContainerDescriptor.ActorDescViewMap.FindByGuid(ParentGuid);
 					
 						if (ParentDescView)
 						{
@@ -392,16 +415,16 @@ class FWorldPartitionStreamingGenerator
 
 					if (TopParentDescView)
 					{
-						References.Emplace(ActorReferenceInfo{ TopParentDescView->GetGuid(), TopParentDescView, ActorDescView.GetGuid(), &ActorDescView });
+						References.Emplace(FActorReferenceInfo{ TopParentDescView->GetGuid(), TopParentDescView, ActorDescView.GetGuid(), &ActorDescView });
 					}
 
 					if (ParentGuid.IsValid())
 					{
 						// In case of missing parent add a missing reference 
-						References.Emplace(ActorReferenceInfo{ ActorDescView.GetGuid(), &ActorDescView, ParentGuid, nullptr });
+						References.Emplace(FActorReferenceInfo{ ActorDescView.GetGuid(), &ActorDescView, ParentGuid, nullptr });
 					}
 
-					for (ActorReferenceInfo& Info : References)
+					for (FActorReferenceInfo& Info : References)
 					{
 						FWorldPartitionActorDescView* RefererActorDescView = Info.ActorDesc;
 						FWorldPartitionActorDescView* ReferenceActorDescView = Info.ReferenceActorDesc;
@@ -415,13 +438,13 @@ class FWorldPartitionStreamingGenerator
 								{
 									ErrorHandler->OnInvalidReferenceGridPlacement(*RefererActorDescView, *ReferenceActorDescView);									
 								}
-								NbErrorsDetected++;
-
-								if (NbValidationPasses)
+								else
 								{
 									RefererActorDescView->SetForcedNonSpatiallyLoaded();
 									ReferenceActorDescView->SetForcedNonSpatiallyLoaded();
 								}
+
+								NbErrorsDetected++;
 							}
 
 							if (!IsReferenceDataLayersValid(*RefererActorDescView, *ReferenceActorDescView))
@@ -430,13 +453,13 @@ class FWorldPartitionStreamingGenerator
 								{
 									ErrorHandler->OnInvalidReferenceDataLayers(*RefererActorDescView, *ReferenceActorDescView);									
 								}
-								NbErrorsDetected++;
-
-								if (NbValidationPasses)
+								else
 								{
 									RefererActorDescView->SetInvalidDataLayers();
 									ReferenceActorDescView->SetInvalidDataLayers();
 								}
+
+								NbErrorsDetected++;
 							}
 
 							if (!IsReferenceRuntimeGridValid(*RefererActorDescView, *ReferenceActorDescView))
@@ -445,13 +468,13 @@ class FWorldPartitionStreamingGenerator
 								{
 									ErrorHandler->OnInvalidReferenceRuntimeGrid(*RefererActorDescView, *ReferenceActorDescView);
 								}
-								NbErrorsDetected++;
-
-								if (NbValidationPasses)
+								else
 								{
 									RefererActorDescView->SetInvalidRuntimeGrid();
 									ReferenceActorDescView->SetInvalidRuntimeGrid();
 								}
+
+								NbErrorsDetected++;
 							}
 
 						}
@@ -464,20 +487,30 @@ class FWorldPartitionStreamingGenerator
 							// Do not increment NbErrorsDetected since it won't be fixed and thus will always occur
 						}
 					}
-				}		
-
-				NbValidationPasses++;
+				});		
 			}
-			while (NbErrorsDetected);
 
-			// Report ActorDescs that need to be resaved
-			for (auto ActorDescIt = ContainerDescriptor.ActorDescViewMap.CreateIterator(); ActorDescIt; ++ActorDescIt)
+			// Report actors that need to be resaved
+			ContainerDescriptor.ActorDescViewMap.ForEachActorDescView([this](FWorldPartitionActorDescView& ActorDescView)
 			{
-				FWorldPartitionActorDescView& ActorDescView = ActorDescIt.Value();
 				if (ActorDescView.IsResaveNeeded())
 				{
 					ErrorHandler->OnActorNeedsResave(ActorDescView);
 				}
+			});
+
+			// Validate data layers
+			if (ContainerID.IsMainContainer() && ContainerDescriptor.Container->GetWorld())
+			{
+				if (UDataLayerSubsystem* DataLayerSubsystem = UWorld::GetSubsystem<UDataLayerSubsystem>(ContainerDescriptor.Container->GetWorld()))
+				{
+					DataLayerSubsystem->ForEachDataLayer([this](const UDataLayerInstance* DataLayerInstance)
+					{
+						DataLayerInstance->Validate(ErrorHandler);
+						return true;
+					});
+					break;
+				}
 			}
 		}
 	}
@@ -493,15 +526,13 @@ class FWorldPartitionStreamingGenerator
 			const FActorContainerID& ContainerID = ContainerIt.Key();
 			FContainerDescriptor& ContainerDescriptor = ContainerIt.Value();
 
-			for (auto ActorDescIt = ContainerDescriptor.ActorDescViewMap.CreateConstIterator(); ActorDescIt; ++ActorDescIt)
+			ContainerDescriptor.ActorDescViewMap.ForEachActorDescView([&ContainerDescriptor](const FWorldPartitionActorDescView& ActorDescView)
 			{
-				const FWorldPartitionActorDescView& ActorDescView = ActorDescIt.Value();
-
 				if (ActorDescView.GetIsSpatiallyLoaded())
 				{
 					ContainerDescriptor.Bounds += ActorDescView.GetBounds().TransformBy(ContainerDescriptor.Transform);
 				}
-			}
+			});
 		}
 
 		// Update parent containers bounds, this relies on the fact that ContainersHierarchy is built bottom up
@@ -513,29 +544,6 @@ class FWorldPartitionStreamingGenerator
 		}
 	}
 
-	void ValidateDataLayers()
-	{
-		for (auto ContainerIt = ContainerDescriptorsMap.CreateIterator(); ContainerIt; ++ContainerIt)
-		{
-			const FActorContainerID& ContainerID = ContainerIt.Key();
-			FContainerDescriptor& ContainerDescriptor = ContainerIt.Value();
-			if (ContainerID.IsMainContainer() && ContainerDescriptor.Container->GetWorld())
-			{
-				UDataLayerSubsystem* DataLayerSubsystem = UWorld::GetSubsystem<UDataLayerSubsystem>(ContainerDescriptor.Container->GetWorld());
-				if(DataLayerSubsystem != nullptr)
-				{
-					DataLayerSubsystem->ForEachDataLayer([this](const UDataLayerInstance* DataLayerInstance)
-						{
-							DataLayerInstance->Validate(ErrorHandler);
-							return true;
-						});
-
-					break;
-				}
-			}
-		}
-	}
-
 public:
 	FWorldPartitionStreamingGenerator(FActorDescList* InModifiedActorsDescList, IStreamingGenerationErrorHandler* InErrorHandler, bool bInEnableStreaming)
 	: bEnableStreaming(bInEnableStreaming)
@@ -551,9 +559,6 @@ public:
 		// Preparation Phase :: Actor Descriptor Views Validation
 		ValidateActorDescriptorViews();
 
-		// Preparation Phase :: Data Layer Validation
-		ValidateDataLayers();
-
 		// Update container descriptors
 		UpdateContainerDescriptors();
 	}
@@ -568,7 +573,7 @@ public:
 			const FActorContainerID& ContainerID = It.Key();
 			const FContainerDescriptor& ContainerDescriptor = It.Value();
 
-			ContainerInstances.Emplace(ContainerID, ContainerDescriptor.Transform, ContainerDescriptor.Bounds, ContainerDescriptor.RuntimeDataLayers, ContainerDescriptor.ClusterMode, ContainerDescriptor.Container, ContainerDescriptor.ActorDescViewMap);
+			ContainerInstances.Emplace(ContainerID, ContainerDescriptor.Transform, ContainerDescriptor.Bounds, ContainerDescriptor.RuntimeDataLayers, ContainerDescriptor.ClusterMode, ContainerDescriptor.Container, ContainerDescriptor.ActorDescViewMap.ActorDescViewsByGuid);
 		}
 
 		return FActorClusterContext(MoveTemp(ContainerInstances), InFilterActorDescViewFunc);
@@ -608,17 +613,17 @@ public:
 					Ar.Printf(TEXT("Container: %s"), *ContainerDescriptor.Container->GetContainerPackage().ToString());
 				}
 
-				if (ContainerDescriptor.ActorDescViewMap.Num())
+				if (ContainerDescriptor.ActorDescViewMap.ActorDescViewsByGuid.Num())
 				{
 					FHierarchicalLogArchive::FIndentScope IndentScope = Ar.PrintfIndent(TEXT("ActorDescs:"));
 
-					TMap<FGuid, FWorldPartitionActorDescView> SortedActorDescViewMap = ContainerDescriptor.ActorDescViewMap;
+					TMap<FGuid, FWorldPartitionActorDescView*> SortedActorDescViewMap = ContainerDescriptor.ActorDescViewMap.ActorDescViewsByGuid;
 					SortedActorDescViewMap.KeySort([](const FGuid& GuidA, const FGuid& GuidB) { return GuidA < GuidB; });
 
 					for (auto ActorDescIt = SortedActorDescViewMap.CreateConstIterator(); ActorDescIt; ++ActorDescIt)
 					{
-						const FWorldPartitionActorDescView& ActorDescView = ActorDescIt.Value();
-						Ar.Printf(TEXT("%s"), *ActorDescView.ToString());
+						const FWorldPartitionActorDescView& ActorDescView = *ActorDescIt.Value();
+						Ar.Print(*ActorDescView.ToString());
 					}
 				}
 
@@ -660,7 +665,7 @@ private:
 		FTransform Transform;
 		const UActorDescContainer* Container;
 		EContainerClusterMode ClusterMode;
-		TMap<FGuid, FWorldPartitionActorDescView> ActorDescViewMap;
+		FActorDescViewMap ActorDescViewMap;
 		TSet<FName> RuntimeDataLayers;
 		FString OwnerName;
 	};
@@ -670,9 +675,6 @@ private:
 
 	/** Maps containers IDs to their parent ID */
 	TMap<FActorContainerID, FActorContainerID> ContainersHierarchy;
-
-	/** Maps actor descriptor views to containers IDs */
-	TMap<FGuid, FActorContainerID> ActorDescViewsContainersMap;
 };
 
 bool UWorldPartition::GenerateStreaming(TArray<FString>* OutPackagesToGenerate)
@@ -697,17 +699,15 @@ bool UWorldPartition::GenerateStreaming(TArray<FString>* OutPackagesToGenerate)
 	FHierarchicalLogArchive HierarchicalLogAr(*LogFileAr);
 
 	FActorClusterContext ActorClusterContext;
-	{
-		FWorldPartitionStreamingGenerator StreamingGenerator(ModifiedActorsDescList, ErrorHandler, IsStreamingEnabled());
+	FWorldPartitionStreamingGenerator StreamingGenerator(ModifiedActorsDescList, ErrorHandler, IsStreamingEnabled());
 
-		// Preparation Phase
-		StreamingGenerator.PreparationPhase(this);
+	// Preparation Phase
+	StreamingGenerator.PreparationPhase(this);
 
-		StreamingGenerator.DumpStateLog(HierarchicalLogAr);
+	StreamingGenerator.DumpStateLog(HierarchicalLogAr);
 
-		// Preparation Phase :: Actor Clusters Creation
-		ActorClusterContext = StreamingGenerator.CreateActorClusters();
-	}
+	// Preparation Phase :: Actor Clusters Creation
+	ActorClusterContext = StreamingGenerator.CreateActorClusters();
 
 	// Generate streaming
 	check(!StreamingPolicy);
@@ -760,12 +760,9 @@ void UWorldPartition::CheckForErrors(IStreamingGenerationErrorHandler* ErrorHand
 
 void UWorldPartition::CheckForErrors(IStreamingGenerationErrorHandler* ErrorHandler, const UActorDescContainer* ActorDescContainer, bool bEnableStreaming)
 {
-	FActorClusterContext ActorClusterContext;
-	{
-		FActorDescList ModifiedActorDescList;
-		FWorldPartitionStreamingGenerator StreamingGenerator(ActorDescContainer->GetWorld() ? &ModifiedActorDescList : nullptr, ErrorHandler, bEnableStreaming);
-		StreamingGenerator.PreparationPhase(ActorDescContainer);
-	}
+	FActorDescList ModifiedActorDescList;
+	FWorldPartitionStreamingGenerator StreamingGenerator(ActorDescContainer->GetWorld() ? &ModifiedActorDescList : nullptr, ErrorHandler, bEnableStreaming);
+	StreamingGenerator.PreparationPhase(ActorDescContainer);
 }
 #endif // WITH_EDITOR
 
diff --git a/Engine/Source/Runtime/Engine/Public/WorldPartition/DataLayer/DataLayerUtils.h b/Engine/Source/Runtime/Engine/Public/WorldPartition/DataLayer/DataLayerUtils.h
--- a/Engine/Source/Runtime/Engine/Public/WorldPartition/DataLayer/DataLayerUtils.h
+++ b/Engine/Source/Runtime/Engine/Public/WorldPartition/DataLayer/DataLayerUtils.h
@@ -8,6 +8,7 @@
 #include "WorldPartition/DataLayer/DataLayerType.h"
 #include "WorldPartition/DataLayer/DataLayerInstance.h"
 #include "WorldPartition/WorldPartitionActorDescView.h"
+#include "WorldPartition/WorldPartitionStreamingGeneration.h"
 
 class FWorldDataLayersActorDesc;
 class FWorldPartitionActorDesc;
@@ -26,7 +27,7 @@ public:
 
 	static TArray<FName> ResolvedDataLayerInstanceNames(const FWorldPartitionActorDesc* InActorDesc, const FWorldDataLayersActorDesc* InWorldDataLayersActorDesc = nullptr, UWorld* InWorld = nullptr, bool* bOutIsResultValid = nullptr);
 	
-	static bool ResolveRuntimeDataLayerInstanceNames(const FWorldPartitionActorDescView& InActorDescView, const TMap<FGuid, FWorldPartitionActorDescView>& ActorDescViews, TArray<FName>& OutRuntimeDataLayerInstanceNames);
+	static bool ResolveRuntimeDataLayerInstanceNames(const FWorldPartitionActorDescView& InActorDescView, const FActorDescViewMap& ActorDescViewMap, TArray<FName>& OutRuntimeDataLayerInstanceNames);
 #endif
 
 #if DATALAYER_TO_INSTANCE_RUNTIME_CONVERSION_ENABLED
diff --git a/Engine/Source/Runtime/Engine/Public/WorldPartition/LevelInstance/LevelInstanceActorDesc.h b/Engine/Source/Runtime/Engine/Public/WorldPartition/LevelInstance/LevelInstanceActorDesc.h
--- a/Engine/Source/Runtime/Engine/Public/WorldPartition/LevelInstance/LevelInstanceActorDesc.h
+++ b/Engine/Source/Runtime/Engine/Public/WorldPartition/LevelInstance/LevelInstanceActorDesc.h
@@ -26,6 +26,7 @@
 
 	inline FName GetLevelPackage() const { return LevelPackage; }
 
+	virtual bool IsContainerInstance() const override;
 	virtual bool GetContainerInstance(const UActorDescContainer*& OutLevelContainer, FTransform& OutLevelTransform, EContainerClusterMode& OutClusterMode) const override;
 
 protected:
diff --git a/Engine/Source/Runtime/Engine/Public/WorldPartition/WorldPartitionActorCluster.h b/Engine/Source/Runtime/Engine/Public/WorldPartition/WorldPartitionActorCluster.h
--- a/Engine/Source/Runtime/Engine/Public/WorldPartition/WorldPartitionActorCluster.h
+++ b/Engine/Source/Runtime/Engine/Public/WorldPartition/WorldPartitionActorCluster.h
@@ -83,7 +83,7 @@ struct FActorCluster
 	FDataLayersID				DataLayersID;
 
 	FActorCluster(UWorld* InWorld, const FWorldPartitionActorDescView& InActorDescView);
-	void Add(const FActorCluster& InActorCluster, const TMap<FGuid, FWorldPartitionActorDescView>& InActorDescViewMap);
+	void Add(const FActorCluster& InActorCluster);
 };
 
 /**
@@ -91,14 +91,14 @@ struct FActorCluster
  */
 struct FActorContainerInstance
 {
-	FActorContainerInstance(const FActorContainerID& InID, const FTransform& InTransform, const FBox& InBounds, const TSet<FName>& InRuntimeDataLayers, EContainerClusterMode InClusterMode, const UActorDescContainer* InContainer, TMap<FGuid, FWorldPartitionActorDescView> InActorDescViewMap);
+	FActorContainerInstance(const FActorContainerID& InID, const FTransform& InTransform, const FBox& InBounds, const TSet<FName>& InRuntimeDataLayers, EContainerClusterMode InClusterMode, const UActorDescContainer* InContainer, TMap<FGuid, FWorldPartitionActorDescView*> InActorDescViewMap);
 	
 	FActorContainerID			ID;
 	FTransform					Transform;
 	FBox						Bounds;
 	EContainerClusterMode		ClusterMode;
 	const UActorDescContainer*	Container;
-	TMap<FGuid, FWorldPartitionActorDescView> ActorDescViewMap;
+	TMap<FGuid, FWorldPartitionActorDescView*> ActorDescViewMap;
 	TSet<const UDataLayerInstance*>		DataLayers;
 
 	const FWorldPartitionActorDescView& GetActorDescView(const FGuid& InGuid) const;
@@ -168,11 +168,11 @@ public:
 	FActorContainerInstance* GetClusterInstance(const UActorDescContainer* InContainer);
 	const FActorContainerInstance* GetClusterInstance(const UActorDescContainer* InContainer) const;
 
-	static void CreateActorClusters(UWorld* World, const TMap<FGuid, FWorldPartitionActorDescView>& ActorDescViewMap, TArray<FActorCluster>& OutActorClusters);
+	static void CreateActorClusters(UWorld* World, const TMap<FGuid, FWorldPartitionActorDescView*>& ActorDescViewMap, TArray<FActorCluster>& OutActorClusters);
 
 private:
 	const TArray<FActorCluster>& CreateActorClusters(const FActorContainerInstance& ContainerInstance);
-	static void CreateActorClusters(UWorld* World, const TMap<FGuid, FWorldPartitionActorDescView>& ActorDescViewMap, TArray<FActorCluster>& OutActorClusters, FFilterActorDescViewFunc InFilterActorDescViewFunc);
+	static void CreateActorClusters(UWorld* World, const TMap<FGuid, FWorldPartitionActorDescView*>& ActorDescViewMap, TArray<FActorCluster>& OutActorClusters, FFilterActorDescViewFunc InFilterActorDescViewFunc);
 	
 	// Init data
 	FFilterActorDescViewFunc FilterActorDescViewFunc;
diff --git a/Engine/Source/Runtime/Engine/Public/WorldPartition/WorldPartitionActorDesc.h b/Engine/Source/Runtime/Engine/Public/WorldPartition/WorldPartitionActorDesc.h
--- a/Engine/Source/Runtime/Engine/Public/WorldPartition/WorldPartitionActorDesc.h
+++ b/Engine/Source/Runtime/Engine/Public/WorldPartition/WorldPartitionActorDesc.h
@@ -104,7 +104,9 @@ public:
 	FName GetActorLabelOrName() const;
 	FName GetDisplayClassName() const;
 
+	virtual bool IsContainerInstance() const { return false; }
 	virtual bool GetContainerInstance(const UActorDescContainer*& OutLevelContainer, FTransform& OutLevelTransform, EContainerClusterMode& OutClusterMode) const { return false; }
+
 	virtual const FGuid& GetSceneOutlinerParent() const { return GetParentActor(); }
 	virtual bool IsResaveNeeded() const { return false; }
 	virtual bool IsRuntimeRelevant(const FActorContainerID& InContainerID) const { return true; }
diff --git a/Engine/Source/Runtime/Engine/Public/WorldPartition/WorldPartitionActorDescView.h b/Engine/Source/Runtime/Engine/Public/WorldPartition/WorldPartitionActorDescView.h
--- a/Engine/Source/Runtime/Engine/Public/WorldPartition/WorldPartitionActorDescView.h
+++ b/Engine/Source/Runtime/Engine/Public/WorldPartition/WorldPartitionActorDescView.h
@@ -47,6 +47,7 @@ public:
 	FName GetActorName() const;
 	const FGuid& GetFolderGuid() const;
 
+	bool IsContainerInstance() const;
 	bool GetContainerInstance(const UActorDescContainer*& OutLevelContainer, FTransform& OutLevelTransform, EContainerClusterMode& OutClusterMode) const;
 
 	FName GetActorLabelOrName() const;
diff --git a/Engine/Source/Runtime/Engine/Public/WorldPartition/WorldPartitionStreamingGeneration.h b/Engine/Source/Runtime/Engine/Public/WorldPartition/WorldPartitionStreamingGeneration.h
new file mode 100644
--- /dev/null
+++ b/Engine/Source/Runtime/Engine/Public/WorldPartition/WorldPartitionStreamingGeneration.h
@@ -0,0 +1,88 @@
+// Copyright Epic Games, Inc. All Rights Reserved.
+#pragma once
+
+#if WITH_EDITOR
+#include "CoreMinimal.h"
+#include "Templates/IntegralConstant.h"
+#include "Templates/UnrealTypeTraits.h"
+#include "WorldPartition/WorldPartitionActorDescView.h"
+
+class ENGINE_API FActorDescViewMap
+{
+	friend class FWorldPartitionStreamingGenerator;
+
+private:
+    template<typename F, typename V>bool DoCall(F Func, V Val, std::true_type) { return Func(*Val); }
+	template<typename F, typename V> bool DoCall(F Func, V Val, std::false_type) { Func(*Val); return true; }
+
+	template <class T>
+	void ForEachActorDescView(T Func)
+	{
+		for (TUniquePtr<FWorldPartitionActorDescView>& ActorDescView : ActorDescViewList)
+		{
+			if (!DoCall(Func, ActorDescView.Get(), std::is_same<bool, decltype(Func(*ActorDescView))>()))
+			{
+				return;
+			}
+		}
+	}
+
+	FWorldPartitionActorDescView* FindByGuid(const FGuid& InGuid)
+	{
+		if (FWorldPartitionActorDescView** ActorDescViewPtr = ActorDescViewsByGuid.Find(InGuid))
+		{
+			return *ActorDescViewPtr;
+		}
+		return nullptr;
+	}
+
+public:
+	FActorDescViewMap();
+
+	// Non-copyable but movable
+	FActorDescViewMap(const FActorDescViewMap&) = delete;
+	FActorDescViewMap(FActorDescViewMap&&) = default;
+	FActorDescViewMap& operator=(const FActorDescViewMap&) = delete;
+	FActorDescViewMap& operator=(FActorDescViewMap&&) = default;
+
+	FWorldPartitionActorDescView* Emplace(const FGuid& InActorGuid, const FWorldPartitionActorDescView& InActorDescView);
+
+	template <class T>
+	void ForEachActorDescView(T Func) const
+	{
+		for (const TUniquePtr<FWorldPartitionActorDescView>& ActorDescView : ActorDescViewList)
+		{
+			if (!DoCall(Func, ActorDescView.Get(), std::is_same<bool, decltype(Func(*ActorDescView))>()))
+			{
+				return;
+			}
+		}
+	}
+
+	const FWorldPartitionActorDescView* FindByGuid(const FGuid& InGuid) const
+	{
+		if (const FWorldPartitionActorDescView* const* ActorDescViewPtr = ActorDescViewsByGuid.Find(InGuid))
+		{
+			return *ActorDescViewPtr;
+		}
+		return nullptr;
+	}
+
+	template <class T>
+	TArray<const FWorldPartitionActorDescView*> FindByExactNativeClass() const
+	{
+		return FindByExactNativeClass(T::StaticClass());
+	}
+
+	TArray<const FWorldPartitionActorDescView*> FindByExactNativeClass(UClass* InExactNativeClass) const;
+
+	const TMap<FGuid, FWorldPartitionActorDescView*>& GetActorDescViewsByGuid() const { return ActorDescViewsByGuid; }
+
+protected:
+	TArray<TUniquePtr<FWorldPartitionActorDescView>> ActorDescViewList;
+
+	TMap<FGuid, FWorldPartitionActorDescView*> ActorDescViewsByGuid;
+	TMultiMap<FName, const FWorldPartitionActorDescView*> ActorDescViewsByClass;
+};
+
+#endif // WITH_EDITOR
\ No newline at end of file
