Add an override of the unreal client that takes in an FInputDeviceId.

diff --git a/Engine/Source/Runtime/Engine/Classes/Engine/GameViewportClient.h b/Engine/Source/Runtime/Engine/Classes/Engine/GameViewportClient.h
--- a/Engine/Source/Runtime/Engine/Classes/Engine/GameViewportClient.h
+++ b/Engine/Source/Runtime/Engine/Classes/Engine/GameViewportClient.h
@@ -169,7 +169,10 @@ public:
 	//~ Begin FViewportClient Interface.
 	virtual void RedrawRequested(FViewport* InViewport) override {}
 	virtual bool InputKey(const FInputKeyEventArgs& EventArgs) override;
+	// TODO: Deprecate this one
 	virtual bool InputAxis(FViewport* Viewport, int32 ControllerId, FKey Key, float Delta, float DeltaTime, int32 NumSamples=1, bool bGamepad=false) override;
+	virtual bool InputAxis(FViewport* Viewport, FInputDeviceId InputDevice, FKey Key, float Delta, float DeltaTime, int32 NumSamples = 1, bool bGamepad = false) override;
+	
 	virtual bool InputChar(FViewport* Viewport,int32 ControllerId, TCHAR Character) override;
 	virtual bool InputTouch(FViewport* Viewport, int32 ControllerId, uint32 Handle, ETouchType::Type Type, const FVector2D& TouchLocation, float Force, FDateTime DeviceTimestamp, uint32 TouchpadIndex) override;
 	virtual bool InputMotion(FViewport* Viewport, int32 ControllerId, const FVector& Tilt, const FVector& RotationRate, const FVector& Gravity, const FVector& Acceleration) override;
diff --git a/Engine/Source/Runtime/Engine/Private/GameViewportClient.cpp b/Engine/Source/Runtime/Engine/Private/GameViewportClient.cpp
--- a/Engine/Source/Runtime/Engine/Private/GameViewportClient.cpp
+++ b/Engine/Source/Runtime/Engine/Private/GameViewportClient.cpp
@@ -664,7 +664,7 @@ bool UGameViewportClient::InputKey(const FInputKeyEventArgs& InEventArgs)
 		ULocalPlayer* const TargetPlayer = GEngine->GetLocalPlayerFromControllerId(this, EventArgs.ControllerId);
 		if (TargetPlayer && TargetPlayer->PlayerController)
 		{
-			bResult = TargetPlayer->PlayerController->InputKey(FInputKeyParams(EventArgs.Key, EventArgs.Event, static_cast<double>(EventArgs.AmountDepressed), EventArgs.IsGamepad()));
+			bResult = TargetPlayer->PlayerController->InputKey(FInputKeyParams(EventArgs.Key, EventArgs.Event, static_cast<double>(EventArgs.AmountDepressed), EventArgs.IsGamepad(), EventArgs.InputDevice));
 		}
 
 		// A gameviewport is always considered to have responded to a mouse buttons to avoid throttling
@@ -695,6 +695,17 @@ bool UGameViewportClient::InputKey(const FInputKeyEventArgs& InEventArgs)
 
 
 bool UGameViewportClient::InputAxis(FViewport* InViewport, int32 ControllerId, FKey Key, float Delta, float DeltaTime, int32 NumSamples, bool bGamepad)
+{
+	// Remap the old int32 ControllerId value to the new InputDeviceId
+	IPlatformInputDeviceMapper& DeviceMapper = IPlatformInputDeviceMapper::Get();
+	FPlatformUserId UserId = PLATFORMUSERID_NONE;
+	FInputDeviceId DeviceId = INPUTDEVICEID_NONE;
+	DeviceMapper.RemapControllerIdToPlatformUserAndDevice(ControllerId, UserId, DeviceId);
+	
+	return InputAxis(InViewport, DeviceId, Key, Delta, DeltaTime, NumSamples, bGamepad);
+}
+
+bool UGameViewportClient::InputAxis(FViewport* InViewport, FInputDeviceId InputDevice, FKey Key, float Delta, float DeltaTime, int32 NumSamples, bool bGamepad)
 {
 	if (IgnoreInput())
 	{
@@ -702,8 +713,7 @@ bool UGameViewportClient::InputAxis(FViewport* InViewport, int32 ControllerId, F
 	}
 
 	// Handle mapping controller id and key if needed
-	FInputKeyEventArgs EventArgs(InViewport, ControllerId, Key, IE_Axis);
-	RemapControllerInput(EventArgs);
+	FInputKeyEventArgs EventArgs(InViewport, InputDevice, Key, IE_Axis);
 
 	OnInputAxisEvent.Broadcast(InViewport, EventArgs.ControllerId, EventArgs.Key, Delta, DeltaTime, NumSamples, EventArgs.IsGamepad());
 	
@@ -712,10 +722,10 @@ bool UGameViewportClient::InputAxis(FViewport* InViewport, int32 ControllerId, F
 	// Don't allow mouse/joystick input axes while in PIE and the console has forced the cursor to be visible.  It's
 	// just distracting when moving the mouse causes mouse look while you are trying to move the cursor over a button
 	// in the editor!
-	if( !( InViewport->IsSlateViewport() && InViewport->IsPlayInEditorViewport() ) || ViewportConsole == NULL || !ViewportConsole->ConsoleActive() )
+	if( !( InViewport->IsSlateViewport() && InViewport->IsPlayInEditorViewport() ) || ViewportConsole == nullptr || !ViewportConsole->ConsoleActive() )
 	{
 		// route to subsystems that care
-		if (ViewportConsole != NULL)
+		if (ViewportConsole != nullptr)
 		{
 			bResult = ViewportConsole->InputAxis(EventArgs.ControllerId, EventArgs.Key, Delta, DeltaTime, NumSamples, EventArgs.IsGamepad());
 		}
@@ -728,10 +738,10 @@ bool UGameViewportClient::InputAxis(FViewport* InViewport, int32 ControllerId, F
 
 		if (!bResult)
 		{
-			ULocalPlayer* const TargetPlayer = GEngine->GetLocalPlayerFromControllerId(this, EventArgs.ControllerId);
+			ULocalPlayer* const TargetPlayer = GEngine->GetLocalPlayerFromInputDevice(this, EventArgs.InputDevice);
 			if (TargetPlayer && TargetPlayer->PlayerController)
 			{
-				bResult = TargetPlayer->PlayerController->InputKey(FInputKeyParams(EventArgs.Key, (double)Delta, DeltaTime, NumSamples, EventArgs.IsGamepad()));
+				bResult = TargetPlayer->PlayerController->InputKey(FInputKeyParams(EventArgs.Key, (double)Delta, DeltaTime, NumSamples, EventArgs.IsGamepad(), EventArgs.InputDevice));
 			}
 		}
 
diff --git a/Engine/Source/Runtime/Engine/Public/UnrealClient.h b/Engine/Source/Runtime/Engine/Public/UnrealClient.h
--- a/Engine/Source/Runtime/Engine/Public/UnrealClient.h
+++ b/Engine/Source/Runtime/Engine/Public/UnrealClient.h
@@ -800,6 +800,16 @@
 	{
 	}
 
+	FInputKeyEventArgs(FViewport* InViewport, FInputDeviceId InInputDevice, FKey InKey, EInputEvent InEvent)
+		: Viewport(InViewport)
+		, InputDevice(InInputDevice)
+		, Key(InKey)
+		, Event(InEvent)
+		, AmountDepressed(1.0f)
+		, bIsTouchEvent(false)
+	{
+	}
+
 	FInputKeyEventArgs(FViewport* InViewport, int32 InControllerId, FKey InKey, EInputEvent InEvent, float InAmountDepressed, bool bInIsTouchEvent)
 		: Viewport(InViewport)
 		, ControllerId(InControllerId)
@@ -809,6 +819,16 @@
 		, bIsTouchEvent(bInIsTouchEvent)
 	{
 	}
+	
+	FInputKeyEventArgs(FViewport* InViewport, FInputDeviceId InInputDevice, FKey InKey, EInputEvent InEvent, float InAmountDepressed, bool bInIsTouchEvent)
+		: Viewport(InViewport)
+		, InputDevice(InInputDevice)
+		, Key(InKey)
+		, Event(InEvent)
+		, AmountDepressed(InAmountDepressed)
+		, bIsTouchEvent(bInIsTouchEvent)
+	{
+	}
 
 	bool IsGamepad() const { return Key.IsGamepadKey(); }
 
@@ -818,6 +838,8 @@
 	FViewport* Viewport;
 	// The controller which the key event is from.
 	int32 ControllerId;
+	// The input device which this key event is from
+	FInputDeviceId InputDevice;
 	// The type of event which occurred.
 	FKey Key;
 	// The type of event which occurred.
@@ -870,9 +892,24 @@
 	 * @param	bGamepad - input came from gamepad (ie xbox controller)
 	 * @return	True to consume the axis movement, false to pass it on.
 	 */
+	// TODO: Deprecate this version
 	virtual bool InputAxis(FViewport* Viewport,int32 ControllerId,FKey Key,float Delta,float DeltaTime,int32 NumSamples=1,bool bGamepad=false) { return false; }
 
 	/**
+	 * Check an axis movement received by the viewport.
+	 * If the viewport client uses the movement, it should return true to consume it.
+	 * @param	Viewport - The viewport which the axis movement is from.
+	 * @param	InputDevice - The input device that triggered this axis movement
+	 * @param	Key - The name of the axis which moved.
+	 * @param	Delta - The axis movement delta.
+	 * @param	DeltaTime - The time since the last axis update.
+	 * @param	NumSamples - The number of device samples that contributed to this Delta, useful for things like smoothing
+	 * @param	bGamepad - input came from gamepad (ie xbox controller)
+	 * @return	True to consume the axis movement, false to pass it on.
+	 */
+	virtual bool InputAxis(FViewport* Viewport, FInputDeviceId InputDevice, FKey Key, float Delta, float DeltaTime, int32 NumSamples = 1, bool bGamepad = false) { return false; }
+
+	/**
 	 * Check a character input received by the viewport.
 	 * If the viewport client uses the character, it should return true to consume it.
 	 * @param	Viewport - The viewport which the axis movement is from.
