Add GameplayStatics functions for getting the player's platform user id instead of the old int32 controllerID

diff --git a/Engine/Source/Runtime/Engine/Classes/Kismet/GameplayStatics.h b/Engine/Source/Runtime/Engine/Classes/Kismet/GameplayStatics.h
--- a/Engine/Source/Runtime/Engine/Classes/Kismet/GameplayStatics.h
+++ b/Engine/Source/Runtime/Engine/Classes/Kismet/GameplayStatics.h
@@ -34,7 +34,7 @@ class USoundBase;
 class USoundConcurrency;
 class UStaticMesh;
 class FMemoryReader;
-
+class APlayerController;
 struct FDialogueContext;
 
 
@@ -225,6 +225,14 @@ class ENGINE_API UGameplayStatics : public UBlueprintFunctionLibrary
 	UFUNCTION(BlueprintPure, Category="Game", meta=(DisplayName="Get Local Player Controller From ID", WorldContext="WorldContextObject", UnsafeDuringActorConstruction="true"))
 	static class APlayerController* GetPlayerControllerFromID(const UObject* WorldContextObject, int32 ControllerID);
 
+	/** 
+	 * Returns the player controller with the specified physical controller ID. This only works for local players.
+	 *
+	 * @param ControllerID	Physical controller ID, the same value returned from Get Player Controller ID
+	 */
+	UFUNCTION(BlueprintPure, Category="Game", meta=(DisplayName="Get Local Player Controller From Platform User", WorldContext="WorldContextObject", UnsafeDuringActorConstruction="true"))
+	static APlayerController* GetPlayerControllerFromPlatformUser(const UObject* WorldContextObject, FPlatformUserId UserId);
+
 	/** 
 	 * Create a new local player for this game, for cases like local multiplayer.
 	 *
@@ -234,6 +242,16 @@ class ENGINE_API UGameplayStatics : public UBlueprintFunctionLibrary
 	 */
 	UFUNCTION(BlueprintCallable, Category = "Game", meta = (DisplayName="Create Local Player", WorldContext = "WorldContextObject", AdvancedDisplay = "2", UnsafeDuringActorConstruction = "true"))
 	static class APlayerController* CreatePlayer(const UObject* WorldContextObject, int32 ControllerId = -1, bool bSpawnPlayerController = true);
+	
+	/** 
+	 * Create a new local player for this game, for cases like local multiplayer.
+	 *
+	 * @param UserId             	   The platform user id that should control the newly created player. A valude of PLATFRMUSERID_NONE specifies to use the next available ID	
+	 * @param bSpawnPlayerController   Whether a player controller should be spawned immediately for this player. If false a player controller will not be created automatically until transition to the next map.
+	 * @return                         The created player controller if one is created.
+	 */
+	UFUNCTION(BlueprintCallable, Category = "Game", meta = (AutoCreateRefTerm="UserId", DisplayName="Create Local Player For Platform User", WorldContext = "WorldContextObject", AdvancedDisplay = "2", UnsafeDuringActorConstruction = "true"))
+	static class APlayerController* CreatePlayerFromPlatformUser(const UObject* WorldContextObject, FPlatformUserId UserId, bool bSpawnPlayerController = true);
 
 	/** 
 	 * Removes a local player from this game.
@@ -253,6 +271,15 @@ class ENGINE_API UGameplayStatics : public UBlueprintFunctionLibrary
 	UFUNCTION(BlueprintPure, Category="Game", meta=(DisplayName="Get Local Player Controller ID", UnsafeDuringActorConstruction="true"))
 	static int32 GetPlayerControllerID(APlayerController* Player);
 
+	/**
+	 * Gets what Platform User Id a player is using. This only works for local player controllers.
+	 *
+	 * @param Player	The player controller of the player to get the ID of
+	 * @return			The Platform User Id of the passed in player. PLATFORMUSERID_NONE if there is no physical controller assigned to the passed in player
+	 */
+	UFUNCTION(BlueprintPure, Category="Game", meta=(DisplayName="Get Local Player Platform User Id", UnsafeDuringActorConstruction="true"))
+	static FPlatformUserId GetPlayerPlatformUserId(const APlayerController* PlayerController);
+
 	/**
 	 * Sets what physical controller ID a player should be using. This only works for local player controllers.
 	 *
@@ -261,6 +288,15 @@ class ENGINE_API UGameplayStatics : public UBlueprintFunctionLibrary
 	 */
 	UFUNCTION(BlueprintCallable, Category="Game", meta=(DisplayName="Set Local Player Controller ID",UnsafeDuringActorConstruction="true"))
 	static void SetPlayerControllerID(APlayerController* Player, int32 ControllerId);
+	
+	/**
+	 * Sets what platform user id a player should be using. This only works for local player controllers.
+	 *
+	 * @param Player			The player controller of the player to change the platform user of
+	 * @param PlatformUserId	The platform user id to assign this player
+	 */
+	UFUNCTION(BlueprintCallable, Category="Game", meta=(DisplayName="Set Local Player Controller Platform User Id",UnsafeDuringActorConstruction="true"))
+	static void SetPlayerPlatformUserId(APlayerController* PlayerController, FPlatformUserId UserId);
 
 	// --- Level Streaming functions ------------------------
 	
diff --git a/Engine/Source/Runtime/Engine/Private/GameplayStatics.cpp b/Engine/Source/Runtime/Engine/Private/GameplayStatics.cpp
--- a/Engine/Source/Runtime/Engine/Private/GameplayStatics.cpp
+++ b/Engine/Source/Runtime/Engine/Private/GameplayStatics.cpp
@@ -362,15 +362,20 @@ class APlayerController* UGameplayStatics::GetPlayerController(const UObject* Wo
 	return nullptr;
 }
 
-class APlayerController* UGameplayStatics::GetPlayerControllerFromID(const UObject* WorldContextObject, int32 ControllerID)
+APlayerController* UGameplayStatics::GetPlayerControllerFromID(const UObject* WorldContextObject, int32 ControllerID)
+{
+	return GetPlayerControllerFromPlatformUser(WorldContextObject, FGenericPlatformMisc::GetPlatformUserForUserIndex(ControllerID));
+}
+
+APlayerController* UGameplayStatics::GetPlayerControllerFromPlatformUser(const UObject* WorldContextObject, FPlatformUserId UserId)
 {
 	if (UWorld* World = GEngine->GetWorldFromContextObject(WorldContextObject, EGetWorldErrorMode::LogAndReturnNull))
 	{
 		for (FConstPlayerControllerIterator Iterator = World->GetPlayerControllerIterator(); Iterator; ++Iterator)
 		{
 			APlayerController* PlayerController = Iterator->Get();
-			int32 PlayerControllerID = GetPlayerControllerID(PlayerController);
-			if (PlayerControllerID != INDEX_NONE && PlayerControllerID == ControllerID)
+			FPlatformUserId PlayerControllerUserID = GetPlayerPlatformUserId(PlayerController);
+			if (PlayerControllerUserID.IsValid() && PlayerControllerUserID == UserId)
 			{
 				return PlayerController;
 			}
@@ -398,11 +403,16 @@ APlayerCameraManager* UGameplayStatics::GetPlayerCameraManager(const UObject* Wo
 }
 
 APlayerController* UGameplayStatics::CreatePlayer(const UObject* WorldContextObject, int32 ControllerId, bool bSpawnPlayerController)
+{
+	return CreatePlayerFromPlatformUser(WorldContextObject, FGenericPlatformMisc::GetPlatformUserForUserIndex(ControllerId), bSpawnPlayerController);
+}
+
+APlayerController* UGameplayStatics::CreatePlayerFromPlatformUser(const UObject* WorldContextObject, FPlatformUserId UserId, bool bSpawnPlayerController)
 {
 	UWorld* World = GEngine->GetWorldFromContextObject(WorldContextObject, EGetWorldErrorMode::LogAndReturnNull);
 	FString Error;
 
-	ULocalPlayer* LocalPlayer = World ? World->GetGameInstance()->CreateLocalPlayer(ControllerId, Error, bSpawnPlayerController) : nullptr;
+	ULocalPlayer* LocalPlayer = World ? World->GetGameInstance()->CreateLocalPlayer(UserId, Error, bSpawnPlayerController) : nullptr;
 
 	if (Error.Len() > 0)
 	{
@@ -431,25 +441,36 @@ void UGameplayStatics::RemovePlayer(APlayerController* PlayerController, bool bD
 }
 
 int32 UGameplayStatics::GetPlayerControllerID(APlayerController* PlayerController)
+{
+	FPlatformUserId UserID = GetPlayerPlatformUserId(PlayerController);
+	return FGenericPlatformMisc::GetUserIndexForPlatformUser(UserID);
+}
+
+FPlatformUserId UGameplayStatics::GetPlayerPlatformUserId(const APlayerController* PlayerController)
 {
 	if (PlayerController)
 	{
-		if (ULocalPlayer* LocalPlayer = PlayerController->GetLocalPlayer())
+		if (const ULocalPlayer* LocalPlayer = PlayerController->GetLocalPlayer())
 		{
-			return LocalPlayer->GetControllerId();
+			return LocalPlayer->GetPlatformUserId();
 		}
 	}
 
-	return INDEX_NONE;
+	return PLATFORMUSERID_NONE;
 }
 
 void UGameplayStatics::SetPlayerControllerID(APlayerController* PlayerController, int32 ControllerId)
+{
+	SetPlayerPlatformUserId(PlayerController, FPlatformMisc::GetPlatformUserForUserIndex(ControllerId));
+}
+
+void UGameplayStatics::SetPlayerPlatformUserId(APlayerController* PlayerController, FPlatformUserId UserId)
 {
 	if (PlayerController)
 	{
 		if (ULocalPlayer* LocalPlayer = PlayerController->GetLocalPlayer())
 		{
-			LocalPlayer->SetControllerId(ControllerId);
+			LocalPlayer->SetPlatformUserId(UserId);
 		}
 	}
 }
