World partition : Fixed regression where Actors not relevant for runtime were still added.
#rnx

diff --git a/Engine/Source/Runtime/Engine/Private/WorldPartition/WorldPartitionStreamingGeneration.cpp b/Engine/Source/Runtime/Engine/Private/WorldPartition/WorldPartitionStreamingGeneration.cpp
--- a/Engine/Source/Runtime/Engine/Private/WorldPartition/WorldPartitionStreamingGeneration.cpp
+++ b/Engine/Source/Runtime/Engine/Private/WorldPartition/WorldPartitionStreamingGeneration.cpp
@@ -134,7 +134,7 @@ class FWorldPartitionStreamingGenerator
 					return ActorDesc->GetActorIsEditorOnly();
 				}
 			}
-			return false;
+			return true;
 		};
 
 		// Create an actor descriptor view for the specified actor (modified or unsaved actors)
