Override MaterialInterface rename to keep the editor only object's name in sync
Fix up the MaterialInterface EdiorOnlyData pointer on load of the editor only data object if not correctly set

diff --git a/Engine/Source/Runtime/Engine/Classes/Materials/MaterialFunctionInterface.h b/Engine/Source/Runtime/Engine/Classes/Materials/MaterialFunctionInterface.h
--- a/Engine/Source/Runtime/Engine/Classes/Materials/MaterialFunctionInterface.h
+++ b/Engine/Source/Runtime/Engine/Classes/Materials/MaterialFunctionInterface.h
@@ -67,6 +67,7 @@ public:
 #if WITH_EDITOR
 	virtual void PostEditChangeProperty(FPropertyChangedEvent& PropertyChangedEvent) override;
 #endif
+	virtual bool Rename(const TCHAR* NewName = nullptr, UObject* NewOuter = nullptr, ERenameFlags Flags = REN_None) override;
 	//~ End UObject Interface.
 
 	/** Used by materials using this function to know when to recompile. */
diff --git a/Engine/Source/Runtime/Engine/Classes/Materials/MaterialInterface.h b/Engine/Source/Runtime/Engine/Classes/Materials/MaterialInterface.h
--- a/Engine/Source/Runtime/Engine/Classes/Materials/MaterialInterface.h
+++ b/Engine/Source/Runtime/Engine/Classes/Materials/MaterialInterface.h
@@ -219,6 +219,8 @@ class UMaterialInterface : public UObject, public IBlendableInterface, public II
 
 #if WITH_EDITORONLY_DATA
 protected:
+	friend class UMaterialInterfaceEditorOnlyData;
+
 	UPROPERTY()
 	TObjectPtr<UMaterialInterfaceEditorOnlyData> EditorOnlyData;
 
@@ -328,6 +330,7 @@ public:
 	ENGINE_API virtual void PostLoad() override;
 	ENGINE_API virtual void PostDuplicate(bool bDuplicateForPIE) override;
 	ENGINE_API virtual void PostCDOContruct() override;
+	ENGINE_API virtual bool Rename(const TCHAR* NewName = nullptr, UObject* NewOuter = nullptr, ERenameFlags Flags = REN_None) override;
 	ENGINE_API static void AddReferencedObjects(UObject* InThis, FReferenceCollector& Collector);
 
 #if WITH_EDITOR
diff --git a/Engine/Source/Runtime/Engine/Private/Materials/MaterialExpressions.cpp b/Engine/Source/Runtime/Engine/Private/Materials/MaterialExpressions.cpp
--- a/Engine/Source/Runtime/Engine/Private/Materials/MaterialExpressions.cpp
+++ b/Engine/Source/Runtime/Engine/Private/Materials/MaterialExpressions.cpp
@@ -13263,6 +13263,28 @@ void UMaterialFunctionInterface::GetAssetRegistryTags(TArray<FAssetRegistryTag>&
 #endif
 }
 
+namespace
+{
+	FString GetEditorOnlyDataName(const TCHAR* InMaterialName)
+	{
+		return FString::Printf(TEXT("%sEditorOnlyData"), InMaterialName);
+	}
+}
+
+bool UMaterialFunctionInterface::Rename(const TCHAR* NewName, UObject* NewOuter, ERenameFlags Flags)
+{
+	bool bRenamed = Super::Rename(NewName, NewOuter, Flags);
+#if WITH_EDITORONLY_DATA
+	// if we have EditorOnlyData, also rename it if we are changing the material's name
+	if (bRenamed && NewName && EditorOnlyData)
+	{
+		FString EditorOnlyDataName = GetEditorOnlyDataName(NewName);
+		bRenamed = EditorOnlyData->Rename(*EditorOnlyDataName, nullptr, Flags);
+	}
+#endif
+	return bRenamed;
+}
+
 UMaterialFunctionInterface* UMaterialFunctionInterface::GetBaseFunctionInterface()
 {
 	return GetBaseFunction();
@@ -13478,7 +13500,7 @@ UMaterialFunctionInterfaceEditorOnlyData* UMaterialFunctionInterface::CreateEdit
 	check(EditorOnlyClass);
 	check(EditorOnlyClass->HasAllClassFlags(CLASS_Optional));
 
-	const FString EditorOnlyName = FString::Printf(TEXT("%sEditorOnlyData"), *GetName());
+	const FString EditorOnlyName = GetEditorOnlyDataName(*GetName());
 	const EObjectFlags EditorOnlyFlags = GetMaskedFlags(RF_PropagateToSubObjects);
 	return NewObject<UMaterialFunctionInterfaceEditorOnlyData>(this, EditorOnlyClass, *EditorOnlyName, EditorOnlyFlags);
 }
diff --git a/Engine/Source/Runtime/Engine/Private/Materials/MaterialInterface.cpp b/Engine/Source/Runtime/Engine/Private/Materials/MaterialInterface.cpp
--- a/Engine/Source/Runtime/Engine/Private/Materials/MaterialInterface.cpp
+++ b/Engine/Source/Runtime/Engine/Private/Materials/MaterialInterface.cpp
@@ -1363,6 +1363,15 @@ void UMaterialInterface::RemoveUserDataOfClass(TSubclassOf<UAssetUserData> InUse
 }
 
 #if WITH_EDITORONLY_DATA
+
+namespace
+{
+	FString GetEditorOnlyDataName(const TCHAR* InMaterialName)
+	{
+		return FString::Printf(TEXT("%sEditorOnlyData"), InMaterialName);
+	}
+}
+
 const UClass* UMaterialInterface::GetEditorOnlyDataClass() const
 {
 	return UMaterialInterfaceEditorOnlyData::StaticClass();
@@ -1376,12 +1385,26 @@ UMaterialInterfaceEditorOnlyData* UMaterialInterface::CreateEditorOnlyData()
 	check(EditorOnlyClass);
 	check(EditorOnlyClass->HasAllClassFlags(CLASS_Optional));
 
-	const FString EditorOnlyName = FString::Printf(TEXT("%sEditorOnlyData"), *GetName());
+	const FString EditorOnlyName = GetEditorOnlyDataName(*GetName());
 	const EObjectFlags EditorOnlyFlags = GetMaskedFlags(RF_PropagateToSubObjects);
 	return NewObject<UMaterialInterfaceEditorOnlyData>(this, EditorOnlyClass, *EditorOnlyName, EditorOnlyFlags);
 }
 #endif // WITH_EDITORONLY_DATA
 
+bool UMaterialInterface::Rename(const TCHAR* NewName, UObject* NewOuter, ERenameFlags Flags)
+{
+	bool bRenamed = Super::Rename(NewName, NewOuter, Flags);
+#if WITH_EDITORONLY_DATA
+	// if we have EditorOnlyData, also rename it if we are changing the material's name
+	if (bRenamed && NewName && EditorOnlyData)
+	{
+		FString EditorOnlyDataName = GetEditorOnlyDataName(NewName);
+		bRenamed = EditorOnlyData->Rename(*EditorOnlyDataName, nullptr, Flags);
+	}
+#endif
+	return bRenamed;
+}
+
 UMaterialInterfaceEditorOnlyData::UMaterialInterfaceEditorOnlyData()
 {
 }
@@ -1408,6 +1431,17 @@ void UMaterialInterfaceEditorOnlyData::Serialize(FArchive& Ar)
 
 	if (bSavedCachedExpressionData)
 	{
+		// if we do not have our CachedExpressionData set at this point, 
+		// it means this object's name doesn't match the default created object name and we need to fix our pointer into the material interface
+		if (CachedExpressionData == nullptr)
+		{
+#if WITH_EDITORONLY_DATA
+			UMaterialInterface* Owner = Cast<UMaterialInterface>(GetOuter());
+			Owner->EditorOnlyData = this;
+			CachedExpressionData = Owner->CachedExpressionData->EditorOnlyData;
+#endif
+		}
+
 		check(CachedExpressionData);
 		UScriptStruct* Struct = FMaterialCachedExpressionEditorOnlyData::StaticStruct();
 		Struct->SerializeTaggedProperties(Ar, (uint8*)CachedExpressionData.Get(), Struct, nullptr);
