A few added DrawDebgCanvas functions for generating boxes, half circles and cylinders.

diff --git a/Engine/Source/Runtime/Engine/Private/DrawDebugHelpers.cpp b/Engine/Source/Runtime/Engine/Private/DrawDebugHelpers.cpp
--- a/Engine/Source/Runtime/Engine/Private/DrawDebugHelpers.cpp
+++ b/Engine/Source/Runtime/Engine/Private/DrawDebugHelpers.cpp
@@ -1113,6 +1113,19 @@ void DrawDebugCanvasCircle(UCanvas* Canvas, const FVector& Base, const FVector&
 	}
 }
 
+void DrawDebugCanvasHalfCircle(UCanvas* Canvas, const FVector& Base, const FVector& X, const FVector& Y, FColor Color, float Radius, int32 NumSides)
+{
+	const float	AngleDelta = 2.0f * UE_PI / NumSides;
+	FVector	LastVertex = Base + X * Radius;
+
+	for (int32 SideIndex = 0;SideIndex < NumSides / 2;SideIndex++)
+	{
+		const FVector Vertex = Base + (X * FMath::Cos(AngleDelta * (SideIndex + 1)) + Y * FMath::Sin(AngleDelta * (SideIndex + 1))) * Radius;
+		DrawDebugCanvasLine(Canvas, LastVertex, Vertex, Color);
+		LastVertex = Vertex;
+	}
+}
+
 void DrawDebugCanvasWireSphere(UCanvas* Canvas, const FVector& Base, FColor Color, float Radius, int32 NumSides)
 {
 	DrawDebugCanvasCircle(Canvas, Base, FVector(1,0,0), FVector(0,1,0), Color, Radius, NumSides);
@@ -1163,6 +1176,66 @@ void DrawDebugCanvasWireCone(UCanvas* Canvas, const FTransform& Transform, float
 	DrawDebugCanvasLine( Canvas, Verts[Verts.Num()-1], Verts[0], Color );
 }
 
+void DrawDebugCanvasWireBox(UCanvas* Canvas, const FMatrix& Transform, const FBox& Box, FColor Color)
+{
+	const FVector Vertices[] =
+	{
+		FVector(Box.Min.X, Box.Min.Y, Box.Min.Z),
+		FVector(Box.Min.X, Box.Min.Y, Box.Max.Z),
+		FVector(Box.Min.X, Box.Max.Y, Box.Min.Z),
+		FVector(Box.Min.X, Box.Max.Y, Box.Max.Z),
+		FVector(Box.Max.X, Box.Min.Y, Box.Min.Z),
+		FVector(Box.Max.X, Box.Min.Y, Box.Max.Z),
+		FVector(Box.Max.X, Box.Max.Y, Box.Min.Z),
+		FVector(Box.Max.X, Box.Max.Y, Box.Max.Z)
+	};
+
+	const FIntVector2 Edges[] =
+	{
+		{ 0, 1 }, { 2, 3 },	{ 4, 5 }, { 6, 7 },
+		{ 0, 4 }, { 4, 6 },	{ 6, 2 }, { 2, 0 },
+		{ 1, 5 }, { 5, 7 },	{ 7, 3 }, { 3, 1 }
+	};
+
+	for (const FIntVector2& Edge : Edges)
+	{
+		DrawDebugCanvasLine(Canvas, Vertices[Edge.X], Vertices[Edge.Y], Color);
+	}
+}
+
+void DrawDebugCanvasCapsule(UCanvas* Canvas, const FMatrix& Transform, float HalfLength, float Radius, const FColor& LineColor)
+{
+	const int32 DrawCollisionSides = 16;
+
+	FVector Origin = Transform.GetOrigin();
+	FVector XAxis = Transform.GetScaledAxis(EAxis::X);
+	FVector YAxis = Transform.GetScaledAxis(EAxis::Y);
+	FVector ZAxis = Transform.GetScaledAxis(EAxis::Z);
+
+	// Draw top and bottom circles
+	float HalfAxis = FMath::Max<float>(HalfLength - Radius, 1.f);
+	FVector TopEnd = Origin + HalfAxis * ZAxis;
+	FVector BottomEnd = Origin - HalfAxis * ZAxis;
+
+	DrawDebugCanvasCircle(Canvas, TopEnd, XAxis, YAxis, LineColor, Radius, DrawCollisionSides);
+	DrawDebugCanvasCircle(Canvas, BottomEnd, XAxis, YAxis, LineColor, Radius, DrawCollisionSides);
+
+	// Draw domed caps
+	DrawDebugCanvasHalfCircle(Canvas, TopEnd, YAxis, ZAxis, LineColor, Radius, DrawCollisionSides);
+	DrawDebugCanvasHalfCircle(Canvas, TopEnd, XAxis, ZAxis, LineColor, Radius, DrawCollisionSides);
+
+	FVector NegZAxis = -ZAxis;
+
+	DrawDebugCanvasHalfCircle(Canvas, BottomEnd, YAxis, NegZAxis, LineColor, Radius, DrawCollisionSides);
+	DrawDebugCanvasHalfCircle(Canvas, BottomEnd, XAxis, NegZAxis, LineColor, Radius, DrawCollisionSides);
+
+	// Draw connected lines
+	DrawDebugCanvasLine(Canvas, TopEnd + Radius * XAxis, BottomEnd + Radius * XAxis, LineColor);
+	DrawDebugCanvasLine(Canvas, TopEnd - Radius * XAxis, BottomEnd - Radius * XAxis, LineColor);
+	DrawDebugCanvasLine(Canvas, TopEnd + Radius * YAxis, BottomEnd + Radius * YAxis, LineColor);
+	DrawDebugCanvasLine(Canvas, TopEnd - Radius * YAxis, BottomEnd - Radius * YAxis, LineColor);
+}
+
 //
 // Canvas 2D
 //
diff --git a/Engine/Source/Runtime/Engine/Public/DrawDebugHelpers.h b/Engine/Source/Runtime/Engine/Public/DrawDebugHelpers.h
--- a/Engine/Source/Runtime/Engine/Public/DrawDebugHelpers.h
+++ b/Engine/Source/Runtime/Engine/Public/DrawDebugHelpers.h
@@ -115,6 +115,19 @@ ENGINE_API void DrawDebugCanvasLine(UCanvas* Canvas, const FVector& Start, const
  */
 ENGINE_API void DrawDebugCanvasCircle(UCanvas* Canvas, const FVector& Base, const FVector& X, const FVector& Y, FColor Color, float Radius, int32 NumSides);
 
+/**
+ * Draws a half circle using lines.
+ *
+ * @param	Canvas			The Canvas to draw on.
+ * @param	Base			Center of the circle.
+ * @param	X				X alignment axis to draw along.
+ * @param	Y				Y alignment axis to draw along.
+ * @param	Color			Color of the circle.
+ * @param	Radius			Radius of the circle.
+ * @param	NumSides		Numbers of sides that the circle has.
+ */
+ENGINE_API void DrawDebugCanvasHalfCircle(UCanvas* Canvas, const FVector& Base, const FVector& X, const FVector& Y, FColor Color, float Radius, int32 NumSides);
+
 /**
  * Draws a sphere using circles.
  *
@@ -138,6 +151,27 @@ ENGINE_API void DrawDebugCanvasWireSphere(UCanvas* Canvas, const FVector& Base,
  */
 ENGINE_API void DrawDebugCanvasWireCone(UCanvas* Canvas, const FTransform& Transform, float ConeRadius, float ConeAngle, int32 ConeSides, FColor Color);
 
+/**
+ * Draws a box using lines.
+ *
+ * @param	Canvas			The Canvas to draw on.
+ * @param	Transform		Transform to apply to the capsule.
+ * @param	Box				The box to be drawn.
+ * @param	Color			Color of the capsule.
+ */
+ENGINE_API void DrawDebugCanvasWireBox(UCanvas* Canvas, const FMatrix& Transform, const FBox& Box, FColor Color);
+
+/**
+ * Draws a capsule using lines.
+ *
+ * @param	Canvas			The Canvas to draw on.
+ * @param	Transform		Transform to apply to the capsule.
+ * @param	HalfLength		Distance from the center of the capsule to the tip (including the cap)
+ * @param	Radius			Radius of the capsule.
+ * @param	Color			Color of the capsule.
+ */
+ENGINE_API void DrawDebugCanvasCapsule(UCanvas* Canvas, const FMatrix& Transform, float HalfLength, float Radius, const FColor& LineColor);
+
 #elif !defined(SHIPPING_DRAW_DEBUG_ERROR) || !SHIPPING_DRAW_DEBUG_ERROR
 
 // Empty versions of above functions
@@ -178,8 +212,11 @@ FORCEINLINE void DrawDebugCanvas2DCircle(UCanvas* Canvas, const FVector2D& Cente
 FORCEINLINE void DrawDebugCanvas2DBox(UCanvas* Canvas, const FBox2D& Box, const FLinearColor& LineColor, const float& LineThickness = 1.f) {}
 FORCEINLINE void DrawDebugCanvasLine(UCanvas* Canvas, const FVector& Start, const FVector& End, const FLinearColor& LineColor) {}
 FORCEINLINE void DrawDebugCanvasCircle(UCanvas* Canvas, const FVector& Base, const FVector& X, const FVector& Y, FColor Color, float Radius, int32 NumSides) {}
+FORCEINLINE void DrawDebugCanvasHalfCircle(UCanvas* Canvas, const FVector& Base, const FVector& X, const FVector& Y, FColor Color, float Radius, int32 NumSides) {}
 FORCEINLINE void DrawDebugCanvasWireSphere(UCanvas* Canvas, const FVector& Base, FColor Color, float Radius, int32 NumSides) {}
 FORCEINLINE void DrawDebugCanvasWireCone(UCanvas* Canvas, const FTransform& Transform, float ConeRadius, float ConeAngle, int32 ConeSides, FColor Color) {}
+FORCEINLINE void DrawDebugCanvasWireBox(UCanvas* Canvas, const FMatrix& Transform, const FBox& Box, FColor Color) {}
+FORCEINLINE void DrawDebugCanvasCapsule(UCanvas* Canvas, const FMatrix& Transform, float HalfLength, float Radius, const FColor& LineColor) {}
 
 #else
 
