Add "GarbageObjectsToVerify" to world context to verify that important objects are cleaned up on world transitions, such as the net driver.

diff --git a/Engine/Source/Editor/UnrealEd/Classes/Editor/EditorEngine.h b/Engine/Source/Editor/UnrealEd/Classes/Editor/EditorEngine.h
--- a/Engine/Source/Editor/UnrealEd/Classes/Editor/EditorEngine.h
+++ b/Engine/Source/Editor/UnrealEd/Classes/Editor/EditorEngine.h
@@ -784,7 +784,7 @@ private:
 
 	virtual bool AreEditorAnalyticsEnabled() const override;
 	virtual void CreateStartupAnalyticsAttributes(TArray<FAnalyticsEventAttribute>& StartSessionAttributes) const override;
-	virtual void VerifyLoadMapWorldCleanup() override;
+	virtual void VerifyLoadMapWorldCleanup(FWorldContext* ForContext) override;
 
 	/** Called during editor init and whenever the vanilla status might have changed, to set the flag on the base class */
 	void UpdateIsVanillaProduct();
diff --git a/Engine/Source/Editor/UnrealEd/Private/EditorEngine.cpp b/Engine/Source/Editor/UnrealEd/Private/EditorEngine.cpp
--- a/Engine/Source/Editor/UnrealEd/Private/EditorEngine.cpp
+++ b/Engine/Source/Editor/UnrealEd/Private/EditorEngine.cpp
@@ -7043,7 +7043,7 @@ bool UEditorEngine::NetworkRemapPath( UPendingNetGame *PendingNetGame, FString&
 	return NetworkRemapPath_local(Context, Str, bReading, PendingNetGame->GetDemoNetDriver() != nullptr);
 }
 
-void UEditorEngine::VerifyLoadMapWorldCleanup()
+void UEditorEngine::VerifyLoadMapWorldCleanup(FWorldContext* ForContext)
 {
 	// This does the same as UEngine::VerifyLoadMapWorldCleanup except it also allows Editor Worlds as a valid world.
 
@@ -7093,6 +7093,21 @@ void UEditorEngine::VerifyLoadMapWorldCleanup()
 			}
 		}
 	}
+	
+	if (ForContext)
+	{
+		for (FObjectKey Key : ForContext->GarbageObjectsToVerify)
+		{
+			if (UObject* Object = Key.ResolveObjectPtrEvenIfPendingKill())
+			{
+				UE_LOG(LogLoad, Error, TEXT("Object %s not cleaned up by garbage collection!"), *Object->GetPathName());
+			
+				FindAndPrintStaleReferencesToObject(Object, 
+					UObjectBaseUtility::IsPendingKillEnabled() ? EPrintStaleReferencesOptions::Fatal : (EPrintStaleReferencesOptions::Error | EPrintStaleReferencesOptions::Ensure));
+			}
+		}
+		ForContext->GarbageObjectsToVerify.Reset();
+	}
 }
 
 void UEditorEngine::UpdateIsVanillaProduct()
diff --git a/Engine/Source/Runtime/Engine/Classes/Engine/Engine.h b/Engine/Source/Runtime/Engine/Classes/Engine/Engine.h
--- a/Engine/Source/Runtime/Engine/Classes/Engine/Engine.h
+++ b/Engine/Source/Runtime/Engine/Classes/Engine/Engine.h
@@ -382,6 +382,9 @@ struct FWorldContext
 	float PIEFixedTickSeconds  = 0.f;
 	float PIEAccumulatedTickSeconds = 0.f;
 
+	/** On a transition to another level (e.g. LoadMap), the engine will verify that these objects have been cleaned up by garbage collection */
+	TSet<FObjectKey> GarbageObjectsToVerify;
+
 	/**************************************************************/
 
 	/** Outside pointers to CurrentWorld that should be kept in sync if current world changes  */
@@ -3182,8 +3185,8 @@ public:
 	 */
 	UWorld* GetCurrentPlayWorld(UWorld* PossiblePlayWorld = nullptr) const;
 
-	/** Verify any remaining World(s) are valid after ::LoadMap destroys a world */
-	virtual void VerifyLoadMapWorldCleanup();
+	/** Verify any remaining World(s) are valid after ::LoadMap destroys a world. May also check that "must-destroy" objects associated with the world context have been destroyed. */
+	virtual void VerifyLoadMapWorldCleanup(FWorldContext* ForWorldContext = nullptr);
 
 	/**
 	 * Attempts to find what is referencing a world that should have been garbage collected
diff --git a/Engine/Source/Runtime/Engine/Private/UnrealEngine.cpp b/Engine/Source/Runtime/Engine/Private/UnrealEngine.cpp
--- a/Engine/Source/Runtime/Engine/Private/UnrealEngine.cpp
+++ b/Engine/Source/Runtime/Engine/Private/UnrealEngine.cpp
@@ -13140,6 +13140,7 @@ void DestroyNamedNetDriver_Local(FWorldContext &Context, FName NetDriverName)
 			ensureMsgf(!(NetDriver->IsInTick() && NetDriverName != NAME_PendingNetDriver), TEXT("Attempting to destroy NetDriver %s [%s] while it is ticking."), *NetDriver->GetName(), *NetDriverName.ToString());
 
 			Context.ActiveNetDrivers.RemoveAtSwap(Index);
+			Context.GarbageObjectsToVerify.Add(NetDriver);
 			NetDriver->SetWorld(NULL);
 			NetDriver->Shutdown();
 			NetDriver->LowLevelDestroy();
@@ -14280,7 +14281,7 @@ bool UEngine::LoadMap( FWorldContext& WorldContext, FURL URL, class UPendingNetG
 	if (GDelayTrimMemoryDuringMapLoadMode == 0)
 	{
 		// Dump info
-		VerifyLoadMapWorldCleanup();
+		VerifyLoadMapWorldCleanup(&WorldContext);
 	}
 
 	FMoviePlayerProxy::BlockingTick();
@@ -15298,12 +15299,12 @@ bool UEngine::IsWorldDuplicate(const UWorld* const InWorld)
 	return false;
 }
 
-void UEngine::VerifyLoadMapWorldCleanup()
+void UEngine::VerifyLoadMapWorldCleanup(FWorldContext* ForWorldContext)
 {
 	TRACE_CPUPROFILER_EVENT_SCOPE(VerifyLoadMapWorldCleanup);
 	// All worlds at this point should be the CurrentWorld of some context, preview worlds, or streaming level
 	// worlds that are owned by the CurrentWorld of some context.
-	TArray<UObject*> LeakedWorlds;
+	TArray<UObject*> LeakedObjects;
 	for( TObjectIterator<UWorld> It; It; ++It )
 	{
 		UWorld* World = *It;
@@ -15312,17 +15313,29 @@ void UEngine::VerifyLoadMapWorldCleanup()
 		{
 			if ((World->PersistentLevel == nullptr || !WorldHasValidContext(World->PersistentLevel->OwningWorld)) && !IsWorldDuplicate(World))
 			{
-				LeakedWorlds.Add(World);
+				LeakedObjects.Add(World);
 			}
 		}
 	}
 
-	if (LeakedWorlds.Num())
+	if (ForWorldContext)
 	{
-		UE_LOG(LogLoad, Error, TEXT("Some previously active worlds were not cleaned up by garbage collection!"));
+		for (FObjectKey Key : ForWorldContext->GarbageObjectsToVerify)
+		{
+			if (UObject* Obj = Key.ResolveObjectPtrEvenIfPendingKill())
+			{
+				LeakedObjects.Add(Obj);
+			}
+		}
+		ForWorldContext->GarbageObjectsToVerify.Reset();
+	}
+
+	if (LeakedObjects.Num())
+	{
+		UE_LOG(LogLoad, Error, TEXT("Some previously active worlds or related objects were not cleaned up by garbage collection!"));
 		UE_LOG(LogLoad, Error, TEXT("Once a world has become active, it cannot be reused and must be destroyed and reloaded. Dumping reference chains:"));
 
-		FindAndPrintStaleReferencesToObjects(LeakedWorlds, UObjectBaseUtility::IsPendingKillEnabled() ? EPrintStaleReferencesOptions::Fatal : (EPrintStaleReferencesOptions::Error | EPrintStaleReferencesOptions::Ensure));
+		FindAndPrintStaleReferencesToObjects(LeakedObjects, UObjectBaseUtility::IsPendingKillEnabled() ? EPrintStaleReferencesOptions::Fatal : (EPrintStaleReferencesOptions::Error | EPrintStaleReferencesOptions::Ensure));
 
 		if (!UObjectBaseUtility::IsPendingKillEnabled())
 		{
diff --git a/Engine/Source/Runtime/Engine/Private/World.cpp b/Engine/Source/Runtime/Engine/Private/World.cpp
--- a/Engine/Source/Runtime/Engine/Private/World.cpp
+++ b/Engine/Source/Runtime/Engine/Private/World.cpp
@@ -7245,7 +7245,7 @@ UWorld* FSeamlessTravelHandler::Tick()
 
 #if !(UE_BUILD_SHIPPING || UE_BUILD_TEST)
 			// verify that we successfully cleaned up the old world
-			GEngine->VerifyLoadMapWorldCleanup();
+			GEngine->VerifyLoadMapWorldCleanup(&CurrentContext);
 #endif
 			// Clean out NetDriver's Packagemaps, since they may have a lot of NULL object ptrs rotting in the lookup maps.
 			if (NetDriver)
