World Partition: added option to disallow enabling streaming on partitioned worlds.

#rnx

diff --git a/Engine/Source/Runtime/Engine/Private/WorldPartition/WorldPartition.cpp b/Engine/Source/Runtime/Engine/Private/WorldPartition/WorldPartition.cpp
--- a/Engine/Source/Runtime/Engine/Private/WorldPartition/WorldPartition.cpp
+++ b/Engine/Source/Runtime/Engine/Private/WorldPartition/WorldPartition.cpp
@@ -264,6 +264,7 @@ UWorldPartition::UWorldPartition(const FObjectInitializer& ObjectInitializer)
 	, bStreamingWasEnabled(true)
 	, bShouldEnableStreamingWarned(false)
 	, bShouldCheckEnableStreamingWarning(false)
+	, bSupportsStreaming(true)
 	, bCanBeUsedByLevelInstance(false)
 	, bForceGarbageCollection(false)
 	, bForceGarbageCollectionPurge(false)
@@ -961,7 +962,7 @@ const FTransform& UWorldPartition::GetInstanceTransform() const
 #if WITH_EDITOR
 bool UWorldPartition::SupportsStreaming() const
 {
-	return World ? World->GetWorldSettings()->SupportsWorldPartitionStreaming() : false;
+	return bSupportsStreaming && (World ? World->GetWorldSettings()->SupportsWorldPartitionStreaming() : false);
 }
 
 bool UWorldPartition::IsStreamingEnabled() const
@@ -1460,4 +1461,4 @@ FBox UWorldPartition::GetEditorWorldBounds() const
 }
 #endif
 
-#undef LOCTEXT_NAMESPACE
\ No newline at end of file
+#undef LOCTEXT_NAMESPACE
diff --git a/Engine/Source/Runtime/Engine/Private/WorldPartition/WorldPartitionConverter.cpp b/Engine/Source/Runtime/Engine/Private/WorldPartition/WorldPartitionConverter.cpp
--- a/Engine/Source/Runtime/Engine/Private/WorldPartition/WorldPartitionConverter.cpp
+++ b/Engine/Source/Runtime/Engine/Private/WorldPartition/WorldPartitionConverter.cpp
@@ -60,6 +60,7 @@ bool FWorldPartitionConverter::Convert()
 		FLevelActorFoldersHelper::SetUseActorFolders(World->PersistentLevel, Parameters.bUseActorFolders);
 		WorldPartition->bEnableStreaming = Parameters.bEnableStreaming;
 		WorldPartition->bStreamingWasEnabled = Parameters.bEnableStreaming;
+		WorldPartition->bSupportsStreaming = Parameters.bSupportsStreaming;
 		WorldPartition->SetCanBeUsedByLevelInstance(!Parameters.bEnableStreaming);
 		bCreatedWorldPartition = true;
 	}
diff --git a/Engine/Source/Runtime/Engine/Public/WorldPartition/WorldPartition.h b/Engine/Source/Runtime/Engine/Public/WorldPartition/WorldPartition.h
--- a/Engine/Source/Runtime/Engine/Public/WorldPartition/WorldPartition.h
+++ b/Engine/Source/Runtime/Engine/Public/WorldPartition/WorldPartition.h
@@ -234,6 +234,10 @@ private:
 	/** Used to know if we need to recheck if the user should enable streaming based on world size. */
 	bool bShouldCheckEnableStreamingWarning;
 
+	/** Used to check if the current world supports enabling streaming. */
+	UPROPERTY()
+	bool bSupportsStreaming;
+
 	/** Whether Level Instance can reference this partition. */
 	UPROPERTY(EditAnywhere, Category = "WorldPartition", AdvancedDisplay, meta = (EditConditionHides, HideEditConditionToggle, EditCondition = "!bEnableStreaming"))
 	bool bCanBeUsedByLevelInstance;
diff --git a/Engine/Source/Runtime/Engine/Public/WorldPartition/WorldPartitionConverter.h b/Engine/Source/Runtime/Engine/Public/WorldPartition/WorldPartitionConverter.h
--- a/Engine/Source/Runtime/Engine/Public/WorldPartition/WorldPartitionConverter.h
+++ b/Engine/Source/Runtime/Engine/Public/WorldPartition/WorldPartitionConverter.h
@@ -21,6 +21,7 @@ public:
 		FParameters()
 		: bConvertSubLevels(true)
 		, bEnableStreaming(true)
+		, bSupportsStreaming(true)
 		, bUseActorFolders(true)
 		, EditorHashClass(nullptr)
 		, RuntimeHashClass(nullptr)
@@ -28,6 +29,7 @@ public:
 
 		bool bConvertSubLevels;
 		bool bEnableStreaming;
+		bool bSupportsStreaming;
 		bool bUseActorFolders;
 		TSubclassOf<UWorldPartitionEditorHash> EditorHashClass;
 		TSubclassOf<UWorldPartitionRuntimeHash> RuntimeHashClass;
