World Partition: temp crash fix when initializing a world on editor client.

#rnx

diff --git a/Engine/Source/Runtime/Engine/Private/WorldPartition/WorldPartition.cpp b/Engine/Source/Runtime/Engine/Private/WorldPartition/WorldPartition.cpp
--- a/Engine/Source/Runtime/Engine/Private/WorldPartition/WorldPartition.cpp
+++ b/Engine/Source/Runtime/Engine/Private/WorldPartition/WorldPartition.cpp
@@ -587,6 +587,12 @@ void UWorldPartition::Initialize(UWorld* InWorld, const FTransform& InTransform)
 				bIsPIE = true;
 			}
 
+			if (StreamingPolicy)
+			{
+				UE_LOG(LogWorldPartition, Warning, TEXT("StreamingPolicy was set when initializing the world partition object"));
+				StreamingPolicy = nullptr;
+			}
+
 			OnBeginPlay();
 		}
 
