UActorElementWorldInterface::GetBounds()
* Fixed unchecked access to root component

diff --git a/Engine/Source/Runtime/Engine/Private/Elements/Actor/ActorElementWorldInterface.cpp b/Engine/Source/Runtime/Engine/Private/Elements/Actor/ActorElementWorldInterface.cpp
--- a/Engine/Source/Runtime/Engine/Private/Elements/Actor/ActorElementWorldInterface.cpp
+++ b/Engine/Source/Runtime/Engine/Private/Elements/Actor/ActorElementWorldInterface.cpp
@@ -35,8 +35,11 @@ bool UActorElementWorldInterface::GetBounds(const FTypedElementHandle& InElement
 	if (const AActor* Actor = ActorElementDataUtil::GetActorFromHandle(InElementHandle))
 	{
 		// TODO: This was taken from FActorOrComponent, but AActor has a function to calculate bounds too...
-		OutBounds = Actor->GetRootComponent()->Bounds;
-		return true;
+		if (const USceneComponent* RootComponent = Actor->GetRootComponent())
+		{
+			OutBounds = RootComponent->Bounds;
+			return true;
+		}
 	}
 
 	return false;
