LevelInstance: Avoid accessing asset registry while saving (fix crash)

diff --git a/Engine/Source/Runtime/Engine/Private/LevelInstance/LevelInstanceSubsystem.cpp b/Engine/Source/Runtime/Engine/Private/LevelInstance/LevelInstanceSubsystem.cpp
--- a/Engine/Source/Runtime/Engine/Private/LevelInstance/LevelInstanceSubsystem.cpp
+++ b/Engine/Source/Runtime/Engine/Private/LevelInstance/LevelInstanceSubsystem.cpp
@@ -635,7 +635,14 @@ bool ULevelInstanceSubsystem::GetLevelInstanceBounds(const ILevelInstanceInterfa
 	}
 	else if(LevelInstance->IsWorldAssetValid())
 	{
-		//@todo_ow: can't find a way to trigger this codepath, maybe deadcode?
+		// @todo_ow: remove this temp fix once it is again safe to call the asset registry while saving.
+		// Currently it can lead to a FindObject which is illegal while saving.
+		if (UE::IsSavingPackage(nullptr))
+		{
+			OutBounds = FBox(ForceInit);
+			return true;
+		}
+
 		FString LevelPackage = LevelInstance->GetWorldAssetPackage();
 
 		if (FBox ContainerBounds = ActorDescContainerInstanceManager.GetContainerBounds(*LevelPackage); ContainerBounds.IsValid)
