Fix FSubsystemCollectionBase::Deinitialize crash when RemoveAndDeinitializeSubsystem gets called while iterating on SubsystemMap
Also update SubsystemArrayMap when removing an entry from SubsystemMap

diff --git a/Engine/Source/Runtime/Engine/Private/Subsystems/SubsystemCollection.cpp b/Engine/Source/Runtime/Engine/Private/Subsystems/SubsystemCollection.cpp
--- a/Engine/Source/Runtime/Engine/Private/Subsystems/SubsystemCollection.cpp
+++ b/Engine/Source/Runtime/Engine/Private/Subsystems/SubsystemCollection.cpp
@@ -7,6 +7,7 @@
 #include "UObject/Package.h"
 #include "Modules/ModuleManager.h"
 #include "Misc/PackageName.h"
+#include "Algo/Transform.h"
 
 DEFINE_LOG_CATEGORY(LogSubsystemCollection);
 
@@ -185,17 +186,26 @@ void FSubsystemCollectionBase::Deinitialize()
 
 	// Deinit and clean up existing systems
 	SubsystemArrayMap.Empty();
-	for (auto Iter = SubsystemMap.CreateIterator(); Iter; ++Iter)
+	if (!SubsystemMap.IsEmpty())
 	{
-		UClass* KeyClass = Iter.Key();
-		USubsystem* Subsystem = Iter.Value();
-		if ( Subsystem != nullptr && Subsystem->GetClass() == KeyClass)
+		TMap<UClass*, TWeakObjectPtr<USubsystem>> SubsystemMapCopy;
+		SubsystemMapCopy.Reserve(SubsystemMap.Num());
+		Algo::Transform(SubsystemMap, SubsystemMapCopy, [](const TPair<UClass*, USubsystem*> Iter){ return TPair<UClass*, TWeakObjectPtr<USubsystem>>(Iter.Key, Iter.Value); });
+
+		for (auto Iter = SubsystemMapCopy.CreateIterator(); Iter; ++Iter)
 		{
-			Subsystem->Deinitialize();
-			Subsystem->InternalOwningSubsystem = nullptr;
+			UClass* KeyClass = Iter.Key();
+			USubsystem* Subsystem = Iter.Value().GetEvenIfUnreachable();
+			if (Subsystem && Subsystem->InternalOwningSubsystem && Subsystem->GetClass() == KeyClass)
+			{
+				Subsystem->Deinitialize();
+				Subsystem->InternalOwningSubsystem = nullptr;
+			}
 		}
+
+		SubsystemMap.Empty();
+		SubsystemArrayMap.Empty(); // Clear again in case it gets repopulated during Deinitialize
 	}
-	SubsystemMap.Empty();
 	Outer = nullptr;
 }
 
@@ -276,6 +286,8 @@ void FSubsystemCollectionBase::RemoveAndDeinitializeSubsystem(USubsystem* Subsys
 	USubsystem* SubsystemFound = SubsystemMap.FindAndRemoveChecked(Subsystem->GetClass());
 	check(Subsystem == SubsystemFound);
 
+	SubsystemArrayMap.Remove(Subsystem->GetClass());
+
 	Subsystem->Deinitialize();
 	Subsystem->InternalOwningSubsystem = nullptr;
 }
