Making it so that we use the owner's Desired Rotation if a Character does not have a Controller on the Autonomus Client when bRunPhysicsWithNoController and bUseControllerDesiredRotation are true.

diff --git a/Engine/Source/Runtime/Engine/Private/Components/CharacterMovementComponent.cpp b/Engine/Source/Runtime/Engine/Private/Components/CharacterMovementComponent.cpp
--- a/Engine/Source/Runtime/Engine/Private/Components/CharacterMovementComponent.cpp
+++ b/Engine/Source/Runtime/Engine/Private/Components/CharacterMovementComponent.cpp
@@ -5931,6 +5931,13 @@ void UCharacterMovementComponent::PhysicsRotation(float DeltaTime)
 	{
 		DesiredRotation = CharacterOwner->Controller->GetDesiredRotation();
 	}
+	else if (!CharacterOwner->Controller && bRunPhysicsWithNoController && bUseControllerDesiredRotation)
+	{
+		if (AController* ControllerOwner = Cast<AController>(CharacterOwner->GetOwner()))
+		{
+			DesiredRotation = ControllerOwner->GetDesiredRotation();
+		}
+	}
 	else
 	{
 		return;
