[KillPendingKill] Rename leaked packages in VerifyLoadMapWorldCleanup and don't fatally assert when leaks happen.
Also fix missing log when ensures are compiled out

diff --git a/Engine/Source/Runtime/Engine/Private/UnrealEngine.cpp b/Engine/Source/Runtime/Engine/Private/UnrealEngine.cpp
--- a/Engine/Source/Runtime/Engine/Private/UnrealEngine.cpp
+++ b/Engine/Source/Runtime/Engine/Private/UnrealEngine.cpp
@@ -15305,6 +15305,7 @@ void UEngine::VerifyLoadMapWorldCleanup(FWorldContext* ForWorldContext)
 	// All worlds at this point should be the CurrentWorld of some context, preview worlds, or streaming level
 	// worlds that are owned by the CurrentWorld of some context.
 	TArray<UObject*> LeakedObjects;
+	TSet<UPackage*> LeakedPackages;
 	for( TObjectIterator<UWorld> It; It; ++It )
 	{
 		UWorld* World = *It;
@@ -15314,6 +15315,7 @@ void UEngine::VerifyLoadMapWorldCleanup(FWorldContext* ForWorldContext)
 			if ((World->PersistentLevel == nullptr || !WorldHasValidContext(World->PersistentLevel->OwningWorld)) && !IsWorldDuplicate(World))
 			{
 				LeakedObjects.Add(World);
+				LeakedPackages.Add(World->GetPackage());
 			}
 		}
 	}
@@ -15339,7 +15341,16 @@ void UEngine::VerifyLoadMapWorldCleanup(FWorldContext* ForWorldContext)
 
 		if (!UObjectBaseUtility::IsPendingKillEnabled())
 		{
-			UE_LOG(LogTemp, Fatal, TEXT("Shutting down because of world leaks."));
+			UE_LOG(LogLoad, Warning, TEXT("Renaming leaked objects in case they need to be reloaded."));
+			for (UPackage* Pkg : LeakedPackages)
+			{
+				FName BaseName = FName(Pkg->GetFName(), NAME_NO_NUMBER_INTERNAL);
+				TStringBuilder<NAME_SIZE> Builder;
+				Builder << TEXT("/Leaked");
+				Builder << BaseName;
+				FName NewName = MakeUniqueObjectName(nullptr, Pkg->GetClass(), FName(*Builder));
+				Pkg->Rename(*NewName.ToString(), nullptr, REN_ForceNoResetLoaders | REN_DontCreateRedirectors | REN_DoNotDirty);
+			}
 		}
 	}
 }
@@ -15515,6 +15526,7 @@ TArray<FString> UEngine::FindAndPrintStaleReferencesToObjects(TConstArrayView<UO
 
 		Paths.Add(PathToCulprit);
 
+#if DO_ENSURE
 		if ((Options & EPrintStaleReferencesOptions::Ensure) == EPrintStaleReferencesOptions::Ensure)
 		{
 			ensureAlwaysMsgf(false, TEXT("Old %s not cleaned up by GC! %s:") LINE_TERMINATOR TEXT("%s"),
@@ -15522,7 +15534,9 @@ TArray<FString> UEngine::FindAndPrintStaleReferencesToObjects(TConstArrayView<UO
 				*GarbageErrorMessage,
 				*PathToCulprit);
 		}
-		else if (Verbosity == ELogVerbosity::Fatal)
+		else
+#endif
+		if (Verbosity == ELogVerbosity::Fatal)
 		{
 			UE_LOG(LogLoad, Fatal, TEXT("Old %s not cleaned up by GC! %s:") LINE_TERMINATOR TEXT("%s"),
 				*ObjectToFindReferencesTo->GetFullName(), 
