Support Meta NavigationModifier which owner is not directly an AActor by selecting the Actor in its Outer chain.

#rnx

diff --git a/Engine/Source/Runtime/Engine/Private/AI/Navigation/NavigationModifier.cpp b/Engine/Source/Runtime/Engine/Private/AI/Navigation/NavigationModifier.cpp
--- a/Engine/Source/Runtime/Engine/Private/AI/Navigation/NavigationModifier.cpp
+++ b/Engine/Source/Runtime/Engine/Private/AI/Navigation/NavigationModifier.cpp
@@ -773,8 +773,21 @@ FCompositeNavModifier FCompositeNavModifier::GetInstantiatedMetaModifier(const F
 	{
 		return Result;
 	}
+	
+	auto FindActorOwner = [](UObject* Obj) -> const AActor*
+	{
+		while (Obj)
+		{
+			if (const AActor* ActorOwner = Cast<AActor>(Obj->GetOuter()))
+			{
+				return ActorOwner;
+			}
+			Obj = Obj->GetOuter();
+		}
+		return nullptr;
+	};
 
-	const AActor* ActorOwner = Cast<AActor>(ObjectOwner) ? (AActor*)ObjectOwner : Cast<AActor>(ObjectOwner->GetOuter());
+	const AActor* ActorOwner = Cast<AActor>(ObjectOwner) ? (AActor*)ObjectOwner : FindActorOwner(ObjectOwner);
 	if (ActorOwner == NULL)
 	{
 		return Result;
