Manual merge of 20563093
Don't inherit CDO editor-only properties if it comes from a cooked package without editor data

diff --git a/Engine/Source/Runtime/Engine/Classes/Engine/BlueprintGeneratedClass.h b/Engine/Source/Runtime/Engine/Classes/Engine/BlueprintGeneratedClass.h
--- a/Engine/Source/Runtime/Engine/Classes/Engine/BlueprintGeneratedClass.h
+++ b/Engine/Source/Runtime/Engine/Classes/Engine/BlueprintGeneratedClass.h
@@ -799,8 +799,9 @@ protected:
 	* @param	InStruct			the current scope for which the given property list applies
 	* @param	DataPtr				destination address (where to start copying values to)
 	* @param	DefaultDataPtr		source address (where to start copying the defaults data from)
+	* @param	bSkipEditorOnly		whether to skip editor-only properties
 	*/
-	static void InitPropertiesFromCustomList(const FCustomPropertyListNode* InPropertyList, UStruct* InStruct, uint8* DataPtr, const uint8* DefaultDataPtr);
+	static void InitPropertiesFromCustomList(const FCustomPropertyListNode* InPropertyList, UStruct* InStruct, uint8* DataPtr, const uint8* DefaultDataPtr, bool bSkipEditorOnly = false);
 
 	/**
 	* Helper method to assist with initializing from an array property with an explicit item list.
@@ -809,8 +810,9 @@ protected:
 	* @param	InPropertyList		only these properties (indices) will be copied from defaults
 	* @param	DataPtr				destination address (where to start copying values to)
 	* @param	DefaultDataPtr		source address (where to start copying the defaults data from)
+	* @param	bSkipEditorOnly		whether to skip editor-only properties
 	*/
-	static void InitArrayPropertyFromCustomList(const FArrayProperty* ArrayProperty, const FCustomPropertyListNode* InPropertyList, uint8* DataPtr, const uint8* DefaultDataPtr);
+	static void InitArrayPropertyFromCustomList(const FArrayProperty* ArrayProperty, const FCustomPropertyListNode* InPropertyList, uint8* DataPtr, const uint8* DefaultDataPtr, bool bSkipEditorOnly = false);
 
 
 	/**
@@ -820,9 +822,10 @@ protected:
 	* @param	CustomPropertySubListNode	the list node of the sub property to copy
 	* @param	PropertyValue				destination address (where to start copying values to)
 	* @param	DefaultPropertyValue		source address (where to start copying the defaults data from)
+	* @param	bSkipEditorOnly		whether to skip editor-only properties
 	* @return	true if the method was able to copy the property successfully via the sub custom property list
 	*/
-	static bool InitPropertyFromSubPropertyList(const FProperty* Property, const FCustomPropertyListNode* SubPropertyList, uint8* PropertyValue, const uint8* DefaultPropertyValue);
+	static bool InitPropertyFromSubPropertyList(const FProperty* Property, const FCustomPropertyListNode* SubPropertyList, uint8* PropertyValue, const uint8* DefaultPropertyValue, bool bSkipEditorOnly = false);
 
 	// @todo: BP2CPP_remove
 	/** Check for and handle manual application of default value overrides to component subobjects that were inherited from a nativized parent class */
diff --git a/Engine/Source/Runtime/Engine/Private/BlueprintGeneratedClass.cpp b/Engine/Source/Runtime/Engine/Private/BlueprintGeneratedClass.cpp
--- a/Engine/Source/Runtime/Engine/Private/BlueprintGeneratedClass.cpp
+++ b/Engine/Source/Runtime/Engine/Private/BlueprintGeneratedClass.cpp
@@ -836,18 +836,32 @@ void UBlueprintGeneratedClass::InitPropertiesFromCustomList(uint8* DataPtr, cons
 
 	if (const FCustomPropertyListNode* CustomPropertyList = GetCustomPropertyListForPostConstruction())
 	{
-		InitPropertiesFromCustomList(CustomPropertyList, this, DataPtr, DefaultDataPtr);
+#if WITH_EDITOR
+		// Don't inherit CDO editor-only properties if it comes from a cooked package without editor data
+		const bool bSkipEditorOnly = RootPackageHasAnyFlags(PKG_FilterEditorOnly);
+#else
+		const bool bSkipEditorOnly = false;
+#endif //WITH_EDITOR
+
+		InitPropertiesFromCustomList(CustomPropertyList, this, DataPtr, DefaultDataPtr, bSkipEditorOnly);
 	}
 }
 
-void UBlueprintGeneratedClass::InitPropertiesFromCustomList(const FCustomPropertyListNode* InPropertyList, UStruct* InStruct, uint8* DataPtr, const uint8* DefaultDataPtr)
+void UBlueprintGeneratedClass::InitPropertiesFromCustomList(const FCustomPropertyListNode* InPropertyList, UStruct* InStruct, uint8* DataPtr, const uint8* DefaultDataPtr, bool bSkipEditorOnly)
 {
 	for (const FCustomPropertyListNode* CustomPropertyListNode = InPropertyList; CustomPropertyListNode; CustomPropertyListNode = CustomPropertyListNode->PropertyListNext)
 	{
+#if WITH_EDITOR
+		if (bSkipEditorOnly && CustomPropertyListNode->Property->IsEditorOnlyProperty())
+		{
+			continue;
+		}
+#endif //WITH_EDITOR
+
 		uint8* PropertyValue = CustomPropertyListNode->Property->ContainerPtrToValuePtr<uint8>(DataPtr, CustomPropertyListNode->ArrayIndex);
 		const uint8* DefaultPropertyValue = CustomPropertyListNode->Property->ContainerPtrToValuePtr<uint8>(DefaultDataPtr, CustomPropertyListNode->ArrayIndex);
 
-		if (!InitPropertyFromSubPropertyList(CustomPropertyListNode->Property, CustomPropertyListNode->SubPropertyList, PropertyValue, DefaultPropertyValue))
+		if (!InitPropertyFromSubPropertyList(CustomPropertyListNode->Property, CustomPropertyListNode->SubPropertyList, PropertyValue, DefaultPropertyValue, bSkipEditorOnly))
 		{
 			// Unable to init properties from sub custom property list, fall back to the default copy value behavior
 			CustomPropertyListNode->Property->CopySingleValue(PropertyValue, DefaultPropertyValue);
@@ -855,7 +869,7 @@ void UBlueprintGeneratedClass::InitPropertiesFromCustomList(const FCustomPropert
 	}
 }
 
-void UBlueprintGeneratedClass::InitArrayPropertyFromCustomList(const FArrayProperty* ArrayProperty, const FCustomPropertyListNode* InPropertyList, uint8* DataPtr, const uint8* DefaultDataPtr)
+void UBlueprintGeneratedClass::InitArrayPropertyFromCustomList(const FArrayProperty* ArrayProperty, const FCustomPropertyListNode* InPropertyList, uint8* DataPtr, const uint8* DefaultDataPtr, bool bSkipEditorOnly)
 {
 	FScriptArrayHelper DstArrayValueHelper(ArrayProperty, DataPtr);
 	FScriptArrayHelper SrcArrayValueHelper(ArrayProperty, DefaultDataPtr);
@@ -884,7 +898,7 @@ void UBlueprintGeneratedClass::InitArrayPropertyFromCustomList(const FArrayPrope
 			continue;
 		}
 
-		if (!InitPropertyFromSubPropertyList(ArrayProperty->Inner, CustomArrayPropertyListNode->SubPropertyList, DstArrayItemValue, SrcArrayItemValue))
+		if (!InitPropertyFromSubPropertyList(ArrayProperty->Inner, CustomArrayPropertyListNode->SubPropertyList, DstArrayItemValue, SrcArrayItemValue, bSkipEditorOnly))
 		{
 			// Unable to init properties from sub custom property list, fall back to the default copy value behavior
 			ArrayProperty->Inner->CopyCompleteValue(DstArrayItemValue, SrcArrayItemValue);
@@ -892,18 +906,18 @@ void UBlueprintGeneratedClass::InitArrayPropertyFromCustomList(const FArrayPrope
 	}
 }
 
-bool UBlueprintGeneratedClass::InitPropertyFromSubPropertyList(const FProperty* Property, const FCustomPropertyListNode* SubPropertyList, uint8* PropertyValue, const uint8* DefaultPropertyValue)
+bool UBlueprintGeneratedClass::InitPropertyFromSubPropertyList(const FProperty* Property, const FCustomPropertyListNode* SubPropertyList, uint8* PropertyValue, const uint8* DefaultPropertyValue, bool bSkipEditorOnly)
 {
 	if (SubPropertyList != nullptr)
 	{
 		if (const FStructProperty* StructProperty = CastField<FStructProperty>(Property))
 		{
-			InitPropertiesFromCustomList(SubPropertyList, StructProperty->Struct, PropertyValue, DefaultPropertyValue);
+			InitPropertiesFromCustomList(SubPropertyList, StructProperty->Struct, PropertyValue, DefaultPropertyValue, bSkipEditorOnly);
 			return true;
 		}
 		else if (const FArrayProperty* ArrayProperty = CastField<FArrayProperty>(Property))
 		{
-			InitArrayPropertyFromCustomList(ArrayProperty, SubPropertyList, PropertyValue, DefaultPropertyValue);
+			InitArrayPropertyFromCustomList(ArrayProperty, SubPropertyList, PropertyValue, DefaultPropertyValue, bSkipEditorOnly);
 			return true;
 		}
 		// No need to handle Sets and Maps as they are not optimized through the custom property list
