The Trace.Screenshot command no longer creates a screenshot file. It only sends the screenshot to the trace system.

diff --git a/Engine/Source/Editor/UnrealEd/Private/EditorViewportClient.cpp b/Engine/Source/Editor/UnrealEd/Private/EditorViewportClient.cpp
--- a/Engine/Source/Editor/UnrealEd/Private/EditorViewportClient.cpp
+++ b/Engine/Source/Editor/UnrealEd/Private/EditorViewportClient.cpp
@@ -5855,23 +5855,29 @@ bool FEditorViewportClient::ProcessScreenShots(FViewport* InViewport)
 				}
 			}
 
+			bool bSuppressWritingToFile = false;
 			if (SHOULD_TRACE_SCREENSHOT())
 			{
+				bSuppressWritingToFile = FTraceScreenshot::ShouldSuppressWritingToFile();
 				FTraceScreenshot::TraceScreenshot(BitmapSize.X, BitmapSize.Y, Bitmap, FScreenshotRequest::GetFilename());
 			}
 
-			// Save the bitmap to disk
-			TFuture<bool> CompletionFuture = HighResScreenshotConfig.ImageWriteQueue->Enqueue(MoveTemp(ImageTask));
-			if (CompletionFuture.IsValid())
+			if (!bSuppressWritingToFile)
 			{
-				// this queues it then immediately waits? what's the point of ImageWriteQueue then?
-				// just use FImageUtils::Save
-				bIsScreenshotSaved = CompletionFuture.Get();
+				// Save the bitmap to disk
+				TFuture<bool> CompletionFuture = HighResScreenshotConfig.ImageWriteQueue->Enqueue(MoveTemp(ImageTask));
+				if (CompletionFuture.IsValid())
+				{
+					// this queues it then immediately waits? what's the point of ImageWriteQueue then?
+					// just use FImageUtils::Save
+					bIsScreenshotSaved = CompletionFuture.Get();
+				}
 			}
 		}
 
 		// Done with the request
 		FScreenshotRequest::Reset();
+		FTraceScreenshot::Reset();
 		FScreenshotRequest::OnScreenshotRequestProcessed().Broadcast();
 
 		// Re-enable screen messages - if we are NOT capturing a movie
diff --git a/Engine/Source/Runtime/Engine/Private/GameViewportClient.cpp b/Engine/Source/Runtime/Engine/Private/GameViewportClient.cpp
--- a/Engine/Source/Runtime/Engine/Private/GameViewportClient.cpp
+++ b/Engine/Source/Runtime/Engine/Private/GameViewportClient.cpp
@@ -2031,19 +2031,24 @@ bool UGameViewportClient::ProcessScreenShots(FViewport* InViewport)
 					ScreenShotName += TEXT(".png");
 				}
 
+				bool bSuppressWritingToFile = false;
 				if (SHOULD_TRACE_SCREENSHOT())
 				{
+					bSuppressWritingToFile = FTraceScreenshot::ShouldSuppressWritingToFile();
 					FTraceScreenshot::TraceScreenshot(Size.X, Size.Y, Bitmap, ScreenShotName);
 				}
 
 				// Save the contents of the array to a png file.
-				
-				FImageView Image((const FColor *)Bitmap.GetData(),Size.X,Size.Y);
-				bIsScreenshotSaved = FImageUtils::SaveImageByExtension(*ScreenShotName,Image);
+				if (!bSuppressWritingToFile)
+				{
+					FImageView Image((const FColor*)Bitmap.GetData(), Size.X, Size.Y);
+					bIsScreenshotSaved = FImageUtils::SaveImageByExtension(*ScreenShotName, Image);
+				}
 			}
 		}
 
 		FScreenshotRequest::Reset();
+		FTraceScreenshot::Reset();
 		FScreenshotRequest::OnScreenshotRequestProcessed().Broadcast();
 
 		// Reeanble screen messages - if we are NOT capturing a movie
diff --git a/Engine/Source/Runtime/Engine/Private/ProfilingDebugging/TraceScreenshot.cpp b/Engine/Source/Runtime/Engine/Private/ProfilingDebugging/TraceScreenshot.cpp
--- a/Engine/Source/Runtime/Engine/Private/ProfilingDebugging/TraceScreenshot.cpp
+++ b/Engine/Source/Runtime/Engine/Private/ProfilingDebugging/TraceScreenshot.cpp
@@ -15,17 +15,12 @@
 static void TraceScreenshotCommandCallback(const TArray<FString>& Args)
 {
 	FString Name;
-	if (Args.Num() > 1)
+	if (Args.Num() > 0)
 	{
 		Name = Args[0];
 	}
 
-	if (FTraceScreenshot::Get() == nullptr)
-	{
-		FTraceScreenshot::CreateInstance();
-	}
-
-	FTraceScreenshot::Get()->RequestScreenshot(Name);
+	FTraceScreenshot::RequestScreenshot(Name);
 }
 
 static FAutoConsoleCommand TraceScreenshotCmd(
@@ -36,65 +31,25 @@ static FAutoConsoleCommand TraceScreenshotCmd(
 	FConsoleCommandWithArgsDelegate::CreateStatic(TraceScreenshotCommandCallback)
 );
 
-FTraceScreenshot* FTraceScreenshot::Instance = nullptr;
-
-void FTraceScreenshot::CreateInstance()
-{
-	if (Instance == nullptr)
-	{
-		Instance = new FTraceScreenshot();
-	}
-}
-
-FTraceScreenshot::FTraceScreenshot()
-{
-	if (GEngine && GEngine->GameViewport)
-	{
-		FlushRenderingCommands();
-		GEngine->GameViewport->OnScreenshotCaptured().AddRaw(this, &FTraceScreenshot::HandleScreenshotData);
-		FWorldDelegates::LevelRemovedFromWorld.AddRaw(this, &FTraceScreenshot::WorldDestroyed);
-	}
-}
-
-FTraceScreenshot::~FTraceScreenshot()
-{
-}
+FString FTraceScreenshot::RequestedScreenshotName;
+bool FTraceScreenshot::bSuppressWritingToFile = false;
 
 void FTraceScreenshot::RequestScreenshot(FString Name)
 {
-	ScreenshotName = Name;
+	RequestedScreenshotName = Name;
 	bool bShowUI = false;
+	bSuppressWritingToFile = true;
 	FScreenshotRequest::RequestScreenshot(bShowUI);
 }
 
-void FTraceScreenshot::HandleScreenshotData(int32 InSizeX, int32 InSizeY, const TArray<FColor>& InImageData)
-{
-	if (SHOULD_TRACE_SCREENSHOT())
-	{
-		TraceScreenshot(InSizeX, InSizeY, InImageData, ScreenshotName, 640);
-	}
-}
-
-void FTraceScreenshot::WorldDestroyed(ULevel* InLevel, UWorld* InWorld)
-{
-	if (InLevel == nullptr && InWorld == World.Get())
-	{
-		UE_LOG(LogCore, Display, TEXT("Tracing screenshot \"%s\" skipped - level was removed from world before we got our screenshot"), *ScreenshotName);
-
-		if (Instance)
-		{
-			delete Instance;
-			Instance = nullptr;
-			GEngine->GameViewport->OnScreenshotCaptured().RemoveAll(this);
-		}
-		FWorldDelegates::LevelRemovedFromWorld.RemoveAll(this);
-	}
-}
-
 void FTraceScreenshot::TraceScreenshot(int32 InSizeX, int32 InSizeY, const TArray<FColor>& InImageData, const FString& InScreenshotName, int32 DesiredX)
 {
 	FString ScreenshotName = InScreenshotName;
-	if (ScreenshotName.IsEmpty())
+	if (!RequestedScreenshotName.IsEmpty())
+	{
+		ScreenshotName = RequestedScreenshotName;
+	}
+	else if (ScreenshotName.IsEmpty())
 	{
 		ScreenshotName = FDateTime::Now().ToString(TEXT("Screenshot_%Y%m%d_%H%M%S"));
 	}
@@ -121,5 +76,11 @@ void FTraceScreenshot::TraceScreenshot(int32 InSizeX, int32 InSizeY, const TArra
 		TRACE_SCREENSHOT(*ScreenshotName, Cycles, InSizeX, InSizeY, CompressedBitmap);
 	}
 
-	ScreenshotName.Empty();
+	Reset();
+}
+
+void FTraceScreenshot::Reset()
+{
+	RequestedScreenshotName.Empty();
+	bSuppressWritingToFile = false;
 }
diff --git a/Engine/Source/Runtime/Engine/Public/ProfilingDebugging/TraceScreenshot.h b/Engine/Source/Runtime/Engine/Public/ProfilingDebugging/TraceScreenshot.h
--- a/Engine/Source/Runtime/Engine/Public/ProfilingDebugging/TraceScreenshot.h
+++ b/Engine/Source/Runtime/Engine/Public/ProfilingDebugging/TraceScreenshot.h
@@ -10,11 +10,7 @@ class ULevel;
 class ENGINE_API FTraceScreenshot
 {
 public:
-	static FTraceScreenshot* Get() { return Instance; }
-	static void CreateInstance();
-	void RequestScreenshot(FString Name);
-	void HandleScreenshotData(int32 InSizeX, int32 InSizeY, const TArray<FColor>& InImageData);
-	void WorldDestroyed(ULevel* InLevel, UWorld* InWorld);
+	static void RequestScreenshot(FString Name);
 
 	/* 
 	* Add the provided screenshot to the trace.
@@ -25,6 +21,16 @@ public:
 	* @param DesiredX - Optionally resize the image to the desired width before tracing. Aspect ratio is preserved.
 	*/
 	static void TraceScreenshot(int32 InSizeX, int32 InSizeY, const TArray<FColor>& InImageData, const FString& InScreenshotName, int32 DesiredX = -1);
+
+	/**
+	* Returns true if the screenshot should not be writted to a file.
+	*/
+	static bool ShouldSuppressWritingToFile() { return bSuppressWritingToFile; }
+
+	/**
+	* Reset the internal state of the TraceScreenshot system.
+	*/
+	static void Reset();
 private:
 	FTraceScreenshot();
 	FTraceScreenshot(const FTraceScreenshot& Other) {}
@@ -34,8 +40,6 @@ private:
 	void Unbind();
 
 private:
-	static FTraceScreenshot* Instance;
-
-	FString ScreenshotName;
-	TWeakObjectPtr<UWorld> World;
+	static bool bSuppressWritingToFile;
+	static FString RequestedScreenshotName;
 };
\ No newline at end of file
