World Partition: fix level instance containers registration regression introduced in recent changes.

#rnx

diff --git a/Engine/Source/Runtime/Engine/Private/WorldPartition/LevelInstance/LevelInstanceActorDesc.cpp b/Engine/Source/Runtime/Engine/Private/WorldPartition/LevelInstance/LevelInstanceActorDesc.cpp
--- a/Engine/Source/Runtime/Engine/Private/WorldPartition/LevelInstance/LevelInstanceActorDesc.cpp
+++ b/Engine/Source/Runtime/Engine/Private/WorldPartition/LevelInstance/LevelInstanceActorDesc.cpp
@@ -68,22 +68,20 @@ void FLevelInstanceActorDesc::RegisterContainerInstance(UWorld* InWorld)
 	if (InWorld)
 	{
 		check(!LevelInstanceContainer.IsValid());
-		if (DesiredRuntimeBehavior == ELevelInstanceRuntimeBehavior::Partitioned && !GLevelInstanceDebugForceLevelStreaming)
+
+		if (IsContainerInstance())
 		{
-			if (!LevelPackage.IsNone() && ULevel::GetIsLevelUsingExternalActorsFromPackage(LevelPackage) && ULevelInstanceSubsystem::CanUsePackage(LevelPackage))
-			{
-				ULevelInstanceSubsystem* LevelInstanceSubsystem = UWorld::GetSubsystem<ULevelInstanceSubsystem>(InWorld);
-				check(LevelInstanceSubsystem);
+			ULevelInstanceSubsystem* LevelInstanceSubsystem = UWorld::GetSubsystem<ULevelInstanceSubsystem>(InWorld);
+			check(LevelInstanceSubsystem);
 
-				LevelInstanceContainer = LevelInstanceSubsystem->RegisterContainer(LevelPackage);
-				check(LevelInstanceContainer.IsValid());
+			LevelInstanceContainer = LevelInstanceSubsystem->RegisterContainer(LevelPackage);
+			check(LevelInstanceContainer.IsValid());
 
-				FTransform LevelInstancePivotOffsetTransform = FTransform(ULevel::GetLevelInstancePivotOffsetFromPackage(LevelPackage));
-				FTransform FinalLevelTransform = LevelInstancePivotOffsetTransform * LevelInstanceTransform;
-				FBox ContainerBounds = LevelInstanceSubsystem->GetContainerBounds(LevelPackage).TransformBy(FinalLevelTransform);
+			FTransform LevelInstancePivotOffsetTransform = FTransform(ULevel::GetLevelInstancePivotOffsetFromPackage(LevelPackage));
+			FTransform FinalLevelTransform = LevelInstancePivotOffsetTransform * LevelInstanceTransform;
+			FBox ContainerBounds = LevelInstanceSubsystem->GetContainerBounds(LevelPackage).TransformBy(FinalLevelTransform);
 
-				ContainerBounds.GetCenterAndExtents(BoundsLocation, BoundsExtent);
-			}
+			ContainerBounds.GetCenterAndExtents(BoundsLocation, BoundsExtent);
 		}
 	}
 }
@@ -116,7 +114,27 @@ void FLevelInstanceActorDesc::SetContainer(UActorDescContainer* InContainer)
 
 bool FLevelInstanceActorDesc::IsContainerInstance() const
 {
-	return DesiredRuntimeBehavior == ELevelInstanceRuntimeBehavior::Partitioned;
+	if (DesiredRuntimeBehavior != ELevelInstanceRuntimeBehavior::Partitioned)
+	{
+		return false;
+	}
+	
+	if (GLevelInstanceDebugForceLevelStreaming)
+	{
+		return false;
+	}
+
+	if (LevelPackage.IsNone())
+	{
+		return false;
+	}
+	
+	if (!ULevel::GetIsLevelUsingExternalActorsFromPackage(LevelPackage))
+	{
+		return false;
+	}
+
+	return ULevelInstanceSubsystem::CanUsePackage(LevelPackage);
 }
 
 bool FLevelInstanceActorDesc::GetContainerInstance(const UActorDescContainer*& OutLevelContainer, FTransform& OutLevelTransform, EContainerClusterMode& OutClusterMode) const
@@ -198,4 +216,4 @@ void FLevelInstanceActorDesc::Serialize(FArchive& Ar)
 	}
 }
 
-#endif
+#endif
\ No newline at end of file
