Data Layers: fix "Iterating over Data Layers without a World Partition" when trying to refresh the data layers outliner for non-partitioned maps.

#rnx

diff --git a/Engine/Source/Editor/DataLayerEditor/Private/DataLayer/DataLayerHierarchy.cpp b/Engine/Source/Editor/DataLayerEditor/Private/DataLayer/DataLayerHierarchy.cpp
--- a/Engine/Source/Editor/DataLayerEditor/Private/DataLayer/DataLayerHierarchy.cpp
+++ b/Engine/Source/Editor/DataLayerEditor/Private/DataLayer/DataLayerHierarchy.cpp
@@ -148,6 +148,11 @@ UWorld* FDataLayerHierarchy::GetOwningWorld() const
 
 void FDataLayerHierarchy::CreateItems(TArray<FSceneOutlinerTreeItemPtr>& OutItems) const
 {
+	if (!GetOwningWorld()->GetWorldDataLayers())
+	{
+		return;
+	}
+
 	const UDataLayerSubsystem* DataLayerSubsystem = UWorld::GetSubsystem<UDataLayerSubsystem>(GetOwningWorld());
 
 	auto IsDataLayerShown = [this](const UDataLayerInstance* DataLayerInstance)
diff --git a/Engine/Source/Runtime/Engine/Private/WorldPartition/DataLayer/DataLayerSubsystem.cpp b/Engine/Source/Runtime/Engine/Private/WorldPartition/DataLayer/DataLayerSubsystem.cpp
--- a/Engine/Source/Runtime/Engine/Private/WorldPartition/DataLayer/DataLayerSubsystem.cpp
+++ b/Engine/Source/Runtime/Engine/Private/WorldPartition/DataLayer/DataLayerSubsystem.cpp
@@ -512,10 +512,6 @@ void UDataLayerSubsystem::ForEachDataLayer(TFunctionRef<bool(UDataLayerInstance*
 		{
 			WorldDataLayers->ForEachDataLayer(Func);
 		}
-		else
-		{
-			UE_LOG(LogWorldPartition, Warning, TEXT("Iterating over Data Layers without a World Partition"));
-		}
 	}
 	else if (const AWorldDataLayers* CurrentLevelWorldDataLayers = (InLevelContext && !InLevelContext->IsPersistentLevel()) ? InLevelContext->GetWorldDataLayers() : nullptr)
 	{
