WorldPartition refactor of cook :
- Always loaded actors are no longer moved in the persistent during the ICookPackageSplitter::GetGenerateList (now done in ICookPackageSplitter::PreSaveGeneratorPackage)
- Removed the loading phase of actors from UWorldPartition's FinalizeGeneratorPackageForCook and PopulateGeneratedPackageForCook : moved it respectively in LoadGeneratorPackageObjectsForCook and LoadGeneratedPackageObjectsForCook.
- This refactor will be used in a future change in the COTFS to retrieve cook package splitter's objects that will be moved from their original packages to either the cooking generator package or their respective cooking generated package.
This fix is necessary for fixing these issues : UE-153904, UE-155321

diff --git a/Engine/Source/Runtime/Engine/Private/WorldPartition/WorldPartition.cpp b/Engine/Source/Runtime/Engine/Private/WorldPartition/WorldPartition.cpp
--- a/Engine/Source/Runtime/Engine/Private/WorldPartition/WorldPartition.cpp
+++ b/Engine/Source/Runtime/Engine/Private/WorldPartition/WorldPartition.cpp
@@ -1245,6 +1245,18 @@ void UWorldPartition::DrawRuntimeHashPreview()
 	RuntimeHash->DrawPreview();
 }
 
+bool UWorldPartition::LoadGeneratorPackageObjectsForCook(TArray<UObject*>& OutLoadedObjects)
+{
+	check(RuntimeHash);
+	return RuntimeHash->LoadGeneratorPackageObjectsForCook(OutLoadedObjects);
+}
+
+bool UWorldPartition::LoadGeneratedPackageObjectsForCook(const FString& InPackageRelativePath, TArray<UObject*>& OutLoadedObjects)
+{
+	check(RuntimeHash);
+	return RuntimeHash->LoadGeneratedPackageObjectsForCook(InPackageRelativePath, OutLoadedObjects);
+}
+
 bool UWorldPartition::PopulateGeneratedPackageForCook(UPackage* InPackage, const FString& InPackageRelativePath)
 {
 	check(RuntimeHash);
diff --git a/Engine/Source/Runtime/Engine/Private/WorldPartition/WorldPartitionCookPackageSplitter.cpp b/Engine/Source/Runtime/Engine/Private/WorldPartition/WorldPartitionCookPackageSplitter.cpp
--- a/Engine/Source/Runtime/Engine/Private/WorldPartition/WorldPartitionCookPackageSplitter.cpp
+++ b/Engine/Source/Runtime/Engine/Private/WorldPartition/WorldPartitionCookPackageSplitter.cpp
@@ -126,7 +126,13 @@ bool FWorldPartitionCookPackageSplitter::TryPopulatePackage(const UPackage* Owne
 	const UWorld* ConstPartitionedWorld = ValidateDataObject(OwnerObject);
 	UWorld* PartitionedWorld = const_cast<UWorld*>(ConstPartitionedWorld);
 	UWorldPartition* WorldPartition = PartitionedWorld->PersistentLevel->GetWorldPartition();
-	bool bResult = WorldPartition->PopulateGeneratedPackageForCook(GeneratedPackage.Package, GeneratedPackage.RelativePath);
+
+	TArray<UObject*> LoadedObjects;
+	bool bResult = WorldPartition->LoadGeneratedPackageObjectsForCook(GeneratedPackage.RelativePath, LoadedObjects);
+	if (bResult)
+	{
+		bResult = WorldPartition->PopulateGeneratedPackageForCook(GeneratedPackage.Package, GeneratedPackage.RelativePath);
+	}
 	return bResult;
 }
 
@@ -135,7 +141,12 @@ void FWorldPartitionCookPackageSplitter::PreSaveGeneratorPackage(UPackage* Owner
 {
 	UWorld* PartitionedWorld = ValidateDataObject(OwnerObject);
 	UWorldPartition* WorldPartition = PartitionedWorld->PersistentLevel->GetWorldPartition();
-	WorldPartition->FinalizeGeneratorPackageForCook(GeneratedPackages);
+
+	TArray<UObject*> LoadedObjects;
+	if (WorldPartition->LoadGeneratorPackageObjectsForCook(LoadedObjects))
+	{
+		WorldPartition->FinalizeGeneratorPackageForCook(GeneratedPackages);
+	}
 }
 
 #endif
diff --git a/Engine/Source/Runtime/Engine/Private/WorldPartition/WorldPartitionLevelStreamingPolicy.cpp b/Engine/Source/Runtime/Engine/Private/WorldPartition/WorldPartitionLevelStreamingPolicy.cpp
--- a/Engine/Source/Runtime/Engine/Private/WorldPartition/WorldPartitionLevelStreamingPolicy.cpp
+++ b/Engine/Source/Runtime/Engine/Private/WorldPartition/WorldPartitionLevelStreamingPolicy.cpp
@@ -77,6 +77,11 @@ void UWorldPartitionLevelStreamingPolicy::PrepareActorToCellRemapping()
 	// Build Actor-to-Cell remapping
 	for (const UWorldPartitionRuntimeCell* Cell : StreamingCells)
 	{
+		if (!Cell->NeedsActorToCellRemapping())
+		{
+			continue;
+		}
+
 		const UWorldPartitionRuntimeLevelStreamingCell* StreamingCell = Cast<const UWorldPartitionRuntimeLevelStreamingCell>(Cell);
 		check(StreamingCell);
 		for (const FWorldPartitionRuntimeCellObjectMapping& CellObjectMap : StreamingCell->GetPackages())
diff --git a/Engine/Source/Runtime/Engine/Private/WorldPartition/WorldPartitionRuntimeCell.cpp b/Engine/Source/Runtime/Engine/Private/WorldPartition/WorldPartitionRuntimeCell.cpp
--- a/Engine/Source/Runtime/Engine/Private/WorldPartition/WorldPartitionRuntimeCell.cpp
+++ b/Engine/Source/Runtime/Engine/Private/WorldPartition/WorldPartitionRuntimeCell.cpp
@@ -24,6 +24,12 @@ UWorldPartitionRuntimeCell::UWorldPartitionRuntimeCell(const FObjectInitializer&
 {}
 
 #if WITH_EDITOR
+bool UWorldPartitionRuntimeCell::NeedsActorToCellRemapping() const
+{
+	// When cooking, always loaded cells content is moved to persistent level (see FinalizeGeneratorPackageForCook)
+	return !(IsAlwaysLoaded() && IsRunningCookCommandlet());
+}
+
 void UWorldPartitionRuntimeCell::SetDataLayers(const TArray<const UDataLayerInstance*>& InDataLayerInstances)
 {
 	check(DataLayers.IsEmpty());
diff --git a/Engine/Source/Runtime/Engine/Private/WorldPartition/WorldPartitionRuntimeLevelStreamingCell.cpp b/Engine/Source/Runtime/Engine/Private/WorldPartition/WorldPartitionRuntimeLevelStreamingCell.cpp
--- a/Engine/Source/Runtime/Engine/Private/WorldPartition/WorldPartitionRuntimeLevelStreamingCell.cpp
+++ b/Engine/Source/Runtime/Engine/Private/WorldPartition/WorldPartitionRuntimeLevelStreamingCell.cpp
@@ -13,6 +13,7 @@
 #include "WorldPartition/WorldPartitionLevelStreamingPolicy.h"
 #include "WorldPartition/WorldPartitionLevelHelper.h"
 #include "Engine/LevelStreamingAlwaysLoaded.h"
+#include "Algo/ForEach.h"
 #endif
 
 UWorldPartitionRuntimeLevelStreamingCell::UWorldPartitionRuntimeLevelStreamingCell(const FObjectInitializer& ObjectInitializer)
@@ -135,22 +136,43 @@ UWorldPartitionLevelStreamingDynamic* UWorldPartitionRuntimeLevelStreamingCell::
 	return nullptr;
 }
 
+bool UWorldPartitionRuntimeLevelStreamingCell::LoadCellObjectsForCook(TArray<UObject*>& OutLoadedObjects)
+{
+	if (GetActorCount() > 0)
+	{
+		const bool bLoadAsync = false;
+		UWorld* OuterWorld = GetOuterUWorldPartition()->GetTypedOuter<UWorld>();
+		verify(FWorldPartitionLevelHelper::LoadActors(OuterWorld, nullptr, Packages, CookPackageReferencer, [](bool) {}, bLoadAsync, FLinkerInstancingContext()));
+
+		for (const FWorldPartitionRuntimeCellObjectMapping& Mapping : Packages)
+		{
+			if (AActor* Actor = FindObject<AActor>(nullptr, *Mapping.LoadedPath.ToString()))
+			{
+				OutLoadedObjects.Add(Actor);
+				CookLoadedActors.Add(Actor);
+			}
+		}
+	}
+	return true;
+}
+
 void UWorldPartitionRuntimeLevelStreamingCell::MoveAlwaysLoadedContentToPersistentLevel()
 {
 	check(IsAlwaysLoaded());
 	if (GetActorCount() > 0)
 	{
-		FWorldPartitionLevelHelper::FPackageReferencer PackageReferencer;
-		const bool bLoadAsync = false;
+		// At this point actors should be loaded
+		Algo::ForEach(CookLoadedActors, [](const TWeakObjectPtr<AActor>& Actor) { check(Actor.Get()); });
 
 		UWorld* OuterWorld = GetOuterUWorldPartition()->GetTypedOuter<UWorld>();
-		verify(FWorldPartitionLevelHelper::LoadActors(OuterWorld, nullptr, Packages, PackageReferencer, [](bool) {}, bLoadAsync, FLinkerInstancingContext()));
-				
 		FWorldPartitionLevelHelper::MoveExternalActorsToLevel(Packages, OuterWorld->PersistentLevel);
 
 		// Empty cell's package list (this ensures that no one can rely on cell's content).
 		Packages.Empty();
 	}
+
+	CookLoadedActors.Empty();
+	CookPackageReferencer.RemoveReferences();
 }
 
 // Do all necessary work to prepare cell object for cook.
@@ -186,15 +208,12 @@ bool UWorldPartitionRuntimeLevelStreamingCell::PopulateGeneratedPackageForCook(U
 			return false;
 		}
 
-		// Load cell Actors
+		// At this point actors should be loaded
+		Algo::ForEach(CookLoadedActors, [](const TWeakObjectPtr<AActor>& Actor) { check(Actor.Get()); });
+
+		// Create a level and move loaded actors in it
 		UWorldPartition* WorldPartition = GetOuterUWorldPartition();
 		UWorld* OuterWorld = WorldPartition->GetTypedOuter<UWorld>();
-
-		FWorldPartitionLevelHelper::FPackageReferencer PackageReferencer;
-		const bool bLoadAsync = false;
-		verify(FWorldPartitionLevelHelper::LoadActors(OuterWorld, nullptr, Packages, PackageReferencer, [](bool) {}, bLoadAsync, FLinkerInstancingContext()));
-
-		// Create a level and move these actors in it
 		ULevel* NewLevel = FWorldPartitionLevelHelper::CreateEmptyLevelForRuntimeCell(this, OuterWorld, LevelStreaming->GetWorldAsset().ToString(), InPackage);
 		check(NewLevel->GetPackage() == InPackage);
 		FWorldPartitionLevelHelper::MoveExternalActorsToLevel(Packages, NewLevel);
@@ -202,6 +221,10 @@ bool UWorldPartitionRuntimeLevelStreamingCell::PopulateGeneratedPackageForCook(U
 		// Remap Level's SoftObjectPaths
 		FWorldPartitionLevelHelper::RemapLevelSoftObjectPaths(NewLevel, WorldPartition);
 	}
+
+	CookLoadedActors.Empty();
+	CookPackageReferencer.RemoveReferences();
+
 	return true;
 }
 
diff --git a/Engine/Source/Runtime/Engine/Private/WorldPartition/WorldPartitionRuntimeSpatialHash.cpp b/Engine/Source/Runtime/Engine/Private/WorldPartition/WorldPartitionRuntimeSpatialHash.cpp
--- a/Engine/Source/Runtime/Engine/Private/WorldPartition/WorldPartitionRuntimeSpatialHash.cpp
+++ b/Engine/Source/Runtime/Engine/Private/WorldPartition/WorldPartitionRuntimeSpatialHash.cpp
@@ -1134,12 +1134,8 @@ bool UWorldPartitionRuntimeSpatialHash::CreateStreamingGrid(const FSpatialHashRu
 
 					if (StreamingCell->GetActorCount())
 					{
-						// Always loaded cell actors are transfered to World's Persistent Level
-						if (StreamingCell->IsAlwaysLoaded())
-						{
-							StreamingCell->MoveAlwaysLoadedContentToPersistentLevel();
-						}
-						else
+						// Always loaded cell actors are transfered to World's Persistent Level (see UWorldPartitionRuntimeSpatialHash::FinalizeGeneratorPackageForCook)
+						if (!StreamingCell->IsAlwaysLoaded())
 						{
 							if (!OutPackagesToGenerate)
 							{
@@ -1163,6 +1159,20 @@ bool UWorldPartitionRuntimeSpatialHash::CreateStreamingGrid(const FSpatialHashRu
 	return true;
 }
 
+bool UWorldPartitionRuntimeSpatialHash::LoadGeneratedPackageObjectsForCook(const FString& InPackageRelativePath, TArray<UObject*>& OutLoadedObjects)
+{
+	OutLoadedObjects.Reset();
+	if (UWorldPartitionRuntimeCell** MatchingCell = PackagesToGenerateForCook.Find(InPackageRelativePath))
+	{
+		UWorldPartitionRuntimeCell* Cell = *MatchingCell;
+		if (ensure(Cell))
+		{
+			return Cell->LoadCellObjectsForCook(OutLoadedObjects);
+		}
+	}
+	return false;
+}
+
 bool UWorldPartitionRuntimeSpatialHash::PopulateGeneratedPackageForCook(UPackage* InPackage, const FString& InPackageRelativePath)
 {
 	if (UWorldPartitionRuntimeCell** MatchingCell = PackagesToGenerateForCook.Find(InPackageRelativePath))
@@ -1176,8 +1186,46 @@ bool UWorldPartitionRuntimeSpatialHash::PopulateGeneratedPackageForCook(UPackage
 	return false;
 }
 
+TArray<UWorldPartitionRuntimeCell*> UWorldPartitionRuntimeSpatialHash::GetAlwaysLoadedCells() const
+{
+	TArray<UWorldPartitionRuntimeCell*> Result;
+	TSet<const UWorldPartitionRuntimeCell*> StreamingCells;
+	GetAllStreamingCells(StreamingCells);
+	for (const UWorldPartitionRuntimeCell* Cell : StreamingCells)
+	{
+		if (Cell->IsAlwaysLoaded())
+		{
+			Result.Add(const_cast<UWorldPartitionRuntimeCell*>(Cell));
+		}
+	}
+	return Result;
+}
+
+bool UWorldPartitionRuntimeSpatialHash::LoadGeneratorPackageObjectsForCook(TArray<UObject*>& OutLoadedObjects)
+{
+	OutLoadedObjects.Reset();
+	for (UWorldPartitionRuntimeCell* Cell : GetAlwaysLoadedCells())
+	{
+		check(Cell->IsAlwaysLoaded());
+		if (!Cell->LoadCellObjectsForCook(OutLoadedObjects))
+		{
+			return false;
+		}
+	}
+	return true;
+}
+
 bool UWorldPartitionRuntimeSpatialHash::FinalizeGeneratorPackageForCook(const TArray<ICookPackageSplitter::FGeneratedPackageForPreSave>& InGeneratedPackages)
 {
+	check(IsRunningCookCommandlet());
+	
+	// Move always loaded content in Persistent Level
+	for (UWorldPartitionRuntimeCell* Cell : GetAlwaysLoadedCells())
+	{
+		check(Cell->IsAlwaysLoaded());
+		Cell->MoveAlwaysLoadedContentToPersistentLevel();
+	}
+	
 	for (const ICookPackageSplitter::FGeneratedPackageForPreSave& GeneratedPackage : InGeneratedPackages)
 	{
 		UWorldPartitionRuntimeCell** MatchingCell = PackagesToGenerateForCook.Find(GeneratedPackage.RelativePath);
diff --git a/Engine/Source/Runtime/Engine/Public/WorldPartition/WorldPartition.h b/Engine/Source/Runtime/Engine/Public/WorldPartition/WorldPartition.h
--- a/Engine/Source/Runtime/Engine/Public/WorldPartition/WorldPartition.h
+++ b/Engine/Source/Runtime/Engine/Public/WorldPartition/WorldPartition.h
@@ -150,6 +150,8 @@ public:
 	void RemapSoftObjectPath(FSoftObjectPath& ObjectPath);
 
 	// Cooking
+	bool LoadGeneratorPackageObjectsForCook(TArray<UObject*>& OutLoadedObjects);
+	bool LoadGeneratedPackageObjectsForCook(const FString& InPackageRelativePath, TArray<UObject*>& OutLoadedObjects);
 	bool PopulateGeneratedPackageForCook(UPackage* InPackage, const FString& InPackageRelativePath);
 	bool FinalizeGeneratorPackageForCook(const TArray<ICookPackageSplitter::FGeneratedPackageForPreSave>& InGeneratedPackages);
 
diff --git a/Engine/Source/Runtime/Engine/Public/WorldPartition/WorldPartitionRuntimeCell.h b/Engine/Source/Runtime/Engine/Public/WorldPartition/WorldPartitionRuntimeCell.h
--- a/Engine/Source/Runtime/Engine/Public/WorldPartition/WorldPartitionRuntimeCell.h
+++ b/Engine/Source/Runtime/Engine/Public/WorldPartition/WorldPartitionRuntimeCell.h
@@ -167,6 +167,7 @@ class UWorldPartitionRuntimeCell : public UObject, public IWorldPartitionCell
 
 	bool GetBlockOnSlowLoading() const { return bBlockOnSlowLoading; }
 #if WITH_EDITOR
+	bool NeedsActorToCellRemapping() const;
 	void SetBlockOnSlowLoading(bool bInBlockOnSlowLoading) { bBlockOnSlowLoading = bInBlockOnSlowLoading; }
 
 	void SetClientOnlyVisible(bool bInClientOnlyVisible) { bClientOnlyVisible = bInClientOnlyVisible; }
@@ -180,6 +181,7 @@ class UWorldPartitionRuntimeCell : public UObject, public IWorldPartitionCell
 
 	// Cook methods
 	virtual bool PrepareCellForCook(UPackage* InPackage) { return false; }
+	virtual bool LoadCellObjectsForCook(TArray<UObject*>& OutLoadedObjects) PURE_VIRTUAL(UWorldPartitionRuntimeCell::LoadCellObjectsForCook, return false;);
 	virtual bool PopulateGeneratedPackageForCook(UPackage* InPackage) PURE_VIRTUAL(UWorldPartitionRuntimeCell::PopulateGeneratedPackageForCook, return false;);
 	virtual void MoveAlwaysLoadedContentToPersistentLevel() PURE_VIRTUAL(UWorldPartitionRuntimeCell::MoveAlwaysLoadedContentToPersistentLevel);
 	virtual FString GetPackageNameToCreate() const PURE_VIRTUAL(UWorldPartitionRuntimeCell::GetPackageNameToCreate, return FString(""););
diff --git a/Engine/Source/Runtime/Engine/Public/WorldPartition/WorldPartitionRuntimeHash.h b/Engine/Source/Runtime/Engine/Public/WorldPartition/WorldPartitionRuntimeHash.h
--- a/Engine/Source/Runtime/Engine/Public/WorldPartition/WorldPartitionRuntimeHash.h
+++ b/Engine/Source/Runtime/Engine/Public/WorldPartition/WorldPartitionRuntimeHash.h
@@ -34,6 +34,8 @@ class ENGINE_API UWorldPartitionRuntimeHash : public UObject
 
 #if WITH_EDITOR
 	virtual void SetDefaultValues() {}
+	virtual bool LoadGeneratorPackageObjectsForCook(TArray<UObject*>& OutLoadedObjects) { return false; }
+	virtual bool LoadGeneratedPackageObjectsForCook(const FString& InPackageRelativePath, TArray<UObject*>& OutLoadedObjects) { return false; }
 	virtual bool PopulateGeneratedPackageForCook(UPackage* InPackage, const FString& InPackageRelativePath) { return false; }
 	virtual bool FinalizeGeneratorPackageForCook(const TArray<ICookPackageSplitter::FGeneratedPackageForPreSave>& InGeneratedPackages) { return false; }
 	virtual bool GenerateStreaming(class UWorldPartitionStreamingPolicy* StreamingPolicy, const FActorClusterContext& ActorClusterContext, TArray<FString>* OutPackagesToGenerate) { return false; }
diff --git a/Engine/Source/Runtime/Engine/Public/WorldPartition/WorldPartitionRuntimeLevelStreamingCell.h b/Engine/Source/Runtime/Engine/Public/WorldPartition/WorldPartitionRuntimeLevelStreamingCell.h
--- a/Engine/Source/Runtime/Engine/Public/WorldPartition/WorldPartitionRuntimeLevelStreamingCell.h
+++ b/Engine/Source/Runtime/Engine/Public/WorldPartition/WorldPartitionRuntimeLevelStreamingCell.h
@@ -3,6 +3,9 @@
 
 #include "WorldPartition/WorldPartitionRuntimeSpatialHashCell.h"
 #include "ProfilingDebugging/ProfilingHelpers.h"
+#if WITH_EDITOR
+#include "WorldPartition/WorldPartitionLevelHelper.h"
+#endif
 #include "WorldPartitionRuntimeLevelStreamingCell.generated.h"
 
 UCLASS()
@@ -37,6 +40,7 @@ class UWorldPartitionRuntimeLevelStreamingCell : public UWorldPartitionRuntimeSp
 	virtual void DumpStateLog(FHierarchicalLogArchive& Ar) override;
 	// Cook methods
 	virtual bool PrepareCellForCook(UPackage* InPackage) override;
+	virtual bool LoadCellObjectsForCook(TArray<UObject*>& OutLoadedObjects) override;
 	virtual bool PopulateGeneratedPackageForCook(UPackage* InPackage) override;
 	virtual FString GetPackageNameToCreate() const override;
 	//~End UWorldPartitionRuntimeCell Interface
@@ -56,6 +60,9 @@ private:
 #if WITH_EDITOR
 	void MoveAlwaysLoadedContentToPersistentLevel();
 	class UWorldPartitionLevelStreamingDynamic* CreateLevelStreaming(const FString& InPackageName = FString()) const;
+
+	FWorldPartitionLevelHelper::FPackageReferencer CookPackageReferencer;
+	TArray<TWeakObjectPtr<AActor>> CookLoadedActors;
 #endif
 
 #if WITH_EDITORONLY_DATA
diff --git a/Engine/Source/Runtime/Engine/Public/WorldPartition/WorldPartitionRuntimeSpatialHash.h b/Engine/Source/Runtime/Engine/Public/WorldPartition/WorldPartitionRuntimeSpatialHash.h
--- a/Engine/Source/Runtime/Engine/Public/WorldPartition/WorldPartitionRuntimeSpatialHash.h
+++ b/Engine/Source/Runtime/Engine/Public/WorldPartition/WorldPartitionRuntimeSpatialHash.h
@@ -229,6 +229,8 @@ public:
 
 #if WITH_EDITOR
 	virtual void SetDefaultValues() override;
+	virtual bool LoadGeneratorPackageObjectsForCook(TArray<UObject*>& OutLoadedObjects) override;
+	virtual bool LoadGeneratedPackageObjectsForCook(const FString& InPackageRelativePath, TArray<UObject*>& OutLoadedObjects) override;
 	virtual bool PopulateGeneratedPackageForCook(UPackage* InPackage, const FString& InPackageRelativePath) override;
 	virtual bool FinalizeGeneratorPackageForCook(const TArray<ICookPackageSplitter::FGeneratedPackageForPreSave>& InGeneratedPackages) override;
 	virtual void FlushStreaming() override;
@@ -243,7 +245,7 @@ public:
 #endif
 
 	// streaming interface
-	virtual int32 GetAllStreamingCells(TSet<const UWorldPartitionRuntimeCell*>& Cells, bool bAllDataLayers, bool bDataLayersOnly, const TSet<FName>& InDataLayers) const override;
+	virtual int32 GetAllStreamingCells(TSet<const UWorldPartitionRuntimeCell*>& Cells, bool bAllDataLayers = false, bool bDataLayersOnly = false, const TSet<FName>& InDataLayers = TSet<FName>()) const override;
 	virtual bool GetStreamingCells(const FWorldPartitionStreamingQuerySource& QuerySource, TSet<const UWorldPartitionRuntimeCell*>& OutCells) const override;
 	virtual bool GetStreamingCells(const TArray<FWorldPartitionStreamingSource>& Sources, UWorldPartitionRuntimeHash::FStreamingSourceCells& OutActivateCells, UWorldPartitionRuntimeHash::FStreamingSourceCells& OutLoadCells) const override;
 		
@@ -298,6 +300,7 @@ private:
 	void GetStreamingCells(const FVector& Position, const FSpatialHashStreamingGrid& StreamingGrid, TSet<const UWorldPartitionRuntimeCell*>& Cells) const;
 	const TMap<FName, const FSpatialHashStreamingGrid*>& GetNameToGridMapping() const;
 #if WITH_EDITOR
+	TArray<UWorldPartitionRuntimeCell*> GetAlwaysLoadedCells() const;
 	bool CreateStreamingGrid(const FSpatialHashRuntimeGrid& RuntimeGrid, const FSquare2DGridHelper& PartionedActors, UWorldPartitionStreamingPolicy* StreamingPolicy, TArray<FString>* OutPackagesToGenerate = nullptr);
 #endif
 	TArray<const FSpatialHashStreamingGrid*> GetFilteredStreamingGrids() const;
