Added World::RepairSingletonActors than will remove duplicated singleton actors from worlds.

#rnx

diff --git a/Engine/Source/Runtime/Engine/Classes/Engine/World.h b/Engine/Source/Runtime/Engine/Classes/Engine/World.h
--- a/Engine/Source/Runtime/Engine/Classes/Engine/World.h
+++ b/Engine/Source/Runtime/Engine/Classes/Engine/World.h
@@ -1597,9 +1597,16 @@ private:
 	/** Finish Async Trace Buffer **/
 	void FinishAsyncTrace();
 
+	/** Utility function that is used to ensure that a World has the correct singleton actor of the provided class */
+	void RepairSingletonActorOfClass(TSubclassOf<AActor> ActorClass);	
+	template <class T> void RepairSingletonActorOfClass() { RepairSingletonActorOfClass(T::StaticClass()); }
+
 	/** Utility function that is used to ensure that a World has the correct WorldSettings */
 	void RepairWorldSettings();
 
+	/** Utility function that is used to ensure that a World has the correct singleton actors*/
+	void RepairSingletonActors();
+
 	/** Utility function to cleanup streaming levels that point to invalid level packages */
 	void RepairStreamingLevels();
 
diff --git a/Engine/Source/Runtime/Engine/Private/World.cpp b/Engine/Source/Runtime/Engine/Private/World.cpp
--- a/Engine/Source/Runtime/Engine/Private/World.cpp
+++ b/Engine/Source/Runtime/Engine/Private/World.cpp
@@ -61,6 +61,7 @@
 #include "Engine/WorldComposition.h"
 #include "WorldPartition/WorldPartition.h"
 #include "WorldPartition/WorldPartitionSubsystem.h"
+#include "WorldPartition/DataLayer/WorldDataLayers.h"
 #include "Misc/NetworkVersion.h"
 #include "GameFramework/GameNetworkManager.h"
 #include "GameFramework/CharacterMovementComponent.h"
@@ -1164,7 +1165,7 @@ void UWorld::PostLoad()
 	CurrentLevel = PersistentLevel;
 #endif
 #if WITH_EDITOR
-	RepairWorldSettings();
+	RepairSingletonActors();
 	RepairStreamingLevels();
 #endif
 
@@ -1543,6 +1544,38 @@ void UWorld::RepairStreamingLevels()
 	}
 }
 
+void UWorld::RepairSingletonActorOfClass(TSubclassOf<AActor> ActorClass)
+{
+	AActor* FoundActor = nullptr;
+
+	for (int32 i=0; i<PersistentLevel->Actors.Num(); i++)
+	{
+		if (AActor* CurrentActor = PersistentLevel->Actors[i])
+		{
+			if (CurrentActor->IsA(ActorClass))
+			{
+				if (FoundActor)
+				{
+					UE_LOG(LogWorld, Warning, TEXT("Extra '%s' actor found. Resave level %s or actors to clean up."), *CurrentActor->GetPathName(), *PersistentLevel->GetPathName());
+
+					if (FoundActor == CurrentActor)
+					{
+						PersistentLevel->Actors[i] = nullptr;
+					}
+					else
+					{
+						CurrentActor->Destroy();
+					}
+				}
+				else
+				{
+					FoundActor = CurrentActor;
+				}
+			}
+		}
+	}
+}
+
 void UWorld::RepairWorldSettings()
 {
 	AWorldSettings* ExistingWorldSettings = PersistentLevel->GetWorldSettings(false);
@@ -1618,6 +1651,12 @@ void UWorld::RepairWorldSettings()
 	check(GetWorldSettings());
 }
 
+void UWorld::RepairSingletonActors()
+{
+	RepairWorldSettings();
+	RepairSingletonActorOfClass<AWorldDataLayers>();
+}
+
 #if WITH_EDITOR
 void UWorld::RepairDefaultBrush()
 {
@@ -1809,7 +1848,7 @@ void UWorld::InitWorld(const InitializationValues IVS)
 	PersistentLevel->bIsVisible = true;
 
 #if WITH_EDITOR
-	RepairWorldSettings();
+	RepairSingletonActors();
 	RepairStreamingLevels();
 #endif
 
