Fix ambiguous resolution in unity build for GetEditorOnlyDataName

diff --git a/Engine/Source/Runtime/Engine/Private/Materials/MaterialExpressions.cpp b/Engine/Source/Runtime/Engine/Private/Materials/MaterialExpressions.cpp
--- a/Engine/Source/Runtime/Engine/Private/Materials/MaterialExpressions.cpp
+++ b/Engine/Source/Runtime/Engine/Private/Materials/MaterialExpressions.cpp
@@ -13263,7 +13263,7 @@ void UMaterialFunctionInterface::GetAssetRegistryTags(TArray<FAssetRegistryTag>&
 #endif
 }
 
-namespace
+namespace MaterialFunctionInterface
 {
 	FString GetEditorOnlyDataName(const TCHAR* InMaterialName)
 	{
@@ -13278,7 +13278,7 @@ bool UMaterialFunctionInterface::Rename(const TCHAR* NewName, UObject* NewOuter,
 	// if we have EditorOnlyData, also rename it if we are changing the material's name
 	if (bRenamed && NewName && EditorOnlyData)
 	{
-		FString EditorOnlyDataName = GetEditorOnlyDataName(NewName);
+		FString EditorOnlyDataName = MaterialFunctionInterface::GetEditorOnlyDataName(NewName);
 		bRenamed = EditorOnlyData->Rename(*EditorOnlyDataName, nullptr, Flags);
 	}
 #endif
@@ -13500,7 +13500,7 @@ UMaterialFunctionInterfaceEditorOnlyData* UMaterialFunctionInterface::CreateEdit
 	check(EditorOnlyClass);
 	check(EditorOnlyClass->HasAllClassFlags(CLASS_Optional));
 
-	const FString EditorOnlyName = GetEditorOnlyDataName(*GetName());
+	const FString EditorOnlyName = MaterialFunctionInterface::GetEditorOnlyDataName(*GetName());
 	const EObjectFlags EditorOnlyFlags = GetMaskedFlags(RF_PropagateToSubObjects);
 	return NewObject<UMaterialFunctionInterfaceEditorOnlyData>(this, EditorOnlyClass, *EditorOnlyName, EditorOnlyFlags);
 }
diff --git a/Engine/Source/Runtime/Engine/Private/Materials/MaterialInterface.cpp b/Engine/Source/Runtime/Engine/Private/Materials/MaterialInterface.cpp
--- a/Engine/Source/Runtime/Engine/Private/Materials/MaterialInterface.cpp
+++ b/Engine/Source/Runtime/Engine/Private/Materials/MaterialInterface.cpp
@@ -1364,7 +1364,7 @@ void UMaterialInterface::RemoveUserDataOfClass(TSubclassOf<UAssetUserData> InUse
 
 #if WITH_EDITORONLY_DATA
 
-namespace
+namespace MaterialInterface
 {
 	FString GetEditorOnlyDataName(const TCHAR* InMaterialName)
 	{
@@ -1385,7 +1385,7 @@ UMaterialInterfaceEditorOnlyData* UMaterialInterface::CreateEditorOnlyData()
 	check(EditorOnlyClass);
 	check(EditorOnlyClass->HasAllClassFlags(CLASS_Optional));
 
-	const FString EditorOnlyName = GetEditorOnlyDataName(*GetName());
+	const FString EditorOnlyName = MaterialInterface::GetEditorOnlyDataName(*GetName());
 	const EObjectFlags EditorOnlyFlags = GetMaskedFlags(RF_PropagateToSubObjects);
 	return NewObject<UMaterialInterfaceEditorOnlyData>(this, EditorOnlyClass, *EditorOnlyName, EditorOnlyFlags);
 }
@@ -1398,7 +1398,7 @@ bool UMaterialInterface::Rename(const TCHAR* NewName, UObject* NewOuter, ERename
 	// if we have EditorOnlyData, also rename it if we are changing the material's name
 	if (bRenamed && NewName && EditorOnlyData)
 	{
-		FString EditorOnlyDataName = GetEditorOnlyDataName(NewName);
+		FString EditorOnlyDataName = MaterialInterface::GetEditorOnlyDataName(NewName);
 		bRenamed = EditorOnlyData->Rename(*EditorOnlyDataName, nullptr, Flags);
 	}
 #endif
