Asset Actions "Replace Selected Actors" now transfers more properties : Actor Folder, Data Layers, IsSpatiallyLoaded flag and RuntimeGrid.

diff --git a/Engine/Source/Editor/DataLayerEditor/Private/DataLayerEditorModule.cpp b/Engine/Source/Editor/DataLayerEditor/Private/DataLayerEditorModule.cpp
--- a/Engine/Source/Editor/DataLayerEditor/Private/DataLayerEditorModule.cpp
+++ b/Engine/Source/Editor/DataLayerEditor/Private/DataLayerEditorModule.cpp
@@ -11,6 +11,7 @@
 #include "DataLayer/DataLayerPropertyTypeCustomization.h"
 #include "DataLayer/SDataLayerBrowser.h"
 #include "DataLayer/DataLayerNameEditSink.h"
+#include "DataLayer/DataLayerEditorSubsystem.h"
 #include "WorldPartition/DataLayer/DataLayerInstance.h"
 
 IMPLEMENT_MODULE(FDataLayerEditorModule, DataLayerEditor );
@@ -50,3 +51,8 @@ void FDataLayerEditorModule::SyncDataLayerBrowserToDataLayer(const UDataLayerIns
 		Browser->SyncDataLayerBrowserToDataLayer(DataLayerInstance);
 	}
 }
+
+bool FDataLayerEditorModule::AddActorToDataLayers(AActor* Actor, const TArray<UDataLayerInstance*>& DataLayers)
+{
+	return UDataLayerEditorSubsystem::Get()->AddActorToDataLayers(Actor, DataLayers);
+}
\ No newline at end of file
diff --git a/Engine/Source/Editor/DataLayerEditor/Public/DataLayerEditorModule.h b/Engine/Source/Editor/DataLayerEditor/Public/DataLayerEditorModule.h
--- a/Engine/Source/Editor/DataLayerEditor/Public/DataLayerEditorModule.h
+++ b/Engine/Source/Editor/DataLayerEditor/Public/DataLayerEditorModule.h
@@ -40,6 +40,9 @@ public:
 	DECLARE_DELEGATE_RetVal_OneParam(TSharedRef<FExtender>, FDataLayersMenuExtender, const TSharedRef<FUICommandList>);
 	virtual TArray<FDataLayersMenuExtender>& GetAllDataLayersMenuExtenders() {return DataLayersMenuExtenders;}
 
+	/* Implement IDataLayerEditorModule */
+	virtual bool AddActorToDataLayers(AActor* Actor, const TArray<UDataLayerInstance*>& DataLayers) override;
+
 private:
 	TWeakPtr<SWidget> DataLayerBrowser;
 
diff --git a/Engine/Source/Editor/UnrealEd/Private/EditorActor.cpp b/Engine/Source/Editor/UnrealEd/Private/EditorActor.cpp
--- a/Engine/Source/Editor/UnrealEd/Private/EditorActor.cpp
+++ b/Engine/Source/Editor/UnrealEd/Private/EditorActor.cpp
@@ -1357,9 +1357,6 @@ AActor* UUnrealEdEngine::ReplaceActor( AActor* CurrentActor, UClass* NewActorCla
 
 			LayersSubsystem->AddActorToLayers( NewActor, CurrentActor->Layers );
 
-			NewActor->SetActorLabel( CurrentActor->GetActorLabel() );
-			NewActor->Tags = CurrentActor->Tags;
-
 			NewActor->EditorReplacedActor( CurrentActor );
 		}
 
diff --git a/Engine/Source/Editor/UnrealEd/Private/EditorEngine.cpp b/Engine/Source/Editor/UnrealEd/Private/EditorEngine.cpp
--- a/Engine/Source/Editor/UnrealEd/Private/EditorEngine.cpp
+++ b/Engine/Source/Editor/UnrealEd/Private/EditorEngine.cpp
@@ -5168,10 +5168,6 @@ void UEditorEngine::ReplaceActors(UActorFactory* Factory, const FAssetData& Asse
 			ULayersSubsystem* LayersSubsystem = GetEditorSubsystem<ULayersSubsystem>();
 			LayersSubsystem->AddActorToLayers( NewActor, OldActor->Layers );
 
-			// Preserve the label and tags from the old actor
-			NewActor->SetActorLabel( OldActor->GetActorLabel() );
-			NewActor->Tags = OldActor->Tags;
-
 			// Allow actor derived classes a chance to replace properties.
 			NewActor->EditorReplacedActor(OldActor);
 
diff --git a/Engine/Source/Runtime/Engine/Classes/GameFramework/Actor.h b/Engine/Source/Runtime/Engine/Classes/GameFramework/Actor.h
--- a/Engine/Source/Runtime/Engine/Classes/GameFramework/Actor.h
+++ b/Engine/Source/Runtime/Engine/Classes/GameFramework/Actor.h
@@ -1241,7 +1241,7 @@ public:
 	bool HasDataLayers() const;
 
 private:
-	TArray<const UDataLayerInstance*> GetDataLayerInstancesInternal(bool bUseLevelContext) const;
+	TArray<const UDataLayerInstance*> GetDataLayerInstancesInternal(bool bUseLevelContext, bool bIncludeParentDataLayers = true) const;
 	bool UseWorldPartitionRuntimeCellDataLayers() const;
 
 	//~=============================================================================
@@ -2372,7 +2372,7 @@ public:
 	virtual void EditorKeyPressed(FKey Key, EInputEvent Event) {}
 
 	/** Called by ReplaceSelectedActors to allow a new actor to copy properties from an old actor when it is replaced */
-	virtual void EditorReplacedActor(AActor* OldActor) {}
+	virtual void EditorReplacedActor(AActor* OldActor);
 
 	/**
 	 * Function that gets called from within Map_Check to allow this actor to check itself
diff --git a/Engine/Source/Runtime/Engine/Private/Actor.cpp b/Engine/Source/Runtime/Engine/Private/Actor.cpp
--- a/Engine/Source/Runtime/Engine/Private/Actor.cpp
+++ b/Engine/Source/Runtime/Engine/Private/Actor.cpp
@@ -5951,7 +5951,7 @@ TArray<const UDataLayerInstance*> AActor::GetDataLayerInstances() const
 // Returns all valid DataLayerInstances for this actor including those inherited from their parent level instance actor.
 // If bUseLevelContext is true, the actor level will be passed down to the DataLayerSubsystem which will use it 
 // to narrow down the resolving of valid datalayers for this particular level.
-TArray<const UDataLayerInstance*> AActor::GetDataLayerInstancesInternal(bool bUseLevelContext) const
+TArray<const UDataLayerInstance*> AActor::GetDataLayerInstancesInternal(bool bUseLevelContext, bool bIncludeParentDataLayers) const
 {
 	if (UseWorldPartitionRuntimeCellDataLayers())
 	{
@@ -5973,13 +5973,16 @@ TArray<const UDataLayerInstance*> AActor::GetDataLayerInstancesInternal(bool bUs
 			PRAGMA_ENABLE_DEPRECATION_WARNINGS
 		}
 
-		// Add parent container Data Layer Instances
-		ULevelInstanceSubsystem* LevelInstanceSubsystem = UWorld::GetSubsystem<ULevelInstanceSubsystem>(GetWorld());
-		if (AActor* ParentActor = LevelInstanceSubsystem ? Cast<AActor>(LevelInstanceSubsystem->GetParentLevelInstance(this)) : nullptr)
+		if (bIncludeParentDataLayers)
 		{
-			for (const UDataLayerInstance* DataLayerInstance : ParentActor->GetDataLayerInstancesInternal(bUseLevelContext))
+			// Add parent container Data Layer Instances
+			ULevelInstanceSubsystem* LevelInstanceSubsystem = UWorld::GetSubsystem<ULevelInstanceSubsystem>(GetWorld());
+			if (AActor* ParentActor = LevelInstanceSubsystem ? Cast<AActor>(LevelInstanceSubsystem->GetParentLevelInstance(this)) : nullptr)
 			{
-				DataLayerInstances.AddUnique(DataLayerInstance);
+				for (const UDataLayerInstance* DataLayerInstance : ParentActor->GetDataLayerInstancesInternal(bUseLevelContext, bIncludeParentDataLayers))
+				{
+					DataLayerInstances.AddUnique(DataLayerInstance);
+				}
 			}
 		}
 		return DataLayerInstances;
diff --git a/Engine/Source/Runtime/Engine/Private/ActorEditor.cpp b/Engine/Source/Runtime/Engine/Private/ActorEditor.cpp
--- a/Engine/Source/Runtime/Engine/Private/ActorEditor.cpp
+++ b/Engine/Source/Runtime/Engine/Private/ActorEditor.cpp
@@ -31,11 +31,14 @@
 #include "Editor.h"
 #include "Engine/LevelStreaming.h"
 #include "WorldPartition/WorldPartitionActorDesc.h"
+#include "WorldPartition/DataLayer/IDataLayerEditorModule.h"
 #include "LevelInstance/LevelInstanceSubsystem.h"
 #include "LevelInstance/LevelInstanceInterface.h"
 #include "Folder.h"
 #include "ActorFolder.h"
 #include "WorldPersistentFolders.h"
+#include "Algo/Transform.h"
+#include "Modules/ModuleManager.h"
 
 #define LOCTEXT_NAMESPACE "ErrorChecking"
 
@@ -1214,6 +1217,34 @@ void AActor::SetFolderPath_Recursively(const FName& NewFolderPath)
 	});
 }
 
+// Transfers some properties from the old actor
+// Ideally, this should be revisited to implement something more generic.
+void AActor::EditorReplacedActor(AActor* OldActor)
+{
+	SetActorLabel(OldActor->GetActorLabel());
+	Tags = OldActor->Tags;
+	
+	SetFolderPath(OldActor->GetFolderPath());
+	
+	if (CanChangeIsSpatiallyLoadedFlag())
+	{
+		SetIsSpatiallyLoaded(OldActor->bIsSpatiallyLoaded);
+	}
+
+	SetRuntimeGrid(OldActor->RuntimeGrid);
+
+	const bool bUseLevelContext = true;
+	const bool bIncludeParentDataLayers = false;
+	TArray<const UDataLayerInstance*> ConstDataLayerInstances = OldActor->GetDataLayerInstancesInternal(bUseLevelContext, bIncludeParentDataLayers);
+	if (!ConstDataLayerInstances.IsEmpty())
+	{
+		TArray<UDataLayerInstance*> DataLayerInstances;
+		Algo::Transform(ConstDataLayerInstances, DataLayerInstances, [](const UDataLayerInstance* ConstDataLayerInstance) { return const_cast<UDataLayerInstance*>(ConstDataLayerInstance); });
+		IDataLayerEditorModule& EditorModule = FModuleManager::LoadModuleChecked<IDataLayerEditorModule>("DataLayerEditor");
+		EditorModule.AddActorToDataLayers(this, DataLayerInstances);
+	}
+}
+
 void AActor::CheckForDeprecated()
 {
 	if ( GetClass()->HasAnyClassFlags(CLASS_Deprecated) )
diff --git a/Engine/Source/Runtime/Engine/Public/WorldPartition/DataLayer/IDataLayerEditorModule.h b/Engine/Source/Runtime/Engine/Public/WorldPartition/DataLayer/IDataLayerEditorModule.h
--- a/Engine/Source/Runtime/Engine/Public/WorldPartition/DataLayer/IDataLayerEditorModule.h
+++ b/Engine/Source/Runtime/Engine/Public/WorldPartition/DataLayer/IDataLayerEditorModule.h
@@ -4,6 +4,9 @@
 #include "CoreMinimal.h"
 #include "Modules/ModuleInterface.h"
 
+class AActor;
+class UDataLayerInstance;
+
 /**
  * The module holding all of the UI related pieces for DataLayer management
  */
@@ -11,4 +14,5 @@ class IDataLayerEditorModule : public IModuleInterface
 {
 public:
 	virtual ~IDataLayerEditorModule() {}
+	virtual bool AddActorToDataLayers(AActor* Actor, const TArray<UDataLayerInstance*>& DataLayers) = 0;
 };
\ No newline at end of file
