UWorldPartitionBuilder - Improved fake engine tick
* Now ticking engine too instead of only graphics
* Unified commandlet engine ticking in new CommandletHelpers::TickEngine() function

diff --git a/Engine/Plugins/NetcodeUnitTest/NetcodeUnitTest/Source/NetcodeUnitTest/Private/UnitTestCommandlet.cpp b/Engine/Plugins/NetcodeUnitTest/NetcodeUnitTest/Source/NetcodeUnitTest/Private/UnitTestCommandlet.cpp
--- a/Engine/Plugins/NetcodeUnitTest/NetcodeUnitTest/Source/NetcodeUnitTest/Private/UnitTestCommandlet.cpp
+++ b/Engine/Plugins/NetcodeUnitTest/NetcodeUnitTest/Source/NetcodeUnitTest/Private/UnitTestCommandlet.cpp
@@ -147,28 +147,7 @@ int32 UUnitTestCommandlet::Main(const FString& Params)
 		// NOTE: This main loop is partly based off of FileServerCommandlet
 		while (GIsRunning && !IsEngineExitRequested())
 		{
-			GEngine->UpdateTimeAndHandleMaxTickRate();
-			GEngine->Tick(FApp::GetDeltaTime(), false);
-
-			if (FSlateApplication::IsInitialized())
-			{
-				FSlateApplication::Get().PumpMessages();
-				FSlateApplication::Get().Tick();
-			}
-
-			// Required for FTimerManager to function - as it blocks ticks, if the frame counter doesn't change
-			GFrameCounter++;
-
-			FTSTicker::GetCoreTicker().Tick(FApp::GetDeltaTime());
-
-			// Execute deferred commands
-			for (int32 DeferredCommandsIndex=0; DeferredCommandsIndex<GEngine->DeferredCommands.Num(); DeferredCommandsIndex++)
-			{
-				GEngine->Exec(UnitTestWorld, *GEngine->DeferredCommands[DeferredCommandsIndex], *GLog);
-			}
-
-			GEngine->DeferredCommands.Empty();
-
+			CommandletHelpers::TickEngine(UnitTestWorld);
 
 			FPlatformProcess::Sleep(0);
 
diff --git a/Engine/Source/Editor/UnrealEd/Private/Commandlets/FileServerCommandlet.cpp b/Engine/Source/Editor/UnrealEd/Private/Commandlets/FileServerCommandlet.cpp
--- a/Engine/Source/Editor/UnrealEd/Private/Commandlets/FileServerCommandlet.cpp
+++ b/Engine/Source/Editor/UnrealEd/Private/Commandlets/FileServerCommandlet.cpp
@@ -109,12 +109,7 @@ int32 UFileServerCommandlet::Main( const FString& Params )
 		FTaskGraphInterface::Get().ProcessThreadUntilIdle(ENamedThreads::GameThread);
 
 		// execute deferred commands
-		for (int32 DeferredCommandsIndex=0; DeferredCommandsIndex<GEngine->DeferredCommands.Num(); DeferredCommandsIndex++)
-		{
-			GEngine->Exec( GWorld, *GEngine->DeferredCommands[DeferredCommandsIndex], *GLog);
-		}
-
-		GEngine->DeferredCommands.Empty();
+		GEngine->TickDeferredCommands();
 
 		// handle server timeout
 		if (InstanceId.IsValid())
diff --git a/Engine/Source/Editor/UnrealEd/Private/WorldPartition/WorldPartitionBuilder.cpp b/Engine/Source/Editor/UnrealEd/Private/WorldPartition/WorldPartitionBuilder.cpp
--- a/Engine/Source/Editor/UnrealEd/Private/WorldPartition/WorldPartitionBuilder.cpp
+++ b/Engine/Source/Editor/UnrealEd/Private/WorldPartition/WorldPartitionBuilder.cpp
@@ -264,6 +264,9 @@ bool UWorldPartitionBuilder::Run(UWorld* World, FPackageSourceControlHelper& Pac
 
 					LoaderAdapters.Emplace_GetRef(MakeUnique<FLoaderAdapterShape>(World, BoundsToLoad, TEXT("Loaded Region")))->Load();
 
+					// Simulate an engine tick
+					FWorldPartitionHelpers::FakeEngineTick(World);
+
 					bResult = RunInternal(World, CellInfo, PackageHelper);
 
 					if (FWorldPartitionHelpers::HasExceededMaxMemory())
@@ -271,17 +274,6 @@ bool UWorldPartitionBuilder::Run(UWorld* World, FPackageSourceControlHelper& Pac
 						LoaderAdapters.Empty();
 						FWorldPartitionHelpers::DoCollectGarbage();
 					}
-
-					// When running with -AllowCommandletRendering we want to simulate an engine tick
-					if (IsAllowCommandletRendering())
-					{
-						FWorldPartitionHelpers::FakeEngineTick(World);
-
-						ENQUEUE_RENDER_COMMAND(VirtualTextureScalability_Release)([](FRHICommandList& RHICmdList)
-						{
-							GetRendererModule().ReleaseVirtualTexturePendingResources();
-						});
-					}
 				}
 			}
 		}
diff --git a/Engine/Source/Runtime/Engine/Classes/Commandlets/Commandlet.h b/Engine/Source/Runtime/Engine/Classes/Commandlets/Commandlet.h
--- a/Engine/Source/Runtime/Engine/Classes/Commandlets/Commandlet.h
+++ b/Engine/Source/Runtime/Engine/Classes/Commandlets/Commandlet.h
@@ -161,5 +161,12 @@ class UCommandlet : public UObject
 
 };
 
-
-
+namespace CommandletHelpers
+{
+	/**
+	 * Simulate an engine frame tick.
+	 * Can be used by commandlets to tick various subsystems. If running with -AllowCommandletRendering, this will
+	 * also tick rendering for the provided world scene(s).
+	 */
+	ENGINE_API void TickEngine(class UWorld* InWorld = nullptr);
+}
\ No newline at end of file
diff --git a/Engine/Source/Runtime/Engine/Private/Commandlets/Commandlet.cpp b/Engine/Source/Runtime/Engine/Private/Commandlets/Commandlet.cpp
--- a/Engine/Source/Runtime/Engine/Private/Commandlets/Commandlet.cpp
+++ b/Engine/Source/Runtime/Engine/Private/Commandlets/Commandlet.cpp
@@ -3,6 +3,13 @@
 #include "Commandlets/Commandlet.h"
 #include "Misc/AutomationTest.h"
 
+#include "Engine/Engine.h"
+#include "EngineModule.h"
+#include "RenderingThread.h"
+#include "RendererInterface.h"
+#include "HAL/ThreadManager.h"
+#include "Framework/Application/SlateApplication.h"
+
 UCommandlet::UCommandlet(const FObjectInitializer& ObjectInitializer)
 	: Super(ObjectInitializer)
 {
@@ -57,3 +64,67 @@ bool FCommandletCommandLineParsingTest::RunTest(const FString& Parameters)
 
 #endif // WITH_DEV_AUTOMATION_TESTS
 
+
+void CommandletHelpers::TickEngine(UWorld* InWorld)
+{
+	// Simulate an engine frame tick
+	// Will make sure systems can perform their internal bookkeeping properly. For example, the VT system needs to 
+	// process deleted VTs.
+
+	// Update FApp::CurrentTime / FApp::DeltaTime while taking into account max tick rate.
+	GEngine->UpdateTimeAndHandleMaxTickRate();
+
+	// Tick the engine.
+	GEngine->Tick(FApp::GetDeltaTime(), false);
+
+	// Tick Slate.
+	if (FSlateApplication::IsInitialized())
+	{
+		FSlateApplication::Get().PumpMessages();
+		FSlateApplication::Get().Tick();
+	}
+
+	// Required for FTimerManager to function - as it blocks ticks, if the frame counter doesn't change
+	GFrameCounter++;
+
+	// Update task graph.
+	FTaskGraphInterface::Get().ProcessThreadUntilIdle(ENamedThreads::GameThread);
+
+	// Ticks all fake threads and their runnable objects.
+	FThreadManager::Get().Tick();
+
+	// Core ticker tick.
+	FTSTicker::GetCoreTicker().Tick(FApp::GetDeltaTime());
+
+	// Execute deferred commands.
+	GEngine->TickDeferredCommands();
+
+	// Tick rendering.
+	if (InWorld && IsAllowCommandletRendering())
+	{
+		if (FSceneInterface* Scene = InWorld->Scene)
+		{
+			// BeingFrame/EndFrame (taken from FEngineLoop)
+
+			ENQUEUE_RENDER_COMMAND(BeginFrame)([](FRHICommandListImmediate& RHICmdList)
+			{
+				GFrameNumberRenderThread++;
+				RHICmdList.BeginFrame();
+				FCoreDelegates::OnBeginFrameRT.Broadcast();
+			});
+
+			ENQUEUE_RENDER_COMMAND(EndFrame)([](FRHICommandListImmediate& RHICmdList)
+			{
+				FCoreDelegates::OnEndFrameRT.Broadcast();
+				RHICmdList.EndFrame();
+			});
+
+			FlushRenderingCommands();
+		}
+
+		ENQUEUE_RENDER_COMMAND(VirtualTextureScalability_Release)([](FRHICommandList& RHICmdList)
+		{
+			GetRendererModule().ReleaseVirtualTexturePendingResources();
+		});
+	}
+}
\ No newline at end of file
diff --git a/Engine/Source/Runtime/Engine/Private/WorldPartition/WorldPartitionHelpers.cpp b/Engine/Source/Runtime/Engine/Private/WorldPartition/WorldPartitionHelpers.cpp
--- a/Engine/Source/Runtime/Engine/Private/WorldPartition/WorldPartitionHelpers.cpp
+++ b/Engine/Source/Runtime/Engine/Private/WorldPartition/WorldPartitionHelpers.cpp
@@ -2,11 +2,11 @@
 
 #include "WorldPartition/WorldPartitionHelpers.h"
 
-#include "Engine/Engine.h"
-#include "RenderingThread.h"
 #include "WorldPartition/WorldPartition.h"
 #include "WorldPartition/WorldPartitionEditorHash.h"
 
+#include "Commandlets/Commandlet.h"
+
 DEFINE_LOG_CATEGORY_STATIC(LogWorldPartitionHelpers, Log, All);
 
 #if WITH_EDITOR
@@ -191,29 +191,7 @@ void FWorldPartitionHelpers::FakeEngineTick(UWorld* InWorld)
 {
 	check(InWorld);
 
-	// Simulate an engine frame tick
-	// Will make sure systems can perform their internal bookkeeping properly. For example, the VT system needs to 
-	// process deleted VTs.
-
-	if (FSceneInterface* Scene = InWorld->Scene)
-	{
-		// BeingFrame/EndFrame (taken from FEngineLoop)
-
-		ENQUEUE_RENDER_COMMAND(BeginFrame)([](FRHICommandListImmediate& RHICmdList)
-			{
-				GFrameNumberRenderThread++;
-				RHICmdList.BeginFrame();
-				FCoreDelegates::OnBeginFrameRT.Broadcast();
-			});
-
-		ENQUEUE_RENDER_COMMAND(EndFrame)([](FRHICommandListImmediate& RHICmdList)
-			{
-				FCoreDelegates::OnEndFrameRT.Broadcast();
-				RHICmdList.EndFrame();
-			});
-
-		FlushRenderingCommands();
-	}
+	CommandletHelpers::TickEngine(InWorld);
 }
 
 #endif // #if WITH_EDITOR
