Add initializers for controller id when constructing key events from a device id, this fixes an infinite loop case with forwarding input events

diff --git a/Engine/Source/Runtime/Engine/Public/UnrealClient.h b/Engine/Source/Runtime/Engine/Public/UnrealClient.h
--- a/Engine/Source/Runtime/Engine/Public/UnrealClient.h
+++ b/Engine/Source/Runtime/Engine/Public/UnrealClient.h
@@ -802,6 +802,7 @@ public:
 
 	FInputKeyEventArgs(FViewport* InViewport, FInputDeviceId InInputDevice, FKey InKey, EInputEvent InEvent)
 		: Viewport(InViewport)
+		, ControllerId(0)
 		, InputDevice(InInputDevice)
 		, Key(InKey)
 		, Event(InEvent)
@@ -822,6 +823,7 @@ public:
 	
 	FInputKeyEventArgs(FViewport* InViewport, FInputDeviceId InInputDevice, FKey InKey, EInputEvent InEvent, float InAmountDepressed, bool bInIsTouchEvent)
 		: Viewport(InViewport)
+		, ControllerId(0)
 		, InputDevice(InInputDevice)
 		, Key(InKey)
 		, Event(InEvent)
