Demote a warning to a log to fix CIS, since it's harmless.

#rnx

diff --git a/Engine/Source/Runtime/Engine/Private/World.cpp b/Engine/Source/Runtime/Engine/Private/World.cpp
--- a/Engine/Source/Runtime/Engine/Private/World.cpp
+++ b/Engine/Source/Runtime/Engine/Private/World.cpp
@@ -1556,7 +1556,7 @@ void UWorld::RepairSingletonActorOfClass(TSubclassOf<AActor> ActorClass)
 			{
 				if (FoundActor)
 				{
-					UE_LOG(LogWorld, Warning, TEXT("Extra '%s' actor found. Resave level %s or actors to clean up."), *CurrentActor->GetPathName(), *PersistentLevel->GetPathName());
+					UE_LOG(LogWorld, Log, TEXT("Extra '%s' actor found. Resave level %s or actors to clean up."), *CurrentActor->GetPathName(), *PersistentLevel->GetPathName());
 
 					if (FoundActor == CurrentActor)
 					{
