Workaround for directional light fallout from new default settings.
- avoid resetting the old defaults if part of a blueprint.

diff --git a/Engine/Source/Runtime/Engine/Private/Components/DirectionalLightComponent.cpp b/Engine/Source/Runtime/Engine/Private/Components/DirectionalLightComponent.cpp
--- a/Engine/Source/Runtime/Engine/Private/Components/DirectionalLightComponent.cpp
+++ b/Engine/Source/Runtime/Engine/Private/Components/DirectionalLightComponent.cpp
@@ -1352,9 +1352,10 @@ void UDirectionalLightComponent::Serialize(FArchive& Ar)
 
 	if (Ar.IsLoading())
 	{
-		// Reset to older defaults if the light was created with a previous version of the engine.
-		if (Ar.CustomVer(FUE5ReleaseStreamObjectVersion::GUID) < FUE5ReleaseStreamObjectVersion::UpdatedDirectionalLightShadowDefaults)
-		{
+		// Reset to older defaults if the light was created with a previous version of the engine & not set up by a blueprint.
+		if (Ar.CustomVer(FUE5ReleaseStreamObjectVersion::GUID) < FUE5ReleaseStreamObjectVersion::UpdatedDirectionalLightShadowDefaults 
+			&& !GetArchetype()->IsInBlueprint())
+		{			
 			DynamicShadowDistanceMovableLight = 20000.0f;
 			DistanceFieldShadowDistance = 30000.0f;
 			DynamicShadowCascades = 3;
