Deprecated local uniform buffers now that RHICreateUniformBuffer works on any thread.

diff --git a/Engine/Source/Runtime/Engine/Private/Materials/MaterialShared.cpp b/Engine/Source/Runtime/Engine/Private/Materials/MaterialShared.cpp
--- a/Engine/Source/Runtime/Engine/Private/Materials/MaterialShared.cpp
+++ b/Engine/Source/Runtime/Engine/Private/Materials/MaterialShared.cpp
@@ -3485,7 +3485,7 @@ void FUniformExpressionCache::ResetAllocatedVTs()
 	OwnedAllocatedVTs.Reset();
 }
 
-void FMaterialRenderProxy::EvaluateUniformExpressions(FUniformExpressionCache& OutUniformExpressionCache, const FMaterialRenderContext& Context, FRHIComputeCommandList* CommandListIfLocalMode) const
+void FMaterialRenderProxy::EvaluateUniformExpressions(FUniformExpressionCache& OutUniformExpressionCache, const FMaterialRenderContext& Context, FRHIComputeCommandList*) const
 {
 	check(IsInParallelRenderingThread());
 
@@ -3534,29 +3534,21 @@ void FMaterialRenderProxy::EvaluateUniformExpressions(FUniformExpressionCache& O
 	check(TempBuffer != nullptr);
 	UniformExpressionSet.FillUniformBuffer(Context, OutUniformExpressionCache, UniformBufferLayout, TempBuffer, UniformBufferLayout->ConstantBufferSize);
 
-	if (CommandListIfLocalMode)
+	if (IsValidRef(OutUniformExpressionCache.UniformBuffer) && !OutUniformExpressionCache.UniformBuffer->IsValid())
 	{
-		OutUniformExpressionCache.LocalUniformBuffer = CommandListIfLocalMode->BuildLocalUniformBuffer(TempBuffer, UniformBufferLayout->ConstantBufferSize, UniformBufferLayout);
-		check(OutUniformExpressionCache.LocalUniformBuffer.IsValid());
+		UE_LOG(LogMaterial, Fatal, TEXT("The Uniformbuffer needs to be valid if it has been set"));
+	}
+
+	if (IsValidRef(OutUniformExpressionCache.UniformBuffer))
+	{
+		// The actual pointer may not match because there are cases (in the editor, during the shader compilation) when material's shader map gets updated without proxy's cache
+		// getting invalidated, but the layout contents must match.
+		check(OutUniformExpressionCache.UniformBuffer->GetLayoutPtr() == UniformBufferLayout || *OutUniformExpressionCache.UniformBuffer->GetLayoutPtr() == *UniformBufferLayout);
+		RHIUpdateUniformBuffer(OutUniformExpressionCache.UniformBuffer, TempBuffer);
 	}
 	else
 	{
-		if (IsValidRef(OutUniformExpressionCache.UniformBuffer) && !OutUniformExpressionCache.UniformBuffer->IsValid())
-		{
-			UE_LOG(LogMaterial, Fatal, TEXT("The Uniformbuffer needs to be valid if it has been set"));
-		}
-
-		if (IsValidRef(OutUniformExpressionCache.UniformBuffer))
-		{
-			// The actual pointer may not match because there are cases (in the editor, during the shader compilation) when material's shader map gets updated without proxy's cache
-			// getting invalidated, but the layout contents must match.
-			check(OutUniformExpressionCache.UniformBuffer->GetLayoutPtr() == UniformBufferLayout || *OutUniformExpressionCache.UniformBuffer->GetLayoutPtr() == *UniformBufferLayout);
-			RHIUpdateUniformBuffer(OutUniformExpressionCache.UniformBuffer, TempBuffer);
-		}
-		else
-		{
-			OutUniformExpressionCache.UniformBuffer = RHICreateUniformBuffer(TempBuffer, UniformBufferLayout, UniformBuffer_MultiFrame);
-		}
+		OutUniformExpressionCache.UniformBuffer = RHICreateUniformBuffer(TempBuffer, UniformBufferLayout, UniformBuffer_MultiFrame);
 	}
 
 	OutUniformExpressionCache.ParameterCollections = UniformExpressionSet.ParameterCollections;
diff --git a/Engine/Source/Runtime/Engine/Private/Particles/ParticleGpuSimulation.cpp b/Engine/Source/Runtime/Engine/Private/Particles/ParticleGpuSimulation.cpp
--- a/Engine/Source/Runtime/Engine/Private/Particles/ParticleGpuSimulation.cpp
+++ b/Engine/Source/Runtime/Engine/Private/Particles/ParticleGpuSimulation.cpp
@@ -1386,7 +1386,6 @@ struct FSimulationCommandGPU
 	FParticlePerFrameSimulationParameters PerFrameParameters;
 	/** Parameters to sample the local vector field for this simulation. */
 	FVectorFieldUniformBufferRef VectorFieldsUniformBuffer;
-	FLocalUniformBuffer VectorFieldsUniformBufferLocal;
 	/** Vector field volume textures for this simulation. */
 	FRHITexture3D* VectorFieldTexturesRHI[MAX_VECTOR_FIELDS];
 	/** The number of tiles to simulate. */
diff --git a/Engine/Source/Runtime/Engine/Public/MaterialShared.h b/Engine/Source/Runtime/Engine/Public/MaterialShared.h
--- a/Engine/Source/Runtime/Engine/Public/MaterialShared.h
+++ b/Engine/Source/Runtime/Engine/Public/MaterialShared.h
@@ -2406,7 +2406,7 @@ struct FUniformExpressionCache
 {
 	/** Material uniform buffer. */
 	FUniformBufferRHIRef UniformBuffer;
-	/** Material uniform buffer. */
+	UE_DEPRECATED(5.1, "Local uniform buffers have been deprecated.")
 	FLocalUniformBuffer LocalUniformBuffer;
 	/** Allocated virtual textures, one for each entry in FUniformExpressionSet::VTStacks */
 	TArray<IAllocatedVirtualTexture*> AllocatedVTs;
diff --git a/Engine/Source/Runtime/RHI/Public/RHICommandList.h b/Engine/Source/Runtime/RHI/Public/RHICommandList.h
--- a/Engine/Source/Runtime/RHI/Public/RHICommandList.h
+++ b/Engine/Source/Runtime/RHI/Public/RHICommandList.h
@@ -2426,6 +2426,7 @@ public:
 		SetShaderUniformBuffer(Shader.GetReference(), BaseIndex, UniformBuffer);
 	}
 
+	UE_DEPRECATED(5.1, "Local uniform buffers are now deprecated. Use RHICreateUniformBuffer instead.")
 	FORCEINLINE_DEBUGGABLE FLocalUniformBuffer BuildLocalUniformBuffer(const void* Contents, uint32 ContentsSize, const FRHIUniformBufferLayout* Layout)
 	{
 		FLocalUniformBuffer Result;
@@ -2445,10 +2446,13 @@ public:
 	UE_DEPRECATED(5.0, "Use Layout pointers instead")
 	FORCEINLINE_DEBUGGABLE FLocalUniformBuffer BuildLocalUniformBuffer(const void* Contents, uint32 ContentsSize, const FRHIUniformBufferLayout& Layout)
 	{
+	PRAGMA_DISABLE_DEPRECATION_WARNINGS
 		return BuildLocalUniformBuffer(Contents, ContentsSize, &Layout);
+	PRAGMA_ENABLE_DEPRECATION_WARNINGS
 	}
 
 	template <typename TRHIShader>
+	UE_DEPRECATED(5.1, "Local uniform buffers are now deprecated. Use RHICreateUniformBuffer instead.")
 	FORCEINLINE_DEBUGGABLE void SetLocalShaderUniformBuffer(TRHIShader* Shader, uint32 BaseIndex, const FLocalUniformBuffer& UniformBuffer)
 	{
 		ValidateBoundShader(Shader);
diff --git a/Engine/Source/Runtime/RHI/Public/RHIValidation.h b/Engine/Source/Runtime/RHI/Public/RHIValidation.h
--- a/Engine/Source/Runtime/RHI/Public/RHIValidation.h
+++ b/Engine/Source/Runtime/RHI/Public/RHIValidation.h
@@ -362,7 +362,6 @@ public:
 	// FlushType: Thread safe, but varies depending on the RHI override final
 	virtual FUniformBufferRHIRef RHICreateUniformBuffer(const void* Contents, const FRHIUniformBufferLayout* Layout, EUniformBufferUsage Usage, EUniformBufferValidation Validation) override final
 	{
-		// RHICreateUniformBuffer can be called from the RenderThread (normal use case) or the RHI Thread (FLocalUniformBuffer)
 		check(Layout->Resources.Num() > 0 || Layout->ConstantBufferSize > 0);
 		FUniformBufferRHIRef UniformBuffer = RHI->RHICreateUniformBuffer(Contents, Layout, Usage, Validation);
 		// Use the render thread frame ID for any non RHI thread allocations.
diff --git a/Engine/Source/Runtime/RenderCore/Public/ShaderParameterMacros.h b/Engine/Source/Runtime/RenderCore/Public/ShaderParameterMacros.h
--- a/Engine/Source/Runtime/RenderCore/Public/ShaderParameterMacros.h
+++ b/Engine/Source/Runtime/RenderCore/Public/ShaderParameterMacros.h
@@ -119,6 +119,7 @@ public:
 	}
 
 	/** Creates a uniform buffer with the given value, and returns a structured reference to it. */
+	UE_DEPRECATED(5.1, "Local uniform buffers are now deprecated. Use RHICreateUniformBuffer instead.")
 	static FLocalUniformBuffer CreateLocalUniformBuffer(FRHICommandList& RHICmdList, const TBufferStruct& Value, EUniformBufferUsage Usage)
 	{
 		return RHICmdList.BuildLocalUniformBuffer(&Value, sizeof(TBufferStruct), TBufferStruct::FTypeInfo::GetStructMetadata()->GetLayoutPtr());
diff --git a/Engine/Source/Runtime/RenderCore/Public/ShaderParameterUtils.h b/Engine/Source/Runtime/RenderCore/Public/ShaderParameterUtils.h
--- a/Engine/Source/Runtime/RenderCore/Public/ShaderParameterUtils.h
+++ b/Engine/Source/Runtime/RenderCore/Public/ShaderParameterUtils.h
@@ -496,6 +496,7 @@ inline void FRWShaderParameter::UnsetUAV(TRHICmdList& RHICmdList, FRHIComputeSha
 
 /** Sets the value of a shader uniform buffer parameter to a uniform buffer containing the struct. */
 template<typename TShaderRHIRef, typename TRHICmdList>
+UE_DEPRECATED(5.1, "Local uniform buffers are now deprecated. Use SetUniformBufferParameter instead.")
 inline void SetLocalUniformBufferParameter(
 	TRHICmdList& RHICmdList,
 	const TShaderRHIRef& Shader,
@@ -507,7 +508,9 @@ inline void SetLocalUniformBufferParameter(
 	checkSlow(Parameter.IsInitialized());
 	if(Parameter.IsBound())
 	{
+	PRAGMA_DISABLE_DEPRECATION_WARNINGS
 		RHICmdList.SetLocalShaderUniformBuffer(Shader, Parameter.GetBaseIndex(), LocalUniformBuffer);
+	PRAGMA_ENABLE_DEPRECATION_WARNINGS
 	}
 }
 
@@ -581,13 +584,11 @@ inline void SetUniformBufferParameterImmediate(
 	checkSlow(Parameter.IsInitialized());
 	if(Parameter.IsBound())
 	{
-		const FRHIUniformBufferLayout* UniformBufferLayout = TBufferStruct::StaticStructMetadata.GetLayoutPtr();
-		FLocalUniformBuffer UniformBuffer = RHICmdList.BuildLocalUniformBuffer(&UniformBufferValue, UniformBufferLayout->ConstantBufferSize, UniformBufferLayout);
-
-		RHICmdList.SetLocalShaderUniformBuffer(
+		RHICmdList.SetShaderUniformBuffer(
 			Shader,
 			Parameter.GetBaseIndex(),
-			UniformBuffer);
+			RHICreateUniformBuffer(&UniformBufferValue, &TBufferStruct::StaticStructMetadata.GetLayout(), UniformBuffer_SingleDraw)
+		);
 	}
 }
 
@@ -604,12 +605,10 @@ inline void SetUniformBufferParameterImmediate(
 	checkSlow(Parameter.IsInitialized());
 	if(Parameter.IsBound())
 	{
-		const FRHIUniformBufferLayout* UniformBufferLayout = TBufferStruct::StaticStructMetadata.GetLayoutPtr();
-		FLocalUniformBuffer UniformBuffer = RHICmdList.BuildLocalUniformBuffer(&UniformBufferValue, UniformBufferLayout->ConstantBufferSize, UniformBufferLayout);
-
-		RHICmdList.SetLocalShaderUniformBuffer(
+		RHICmdList.SetShaderUniformBuffer(
 			Shader,
 			Parameter.GetBaseIndex(),
-			UniformBuffer);
+			RHICreateUniformBuffer(&UniformBufferValue, &TBufferStruct::StaticStructMetadata.GetLayout(), UniformBuffer_SingleDraw)
+		);
 	}
 }
\ No newline at end of file
diff --git a/Engine/Source/Runtime/Renderer/Private/ShaderBaseClasses.cpp b/Engine/Source/Runtime/Renderer/Private/ShaderBaseClasses.cpp
--- a/Engine/Source/Runtime/Renderer/Private/ShaderBaseClasses.cpp
+++ b/Engine/Source/Runtime/Renderer/Private/ShaderBaseClasses.cpp
@@ -131,6 +131,7 @@ void FMaterialShader::VerifyExpressionAndShaderMaps(const FMaterialRenderProxy*
 		|| UniformExpressionCache->CachedUniformExpressionShaderMap != ShaderMap;
 	if (!bUniformExpressionSetMismatch)
 	{
+	PRAGMA_DISABLE_DEPRECATION_WARNINGS
 		if (UniformExpressionCache->LocalUniformBuffer.IsValid())
 		{
 			if (UniformExpressionCache->LocalUniformBuffer.BypassUniform)
@@ -154,6 +155,7 @@ void FMaterialShader::VerifyExpressionAndShaderMaps(const FMaterialRenderProxy*
 				}
 			}
 		}
+	PRAGMA_ENABLE_DEPRECATION_WARNINGS
 		else
 		{
 			if (DebugUniformExpressionUBLayout.GetHash() != UniformExpressionCache->UniformBuffer->GetLayout().GetHash())
@@ -248,14 +250,11 @@ void FMaterialShader::SetParameters(
 		FMaterialRenderContext MaterialRenderContext(MaterialRenderProxy, Material, &View);
 		bUniformExpressionCacheNeedsDelete = true;
 		UniformExpressionCache = new FUniformExpressionCache();
-		MaterialRenderProxy->EvaluateUniformExpressions(*UniformExpressionCache, MaterialRenderContext, &RHICmdList);
-		SetLocalUniformBufferParameter(RHICmdList, ShaderRHI, MaterialUniformBuffer, UniformExpressionCache->LocalUniformBuffer);
-	}
-	else
-	{
-		SetUniformBufferParameter(RHICmdList, ShaderRHI, MaterialUniformBuffer, UniformExpressionCache->UniformBuffer);
+		MaterialRenderProxy->EvaluateUniformExpressions(*UniformExpressionCache, MaterialRenderContext);
 	}
 
+	SetUniformBufferParameter(RHICmdList, ShaderRHI, MaterialUniformBuffer, UniformExpressionCache->UniformBuffer);
+
 #if !(UE_BUILD_TEST || UE_BUILD_SHIPPING || !WITH_EDITOR)
 	VerifyExpressionAndShaderMaps(MaterialRenderProxy, Material, UniformExpressionCache);
 #endif
