World Partition: streaming policy is now only serialized for cooked builds and PIE duplication.

#rnx

diff --git a/Engine/Source/Runtime/Core/Public/UObject/UE5MainStreamObjectVersion.h b/Engine/Source/Runtime/Core/Public/UObject/UE5MainStreamObjectVersion.h
--- a/Engine/Source/Runtime/Core/Public/UObject/UE5MainStreamObjectVersion.h
+++ b/Engine/Source/Runtime/Core/Public/UObject/UE5MainStreamObjectVersion.h
@@ -233,6 +233,9 @@ struct CORE_API FUE5MainStreamObjectVersion
 		// GPU none interpolated spawning no longer calls the update script
 		FixGpuAlwaysRunningUpdateScriptNoneInterpolated,
 
+		// World partition streaming policy serialization only for cooked builds
+		WorldPartitionSerializeStreamingPolicyOnCook,
+
 		// -----<new versions can be added above this line>-------------------------------------------------
 		VersionPlusOne,
 		LatestVersion = VersionPlusOne - 1
diff --git a/Engine/Source/Runtime/Engine/Private/WorldPartition/WorldPartition.cpp b/Engine/Source/Runtime/Engine/Private/WorldPartition/WorldPartition.cpp
--- a/Engine/Source/Runtime/Engine/Private/WorldPartition/WorldPartition.cpp
+++ b/Engine/Source/Runtime/Engine/Private/WorldPartition/WorldPartition.cpp
@@ -18,6 +18,7 @@
 #include "Engine/LevelStreaming.h"
 #include "GameFramework/WorldSettings.h"
 #include "ProfilingDebugging/ScopedTimers.h"
+#include "UObject/UE5MainStreamObjectVersion.h"
 
 #if WITH_EDITOR
 #include "LevelUtils.h"
@@ -1051,6 +1052,24 @@ void UWorldPartition::UnhashActorDesc(FWorldPartitionActorDesc* ActorDesc)
 }
 #endif
 
+void UWorldPartition::Serialize(FArchive& Ar)
+{
+	Ar.UsingCustomVersion(FUE5MainStreamObjectVersion::GUID);
+
+	Super::Serialize(Ar);
+
+	if (Ar.CustomVer(FUE5MainStreamObjectVersion::GUID) >= FUE5MainStreamObjectVersion::WorldPartitionSerializeStreamingPolicyOnCook)
+	{
+		bool bCooked = Ar.IsCooking();
+		Ar << bCooked;
+
+		if (bCooked)
+		{
+			Ar << StreamingPolicy;
+		}
+	}
+}
+
 bool UWorldPartition::ResolveSubobject(const TCHAR* SubObjectPath, UObject*& OutObject, bool bLoadIfExists)
 {
 	if (GetWorld())
diff --git a/Engine/Source/Runtime/Engine/Public/WorldPartition/WorldPartition.h b/Engine/Source/Runtime/Engine/Public/WorldPartition/WorldPartition.h
--- a/Engine/Source/Runtime/Engine/Public/WorldPartition/WorldPartition.h
+++ b/Engine/Source/Runtime/Engine/Public/WorldPartition/WorldPartition.h
@@ -136,6 +136,7 @@ public:
 	//~ End UActorDescContainer Interface
 
 	//~ Begin UObject Interface
+	virtual void Serialize(FArchive& Ar) override;
 	virtual bool ResolveSubobject(const TCHAR* SubObjectPath, UObject*& OutObject, bool bLoadIfExists) override;
 	static void AddReferencedObjects(UObject* InThis, FReferenceCollector& Collector);
 	//~ End UObject Interface
@@ -269,7 +270,7 @@ private:
 	EWorldPartitionInitState InitState;
 	TOptional<FTransform> InstanceTransform;
 
-	UPROPERTY()
+	UPROPERTY(Transient)
 	mutable TObjectPtr<UWorldPartitionStreamingPolicy> StreamingPolicy;
 
 #if WITH_EDITORONLY_DATA
