LevelInstance: Avoid recalculating bounds

diff --git a/Engine/Source/Runtime/Engine/Private/LevelInstance/LevelInstanceLevelStreaming.cpp b/Engine/Source/Runtime/Engine/Private/LevelInstance/LevelInstanceLevelStreaming.cpp
--- a/Engine/Source/Runtime/Engine/Private/LevelInstance/LevelInstanceLevelStreaming.cpp
+++ b/Engine/Source/Runtime/Engine/Private/LevelInstance/LevelInstanceLevelStreaming.cpp
@@ -19,6 +19,9 @@
 
 ULevelStreamingLevelInstance::ULevelStreamingLevelInstance(const FObjectInitializer& ObjectInitializer)
 	: Super(ObjectInitializer)
+#if WITH_EDITOR
+	, CachedBounds(ForceInit)
+#endif
 {
 #if WITH_EDITOR
 	SetShouldBeVisibleInEditor(true);
@@ -53,7 +56,20 @@ TOptional<FFolder::FRootObject> ULevelStreamingLevelInstance::GetFolderRootObjec
 FBox ULevelStreamingLevelInstance::GetBounds() const
 {
 	check(GetLoadedLevel());
-	return ALevelBounds::CalculateLevelBounds(GetLoadedLevel());
+	FTransform LevelInstanceTransform = Cast<AActor>(GetLevelInstance())->GetTransform();
+	if (!CachedBounds.IsValid || !CachedTransform.Equals(LevelInstanceTransform))
+	{
+		CachedTransform = LevelInstanceTransform;
+		CachedBounds = ALevelBounds::CalculateLevelBounds(GetLoadedLevel());
+		
+		// Possible if Level has no bounds relevant actors
+		if (!CachedBounds.IsValid)
+		{
+			CachedBounds = FBox(CachedTransform.GetLocation(), CachedTransform.GetLocation());
+		}
+	}
+	check(CachedBounds.IsValid);
+	return CachedBounds;
 }
 
 void ULevelStreamingLevelInstance::OnLoadedActorAddedToLevel(AActor& InActor)
diff --git a/Engine/Source/Runtime/Engine/Public/LevelInstance/LevelInstanceLevelStreaming.h b/Engine/Source/Runtime/Engine/Public/LevelInstance/LevelInstanceLevelStreaming.h
--- a/Engine/Source/Runtime/Engine/Public/LevelInstance/LevelInstanceLevelStreaming.h
+++ b/Engine/Source/Runtime/Engine/Public/LevelInstance/LevelInstanceLevelStreaming.h
@@ -42,6 +42,9 @@ private:
 	void OnLoadedActorRemovedFromLevel(AActor& InActor);
 
 	TWeakObjectPtr<ALevelInstanceEditorInstanceActor> LevelInstanceEditorInstanceActor;
+
+	mutable FTransform CachedTransform;
+	mutable FBox CachedBounds;
 #endif
 	FLevelInstanceID LevelInstanceID;
 };
\ No newline at end of file
