Fixed duplicate actors in Level. Could happen when renaming a Level that contains actors that return true in AActor::ShouldLevelKeepRefIfExternal.

diff --git a/Engine/Source/Runtime/Engine/Private/Level.cpp b/Engine/Source/Runtime/Engine/Private/Level.cpp
--- a/Engine/Source/Runtime/Engine/Private/Level.cpp
+++ b/Engine/Source/Runtime/Engine/Private/Level.cpp
@@ -1107,6 +1107,7 @@ void ULevel::PostLoad()
 				}
 			}
 
+			TSet<AActor*> ActorsSet(Actors);
 			for (int32 i=0; i < ActorPackageNames.Num(); i++)
 			{
 				const FString& ActorPackageName = ActorPackageNames[i];
@@ -1121,12 +1122,17 @@ void ULevel::PostLoad()
 				ActorPackage = LoadPackage(ActorPackage, *ActorPackageName, bPackageForPIE ? LOAD_PackageForPIE : LOAD_None, nullptr, &InstancingContext);
 
 				int32 PreviousActorCount = Actors.Num();
-				ForEachObjectWithPackage(ActorPackage, [this](UObject* PackageObject)
+				ForEachObjectWithPackage(ActorPackage, [this, &ActorsSet](UObject* PackageObject)
 				{
 					// There might be multiple actors per package in the case where an actor as a child actor component as we put child actor in the same package as their parent
 					if (PackageObject->IsA<AActor>() && !PackageObject->IsTemplate())
 					{
-						Actors.Add((AActor*)PackageObject);
+						AActor* Actor = (AActor*)PackageObject;
+						// Verity that the actor is not already in the array (this is valid if, during last save, the actor returned true in AActor::ShouldLevelKeepRefIfExternal)
+						if (!ActorsSet.Contains(Actor))
+						{
+							Actors.Add(Actor);
+						}
 					}
 					return true;
 				}, false);
