World Partition: minor optimization when adding actor descriptors to load in deferred context + better checks.

#rnx

diff --git a/Engine/Source/Runtime/Engine/Private/WorldPartition/WorldPartitionHandle.cpp b/Engine/Source/Runtime/Engine/Private/WorldPartition/WorldPartitionHandle.cpp
--- a/Engine/Source/Runtime/Engine/Private/WorldPartition/WorldPartitionHandle.cpp
+++ b/Engine/Source/Runtime/Engine/Private/WorldPartition/WorldPartitionHandle.cpp
@@ -105,7 +105,7 @@ FWorldPartitionLoadingContext::FDeferred::~FDeferred()
 		const FTransform& ContainerTransform = Container->GetInstanceTransform();
 		const FTransform* ContainerTransformPtr = ContainerTransform.Equals(FTransform::Identity) ? nullptr : &ContainerTransform;
 
-		auto CreateActorList = [](TArray<AActor*>& ActorList, const TArray<FWorldPartitionActorDesc*>& SourceList) -> ULevel*
+		auto CreateActorList = [](TArray<AActor*>& ActorList, const TSet<FWorldPartitionActorDesc*>& SourceList) -> ULevel*
 		{
 			ActorList.Empty(SourceList.Num());
 
@@ -162,7 +162,10 @@ void FWorldPartitionLoadingContext::FDeferred::RegisterActor(FWorldPartitionActo
 	UActorDescContainer* Container = ActorDesc->GetContainer();
 	check(Container);
 
-	ContainerOps.FindOrAdd(Container).Registrations.Add(ActorDesc);
+	bool bIsAlreadyInSetPtr;
+	ContainerOps.FindOrAdd(Container).Registrations.Add(ActorDesc, &bIsAlreadyInSetPtr);
+	check(!bIsAlreadyInSetPtr);
+
 	check(!ContainerOps.FindChecked(Container).Unregistrations.Contains(ActorDesc));
 
 	ActorDesc->Load();
@@ -175,7 +178,10 @@ void FWorldPartitionLoadingContext::FDeferred::UnregisterActor(FWorldPartitionAc
 	UActorDescContainer* Container = ActorDesc->GetContainer();
 	check(Container);
 
-	ContainerOps.FindOrAdd(Container).Unregistrations.Add(ActorDesc);
+	bool bIsAlreadyInSetPtr;
+	ContainerOps.FindOrAdd(Container).Unregistrations.Add(ActorDesc, &bIsAlreadyInSetPtr);
+	check(!bIsAlreadyInSetPtr);
+
 	check(!ContainerOps.FindChecked(Container).Registrations.Contains(ActorDesc));
 }
 
diff --git a/Engine/Source/Runtime/Engine/Public/WorldPartition/WorldPartitionHandle.h b/Engine/Source/Runtime/Engine/Public/WorldPartition/WorldPartitionHandle.h
--- a/Engine/Source/Runtime/Engine/Public/WorldPartition/WorldPartitionHandle.h
+++ b/Engine/Source/Runtime/Engine/Public/WorldPartition/WorldPartitionHandle.h
@@ -57,8 +57,8 @@ public:
 
 		struct FContainerOps
 		{
-			TArray<FWorldPartitionActorDesc*> Registrations;
-			TArray<FWorldPartitionActorDesc*> Unregistrations;
+			TSet<FWorldPartitionActorDesc*> Registrations;
+			TSet<FWorldPartitionActorDesc*> Unregistrations;
 		};
 
 		TMap<UActorDescContainer*, FContainerOps> ContainerOps;
