LevelInstance: Update FLevelInstanceActorDesc bounds when Level Instance is saved

diff --git a/Engine/Source/Runtime/Engine/Private/LevelInstance/LevelInstanceSubsystem.cpp b/Engine/Source/Runtime/Engine/Private/LevelInstance/LevelInstanceSubsystem.cpp
--- a/Engine/Source/Runtime/Engine/Private/LevelInstance/LevelInstanceSubsystem.cpp
+++ b/Engine/Source/Runtime/Engine/Private/LevelInstance/LevelInstanceSubsystem.cpp
@@ -2108,6 +2108,12 @@ bool ULevelInstanceSubsystem::CommitLevelInstanceInternal(TUniquePtr<FLevelInsta
 	// Wait for Level Instances to be loaded
 	BlockOnLoading();
 
+	// Send out Event if changes were committed
+	if (bChangesCommitted)
+	{
+		LevelInstanceChangedEvent.Broadcast(FName(*EditPackage));
+	}
+
 	GEngine->BroadcastLevelActorListChanged();
 
 	// Restore actor editor context
diff --git a/Engine/Source/Runtime/Engine/Private/WorldPartition/ActorDescContainer.cpp b/Engine/Source/Runtime/Engine/Private/WorldPartition/ActorDescContainer.cpp
--- a/Engine/Source/Runtime/Engine/Private/WorldPartition/ActorDescContainer.cpp
+++ b/Engine/Source/Runtime/Engine/Private/WorldPartition/ActorDescContainer.cpp
@@ -208,6 +208,7 @@ void UActorDescContainer::OnObjectPreSave(UObject* Object, FObjectPreSaveContext
 					// Transfer any internal values not coming from the actor
 					NewActorDesc->TransferFrom(ExistingActorDesc->Get());
 
+					// This moves the TUniquePtr (not the content of the actor desc)
 					*ExistingActorDesc = MoveTemp(NewActorDesc);
 
 					OnActorDescUpdated(ExistingActorDesc->Get());
diff --git a/Engine/Source/Runtime/Engine/Private/WorldPartition/LevelInstance/LevelInstanceActorDesc.cpp b/Engine/Source/Runtime/Engine/Private/WorldPartition/LevelInstance/LevelInstanceActorDesc.cpp
--- a/Engine/Source/Runtime/Engine/Private/WorldPartition/LevelInstance/LevelInstanceActorDesc.cpp
+++ b/Engine/Source/Runtime/Engine/Private/WorldPartition/LevelInstance/LevelInstanceActorDesc.cpp
@@ -63,6 +63,26 @@ bool FLevelInstanceActorDesc::Equals(const FWorldPartitionActorDesc* Other) cons
 	return false;
 }
 
+void FLevelInstanceActorDesc::UpdateBounds()
+{
+	ULevelInstanceSubsystem* LevelInstanceSubsystem = UWorld::GetSubsystem<ULevelInstanceSubsystem>(LevelInstanceContainer->GetWorld());
+	check(LevelInstanceSubsystem);
+
+	FTransform LevelInstancePivotOffsetTransform = FTransform(ULevel::GetLevelInstancePivotOffsetFromPackage(LevelPackage));
+	FTransform FinalLevelTransform = LevelInstancePivotOffsetTransform * LevelInstanceTransform;
+	FBox ContainerBounds = LevelInstanceSubsystem->GetContainerBounds(LevelPackage).TransformBy(FinalLevelTransform);
+
+	ContainerBounds.GetCenterAndExtents(BoundsLocation, BoundsExtent);
+}
+
+void FLevelInstanceActorDesc::OnLevelInstanceChanged(FName PackageName)
+{
+	if (LevelInstanceContainer.IsValid() && (LevelInstanceContainer->GetContainerPackage() == PackageName))
+	{
+		UpdateBounds();
+	}
+}
+
 void FLevelInstanceActorDesc::RegisterContainerInstance(UWorld* InWorld)
 {
 	if (InWorld)
@@ -73,15 +93,12 @@ void FLevelInstanceActorDesc::RegisterContainerInstance(UWorld* InWorld)
 		{
 			ULevelInstanceSubsystem* LevelInstanceSubsystem = UWorld::GetSubsystem<ULevelInstanceSubsystem>(InWorld);
 			check(LevelInstanceSubsystem);
+			LevelInstanceSubsystem->OnLevelInstanceChanged().AddRaw(this, &FLevelInstanceActorDesc::OnLevelInstanceChanged);
 
 			LevelInstanceContainer = LevelInstanceSubsystem->RegisterContainer(LevelPackage);
 			check(LevelInstanceContainer.IsValid());
 
-			FTransform LevelInstancePivotOffsetTransform = FTransform(ULevel::GetLevelInstancePivotOffsetFromPackage(LevelPackage));
-			FTransform FinalLevelTransform = LevelInstancePivotOffsetTransform * LevelInstanceTransform;
-			FBox ContainerBounds = LevelInstanceSubsystem->GetContainerBounds(LevelPackage).TransformBy(FinalLevelTransform);
-
-			ContainerBounds.GetCenterAndExtents(BoundsLocation, BoundsExtent);
+			UpdateBounds();
 		}
 	}
 }
@@ -93,6 +110,8 @@ void FLevelInstanceActorDesc::UnregisterContainerInstance()
 		ULevelInstanceSubsystem* LevelInstanceSubsystem = UWorld::GetSubsystem<ULevelInstanceSubsystem>(LevelInstanceContainer->GetWorld());
 		check(LevelInstanceSubsystem);
 
+		LevelInstanceSubsystem->OnLevelInstanceChanged().RemoveAll(this);
+
 		LevelInstanceSubsystem->UnregisterContainer(LevelInstanceContainer.Get());
 		LevelInstanceContainer.Reset();
 	}
@@ -165,8 +184,13 @@ void FLevelInstanceActorDesc::TransferFrom(const FWorldPartitionActorDesc* From)
 	FWorldPartitionActorDesc::TransferFrom(From);
 
 	FLevelInstanceActorDesc* FromLevelInstanceActorDesc = (FLevelInstanceActorDesc*)From;
-	LevelInstanceContainer = FromLevelInstanceActorDesc->LevelInstanceContainer;
-	FromLevelInstanceActorDesc->LevelInstanceContainer = nullptr;
+
+	// Use the Register/Unregister so callbacks are added/removed
+	if (FromLevelInstanceActorDesc->LevelInstanceContainer.IsValid())
+	{
+		RegisterContainerInstance(Container->World);
+		FromLevelInstanceActorDesc->UnregisterContainerInstance();
+	}
 }
 
 void FLevelInstanceActorDesc::Serialize(FArchive& Ar)
@@ -216,4 +240,4 @@ void FLevelInstanceActorDesc::Serialize(FArchive& Ar)
 	}
 }
 
-#endif
\ No newline at end of file
+#endif
diff --git a/Engine/Source/Runtime/Engine/Public/LevelInstance/LevelInstanceSubsystem.h b/Engine/Source/Runtime/Engine/Public/LevelInstance/LevelInstanceSubsystem.h
--- a/Engine/Source/Runtime/Engine/Public/LevelInstance/LevelInstanceSubsystem.h
+++ b/Engine/Source/Runtime/Engine/Public/LevelInstance/LevelInstanceSubsystem.h
@@ -117,6 +117,10 @@ public:
 	static bool CheckForLoop(const ILevelInstanceInterface* LevelInstance, TSoftObjectPtr<UWorld> WorldAsset, TArray<TPair<FText, TSoftObjectPtr<UWorld>>>* LoopInfo = nullptr, const ILevelInstanceInterface** LoopStart = nullptr);
 	static bool CanUseWorldAsset(const ILevelInstanceInterface* LevelInstance, TSoftObjectPtr<UWorld> WorldAsset, FString* OutReason);
 	static bool CanUsePackage(FName InPackageName);
+
+	/** Editor-only event triggered when level instance is committed with changes */
+	DECLARE_EVENT_OneParam(ULevelInstanceSubsystem, FLevelInstanceChanged, FName);
+	FLevelInstanceChanged& OnLevelInstanceChanged() { return LevelInstanceChangedEvent; }
 #endif
 
 private:
@@ -223,6 +227,8 @@ private:
 #if WITH_EDITOR
 	bool bIsCreatingLevelInstance;
 	bool bIsCommittingLevelInstance;
+
+	FLevelInstanceChanged LevelInstanceChangedEvent;
 #endif
 
 	struct FLevelInstance
diff --git a/Engine/Source/Runtime/Engine/Public/WorldPartition/LevelInstance/LevelInstanceActorDesc.h b/Engine/Source/Runtime/Engine/Public/WorldPartition/LevelInstance/LevelInstanceActorDesc.h
--- a/Engine/Source/Runtime/Engine/Public/WorldPartition/LevelInstance/LevelInstanceActorDesc.h
+++ b/Engine/Source/Runtime/Engine/Public/WorldPartition/LevelInstance/LevelInstanceActorDesc.h
@@ -46,5 +46,7 @@ protected:
 private:
 	void RegisterContainerInstance(UWorld* InWorld);
 	void UnregisterContainerInstance();
+	void OnLevelInstanceChanged(FName PackageName);
+	void UpdateBounds();
 #endif
 };
