Fix ExpectsAdditive in MakeDynamicAdditive Node

diff --git a/Engine/Source/Runtime/AnimGraphRuntime/Private/AnimNodes/AnimNode_MakeDynamicAdditive.cpp b/Engine/Source/Runtime/AnimGraphRuntime/Private/AnimNodes/AnimNode_MakeDynamicAdditive.cpp
--- a/Engine/Source/Runtime/AnimGraphRuntime/Private/AnimNodes/AnimNode_MakeDynamicAdditive.cpp
+++ b/Engine/Source/Runtime/AnimGraphRuntime/Private/AnimNodes/AnimNode_MakeDynamicAdditive.cpp
@@ -40,6 +40,7 @@ void FAnimNode_MakeDynamicAdditive::Update_AnyThread(const FAnimationUpdateConte
 void FAnimNode_MakeDynamicAdditive::Evaluate_AnyThread(FPoseContext& Output)
 {
 	DECLARE_SCOPE_HIERARCHICAL_COUNTER_ANIMNODE(Evaluate_AnyThread)
+	FScopedExpectsAdditiveOverride(Output, false);
 	FPoseContext BaseEvalContext(Output);
 
 	Base.Evaluate(BaseEvalContext);
diff --git a/Engine/Source/Runtime/Engine/Classes/Animation/AnimNodeBase.h b/Engine/Source/Runtime/Engine/Classes/Animation/AnimNodeBase.h
--- a/Engine/Source/Runtime/Engine/Classes/Animation/AnimNodeBase.h
+++ b/Engine/Source/Runtime/Engine/Classes/Animation/AnimNodeBase.h
@@ -444,6 +444,8 @@ public:
 	UE::Anim::FStackAttributeContainer CustomAttributes;
 
 public:
+	friend class FScopedExpectsAdditiveOverride;
+	
 	// This constructor allocates a new uninitialized pose for the specified anim instance
 	FPoseContext(FAnimInstanceProxy* InAnimInstanceProxy, bool bInExpectsAdditivePose = false)
 		: FAnimationBaseContext(InAnimInstanceProxy)
@@ -518,6 +520,27 @@ private:
 	bool bExpectsAdditivePose;
 };
 
+// Helper for modifying and resetting ExpectsAdditivePose on a FPoseContext
+class FScopedExpectsAdditiveOverride
+{
+public:
+	FScopedExpectsAdditiveOverride(FPoseContext& InContext, bool bInExpectsAdditive)
+		: Context(InContext)
+	{
+		bPreviousValue = Context.ExpectsAdditivePose();
+		Context.bExpectsAdditivePose = bInExpectsAdditive;
+	}
+	
+	~FScopedExpectsAdditiveOverride()
+	{
+		Context.bExpectsAdditivePose = bPreviousValue;
+	}
+private:
+	FPoseContext& Context;
+	bool bPreviousValue;
+};
+	
+
 
 /** Evaluation context passed around during animation tree evaluation */
 struct FComponentSpacePoseContext : public FAnimationBaseContext
