Fix for race condition in for loading zeroth chunk.

diff --git a/Engine/Source/Runtime/Engine/Classes/Sound/SoundWave.h b/Engine/Source/Runtime/Engine/Classes/Sound/SoundWave.h
--- a/Engine/Source/Runtime/Engine/Classes/Sound/SoundWave.h
+++ b/Engine/Source/Runtime/Engine/Classes/Sound/SoundWave.h
@@ -1296,6 +1296,7 @@ private:
 
 	/** Zeroth Chunk of audio for sources that use Load On Demand. */
 	FBulkDataBuffer<uint8> ZerothChunkData;
+	mutable FCriticalSection LoadZerothChunkDataCriticalSection;
 
 	/* Accessor to get the zeroth chunk which might perform additional work in editor to handle async tasks. */
 	FBulkDataBuffer<uint8>& GetZerothChunkData() const;
diff --git a/Engine/Source/Runtime/Engine/Private/SoundWave.cpp b/Engine/Source/Runtime/Engine/Private/SoundWave.cpp
--- a/Engine/Source/Runtime/Engine/Private/SoundWave.cpp
+++ b/Engine/Source/Runtime/Engine/Private/SoundWave.cpp
@@ -346,33 +346,8 @@ FBulkDataBuffer<uint8>& FSoundWaveData::GetZerothChunkData() const
 	FSoundWaveData* MutableThis = const_cast<FSoundWaveData*>(this);
 
 #if WITH_EDITOR
-	// In editor, we wait until the zeroth chunk is required to force finish compilation
-	// and retrieve it from DDC.  The rational is that async compilation will either prefetch
-	// chunks or rebuild them if unavailable so that this DDC request will be fulfilled from local cache.
-	if (ZerothChunkData.GetView().Num() == 0)
-	{
-		// If we're running the editor, we'll need to retrieve the chunked audio from the DDC:
-		uint8* TempChunkBuffer = nullptr;
-		int32 ChunkSizeInBytes = 0;
-
-		// GetNumChunks will force finish compilation if an async task still exists.
-		// Do not query chunk from DDC if no chunk exists as it would ensure inside GetChunkFromDDC.
-		if (MutableThis->RunningPlatformData.GetNumChunks() > 0)
-		{
-			ChunkSizeInBytes = MutableThis->GetChunkFromDDC(0, &TempChunkBuffer, true);
-		}
-
-		// Since we block for the DDC in the previous call we should always have the chunk loaded.
-		if (ChunkSizeInBytes != 0)
-		{
-			MutableThis->ZerothChunkData.Reset(TempChunkBuffer, ChunkSizeInBytes);
-		}
-		else
-		{
-			MutableThis->ZerothChunkData.Empty();
-		}
-	}
-#endif
+	MutableThis->LoadZerothChunk();
+#endif // #if WITH_EDITOR
 
 	return MutableThis->ZerothChunkData;
 }
@@ -390,45 +365,61 @@ bool FSoundWaveData::LoadZerothChunk()
 
 	LLM_SCOPE(ELLMTag::AudioSoundWaves);
 
-	// If the zeroth chunk is already loaded, early exit.
-	if (IsZerothChunkDataLoaded())
-	{
-		return true;
-	}
-
 	if (!ShouldUseStreamCaching())
 	{
 		return false;
 	}
 
-#if WITH_EDITOR
-	// If running editor, retrieve chunked audio from the DDC
-	uint8* TempChunkBuffer = nullptr;
-	int32 ChunkSizeInBytes = GetChunkFromDDC(0, &TempChunkBuffer, true);
-
-	if (ChunkSizeInBytes == 0)
 	{
-		UE_LOG(LogAudio, Warning, TEXT("FSoundWaveData::LoadZerothChunk: Unsuccessful load of zeroth chunk from DDC. Asset requires manual re-cook.\n\tAsset: '%s'\n\tDerivedDataKey: '%s'"),
-			*GetFName().ToString(),
-			*RunningPlatformData.DerivedDataKey);
-		return false;
-	}
-	ZerothChunkData.Reset(TempChunkBuffer, ChunkSizeInBytes);
+		FScopeLock ZerothChunkDataScopeLock(&LoadZerothChunkDataCriticalSection);
+
+		if (ZerothChunkData.GetView().Num() == 0)
+		{
+#if WITH_EDITOR
+			// In editor, we wait until the zeroth chunk is required to force finish compilation
+			// and retrieve it from DDC.  The rational is that async compilation will either prefetch
+			// chunks or rebuild them if unavailable so that this DDC request will be fulfilled from local cache.
+			uint8* TempChunkBuffer = nullptr;
+			int32 ChunkSizeInBytes = 0;
+
+			// GetNumChunks will force finish compilation if an async task still exists.
+			// Do not query chunk from DDC if no chunk exists as it would ensure inside GetChunkFromDDC.
+			if (RunningPlatformData.GetNumChunks() > 0)
+			{
+				ChunkSizeInBytes = GetChunkFromDDC(0, &TempChunkBuffer, true);
+			}
 
+			// Since we block for the DDC in the previous call we should always have the chunk loaded.
+			if (ChunkSizeInBytes != 0)
+			{
+				ZerothChunkData.Reset(TempChunkBuffer, ChunkSizeInBytes);
+			}
+			else
+			{
+				ZerothChunkData.Empty();
+				UE_LOG(LogAudio, Warning, TEXT("FSoundWaveData::LoadZerothChunk: Unsuccessful load of zeroth chunk from DDC. Asset requires manual re-cook.\n\tAsset: '%s'\n\tDerivedDataKey: '%s'"),
+					*GetFName().ToString(),
+					*RunningPlatformData.DerivedDataKey);
+				return false;
+			}
 #else // WITH_EDITOR
-	// Otherwise, the zeroth chunk is cooked out to RunningPlatformData, so retrieve it.
-	check(GetNumChunks() > 0);
-	FStreamedAudioChunk& ZerothChunk = GetChunk(0);
-
-	// Sanity check to ensure bulk size is set up
-	UE_CLOG(ZerothChunk.BulkData.GetBulkDataSize() != ZerothChunk.DataSize, LogAudio, Warning
-		, TEXT("SoundWave '%s' bulk data serialized out had a mismatched size with the DataSize field."
-			"\nBulk Data Reported Size: %d"
-			"\nBulk Data Actual Size: %ld")
-		, *GetFName().ToString(), ZerothChunk.DataSize, ZerothChunk.BulkData.GetBulkDataSize());
-
-	ZerothChunkData = ZerothChunk.BulkData.GetCopyAsBuffer(ZerothChunk.AudioDataSize, true);
+
+			// The zeroth chunk is cooked out to RunningPlatformData, so retrieve it.
+			check(GetNumChunks() > 0);
+			FStreamedAudioChunk& ZerothChunk = GetChunk(0);
+
+			// Sanity check to ensure bulk size is set up
+			UE_CLOG(ZerothChunk.BulkData.GetBulkDataSize() != ZerothChunk.DataSize, LogAudio, Warning
+				, TEXT("SoundWave '%s' bulk data serialized out had a mismatched size with the DataSize field."
+					"\nBulk Data Reported Size: %d"
+					"\nBulk Data Actual Size: %ld")
+				, *GetFName().ToString(), ZerothChunk.DataSize, ZerothChunk.BulkData.GetBulkDataSize());
+
+			ZerothChunkData = ZerothChunk.BulkData.GetCopyAsBuffer(ZerothChunk.AudioDataSize, true);
+
 #endif // WITH_EDITOR
+		}
+	}
 
 	return true;
 }
@@ -4196,4 +4187,4 @@ TArrayView<const uint8> FSoundWaveProxy::GetZerothChunk(const FSoundWaveProxyPtr
 	}
 
 	return {};
-}
\ No newline at end of file
+}
