Prevent editing "Render in depth pass" when "Render in main pass" is checked (UE-91213).

- Rendering in main pass implies it will be rendered in depth pass so ability to set it to false is just confusing for the user.

diff --git a/Engine/Source/Runtime/Engine/Classes/Components/PrimitiveComponent.h b/Engine/Source/Runtime/Engine/Classes/Components/PrimitiveComponent.h
--- a/Engine/Source/Runtime/Engine/Classes/Components/PrimitiveComponent.h
+++ b/Engine/Source/Runtime/Engine/Classes/Components/PrimitiveComponent.h
@@ -372,7 +372,7 @@ public:
 	uint8 bRenderInMainPass:1;
 
 	/** If true, this component will be rendered in the depth pass even if it's not rendered in the main pass */
-	UPROPERTY(EditAnywhere, AdvancedDisplay, BlueprintReadOnly, Category = Rendering)
+	UPROPERTY(EditAnywhere, AdvancedDisplay, BlueprintReadOnly, Category = Rendering, meta = (EditCondition = "!bRenderInMainPass"))
 	uint8 bRenderInDepthPass:1;
 
 	/** Whether the primitive receives decals. */
