Reset default value for RepLayout Array OnRep fix, as it was causing crashes in many real world cases.

diff --git a/Engine/Source/Runtime/Engine/Private/RepLayout.cpp b/Engine/Source/Runtime/Engine/Private/RepLayout.cpp
--- a/Engine/Source/Runtime/Engine/Private/RepLayout.cpp
+++ b/Engine/Source/Runtime/Engine/Private/RepLayout.cpp
@@ -77,7 +77,7 @@ static FAutoConsoleVariableRef CVarShareInitialCompareState(TEXT("net.ShareIniti
 bool GbTrackNetSerializeObjectReferences = false;
 static FAutoConsoleVariableRef CVarTrackNetSerializeObjectReferences(TEXT("net.TrackNetSerializeObjectReferences"), GbTrackNetSerializeObjectReferences, TEXT("If true, we will create small layouts for Net Serialize Structs if they have Object Properties. This can prevent some Shadow State GC crashes."));
 
-bool GbWithArrayOnRepFix = true;
+bool GbWithArrayOnRepFix = false;
 static FAutoConsoleVariableRef CVarWithArrayOnRepFix(TEXT("net.WithArrayOnRepFix"), GbWithArrayOnRepFix, TEXT("If true, attempt to prevent issues with Arrays not receiving OnRep calls until their size changes if their Archetypes have different values from instances in levels."));
 
 #if WITH_PUSH_VALIDATION_SUPPORT
