PR #8833: Render in Depth Pass fixed (Contributed by [REDACTED_PERSON])

#rnx

diff --git a/Engine/Source/Runtime/Engine/Private/StaticMeshRender.cpp b/Engine/Source/Runtime/Engine/Private/StaticMeshRender.cpp
--- a/Engine/Source/Runtime/Engine/Private/StaticMeshRender.cpp
+++ b/Engine/Source/Runtime/Engine/Private/StaticMeshRender.cpp
@@ -1898,6 +1898,7 @@ FPrimitiveViewRelevance FStaticMeshSceneProxy::GetViewRelevance(const FSceneView
 	Result.bDrawRelevance = IsShown(View) && View->Family->EngineShowFlags.StaticMeshes;
 	Result.bRenderCustomDepth = ShouldRenderCustomDepth();
 	Result.bRenderInMainPass = ShouldRenderInMainPass();
+	Result.bRenderInDepthPass = ShouldRenderInDepthPass();
 	Result.bUsesLightingChannels = GetLightingChannelMask() != GetDefaultLightingChannelMask();
 	Result.bTranslucentSelfShadow = bCastVolumetricTranslucentShadow;
 
