Fixed "Use Complex As Simple" project-wide setting to actually take effect once the setting changes.
	* BodyInstance wasn't querying the project settings when initialising a body.
	* PhysicsSettings wasn't correctly invalidating existing bodies when the setting was changed.

diff --git a/Engine/Source/Runtime/Engine/Private/PhysicsEngine/BodyInstance.cpp b/Engine/Source/Runtime/Engine/Private/PhysicsEngine/BodyInstance.cpp
--- a/Engine/Source/Runtime/Engine/Private/PhysicsEngine/BodyInstance.cpp
+++ b/Engine/Source/Runtime/Engine/Private/PhysicsEngine/BodyInstance.cpp
@@ -1399,7 +1399,7 @@ void FInitBodiesHelperBase::InitBodies()
 						{
 							for (int32 ShapeIndex = 0; ShapeIndex < NumShapes; ++ShapeIndex)
 							{
-								Body_External.SetShapeCollisionTraceType(ShapeIndex, ChaosInterface::ConvertCollisionTraceFlag(BI->BodySetup->CollisionTraceFlag)) ;
+								Body_External.SetShapeCollisionTraceType(ShapeIndex, ChaosInterface::ConvertCollisionTraceFlag(BI->BodySetup->GetCollisionTraceFlag())) ;
 							}
 						}
 					}
diff --git a/Engine/Source/Runtime/Engine/Private/PhysicsEngine/PhysicsSettings.cpp b/Engine/Source/Runtime/Engine/Private/PhysicsEngine/PhysicsSettings.cpp
--- a/Engine/Source/Runtime/Engine/Private/PhysicsEngine/PhysicsSettings.cpp
+++ b/Engine/Source/Runtime/Engine/Private/PhysicsEngine/PhysicsSettings.cpp
@@ -114,6 +114,20 @@ void UPhysicsSettings::PostEditChangeProperty(struct FPropertyChangedEvent& Prop
 	{
 		ChaosSettings.OnSettingsUpdated();
 	}
+
+	if(PropertyChangedEvent.GetPropertyName() == GET_MEMBER_NAME_CHECKED(UPhysicsSettingsCore, DefaultShapeComplexity))
+	{
+		for(TObjectIterator<UBodySetup> It; It; ++It)
+		{
+			UBodySetup* Setup = *It;
+			check(Setup);
+
+			if(Setup->bCreatedPhysicsMeshes)
+			{
+				Setup->InvalidatePhysicsData();
+			}
+		}
+	}
 }
 
 void UPhysicsSettings::LoadSurfaceType()
