Notify states no longer re-trigger erratically when a montage is re-played

diff --git a/Engine/Source/Runtime/Engine/Classes/Animation/AnimInstance.h b/Engine/Source/Runtime/Engine/Classes/Animation/AnimInstance.h
--- a/Engine/Source/Runtime/Engine/Classes/Animation/AnimInstance.h
+++ b/Engine/Source/Runtime/Engine/Classes/Animation/AnimInstance.h
@@ -238,16 +238,19 @@ struct FQueuedMontageBlendingOutEvent
 struct FQueuedMontageEndedEvent
 {
 	class UAnimMontage* Montage;
+	int32 MontageInstanceID;
 	bool bInterrupted;
 	FOnMontageEnded Delegate;
 
 	FQueuedMontageEndedEvent()
 		: Montage(NULL)
+		, MontageInstanceID(INDEX_NONE)
 		, bInterrupted(false)
 	{}
 
-	FQueuedMontageEndedEvent(class UAnimMontage* InMontage, bool InbInterrupted, FOnMontageEnded InDelegate)
+	FQueuedMontageEndedEvent(class UAnimMontage* InMontage, int32 InInstanceID, bool InbInterrupted, FOnMontageEnded InDelegate)
 		: Montage(InMontage)
+		, MontageInstanceID(InInstanceID)
 		, bInterrupted(InbInterrupted)
 		, Delegate(InDelegate)
 	{}
diff --git a/Engine/Source/Runtime/Engine/Classes/Animation/AnimationAsset.h b/Engine/Source/Runtime/Engine/Classes/Animation/AnimationAsset.h
--- a/Engine/Source/Runtime/Engine/Classes/Animation/AnimationAsset.h
+++ b/Engine/Source/Runtime/Engine/Classes/Animation/AnimationAsset.h
@@ -392,6 +392,9 @@ struct FAnimTickRecord
 	// Return the root motion weight for this tick record
 	float GetRootMotionWeight() const { return EffectiveBlendWeight * RootMotionWeightModifier; }
 
+private:
+	void AllocateContextDataContainer();
+
 public:
 	FAnimTickRecord()
 		: SourceAsset(nullptr)
@@ -428,6 +431,19 @@ public:
 	// Gather any data from the current update context
 	ENGINE_API void GatherContextData(const FAnimationUpdateContext& InContext);
 	
+	// Explicitly add typed context data to the tick record
+	template<typename Type, typename... TArgs>
+	void MakeContextData(TArgs&&... Args)
+	{
+		static_assert(TPointerIsConvertibleFromTo<Type, const UE::Anim::IAnimNotifyEventContextDataInterface>::Value, "'Type' template parameter to MakeContextData must be derived from IAnimNotifyEventContextDataInterface");
+		if (!ContextData.IsValid())
+		{
+			AllocateContextDataContainer();
+		}
+
+		ContextData->Add(MakeUnique<Type>(Forward<TArgs>(Args)...));
+	}
+
 	/** This can be used with the Sort() function on a TArray of FAnimTickRecord to sort from higher leader score */
 	ENGINE_API bool operator <(const FAnimTickRecord& Other) const { return LeaderScore > Other.LeaderScore; }
 };
diff --git a/Engine/Source/Runtime/Engine/Private/Animation/ActiveMontageInstanceScope.cpp b/Engine/Source/Runtime/Engine/Private/Animation/ActiveMontageInstanceScope.cpp
new file mode 100644
--- /dev/null
+++ b/Engine/Source/Runtime/Engine/Private/Animation/ActiveMontageInstanceScope.cpp
@@ -0,0 +1,14 @@
+// Copyright Epic Games, Inc. All Rights Reserved.
+
+#include "Animation/ActiveMontageInstanceScope.h"
+
+namespace UE::Anim {
+
+IMPLEMENT_NOTIFY_CONTEXT_INTERFACE(FAnimNotifyMontageInstanceContext)
+
+FAnimNotifyMontageInstanceContext::FAnimNotifyMontageInstanceContext(const int32 InMontageInstanceID)
+	: MontageInstanceID(InMontageInstanceID)
+{
+}
+	
+}	// namespace UE::Anim
\ No newline at end of file
diff --git a/Engine/Source/Runtime/Engine/Private/Animation/AnimInstance.cpp b/Engine/Source/Runtime/Engine/Private/Animation/AnimInstance.cpp
--- a/Engine/Source/Runtime/Engine/Private/Animation/AnimInstance.cpp
+++ b/Engine/Source/Runtime/Engine/Private/Animation/AnimInstance.cpp
@@ -33,6 +33,7 @@
 #include "Animation/AnimSubsystem.h"
 #include "Animation/AnimSubsystem_Tag.h"
 #include "Animation/AnimStateMachineTypes.h"
+#include "Animation/ActiveMontageInstanceScope.h"
 #if WITH_EDITOR
 #include "Animation/DebugSkelMeshComponent.h"
 #endif
@@ -1841,7 +1842,16 @@ void UAnimInstance::TriggerMontageEndedEvent(const FQueuedMontageEndedEvent& Mon
 			const FAnimNotifyEventReference& EventReference = ActiveAnimNotifyEventReference[Index];
 			UAnimMontage* NotifyMontage = Cast<UAnimMontage>(AnimNotifyEvent.NotifyStateClass->GetOuter());
 
-			if (NotifyMontage && (NotifyMontage == MontageEndedEvent.Montage))
+			// Grab the montage instance ID from the notify's event context
+			int32 EventReferenceMontageInstanceID = INDEX_NONE;
+			const UE::Anim::FAnimNotifyMontageInstanceContext* ActiveMontageContext = EventReference.GetContextData<UE::Anim::FAnimNotifyMontageInstanceContext>();
+			if (ActiveMontageContext)
+			{
+				EventReferenceMontageInstanceID = ActiveMontageContext->MontageInstanceID;
+			}
+
+			// Compare against the montage instance ID to prevent ending notify states from other instances of the same montage
+			if (NotifyMontage && (EventReferenceMontageInstanceID == MontageEndedEvent.MontageInstanceID))
 			{
 				if (ShouldTriggerAnimNotifyState(AnimNotifyEvent.NotifyStateClass))
 				{
diff --git a/Engine/Source/Runtime/Engine/Private/Animation/AnimMontage.cpp b/Engine/Source/Runtime/Engine/Private/Animation/AnimMontage.cpp
--- a/Engine/Source/Runtime/Engine/Private/Animation/AnimMontage.cpp
+++ b/Engine/Source/Runtime/Engine/Private/Animation/AnimMontage.cpp
@@ -19,6 +19,7 @@
 #include "Animation/AnimSingleNodeInstance.h"
 #include "Engine/Engine.h"
 #include "Animation/AnimTrace.h"
+#include "Animation/ActiveMontageInstanceScope.h"
 
 DEFINE_LOG_CATEGORY(LogAnimMontage);
 
@@ -1620,7 +1621,7 @@ void FAnimMontageInstance::Terminate()
 		ActiveStateBranchingPoints.Empty();
 
 		// terminating, trigger end
-		AnimInstance->QueueMontageEndedEvent(FQueuedMontageEndedEvent(OldMontage, bInterrupted, OnMontageEnded));
+		AnimInstance->QueueMontageEndedEvent(FQueuedMontageEndedEvent(OldMontage, InstanceID, bInterrupted, OnMontageEnded));
 
 		// Clear references to this MontageInstance. Needs to happen before Montage is cleared to nullptr, as TMaps can use that as a key.
 		AnimInstance->ClearMontageInstanceReferences(*this);
@@ -2539,6 +2540,10 @@ void FAnimMontageInstance::HandleEvents(float PreviousTrackPos, float CurrentTra
 	{
 		TMap<FName, TArray<FAnimNotifyEventReference>> NotifyMap;
 		FAnimTickRecord TickRecord;
+
+		// Add instance ID to context to differentiate notifies between different instances of the same montage
+		TickRecord.MakeContextData<UE::Anim::FAnimNotifyMontageInstanceContext>(InstanceID);
+
 		FAnimNotifyContext NotifyContext(TickRecord);
 		// We already break up AnimMontage update to handle looping, so we guarantee that PreviousPos and CurrentPos are contiguous.
 		Montage->GetAnimNotifiesFromDeltaPositions(PreviousTrackPos, CurrentTrackPos, NotifyContext);
diff --git a/Engine/Source/Runtime/Engine/Private/AnimationAsset.cpp b/Engine/Source/Runtime/Engine/Private/AnimationAsset.cpp
--- a/Engine/Source/Runtime/Engine/Private/AnimationAsset.cpp
+++ b/Engine/Source/Runtime/Engine/Private/AnimationAsset.cpp
@@ -182,6 +182,11 @@ void FAnimGroupInstance::Prepare(const FAnimGroupInstance* PreviousGroup)
 	}
 }
 
+void FAnimTickRecord::AllocateContextDataContainer()
+{
+	ContextData = MakeShared<TArray<TUniquePtr<const UE::Anim::IAnimNotifyEventContextDataInterface>>>();
+}
+
 FAnimTickRecord::FAnimTickRecord(UAnimSequenceBase* InSequence, bool bInLooping, float InPlayRate, float InFinalBlendWeight, float& InCurrentTime, FMarkerTickRecord& InMarkerTickRecord)
 {
 	SourceAsset = InSequence;
@@ -240,9 +245,9 @@ void FAnimTickRecord::GatherContextData(const FAnimationUpdateContext& InContext
 		InContext.GetSharedContext()->MessageStack.MakeEventContextData(NewContextData);
 		if(NewContextData.Num())
 		{
-			if(!ContextData.IsValid())
+			if (!ContextData.IsValid())
 			{
-				ContextData = MakeShared<TArray<TUniquePtr<const UE::Anim::IAnimNotifyEventContextDataInterface>>>();
+				AllocateContextDataContainer();
 			}
 
 			ContextData->Append(MoveTemp(NewContextData));
diff --git a/Engine/Source/Runtime/Engine/Public/Animation/ActiveMontageInstanceScope.h b/Engine/Source/Runtime/Engine/Public/Animation/ActiveMontageInstanceScope.h
new file mode 100644
--- /dev/null
+++ b/Engine/Source/Runtime/Engine/Public/Animation/ActiveMontageInstanceScope.h
@@ -0,0 +1,19 @@
+// Copyright Epic Games, Inc. All Rights Reserved.
+
+#pragma once
+
+#include "AnimNodeMessages.h"
+
+namespace UE::Anim {
+	
+// Notify Context Data containing the Instance ID of the montage that triggered this notify
+// This is useful i.e When there's simultaneous instances of the same montage, to trigger start/end events of its notify states.
+class ENGINE_API FAnimNotifyMontageInstanceContext : public UE::Anim::IAnimNotifyEventContextDataInterface
+{
+	DECLARE_NOTIFY_CONTEXT_INTERFACE(FAnimNotifyMontageInstanceContext)
+public:
+	FAnimNotifyMontageInstanceContext(const int32 InMontageInstanceID);
+	const int32 MontageInstanceID;
+};
+
+}	// namespace UE::Anim
