Crash fix for changing the anim class of a SKM actor with no skeleton
* Added nullptr checks for skeletons when comparing compatibility

diff --git a/Engine/Source/Runtime/Engine/Private/Components/SkeletalMeshComponent.cpp b/Engine/Source/Runtime/Engine/Private/Components/SkeletalMeshComponent.cpp
--- a/Engine/Source/Runtime/Engine/Private/Components/SkeletalMeshComponent.cpp
+++ b/Engine/Source/Runtime/Engine/Private/Components/SkeletalMeshComponent.cpp
@@ -524,7 +524,7 @@ bool USkeletalMeshComponent::NeedToSpawnAnimScriptInstance() const
 {
 	IAnimClassInterface* AnimClassInterface = IAnimClassInterface::GetFromClass(AnimClass);
 	const USkeleton* AnimSkeleton = (AnimClassInterface) ? AnimClassInterface->GetTargetSkeleton() : nullptr;
-	const bool bSkeletonCompatible = (SkeletalMesh && AnimSkeleton) ? SkeletalMesh->GetSkeleton()->IsCompatible(AnimSkeleton) : false;
+	const bool bSkeletonCompatible = (SkeletalMesh && AnimSkeleton && SkeletalMesh->GetSkeleton()) ? SkeletalMesh->GetSkeleton()->IsCompatible(AnimSkeleton) : false;
 	const bool bSkelMeshCompatible = (SkeletalMesh && AnimSkeleton) ? AnimSkeleton->IsCompatibleMesh(SkeletalMesh, false) : false;
 	const bool bAnimSkelValid = !AnimClassInterface || (bSkeletonCompatible && bSkelMeshCompatible);
 
@@ -761,11 +761,11 @@ void USkeletalMeshComponent::InitAnim(bool bForceReinit)
 
 		const bool bClearAnimInstance = AnimScriptInstance && !AnimSkeleton;
 		const bool bSkeletonMismatch = AnimSkeleton && (AnimScriptInstance->CurrentSkeleton!=SkeletalMesh->GetSkeleton());
-		const bool bSkeletonNotCompatible = AnimSkeleton && !bSkeletonMismatch && !SkeletalMesh->GetSkeleton()->IsCompatible(AnimSkeleton);
+		const bool bSkeletonCompatible = AnimSkeleton && SkeletalMesh->GetSkeleton() && !bSkeletonMismatch && SkeletalMesh->GetSkeleton()->IsCompatible(AnimSkeleton);
 
 		LastPoseTickFrame = 0;
 
-		if (bBlueprintMismatch || bSkeletonMismatch || bSkeletonNotCompatible || bClearAnimInstance)
+		if (bBlueprintMismatch || bSkeletonMismatch || !bSkeletonCompatible || bClearAnimInstance)
 		{
 			ClearAnimScriptInstance();
 		}
@@ -1086,14 +1086,14 @@ void USkeletalMeshComponent::PostEditChangeProperty(FPropertyChangedEvent& Prope
 		if ( PropertyThatChanged->GetFName() == GET_MEMBER_NAME_CHECKED( FSingleAnimationPlayData, AnimToPlay ))
 		{
 			// make sure the animation skeleton matches the current skeletalmesh
-			if (AnimationData.AnimToPlay != nullptr && SkeletalMesh && !SkeletalMesh->GetSkeleton()->IsCompatible(AnimationData.AnimToPlay->GetSkeleton()))
+			if (AnimationData.AnimToPlay && SkeletalMesh && SkeletalMesh->GetSkeleton() && SkeletalMesh->GetSkeleton()->IsCompatible(AnimationData.AnimToPlay->GetSkeleton()))
 			{
-				UE_LOG(LogAnimation, Warning, TEXT("Invalid animation"));
-				AnimationData.AnimToPlay = nullptr;
+				PlayAnimation(AnimationData.AnimToPlay, false);
 			}
 			else
 			{
-				PlayAnimation(AnimationData.AnimToPlay, false);
+				UE_LOG(LogAnimation, Warning, TEXT("Invalid animation"));
+				AnimationData.AnimToPlay = nullptr;
 			}
 		}
 
@@ -3792,18 +3792,18 @@ void USkeletalMeshComponent::ValidateAnimation()
 
 	if(AnimationMode == EAnimationMode::AnimationSingleNode)
 	{
-		if(AnimationData.AnimToPlay && SkeletalMesh && !AnimationData.AnimToPlay->GetSkeleton()->IsCompatible(SkeletalMesh->GetSkeleton()))
+		if (AnimationData.AnimToPlay && SkeletalMesh)
 		{
-			if (SkeletalMesh->GetSkeleton())
+			if (AnimationData.AnimToPlay->GetSkeleton() == nullptr)
 			{
-				UE_LOG(LogAnimation, Warning, TEXT("Animation %s is incompatible with skeleton %s, removing animation from actor."), *AnimationData.AnimToPlay->GetName(), *SkeletalMesh->GetSkeleton()->GetName());
+				UE_LOG(LogAnimation, Warning, TEXT("Animation %s is incompatible because it has no skeleton, removing animation from actor."), *AnimationData.AnimToPlay->GetName());
+				AnimationData.AnimToPlay = nullptr;
 			}
-			else
+			else if (!AnimationData.AnimToPlay->GetSkeleton()->IsCompatible(SkeletalMesh->GetSkeleton()))
 			{
-				UE_LOG(LogAnimation, Warning, TEXT("Animation %s is incompatible because mesh %s has no skeleton, removing animation from actor."), *AnimationData.AnimToPlay->GetName());
+				UE_LOG(LogAnimation, Warning, TEXT("Animation %s is incompatible with the skeletal mesh's skeleton, removing animation from actor."), *AnimationData.AnimToPlay->GetName());
+				AnimationData.AnimToPlay = nullptr;
 			}
-
-			AnimationData.AnimToPlay = nullptr;
 		}
 	}
 	else if (AnimationMode == EAnimationMode::AnimationBlueprint)
