Use proper gamma correction when previewing HDR virtual textures.

diff --git a/Engine/Source/Runtime/Engine/Private/Texture2D.cpp b/Engine/Source/Runtime/Engine/Private/Texture2D.cpp
--- a/Engine/Source/Runtime/Engine/Private/Texture2D.cpp
+++ b/Engine/Source/Runtime/Engine/Private/Texture2D.cpp
@@ -1471,7 +1471,7 @@ void FVirtualTexture2DResource::InitializeEditorResources(IVirtualTexture* InVir
 		RHIBindDebugLabelName(TextureRHI, *Name);
 		RHIUpdateTextureReference(TextureOwner->TextureReference.TextureReferenceRHI, TextureRHI);
 
-		bIgnoreGammaConversions = !TextureOwner->SRGB && ! UE::TextureDefines::IsHDR(TextureOwner->CompressionSettings);
+		bIgnoreGammaConversions = !TextureOwner->SRGB && !IsHDR(PixelFormat);
 
 		// re factored to ensure this is set earlier...make sure it's correct
 		ensure(bSRGB == TextureOwner->SRGB);
