Removing erroneous debug spam (logging should have been removed before submission)

diff --git a/Engine/Source/Runtime/Engine/Private/Components/AudioComponent.cpp b/Engine/Source/Runtime/Engine/Private/Components/AudioComponent.cpp
--- a/Engine/Source/Runtime/Engine/Private/Components/AudioComponent.cpp
+++ b/Engine/Source/Runtime/Engine/Private/Components/AudioComponent.cpp
@@ -588,8 +588,7 @@ void UAudioComponent::PlayQueuedQuantizedInternal(const UObject* WorldContextObj
 
 					// Now play the quartz command
 					PlayInternal(InternalRequestData);
-					UE_LOG(LogAudioQuartz, Warning, TEXT("AudioComponent::PlayQueueQUantizedInternal()->PlayInternal()"));
-				} 
+				}
 			}
 			else // this component does not have a Sound to play
 			{
