Fix non-unity and deprecation warning
#rnx

diff --git a/Engine/Plugins/FX/Niagara/Source/NiagaraEditor/Private/NiagaraBakerRendererOutputVolumeTexture.cpp b/Engine/Plugins/FX/Niagara/Source/NiagaraEditor/Private/NiagaraBakerRendererOutputVolumeTexture.cpp
--- a/Engine/Plugins/FX/Niagara/Source/NiagaraEditor/Private/NiagaraBakerRendererOutputVolumeTexture.cpp
+++ b/Engine/Plugins/FX/Niagara/Source/NiagaraEditor/Private/NiagaraBakerRendererOutputVolumeTexture.cpp
@@ -12,6 +12,7 @@
 
 #include "Engine/Canvas.h"
 #include "Engine/VolumeTexture.h"
+#include "HAL/PlatformFileManager.h"
 #include "Misc/PathViews.h"
 #include "Modules/ModuleManager.h"
 #include "Factories/VolumeTextureFactory.h"
diff --git a/Engine/Source/Developer/TraceServices/Private/Analyzers/StatsTraceAnalysis.cpp b/Engine/Source/Developer/TraceServices/Private/Analyzers/StatsTraceAnalysis.cpp
--- a/Engine/Source/Developer/TraceServices/Private/Analyzers/StatsTraceAnalysis.cpp
+++ b/Engine/Source/Developer/TraceServices/Private/Analyzers/StatsTraceAnalysis.cpp
@@ -5,6 +5,7 @@
 #include "AnalysisServicePrivate.h"
 #include "Common/Utils.h"
 #include "HAL/LowLevelMemTracker.h"
+#include "ProfilingDebugging/MiscTrace.h"
 #include "TraceServices/Model/Counters.h"
 
 #include <limits>
diff --git a/Engine/Source/Runtime/Engine/Private/Materials/MaterialShared.cpp b/Engine/Source/Runtime/Engine/Private/Materials/MaterialShared.cpp
--- a/Engine/Source/Runtime/Engine/Private/Materials/MaterialShared.cpp
+++ b/Engine/Source/Runtime/Engine/Private/Materials/MaterialShared.cpp
@@ -3469,11 +3469,13 @@ IAllocatedVirtualTexture* FMaterialRenderProxy::AllocateVTStack(const FMaterialR
 	return nullptr;
 }
 
+PRAGMA_DISABLE_DEPRECATION_WARNINGS
 FUniformExpressionCache::~FUniformExpressionCache()
 {
 	ResetAllocatedVTs();
 	UniformBuffer.SafeRelease();
 }
+PRAGMA_ENABLE_DEPRECATION_WARNINGS
 
 void FUniformExpressionCache::ResetAllocatedVTs()
 {
diff --git a/Engine/Source/Runtime/Engine/Private/PhysicsEngine/PhysicsSettings.cpp b/Engine/Source/Runtime/Engine/Private/PhysicsEngine/PhysicsSettings.cpp
--- a/Engine/Source/Runtime/Engine/Private/PhysicsEngine/PhysicsSettings.cpp
+++ b/Engine/Source/Runtime/Engine/Private/PhysicsEngine/PhysicsSettings.cpp
@@ -3,7 +3,9 @@
 #include "PhysicsEngine/PhysicsSettings.h"
 #include "GameFramework/MovementComponent.h"
 #include "PhysicalMaterials/PhysicalMaterial.h"
+#include "PhysicsEngine/BodySetup.h"
 #include "UObject/Package.h"
+#include "UObject/UObjectIterator.h"
 
 #include "Framework/Threading.h"
 
