Actor Descriptors: remove unused LevelBoundsRelevant property.

#rnx

diff --git a/Engine/Source/Runtime/Core/Public/UObject/UE5MainStreamObjectVersion.h b/Engine/Source/Runtime/Core/Public/UObject/UE5MainStreamObjectVersion.h
--- a/Engine/Source/Runtime/Core/Public/UObject/UE5MainStreamObjectVersion.h
+++ b/Engine/Source/Runtime/Core/Public/UObject/UE5MainStreamObjectVersion.h
@@ -236,6 +236,9 @@ struct CORE_API FUE5MainStreamObjectVersion
 		// World partition streaming policy serialization only for cooked builds
 		WorldPartitionSerializeStreamingPolicyOnCook,
 
+		// Remove serialization of bounds relevant from  WorldPartitionActorDesc
+		WorldPartitionActorDescRemoveBoundsRelevantSerialization,
+
 		// -----<new versions can be added above this line>-------------------------------------------------
 		VersionPlusOne,
 		LatestVersion = VersionPlusOne - 1
diff --git a/Engine/Source/Runtime/Engine/Private/WorldPartition/WorldPartitionActorDesc.cpp b/Engine/Source/Runtime/Engine/Private/WorldPartition/WorldPartitionActorDesc.cpp
--- a/Engine/Source/Runtime/Engine/Private/WorldPartition/WorldPartitionActorDesc.cpp
+++ b/Engine/Source/Runtime/Engine/Private/WorldPartition/WorldPartitionActorDesc.cpp
@@ -58,7 +58,6 @@ void FWorldPartitionActorDesc::Init(const AActor* InActor)
 	RuntimeGrid = InActor->GetRuntimeGrid();
 	bIsSpatiallyLoaded = InActor->GetIsSpatiallyLoaded();
 	bActorIsEditorOnly = InActor->IsEditorOnly();
-	bLevelBoundsRelevant = InActor->IsLevelBoundsRelevant();
 	bActorIsHLODRelevant = InActor->IsHLODRelevant();
 	HLODLayer = InActor->GetHLODLayer() ? FName(InActor->GetHLODLayer()->GetPathName()) : FName();
 	
@@ -164,7 +163,6 @@ bool FWorldPartitionActorDesc::Equals(const FWorldPartitionActorDesc* Other) con
 		BoundsExtent.Equals(Other->BoundsExtent, 0.1f) &&
 		RuntimeGrid == Other->RuntimeGrid &&
 		bActorIsEditorOnly == Other->bActorIsEditorOnly &&
-		bLevelBoundsRelevant == Other->bLevelBoundsRelevant &&
 		bActorIsHLODRelevant == Other->bActorIsHLODRelevant &&
 		bIsUsingDataLayerAsset == Other->bIsUsingDataLayerAsset &&
 		HLODLayer == Other->HLODLayer &&
@@ -220,7 +218,7 @@ FString FWorldPartitionActorDesc::ToString() const
 	};
 
 	return FString::Printf(
-		TEXT("Guid:%s BaseClass:%s NativeClass:%s Name:%s Label:%s SpatiallyLoaded:%s Bounds:%s RuntimeGrid:%s EditorOnly:%s LevelBoundsRelevant:%s HLODRelevant:%s FolderPath:%s FolderGuid:%s Parent:%s"), 
+		TEXT("Guid:%s BaseClass:%s NativeClass:%s Name:%s Label:%s SpatiallyLoaded:%s Bounds:%s RuntimeGrid:%s EditorOnly:%s HLODRelevant:%s FolderPath:%s FolderGuid:%s Parent:%s"), 
 		*Guid.ToString(), 
 		*BaseClass.ToString(), 
 		*NativeClass.ToString(), 
@@ -230,7 +228,6 @@ FString FWorldPartitionActorDesc::ToString() const
 		*GetBounds().ToString(),
 		*RuntimeGrid.ToString(),
 		GetBoolStr(bActorIsEditorOnly),
-		GetBoolStr(bLevelBoundsRelevant),
 		GetBoolStr(bActorIsHLODRelevant),
 		*FolderPath.ToString(),
 		*FolderGuid.ToString(),
@@ -278,7 +275,13 @@ void FWorldPartitionActorDesc::Serialize(FArchive& Ar)
 		Ar << bIsSpatiallyLoaded;
 	}
 		
-	Ar << RuntimeGrid << bActorIsEditorOnly << bLevelBoundsRelevant;
+	Ar << RuntimeGrid << bActorIsEditorOnly;
+
+	if (Ar.CustomVer(FUE5MainStreamObjectVersion::GUID) < FUE5MainStreamObjectVersion::WorldPartitionActorDescRemoveBoundsRelevantSerialization)
+	{
+		bool bLevelBoundsRelevant;
+		Ar << bLevelBoundsRelevant;
+	}
 	
 	if (Ar.CustomVer(FUE5MainStreamObjectVersion::GUID) < FUE5MainStreamObjectVersion::WorldPartitionActorDescSerializeDataLayers)
 	{
@@ -480,4 +483,4 @@ void FWorldPartitionActorDesc::Unload()
 		ActorPtr = nullptr;
 	}
 }
-#endif
\ No newline at end of file
+#endif
diff --git a/Engine/Source/Runtime/Engine/Private/WorldPartition/WorldPartitionActorDescView.cpp b/Engine/Source/Runtime/Engine/Private/WorldPartition/WorldPartitionActorDescView.cpp
--- a/Engine/Source/Runtime/Engine/Private/WorldPartition/WorldPartitionActorDescView.cpp
+++ b/Engine/Source/Runtime/Engine/Private/WorldPartition/WorldPartitionActorDescView.cpp
@@ -63,11 +63,6 @@ bool FWorldPartitionActorDescView::GetIsSpatiallyLoaded() const
 	return bIsForcedNonSpatiallyLoaded ? false : ActorDesc->GetIsSpatiallyLoaded();
 }
 
-bool FWorldPartitionActorDescView::GetLevelBoundsRelevant() const
-{
-	return ActorDesc->GetLevelBoundsRelevant();
-}
-
 bool FWorldPartitionActorDescView::GetActorIsHLODRelevant() const
 {
 	return ActorDesc->GetActorIsHLODRelevant();
diff --git a/Engine/Source/Runtime/Engine/Public/WorldPartition/WorldPartitionActorDesc.h b/Engine/Source/Runtime/Engine/Public/WorldPartition/WorldPartitionActorDesc.h
--- a/Engine/Source/Runtime/Engine/Public/WorldPartition/WorldPartitionActorDesc.h
+++ b/Engine/Source/Runtime/Engine/Public/WorldPartition/WorldPartitionActorDesc.h
@@ -82,7 +82,10 @@ public:
 	inline bool GetIsSpatiallyLoaded() const { return bIsForcedNonSpatiallyLoaded ? false : bIsSpatiallyLoaded; }
 	inline bool GetIsSpatiallyLoadedRaw() const { return bIsSpatiallyLoaded; }
 	inline bool GetActorIsEditorOnly() const { return bActorIsEditorOnly; }
-	inline bool GetLevelBoundsRelevant() const { return bLevelBoundsRelevant; }
+
+	UE_DEPRECATED(5.1, "GetLevelBoundsRelevant is deprecated.")
+	inline bool GetLevelBoundsRelevant() const { return false; }
+
 	inline bool GetActorIsHLODRelevant() const { return bActorIsHLODRelevant; }
 	inline FName GetHLODLayer() const { return HLODLayer; }
 	inline const TArray<FName>& GetDataLayers() const { return DataLayers; }
@@ -223,7 +226,6 @@ protected:
 	FName							RuntimeGrid;
 	bool							bIsSpatiallyLoaded;
 	bool							bActorIsEditorOnly;
-	bool							bLevelBoundsRelevant;
 	bool							bActorIsHLODRelevant;
 	bool							bIsUsingDataLayerAsset; // Used to know if DataLayers array represents DataLayers Asset paths or the FNames of the deprecated version of Data Layers
 	FName							HLODLayer;
diff --git a/Engine/Source/Runtime/Engine/Public/WorldPartition/WorldPartitionActorDescView.h b/Engine/Source/Runtime/Engine/Public/WorldPartition/WorldPartitionActorDescView.h
--- a/Engine/Source/Runtime/Engine/Public/WorldPartition/WorldPartitionActorDescView.h
+++ b/Engine/Source/Runtime/Engine/Public/WorldPartition/WorldPartitionActorDescView.h
@@ -30,7 +30,6 @@ public:
 	FName GetRuntimeGrid() const;
 	bool GetIsSpatiallyLoaded() const;
 	bool GetActorIsEditorOnly() const;
-	bool GetLevelBoundsRelevant() const;
 	bool GetActorIsHLODRelevant() const;
 	FName GetHLODLayer() const;
 	const TArray<FName>& GetDataLayers() const;
