World Partition: fixed streaming policy not duplicated for PIE.

#rnx

diff --git a/Engine/Source/Runtime/Engine/Private/WorldPartition/WorldPartition.cpp b/Engine/Source/Runtime/Engine/Private/WorldPartition/WorldPartition.cpp
--- a/Engine/Source/Runtime/Engine/Private/WorldPartition/WorldPartition.cpp
+++ b/Engine/Source/Runtime/Engine/Private/WorldPartition/WorldPartition.cpp
@@ -1058,7 +1058,11 @@ void UWorldPartition::Serialize(FArchive& Ar)
 
 	Super::Serialize(Ar);
 
-	if (Ar.CustomVer(FUE5MainStreamObjectVersion::GUID) >= FUE5MainStreamObjectVersion::WorldPartitionSerializeStreamingPolicyOnCook)
+	if (Ar.GetPortFlags() & PPF_DuplicateForPIE)
+	{
+		Ar << StreamingPolicy;
+	}
+	else if (Ar.CustomVer(FUE5MainStreamObjectVersion::GUID) >= FUE5MainStreamObjectVersion::WorldPartitionSerializeStreamingPolicyOnCook)
 	{
 		bool bCooked = Ar.IsCooking();
 		Ar << bCooked;
