Do not alter bIsLowAreaModifier & ApplyMode values when AreaClassToPlace is null in FAreaNavModifier::SetAreaClassToReplace
Update bHasMetaAreas value after picking the area class for them in FCompositeNavModifier::GetInstantiatedMetaModifier

diff --git a/Engine/Source/Runtime/Engine/Private/AI/Navigation/NavigationModifier.cpp b/Engine/Source/Runtime/Engine/Private/AI/Navigation/NavigationModifier.cpp
--- a/Engine/Source/Runtime/Engine/Private/AI/Navigation/NavigationModifier.cpp
+++ b/Engine/Source/Runtime/Engine/Private/AI/Navigation/NavigationModifier.cpp
@@ -445,8 +445,11 @@ void FAreaNavModifier::SetAreaClassToReplace(const TSubclassOf<UNavAreaBase> InA
 	bHasMetaAreas = (AreaClass1 && IsMetaAreaClass(*AreaClass1))
 		|| (AreaClass2 && IsMetaAreaClass(*AreaClass2));
 
-	bIsLowAreaModifier = (AreaClass2 && AreaClass2->GetDefaultObject<UNavAreaBase>()->IsLowArea());
-	ApplyMode = bIsLowAreaModifier ? ENavigationAreaMode::ReplaceInLowPass : AreaClass2 ? ENavigationAreaMode::Replace : ENavigationAreaMode::Apply;
+	if (AreaClass2)
+	{
+		bIsLowAreaModifier = AreaClass2->GetDefaultObject<UNavAreaBase>()->IsLowArea();
+		ApplyMode = bIsLowAreaModifier ? ENavigationAreaMode::ReplaceInLowPass : ENavigationAreaMode::Replace;
+	}
 }
 
 void FAreaNavModifier::SetApplyMode(ENavigationAreaMode::Type InApplyMode)
@@ -797,6 +800,7 @@ FCompositeNavModifier FCompositeNavModifier::GetInstantiatedMetaModifier(const F
 	Result = *this;
 
 	{
+		Result.bHasMetaAreas = false;
 		FAreaNavModifier* Area = Result.Areas.GetData();
 		for (int32 Index = 0; Index < Result.Areas.Num(); ++Index, ++Area)
 		{
@@ -804,6 +808,7 @@ FCompositeNavModifier FCompositeNavModifier::GetInstantiatedMetaModifier(const F
 			{
 				Area->SetAreaClass(UNavAreaBase::PickAreaClassForAgent(Area->GetAreaClass(), *ActorOwner, *NavAgent));
 				Area->SetAreaClassToReplace(UNavAreaBase::PickAreaClassForAgent(Area->GetAreaClassToReplace(), *ActorOwner, *NavAgent));
+				Result.bHasMetaAreas = true;
 			}
 		}
 	}
