Physics replication Async path previously copied a particle's game thread object state into FAsyncPhysicsDesiredState, marshaled it to game thread and in ApplyAsyncDesiredPhysicsState would overwrite whatever state occurred there.

The upshot is that if a particle was asleep, even if replication should have woken it up by moving it, it would get put back to sleep.

This was temporarily worked around in 20560714 by optionally skipping the Apply step unless it was putting it to sleep. This change simplifies by removing ObjectState from desired state entirely.

diff --git a/Engine/Source/Runtime/Engine/Private/PhysicsEngine/PhysicsReplication.cpp b/Engine/Source/Runtime/Engine/Private/PhysicsEngine/PhysicsReplication.cpp
--- a/Engine/Source/Runtime/Engine/Private/PhysicsEngine/PhysicsReplication.cpp
+++ b/Engine/Source/Runtime/Engine/Private/PhysicsEngine/PhysicsReplication.cpp
@@ -80,9 +80,6 @@ namespace CharacterMovementCVars
 
 	static int32 ApplyAsyncSleepState = 1;
 	static FAutoConsoleVariableRef CVarApplyAsyncSleepState(TEXT("p.ApplyAsyncSleepState"), ApplyAsyncSleepState, TEXT(""));
-
-	static bool OnlyApplyAsyncSleepStateIfShouldSleep = true;
-	static FAutoConsoleVariableRef CVarOnlyApplyAsyncSleepStateIfShouldSleep(TEXT("p.OnlyApplyAsyncSleepStateIfShouldSleep"), OnlyApplyAsyncSleepStateIfShouldSleep, TEXT(""));
 }
 
 namespace PhysicsReplicationCVars
@@ -410,8 +407,6 @@ bool FPhysicsReplication::ApplyRigidBodyState(float DeltaSeconds, FBodyInstance*
 				AsyncDesiredState.LinearVelocity = NewState.LinVel;
 				AsyncDesiredState.AngularVelocity = NewState.AngVel;
 				AsyncDesiredState.Proxy = static_cast<Chaos::FSingleParticlePhysicsProxy*>(BI->GetPhysicsActorHandle());
-				// ObjectState may not be needed at all.
-				AsyncDesiredState.ObjectState = AsyncDesiredState.Proxy->GetGameThreadAPI().ObjectState();
 				AsyncDesiredState.ErrorCorrection = { ErrorCorrection.LinearVelocityCoefficient, ErrorCorrection.AngularVelocityCoefficient, ErrorCorrection.PositionLerp, ErrorCorrection.AngleLerp };
 				AsyncDesiredState.bShouldSleep = bShouldSleep;
 
@@ -683,21 +678,13 @@ void FPhysicsReplication::ApplyAsyncDesiredState(const float DeltaSeconds, const
 				Handle->SetV(NewLinVel);
 				Handle->SetW(FMath::DegreesToRadians(NewAngVel));
 
-				// The old state should probably not be the default here.
-				EObjectStateType ObjectStateType = State.ObjectState;
-				if ((CharacterMovementCVars::ApplyAsyncSleepState != 0) && State.bShouldSleep)
-				{
-					ObjectStateType = EObjectStateType::Sleeping;
-				}
-
-				if (State.bShouldSleep || !CharacterMovementCVars::OnlyApplyAsyncSleepStateIfShouldSleep)
+				if (State.bShouldSleep)
 				{
 					// don't allow kinematic to sleeping transition
-					bool bInvalidObjectState = (ObjectStateType == EObjectStateType::Sleeping && Handle->ObjectState() == EObjectStateType::Kinematic);
-					if (!bInvalidObjectState)
+					if (Handle->ObjectState() != EObjectStateType::Kinematic)
 					{
 						auto* Solver = Proxy->GetSolver<FPBDRigidsSolver>();
-						Solver->GetEvolution()->SetParticleObjectState(Proxy->GetHandle_LowLevel()->CastToRigidParticle(), ObjectStateType);	//todo: move object state into physics thread api
+						Solver->GetEvolution()->SetParticleObjectState(Proxy->GetHandle_LowLevel()->CastToRigidParticle(), EObjectStateType::Sleeping);	//todo: move object state into physics thread api
 					}
 				}
 			}
diff --git a/Engine/Source/Runtime/Engine/Public/PhysicsReplication.h b/Engine/Source/Runtime/Engine/Public/PhysicsReplication.h
--- a/Engine/Source/Runtime/Engine/Public/PhysicsReplication.h
+++ b/Engine/Source/Runtime/Engine/Public/PhysicsReplication.h
@@ -74,7 +74,6 @@ struct FAsyncPhysicsDesiredState
 	FVector LinearVelocity;
 	FVector AngularVelocity;
 	Chaos::FSingleParticlePhysicsProxy* Proxy;
-	Chaos::EObjectStateType ObjectState;
 	TOptional<ErrorCorrectionData> ErrorCorrection;
 	bool bShouldSleep;
 	int32 ServerFrame;
