GitHub 9274 : Fix OrthoZoomFactor being overwritten

#9274

diff --git a/Engine/Source/Runtime/Engine/Private/BatchedElements.cpp b/Engine/Source/Runtime/Engine/Private/BatchedElements.cpp
--- a/Engine/Source/Runtime/Engine/Private/BatchedElements.cpp
+++ b/Engine/Source/Runtime/Engine/Private/BatchedElements.cpp
@@ -897,12 +897,12 @@ bool FBatchedElements::Draw(FRHICommandList& RHICmdList, const FMeshPassProcesso
 						const float ScalingStart	= Line.bScreenSpace ? StartW / ViewportSizeX : 1.0f;
 						const float ScalingEnd		= Line.bScreenSpace ? EndW   / ViewportSizeX : 1.0f;
 
-						OrthoZoomFactor = Line.bScreenSpace ? OrthoZoomFactor : 1.0f;
+						const float CurrentOrthoZoomFactor = Line.bScreenSpace ? OrthoZoomFactor : 1.0f;
 
 						const float ScreenSpaceScaling = Line.bScreenSpace ? 2.0f : 1.0f;
 
-						const float StartThickness	= Thickness * ScreenSpaceScaling * OrthoZoomFactor * ScalingStart;
-						const float EndThickness	= Thickness * ScreenSpaceScaling * OrthoZoomFactor * ScalingEnd;
+						const float StartThickness	= Thickness * ScreenSpaceScaling * CurrentOrthoZoomFactor * ScalingStart;
+						const float EndThickness	= Thickness * ScreenSpaceScaling * CurrentOrthoZoomFactor * ScalingEnd;
 
 						const FVector WorldPointXS	= CameraX * StartThickness * 0.5f;
 						const FVector WorldPointYS	= CameraY * StartThickness * 0.5f;
