**New**

diff --git a/Engine/Source/Runtime/Engine/Classes/Animation/AnimData/AnimDataModel.h b/Engine/Source/Runtime/Engine/Classes/Animation/AnimData/AnimDataModel.h
--- a/Engine/Source/Runtime/Engine/Classes/Animation/AnimData/AnimDataModel.h
+++ b/Engine/Source/Runtime/Engine/Classes/Animation/AnimData/AnimDataModel.h
@@ -4,279 +4,75 @@
 
 #include "CoreMinimal.h"
 #include "UObject/ObjectMacros.h"
-#include "AnimDataNotifications.h"
-#include "Animation/AttributeCurve.h"
 
-#include "AnimDataModel.generated.h"
-
-/**
- * Structure encapsulating a single bone animation track.
- */
-USTRUCT(BlueprintType)
-struct ENGINE_API FBoneAnimationTrack
-{
-	GENERATED_BODY()
-
-	/** Internally stored data representing the animation bone data */
-	UPROPERTY(VisibleAnywhere, BlueprintReadOnly, Category = "Model")
-	FRawAnimSequenceTrack InternalTrackData;
-
-	/** Index corresponding to the bone this track corresponds to within the target USkeleton */
-	UPROPERTY(VisibleAnywhere, BlueprintReadOnly, Category = "Model")
-	int32 BoneTreeIndex = INDEX_NONE;
-
-	/** Name of the bone this track corresponds to */
-	UPROPERTY(VisibleAnywhere, BlueprintReadOnly, Category = "Model")
-	FName Name;
-};
-
-/**
- * Structure encapsulating animated curve data. Currently only contains Float and Transform curves.
- */
-USTRUCT(BlueprintType)
-struct ENGINE_API FAnimationCurveData
-{
-	GENERATED_BODY()
-
-	/** Float-based animation curves */
-	UPROPERTY(VisibleAnywhere, BlueprintReadOnly, Category = "Model")
-	TArray<FFloatCurve>	FloatCurves;
-
-	/** FTransform-based animation curves, used for animation layer editing */
-	UPROPERTY(VisibleAnywhere, BlueprintReadOnly, Category = "Model")
-	TArray<FTransformCurve>	TransformCurves;
-};
-
-/**
- * Structure encapsulating animated (bone) attribute data.
- */
-USTRUCT(BlueprintType)
-struct ENGINE_API FAnimatedBoneAttribute
-{
-	GENERATED_BODY()
-
-	/** Identifier to reference this attribute by */
-	UPROPERTY(VisibleAnywhere, BlueprintReadOnly, Category = "Model")
-	FAnimationAttributeIdentifier Identifier;	
+#include "IAnimationDataModel.h"
 
-	/** Curve containing the (animated) attribute data */
-	UPROPERTY(VisibleAnywhere, BlueprintReadOnly, Category = "Model")
-	FAttributeCurve Curve;
-};
-
-struct FTrackToSkeletonMap;
-struct FAnimationCurveIdentifier;
+#include "AnimDataModel.generated.h"
 
 /**
  * The Model represents the source data for animations. It contains both bone animation data as well as animated curves.
  * They are currently only a sub-object of a AnimSequenceBase instance. The instance derives all runtime data from the source data. 
  */
 UCLASS(BlueprintType)
-class ENGINE_API UAnimDataModel : public UObject
+class ENGINE_API UAnimDataModel : public UObject, public IAnimationDataModel
 {
 	GENERATED_BODY()
 public:
-
-	/** Begin UAnimDataModel overrides */
+	/** Begin UObject overrides */
 	virtual void PostLoad() override;
 	virtual void PostDuplicate(bool bDuplicateForPIE) override;
 	virtual bool IsEditorOnly() const override { return true; }
-	/** End UAnimDataModel overrides */
-
-	/**
-	* @return	Total length of play-able animation data 
-	*/
-	UFUNCTION(BlueprintPure, Category=AnimationDataModel)
-	float GetPlayLength() const;
+	/** End UObject overrides */
 	
-	/**
-	* @return	Total number of frames of animation data stored 
-	*/
-	UFUNCTION(BlueprintPure, Category = AnimationDataModel)
-	int32 GetNumberOfFrames() const;
-
-	/**
-	* @return	Total number of animation data keys stored 
-	*/
-	UFUNCTION(BlueprintPure, Category = AnimationDataModel)
-	int32 GetNumberOfKeys() const;
-
-	/**
-	* @return	Frame rate at which the animation data is key-ed 
-	*/
-	UFUNCTION(BlueprintPure, Category = AnimationDataModel)
-	const FFrameRate& GetFrameRate() const;
-	
-	/**
-	* @return	Array containg all bone animation tracks 
-	*/
-	UFUNCTION(BlueprintPure, Category = AnimationDataModel)
-	const TArray<FBoneAnimationTrack>& GetBoneAnimationTracks() const;
-	
-	/**
-	* @return	Bone animation track for the provided index
-	*/
-	UFUNCTION(BlueprintPure, Category = AnimationDataModel)
-	const FBoneAnimationTrack& GetBoneTrackByIndex(int32 TrackIndex) const;
-
-	/**
-	* @return	Bone animation track for the provided (bone) name
-	*/
-	UFUNCTION(BlueprintPure, Category = AnimationDataModel)
-	const FBoneAnimationTrack& GetBoneTrackByName(FName TrackName) const;
-	
-	/**
-	* @return	Bone animation track for the provided (bone) name if found, otherwise returns a nullptr 
-	*/
-	const FBoneAnimationTrack* FindBoneTrackByName(FName Name) const;
-
-	/**
-	* @return	Bone animation track for the provided index if valid, otherwise returns a nullptr 
-	*/
-	const FBoneAnimationTrack* FindBoneTrackByIndex(int32 BoneIndex) const;
-
-	/**
-	* @return	Internal track index for the provided bone animation track if found, otherwise returns INDEX_NONE 
-	*/
-	UFUNCTION(BlueprintPure, Category = AnimationDataModel)
-	int32 GetBoneTrackIndex(const FBoneAnimationTrack& Track) const;
-
-	/**
-	* @return	Internal track index for the provided (bone) name if found, otherwise returns INDEX_NONE 
-	*/
-	UFUNCTION(BlueprintPure, Category = AnimationDataModel)
-	int32 GetBoneTrackIndexByName(FName TrackName) const;
+	/** Begin IAnimationDataModel overrides */
+	virtual double GetPlayLength() const override;
+	virtual int32 GetNumberOfFrames() const override;
+	virtual int32 GetNumberOfKeys() const override;
+	virtual FFrameRate GetFrameRate() const override;
+	virtual const TArray<FBoneAnimationTrack>& GetBoneAnimationTracks() const override;
+	virtual const FBoneAnimationTrack& GetBoneTrackByIndex(int32 TrackIndex) const override;
+	virtual const FBoneAnimationTrack& GetBoneTrackByName(FName TrackName) const override;
+	virtual const FBoneAnimationTrack* FindBoneTrackByName(FName Name) const override;
+	virtual const FBoneAnimationTrack* FindBoneTrackByIndex(int32 BoneIndex) const override;
+	virtual int32 GetBoneTrackIndex(const FBoneAnimationTrack& Track) const override;
+	virtual int32 GetBoneTrackIndexByName(FName TrackName) const override;
+	virtual bool IsValidBoneTrackIndex(int32 TrackIndex) const override;
+	virtual int32 GetNumBoneTracks() const override;
+	virtual void GetBoneTrackNames(TArray<FName>& OutNames) const override;
+
+	virtual const FAnimCurveBase* FindCurve(const FAnimationCurveIdentifier& CurveIdentifier) const override;
+	virtual const FFloatCurve* FindFloatCurve(const FAnimationCurveIdentifier& CurveIdentifier) const override;
+	virtual const FTransformCurve* FindTransformCurve(const FAnimationCurveIdentifier& CurveIdentifier) const override;
+	virtual const FRichCurve* FindRichCurve(const FAnimationCurveIdentifier& CurveIdentifier) const override;
+	virtual const FAnimationCurveData& GetCurveData() const override;
+	virtual int32 GetNumberOfTransformCurves() const override;
+	virtual int32 GetNumberOfFloatCurves() const override;
+	virtual const TArray<struct FFloatCurve>& GetFloatCurves() const override;
+	virtual const TArray<struct FTransformCurve>& GetTransformCurves() const override;	
+	virtual const FAnimCurveBase& GetCurve(const FAnimationCurveIdentifier& CurveIdentifier) const override;
+	virtual const FFloatCurve& GetFloatCurve(const FAnimationCurveIdentifier& CurveIdentifier) const override;
+	virtual const FTransformCurve& GetTransformCurve(const FAnimationCurveIdentifier& CurveIdentifier) const override;
+	virtual const FRichCurve& GetRichCurve(const FAnimationCurveIdentifier& CurveIdentifier) const override;
+	virtual TArrayView<const FAnimatedBoneAttribute> GetAttributes() const override;
+	virtual int32 GetNumberOfAttributes() const override;
+	virtual int32 GetNumberOfAttributesForBoneIndex(const int32 BoneIndex) const override;
+	virtual void GetAttributesForBone(const FName& BoneName, TArray<const FAnimatedBoneAttribute*>& OutBoneAttributes) const override;
+	virtual const FAnimatedBoneAttribute& GetAttribute(const FAnimationAttributeIdentifier& AttributeIdentifier) const override;
+	virtual const FAnimatedBoneAttribute* FindAttribute(const FAnimationAttributeIdentifier& AttributeIdentifier) const override;
+	virtual UAnimSequence* GetAnimationSequence() const override;
+	virtual FAnimDataModelModifiedEvent& GetModifiedEvent() override { return ModifiedEvent; }
+	virtual FGuid GenerateGuid() const override;
+#if WITH_EDITOR
+	virtual void Evaluate(FAnimationPoseData& InOutPoseData, const UE::Anim::DataModel::FEvaluationContext& EvaluationContext) const override;
+#endif
+	virtual TScriptInterface<IAnimationDataController> GetController() override;
+	virtual bool HasBeenPopulated() const override { return bPopulated; }
+protected:
+	virtual IAnimationDataModel::FModelNotifier& GetNotifier() override;
+	virtual FAnimDataModelModifiedDynamicEvent& GetModifiedDynamicEvent() override { return ModifiedEventDynamic; }
+	virtual void OnNotify(const EAnimDataModelNotifyType& NotifyType, const FAnimDataModelNotifPayload& Payload) override {}
+	/** End IAnimationDataModel overrides */
 
-	/**
-	* @return	Whether or not the provided track index is valid 
-	*/
-	UFUNCTION(BlueprintPure, Category = AnimationDataModel)
-	bool IsValidBoneTrackIndex(int32 TrackIndex) const;
-
-	/**
-	* @return	Total number of bone animation tracks
-	*/
-	UFUNCTION(BlueprintPure, Category = AnimationDataModel)
-	const int32 GetNumBoneTracks() const;
-
-	/**
-	* Populates the provided array with all contained (bone) track names
-	*
-	* @param	OutNames	[out] Array containing all bone track names
-	*/
-	UFUNCTION(BlueprintPure, Category = AnimationDataModel)
-	void GetBoneTrackNames(TArray<FName>& OutNames) const;
-
-	/** Returns all contained curve animation data */
-	const FAnimationCurveData& GetCurveData() const;
-
-	/**
-	* @return	Total number of stored FTransform curves
-	*/
-	UFUNCTION(BlueprintPure, Category = AnimationDataModel)
-	int32 GetNumberOfTransformCurves() const;
-
-	/**
-	* @return	Total number of stored float curves
-	*/
-	UFUNCTION(BlueprintPure, Category = AnimationDataModel)
-	int32 GetNumberOfFloatCurves() const;
-
-	/**
-	* @return	Array containing all stored float curves 
-	*/
-	const TArray<struct FFloatCurve>& GetFloatCurves() const;
-
-	/**
-	* @return	Array containing all stored FTransform curves 
-	*/
-	const TArray<struct FTransformCurve>& GetTransformCurves() const;
-
-	/**
-	* @return	Curve ptr for the provided identifier if valid, otherwise returns a nullptr 
-	*/
-	const FAnimCurveBase* FindCurve(const FAnimationCurveIdentifier& CurveIdentifier) const;
-
-	/**
-	* @return	Float Curve ptr for the provided identifier if valid, otherwise returns a nullptr
-	*/
-	const FFloatCurve* FindFloatCurve(const FAnimationCurveIdentifier& CurveIdentifier) const;
-
-	/**
-	* @return	Transform Curve ptr for the provided identifier if valid, otherwise returns a nullptr
-	*/
-	const FTransformCurve* FindTransformCurve(const FAnimationCurveIdentifier& CurveIdentifier) const;
-
-	/**
-	* @return	Rich curve ptr for the provided identifier if valid, otherwise returns a nullptr
-	*/
-	const FRichCurve* FindRichCurve(const FAnimationCurveIdentifier& CurveIdentifier) const;
-
-	/**
-	* @return	Curve object for the provided identifier if valid
-	*/
-	const FAnimCurveBase& GetCurve(const FAnimationCurveIdentifier& CurveIdentifier) const;
-
-	/**
-	* @return	Float Curve object for the provided identifier if valid
-	*/
-	const FFloatCurve& GetFloatCurve(const FAnimationCurveIdentifier& CurveIdentifier) const;
-
-	/**
-	* @return	Transform Curve object for the provided identifier if valid
-	*/
-	const FTransformCurve& GetTransformCurve(const FAnimationCurveIdentifier& CurveIdentifier) const;
-
-	/**
-	* @return	Rich Curve object for the provided identifier if valid
-	*/
-	const FRichCurve& GetRichCurve(const FAnimationCurveIdentifier& CurveIdentifier) const;
-		
-	/**
-	* @return	Animated (bone) attributes stored
-	*/
-	TArrayView<const FAnimatedBoneAttribute> GetAttributes() const;
-
-	/**
-	* @return	Number of animated (bone) attributes stored
-	*/
-	int32 GetNumberOfAttributes() const;
-
-	/**
-	* @return	Number of animated (bone) attributes stored for the specified bone index
-	*/
-	int32 GetNumberOfAttributesForBoneIndex(const int32 BoneIndex) const;
-
-	/**
-	* @return	All animated (bone) attributes stored for the specified bone name
-	*/
-	void GetAttributesForBone(const FName& BoneName, TArray<const FAnimatedBoneAttribute*>& OutBoneAttributes) const;
-
-	/**
-	* @return	Animated (bone) attribute object for the provided identifier if valid
-	*/
-	const FAnimatedBoneAttribute& GetAttribute(const FAnimationAttributeIdentifier& AttributeIdentifier) const;
-
-	/**
-	* @return	Animated (bone) attribute ptr for the provided identifier if valid, otherwise returns a nullptr 
-	*/
-	const FAnimatedBoneAttribute* FindAttribute(const FAnimationAttributeIdentifier& AttributeIdentifier) const;
-		
-	/**
-	* @return	The outer UAnimSequence object if found, otherwise returns a nullptr 
-	*/
-	UFUNCTION(BlueprintPure, Category = AnimationDataModel)
-	UAnimSequence* GetAnimationSequence() const;
-
-	/**
-	* @return	Multicast delegate which is broadcasted to propagated changes to any internal data, see FAnimDataModelModifiedEvent and EAnimDataModelNotifyType
-	*/
-	FAnimDataModelModifiedEvent& GetModifiedEvent() { return ModifiedEvent; }
-
-	/**
-	* @return	GUID representing the contained data and state 
-	*/
-	FGuid GenerateGuid() const;
 private:
 	/** Helper functionality used by UAnimDataController to retrieve mutable data */ 
 	FBoneAnimationTrack* FindMutableBoneTrackByName(FName Name);
@@ -286,77 +82,10 @@ private:
 	FAnimCurveBase* FindMutableCurveById(const FAnimationCurveIdentifier& CurveIdentifier);	   
 	FRichCurve* GetMutableRichCurve(const FAnimationCurveIdentifier& CurveIdentifier);
 
-	/**
-	* Broadcasts a new EAnimDataModelNotifyType with the provided payload data alongside it.
-	*
-	* @param	NotifyType			Type of notify to broadcast
-	* @param	PayloadData			Typed payload data
-	*/
-	template<typename T>
-	void Notify(EAnimDataModelNotifyType NotifyType, const T& PayloadData)
-	{
-		UScriptStruct* TypeScriptStruct = T::StaticStruct();
-
-		const FAnimDataModelNotifPayload Payload((int8*)&PayloadData, TypeScriptStruct);
-		ModifiedEvent.Broadcast(NotifyType, this, Payload);
-
-		if (ModifiedEventDynamic.IsBound())
-		{
-			ModifiedEventDynamic.Broadcast(NotifyType, this, Payload);
-		}
-
-		// Only regenerate transient data when not in a bracket, or at the end of one
-		{
-			if (NotifyType == EAnimDataModelNotifyType::BracketOpened)
-			{
-				++BracketCounter;
-			}
-			if (NotifyType == EAnimDataModelNotifyType::BracketClosed)
-			{
-				--BracketCounter;
-			}
-
-			check(BracketCounter >= 0);
-		}
-	}
-
-	/**
-	* Broadcasts a new EAnimDataModelNotifyType alongside of an empty payload.
-	*
-	* @param	NotifyType			Type of notify to broadcast
-	*/
-	void Notify(EAnimDataModelNotifyType NotifyType)
-	{
-		FEmptyPayload EmptyPayload;
-		const FAnimDataModelNotifPayload Payload((int8*)&EmptyPayload, FEmptyPayload::StaticStruct());
-
-		ModifiedEvent.Broadcast(NotifyType, this, Payload);
-
-		if (ModifiedEventDynamic.IsBound())
-		{
-			ModifiedEventDynamic.Broadcast(NotifyType, this, Payload);
-		}
-
-		// Only regenerate transient data when not in a bracket, or at the end of one
-		{
-			if (NotifyType == EAnimDataModelNotifyType::BracketOpened)
-			{
-				++BracketCounter;
-			}
-			if (NotifyType == EAnimDataModelNotifyType::BracketClosed)
-			{
-				--BracketCounter;
-			}
-
-			check(BracketCounter >= 0);
-		}
-	}
-
 private:
-
 	UPROPERTY(Transient)
 	int32 BracketCounter = 0;
-private:
+
 	/** Dynamic delegate event allows scripting to register to any broadcasted notify. */
 	UPROPERTY(BlueprintAssignable, Transient, Category = AnimationDataModel, meta = (ScriptName = "ModifiedEvent", AllowPrivateAccess = "true"))
 	FAnimDataModelModifiedDynamicEvent ModifiedEventDynamic;
@@ -369,6 +98,7 @@ private:
 	TArray<FBoneAnimationTrack> BoneAnimationTracks;
 
 	/** Total playable length of the contained animation data */
+	UE_DEPRECATED(5.1, "PlayLength is deprecated use GetPlayLength instead, as it is now calculated with Number of Frames * FrameRate instead of stored as a value")
 	UPROPERTY(VisibleAnywhere, BlueprintReadOnly, Category = "Animation Data Model", meta = (AllowPrivateAccess = "true"))
 	float PlayLength;
 	
@@ -390,7 +120,12 @@ private:
 	
 	/** Container with all animated (bone) attribute data */
 	UPROPERTY(VisibleAnywhere, BlueprintReadOnly, Category = "Animation Data Model", meta = (AllowPrivateAccess = "true"))
-	TArray<FAnimatedBoneAttribute> AnimatedBoneAttributes;	
+	TArray<FAnimatedBoneAttribute> AnimatedBoneAttributes;
+
+	UPROPERTY()
+	bool bPopulated = false;
+
+	TUniquePtr<IAnimationDataModel::FModelNotifier> Notifier;
 
 	friend class UAnimDataController;
 	friend class FAnimDataControllerTestBase;
diff --git a/Engine/Source/Runtime/Engine/Classes/Animation/AnimData/AnimDataNotifications.h b/Engine/Source/Runtime/Engine/Classes/Animation/AnimData/AnimDataNotifications.h
--- a/Engine/Source/Runtime/Engine/Classes/Animation/AnimData/AnimDataNotifications.h
+++ b/Engine/Source/Runtime/Engine/Classes/Animation/AnimData/AnimDataNotifications.h
@@ -13,7 +13,7 @@
 
 #include "AnimDataNotifications.generated.h"
 
-class UAnimDataModel;
+class IAnimationDataModel;
 
 UENUM(BlueprintType)
 enum class EAnimDataModelNotifyType : uint8
@@ -122,17 +122,40 @@ struct FSequenceLengthChangedPayload : public FEmptyPayload
 {
 	GENERATED_BODY()
 
+	PRAGMA_DISABLE_DEPRECATION_WARNINGS
+	FSequenceLengthChangedPayload() = default;
+	FSequenceLengthChangedPayload(const FSequenceLengthChangedPayload&) = default;
+	FSequenceLengthChangedPayload(FSequenceLengthChangedPayload&&) = default;
+	FSequenceLengthChangedPayload& operator=(const FSequenceLengthChangedPayload&) = default;
+	FSequenceLengthChangedPayload& operator=(FSequenceLengthChangedPayload&&) = default;
+	PRAGMA_ENABLE_DEPRECATION_WARNINGS
+
 	/** Previous playable length for the Model */
+	UE_DEPRECATED(5.0, "Time in seconds is deprecated use FFrameNumber members (PreviousNumberOfFrames) instead")
 	UPROPERTY(BlueprintReadOnly, Category = Payload)
 	float PreviousLength = 0.f;
 
 	/** Time at which the change in length has been made */
+	UE_DEPRECATED(5.0, "Time in seconds is deprecated use FFrameNumber members (Frame0) instead")
 	UPROPERTY(BlueprintReadOnly, Category = Payload)
 	float T0 = 0.f;
 
 	/** Length of time which is inserted or removed starting at T0 */
+	UE_DEPRECATED(5.0, "Time in seconds is deprecated use FFrameNumber members (Frame1) instead")
 	UPROPERTY(BlueprintReadOnly, Category = Payload)
 	float T1 = 0.f;
+
+	/** Previous playable number of frames for the Model */
+    UPROPERTY(BlueprintReadOnly, Category = Payload)
+    FFrameNumber PreviousNumberOfFrames;
+
+    /** Frame number at which the change in frames has been made */
+    UPROPERTY(BlueprintReadOnly, Category = Payload)
+    FFrameNumber Frame0;
+
+    /** Amount of frames which is inserted or removed starting at Frame0 */
+    UPROPERTY(BlueprintReadOnly, Category = Payload)
+	FFrameNumber Frame1;
 };
 
 USTRUCT(BlueprintType)
@@ -257,5 +280,5 @@ public:
 #endif // WITH_EDITOR
 };
 
-DECLARE_MULTICAST_DELEGATE_ThreeParams(FAnimDataModelModifiedEvent, const EAnimDataModelNotifyType& /* type */, UAnimDataModel* /* Model */, const FAnimDataModelNotifPayload& /* payload */);
-DECLARE_DYNAMIC_MULTICAST_DELEGATE_ThreeParams(FAnimDataModelModifiedDynamicEvent, EAnimDataModelNotifyType, NotifType, UAnimDataModel*, Model, const FAnimDataModelNotifPayload&, Payload);
\ No newline at end of file
+DECLARE_MULTICAST_DELEGATE_ThreeParams(FAnimDataModelModifiedEvent, const EAnimDataModelNotifyType& /* type */, IAnimationDataModel* /* Model */, const FAnimDataModelNotifPayload& /* payload */);
+DECLARE_DYNAMIC_MULTICAST_DELEGATE_ThreeParams(FAnimDataModelModifiedDynamicEvent, EAnimDataModelNotifyType, NotifType, TScriptInterface<IAnimationDataModel>, Model, const FAnimDataModelNotifPayload&, Payload);
\ No newline at end of file
diff --git a/Engine/Source/Runtime/Engine/Classes/Animation/AnimData/AttributeIdentifier.h b/Engine/Source/Runtime/Engine/Classes/Animation/AnimData/AttributeIdentifier.h
--- a/Engine/Source/Runtime/Engine/Classes/Animation/AnimData/AttributeIdentifier.h
+++ b/Engine/Source/Runtime/Engine/Classes/Animation/AnimData/AttributeIdentifier.h
@@ -80,7 +80,6 @@ protected:
 	FSoftObjectPath ScriptStructPath;
 
 	friend class UAnimationAttributeIdentifierExtensions;
-	friend class UAnimDataController;
 };
 
 template<>
diff --git a/Engine/Source/Runtime/Engine/Classes/Animation/AnimData/CurveIdentifier.h b/Engine/Source/Runtime/Engine/Classes/Animation/AnimData/CurveIdentifier.h
--- a/Engine/Source/Runtime/Engine/Classes/Animation/AnimData/CurveIdentifier.h
+++ b/Engine/Source/Runtime/Engine/Classes/Animation/AnimData/CurveIdentifier.h
@@ -69,10 +69,22 @@ struct ENGINE_API FAnimationCurveIdentifier
 
 		return *this;
 	}
+	
+	friend uint32 GetTypeHash(const FAnimationCurveIdentifier& CurveId)
+	{
+		return HashCombine(HashCombine(HashCombine((uint32)CurveId.Channel, (uint32)CurveId.Axis), (uint32)CurveId.CurveType), GetTypeHash(CurveId.InternalName.DisplayName));
+	}
 
+	UPROPERTY()
 	FSmartName InternalName;
+	
+	UPROPERTY()
 	ERawCurveTrackTypes CurveType;
+
+	UPROPERTY()
 	ETransformCurveChannel Channel;
+	
+	UPROPERTY()
 	EVectorCurveChannel Axis;
 };
 
diff --git a/Engine/Source/Runtime/Engine/Classes/Animation/AnimData/IAnimationDataController.h b/Engine/Source/Runtime/Engine/Classes/Animation/AnimData/IAnimationDataController.h
--- a/Engine/Source/Runtime/Engine/Classes/Animation/AnimData/IAnimationDataController.h
+++ b/Engine/Source/Runtime/Engine/Classes/Animation/AnimData/IAnimationDataController.h
@@ -7,12 +7,25 @@
 #include "CurveIdentifier.h"
 #include "Misc/FrameRate.h"
 #include "AttributeIdentifier.h"
+#include "IAnimationDataModel.h"
+#include "Logging/LogVerbosity.h"
 #include "Algo/Transform.h"
+#include "Animation/AnimCurveTypes.h"
+
+#if WITH_EDITOR
+#include "ChangeTransactor.h"
+#endif // WITH_EDITOR
 
 #include "IAnimationDataController.generated.h"
 
+#define LOCTEXT_NAMESPACE "IAnimationDataController"
+
+class FText;
 class UAssetUserData;
-class UAnimDataModel;
+class IAnimationDataModel;
+
+struct FAnimationCurveIdentifier;
+struct FAnimationAttributeIdentifier;
 
 namespace UE {
 namespace Anim {
@@ -22,7 +35,7 @@ namespace Anim {
 
 /**
  * The Controller is the sole authority to perform changes on the Animation Data Model. Any mutation to the model made will
- * cause a subsequent notify (EAnimDataModelNotifyType) to be broadcasted from the Model's ModifiedEvent. Alongside of it is a 
+ * cause a subsequent notify (EAnimDataModelNotifyType) to be broadcast-ed from the Model's ModifiedEvent. Alongside of it is a 
  * payload containing information relevant to the mutation. These notifies should be relied upon to update any dependent views 
  * or generated (derived) data.
  */
@@ -53,7 +66,7 @@ public:
 			Controller.OpenBracket(InDescription, bShouldTransact);
 		}
 		
-		FScopedBracket(TScriptInterface<IAnimationDataController>& InController, const FText& InDescription, bool bInShouldTransact=true)
+		FScopedBracket(const TScriptInterface<IAnimationDataController>& InController, const FText& InDescription, bool bInShouldTransact=true)
             : Controller(*InController), bShouldTransact(bInShouldTransact)
 		{
 			Controller.OpenBracket(InDescription, bShouldTransact);
@@ -72,25 +85,25 @@ public:
 	/**
 	* Sets the AnimDataModel instance this controller is supposed to be targeting
 	*
-	* @param	InModel		UAnimDataModel instance to target
+	* @param	InModel		IAnimationDataModel instance to target
 	*/
 	UFUNCTION(BlueprintCallable, Category = AnimationData)
-	virtual void SetModel(UAnimDataModel* InModel) = 0;
+	virtual void SetModel(TScriptInterface<IAnimationDataModel> InModel) = 0;
 
 	/**
-	* @return		The AnimDataModel instance this controller is currently targeting
+	* @return		The IAnimationDataModel instance this controller is currently targeting
 	*/
 	UFUNCTION(BlueprintCallable, Category = AnimationData)
-	virtual UAnimDataModel* GetModel() = 0;
+	virtual TScriptInterface<IAnimationDataModel> GetModelInterface() const = 0;
 
 	/**
-	* @return		The AnimDataModel instance this controller is currently targeting
+	* @return		The IAnimationDataModel instance this controller is currently targeting
 	*/
-	virtual const UAnimDataModel* const GetModel() const = 0;
+	virtual const IAnimationDataModel* const GetModel() const = 0;
 
 	/**
 	* Opens an interaction bracket, used for combining a set of controller actions. Broadcasts a EAnimDataModelNotifyType::BracketOpened notify,
-	* this can be used by any Views or dependendent systems to halt any unnecessary or invalid operations until the (last) bracket is closed.
+	* this can be used by any Views or dependent systems to halt any unnecessary or invalid operations until the (last) bracket is closed.
 	*
 	* @param	InTitle				Description of the bracket, e.g. "Generating Curve Data"
 	* @param	bShouldTransact		Whether or not any undo-redo changes should be generated
@@ -101,7 +114,6 @@ public:
 	/**
 	* Closes a previously opened interaction bracket, used for combining a set of controller actions. Broadcasts a EAnimDataModelNotifyType::BracketClosed notify.
 	*
-	* @param	InTitle				Description of the bracket, e.g. "Generating Curve Data"
 	* @param	bShouldTransact		Whether or not any undo-redo changes should be generated
 	*/
 	UFUNCTION(BlueprintCallable, Category = AnimationData)
@@ -111,22 +123,30 @@ public:
 	* Sets the total play-able length in seconds. Broadcasts a EAnimDataModelNotifyType::SequenceLengthChanged notify if successful.
 	* The number of frames and keys for the provided length is recalculated according to the current value of UAnimDataModel::FrameRate.
 	*
-	* @param	Length				New play-able length value, has to be positive and non-zero
+	* @param	NewLengthInFrames	Total new play-able number of frames value (according to frame rate), has to be positive and non-zero
 	* @param	bShouldTransact		Whether or not any undo-redo changes should be generated
 	*/
 	UFUNCTION(BlueprintCallable, Category = AnimationData)
+	virtual void SetNumberOfFrames(FFrameNumber NewLengthInFrames, bool bShouldTransact = true) = 0;
+
+	UE_DEPRECATED(5.1, "SetPlayLength with length in seconds in deprecated use SetNumberOfFrames using FFrameNumber instead")
+	UFUNCTION(BlueprintCallable, Category = AnimationData, meta=(DeprecatedFunction, DeprecationMessage="SetPlayLength is deprecated use SetNumberOfFrames instead."))
 	virtual void SetPlayLength(float Length, bool bShouldTransact = true) = 0;
 
 	/*** Sets the total play-able length in seconds. Broadcasts a EAnimDataModelNotifyType::SequenceLengthChanged notify if successful.
 	* T0 and T1 are expected to represent the window of time that was either added or removed. E.g. for insertion T0 indicates the time
 	* at which additional time starts and T1 were it ends. For removal T0 indicates the time at which time should be started to remove, and T1 indicates the end. Giving a total of T1 - T0 added or removed length.
 	* The number of frames and keys for the provided length is recalculated according to the current value of UAnimDataModel::FrameRate.
-	* @param	Length				Total new play-able length value, has to be positive and non-zero
-	* @param	T0					Point between 0 and Length at which the change in time starts
-	* @param	T1					Point between 0 and Length at which the change in time ends
+	* @param	NewLengthInFrames	Total new play-able number of frames value (according to frame rate), has to be positive and non-zero
+	* @param	T0					Point between 0 and NewLengthInFrames at which the change in length starts
+	* @param	T1					Point between 0 and NewLengthInFrames at which the change in length ends
 	* @param	bShouldTransact		Whether or not any undo-redo changes should be generated
 	*/
 	UFUNCTION(BlueprintCallable, Category = AnimationData)
+    virtual void ResizeNumberOfFrames(FFrameNumber NewLengthInFrames, FFrameNumber T0, FFrameNumber T1, bool bShouldTransact = true) = 0;
+
+	UE_DEPRECATED(5.1, "ResizePlayLength with length in seconds in deprecated use ResizeNumberOfFrames using FFrameNumber instead")
+	UFUNCTION(BlueprintCallable, Category = AnimationData, meta=(DeprecatedFunction, DeprecationMessage="ResizePlayLength is deprecated use ResizeNumberOfFrames instead."))
 	virtual void ResizePlayLength(float NewLength, float T0, float T1, bool bShouldTransact = true) = 0;
 
 	/**
@@ -136,14 +156,18 @@ public:
 	* at which additional time starts and T1 were it ends. For removal T0 indicates the time at which time should be started to remove, and T1 indicates the end. Giving a total of T1 - T0 added or removed length.
 	* The number of frames and keys for the provided length is recalculated according to the current value of UAnimDataModel::FrameRate.
 	*
-	* @param	Length				Total new play-able length value, has to be positive and non-zero
-	* @param	T0					Point between 0 and Length at which the change in time starts
-	* @param	T1					Point between 0 and Length at which the change in time ends
+	* @param	NewLengthInFrames	Total new play-able number of frames value (according to frame rate), has to be positive and non-zero
+	* @param	T0					Point between 0 and NewLengthInFrames at which the change in length starts
+	* @param	T1					Point between 0 and NewLengthInFrames at which the change in length ends
 	* @param	bShouldTransact		Whether or not any undo-redo changes should be generated
 	*/
 	UFUNCTION(BlueprintCallable, Category = AnimationData)
-	virtual void Resize(float Length, float T0, float T1, bool bShouldTransact = true) = 0;
+	virtual void ResizeInFrames(FFrameNumber NewLengthInFrames, FFrameNumber T0, FFrameNumber T1, bool bShouldTransact = true) = 0;
 
+	UE_DEPRECATED(5.1, "Resize with length in seconds in deprecated use ResizeInFrames using FFrameNumber instead")
+	UFUNCTION(BlueprintCallable, Category = AnimationData, meta=(DeprecatedFunction, DeprecationMessage="Resize is deprecated use ResizeInFrames instead."))
+	virtual void Resize(float Length, float T0, float T1, bool bShouldTransact = true) = 0;
+	
 	/**
 	* Sets the frame rate according to which the bone animation is expected to be sampled. Broadcasts a EAnimDataModelNotifyType::FrameRateChanged notify if successful.
 	* The number of frames and keys for the provided frame rate is recalculated according to the current value of UAnimDataModel::PlayLength.
@@ -175,7 +199,8 @@ public:
 	*
 	* @return	The index at which the bone track was inserted, INDEX_NONE if the insertion failed
 	*/
-	UFUNCTION(BlueprintCallable, Category = AnimationData)
+	UE_DEPRECATED(5.1, "InsertBoneTrack has been deprecated, use AddBoneTrack instead")
+	UFUNCTION(BlueprintCallable, Category = AnimationData, meta=(DeprecatedFunction, DeprecationMessage="InsertBoneTrack is deprecated use AddBoneTrack instead."))
 	virtual int32 InsertBoneTrack(FName BoneName, int32 DesiredIndex, bool bShouldTransact = true) = 0;
 
 	/**
@@ -231,7 +256,7 @@ public:
 	* The provided number of keys provided is expected to match for each component, be between FrameLowerBound and FrameUpperBound (inclusive), and be non-zero.
 	*
 	* @param	BoneName			Bone name of the track for which the keys should be set
-	* @param	FrameRange			Range of frames to set keys for
+	* @param	KeyRangeToSet			Range of frames to set keys for
 	* @param	PositionalKeys		Array of keys for the translation component
 	* @param	RotationalKeys		Array of keys for the rotation component
 	* @param	ScalingKeys			Array of keys for the scale component
@@ -246,7 +271,7 @@ public:
 	* The provided number of keys provided is expected to match for each component, be between FrameLowerBound and FrameUpperBound (inclusive), and be non-zero.
 	*
 	* @param	BoneName			Bone name of the track for which the keys should be set
-	* @param	FrameRange			Range of frames to set keys for
+	* @param	KeyRangeToSet			Range of frames to set keys for
 	* @param	PositionalKeys		Array of keys for the translation component
 	* @param	RotationalKeys		Array of keys for the rotation component
 	* @param	ScalingKeys			Array of keys for the scale component
@@ -284,7 +309,7 @@ public:
 	/**
 	* Remove the curve with provided identifier. Broadcasts a EAnimDataModelNotifyType::CurveRemoved notify if successful.
 	*
-	* @param	CopyCurveId			Identifier for the to-be-removed curve
+	* @param	CurveId				Identifier for the to-be-removed curve
 	* @param	bShouldTransact		Whether or not any undo-redo changes should be generated
 	*
 	* @return	Whether or not the curve was successfully removed
@@ -318,7 +343,7 @@ public:
 	* Replace the flags for the curve with provided identifier. Broadcasts a EAnimDataModelNotifyType::CurveFlagsChanged notify if successful.
 	*
 	* @param	CurveId			    Identifier for the curve for which the flag state is to be set
-	* @param	Flags				Flag mask with which the existings flags are to be replaced
+	* @param	Flags				Flag mask with which the existing flags are to be replaced
 	* @param	bShouldTransact		Whether or not any undo-redo changes should be generated
 	*
 	* @return	Whether or not the flag mask was successfully set
@@ -451,7 +476,7 @@ public:
 	virtual void UpdateCurveNamesFromSkeleton(const USkeleton* Skeleton, ERawCurveTrackTypes SupportedCurveType, bool bShouldTransact = true) = 0;
 
 	/**
-	* Updates the curve names with the provided skeleton, if a display name is not found it will be added thus modifying the skeleton. Broadcasts a EAnimDataModelNotifyType::CurveRenamed for each curve name for which the UID was different or if it was added as a new smartname, wrapped within BracketOpened/BracketClosed notifies.
+	* Updates the curve names with the provided skeleton, if a display name is not found it will be added thus modifying the skeleton. Broadcasts a EAnimDataModelNotifyType::CurveRenamed for each curve name for which the UID was different or if it was added as a new smart-name, wrapped within BracketOpened/BracketClosed notifies.
 	*
 	* @param	Skeleton			Skeleton to retrieve the display name values from
 	* @param	SupportedCurveType	Curve type for which the names should be updated
@@ -548,9 +573,9 @@ public:
 	* @return	Whether or not the key was successfully set
 	*/
 	template<typename AttributeType>
-	bool SetTypedAttributeKey(const FAnimationAttributeIdentifier& AttributeIdentifier, float Time, AttributeType& KeyValue, bool bShouldTransact = true)
+	bool SetTypedAttributeKey(const FAnimationAttributeIdentifier& AttributeIdentifier, float Time, const AttributeType& KeyValue, bool bShouldTransact = true)
 	{
-		return SetAttributeKey(AttributeIdentifier, Time, (const void*)&KeyValue, AttributeType::StaticStruct(), bShouldTransact);
+		return SetAttributeKey(AttributeIdentifier, Time, static_cast<const void*>(&KeyValue), AttributeType::StaticStruct(), bShouldTransact);
 	}
 
 	/**
@@ -560,6 +585,7 @@ public:
 	* @param	AttributeIdentifier		Identifier for the attribute for which the key is to be set
 	* @param	Time					Time of the to-be-set key
 	* @param	KeyValue				Value of the to-be-set key
+	* @param	TypeStruct				UScriptStruct describing the type of KeyValue
 	* @param	bShouldTransact			Whether or not any undo-redo changes should be generated
 	*
 	* @return	Whether or not the key was successfully set
@@ -572,6 +598,7 @@ public:
 	* @param	AttributeIdentifier		Identifier for the attribute for which the keys are to be replaced
 	* @param	Times					Times with which the existing key timings are to be replaced
 	* @param	KeyValues				Values with which the existing key values are to be replaced
+	* @param	TypeStruct				UScriptStruct describing the type of KeyValues
 	* @param	bShouldTransact			Whether or not any undo-redo changes should be generated
 	*
 	* @return	Whether or not replacing the attribute keys was successful
@@ -589,12 +616,12 @@ public:
 	* @return	Whether or not replacing the attribute keys was successful
 	*/
 	template<typename AttributeType>
-	bool SetTypedAttributeKeys(const FAnimationAttributeIdentifier& AttributeIdentifier, TArrayView<const float> Times, TArrayView<AttributeType> KeyValues, bool bShouldTransact = true)
+	bool SetTypedAttributeKeys(const FAnimationAttributeIdentifier& AttributeIdentifier, TArrayView<const float> Times, TArrayView<const AttributeType> KeyValues, bool bShouldTransact = true)
 	{
 		TArray<const void*> KeyValuePtrs;
 		Algo::Transform(KeyValues, KeyValuePtrs, [](const AttributeType& Value)
 		{
-			return (const void*)&Value;
+			return static_cast<const void*>(&Value);
 		});
 
 		return SetAttributeKeys(AttributeIdentifier, Times, MakeArrayView(KeyValuePtrs), AttributeType::StaticStruct(), bShouldTransact);
@@ -623,13 +650,175 @@ public:
 	*/
 	UFUNCTION(BlueprintCallable, Category = AttributeData)
 	virtual bool DuplicateAttribute(const FAnimationAttributeIdentifier& AttributeIdentifier, const FAnimationAttributeIdentifier& NewAttributeIdentifier, bool bShouldTransact = true) = 0;
+
+	/** Updates/removes/remaps contained animation data according to the newly assigned skeleton
+	 *
+	* @param	TargetSkeleton		Skeleton to which the contained animation data is to be mapped
+	* @param	bShouldTransact		Whether or not any undo-redo changes should be generated
+	* 
+	 */
+	virtual void UpdateWithSkeleton(USkeleton* TargetSkeleton, bool bShouldTransact = true) = 0;
+
+	/** Copies any animation relevant data from an already existing IAnimationDataModel object */
+	virtual void PopulateWithExistingModel(TScriptInterface<IAnimationDataModel> InModel) = 0;
+
+	/** Initializes model data structures */
+	virtual void InitializeModel() = 0;
+
+	/** Returns the final frame number calculating according to the Model its frame-rate, additionally outputs log information for invalid/loss of precision */
+	FFrameNumber ConvertSecondsToFrameNumber(double Seconds) const
+	{
+		ValidateModel();
+		
+		const FFrameRate& ModelFrameRate = GetModel()->GetFrameRate();
+		
+		static const FNumberFormattingOptions DurationFormatOptions = FNumberFormattingOptions()
+			.SetMinimumFractionalDigits(8)
+			.SetMaximumFractionalDigits(8);
+	
+		const FFrameTime FrameTime = ModelFrameRate.AsFrameTime(Seconds);
+		// Check for either small sub-frame value or near zero seconds representation (of sub-frame only)
+		if (!FMath::IsNearlyZero(FrameTime.GetSubFrame(), UE_KINDA_SMALL_NUMBER) &&
+			!FMath::IsNearlyZero(ModelFrameRate.AsSeconds(FFrameTime(0, FrameTime.GetSubFrame())), UE_DOUBLE_KINDA_SMALL_NUMBER))
+		{
+			ReportWarningf(LOCTEXT("SecondsToFrameNumberPrecisionWarning", "Insufficient precision while converting seconds to frames: {0} seconds {1} frames using {2} (sub-frame in seconds {3})"), FText::AsNumber(Seconds), FText::AsNumber(FrameTime.AsDecimal(), &DurationFormatOptions), ModelFrameRate.ToPrettyText(), FText::AsNumber(ModelFrameRate.AsSeconds(FFrameTime(0, FrameTime.GetSubFrame())), &DurationFormatOptions));
+		}
+
+		return FrameTime.GetFrame();		
+	}	
 protected:
 	/** Functionality used by FOpenBracketAction and FCloseBracketAction to broadcast their equivalent notifies without actually opening a bracket. */
 	virtual void NotifyBracketOpen() = 0;
 	virtual void NotifyBracketClosed() = 0;
 
+	template <typename FmtType, typename... Types>
+    void ReportWarningf(const FmtType& Fmt, Types... Args) const
+	{
+		ReportWarning(FText::Format(Fmt, Args...));
+	}
+
+	template <typename FmtType, typename... Types>
+    void ReportErrorf(const FmtType& Fmt, Types... Args) const
+	{
+		ReportError(FText::Format(Fmt, Args...));
+	}
+
+	void ReportWarning(const FText& Message) const
+	{
+		ReportMessage(GetModelInterface().GetObject(), Message, ELogVerbosity::Warning);
+	}
+	
+	void ReportError(const FText& Message) const
+    {
+    	ReportMessage(GetModelInterface().GetObject(), Message, ELogVerbosity::Error);
+    }
+
+	/** Returns whether or not the supplied curve type is supported by the controller functionality */
+	static bool IsSupportedCurveType(ERawCurveTrackTypes CurveType)
+	{
+		const TArray<ERawCurveTrackTypes> SupportedTypes = { ERawCurveTrackTypes::RCT_Float, ERawCurveTrackTypes::RCT_Transform };
+		return SupportedTypes.Contains(CurveType);
+	}
+
+	/** Ensures that a valid model is currently targeted */
+	void ValidateModel() const
+	{
+		checkf(GetModel() != nullptr, TEXT("Invalid Model"));
+	}
+		
+	/** Verifies whether or not the Model's outer object is (or is derived from) the specified UClass */
+	bool CheckOuterClass(UClass* InClass) const
+	{
+		ValidateModel();
+	
+		if (const UObject* ModelOuter = GetModelInterface().GetObject()->GetOuter())
+		{
+			if (const UClass* OuterClass = ModelOuter->GetClass())
+			{
+				if (OuterClass == InClass || OuterClass->IsChildOf(InClass))
+				{
+					return true;
+				}
+				else
+				{
+					ReportErrorf(NSLOCTEXT("IAnimationDataController", "NoValidOuterClassError", "Incorrect outer object class found for Animation Data Model {0}, expected {1} actual {2}"), FText::FromString(GetModelInterface().GetObject()->GetName()), FText::FromString(InClass->GetName()), FText::FromString(OuterClass->GetName()));
+				}
+			}
+		}
+		else
+		{
+			ReportErrorf(NSLOCTEXT("IAnimationDataController", "NoValidOuterObjectFoundError", "No valid outer object found for Animation Data Model {0}"), FText::FromString(GetModelInterface().GetObject()->GetName()));
+		}
+
+		return false;
+	}
+	
+	/** Returns the string representation of the provided curve enum type value */
+	static FString GetCurveTypeValueName(ERawCurveTrackTypes InType)
+	{
+		FString ValueString;
+		if (const UEnum* Enum = FindObject<UEnum>(nullptr, TEXT("/Script/Engine.ERawCurveTrackTypes")))
+		{
+			ValueString = Enum->GetNameStringByValue(static_cast<int64>(InType));
+		}
+
+		return ValueString;
+	}
+	
+	static constexpr int32 DefaultCurveFlags = EAnimAssetCurveFlags::AACF_Editable;
+
+#if WITH_EDITOR
+	typedef TUniquePtr<UE::FScopedCompoundTransaction> FTransaction;
+	FTransaction ConditionalTransaction(const FText& Description, bool bCondition)
+	{
+		if (UE::FChangeTransactor::CanTransactChanges() && bCondition)
+		{ 
+			return MakeUnique<UE::FScopedCompoundTransaction>(ChangeTransactor, Description);
+		}
+
+		return nullptr;
+	}
+
+	typedef TUniquePtr<IAnimationDataController::FScopedBracket> FBracket; 
+	FBracket ConditionalBracket(const FText& Description, bool bCondition)
+	{	
+		return MakeUnique<IAnimationDataController::FScopedBracket>(this, Description, UE::FChangeTransactor::CanTransactChanges() && bCondition);	
+	}
+
+	template<typename ActionClass, class... ActionArguments>
+	void ConditionalAction(bool bCondition, ActionArguments... Arguments)
+	{
+		if (UE::FChangeTransactor::CanTransactChanges() && bCondition) \
+		{ 
+			ChangeTransactor.AddTransactionChange<ActionClass>(Arguments...); \
+		}
+	}
+protected:
+	UE::FChangeTransactor ChangeTransactor;
+#endif // WITH_EDITOR
+
+public:	
+	template <typename FmtType, typename... Types>
+	static void ReportObjectWarningf(const UObject* ErrorObject, const FmtType& Fmt, Types... Args)
+	{
+		ReportMessage(ErrorObject, FText::Format(Fmt, Args...), ELogVerbosity::Warning);
+	}
+
+	template <typename FmtType, typename... Types>
+	static void ReportObjectErrorf(const UObject* ErrorObject, const FmtType& Fmt, Types... Args)
+	{
+		ReportMessage(ErrorObject, FText::Format(Fmt, Args...), ELogVerbosity::Error);
+	}
+
+	static void ReportMessage(const UObject* ErrorObject, const FText& InMessage, ELogVerbosity::Type LogVerbosity);	
 private:
 	friend class FAnimDataControllerTestBase;
 	friend UE::Anim::FOpenBracketAction;
 	friend UE::Anim::FCloseBracketAction;
 };
+
+#define CONDITIONAL_TRANSACTION(Text) FTransaction Transaction = ConditionalTransaction(Test, bShouldTransact);
+#define CONDITIONAL_BRACKET(Text) FBracket Bracket = ConditionalBracket(Text, bShouldTransact);
+#define CONDITIONAL_ACTION(ActionClass, ...) ConditionalAction<ActionClass>(bShouldTransact, __VA_ARGS__);
+
+#undef LOCTEXT_NAMESPACE // "IAnimationDataController"
\ No newline at end of file
diff --git a/Engine/Source/Runtime/Engine/Classes/Animation/AnimMontage.h b/Engine/Source/Runtime/Engine/Classes/Animation/AnimMontage.h
--- a/Engine/Source/Runtime/Engine/Classes/Animation/AnimMontage.h
+++ b/Engine/Source/Runtime/Engine/Classes/Animation/AnimMontage.h
@@ -496,7 +496,7 @@ public:
 	bool SetNextSectionName(FName const & SectionName, FName const & NewNextSectionName);
 	bool SetNextSectionID(int32 const & SectionID, int32 const & NewNextSectionID);
 
-	bool IsValid() const { return (Montage!=NULL); }
+	bool IsValid() const { return (Montage!=nullptr); }
 	bool IsPlaying() const { return IsValid() && bPlaying; }
 	void SetPlaying(bool bInPlaying) { bPlaying = bInPlaying; }
 	bool IsStopped() const { return Blend.GetDesiredValue() == 0.f; }
@@ -924,6 +924,9 @@ public:
 	virtual void InvalidateRecursiveAsset() override;
 	virtual bool ContainRecursive(TArray<UAnimCompositeBase*>& CurrentAccumulatedList) override;
 	virtual void SetCompositeLength(float InLength) override;
+#if WITH_EDITOR
+	virtual void PopulateWithExistingModel(TScriptInterface<IAnimationDataModel> ExistingDataModel) override;
+#endif // WITH_EDITOR
 	//~End UAnimCompositeBase Interface
 
 	/** Utility function to create dynamic montage from AnimSequence */
diff --git a/Engine/Source/Runtime/Engine/Classes/Animation/AnimSequence.h b/Engine/Source/Runtime/Engine/Classes/Animation/AnimSequence.h
--- a/Engine/Source/Runtime/Engine/Classes/Animation/AnimSequence.h
+++ b/Engine/Source/Runtime/Engine/Classes/Animation/AnimSequence.h
@@ -23,7 +23,7 @@
 #include "Animation/AttributeCurve.h"
 
 #if WITH_EDITOR
-#include "AnimData/AnimDataModel.h"
+#include "AnimData/IAnimationDataModel.h"
 #endif // WITH_EDITOR
 
 #include "AnimSequence.generated.h"
@@ -233,21 +233,11 @@ protected:
 	UPROPERTY()
 	FFrameRate SamplingFrameRate;
 
-	/**
-	 * In the future, maybe keeping RawAnimSequenceTrack + TrackMap as one would be good idea to avoid inconsistent array size
-	 * TrackToSkeletonMapTable(i) should contains  track mapping data for RawAnimationData(i). 
-	 */
-	UE_DEPRECATED(5.0, "TrackToSkeletonMapTable has been deprecated see FBoneAnimationTrack::BoneTreeIndex")
-	UPROPERTY()
-	TArray<struct FTrackToSkeletonMap> TrackToSkeletonMapTable;
-
-	/**
-	 * Raw uncompressed keyframe data. 
-	 */
 	UE_DEPRECATED(5.0, "RawAnimationData has been deprecated see FBoneAnimationTrack::InternalTrackData")
 	TArray<struct FRawAnimSequenceTrack> RawAnimationData;
 
-	// Update this if the contents of RawAnimationData changes;
+	// Update this if the contents of RawAnimationData changes
+	UE_DEPRECATED(5.1, "RawDataGuid has been deprecated see GenerateGuidFromModel instead")
 	UPROPERTY()
 	FGuid RawDataGuid;
 
@@ -262,7 +252,6 @@ protected:
 	 * Source RawAnimationData. Only can be overridden by when transform curves are added first time OR imported
 	 */
 	TArray<struct FRawAnimSequenceTrack> SourceRawAnimationData_DEPRECATED;
-
 public:
 
 	/**
@@ -370,10 +359,6 @@ public:
 	UPROPERTY()
 	FString SourceFileTimestamp_DEPRECATED;
 
-	UE_DEPRECATED(5.0, "bNeedsRebake has been deprecated, transform curves are now baked during compression")
-	UPROPERTY(transient)
-	bool bNeedsRebake;
-
 	// Track whether we have updated markers so cached data can be updated
 	int32 MarkerDataUpdateCounter;
 #endif // WITH_EDITORONLY_DATA
@@ -414,24 +399,10 @@ public:
 	virtual void RefreshCacheData() override;
 	virtual EAdditiveAnimationType GetAdditiveAnimType() const override { return AdditiveAnimType; }
 	virtual int32 GetNumberOfSampledKeys() const override;
-	virtual const FFrameRate& GetSamplingFrameRate() const override { return TargetFrameRate; }
+	virtual FFrameRate GetSamplingFrameRate() const override { return TargetFrameRate; }
 	virtual void EvaluateCurveData(FBlendedCurve& OutCurve, float CurrentTime, bool bForceUseRawData = false) const override;
 	virtual float EvaluateCurveData(SmartName::UID_Type CurveUID, float CurrentTime, bool bForceUseRawData = false) const override;
 	virtual bool HasCurveData(SmartName::UID_Type CurveUID, bool bForceUseRawData) const override;
-
-#if WITH_EDITOR
-	UE_DEPRECATED(5.0, "MarkRawDataAsModified has been deprecated, any (Raw Data) modification should be applied using the UAnimDataController API instead. This will handle updating the GUID instead.")
-	PRAGMA_DISABLE_DEPRECATION_WARNINGS
-	virtual void MarkRawDataAsModified(bool bForceNewRawDatGuid = true) override
-	{
-		Super::MarkRawDataAsModified();
-		bUseRawDataOnly = true;
-		RawDataGuid = bForceNewRawDatGuid ? FGuid::NewGuid() : GenerateGuidFromRawData();
-		FlagDependentAnimationsAsRawDataOnly();
-		UpdateDependentStreamingAnimations();
-	}
-	PRAGMA_ENABLE_DEPRECATION_WARNINGS
-#endif
 	//~ End UAnimSequenceBase Interface
 
 	// Extract Root Motion transform from the animation
@@ -472,9 +443,7 @@ public:
 #if WITH_EDITORONLY_DATA
 protected:
 	void UpdateCompressedCurveName(SmartName::UID_Type CurveUID, const struct FSmartName& NewCurveName);
-#endif
 private:
-#if WITH_EDITORONLY_DATA
 	void UpdateRetargetSourceAsset();
 
 	/** Updates the stored sampling frame-rate using the sequence length and number of sampling keys */
@@ -526,7 +495,9 @@ public:
 	 * @param	Time			Time on track to interpolate to.
 	 * @param	bUseRawData		If true, use raw animation data instead of compressed data.
 	 */
+	UE_DEPRECATED(5.1, "Use other GetBoneTransform signature")
 	void GetBoneTransform(FTransform& OutAtom, int32 TrackIndex, float Time, bool bUseRawData) const;
+	void GetBoneTransform(FTransform& OutAtom, int32 TrackIndex, double Time, bool bUseRawData) const;
 
 	/**
 	 * Get Bone Transform of the Time given, relative to Parent for the Track Given
@@ -587,10 +558,14 @@ public:
 
 protected:
 	void ApplyCompressedData(const TArray<uint8>& Data);
-    void ClearCompressedBoneData();
-    void ClearCompressedCurveData();
-    // Write the compressed data to the supplied FArchive
-    void SerializeCompressedData(FArchive& Ar, bool bDDCData);
+	void ClearCompressedBoneData();
+	void ClearCompressedCurveData();
+	// Write the compressed data to the supplied FArchive
+	void SerializeCompressedData(FArchive& Ar, bool bDDCData);
+
+#if WITH_EDITOR
+	virtual void OnAnimModelLoaded() override;
+#endif
 public:
 	bool IsCompressedDataValid() const;
 	bool IsCurveCompressedDataValid() const;
@@ -700,6 +675,7 @@ public:
 	void UpdateDependentStreamingAnimations() const;
 
 	// Generate a GUID from a hash of our own raw data
+	UE_DEPRECATED(5.1, "GenerateGuidFromRawData has been deprecated use IAnimationDataModel::GenerateGuid instead")
 	FGuid GenerateGuidFromRawData() const;
 
 	// Should we be always using our raw data (i.e is our compressed data stale)
@@ -707,9 +683,12 @@ public:
 	void SetUseRawDataOnly(bool bInUseRawDataOnly) { bUseRawDataOnly = bInUseRawDataOnly; }
 
 	// Return this animations guid for the raw data
+	UE_DEPRECATED(5.1, "GetRawDataGuid has been deprecated use GenerateGuidFromModel instead")
 	FGuid GetRawDataGuid() const
 	{ 
+		PRAGMA_DISABLE_DEPRECATION_WARNINGS
 		return RawDataGuid;
+		PRAGMA_ENABLE_DEPRECATION_WARNINGS
 	}
 
 	/** Resets Bone Animation, Curve data and Notify tracks **/
@@ -739,18 +718,13 @@ private:
 	 */ 
 	virtual void RemapTracksToNewSkeleton( USkeleton* NewSkeleton, bool bConvertSpaces ) override;
 
-	/**
-	 * Remap NaN tracks from the RawAnimation data and recompress
-	 */	
-	void RemoveNaNTracks();
-
 	/** Retargeting functions */
 	bool ConvertAnimationDataToRiggingData(FAnimSequenceTrackContainer & RiggingAnimationData);
 	bool ConvertRiggingDataToAnimationData(FAnimSequenceTrackContainer & RiggingAnimationData);
 	int32 GetSpaceBasedAnimationData(TArray< TArray<FTransform> > & AnimationDataInComponentSpace, FAnimSequenceTrackContainer * RiggingAnimationData) const;
+#endif
 
 public:
-#endif
 	/** Refresh sync marker data*/
 	void RefreshSyncMarkerDataFromAuthored();
 
@@ -811,7 +785,7 @@ protected:
 protected:
 #if WITH_EDITOR
 	// Begin UAnimSequenceBase virtual overrides
-	virtual void OnModelModified(const EAnimDataModelNotifyType& NotifyType, UAnimDataModel* Model, const FAnimDataModelNotifPayload& Payload) override;
+	virtual void OnModelModified(const EAnimDataModelNotifyType& NotifyType, IAnimationDataModel* Model, const FAnimDataModelNotifPayload& Payload) override;
 	virtual void PopulateModel() override;
 	// End UAnimSequenceBase virtual overrides
 
@@ -862,4 +836,5 @@ public:
 	friend class FCustomAttributeCustomization;
 	friend class FAnimSequenceTestBase;
 	friend struct UE::Anim::Compression::FScopedCompressionGuard;
+	friend class FAnimDataControllerTestBase;
 };
diff --git a/Engine/Source/Runtime/Engine/Classes/Animation/AnimSequenceBase.h b/Engine/Source/Runtime/Engine/Classes/Animation/AnimSequenceBase.h
--- a/Engine/Source/Runtime/Engine/Classes/Animation/AnimSequenceBase.h
+++ b/Engine/Source/Runtime/Engine/Classes/Animation/AnimSequenceBase.h
@@ -29,6 +29,7 @@ enum ETypeAdvanceAnim
 struct FAnimationPoseData;
 struct FAnimDataModelNotifPayload;
 class UAnimDataModel;
+class IAnimationDataModel;
 enum class EAnimDataModelNotifyType : uint8;
 
 UCLASS(abstract, BlueprintType)
@@ -41,14 +42,20 @@ public:
 	UPROPERTY()
 	TArray<struct FAnimNotifyEvent> Notifies;
 
+protected:
 	/** Length (in seconds) of this AnimSequence if played back with a speed of 1.0. */
 	UE_DEPRECATED(5.0, "Public access to SequenceLength is deprecated, use GetPlayLength or UAnimDataController::SetPlayLength instead")
 	UPROPERTY(Category=Length, AssetRegistrySearchable, VisibleAnywhere, BlueprintReadOnly)
 	float SequenceLength;
 
-	UE_DEPRECATED(5.0, "SetSequenceLength is deprecated use UAnimDataController::SetPlayLength instead")
-	virtual void SetSequenceLength(float NewLength);
-
+	/**
+	 * Raw uncompressed float curve data
+	 */
+	UE_DEPRECATED(5.0, "Public access to RawCurveData is deprecated, see UAnimDataModel for source data or use GetCurveData for runtime instead")
+	UPROPERTY()
+	struct FRawCurveTracks RawCurveData;
+	
+public:
 	/** Number for tweaking playback rate of this animation globally. */
 	UPROPERTY(EditAnywhere, Category=Animation)
 	float RateScale;
@@ -59,14 +66,6 @@ public:
 	 */
 	UPROPERTY(EditAnywhere, Category=Animation)
 	bool bLoop;
-
-	/**
-	 * Raw uncompressed float curve data 
-	 */
-	UE_DEPRECATED(5.0, "Public access to RawCurveData is deprecated, see UAnimDataModel for source data or use GetCurveData for runtime instead")
-	UPROPERTY()
-	struct FRawCurveTracks RawCurveData;
-
 #if WITH_EDITORONLY_DATA
 	// if you change Notifies array, this will need to be rebuilt
 	UPROPERTY()
@@ -145,7 +144,7 @@ public:
 	virtual int32 GetNumberOfSampledKeys() const;
 
 	/** Return rate at which the animation is sampled **/
-	virtual const FFrameRate& GetSamplingFrameRate() const;
+	virtual FFrameRate GetSamplingFrameRate() const;
 
 #if WITH_EDITOR
 	/** Get the frame number for the provided time */
@@ -238,10 +237,6 @@ public:
 	virtual void EnableRootMotionSettingFromMontage(bool bInEnableRootMotion, const ERootMotionRootLock::Type InRootMotionRootLock) {};
 
 #if WITH_EDITOR
-	// Store that our raw data has changed so that we can get correct compressed data later on
-	UE_DEPRECATED(5.0, "MarkRawDataAsModified has been deprecated, any (Raw Data) modification should be applied using the UAnimDataController API instead. This will handle updating the GUID instead.")
-	virtual void MarkRawDataAsModified(bool bForceNewRawDatGuid = true) {}
-
 private:
 	DECLARE_MULTICAST_DELEGATE( FOnNotifyChangedMulticaster );
 	FOnNotifyChangedMulticaster OnNotifyChanged;
@@ -269,18 +264,24 @@ protected:
 		}
 	}
 #if WITH_EDITOR
+	void OnEndLoadPackage(const FEndLoadPackageContext& Context);
+	virtual void OnAnimModelLoaded();
 public:
-	/** Returns the UAnimDataModel object embedded in this UAnimSequenceBase */
-	UAnimDataModel* GetDataModel() const;
+	/** Returns the IAnimationDataModel object embedded in this UAnimSequenceBase */
+	IAnimationDataModel* GetDataModel() const;
+
+	/** Returns the IAnimationDataModel as a script-interface, provides access to UObject and Interface */
+	TScriptInterface<IAnimationDataModel> GetDataModelInterface() const;
 
 	/** Returns the transient UAnimDataController set to operate on DataModel */
 	IAnimationDataController& GetController();
 protected:
 	/** Populates the UAnimDataModel object according to any pre-existing data. (overrides expect to populate the model according to their data) */
 	virtual void PopulateModel();
+	virtual void PopulateWithExistingModel(TScriptInterface<IAnimationDataModel> ExistingDataModel);
 
 	/** Callback registered to UAnimDatModel::GetModifiedEvent for the embedded object */
-	virtual void OnModelModified(const EAnimDataModelNotifyType& NotifyType, UAnimDataModel* Model, const FAnimDataModelNotifPayload& Payload);
+	virtual void OnModelModified(const EAnimDataModelNotifyType& NotifyType, IAnimationDataModel* Model, const FAnimDataModelNotifPayload& Payload);
 
 	/** Validates that DataModel contains a valid UAnimDataModel object */
 	void ValidateModel() const;
@@ -289,7 +290,7 @@ protected:
 	void BindToModelModificationEvent();
 
 	/** Replaces the current DataModel, if any, with the provided one */
-	void CopyDataModel(const UAnimDataModel* ModelToDuplicate);
+	void CopyDataModel(const TScriptInterface<IAnimationDataModel>& ModelToDuplicate);
 private:
 	/** Creates a new UAnimDataModel instance and sets DataModel accordingly */
 	void CreateModel();
@@ -301,7 +302,7 @@ public:
 		if(ShouldDataModelBeValid())
 		{
 			ValidateModel();
-			return DataModel != nullptr;
+			return DataModelInterface != nullptr;
 		}
 
 		return false;
@@ -310,10 +311,14 @@ public:
 
 #if WITH_EDITORONLY_DATA
 protected:
-	/** UAnimDataModel instance containing source animation data */
+	UE_DEPRECATED(5.1, "DataModel has been converted to DataModelInterface")
 	UPROPERTY(VisibleAnywhere, BlueprintReadOnly, Category = "Animation Model")
 	TObjectPtr<UAnimDataModel> DataModel;
 
+	/** IAnimationDataModel instance containing (source) animation data */
+	UPROPERTY(VisibleAnywhere, BlueprintReadOnly, Category = "Animation Model")
+	TScriptInterface<IAnimationDataModel> DataModelInterface;
+
 	/** Flag set whenever the data-model is initially populated (during upgrade path) */
 	bool bPopulatingDataModel;
 
diff --git a/Engine/Source/Runtime/Engine/Classes/Animation/AnimSequenceHelpers.h b/Engine/Source/Runtime/Engine/Classes/Animation/AnimSequenceHelpers.h
--- a/Engine/Source/Runtime/Engine/Classes/Animation/AnimSequenceHelpers.h
+++ b/Engine/Source/Runtime/Engine/Classes/Animation/AnimSequenceHelpers.h
@@ -8,7 +8,7 @@
 #include "Delegates/IDelegateInstance.h"
 
 class UAnimSequence;
-class UAnimDataModel;
+class IAnimationDataModel;
 class UAnimDataController;
 class IAnimationDataController;
 class FName;
@@ -17,6 +17,7 @@ struct FRawCurveTracks;
 struct FRawAnimSequenceTrack;
 struct FBoneAnimationTrack;
 struct FRawAnimSequenceTrack;
+struct FAnimationPoseData;
 class IAnimationDataController;
 
 enum class EAnimInterpolationType : uint8;
@@ -27,7 +28,7 @@ namespace Anim {
 
 #if WITH_EDITOR
 	/**
-	* Populates a FCompactPose according to the data stored within a UAnimDataModel provided a time to sample.
+	* Populates a FCompactPose according to the data stored within a IAnimationDataModel provided a time to sample.
 	*
 	* @param	Model				Model to sample from
 	* @param	OutPose				[out] Pose to be populated using the model
@@ -36,29 +37,42 @@ namespace Anim {
 	* @param	RetargetSource		Name of the retarget source to retrieve from the skeleton
 	* @param	RetargetTransforms	Base transforms to use when retargetting the pose
 	*/
-	ENGINE_API void BuildPoseFromModel(const UAnimDataModel* Model, FCompactPose& OutPose, const float Time, const EAnimInterpolationType& InterpolationType, const FName& RetargetSource, const TArray<FTransform>& RetargetTransforms);
+	UE_DEPRECATED(5.1, "Use BuildPoseFromModel with different signature instead")
+	ENGINE_API void BuildPoseFromModel(const IAnimationDataModel* Model, FCompactPose& OutPose, const float Time, const EAnimInterpolationType& InterpolationType, const FName& RetargetSource, const TArray<FTransform>& RetargetTransforms);
 
 	/**
-	* Populates float curves according to the data stored within a UAnimDataModel provided a time to sample.
+	* Populates a FAnimationPoseData according to the data stored within a IAnimationDataModel provided a time to sample.
+	*
+	* @param	Model				Model to sample from
+	* @param	OutPoseData			[out] Pose data to be populated using the model
+	* @param	Time				Time value to sample at
+	* @param	InterpolationType	Method to be used to interpolate between different sample points
+	* @param	RetargetSource		Name of the retarget source to retrieve from the skeleton
+	* @param	RetargetTransforms	Base transforms to use when retargetting the pose
+	*/
+	ENGINE_API void BuildPoseFromModel(const IAnimationDataModel* Model, FAnimationPoseData& OutPoseData, const double Time, const EAnimInterpolationType& InterpolationType, const FName& RetargetSource, const TArray<FTransform>& RetargetTransforms);
+
+	/**
+	* Populates float curves according to the data stored within a IAnimationDataModel provided a time to sample.
 	*
 	* @param	Model				Model to sample from
 	* @param	OutCurves			[out] Curves to be populated using the model
 	* @param	Time				Time value to sample at
 	*/
-	ENGINE_API void EvaluateFloatCurvesFromModel(const UAnimDataModel* Model, FBlendedCurve& OutCurves, float Time);
+	ENGINE_API void EvaluateFloatCurvesFromModel(const IAnimationDataModel* Model, FBlendedCurve& OutCurves, double Time);
 
 	/**
-	* Populates transform curves according to the data stored within a UAnimDataModel provided a time to sample.
+	* Populates transform curves according to the data stored within a IAnimationDataModel provided a time to sample.
 	*
 	* @param	Model				Model to sample from
 	* @param	OutCurves			[out] Curves which were sampled from the model, key-ed with their equivalent name
 	* @param	Time				Time value to sample at
 	* @param	BlendWeight			Blending weight to apply when sampling
 	*/
-	ENGINE_API void EvaluateTransformCurvesFromModel(const UAnimDataModel* Model, TMap<FName, FTransform>& OutCurves, float Time, float BlendWeight);
+	ENGINE_API void EvaluateTransformCurvesFromModel(const IAnimationDataModel* Model, TMap<FName, FTransform>& OutCurves, double Time, float BlendWeight);
 
 	/**
-	* Retrieves a single bone transform according to the data stored within a UAnimDataModel provided a time and track to sample.
+	* Retrieves a single bone transform according to the data stored within a IAnimationDataModel provided a time and track to sample.
 	*
 	* @param	Model				Model to sample from
 	* @param	OutTransform		[out] Transform which was sampled from the model
@@ -66,20 +80,20 @@ namespace Anim {
 	* @param	Time				Time value to sample at
 	* @param	InterpolationType	Method to be used to interpolate between different sample points
 	*/
-	ENGINE_API void GetBoneTransformFromModel(const UAnimDataModel* Model, FTransform& OutTransform, int32 TrackIndex, float Time, const EAnimInterpolationType& Interpolation);
+	ENGINE_API void GetBoneTransformFromModel(const IAnimationDataModel* Model, FTransform& OutTransform, int32 TrackIndex, double Time, const EAnimInterpolationType& Interpolation);
 	
 	/**
-	* Retrieves a single bone transform according to the data stored within a UAnimDataModel provided a key and track to sample.
+	* Retrieves a single bone transform according to the data stored within a IAnimationDataModel provided a key and track to sample.
 	*
 	* @param	Model				Model to sample from
 	* @param	OutTransform		[out] Transform which was sampled from the model
 	* @param	TrackIndex			Track / bone index which should be sampld
 	* @param	KeyIndex			Specific key index to sample
 	*/
-	ENGINE_API void GetBoneTransformFromModel(const UAnimDataModel* Model, FTransform& OutTransform, int32 TrackIndex, int32 KeyIndex);
+	ENGINE_API void GetBoneTransformFromModel(const IAnimationDataModel* Model, FTransform& OutTransform, int32 TrackIndex, int32 KeyIndex);
 
 	/**
-	* Copies over any individual curve from FRawCurveTracks to a UAnimDataModel instance targeted by the provided controller.
+	* Copies over any individual curve from FRawCurveTracks to a IAnimationDataModel instance targeted by the provided controller.
 	*
 	* @param	CurveData			Container of the curves to be copied over
 	* @param	Skeleton			Skeleton to use for verifying and or add the curve names
@@ -88,7 +102,7 @@ namespace Anim {
 	ENGINE_API void CopyCurveDataToModel(const FRawCurveTracks& CurveData, const USkeleton* Skeleton, IAnimationDataController& Controller);
 
 	/**
-	* Copies over any individual curve from FRawCurveTracks to a UAnimDataModel instance targeted by the provided controller.
+	* Copies over any individual curve from FRawCurveTracks to a IAnimationDataModel instance targeted by the provided controller.
 	*
 	* @param	SourceAnimSeq				Animation Sequence to copy the notifies from
 	* @param	DestAnimSeq					Animation Sequence to copy the notifies to
@@ -212,8 +226,30 @@ namespace Anim {
 #endif // WITH_EDITOR
 	}
 
+	namespace DataModel { struct FEvaluationContext; }
 	namespace Retargeting
 	{
+		struct ENGINE_API FRetargetTracking
+		{
+			const FCompactPoseBoneIndex PoseBoneIndex;
+			const int32 SkeletonBoneIndex;
+
+			FRetargetTracking(const FCompactPoseBoneIndex InPoseBoneIndex, const int32 InSkeletonBoneIndex) : PoseBoneIndex(InPoseBoneIndex), SkeletonBoneIndex(InSkeletonBoneIndex) {}
+		};
+
+		struct ENGINE_API FRetargetingScope
+		{
+			FRetargetingScope(FCompactPose& ToRetargetPose, const DataModel::FEvaluationContext& InEvaluationContext);
+			~FRetargetingScope();
+			
+			void AddTrackedBone(FCompactPoseBoneIndex CompactBoneIndex, int32 SkeletonBoneIndex) const;
+		private:
+			FCompactPose& RetargetPose;
+			const DataModel::FEvaluationContext& EvaluationContext;
+			TArray<FRetargetTracking>& RetargetTracking;
+			bool bShouldRetarget;
+		};
+		
 		/**
 		* Retargeting the provided pose using the retarget source or transforms
 		*
@@ -221,9 +257,16 @@ namespace Anim {
 		* @param	RetargetSource		Name of the retarget source to retrieve data for
 		* @param	RetargetTransforms	Set of transforms to use as the ref skeleton bone transforms
 		*/
+		UE_DEPRECATED(5.1, "UE::Anim::Retargeting::RetargetPose has been deprecated as retargeting while evaluating pose data from a DataModel is now handled/performed implicitly")
 		ENGINE_API void RetargetPose(FCompactPose& InOutPose, const FName& RetargetSource, const TArray<FTransform>& RetargetTransforms);
 	}
 
+	struct ENGINE_API FBuildRawPoseScratchArea : public TThreadSingleton<FBuildRawPoseScratchArea>
+	{
+		TArray<Retargeting::FRetargetTracking> RetargetTracking;
+		TArray<FVirtualBoneCompactPoseData> VirtualBoneCompactPoseData;
+	};
+
 } // namespace Anim	
 
 } // namespace UE
\ No newline at end of file
diff --git a/Engine/Source/Runtime/Engine/Classes/Animation/AnimStreamable.h b/Engine/Source/Runtime/Engine/Classes/Animation/AnimStreamable.h
--- a/Engine/Source/Runtime/Engine/Classes/Animation/AnimStreamable.h
+++ b/Engine/Source/Runtime/Engine/Classes/Animation/AnimStreamable.h
@@ -38,6 +38,7 @@ public:
 	float StartTime;
 
 	float SequenceLength;
+	int32 NumFrames;
 
 	// Compressed Data for this chunk (if nullptr then data needs to be loaded via BulkData)
 	FCompressedAnimSequence* CompressedAnimSequence;
@@ -103,6 +104,9 @@ public:
 	UPROPERTY(EditAnywhere, AssetRegistrySearchable, Category = Animation)
 	FName RetargetSource;
 
+	UPROPERTY()
+	FFrameRate SamplingFrameRate;
+
 #if WITH_EDITORONLY_DATA
 
 	// Sequence the streamable was created from (used for reflecting changes to the source in editor)
diff --git a/Engine/Source/Runtime/Engine/Classes/Animation/AnimationAsset.h b/Engine/Source/Runtime/Engine/Classes/Animation/AnimationAsset.h
--- a/Engine/Source/Runtime/Engine/Classes/Animation/AnimationAsset.h
+++ b/Engine/Source/Runtime/Engine/Classes/Animation/AnimationAsset.h
@@ -245,7 +245,7 @@ struct FPoseCurve
 struct FAnimExtractContext
 {
 	/** Position in animation to extract pose from */
-	float CurrentTime;
+	double CurrentTime;
 
 	/** Is root motion being extracted? */
 	bool bExtractRootMotion;
@@ -268,8 +268,19 @@ struct FAnimExtractContext
 	 * by several animation nodes to optimize evaluation time.
 	 */
 	TArray<bool> BonesRequired;
+	
+	UE_DEPRECATED(5.1, "FAnimExtractContext construct with float-based time value is deprecated, use other signature")
+	FAnimExtractContext(float InCurrentTime, bool InbExtractRootMotion = false, FDeltaTimeRecord InDeltaTimeRecord = {}, bool InbLooping = false)
+		: CurrentTime((double)InCurrentTime)
+		, bExtractRootMotion(InbExtractRootMotion)
+		, DeltaTimeRecord(InDeltaTimeRecord)
+		, bLooping(InbLooping)
+		, PoseCurves()
+		, BonesRequired()
+	{
+	}
 
-	FAnimExtractContext(float InCurrentTime = 0.f, bool InbExtractRootMotion = false, FDeltaTimeRecord InDeltaTimeRecord = {}, bool InbLooping = false)
+	FAnimExtractContext(double InCurrentTime = 0.0, bool InbExtractRootMotion = false, FDeltaTimeRecord InDeltaTimeRecord = {}, bool InbLooping = false)
 		: CurrentTime(InCurrentTime)
 		, bExtractRootMotion(InbExtractRootMotion)
 		, DeltaTimeRecord(InDeltaTimeRecord)
@@ -1091,7 +1102,7 @@ protected:
 
 public:
 	/** Return a list of unique marker names for blending compatibility */
-	virtual TArray<FName>* GetUniqueMarkerNames() { return NULL; }
+	virtual TArray<FName>* GetUniqueMarkerNames() { return nullptr; }
 
 	//~ Begin IInterface_AssetUserData Interface
 	virtual void AddAssetUserData(UAssetUserData* InUserData) override;
diff --git a/Engine/Source/Runtime/Engine/Classes/Animation/AnimationSettings.h b/Engine/Source/Runtime/Engine/Classes/Animation/AnimationSettings.h
--- a/Engine/Source/Runtime/Engine/Classes/Animation/AnimationSettings.h
+++ b/Engine/Source/Runtime/Engine/Classes/Animation/AnimationSettings.h
@@ -95,9 +95,15 @@ class ENGINE_API UAnimationSettings : public UDeveloperSettings
 	UPROPERTY(config, EditAnywhere, Category = Mirroring)
 	TArray<FMirrorFindReplaceExpression> MirrorFindReplaceExpressions;
 
+	/** Project specific default frame-rate used when (re)initializing any animation based data */
+	UPROPERTY(config, EditAnywhere, Category = AnimationData)
+	FFrameRate DefaultFrameRate;
 public:
 	static UAnimationSettings * Get() { return CastChecked<UAnimationSettings>(UAnimationSettings::StaticClass()->GetDefaultObject()); }
 
+	/** Returns the project specific default frame-rate */
+	const FFrameRate& GetDefaultFrameRate() const;
+
 #if WITH_EDITOR
 	virtual void PostEditChangeProperty(struct FPropertyChangedEvent& PropertyChangedEvent) override;
 #endif // WITH_EDITOR
diff --git a/Engine/Source/Runtime/Engine/Classes/Animation/AttributeCurve.h b/Engine/Source/Runtime/Engine/Classes/Animation/AttributeCurve.h
--- a/Engine/Source/Runtime/Engine/Classes/Animation/AttributeCurve.h
+++ b/Engine/Source/Runtime/Engine/Classes/Animation/AttributeCurve.h
@@ -189,7 +189,6 @@ protected:
 	const UE::Anim::IAttributeBlendOperator* Operator;
 
 	friend class UAnimSequence;
-	friend class UAnimDataController;
 	friend struct UE::Anim::Attributes;
 };
 
diff --git a/Engine/Source/Runtime/Engine/Private/Animation/AnimComposite.cpp b/Engine/Source/Runtime/Engine/Private/Animation/AnimComposite.cpp
--- a/Engine/Source/Runtime/Engine/Private/Animation/AnimComposite.cpp
+++ b/Engine/Source/Runtime/Engine/Private/Animation/AnimComposite.cpp
@@ -144,10 +144,10 @@ bool UAnimComposite::ContainRecursive(TArray<UAnimCompositeBase*>& CurrentAccumu
 void UAnimComposite::SetCompositeLength(float InLength)
 {
 #if WITH_EDITOR		
-	Controller->SetPlayLength(InLength);
+	Controller->SetNumberOfFrames(DataModelInterface->GetFrameRate().AsFrameNumber(InLength));
 #else
 	PRAGMA_DISABLE_DEPRECATION_WARNINGS
-	SetSequenceLength(InLength);
+	SequenceLength = InLength;
 	PRAGMA_ENABLE_DEPRECATION_WARNINGS
 #endif	
 }
diff --git a/Engine/Source/Runtime/Engine/Private/Animation/AnimCompositeBase.cpp b/Engine/Source/Runtime/Engine/Private/Animation/AnimCompositeBase.cpp
--- a/Engine/Source/Runtime/Engine/Private/Animation/AnimCompositeBase.cpp
+++ b/Engine/Source/Runtime/Engine/Private/Animation/AnimCompositeBase.cpp
@@ -16,7 +16,7 @@
 // FAnimSegment
 ///////////////////////////////////////////////////////
 
-UAnimSequenceBase * FAnimSegment::GetAnimationData(float PositionInTrack, float& PositionInAnim) const
+UAnimSequenceBase* FAnimSegment::GetAnimationData(float PositionInTrack, float& PositionInAnim) const
 {
 	if( bValid && IsInRange(PositionInTrack) )
 	{
@@ -569,9 +569,11 @@ void FAnimTrack::GetAnimationPose(FAnimationPoseData& OutAnimationPoseData, cons
 		if (AnimSegment->bValid)
 		{
 			// Copy passed in Extraction Context, but override position and root motion parameters.
-			FAnimExtractContext SequenceExtractionContext(ExtractionContext);
-			if (const UAnimSequenceBase* const AnimRef = AnimSegment->GetAnimationData(ClampedTime, SequenceExtractionContext.CurrentTime))
+			float PositionInAnim = 0.f;
+			if (const UAnimSequenceBase* const AnimRef = AnimSegment->GetAnimationData(ClampedTime, PositionInAnim))
 			{
+				FAnimExtractContext SequenceExtractionContext(ExtractionContext);
+				SequenceExtractionContext.CurrentTime = static_cast<double>(PositionInAnim);
 				SequenceExtractionContext.DeltaTimeRecord.SetPrevious(
 					SequenceExtractionContext.CurrentTime - SequenceExtractionContext.DeltaTimeRecord.Delta);
 				SequenceExtractionContext.bExtractRootMotion &= AnimRef->HasRootMotion();
@@ -766,6 +768,13 @@ void UAnimCompositeBase::ExtractRootMotionFromTrack(const FAnimTrack &SlotAnimTr
 	}
 }
 
+FFrameRate UAnimCompositeBase::GetSamplingFrameRate() const
+{
+	// Allowing for 0.00001s precision in composite/montage length
+	static const FFrameRate CompositeFrameRate(100000, 1);
+	return CompositeFrameRate;
+}
+
 void UAnimCompositeBase::PostLoad()
 {
 	Super::PostLoad();
@@ -817,11 +826,17 @@ void FAnimSegment::UpdateCachedPlayLength()
 {
 	CachedPlayLength = 0.f;
 	PRAGMA_DISABLE_DEPRECATION_WARNINGS
-	const UAnimDataModel* DataModel = AnimReference ? AnimReference->GetDataModel() : nullptr;	
+	const IAnimationDataModel* DataModel = AnimReference ? AnimReference->GetDataModel() : nullptr;	
 	PRAGMA_ENABLE_DEPRECATION_WARNINGS
 	if(DataModel)
 	{
 		CachedPlayLength = DataModel->GetPlayLength();
 	}
 }
+
+void UAnimCompositeBase::PopulateWithExistingModel(TScriptInterface<IAnimationDataModel> ExistingDataModel)
+{
+	Super::PopulateWithExistingModel(ExistingDataModel);
+	Controller->SetFrameRate(GetSamplingFrameRate());
+}
 #endif // WITH_EDITOR
\ No newline at end of file
diff --git a/Engine/Source/Runtime/Engine/Private/Animation/AnimCompress_RemoveLinearKeys.cpp b/Engine/Source/Runtime/Engine/Private/Animation/AnimCompress_RemoveLinearKeys.cpp
--- a/Engine/Source/Runtime/Engine/Private/Animation/AnimCompress_RemoveLinearKeys.cpp
+++ b/Engine/Source/Runtime/Engine/Private/Animation/AnimCompress_RemoveLinearKeys.cpp
@@ -304,20 +304,19 @@ void UAnimCompress_RemoveLinearKeys::UpdateWorldBoneTransformTable(
 {
 	const FBoneData& Bone		= CompressibleAnimData.BoneData[BoneIndex];
 	const int32 NumKeys		= CompressibleAnimData.NumberOfKeys;
-	const float SequenceLength	= CompressibleAnimData.SequenceLength;
 	const int32 FrameStart		= (BoneIndex*NumKeys);
 	const int32 TrackIndex = FAnimationUtils::GetAnimTrackIndexForSkeletonBone(BoneIndex, CompressibleAnimData.TrackToSkeletonMapTable);
 	
 	check(OutputWorldBones.Num() >= (FrameStart+NumKeys));
 
-	const float TimePerFrame = SequenceLength / (float)(NumKeys-1);
+	const FFrameRate& SamplingRate = CompressibleAnimData.SampledFrameRate;
 
 	if( TrackIndex != INDEX_NONE )
 	{
 		// get the local-space bone transforms using the animation solver
 		for ( int32 KeyIndex = 0; KeyIndex < NumKeys; ++KeyIndex )
 		{
-			float Time = (float)KeyIndex * TimePerFrame;
+			const double Time = SamplingRate.AsSeconds(KeyIndex);
 			FTransform LocalAtom;
 
 			FAnimationUtils::ExtractTransformFromCompressionData(CompressibleAnimData, OutCompressedData, Time, TrackIndex, UseRaw, LocalAtom);
@@ -417,9 +416,10 @@ void UAnimCompress_RemoveLinearKeys::UpdateBoneAtomList(
 	TArray<FTransform>& BoneAtoms)
 {
 	BoneAtoms.Reset(NumFrames);
+	const FFrameRate& SamplingRate = CompressibleAnimData.SampledFrameRate;
 	for ( int32 FrameIndex = 0; FrameIndex < NumFrames; ++FrameIndex )
 	{
-		float Time = (float)FrameIndex * TimePerFrame;
+		const double Time = SamplingRate.AsSeconds(FrameIndex);
 		FTransform LocalAtom;
 		FAnimationUtils::ExtractTransformFromCompressionData(CompressibleAnimData, OutCompressedData, Time, TrackIndex, false, LocalAtom);
 
@@ -465,11 +465,10 @@ void UAnimCompress_RemoveLinearKeys::ConvertFromRelativeSpace(FCompressibleAnimD
 		// make sure scale key exists
 		if (RawTrack.ScaleKeys.Num() > 0)
 		{
-			const FVector3f DefaultScale(1.f);
-			const FVector3f& RefBoneScale = (BasePoseTrack.ScaleKeys.Num() > 0)? BasePoseTrack.ScaleKeys[0] : DefaultScale;
+			const FVector3f& RefBoneScale = (BasePoseTrack.ScaleKeys.Num() > 0)? BasePoseTrack.ScaleKeys[0] : FVector3f::OneVector;
 			for (int32 ScaleIndex = 0; ScaleIndex < RawTrack.ScaleKeys.Num(); ++ScaleIndex)
 			{
-				RawTrack.ScaleKeys[ScaleIndex] = RefBoneScale * (DefaultScale + RawTrack.ScaleKeys[ScaleIndex]);
+				RawTrack.ScaleKeys[ScaleIndex] = RefBoneScale * (FVector3f::OneVector + RawTrack.ScaleKeys[ScaleIndex]);
 			}
 		}
 	}
@@ -515,7 +514,8 @@ void UAnimCompress_RemoveLinearKeys::ConvertToRelativeSpace(FCompressibleAnimDat
 		// scale key
 		if (RawTrack.ScaleKeys.Num() > 0)
 		{
-			const FVector3f InvRefBoneScale = (FVector3f)FTransform::GetSafeScaleReciprocal((FVector)BasePoseTrack.ScaleKeys[0]);
+        	const FVector3f& RefBoneScale = (BasePoseTrack.ScaleKeys.Num() > 0)? BasePoseTrack.ScaleKeys[0] : FVector3f::OneVector;
+			const FVector3f InvRefBoneScale = (FVector3f)FTransform::GetSafeScaleReciprocal((FVector)RefBoneScale);
 
 			// transform scale keys.
 			for (int32 ScaleIndex = 0; ScaleIndex < RawTrack.ScaleKeys.Num(); ++ScaleIndex)
@@ -563,7 +563,8 @@ void UAnimCompress_RemoveLinearKeys::ConvertToRelativeSpace(
 		// scale key
 		if (ScaleData.Num() > 0)
 		{
-			const FVector3f InvRefBoneScale = (FVector3f)FTransform::GetSafeScaleReciprocal((FVector)BasePoseTrack.ScaleKeys[0]);
+        	const FVector3f& RefBoneScale = (BasePoseTrack.ScaleKeys.Num() > 0)? BasePoseTrack.ScaleKeys[0] : FVector3f::OneVector;
+			const FVector3f InvRefBoneScale = (FVector3f)FTransform::GetSafeScaleReciprocal((FVector)RefBoneScale);
 
 			// convert the new scale tracks to additive space
 			FScaleTrack& ScaleTrack = ScaleData[TrackIndex];
@@ -879,7 +880,7 @@ void UAnimCompress_RemoveLinearKeys::ProcessAnimationTracks(
 			}
 
 			// Determine if a guidance track should be used to aid in choosing keys to retain
-			TArray<float>* GuidanceTrack = NULL;
+			TArray<float>* GuidanceTrack = nullptr;
 			float GuidanceScale = 1.0f;
 			if (GuideTrackIndex != INDEX_NONE)
 			{
diff --git a/Engine/Source/Runtime/Engine/Private/Animation/AnimCompressionDerivedData.cpp b/Engine/Source/Runtime/Engine/Private/Animation/AnimCompressionDerivedData.cpp
--- a/Engine/Source/Runtime/Engine/Private/Animation/AnimCompressionDerivedData.cpp
+++ b/Engine/Source/Runtime/Engine/Private/Animation/AnimCompressionDerivedData.cpp
@@ -35,7 +35,7 @@ const TCHAR* FDerivedDataAnimationCompression::GetVersionString() const
 	// This is a version string that mimics the old versioning scheme. If you
 	// want to bump this version, generate a new guid using VS->Tools->Create GUID and
 	// return it here. Ex.
-	return TEXT("3A3F9FBB7A4047278CABE35820CC44D7");
+	return TEXT("6CAE774196A34FEABF46EF3B30EF6DDC");
 }
 
 bool FDerivedDataAnimationCompression::Build( TArray<uint8>& OutDataArray )
