**New**

diff --git a/Engine/Source/Runtime/Engine/Private/Animation/AnimCompressionTypes.cpp b/Engine/Source/Runtime/Engine/Private/Animation/AnimCompressionTypes.cpp
--- a/Engine/Source/Runtime/Engine/Private/Animation/AnimCompressionTypes.cpp
+++ b/Engine/Source/Runtime/Engine/Private/Animation/AnimCompressionTypes.cpp
@@ -16,11 +16,6 @@
 #include "Misc/FileHelper.h"
 #include "Serialization/JsonWriter.h"
 
-CSV_DECLARE_CATEGORY_MODULE_EXTERN(ENGINE_API, Animation);
-
-DECLARE_CYCLE_STAT(TEXT("Build Anim Track Pairs"), STAT_BuildAnimTrackPairs, STATGROUP_Anim);
-DECLARE_CYCLE_STAT(TEXT("Extract Pose From Anim Data"), STAT_ExtractPoseFromAnimData, STATGROUP_Anim);
-
 static FString GCompressionJsonOutput;
 static FAutoConsoleVariableRef CVarCompressionJsonOutput(
 	TEXT("a.Compression.CompressibleDataOutput"),
@@ -337,6 +332,8 @@ FCompressibleAnimData::FCompressibleAnimData(class UAnimSequence* InSeq, const b
 	/* Always get the resampled data to start off with */
 	const TArray<FBoneAnimationTrack>& ResampledTrackData = InSeq->GetResampledTrackData();
 
+	SampledFrameRate = InSeq->GetSamplingFrameRate();
+
 	TArray<FName> OriginalTrackNames;
 	RawAnimationData.Empty(ResampledTrackData.Num());
 	TrackToSkeletonMapTable.Empty(ResampledTrackData.Num());
@@ -350,8 +347,8 @@ FCompressibleAnimData::FCompressibleAnimData(class UAnimSequence* InSeq, const b
 		OriginalTrackNames.Add(AnimTrack.Name);
 	}
 
-	const bool bBakeAdditive = InSeq->CanBakeAdditive();
-	if (bBakeAdditive)
+	const bool bIsAdditiveAnimation = InSeq->CanBakeAdditive();
+	if (bIsAdditiveAnimation)
 	{
 		InSeq->BakeOutAdditiveIntoRawData(RawAnimationData, OriginalTrackNames, TrackToSkeletonMapTable, RawFloatCurves, AdditiveBaseAnimationData);
 	}
@@ -384,7 +381,7 @@ FCompressibleAnimData::FCompressibleAnimData(class UAnimSequence* InSeq, const b
 		return Keys.Num() == 0 || (Keys.Num() == 1 && Keys[0].Equals(IdentityValue));
 	};
 
-	auto IsRawTrackValidForRemoval = [IsKeyArrayValidForRemoval](const FRawAnimSequenceTrack& Track) -> bool
+	auto IsRawTrackZeroAdditive = [IsKeyArrayValidForRemoval](const FRawAnimSequenceTrack& Track) -> bool
 	{
 		return IsKeyArrayValidForRemoval(Track.PosKeys, FVector3f::ZeroVector) &&
 			IsKeyArrayValidForRemoval(Track.RotKeys, FQuat4f::Identity) &&
@@ -395,21 +392,59 @@ FCompressibleAnimData::FCompressibleAnimData(class UAnimSequence* InSeq, const b
     // And for additive animations remove any track deemed not to add any additive animation (identity rotation and zero-vector translation and scale)
 	// Note on (TrackIndex > 0) below : deliberately stop before track 0, compression code doesn't like getting a completely empty animation
 	TArray<FName> FinalTrackNames;
-	FinalTrackNames.Reserve(ResampledTrackData.Num());	
-	for (int32 TrackIndex = RawAnimationData.Num() - 1; TrackIndex > 0; --TrackIndex)
+
+	// Ensure we have any tracks at all	to begin with
+	if (OriginalTrackNames.Num())
 	{
-		const FRawAnimSequenceTrack& Track = RawAnimationData[TrackIndex];
-		// Try find correct bone index
-		const int32 BoneIndex = RefSkeleton.FindBoneIndex(OriginalTrackNames[TrackIndex]);
-		if ((bBakeAdditive && IsRawTrackValidForRemoval(Track)) || BoneIndex == INDEX_NONE)
+		TArray<FRawAnimSequenceTrack> TempRawAnimationData;
+		TArray<FRawAnimSequenceTrack> TempAdditiveBaseAnimationData;
+		TArray<FTrackToSkeletonMap> TempTrackToSkeletonMapTable;
+		TempTrackToSkeletonMapTable.Reserve(OriginalTrackNames.Num());
+		TempRawAnimationData.Reserve(OriginalTrackNames.Num());
+		FinalTrackNames.Reserve(ResampledTrackData.Num());
+		TempAdditiveBaseAnimationData.Reserve(AdditiveBaseAnimationData.Num() ? AdditiveBaseAnimationData.Num() : 0);
+
+		// Include root bone track
+		FinalTrackNames.Add(OriginalTrackNames[0]);
+		TempTrackToSkeletonMapTable.Add(TrackToSkeletonMapTable[0]);
+		TempRawAnimationData.Add(RawAnimationData[0]);
+		if (AdditiveBaseAnimationData.Num())
 		{
-			RawAnimationData.RemoveAtSwap(TrackIndex, 1, false);
-			TrackToSkeletonMapTable.RemoveAtSwap(TrackIndex, 1, false);
+			TempAdditiveBaseAnimationData.Add(AdditiveBaseAnimationData[0]);
+		}
+
+		const int32 NumTracks = RawAnimationData.Num();
+		for (int32 TrackIndex = 1; TrackIndex < NumTracks; ++TrackIndex)
+		{
+			const FRawAnimSequenceTrack& Track = RawAnimationData[TrackIndex];
+			// Try find correct bone index
+			const int32 BoneIndex = RefSkeleton.FindBoneIndex(OriginalTrackNames[TrackIndex]);
+
+			const bool bValidBoneIndex = BoneIndex != INDEX_NONE;
+			const bool bValidAdditiveTrack = !IsRawTrackZeroAdditive(Track);
+
+			// Only include track if it contains valid (additive) data and its name corresponds to a bone on the skeleton
+			if ((!bIsAdditiveAnimation || bValidAdditiveTrack) && bValidBoneIndex)
+			{
+				FinalTrackNames.Add(OriginalTrackNames[TrackIndex]);
+				TempTrackToSkeletonMapTable.Add(TrackToSkeletonMapTable[TrackIndex]);
+				TempRawAnimationData.Add(RawAnimationData[TrackIndex]);
+
+				if (AdditiveBaseAnimationData.Num())
+				{
+					TempAdditiveBaseAnimationData.Add(AdditiveBaseAnimationData[TrackIndex]);
+				}
+			}
+		}
+
+		// Swap out maintained track data
+		Swap(RawAnimationData, TempRawAnimationData);
+		Swap(TrackToSkeletonMapTable, TempTrackToSkeletonMapTable);
+
+		if (AdditiveBaseAnimationData.Num())
+		{
+			Swap(AdditiveBaseAnimationData, TempAdditiveBaseAnimationData);
 		}
-		else
-        {
-			FinalTrackNames.Add(OriginalTrackNames[TrackIndex]);
-        }
 	}
 
 	// Find or add curve names on skeleton
@@ -459,12 +494,6 @@ FCompressibleAnimData::FCompressibleAnimData(class UAnimSequence* InSeq, const b
 
 	if (GCompressionJsonOutput.Len())
 	{
-		// Insert root-bone name
-		if (RawAnimationData.Num())
-		{			
-			FinalTrackNames.Insert(OriginalTrackNames[0], 0);
-		}
-		
 		WriteCompressionDataToJSON(MakeArrayView(OriginalTrackNames), MakeArrayView(RawAnimationData), MakeArrayView(FinalTrackNames));	
 	}
 #endif
@@ -910,7 +939,7 @@ void FCompressedAnimSequence::SerializeCompressedData(FArchive& Ar, bool bDDCDat
 			FMemoryWriter MemoryWriter(TempSerialized, true);
 			MemoryWriter.SetByteSwapping(Ar.ForceByteSwapping());
 
-			check(RotationCodec != NULL);
+			check(RotationCodec != nullptr);
 
 			FMemoryReader MemoryReader(TempSerialized, true);
 			MemoryReader.SetByteSwapping(Ar.ForceByteSwapping());
@@ -1004,269 +1033,23 @@ void FCompressedAnimSequence::ClearCompressedCurveData()
 	CurveCompressionCodec = nullptr;
 }
 
-struct FGetBonePoseScratchArea : public TThreadSingleton<FGetBonePoseScratchArea>
-{
-	BoneTrackArray RotationScalePairs;
-	BoneTrackArray TranslationPairs;
-	BoneTrackArray AnimScaleRetargetingPairs;
-	BoneTrackArray AnimRelativeRetargetingPairs;
-	BoneTrackArray OrientAndScaleRetargetingPairs;
-};
+
 
 void DecompressPose(FCompactPose& OutPose, const FCompressedAnimSequence& CompressedData, const FAnimExtractContext& ExtractionContext, USkeleton* SourceSkeleton, float SequenceLength, EAnimInterpolationType Interpolation, bool bIsBakedAdditive, FName RetargetSource, FName SourceName, const FRootMotionReset& RootMotionReset)
 {
 	const TArray<FTransform>& RetargetTransforms = SourceSkeleton->GetRefLocalPoses(RetargetSource);
-	DecompressPose(OutPose, CompressedData, ExtractionContext, SourceSkeleton, SequenceLength, Interpolation, bIsBakedAdditive, RetargetTransforms, SourceName, RootMotionReset);
+	PRAGMA_DISABLE_DEPRECATION_WARNINGS
+	FAnimSequenceDecompressionContext DecompressionContext(SequenceLength, Interpolation, SourceName, *CompressedData.CompressedDataStructure.Get(), SourceSkeleton, bIsBakedAdditive);
+	PRAGMA_ENABLE_DEPRECATION_WARNINGS
+	UE::Anim::Decompression::DecompressPose(OutPose, CompressedData, ExtractionContext, DecompressionContext, RetargetTransforms, RootMotionReset);
 }
 
 void DecompressPose(FCompactPose& OutPose, const FCompressedAnimSequence& CompressedData, const FAnimExtractContext& ExtractionContext, USkeleton* SourceSkeleton, float SequenceLength, EAnimInterpolationType Interpolation, bool bIsBakedAdditive, const TArray<FTransform>& RetargetTransforms, FName SourceName, const FRootMotionReset& RootMotionReset)
 {
-	const FBoneContainer& RequiredBones = OutPose.GetBoneContainer();
-	const int32 NumTracks = CompressedData.CompressedTrackToSkeletonMapTable.Num();
-
-	const USkeleton* TargetSkeleton = RequiredBones.GetSkeletonAsset();
-	const FSkeletonRemapping* SkeletonRemapping = TargetSkeleton->GetSkeletonRemapping(SourceSkeleton);
-
-	BoneTrackArray& RotationScalePairs = FGetBonePoseScratchArea::Get().RotationScalePairs;
-	BoneTrackArray& TranslationPairs = FGetBonePoseScratchArea::Get().TranslationPairs;
-	BoneTrackArray& AnimScaleRetargetingPairs = FGetBonePoseScratchArea::Get().AnimScaleRetargetingPairs;
-	BoneTrackArray& AnimRelativeRetargetingPairs = FGetBonePoseScratchArea::Get().AnimRelativeRetargetingPairs;
-	BoneTrackArray& OrientAndScaleRetargetingPairs = FGetBonePoseScratchArea::Get().OrientAndScaleRetargetingPairs;
-
-	// build a list of desired bones
-	RotationScalePairs.Reset();
-	TranslationPairs.Reset();
-	AnimScaleRetargetingPairs.Reset();
-	AnimRelativeRetargetingPairs.Reset();
-	OrientAndScaleRetargetingPairs.Reset();
-
-	// Optimization: assuming first index is root bone. That should always be the case in Skeletons.
-	checkSlow((RequiredBones.GetMeshPoseIndexFromSkeletonPoseIndex(FSkeletonPoseBoneIndex(0)) == FMeshPoseBoneIndex(0)));
-	// this is not guaranteed for AnimSequences though... If Root is not animated, Track will not exist.
-	const bool bFirstTrackIsRootBone = (CompressedData.GetSkeletonIndexFromTrackIndex(0) == 0);
-
-	{
-		SCOPE_CYCLE_COUNTER(STAT_BuildAnimTrackPairs);
-
-		// Handle root bone separately if it is track 0. so we start w/ Index 1.
-		for (int32 TrackIndex = (bFirstTrackIsRootBone ? 1 : 0); TrackIndex < NumTracks; TrackIndex++)
-		{
-			const int32 SourceSkeletonBoneIndex = CompressedData.GetSkeletonIndexFromTrackIndex(TrackIndex);
-			const int32 TargetSkeletonBoneIndex = (SkeletonRemapping) ? SkeletonRemapping->GetTargetSkeletonBoneIndex(SourceSkeletonBoneIndex) : SourceSkeletonBoneIndex;
-
-			if (TargetSkeletonBoneIndex != INDEX_NONE)
-			{
-				const FCompactPoseBoneIndex BoneIndex = RequiredBones.GetCompactPoseIndexFromSkeletonIndex(TargetSkeletonBoneIndex);
-				//Nasty, we break our type safety, code in the lower levels should be adjusted for this
-				const int32 CompactPoseBoneIndex = BoneIndex.GetInt();
-				if (CompactPoseBoneIndex != INDEX_NONE)
-				{
-					RotationScalePairs.Add(BoneTrackPair(CompactPoseBoneIndex, TrackIndex));
-
-					// Skip extracting translation component for EBoneTranslationRetargetingMode::Skeleton.
-					switch (TargetSkeleton->GetBoneTranslationRetargetingMode(TargetSkeletonBoneIndex))
-					{
-					case EBoneTranslationRetargetingMode::Animation:
-						TranslationPairs.Add(BoneTrackPair(CompactPoseBoneIndex, TrackIndex));
-						break;
-					case EBoneTranslationRetargetingMode::AnimationScaled:
-						TranslationPairs.Add(BoneTrackPair(CompactPoseBoneIndex, TrackIndex));
-						AnimScaleRetargetingPairs.Add(BoneTrackPair(CompactPoseBoneIndex, SourceSkeletonBoneIndex));
-						break;
-					case EBoneTranslationRetargetingMode::AnimationRelative:
-						TranslationPairs.Add(BoneTrackPair(CompactPoseBoneIndex, TrackIndex));
-
-						// With baked additives, we can skip 'AnimationRelative' tracks, as the relative transform gets canceled out.
-						// (A1 + Rel) - (A2 + Rel) = A1 - A2.
-						if (!bIsBakedAdditive)
-						{
-							AnimRelativeRetargetingPairs.Add(BoneTrackPair(CompactPoseBoneIndex, SourceSkeletonBoneIndex));
-						}
-						break;
-					case EBoneTranslationRetargetingMode::OrientAndScale:
-						TranslationPairs.Add(BoneTrackPair(CompactPoseBoneIndex, TrackIndex));
-
-						// Additives remain additives, they're not retargeted.
-						if (!bIsBakedAdditive)
-						{
-							OrientAndScaleRetargetingPairs.Add(BoneTrackPair(CompactPoseBoneIndex, SourceSkeletonBoneIndex));
-						}
-						break;
-					}
-				}
-			}
-		}
-	}
-
-	{
-		SCOPE_CYCLE_COUNTER(STAT_ExtractPoseFromAnimData);
-		CSV_SCOPED_TIMING_STAT(Animation, ExtractPoseFromAnimData);
-		CSV_CUSTOM_STAT(Animation, NumberOfExtractedAnimations, 1, ECsvCustomStatOp::Accumulate);
-
-		FAnimSequenceDecompressionContext EvalDecompContext(SequenceLength, Interpolation, SourceName, *CompressedData.CompressedDataStructure);
-		EvalDecompContext.Seek(ExtractionContext.CurrentTime);
-
-		// Handle Root Bone separately
-		if (bFirstTrackIsRootBone)
-		{
-			const int32 TrackIndex = 0;
-			FCompactPoseBoneIndex RootBone(0);
-			FTransform& RootAtom = OutPose[RootBone];
-
-			CompressedData.BoneCompressionCodec->DecompressBone(EvalDecompContext, TrackIndex, RootAtom);
-
-			// Retarget the root onto the target skeleton (correcting for differences in rest poses)
-			if (SkeletonRemapping)
-			{
-				const int32 TargetSkeletonBoneIndex = 0;
-
-				if (bIsBakedAdditive)
-				{
-					RootAtom.SetRotation(SkeletonRemapping->RetargetAdditiveRotationToTargetSkeleton(TargetSkeletonBoneIndex, RootAtom.GetRotation()));
-					if (TargetSkeleton->GetBoneTranslationRetargetingMode(TargetSkeletonBoneIndex) != EBoneTranslationRetargetingMode::Skeleton)
-					{
-						RootAtom.SetTranslation(SkeletonRemapping->RetargetAdditiveTranslationToTargetSkeleton(TargetSkeletonBoneIndex, RootAtom.GetTranslation()));
-					}
-				}
-				else
-				{
-					RootAtom.SetRotation(SkeletonRemapping->RetargetBoneRotationToTargetSkeleton(TargetSkeletonBoneIndex, RootAtom.GetRotation()));
-					if (TargetSkeleton->GetBoneTranslationRetargetingMode(TargetSkeletonBoneIndex) != EBoneTranslationRetargetingMode::Skeleton)
-					{
-						RootAtom.SetTranslation(SkeletonRemapping->RetargetBoneTranslationToTargetSkeleton(TargetSkeletonBoneIndex, RootAtom.GetTranslation()));
-					}
-				}
-			}
-
-			// @laurent - we should look into splitting rotation and translation tracks, so we don't have to process translation twice.
-			FAnimationRuntime::RetargetBoneTransform(SourceSkeleton, SourceName, RetargetTransforms, RootAtom, 0, RootBone, RequiredBones, bIsBakedAdditive);
-		}
-
-		if (RotationScalePairs.Num() > 0)
-		{
-			// get the remaining bone atoms
-			TArrayView<FTransform> OutPoseBones = OutPose.GetMutableBones();
-			CompressedData.BoneCompressionCodec->DecompressPose(EvalDecompContext, RotationScalePairs, TranslationPairs, RotationScalePairs, OutPoseBones);
-		}
-	}
-
-	// Retarget the pose onto the target skeleton (correcting for differences in rest poses)
-	if (SkeletonRemapping)
-	{
-		if (bIsBakedAdditive)
-		{
-			for (FCompactPoseBoneIndex BoneIndex(bFirstTrackIsRootBone ? 1 : 0); BoneIndex < OutPose.GetNumBones(); ++BoneIndex)
-			{
-				const int32 TargetSkeletonBoneIndex = RequiredBones.GetSkeletonIndex(BoneIndex);
-				OutPose[BoneIndex].SetRotation(SkeletonRemapping->RetargetAdditiveRotationToTargetSkeleton(TargetSkeletonBoneIndex, OutPose[BoneIndex].GetRotation()));
-				if (TargetSkeleton->GetBoneTranslationRetargetingMode(TargetSkeletonBoneIndex) != EBoneTranslationRetargetingMode::Skeleton)
-				{
-					OutPose[BoneIndex].SetTranslation(SkeletonRemapping->RetargetAdditiveTranslationToTargetSkeleton(TargetSkeletonBoneIndex, OutPose[BoneIndex].GetTranslation()));
-				}
-			}
-		}
-		else
-		{
-			for (FCompactPoseBoneIndex BoneIndex(bFirstTrackIsRootBone ? 1 : 0); BoneIndex < OutPose.GetNumBones(); ++BoneIndex)
-			{
-				const int32 TargetSkeletonBoneIndex = RequiredBones.GetSkeletonIndex(BoneIndex);
-				OutPose[BoneIndex].SetRotation(SkeletonRemapping->RetargetBoneRotationToTargetSkeleton(TargetSkeletonBoneIndex, OutPose[BoneIndex].GetRotation()));
-				if (TargetSkeleton->GetBoneTranslationRetargetingMode(TargetSkeletonBoneIndex) != EBoneTranslationRetargetingMode::Skeleton)
-				{
-					OutPose[BoneIndex].SetTranslation(SkeletonRemapping->RetargetBoneTranslationToTargetSkeleton(TargetSkeletonBoneIndex, OutPose[BoneIndex].GetTranslation()));
-				}
-			}
-		}
-	}
-
-	// Once pose has been extracted, snap root bone back to first frame if we are extracting root motion.
-	if ((ExtractionContext.bExtractRootMotion && RootMotionReset.bEnableRootMotion) || RootMotionReset.bForceRootLock)
-	{
-		RootMotionReset.ResetRootBoneForRootMotion(OutPose[FCompactPoseBoneIndex(0)], RequiredBones);
-	}
-
-	// Anim Scale Retargeting
-	int32 const NumBonesToScaleRetarget = AnimScaleRetargetingPairs.Num();
-	if (NumBonesToScaleRetarget > 0)
-	{
-		TArray<FTransform> const& AuthoredOnRefSkeleton = RetargetTransforms;
-
-		for (const BoneTrackPair& BonePair : AnimScaleRetargetingPairs)
-		{
-			const FCompactPoseBoneIndex BoneIndex(BonePair.AtomIndex); //Nasty, we break our type safety, code in the lower levels should be adjusted for this
-			const int32 SourceSkeletonBoneIndex = BonePair.TrackIndex;
-
-			// @todo - precache that in FBoneContainer when we have SkeletonIndex->TrackIndex mapping. So we can just apply scale right away.
-			float const SourceTranslationLength = AuthoredOnRefSkeleton[SourceSkeletonBoneIndex].GetTranslation().Size();
-			if (SourceTranslationLength > UE_KINDA_SMALL_NUMBER)
-			{
-				float const TargetTranslationLength = RequiredBones.GetRefPoseTransform(BoneIndex).GetTranslation().Size();
-				OutPose[BoneIndex].ScaleTranslation(TargetTranslationLength / SourceTranslationLength);
-			}
-		}
-	}
-
-	// Anim Relative Retargeting
-	int32 const NumBonesToRelativeRetarget = AnimRelativeRetargetingPairs.Num();
-	if (NumBonesToRelativeRetarget > 0)
-	{
-		TArray<FTransform> const& AuthoredOnRefSkeleton = RetargetTransforms;
-
-		for (const BoneTrackPair& BonePair : AnimRelativeRetargetingPairs)
-		{
-			const FCompactPoseBoneIndex BoneIndex(BonePair.AtomIndex); //Nasty, we break our type safety, code in the lower levels should be adjusted for this
-			const int32 SourceSkeletonBoneIndex = BonePair.TrackIndex;
-
-			const FTransform& RefPoseTransform = RequiredBones.GetRefPoseTransform(BoneIndex);
-
-			// Remap the base pose onto the target skeleton so that we are working entirely in target space
-			FTransform BaseTransform = AuthoredOnRefSkeleton[SourceSkeletonBoneIndex];
-			if (SkeletonRemapping)
-			{
-				const int32 TargetSkeletonBoneIndex = SkeletonRemapping->GetTargetSkeletonBoneIndex(SourceSkeletonBoneIndex);
-				BaseTransform = SkeletonRemapping->RetargetBoneTransformToTargetSkeleton(TargetSkeletonBoneIndex, BaseTransform);
-			}
-
-			// Apply the retargeting as if it were an additive difference between the current skeleton and the retarget skeleton. 
-			OutPose[BoneIndex].SetRotation(OutPose[BoneIndex].GetRotation() * BaseTransform.GetRotation().Inverse() * RefPoseTransform.GetRotation());
-			OutPose[BoneIndex].SetTranslation(OutPose[BoneIndex].GetTranslation() + (RefPoseTransform.GetTranslation() - BaseTransform.GetTranslation()));
-			OutPose[BoneIndex].SetScale3D(OutPose[BoneIndex].GetScale3D() * (RefPoseTransform.GetScale3D() * BaseTransform.GetSafeScaleReciprocal(BaseTransform.GetScale3D())));
-			OutPose[BoneIndex].NormalizeRotation();
-		}
-	}
-
-	// Translation 'Orient and Scale' Translation Retargeting
-	const int32 NumBonesToOrientAndScaleRetarget = OrientAndScaleRetargetingPairs.Num();
-	if (NumBonesToOrientAndScaleRetarget > 0)
-	{
-		const FRetargetSourceCachedData& RetargetSourceCachedData = RequiredBones.GetRetargetSourceCachedData(SourceName, RetargetTransforms);
-		const TArray<FOrientAndScaleRetargetingCachedData>& OrientAndScaleDataArray = RetargetSourceCachedData.OrientAndScaleData;
-		const TArray<int32>& CompactPoseIndexToOrientAndScaleIndex = RetargetSourceCachedData.CompactPoseIndexToOrientAndScaleIndex;
-
-		// If we have any cached retargeting data.
-		if ((OrientAndScaleDataArray.Num() > 0) && (CompactPoseIndexToOrientAndScaleIndex.Num() == RequiredBones.GetCompactPoseNumBones()))
-		{
-			for (int32 Index = 0; Index < NumBonesToOrientAndScaleRetarget; Index++)
-			{
-				const BoneTrackPair& BonePair = OrientAndScaleRetargetingPairs[Index];
-				const FCompactPoseBoneIndex CompactPoseBoneIndex(BonePair.AtomIndex);
-				const int32 OrientAndScaleIndex = CompactPoseIndexToOrientAndScaleIndex[CompactPoseBoneIndex.GetInt()];
-				if (OrientAndScaleIndex != INDEX_NONE)
-				{
-					const FOrientAndScaleRetargetingCachedData& OrientAndScaleData = OrientAndScaleDataArray[OrientAndScaleIndex];
-					FTransform& BoneTransform = OutPose[CompactPoseBoneIndex];
-					const FVector AnimatedTranslation = BoneTransform.GetTranslation();
-
-					// If Translation is not animated, we can just copy the TargetTranslation. No retargeting needs to be done.
-					const FVector NewTranslation = (AnimatedTranslation - OrientAndScaleData.SourceTranslation).IsNearlyZero(BONE_TRANS_RT_ORIENT_AND_SCALE_PRECISION) ?
-						OrientAndScaleData.TargetTranslation :
-						OrientAndScaleData.TranslationDeltaOrient.RotateVector(AnimatedTranslation) * OrientAndScaleData.TranslationScale;
-
-					BoneTransform.SetTranslation(NewTranslation);
-				}
-			}
-		}
-	}
+	PRAGMA_DISABLE_DEPRECATION_WARNINGS
+	FAnimSequenceDecompressionContext DecompressionContext(SequenceLength, Interpolation, SourceName, *CompressedData.CompressedDataStructure.Get(), SourceSkeleton, bIsBakedAdditive);
+	PRAGMA_ENABLE_DEPRECATION_WARNINGS
+	UE::Anim::Decompression::DecompressPose(OutPose, CompressedData, ExtractionContext, DecompressionContext, RetargetTransforms, RootMotionReset);
 }
 
 FArchive& operator<<(FArchive& Ar, FCompressedOffsetData& D)
diff --git a/Engine/Source/Runtime/Engine/Private/Animation/AnimCurveCompressionCodec_CompressedRichCurve.cpp b/Engine/Source/Runtime/Engine/Private/Animation/AnimCurveCompressionCodec_CompressedRichCurve.cpp
--- a/Engine/Source/Runtime/Engine/Private/Animation/AnimCurveCompressionCodec_CompressedRichCurve.cpp
+++ b/Engine/Source/Runtime/Engine/Private/Animation/AnimCurveCompressionCodec_CompressedRichCurve.cpp
@@ -45,8 +45,8 @@ bool UAnimCurveCompressionCodec_CompressedRichCurve::Compress(const FCompressibl
 	int32 KeyDataOffset = 0;
 	KeyDataOffset += sizeof(FCurveDesc) * NumCurves;
 
-	const FAnimKeyHelper Helper(AnimSeq.SequenceLength, AnimSeq.NumberOfKeys);
-	const float SampleRate = UseAnimSequenceSampleRate ? Helper.KeysPerSecond() : ErrorSampleRate;
+	const FFrameRate& SamplingFrameRate = AnimSeq.SampledFrameRate;
+	const double SampleRate = AnimSeq.SampledFrameRate.AsDecimal();
 
 	TArray<uint8> KeyData;
 
@@ -78,15 +78,15 @@ bool UAnimCurveCompressionCodec_CompressedRichCurve::Compress(const FCompressibl
 			for (int32 KeyIndex = 0; KeyIndex < SamplesToTake; ++KeyIndex)
 			{
 				// Evaluated RawCurve and Compressed Curve
-				const float EvalTime = KeyIndex * Helper.TimePerKey();
-				const float RawValue = RawCurve.Eval(EvalTime);
+				const double EvalTime = SamplingFrameRate.AsSeconds(KeyIndex);
+				const float RawValue = Curve.FloatCurve.Eval(EvalTime);
 				const float CompressedValue = CompressedCurve.Eval(EvalTime);
 
 				const float AbsDelta = FMath::Abs(CompressedValue - RawValue);
-				if (AbsDelta > MaxCurveError)
+				if (!FMath::IsNearlyZero(AbsDelta, MaxCurveError + UE_KINDA_SMALL_NUMBER))
 				{
 					// Delta larger than tolerated error value
-					UE_LOG(LogAnimationCompression, Warning, TEXT("Curve %s: delta too large %f at %f"), *Curve.Name.DisplayName.ToString(), AbsDelta, EvalTime);
+					UE_LOG(LogAnimationCompression, Warning, TEXT("Curve %s: delta too large %f, between %f and %f, at %f"), *Curve.Name.DisplayName.ToString(), AbsDelta, RawValue, CompressedValue, EvalTime);
 				}
 			}
 		}
diff --git a/Engine/Source/Runtime/Engine/Private/Animation/AnimCurveCompressionCodec_UniformIndexable.cpp b/Engine/Source/Runtime/Engine/Private/Animation/AnimCurveCompressionCodec_UniformIndexable.cpp
--- a/Engine/Source/Runtime/Engine/Private/Animation/AnimCurveCompressionCodec_UniformIndexable.cpp
+++ b/Engine/Source/Runtime/Engine/Private/Animation/AnimCurveCompressionCodec_UniformIndexable.cpp
@@ -89,11 +89,7 @@ UAnimCurveCompressionCodec_UniformIndexable::UAnimCurveCompressionCodec_UniformI
 bool UAnimCurveCompressionCodec_UniformIndexable::Compress(const FCompressibleAnimData& AnimSeq, FAnimCurveCompressionResult& OutResult)
 {
 	const int32 NumCurves = AnimSeq.RawFloatCurves.Num();
-	const float Duration = AnimSeq.SequenceLength;
-
-	const FAnimKeyHelper Helper(AnimSeq.SequenceLength, AnimSeq.NumberOfKeys);
-	const float SampleRate = Helper.KeysPerSecond();
-	const int32 NumSamples = FMath::RoundToInt(Duration * SampleRate) + 1;
+	const int32 NumSamples = AnimSeq.NumberOfKeys;
 
 	int32 BufferSize = 0;
 	BufferSize += sizeof(int32);	// NumSamples
@@ -113,7 +109,7 @@ bool UAnimCurveCompressionCodec_UniformIndexable::Compress(const FCompressibleAn
 	BufferOffset += sizeof(float);
 
 	*NumSamplesPtr = NumSamples;
-	*SampleRatePtr = SampleRate;
+	*SampleRatePtr = AnimSeq.SampledFrameRate.AsDecimal();
 
 	if (NumCurves > 0 && NumSamples > 0)
 	{
@@ -123,12 +119,12 @@ bool UAnimCurveCompressionCodec_UniformIndexable::Compress(const FCompressibleAn
 		// Write out samples by curve so that we can more easily index curve keys
 
 		// Curve 0 Key 0, Curve 0 Key 1, Curve 0 Key N, Curve 1 Key 0, Curve 1 Key 1, Curve 1 Key N, Curve M Key 0, ...
-		const float InvSampleRate = 1.0f / SampleRate;
 		for (const FFloatCurve& Curve : AnimSeq.RawFloatCurves)
 		{
 			for (int32 SampleIndex = 0; SampleIndex < NumSamples; ++SampleIndex)
 			{
-				const float SampleTime = FMath::Clamp(SampleIndex * InvSampleRate, 0.0f, Duration);
+				const double FrameSampleTime = AnimSeq.SampledFrameRate.AsSeconds(SampleIndex);
+				const float SampleTime = FMath::Clamp(FrameSampleTime, 0.0f, AnimSeq.SequenceLength);
 				const float SampleValue = Curve.FloatCurve.Eval(SampleTime);
 
 				(*AnimatedSamplesPtr) = SampleValue;
diff --git a/Engine/Source/Runtime/Engine/Private/Animation/AnimCurveCompressionCodec_UniformlySampled.cpp b/Engine/Source/Runtime/Engine/Private/Animation/AnimCurveCompressionCodec_UniformlySampled.cpp
--- a/Engine/Source/Runtime/Engine/Private/Animation/AnimCurveCompressionCodec_UniformlySampled.cpp
+++ b/Engine/Source/Runtime/Engine/Private/Animation/AnimCurveCompressionCodec_UniformlySampled.cpp
@@ -23,9 +23,8 @@ bool UAnimCurveCompressionCodec_UniformlySampled::Compress(const FCompressibleAn
 	float SampleRate_;
 	if (UseAnimSequenceSampleRate)
 	{
-		const FAnimKeyHelper Helper(AnimSeq.SequenceLength, AnimSeq.NumberOfKeys);
-		SampleRate_ = Helper.KeysPerSecond();
-		NumSamples = FMath::RoundToInt(Duration * SampleRate_) + 1;
+		NumSamples = AnimSeq.NumberOfKeys;
+		SampleRate_ = AnimSeq.SampledFrameRate.AsDecimal();
 	}
 	else
 	{
diff --git a/Engine/Source/Runtime/Engine/Private/Animation/AnimDataModel.cpp b/Engine/Source/Runtime/Engine/Private/Animation/AnimDataModel.cpp
--- a/Engine/Source/Runtime/Engine/Private/Animation/AnimDataModel.cpp
+++ b/Engine/Source/Runtime/Engine/Private/Animation/AnimDataModel.cpp
@@ -1,16 +1,28 @@
 // Copyright Epic Games, Inc. All Rights Reserved.
 
 #include "Animation/AnimData/AnimDataModel.h"
+
+#include "AnimationRuntime.h"
+#include "AnimEncoding.h"
 #include "UObject/NameTypes.h"
 
 #include "Animation/AnimSequence.h"
 
 #include "Algo/Transform.h"
 #include "Algo/Accumulate.h"
+#include "Animation/AnimSequenceHelpers.h"
 #include "Animation/SmartName.h"
+#include "Animation/AttributesRuntime.h"
 #include "UObject/UE5MainStreamObjectVersion.h"
 #include "UObject/FortniteMainBranchObjectVersion.h"
 
+#if WITH_EDITOR
+#include "IAnimationDataControllerModule.h"
+#include "Modules/ModuleManager.h"
+#endif // WITH_EDITOR
+
+#define LOCTEXT_NAMESPACE "AnimDataModel"
+
 void UAnimDataModel::PostLoad()
 {
 	UObject::PostLoad();
@@ -18,20 +30,21 @@ void UAnimDataModel::PostLoad()
 	if (GetLinkerCustomVersion(FUE5MainStreamObjectVersion::GUID) < FUE5MainStreamObjectVersion::ForceUpdateAnimationAssetCurveTangents)
 	{
 		// Forcefully AutoSetTangents to fix-up any imported sequences pre the fix for flattening first/last key leave/arrive tangents
-		Notify(EAnimDataModelNotifyType::BracketOpened);
+		GetNotifier().Notify(EAnimDataModelNotifyType::BracketOpened);
 		for (FFloatCurve& FloatCurve : CurveData.FloatCurves)
 		{
 			FloatCurve.FloatCurve.AutoSetTangents();
 			FCurvePayload Payload;
 			Payload.Identifier = FAnimationCurveIdentifier(FloatCurve.Name.UID, ERawCurveTrackTypes::RCT_Float);
-			Notify(EAnimDataModelNotifyType::CurveChanged, Payload);
+			GetNotifier().Notify(EAnimDataModelNotifyType::CurveChanged, Payload);
 		}
-		Notify(EAnimDataModelNotifyType::BracketClosed);
+		GetNotifier().Notify(EAnimDataModelNotifyType::BracketClosed);
 	}
 
+	const bool bHasBoneTracks = BoneAnimationTracks.Num() > 0;
+
 	if (GetLinkerCustomVersion(FFortniteMainBranchObjectVersion::GUID) < FFortniteMainBranchObjectVersion::SingleFrameAndKeyAnimModel)
 	{
-		const bool bHasBoneTracks = BoneAnimationTracks.Num() > 0;
 		if (bHasBoneTracks)
 		{
 			// Number of keys was used directly rather than Max(Value,2), as a single _frame_ animation should always have two _keys_ 
@@ -52,17 +65,47 @@ void UAnimDataModel::PostLoad()
 					AddKey(BoneTrack.InternalTrackData.ScaleKeys);
 				}
 
-				Notify(EAnimDataModelNotifyType::TrackChanged);
+				GetNotifier().Notify(EAnimDataModelNotifyType::TrackChanged);
 			}
 		}
-	}	
+	}
+	
+	if (GetLinkerCustomVersion(FUE5MainStreamObjectVersion::GUID) < FUE5MainStreamObjectVersion::AnimationDataModelInterface)
+	{		
+		const double NumFramesAsSeconds = FrameRate.AsSeconds(NumberOfFrames);
+		PRAGMA_DISABLE_DEPRECATION_WARNINGS
+		const float CachedPlayLength = PlayLength;
+		if (static_cast<float>(NumFramesAsSeconds) != PlayLength)
+		{
+			if (bHasBoneTracks)
+			{
+				// Number of keys was rounded up from sequence length, so set play length from stored number of bone keys 
+				const int32 ActualNumKeys = bHasBoneTracks ? FMath::Min(BoneAnimationTracks[0].InternalTrackData.PosKeys.Num(), NumberOfKeys) : NumberOfKeys;
+				if (ActualNumKeys != NumberOfKeys)
+				{
+					PlayLength = FrameRate.AsSeconds(ActualNumKeys - 1);
+					NumberOfFrames = ActualNumKeys - 1;
+					NumberOfKeys = ActualNumKeys;
+
+					IAnimationDataController::ReportObjectWarningf(this, LOCTEXT("AnimDataModelPlayLengthAlteredToBones", "Playable length {0} has been updated according to number of frames {1} and frame-rate {2}, new length: {3}"), FText::AsNumber(CachedPlayLength), FText::AsNumber(GetNumberOfFrames()), FrameRate.ToPrettyText(), FText::AsNumber(PlayLength));
+				}
+			}
+			else
+			{
+				// No way to determine 'correct' length so set current PlayLength value to # of frames * frame-rate
+				PlayLength = FrameRate.AsSeconds(NumberOfFrames);
+			}
+			GetNotifier().Notify(EAnimDataModelNotifyType::Populated);
+		}
+		PRAGMA_ENABLE_DEPRECATION_WARNINGS
+	}
 }
 
 void UAnimDataModel::PostDuplicate(bool bDuplicateForPIE)
 {
 	UObject::PostDuplicate(bDuplicateForPIE);
 
-	Notify(EAnimDataModelNotifyType::Populated);
+	GetNotifier().Notify(EAnimDataModelNotifyType::Populated);
 }
 
 const TArray<FBoneAnimationTrack>& UAnimDataModel::GetBoneAnimationTracks() const
@@ -116,9 +159,9 @@ int32 UAnimDataModel::GetBoneTrackIndexByName(FName TrackName) const
 	return INDEX_NONE;
 }
 
-float UAnimDataModel::GetPlayLength() const
+double UAnimDataModel::GetPlayLength() const
 {
-	return PlayLength;
+	return FrameRate.AsSeconds(NumberOfFrames);
 }
 
 int32 UAnimDataModel::GetNumberOfFrames() const
@@ -131,7 +174,7 @@ int32 UAnimDataModel::GetNumberOfKeys() const
 	return NumberOfKeys;
 }
 
-const FFrameRate& UAnimDataModel::GetFrameRate() const
+FFrameRate UAnimDataModel::GetFrameRate() const
 {
 	return FrameRate;
 }
@@ -141,14 +184,15 @@ bool UAnimDataModel::IsValidBoneTrackIndex(int32 TrackIndex) const
 	return BoneAnimationTracks.IsValidIndex(TrackIndex);
 }
 
-const int32 UAnimDataModel::GetNumBoneTracks() const
+int32 UAnimDataModel::GetNumBoneTracks() const
 {
 	return BoneAnimationTracks.Num();
 }
 
 UAnimSequence* UAnimDataModel::GetAnimationSequence() const
 {
-	return Cast<UAnimSequence>(GetOuter());
+	UObject* AnimObject = static_cast<UObject*>(FindObjectWithOuter(GetOuter(), UAnimSequence::StaticClass()));
+	return CastChecked<UAnimSequence>(AnimObject ? AnimObject : GetOuter());
 }
 
 void UAnimDataModel::GetBoneTrackNames(TArray<FName>& OutNames) const
@@ -164,6 +208,16 @@ const FAnimationCurveData& UAnimDataModel::GetCurveData() const
 	return CurveData;
 }
 
+IAnimationDataModel::FModelNotifier& UAnimDataModel::GetNotifier()
+{
+	if (!Notifier)
+	{
+		Notifier.Reset(new IAnimationDataModel::FModelNotifier(this));
+	}
+	
+	return *Notifier.Get();
+}
+
 int32 UAnimDataModel::GetNumberOfTransformCurves() const
 {
 	return CurveData.TransformCurves.Num();
@@ -299,7 +353,6 @@ int32 UAnimDataModel::GetNumberOfAttributesForBoneIndex(const int32 BoneIndex) c
 
 void UAnimDataModel::GetAttributesForBone(const FName& BoneName, TArray<const FAnimatedBoneAttribute*>& OutBoneAttributes) const
 {
-	// Sum up total number of attributes with provided bone name
 	Algo::TransformIf(AnimatedBoneAttributes, OutBoneAttributes, [BoneName](const FAnimatedBoneAttribute& Attribute) -> bool
 	{
 		return Attribute.Identifier.GetBoneName() == BoneName;
@@ -378,7 +431,7 @@ FGuid UAnimDataModel::GenerateGuid() const
     {
     	UpdateWithData(Curve.Name.UID);
 
-    	auto UpdateWithComponent = [&Sha, &UpdateWithFloatCurve](const FVectorCurve& VectorCurve)
+    	auto UpdateWithComponent = [&UpdateWithFloatCurve](const FVectorCurve& VectorCurve)
     	{
     		for (int32 ChannelIndex = 0; ChannelIndex < 3; ++ChannelIndex)
     		{
@@ -406,6 +459,252 @@ FGuid UAnimDataModel::GenerateGuid() const
 	return Guid;
 }
 
+TScriptInterface<IAnimationDataController> UAnimDataModel::GetController()
+{
+	TScriptInterface<IAnimationDataController> Controller = nullptr;
+#if WITH_EDITOR
+	IAnimationDataControllerModule& ControllerModule = FModuleManager::Get().GetModuleChecked<IAnimationDataControllerModule>("AnimationDataController");
+	Controller = ControllerModule.GetController();
+	Controller->SetModel(this);
+#endif // WITH_EDITOR
+
+	return Controller;
+}
+
+#if WITH_EDITOR
+FTransform ExtractTransformForKeyIndex(int32 Key, const FRawAnimSequenceTrack& TrackToExtract)
+{
+	static const FVector DefaultScale3D = FVector(1.f);
+	const bool bHasScaleKey = TrackToExtract.ScaleKeys.Num() > 0;
+
+	const int32 PosKeyIndex = FMath::Min(Key, TrackToExtract.PosKeys.Num() - 1);
+	const int32 RotKeyIndex = FMath::Min(Key, TrackToExtract.RotKeys.Num() - 1);
+	if (bHasScaleKey)
+	{
+		const int32 ScaleKeyIndex = FMath::Min(Key, TrackToExtract.ScaleKeys.Num() - 1);
+		return FTransform(FQuat(TrackToExtract.RotKeys[RotKeyIndex]), FVector(TrackToExtract.PosKeys[PosKeyIndex]), FVector(TrackToExtract.ScaleKeys[ScaleKeyIndex]));
+	}
+	else
+	{
+		return FTransform(FQuat(TrackToExtract.RotKeys[RotKeyIndex]), FVector(TrackToExtract.PosKeys[PosKeyIndex]), DefaultScale3D);
+	}
+}
+
+template<bool bInterpolateT>
+void ExtractPose(const TArray<FBoneAnimationTrack>& BoneAnimationTracks, const TMap<FName, const FTransformCurve*>& OverrideBoneTransforms, FCompactPose& InOutPose, const int32 KeyIndex1, const int32 KeyIndex2, float Alpha, float TimePerKey, UE::Anim::Retargeting::FRetargetingScope& RetargetingScope)
+{
+	const int32 NumAnimationTracks = BoneAnimationTracks.Num();
+	const FBoneContainer& RequiredBones = InOutPose.GetBoneContainer();
+
+	TArray<FVirtualBoneCompactPoseData>& VBCompactPoseData = UE::Anim::FBuildRawPoseScratchArea::Get().VirtualBoneCompactPoseData;
+	VBCompactPoseData = RequiredBones.GetVirtualBoneCompactPoseData();
+
+	FCompactPose Key2Pose;
+	if (bInterpolateT)
+	{
+		Key2Pose.CopyBonesFrom(InOutPose);
+	}
+
+	for (const FBoneAnimationTrack& AnimationTrack : BoneAnimationTracks)
+	{
+		const int32 SkeletonBoneIndex = AnimationTrack.BoneTreeIndex;
+		// not sure it's safe to assume that SkeletonBoneIndex can never be INDEX_NONE
+		if ((SkeletonBoneIndex != INDEX_NONE) && (SkeletonBoneIndex < MAX_BONES))
+		{
+			const FCompactPoseBoneIndex PoseBoneIndex = RequiredBones.GetCompactPoseIndexFromSkeletonIndex(SkeletonBoneIndex);
+			if (PoseBoneIndex != INDEX_NONE)
+			{
+				for (int32 Idx = 0; Idx < VBCompactPoseData.Num(); ++Idx)
+				{
+					FVirtualBoneCompactPoseData& VB = VBCompactPoseData[Idx];
+					if (PoseBoneIndex == VB.VBIndex)
+					{
+						// Remove this bone as we have written data for it (false so we dont resize allocation)
+						VBCompactPoseData.RemoveAtSwap(Idx, 1, false);
+						break; //Modified TArray so must break here
+					}
+				}
+				
+				// extract animation
+				const FRawAnimSequenceTrack& TrackToExtract = AnimationTrack.InternalTrackData;
+				auto OverrideTransform = OverrideBoneTransforms.Find(AnimationTrack.Name);
+				{
+					// Bail out (with rather wacky data) if data is empty for some reason.
+					if (TrackToExtract.PosKeys.Num() == 0 || TrackToExtract.RotKeys.Num() == 0)
+					{
+						InOutPose[PoseBoneIndex].SetIdentity();
+
+						if (bInterpolateT)
+						{
+							Key2Pose[PoseBoneIndex].SetIdentity();
+						}
+					}
+					else
+					{
+						InOutPose[PoseBoneIndex] = ExtractTransformForKeyIndex(KeyIndex1, TrackToExtract);
+
+						if (bInterpolateT)
+						{
+							Key2Pose[PoseBoneIndex] = ExtractTransformForKeyIndex(KeyIndex2, TrackToExtract);
+						}
+					}
+				}
+
+				if (OverrideTransform)
+				{
+					{
+						const FTransform PoseOneAdditive = (*OverrideTransform)->Evaluate(KeyIndex1 * TimePerKey, 1.f);
+						const FTransform PoseOneLocalTransform = InOutPose[PoseBoneIndex];
+						InOutPose[PoseBoneIndex].SetRotation(PoseOneLocalTransform.GetRotation() * PoseOneAdditive.GetRotation());
+						InOutPose[PoseBoneIndex].SetTranslation(PoseOneLocalTransform.TransformPosition(PoseOneAdditive.GetTranslation()));
+						InOutPose[PoseBoneIndex].SetScale3D(PoseOneLocalTransform.GetScale3D() * PoseOneAdditive.GetScale3D());
+					}
+					
+
+					if (bInterpolateT)
+					{
+						const FTransform PoseTwoAdditive = (*OverrideTransform)->Evaluate(KeyIndex2 * TimePerKey, 1.f);
+						const FTransform PoseTwoLocalTransform = Key2Pose[PoseBoneIndex];						
+						Key2Pose[PoseBoneIndex].SetRotation(PoseTwoLocalTransform.GetRotation() * PoseTwoAdditive.GetRotation());
+						Key2Pose[PoseBoneIndex].SetTranslation(PoseTwoLocalTransform.TransformPosition(PoseTwoAdditive.GetTranslation()));
+						Key2Pose[PoseBoneIndex].SetScale3D(PoseTwoLocalTransform.GetScale3D() * PoseTwoAdditive.GetScale3D());
+					}
+				}
+				RetargetingScope.AddTrackedBone(PoseBoneIndex, SkeletonBoneIndex);
+			}
+		}
+	}
+
+	//Build Virtual Bones
+	if (VBCompactPoseData.Num() > 0)
+	{
+		FCSPose<FCompactPose> CSPose1;
+		CSPose1.InitPose(InOutPose);
+
+		FCSPose<FCompactPose> CSPose2;
+		if (bInterpolateT)
+		{
+			CSPose2.InitPose(Key2Pose);
+		}
+
+		for (FVirtualBoneCompactPoseData& VB : VBCompactPoseData)
+		{
+			FTransform Source = CSPose1.GetComponentSpaceTransform(VB.SourceIndex);
+			FTransform Target = CSPose1.GetComponentSpaceTransform(VB.TargetIndex);
+			InOutPose[VB.VBIndex] = Target.GetRelativeTransform(Source);
+
+			if (bInterpolateT)
+			{
+				FTransform Source2 = CSPose2.GetComponentSpaceTransform(VB.SourceIndex);
+				FTransform Target2 = CSPose2.GetComponentSpaceTransform(VB.TargetIndex);
+				Key2Pose[VB.VBIndex] = Target2.GetRelativeTransform(Source2);
+			}
+		}
+	}
+
+	if (bInterpolateT)
+	{
+		for (FCompactPoseBoneIndex BoneIndex : InOutPose.ForEachBoneIndex())
+		{
+			InOutPose[BoneIndex].Blend(InOutPose[BoneIndex], Key2Pose[BoneIndex], Alpha);
+		}
+	}
+}
+
+void UAnimDataModel::Evaluate(FAnimationPoseData& InOutPoseData, const UE::Anim::DataModel::FEvaluationContext& EvaluationContext) const
+{
+	const float SequenceLength = GetPlayLength();
+
+	const double Time = EvaluationContext.SampleFrameRate.AsSeconds(EvaluationContext.SampleTime);
+	const EAnimInterpolationType& InterpolationType = EvaluationContext.InterpolationType;
+	FCompactPose& OutPose = InOutPoseData.GetPose();
+	
+	// Generate keys to interpolate between
+	int32 KeyIndex1, KeyIndex2;
+	float Alpha;
+	FAnimationRuntime::GetKeyIndicesFromTime(KeyIndex1, KeyIndex2, Alpha, Time, FrameRate, NumberOfKeys);
+
+	if (InterpolationType == EAnimInterpolationType::Step)
+	{
+		// Force stepping between keys
+		Alpha = 0.f;
+	}
+
+	bool bShouldInterpolate = true;
+
+	if (Alpha < UE_KINDA_SMALL_NUMBER)
+	{
+		Alpha = 0.f;
+		bShouldInterpolate = false;
+	}
+	else if (Alpha > 1.f - UE_KINDA_SMALL_NUMBER)
+	{
+		bShouldInterpolate = false;
+		KeyIndex1 = KeyIndex2;
+	}
+
+	// Evaluate animation float curve data
+	FBlendedCurve& Curves = InOutPoseData.GetCurve();
+	if (Curves.NumValidCurveCount > 0)
+	{
+		// evaluate the curve data at the CurrentTime and add to Instance
+		for (const FFloatCurve& Curve : CurveData.FloatCurves)
+		{
+			if (Curves.IsEnabled(Curve.Name.UID))
+			{
+				float Value = Curve.Evaluate(Time);
+				Curves.Set(Curve.Name.UID, Value);
+			}
+		}
+	}
+
+	const double TimePerFrame = FrameRate.AsInterval();
+	TMap<FName, const FTransformCurve*> ActiveCurves;
+
+	if (!OutPose.GetBoneContainer().ShouldUseSourceData())
+	{
+		for (const FTransformCurve& Curve : GetTransformCurves())
+		{
+			// if disabled, do not handle
+			if (Curve.GetCurveTypeFlag(AACF_Disabled))
+			{
+				continue;
+			}
+
+			// Add or retrieve curve
+			const FName& CurveName = Curve.Name.DisplayName;
+
+			// note we're not checking Curve.GetCurveTypeFlags() yet
+			ActiveCurves.FindOrAdd(CurveName, &Curve);
+		}
+	}
+
+	{		
+		UE::Anim::Retargeting::FRetargetingScope RetargetingScope(OutPose, EvaluationContext);
+		if (bShouldInterpolate)
+		{		
+			ExtractPose<true>(GetBoneAnimationTracks(), ActiveCurves, OutPose, KeyIndex1, KeyIndex2, Alpha, TimePerFrame, RetargetingScope);
+		}
+		else
+		{
+			ExtractPose<false>(GetBoneAnimationTracks(), ActiveCurves, OutPose, KeyIndex1, KeyIndex2, Alpha, TimePerFrame, RetargetingScope);
+		}
+	}
+	
+
+	const FBoneContainer& RequiredBones = InOutPoseData.GetPose().GetBoneContainer();
+	for (const FAnimatedBoneAttribute& Attribute : AnimatedBoneAttributes)
+	{
+		const FCompactPoseBoneIndex PoseBoneIndex = RequiredBones.GetCompactPoseIndexFromSkeletonIndex(Attribute.Identifier.GetBoneIndex());
+		// Only add attribute if the bone its tied to exists in the currently evaluated set of bones
+		if(PoseBoneIndex.IsValid())
+		{
+			UE::Anim::Attributes::GetAttributeValue(InOutPoseData.GetAttributes(), PoseBoneIndex, Attribute, EvaluationContext.SampleFrameRate.AsSeconds(EvaluationContext.SampleTime));
+		}
+	}	
+}
+#endif // WITH_EDITOR
+
 FRichCurve* UAnimDataModel::GetMutableRichCurve(const FAnimationCurveIdentifier& CurveIdentifier)
 {
 	FRichCurve* RichCurve = nullptr;
@@ -509,3 +808,4 @@ FAnimCurveBase* UAnimDataModel::FindMutableCurveById(const FAnimationCurveIdenti
 	return nullptr;
 }
 
+#undef LOCTEXT_NAMESPACE // "AnimDataModel"
diff --git a/Engine/Source/Runtime/Engine/Private/Animation/AnimInstanceProxy.cpp b/Engine/Source/Runtime/Engine/Private/Animation/AnimInstanceProxy.cpp
--- a/Engine/Source/Runtime/Engine/Private/Animation/AnimInstanceProxy.cpp
+++ b/Engine/Source/Runtime/Engine/Private/Animation/AnimInstanceProxy.cpp
@@ -1346,7 +1346,7 @@ void FAnimInstanceProxy::EvaluateAnimationNode(FPoseContext& Output)
 
 void FAnimInstanceProxy::EvaluateAnimationNode_WithRoot(FPoseContext& Output, FAnimNode_Base* InRootNode)
 {
-	if (InRootNode != NULL)
+	if (InRootNode != nullptr)
 	{
 		ANIM_MT_SCOPE_CYCLE_COUNTER(EvaluateAnimGraph, !IsInGameThread());
 		if(InRootNode == RootNode)
@@ -1467,7 +1467,7 @@ void FAnimInstanceProxy::SlotEvaluatePoseWithBlendProfiles(const FName& SlotNode
 			NewPose.Curve.InitFrom(RequiredBones);
 
 			// Extract pose from Track.
-			FAnimExtractContext ExtractionContext(EvalState.MontagePosition, Montage->HasRootMotion() && RootMotionMode != ERootMotionMode::NoRootMotionExtraction, EvalState.DeltaTimeRecord);
+			FAnimExtractContext ExtractionContext(static_cast<double>(EvalState.MontagePosition), Montage->HasRootMotion() && RootMotionMode != ERootMotionMode::NoRootMotionExtraction, EvalState.DeltaTimeRecord);
 			FAnimationPoseData NewAnimationPoseData(NewPose);
 			AnimTrack->GetAnimationPose(NewAnimationPoseData, ExtractionContext);
 
@@ -1806,7 +1806,7 @@ void FAnimInstanceProxy::SlotEvaluatePose(const FName& SlotNodeName, const FAnim
 			NewPose.Curve.InitFrom(RequiredBones);
 
 			// Extract pose from Track
-			FAnimExtractContext ExtractionContext(EvalState.MontagePosition, Montage->HasRootMotion() && RootMotionMode != ERootMotionMode::NoRootMotionExtraction, EvalState.DeltaTimeRecord);
+			FAnimExtractContext ExtractionContext(static_cast<double>(EvalState.MontagePosition), Montage->HasRootMotion() && RootMotionMode != ERootMotionMode::NoRootMotionExtraction, EvalState.DeltaTimeRecord);
 
 			FAnimationPoseData NewAnimationPoseData(NewPose);
 			AnimTrack->GetAnimationPose(NewAnimationPoseData, ExtractionContext);
diff --git a/Engine/Source/Runtime/Engine/Private/Animation/AnimMontage.cpp b/Engine/Source/Runtime/Engine/Private/Animation/AnimMontage.cpp
--- a/Engine/Source/Runtime/Engine/Private/Animation/AnimMontage.cpp
+++ b/Engine/Source/Runtime/Engine/Private/Animation/AnimMontage.cpp
@@ -92,7 +92,7 @@ const FAnimTrack* UAnimMontage::GetAnimationData(FName InSlotName) const
 		}
 	}
 
-	return NULL;
+	return nullptr;
 }
 
 bool UAnimMontage::IsWithinPos(int32 FirstIndex, int32 SecondIndex, float CurrentTime) const
@@ -369,7 +369,7 @@ void UAnimMontage::PostLoad()
 		Track.ValidateSegmentTimes();
 
 		const float CurrentCalculatedLength = CalculateSequenceLength();
-		if(CurrentCalculatedLength != GetPlayLength())		
+		if(!FMath::IsNearlyEqual(CurrentCalculatedLength, GetPlayLength(), UE_KINDA_SMALL_NUMBER))		
 		{
 			UE_LOG(LogAnimMontage, Display, TEXT("UAnimMontage::PostLoad: The actual sequence length for %s does not match the length stored in the asset, please resave the asset."), *GetFullName());
 			SetCompositeLength(CurrentCalculatedLength);
@@ -418,7 +418,7 @@ void UAnimMontage::PostLoad()
 	}
 	// find preview base pose if it can
 #if WITH_EDITORONLY_DATA
-	if ( IsValidAdditive() && PreviewBasePose == NULL )
+	if ( IsValidAdditive() && PreviewBasePose == nullptr )
 	{
 		for (int32 I=0; I<SlotAnimTracks.Num(); ++I)
 		{
@@ -555,8 +555,8 @@ void UAnimMontage::ConvertBranchingPointsToAnimNotifies()
 #if WITH_EDITORONLY_DATA
 			NewEvent.TrackIndex = TrackIndex;
 #endif
-			NewEvent.Notify = NULL;
-			NewEvent.NotifyStateClass = NULL;
+			NewEvent.Notify = nullptr;
+			NewEvent.NotifyStateClass = nullptr;
 			NewEvent.bConvertedFromBranchingPoint = true;
 			NewEvent.MontageTickType = EMontageNotifyTickType::BranchingPoint;
 
@@ -1310,11 +1310,14 @@ bool UAnimMontage::ContainRecursive(TArray<UAnimCompositeBase*>& CurrentAccumula
 
 void UAnimMontage::SetCompositeLength(float InLength)
 {
-#if WITH_EDITOR	
-	Controller->SetPlayLength(InLength);
+#if WITH_EDITOR
+	const FFrameTime LengthInFrameTime = DataModelInterface->GetFrameRate().AsFrameTime(InLength);
+	PRAGMA_DISABLE_DEPRECATION_WARNINGS
+	Controller->SetNumberOfFrames(LengthInFrameTime.RoundToFrame());
+	PRAGMA_ENABLE_DEPRECATION_WARNINGS
 #else
 	PRAGMA_DISABLE_DEPRECATION_WARNINGS
-	SetSequenceLength(InLength);
+	SequenceLength = InLength;
 	PRAGMA_ENABLE_DEPRECATION_WARNINGS
 #endif	
 }
@@ -1324,11 +1327,11 @@ void UAnimMontage::SetCompositeLength(float InLength)
 /////////////////////////////////////////////////////////////////////////////////////////////
 
 FAnimMontageInstance::FAnimMontageInstance()
-	: Montage(NULL)
+	: Montage(nullptr)
 	, bPlaying(false)
 	, DefaultBlendTimeMultiplier(1.0f)
 	, bDidUseMarkerSyncThisTick(false)
-	, AnimInstance(NULL)
+	, AnimInstance(nullptr)
 	, InstanceID(INDEX_NONE)
 	, Position(0.f)
 	, PlayRate(1.f)
@@ -1340,13 +1343,13 @@ FAnimMontageInstance::FAnimMontageInstance()
 	, ActiveBlendProfile(nullptr)
 	, ActiveBlendProfileMode(EBlendProfileMode::TimeFactor)
 	, DisableRootMotionCount(0)
-	, MontageSyncLeader(NULL)
+	, MontageSyncLeader(nullptr)
 	, MontageSyncUpdateFrameCounter(INDEX_NONE)
 {
 }
 
 FAnimMontageInstance::FAnimMontageInstance(UAnimInstance * InAnimInstance)
-	: Montage(NULL)
+	: Montage(nullptr)
 	, bPlaying(false)
 	, DefaultBlendTimeMultiplier(1.0f)
 	, bDidUseMarkerSyncThisTick(false)
@@ -1363,7 +1366,7 @@ FAnimMontageInstance::FAnimMontageInstance(UAnimInstance * InAnimInstance)
 	, ActiveBlendProfile(nullptr)
 	, ActiveBlendProfileMode(EBlendProfileMode::TimeFactor)
 	, DisableRootMotionCount(0)
-	, MontageSyncLeader(NULL)
+	, MontageSyncLeader(nullptr)
 	, MontageSyncUpdateFrameCounter(INDEX_NONE)
 {
 }
@@ -1589,7 +1592,7 @@ void FAnimMontageInstance::Terminate()
 {
 	SCOPE_CYCLE_COUNTER(STAT_AnimMontageInstance_Terminate);
 
-	if (Montage == NULL)
+	if (Montage == nullptr)
 	{
 		return;
 	}
@@ -1628,7 +1631,7 @@ void FAnimMontageInstance::Terminate()
 	}
 
 	// clear Blend curve
-	Blend.SetCustomCurve(NULL);
+	Blend.SetCustomCurve(nullptr);
 	Blend.SetBlendOption(EAlphaBlendOption::Linear);
 
 	ActiveBlendProfile = nullptr;
@@ -1774,7 +1777,7 @@ void FAnimMontageInstance::MontageSync_StopLeading()
 		if (MontageSyncFollower)
 		{
 			ensure(MontageSyncFollower->MontageSyncLeader == this);
-			MontageSyncFollower->MontageSyncLeader = NULL;
+			MontageSyncFollower->MontageSyncLeader = nullptr;
 		}
 	}
 	MontageSyncFollowers.Empty();
@@ -1785,7 +1788,7 @@ void FAnimMontageInstance::MontageSync_StopFollowing()
 	if (MontageSyncLeader)
 	{
 		MontageSyncLeader->MontageSyncFollowers.RemoveSingleSwap(this);
-		MontageSyncLeader = NULL;
+		MontageSyncLeader = nullptr;
 	}
 }
 
@@ -2658,7 +2661,7 @@ void FAnimMontageInstance::BranchingPointEventHandler(const FBranchingPointMarke
 		// Must grab a reference on the stack in case "this" is deleted during iteration
 		TWeakObjectPtr<UAnimInstance> AnimInstanceLocal = AnimInstance;
 
-		FAnimNotifyEvent* NotifyEvent = (BranchingPointMarker->NotifyIndex < Montage->Notifies.Num()) ? &Montage->Notifies[BranchingPointMarker->NotifyIndex] : NULL;
+		FAnimNotifyEvent* NotifyEvent = (BranchingPointMarker->NotifyIndex < Montage->Notifies.Num()) ? &Montage->Notifies[BranchingPointMarker->NotifyIndex] : nullptr;
 		if (NotifyEvent)
 		{
 			// Handle backwards compatibility with older BranchingPoints.
@@ -2670,7 +2673,7 @@ void FAnimMontageInstance::BranchingPointEventHandler(const FBranchingPointMarke
 				UFunction* Function = AnimInstance.Get()->FindFunction(FuncFName);
 				if (Function)
 				{
-					AnimInstance.Get()->ProcessEvent(Function, NULL);
+					AnimInstance.Get()->ProcessEvent(Function, nullptr);
 				}
 				// In case older BranchingPoint has been re-implemented as a new Custom Notify, this is if BranchingPoint function hasn't been found.
 				else
@@ -2678,7 +2681,7 @@ void FAnimMontageInstance::BranchingPointEventHandler(const FBranchingPointMarke
 					AnimInstance.Get()->TriggerSingleAnimNotify(NotifyEvent);
 				}
 			}
-			else if (NotifyEvent->NotifyStateClass != NULL)
+			else if (NotifyEvent->NotifyStateClass != nullptr)
 			{
 				if (BranchingPointMarker->NotifyEventType == EAnimNotifyEventType::Begin)
 				{
@@ -2807,7 +2810,7 @@ UAnimMontage* FAnimMontageInstance::PreviewSequencerMontagePosition(FName SlotNa
 			// we have to get it again in case if this is new
 			MontageInstanceToUpdate = AnimInst->GetMontageInstanceForID(InOutInstanceId);
 			// since we don't advance montage in the tick, we manually have to handle notifies
-			MontageInstanceToUpdate->HandleEvents(InFromPosition, InToPosition, NULL);
+			MontageInstanceToUpdate->HandleEvents(InFromPosition, InToPosition, nullptr);
 			if (!bFireNotifies)
 			{
 				AnimInst->NotifyQueue.Reset(SkeletalMeshComponent);
@@ -2926,6 +2929,15 @@ void UAnimMontage::BakeTimeStretchCurve()
 	TimeStretchCurve.BakeFromFloatCurve(*TimeStretchFloatCurve, SequenceLength);
 	PRAGMA_ENABLE_DEPRECATION_WARNINGS
 }
+
+void UAnimMontage::PopulateWithExistingModel(TScriptInterface<IAnimationDataModel> ExistingDataModel)
+{
+	Super::PopulateWithExistingModel(ExistingDataModel);
+	
+	// Set composite length while model is being populated
+	const float CurrentCalculatedLength = CalculateSequenceLength();
+	SetCompositeLength(CurrentCalculatedLength);
+}
 #endif // WITH_EDITOR
 
 FMontageBlendSettings::FMontageBlendSettings()
diff --git a/Engine/Source/Runtime/Engine/Private/Animation/AnimNode_SequencePlayer.cpp b/Engine/Source/Runtime/Engine/Private/Animation/AnimNode_SequencePlayer.cpp
--- a/Engine/Source/Runtime/Engine/Private/Animation/AnimNode_SequencePlayer.cpp
+++ b/Engine/Source/Runtime/Engine/Private/Animation/AnimNode_SequencePlayer.cpp
@@ -128,7 +128,7 @@ void FAnimNode_SequencePlayerBase::Evaluate_AnyThread(FPoseContext& Output)
 		}
 
 		FAnimationPoseData AnimationPoseData(Output);
-		CurrentSequence->GetAnimationPose(AnimationPoseData, FAnimExtractContext(InternalTimeAccumulator, Output.AnimInstanceProxy->ShouldExtractRootMotion(), DeltaTimeRecord, GetLoopAnimation()));
+		CurrentSequence->GetAnimationPose(AnimationPoseData, FAnimExtractContext(static_cast<double>(InternalTimeAccumulator), Output.AnimInstanceProxy->ShouldExtractRootMotion(), DeltaTimeRecord, GetLoopAnimation()));
 	}
 	else
 	{
@@ -141,7 +141,7 @@ void FAnimNode_SequencePlayerBase::GatherDebugData(FNodeDebugData& DebugData)
 	FString DebugLine = DebugData.GetNodeName(this);
 
 	UAnimSequenceBase* CurrentSequence = GetSequence();
-	DebugLine += FString::Printf(TEXT("('%s' Play Time: %.3f)"), CurrentSequence ? *CurrentSequence->GetName() : TEXT("NULL"), InternalTimeAccumulator);
+	DebugLine += FString::Printf(TEXT("('%s' Play Time: %.3f)"), CurrentSequence ? *CurrentSequence->GetName() : TEXT("nullptr"), InternalTimeAccumulator);
 	DebugData.AddDebugItem(DebugLine, true);
 }
 
