**New**

diff --git a/Engine/Source/Runtime/Engine/Private/Animation/AnimSequence.cpp b/Engine/Source/Runtime/Engine/Private/Animation/AnimSequence.cpp
--- a/Engine/Source/Runtime/Engine/Private/Animation/AnimSequence.cpp
+++ b/Engine/Source/Runtime/Engine/Private/Animation/AnimSequence.cpp
@@ -52,8 +52,9 @@
 #include "Animation/AttributesRuntime.h"
 #include "Stats/StatsHierarchical.h"
 #include "Animation/AnimationPoseData.h"
-#include "ITimeManagementModule.h"
 #include "CommonFrameRates.h"
+#include "ITimeManagementModule.h"
+#include "Algo/Accumulate.h"
 #include "HAL/LowLevelMemTracker.h"
 
 LLM_DEFINE_TAG(SequenceData);
@@ -61,7 +62,7 @@ LLM_DEFINE_TAG(SequenceData);
 #include "HAL/FileManager.h"
 
 #if WITH_EDITOR
-#include "Animation/AnimData/AnimDataModel.h"
+#include "Animation/AnimData/IAnimationDataModel.h"
 #include "Animation/BuiltInAttributeTypes.h"
 #endif // WITH_EDITOR
 
@@ -271,7 +272,7 @@ FString GetAnimSequenceSpecificCacheKeySuffix(const UAnimSequence& Seq, bool bPe
 	FString Ret = FString::Printf(TEXT("%i_%s%s%s_%c%c%i_%s_%s_%i"),
 		Seq.CompressCommandletVersion,
 		PRAGMA_DISABLE_DEPRECATION_WARNINGS
-		*Seq.GetRawDataGuid().ToString(),
+		*Seq.GetDataModel()->GenerateGuid().ToString(),
 		PRAGMA_ENABLE_DEPRECATION_WARNINGS
 		*Seq.GetSkeleton()->GetGuid().ToString(),
 		*Seq.GetSkeleton()->GetVirtualBoneGuid().ToString(),
@@ -279,7 +280,7 @@ FString GetAnimSequenceSpecificCacheKeySuffix(const UAnimSequence& Seq, bool bPe
 		RefType,
 		Seq.RefFrameIndex,
 		PRAGMA_DISABLE_DEPRECATION_WARNINGS
-		(bIsValidAdditive && Seq.RefPoseSeq) ? *Seq.RefPoseSeq->GetRawDataGuid().ToString() : TEXT("NoAdditiveGuid"),
+		(bIsValidAdditive && Seq.RefPoseSeq) ? *Seq.RefPoseSeq->GetDataModel()->GenerateGuid().ToString() : TEXT("NoAdditiveGuid"),
 		PRAGMA_ENABLE_DEPRECATION_WARNINGS
 		*ArcToHexString.MakeString(),
 		Seq.GetSamplingFrameRate().Numerator / Seq.GetSamplingFrameRate().Denominator
@@ -409,15 +410,6 @@ bool FCurveTrack::CompressCurveWeights()
 	return false;
 }
 
-/////////////////////////////////////////////////////////////
-
-// since we want this change for hot fix, I can't change header file, 
-// next time move this to the header
-float GetIntervalPerKey(int32 NumFrames, float SequenceLength) 
-{
-	return (NumFrames > 1) ? (SequenceLength / (NumFrames-1)) : MINIMUM_ANIMATION_LENGTH;
-}
-
 /////////////////////////////////////////////////////
 // UAnimSequence
 
@@ -444,7 +436,7 @@ UAnimSequence::UAnimSequence(const FObjectInitializer& ObjectInitializer)
 
 	if (!HasAnyFlags(EObjectFlags::RF_ClassDefaultObject| EObjectFlags::RF_NeedLoad))
 	{
-		check(IsDataModelValid());
+		check(DataModelInterface.GetObject());
 	}
 #endif
 }
@@ -469,12 +461,12 @@ void UAnimSequence::GetAssetRegistryTags(TArray<FAssetRegistryTag>& OutTags) con
 		OutTags.Add( FAssetRegistryTag(SourceFileTagName(), AssetImportData->GetSourceData().ToJson(), FAssetRegistryTag::TT_Hidden) );
 	}
 
-	OutTags.Add(FAssetRegistryTag(TEXT("Compression Ratio"), FString::Printf(TEXT("%.03f"), (float)GetApproxCompressedSize() / (float)GetUncompressedRawSize()), FAssetRegistryTag::TT_Numerical));
-	if (DataModel)
+	if (DataModelInterface.GetObject() && DataModelInterface->HasBeenPopulated())
 	{
-		OutTags.Add(FAssetRegistryTag(TEXT("Source Frame Rate"), FString::Printf(TEXT("%.2f"), GetDataModel()->GetFrameRate().AsDecimal()), FAssetRegistryTag::TT_Numerical));
-		OutTags.Add(FAssetRegistryTag(TEXT("Number of Frames"), FString::Printf(TEXT("%.i"), GetDataModel()->GetNumberOfFrames()), FAssetRegistryTag::TT_Numerical));
-		OutTags.Add(FAssetRegistryTag(TEXT("Number of Keys"), FString::Printf(TEXT("%i"), GetDataModel()->GetNumberOfKeys()), FAssetRegistryTag::TT_Numerical));
+		OutTags.Add(FAssetRegistryTag(TEXT("Compression Ratio"), FString::Printf(TEXT("%.03f"), (float)GetApproxCompressedSize() / (float)GetUncompressedRawSize()), FAssetRegistryTag::TT_Numerical));
+		OutTags.Add(FAssetRegistryTag(TEXT("Source Frame Rate"), FString::Printf(TEXT("%.2f"), DataModelInterface->GetFrameRate().AsDecimal()), FAssetRegistryTag::TT_Numerical));
+		OutTags.Add(FAssetRegistryTag(TEXT("Number of Frames"), FString::Printf(TEXT("%.i"), DataModelInterface->GetNumberOfFrames()), FAssetRegistryTag::TT_Numerical));
+		OutTags.Add(FAssetRegistryTag(TEXT("Number of Keys"), FString::Printf(TEXT("%i"), DataModelInterface->GetNumberOfKeys()), FAssetRegistryTag::TT_Numerical));
 	}
 	else
 	{
@@ -501,12 +493,12 @@ void UAnimSequence::AddReferencedObjects(UObject* This, FReferenceCollector& Col
 #if WITH_EDITOR
 int32 UAnimSequence::GetUncompressedRawSize() const
 {
-	if (DataModel)
+	if (DataModelInterface.GetObject())
 	{
-		const int32 BoneRawSize = FRawAnimSequenceTrack::SingleKeySize * DataModel->GetNumBoneTracks() * DataModel->GetNumberOfKeys();
+		const int32 BoneRawSize = FRawAnimSequenceTrack::SingleKeySize * DataModelInterface->GetNumBoneTracks() * DataModelInterface->GetNumberOfKeys();
 		int32 CurveRawSize = 0;
 
-		for (const FFloatCurve& Curve : DataModel->GetCurveData().FloatCurves)
+		for (const FFloatCurve& Curve : DataModelInterface->GetCurveData().FloatCurves)
 		{
 			CurveRawSize += sizeof(FFloatCurve);
 			CurveRawSize += sizeof(FRichCurveKey) * Curve.FloatCurve.Keys.Num();
@@ -518,17 +510,18 @@ int32 UAnimSequence::GetUncompressedRawSize() const
 
 int32 UAnimSequence::GetApproxRawSize() const
 {
-	if (DataModel)
+	if (ShouldDataModelBeValid())
 	{
-		int32 Total = sizeof(FRawAnimSequenceTrack) * DataModel->GetNumBoneTracks();
-		for (const FBoneAnimationTrack& AnimationTrack : DataModel->GetBoneAnimationTracks())
+		ValidateModel();
+		int32 Total = sizeof(FRawAnimSequenceTrack) * DataModelInterface->GetNumBoneTracks();
+		for (const FBoneAnimationTrack& AnimationTrack : DataModelInterface->GetBoneAnimationTracks())
 		{
-				Total += AnimationTrack.InternalTrackData.PosKeys.Num() * AnimationTrack.InternalTrackData.PosKeys.GetTypeSize();
-				Total += AnimationTrack.InternalTrackData.RotKeys.Num() * AnimationTrack.InternalTrackData.RotKeys.GetTypeSize();
-				Total += AnimationTrack.InternalTrackData.ScaleKeys.Num() * AnimationTrack.InternalTrackData.ScaleKeys.GetTypeSize();
+			Total += AnimationTrack.InternalTrackData.PosKeys.Num() * AnimationTrack.InternalTrackData.PosKeys.GetTypeSize();
+			Total += AnimationTrack.InternalTrackData.RotKeys.Num() * AnimationTrack.InternalTrackData.RotKeys.GetTypeSize();
+			Total += AnimationTrack.InternalTrackData.ScaleKeys.Num() * AnimationTrack.InternalTrackData.ScaleKeys.GetTypeSize();
 		}
-	
-			for (const FFloatCurve& Curve : DataModel->GetCurveData().FloatCurves)
+
+		for (const FFloatCurve& Curve : DataModelInterface->GetCurveData().FloatCurves)
 		{
 			Total += sizeof(FFloatCurve);
 			Total += sizeof(FRichCurveKey) * Curve.FloatCurve.Keys.Num();
@@ -629,7 +622,6 @@ void UAnimSequence::Serialize(FArchive& Ar)
 						// Set RawAnimationData to Source
 						RawAnimationData = SourceRawAnimationData_DEPRECATED;
 						SourceRawAnimationData_DEPRECATED.Empty();
-						MarkRawDataAsModified();
 					}
 					PRAGMA_ENABLE_DEPRECATION_WARNINGS
 				}
@@ -862,20 +854,6 @@ void UAnimSequence::PostLoad()
 	{
 		RefPoseSeq->ConditionalPostLoad();
 	}
-
-	EnsureValidRawDataGuid();
-
-	// I have to do this first thing in here
-	// so that remove all NaNs before even being read
-	if(GetLinkerUEVersion() < VER_UE4_ANIMATION_REMOVE_NANS)
-	{
-		RemoveNaNTracks();
-	}
-
-	if (ShouldDataModelBeValid())
-	{
-	Controller->RemoveBoneTracksMissingFromSkeleton(GetSkeleton());
-	}
 #endif // WITH_EDITOR
 
 	// if valid additive, but if base additive isn't 
@@ -883,7 +861,7 @@ void UAnimSequence::PostLoad()
 	// but this causes issue since once this happens this is unrecoverable until you delete from outside of editor
 	if (IsValidAdditive())
 	{
-		if (RefPoseSeq && !GetSkeleton()->IsCompatible(RefPoseSeq->GetSkeleton()))
+		if (RefPoseSeq && GetSkeleton() && !GetSkeleton()->IsCompatible(RefPoseSeq->GetSkeleton()))
 		{
 			// if this happens, there was a issue with retargeting
 			UE_LOG(LogAnimation, Warning, TEXT("Animation %s - Invalid additive animation base animation (%s)"), *GetName(), *RefPoseSeq->GetName());
@@ -891,27 +869,6 @@ void UAnimSequence::PostLoad()
 		}
 	}
 
-#if WITH_EDITOR
-	// Generate transient animation data, used to fallback on whenever the animation modified and not yet compressed (what previously used to be the raw data)
-	ResampleAnimationTrackData();
-
-	static bool ForcedRecompressionSetting = FAnimationUtils::GetForcedRecompressionSetting();
-
-	if (ForcedRecompressionSetting)
-	{
-		//Force recompression
-		PRAGMA_DISABLE_DEPRECATION_WARNINGS
-		RawDataGuid = FGuid::NewGuid();
-		PRAGMA_ENABLE_DEPRECATION_WARNINGS
-		bUseRawDataOnly = true;
-	}
-
-	if (bUseRawDataOnly)
-	{
-		RequestAnimCompression(FRequestAnimCompressionParams(true, false, false));
-	}
-#endif
-
 	// Ensure notifies are sorted.
 	SortNotifies();
 
@@ -921,48 +878,15 @@ void UAnimSequence::PostLoad()
 	if (GetOutermost()->HasAnyPackageFlags(PKG_Cooked | PKG_FilterEditorOnly))
 	{
 		// Update the virtual bone GUID according to the skeleton
-		if (GetSkeleton())
+		if (const USkeleton* CurrentSkeleton = GetSkeleton())
 		{
-			SetSkeletonVirtualBoneGuid(GetSkeleton()->GetVirtualBoneGuid());
-		}
-	}
-	else
-	{
-#if WITH_EDITOR
-		if (IsDataModelValid() && DataModel->GetNumberOfKeys() == 0 && DataModel->GetNumberOfFloatCurves() == 0 )
-		{
-			UE_LOG(LogAnimation, Warning, TEXT("No animation data exists for sequence %s (%s)"), *GetName(), (GetOuter() ? *GetOuter()->GetFullName() : *GetFullName()) );
+			SetSkeletonVirtualBoneGuid(CurrentSkeleton->GetVirtualBoneGuid());
 
-			if (!IsRunningGame())
+			if(!IsCompressedDataValid())
 			{
-				static FName NAME_LoadErrors("LoadErrors");
-				FMessageLog LoadErrors(NAME_LoadErrors);
-
-				TSharedRef<FTokenizedMessage> Message = LoadErrors.Warning();
-				Message->AddToken(FTextToken::Create(LOCTEXT("EmptyAnimationData1", "The Animation ")));
-				Message->AddToken(FAssetNameToken::Create(GetPathName(), FText::FromString(GetName())));
-				Message->AddToken(FTextToken::Create(LOCTEXT("EmptyAnimationData2", " has no animation data. Recommend to remove.")));
-				LoadErrors.Notify();
+				UE_LOG(LogAnimation, Fatal, TEXT("No animation compression exists for sequence %s (%s)"), *GetName(), (GetOuter() ? *GetOuter()->GetFullName() : *GetFullName()) );
 			}
 		}
-		// @remove temp hack for fixing length
-		// @todo need to fix importer/editing feature
-		else if (IsDataModelValid() && GetPlayLength() == 0.f)
-		{
-			ensure(DataModel->GetNumberOfKeys() == 1);
-			Controller->SetPlayLength(MINIMUM_ANIMATION_LENGTH);
-		}
-		// Raw data exists, but missing compress animation data
-		else
-#endif
-		if( GetSkeleton() && !IsCompressedDataValid()
-#if WITH_EDITOR
-			&& !bCompressionInProgress
-#endif
-			)
-		{
-			UE_LOG(LogAnimation, Fatal, TEXT("No animation compression exists for sequence %s (%s)"), *GetName(), (GetOuter() ? *GetOuter()->GetFullName() : *GetFullName()) );
-		}
 	}
 
 	for(FAnimNotifyEvent& Notify : Notifies)
@@ -990,26 +914,12 @@ void UAnimSequence::PostLoad()
 			CurrentSkeleton->VerifySmartName(USkeleton::AnimCurveMappingName, CurveName);
 		}
 #if WITH_EDITOR
-		if (IsDataModelValid())
-		{
-		static const bool DoNotTransactAction = false;
-		Controller->FindOrAddCurveNamesOnSkeleton(CurrentSkeleton, ERawCurveTrackTypes::RCT_Transform, DoNotTransactAction);
-		}
-
 		for (const FAnimSyncMarker& SyncMarker : AuthoredSyncMarkers)
 		{
 			CurrentSkeleton->RegisterMarkerName(SyncMarker.MarkerName);
 		}
 #endif
 	}
-
-#if WITH_EDITOR
-	if (IsDataModelValid() && GetLinkerCustomVersion(FUE5MainStreamObjectVersion::GUID) < FUE5MainStreamObjectVersion::MoveCustomAttributesToDataModel)
-	{
-		ValidateModel();
-		MoveAttributesToModel();
-	}
-#endif // WITH_EDITOR
 }
 
 #if WITH_EDITOR
@@ -1068,11 +978,11 @@ void UAnimSequence::PostEditChangeProperty(FPropertyChangedEvent& PropertyChange
 		}
 	}
 
-	if (RefPoseSeq != NULL)
+	if (RefPoseSeq != nullptr)
 	{
 		if (RefPoseSeq->GetSkeleton() != GetSkeleton()) // @todo this may require to be changed when hierarchy of skeletons is introduced
 		{
-			RefPoseSeq = NULL;
+			RefPoseSeq = nullptr;
 		}
 	}
 
@@ -1125,12 +1035,8 @@ void UAnimSequence::PostEditChangeProperty(FPropertyChangedEvent& PropertyChange
 
 	// @Todo fix me: This is temporary fix to make sure they always have compressed data
 	ValidateModel();
-	if (IsDataModelValid() && DataModel->GetNumBoneTracks() && bNeedPostProcess)
+	if (IsDataModelValid() && DataModelInterface->GetNumBoneTracks() && bNeedPostProcess)
 	{
-		// QQ need new function here to queue async compression
-		PRAGMA_DISABLE_DEPRECATION_WARNINGS
-		MarkRawDataAsModified(false);
-		PRAGMA_ENABLE_DEPRECATION_WARNINGS
 		RecompressAnimationData();
 	}
 }
@@ -1156,15 +1062,33 @@ static int32 FindKeyIndex(float Time, const TimeArray& Times)
 	return FoundIndex;
 }
 
+#if WITH_EDITOR
+int32 GForceRawData = 0;
+static FAutoConsoleVariableRef CVarForceRawData(
+	TEXT("a.ForceEvalRawData"),
+	GForceRawData,
+	TEXT("Values: 0/1\n")
+	TEXT("Controls whether or not to forcefully sample non-compressed anim data."),
+	ECVF_Cheat);
+#endif // WITH_EDITOR
+
 void UAnimSequence::GetBoneTransform(FTransform& OutAtom, int32 TrackIndex, float Time, bool bUseRawData) const
+{
+	GetBoneTransform(OutAtom, TrackIndex, (double)Time, bUseRawData);
+}
+
+void UAnimSequence::GetBoneTransform(FTransform& OutAtom, int32 TrackIndex, double Time, bool bUseRawData) const
 {
 	// If the caller didn't request that raw animation data be used . . .
-	const bool bEvaluateCompressedData = (!bUseRawData || !CanEvaluateRawAnimationData());
+	const bool bEvaluateCompressedData =
+#if WITH_EDITOR 
+		!GForceRawData && 
+#endif 
+		(!bUseRawData || !CanEvaluateRawAnimationData());
+
 	if ( bEvaluateCompressedData && IsCompressedDataValid())
 	{
-		PRAGMA_DISABLE_DEPRECATION_WARNINGS
-		FAnimSequenceDecompressionContext DecompContext(SequenceLength, Interpolation, GetFName(), *CompressedData.CompressedDataStructure);
-		PRAGMA_ENABLE_DEPRECATION_WARNINGS
+		FAnimSequenceDecompressionContext DecompContext(TargetFrameRate, TargetFrameRate.AsFrameNumber(GetPlayLength()).Value, Interpolation, GetRetargetTransformsSourceName(), *CompressedData.CompressedDataStructure, GetSkeleton(), IsValidAdditive());
 		DecompContext.Seek(Time);
 		if (CompressedData.BoneCompressionCodec)
 		{
@@ -1176,7 +1100,7 @@ void UAnimSequence::GetBoneTransform(FTransform& OutAtom, int32 TrackIndex, floa
 	{
 #if WITH_EDITOR
 		ValidateModel();
-		UE::Anim::GetBoneTransformFromModel(DataModel, OutAtom, TrackIndex, Time, Interpolation);
+		UE::Anim::GetBoneTransformFromModel(DataModelInterface.GetInterface(), OutAtom, TrackIndex, Time, Interpolation);
 #endif
 	}
 }
@@ -1197,7 +1121,7 @@ void UAnimSequence::GetBoneTransform(FTransform& OutAtom, int32 TrackIndex, FAni
 	{
 #if WITH_EDITOR
 		ValidateModel();
-		UE::Anim::GetBoneTransformFromModel(DataModel, OutAtom, TrackIndex, DecompContext.GetEvaluationTime(), Interpolation);
+		UE::Anim::GetBoneTransformFromModel(DataModelInterface.GetInterface(), OutAtom, TrackIndex, DecompContext.GetEvaluationTime(), Interpolation);
 #endif
 	}
 }
@@ -1221,9 +1145,9 @@ FTransform UAnimSequence::ExtractRootTrackTransform(float Pos, const FBoneContai
 		if (CanEvaluateRawAnimationData() && bUseRawDataOnly)
 		{
 			ValidateModel();
-			if (const FBoneAnimationTrack* RootTrack = DataModel->FindBoneTrackByIndex(RootBoneIndex))
+			if (const FBoneAnimationTrack* RootTrack = DataModelInterface->FindBoneTrackByIndex(RootBoneIndex))
 			{
-				return DataModel->GetBoneTrackIndex(*RootTrack);
+				return DataModelInterface->GetBoneTrackIndex(*RootTrack);
 			}
 		}
 		else
@@ -1244,7 +1168,7 @@ FTransform UAnimSequence::ExtractRootTrackTransform(float Pos, const FBoneContai
 	{
 		// if we do have root data, then return root data
 		FTransform RootTransform;
-		GetBoneTransform(RootTransform, TrackIndex, Pos, bUseRawDataOnly);
+		GetBoneTransform(RootTransform, TrackIndex, static_cast<double>(Pos), bUseRawDataOnly);
 		return RootTransform;
 	}
 
@@ -1299,7 +1223,7 @@ FTransform UAnimSequence::ExtractRootMotion(float StartTime, float DeltaTime, bo
 			// If we've hit the end of the animation, and we're allowed to loop, keep going.
 			if ((AdvanceType == ETAA_Finished) && bAllowLooping)
 			{
-				const float ActualDeltaMove = (CurrentPosition - PreviousPosition);
+				const double ActualDeltaMove = (CurrentPosition - PreviousPosition);
 				DesiredDeltaMove -= ActualDeltaMove;
 
 				PreviousPosition = bPlayingBackwards ? GetPlayLength() : 0.f;
@@ -1329,8 +1253,8 @@ FTransform UAnimSequence::ExtractRootMotionFromRange(float StartTrackPosition, f
 		}
 	}
 
-	FTransform StartTransform = ExtractRootTrackTransform(StartTrackPosition, NULL);
-	FTransform EndTransform = ExtractRootTrackTransform(EndTrackPosition, NULL);
+	FTransform StartTransform = ExtractRootTrackTransform(StartTrackPosition, nullptr);
+	FTransform EndTransform = ExtractRootTrackTransform(EndTrackPosition, nullptr);
 
 	// Use old calculation if needed.
 	if (bUseNormalizedRootMotionScale)
@@ -1464,7 +1388,11 @@ void UAnimSequence::GetBonePose(FAnimationPoseData& OutAnimationPoseData, const
 
 	FBoneContainer& RequiredBones = OutPose.GetBoneContainer();
 	check(!bForceUseRawData || CanEvaluateRawAnimationData());
-	const bool bUseRawDataForPoseExtraction = (CanEvaluateRawAnimationData() && bForceUseRawData) || UseRawDataForPoseExtraction(RequiredBones);
+	const bool bUseRawDataForPoseExtraction = (CanEvaluateRawAnimationData() && bForceUseRawData) ||
+#if WITH_EDITOR
+ 		GForceRawData == 1 || 
+#endif // WITH_EDITOR
+		UseRawDataForPoseExtraction(RequiredBones);
 
 	const bool bIsBakedAdditive = !bUseRawDataForPoseExtraction && IsValidAdditive();
 
@@ -1506,7 +1434,7 @@ void UAnimSequence::GetBonePose(FAnimationPoseData& OutAnimationPoseData, const
 				// Virtual bones are part of the retarget transform pose, so if the pose has not been updated (recently) there might be a mismatch
 				if (SkeletonBoneIndex < NumRawBones || AuthoredOnRefSkeleton.IsValidIndex(SkeletonBoneIndex.GetInt()))
 				{
-					OutPose[PoseBoneIndex] = AuthoredOnRefSkeleton[SkeletonBoneIndex.GetInt()];					
+					OutPose[PoseBoneIndex] = AuthoredOnRefSkeleton[SkeletonBoneIndex.GetInt()];
 				}
 			}
 		}
@@ -1516,35 +1444,22 @@ void UAnimSequence::GetBonePose(FAnimationPoseData& OutAnimationPoseData, const
 		}
 	}
 
-	// Evaluate the (remapped if needed) curves.
-	// Scoped so that the RemappedCurve destructs and isn't kept alive longer than needed.
-	{
-		FSkeletonRemappingCurve RemappedCurve(OutAnimationPoseData.GetCurve(), RequiredBones, GetSkeleton());
-		EvaluateCurveData(RemappedCurve.GetCurve(), ExtractionContext.CurrentTime, bUseRawDataForPoseExtraction);
-	}
-
 #if WITH_EDITOR
-	const int32 NumTracks = bUseRawDataForPoseExtraction ? DataModel->GetNumBoneTracks() : CompressedData.CompressedTrackToSkeletonMapTable.Num();
+	const int32 NumTracks = bUseRawDataForPoseExtraction ? DataModelInterface->GetNumBoneTracks() : CompressedData.CompressedTrackToSkeletonMapTable.Num();
 #else
 	const int32 NumTracks = CompressedData.CompressedTrackToSkeletonMapTable.Num();
 #endif 
-	if (NumTracks == 0)
-	{
-		return;
-	}
-
 	const bool bTreatAnimAsAdditive = (IsValidAdditive() && !bUseRawDataForPoseExtraction); // Raw data is never additive
-	FRootMotionReset RootMotionReset(bEnableRootMotion, RootMotionRootLock, bForceRootLock, ExtractRootTrackTransform(0.f, &RequiredBones), bTreatAnimAsAdditive);
+	const FRootMotionReset RootMotionReset(bEnableRootMotion, RootMotionRootLock, bForceRootLock, ExtractRootTrackTransform(0.0f, &RequiredBones), bTreatAnimAsAdditive);
 
 #if WITH_EDITOR
-	// this happens only with editor data
-	// Slower path for disable retargeting, that's only used in editor and for debugging.
+	// Evaluate raw (source) curve and bone data
 	if (bUseRawDataForPoseExtraction)
 	{
 		// Warning if we have invalid data
 		for (int32 TrackIndex = 0; TrackIndex < NumTracks; TrackIndex++)
 		{
-			const FBoneAnimationTrack& BoneAnimationTrack = DataModel->GetBoneTrackByIndex(TrackIndex);
+			const FBoneAnimationTrack& BoneAnimationTrack = DataModelInterface->GetBoneTrackByIndex(TrackIndex);
 			const FRawAnimSequenceTrack& TrackToExtract = BoneAnimationTrack.InternalTrackData;
 
 			// Bail out (with rather wacky data) if data is empty for some reason.
@@ -1555,38 +1470,47 @@ void UAnimSequence::GetBonePose(FAnimationPoseData& OutAnimationPoseData, const
 			}
 		}
 
-		UE::Anim::BuildPoseFromModel(DataModel, OutPose, ExtractionContext.CurrentTime, Interpolation, GetRetargetTransformsSourceName(), GetRetargetTransforms());
+		{
+			FSkeletonRemappingCurve RemappedCurve(OutAnimationPoseData.GetCurve(), RequiredBones, GetSkeleton());
+			FAnimationPoseData RemappedPoseData(OutAnimationPoseData.GetPose(), RemappedCurve.GetCurve(), OutAnimationPoseData.GetAttributes());
+
+			const UE::Anim::DataModel::FEvaluationContext EvaluationContext(ExtractionContext.CurrentTime, DataModelInterface->GetFrameRate(), GetRetargetTransformsSourceName(), GetRetargetTransforms(), Interpolation);
+			DataModelInterface->Evaluate(RemappedPoseData, EvaluationContext);
+		}
 
 		if ((ExtractionContext.bExtractRootMotion && RootMotionReset.bEnableRootMotion) || RootMotionReset.bForceRootLock)
 		{
 			RootMotionReset.ResetRootBoneForRootMotion(OutPose[FCompactPoseBoneIndex(0)], RequiredBones);
 		}
 
-		EvaluateAttributes(OutAnimationPoseData, ExtractionContext, true);
-
 		return;
 	}
+	else
 #endif // WITH_EDITOR
+	// Only try and evaluate compressed bone data if the animation contains any bone tracks
+	if (NumTracks != 0)
+	{
+		// Evaluate compressed bone data
+		FAnimSequenceDecompressionContext DecompContext(TargetFrameRate, TargetFrameRate.AsFrameTime(GetPlayLength()).RoundToFrame().Value, Interpolation, GetRetargetTransformsSourceName(), *CompressedData.CompressedDataStructure, GetSkeleton(), IsValidAdditive());
+		UE::Anim::Decompression::DecompressPose(OutPose, CompressedData, ExtractionContext, DecompContext, GetRetargetTransforms(), RootMotionReset);
+	}
 
-	DecompressPose(OutPose, CompressedData, ExtractionContext, GetSkeleton(), GetPlayLength(), Interpolation, bIsBakedAdditive, GetRetargetTransforms(), GetRetargetTransformsSourceName(), RootMotionReset);
-
-	EvaluateAttributes(OutAnimationPoseData, ExtractionContext, false);
-}
-
-
-#if WITH_EDITOR
-void UAnimSequence::UpdateCompressedCurveName(SmartName::UID_Type CurveUID, const struct FSmartName& NewCurveName)
-{
-	for (FSmartName& CurveName : CompressedData.CompressedCurveNames)
+	// (Always) evaluate compressed curve data
 	{
-		if (CurveName.UID == CurveUID)
+		// Scoped so that the RemappedCurve destructs and isn't kept alive longer than needed.
+		FSkeletonRemappingCurve RemappedCurve(OutAnimationPoseData.GetCurve(), RequiredBones, GetSkeleton());
+#if WITH_EDITOR
+		// When evaluating from raw animation data, UE::Anim::BuildPoseFromModel will populate the curve data
+		if (!bUseRawDataForPoseExtraction)
+#endif // WITH_EDITOR
 		{
-			CurveName = NewCurveName;
-			break;
+			EvaluateCurveData(RemappedCurve.GetCurve(), ExtractionContext.CurrentTime, bUseRawDataForPoseExtraction);
 		}
 	}
+
+	// Evaluate animation attributes (no compressed format yet)
+	EvaluateAttributes(OutAnimationPoseData, ExtractionContext, false);
 }
-#endif
 
 void UAnimSequence::GetBonePose_Additive(FCompactPose& OutPose, FBlendedCurve& OutCurve, const FAnimExtractContext& ExtractionContext) const
 {
@@ -1638,10 +1562,10 @@ void UAnimSequence::GetAdditiveBasePose(FCompactPose& OutPose, FBlendedCurve& Ou
 	GetAdditiveBasePose(OutAnimationPoseData, ExtractionContext);
 }
 
-static void GetSequencePose(FAnimationPoseData& OutAnimationPoseData, const FAnimExtractContext& ExtractionContext, const UAnimSequence &Seq, int FrameIndex)
+static void GetSequencePose(FAnimationPoseData& OutAnimationPoseData, const FAnimExtractContext& ExtractionContext, const UAnimSequence &Seq, int32 FrameIndex)
 {
-	const float Fraction = (Seq.GetNumberOfSampledKeys() > 0) ? FMath::Clamp<float>((float)FrameIndex / (float)Seq.GetNumberOfSampledKeys(), 0.f, 1.f) : 0.f;
-	const float BasePoseTime = Seq.GetPlayLength() * Fraction;
+	const double Fraction = (Seq.GetNumberOfSampledKeys() > 0) ? FMath::Clamp<double>((double)FrameIndex / (double)Seq.GetNumberOfSampledKeys(), 0.0, 1.0) : 0.0;
+	const double BasePoseTime = Seq.GetPlayLength() * Fraction;
 	FAnimExtractContext BasePoseExtractionContext(ExtractionContext);
 	BasePoseExtractionContext.CurrentTime = BasePoseTime;
 	Seq.GetBonePose(OutAnimationPoseData, BasePoseExtractionContext, true);
@@ -1656,8 +1580,8 @@ void UAnimSequence::GetAdditiveBasePose(FAnimationPoseData& OutAnimationPoseData
 		{
 			// normalize time to fit base seq
 			PRAGMA_DISABLE_DEPRECATION_WARNINGS
-			const float Fraction = (SequenceLength > 0.f)? FMath::Clamp<float>(ExtractionContext.CurrentTime / SequenceLength, 0.f, 1.f) : 0.f;
-			const float BasePoseTime = RefPoseSeq->SequenceLength * Fraction;
+			const double Fraction = (GetPlayLength() > 0.0)? FMath::Clamp<double>(ExtractionContext.CurrentTime / (double)GetPlayLength(), 0.0, 1.0) : 0.0;
+			const double BasePoseTime = (double)RefPoseSeq->GetPlayLength() * Fraction;
 			PRAGMA_ENABLE_DEPRECATION_WARNINGS
 
 			FAnimExtractContext BasePoseExtractionContext(ExtractionContext);
@@ -1741,6 +1665,20 @@ bool UAnimSequence::CanEvaluateRawAnimationData() const
 	return false;
 }
 
+#if WITH_EDITOR
+void UAnimSequence::UpdateCompressedCurveName(SmartName::UID_Type CurveUID, const struct FSmartName& NewCurveName)
+{
+	for (FSmartName& CurveName : CompressedData.CompressedCurveNames)
+	{
+		if (CurveName.UID == CurveUID)
+		{
+			CurveName = NewCurveName;
+			break;
+		}
+	}
+}
+#endif // WITH_EDITOR
+
 #if WITH_EDITORONLY_DATA
 void UAnimSequence::UpdateRetargetSourceAsset()
 {
@@ -1828,9 +1766,18 @@ void UAnimSequence::WaitOnExistingCompression(const bool bWantResults)
 		COOK_STAT(Timer.TrackCyclesOnly()); // Need to get hit/miss and size from WaitOnExistingCompression!
 	}
 }
+#endif
+
+#if WITH_EDITOR
+int32 GSkipDDC = 0;
+static FAutoConsoleVariableRef CVarSkipDDC(
+	TEXT("a.SkipDDC"),
+	GSkipDDC,
+	TEXT("1 = Skip DDC during compression. 0 = Include DDC results during compresion "));
 
 void UAnimSequence::RequestAnimCompression(FRequestAnimCompressionParams Params)
 {
+	QUICK_SCOPE_CYCLE_COUNTER(STAT_RequestAnimCompression);
 	check(IsInGameThread());
 	USkeleton* CurrentSkeleton = GetSkeleton();
 	if (CurrentSkeleton == nullptr)
@@ -1865,6 +1812,11 @@ void UAnimSequence::RequestAnimCompression(FRequestAnimCompressionParams Params)
 
 	EnsureValidRawDataGuid();
 
+	if(ResampledAnimationTrackData.Num() == 0)
+	{		
+		ResampleAnimationTrackData();
+	}
+
 	bUseRawDataOnly = true;
 
 	check(!bCompressionInProgress);
@@ -1886,7 +1838,7 @@ void UAnimSequence::RequestAnimCompression(FRequestAnimCompressionParams Params)
 		FDerivedDataAnimationCompression* AnimCompressor = new FDerivedDataAnimationCompression(TEXT("AnimSeq"), AssetDDCKey, Params.CompressContext);
 
 		// For debugging DDC/Compression issues
-		const bool bSkipDDC = false;
+		const bool bSkipDDC = GSkipDDC == 1;
 
 		// Data does not exist, need to build it.
 		FCompressibleAnimRef CompressibleData = MakeShared<FCompressibleAnimData, ESPMode::ThreadSafe>(this, bPerformStripping);
@@ -1938,7 +1890,6 @@ void UAnimSequence::RequestAnimCompression(FRequestAnimCompressionParams Params)
 			FAnimationUtils::BuildSkeletonMetaData(GetSkeleton(), BoneData);
 			Params.CompressContext->GatherPostCompressionStats(CompressedData, BoneData, GetFName(), CompressionTime, false);
 		}
-
 	}
 }
 
@@ -1981,11 +1932,51 @@ void UAnimSequence::SerializeCompressedData(FArchive& Ar, bool bDDCData)
 	if (!HasAnyFlags(RF_ClassDefaultObject))
 	{
 		CompressedData.SerializeCompressedData(Ar, bDDCData, this, this->GetSkeleton(), BoneCompressionSettings, CurveCompressionSettings);
-		}
 	}
-
+}
 
 #if WITH_EDITOR
+void UAnimSequence::OnAnimModelLoaded()
+{
+	Super::OnAnimModelLoaded();
+
+	if (ShouldDataModelBeValid())
+	{
+		Controller->RemoveBoneTracksMissingFromSkeleton(GetSkeleton());
+
+		// Generate transient animation data, used to fallback on whenever the animation modified and not yet compressed (what previously used to be the raw data)
+		ResampleAnimationTrackData();
+
+		if (FAnimationUtils::GetForcedRecompressionSetting())
+		{
+			RequestAnimCompression(FRequestAnimCompressionParams(true, false, false));
+		}
+
+		if ( DataModelInterface->GetNumberOfKeys() == 0 && DataModelInterface->GetNumberOfFloatCurves() == 0 )
+		{
+			UE_LOG(LogAnimation, Warning, TEXT("No animation data exists for sequence %s (%s)"), *GetName(), (GetOuter() ? *GetOuter()->GetFullName() : *GetFullName()) );
+
+			if (!IsRunningGame())
+			{
+				static FName NAME_LoadErrors("LoadErrors");
+				FMessageLog LoadErrors(NAME_LoadErrors);
+
+				TSharedRef<FTokenizedMessage> Message = LoadErrors.Warning();
+				Message->AddToken(FTextToken::Create(LOCTEXT("EmptyAnimationData1", "The Animation ")));
+				Message->AddToken(FAssetNameToken::Create(GetPathName(), FText::FromString(GetName())));
+				Message->AddToken(FTextToken::Create(LOCTEXT("EmptyAnimationData2", " has no animation data. Recommend to remove.")));
+				LoadErrors.Notify();
+			}
+		}
+
+		if (GetLinkerCustomVersion(FUE5MainStreamObjectVersion::GUID) < FUE5MainStreamObjectVersion::MoveCustomAttributesToDataModel)
+		{
+			ValidateModel();
+			MoveAttributesToModel();
+		}
+	}
+}
+
 bool UAnimSequence::CanBakeAdditive() const
 {
 	return	(NumberOfSampledKeys > 0) &&
@@ -1997,7 +1988,7 @@ bool UAnimSequence::DoesSequenceContainZeroScale() const
 {
 	ValidateModel();
 
-	const TArray<FBoneAnimationTrack>& BoneAnimationTracks = DataModel->GetBoneAnimationTracks();
+	const TArray<FBoneAnimationTrack>& BoneAnimationTracks = DataModelInterface->GetBoneAnimationTracks();
 	for (const FBoneAnimationTrack& Track : BoneAnimationTracks)
 	{
 		for (const FVector3f& ScaleKey : Track.InternalTrackData.ScaleKeys)
@@ -2015,7 +2006,7 @@ bool UAnimSequence::DoesSequenceContainZeroScale() const
 FGuid UAnimSequence::GenerateGuidFromRawData() const
 {
 	ValidateModel();
-	return DataModel->GenerateGuid();
+	return DataModelInterface->GenerateGuid();
 }
 
 void CopyTransformToRawAnimationData(const FTransform& BoneTransform, FRawAnimSequenceTrack& Track, int32 Frame)
@@ -2031,17 +2022,13 @@ struct FByFramePoseEvalContext
 public:
 	FBoneContainer RequiredBones;
 
-	// Length of one frame.
-	const float IntervalTime;
-
 	TArray<FBoneIndexType> RequiredBoneIndexArray;
 
 	FByFramePoseEvalContext(const UAnimSequence* InAnimToEval)
-		: FByFramePoseEvalContext(InAnimToEval->GetPlayLength(), InAnimToEval->GetNumberOfSampledKeys(), InAnimToEval->GetSkeleton())
+		: FByFramePoseEvalContext(InAnimToEval->GetSkeleton())
 	{}
 		
-	FByFramePoseEvalContext(float InSequenceLength, int32 InRawNumOfFrames, USkeleton* InSkeleton)
-		: IntervalTime(InSequenceLength / ((float)FMath::Max(InRawNumOfFrames - 1, 1)))
+	FByFramePoseEvalContext(USkeleton* InSkeleton)
 	{
 		// Initialize RequiredBones for pose evaluation
 		RequiredBones.SetUseRAWData(true);
@@ -2170,15 +2157,14 @@ void UAnimSequence::BakeOutVirtualBoneTracks(TArray<FRawAnimSequenceTrack>& NewR
 
 	FAnimExtractContext ExtractContext;
 	const TArray<FVirtualBoneRefData>& VBRefData = GetSkeleton()->GetReferenceSkeleton().GetVirtualBoneRefData();
-	for (int32 Frame = 0; Frame < NumTrackKeys; ++Frame)
+	for (int32 KeyIndex = 0; KeyIndex < NumTrackKeys; ++KeyIndex)
 	{
 		// Initialise curve data from Skeleton
 		FBlendedCurve Curve;
 		Curve.InitFrom(EvalContext.RequiredBones);
 
 		//Grab pose for this frame
-		const float CurrentFrameTime = Frame * EvalContext.IntervalTime;
-		ExtractContext.CurrentTime = CurrentFrameTime;
+		ExtractContext.CurrentTime = TargetFrameRate.AsSeconds(KeyIndex);
 
 		UE::Anim::FStackAttributeContainer TempAttributes;
 		FAnimationPoseData AnimPoseData(Pose, Curve, TempAttributes);
@@ -2189,7 +2175,7 @@ void UAnimSequence::BakeOutVirtualBoneTracks(TArray<FRawAnimSequenceTrack>& NewR
 		{
 			const int32 VBIndex = VirtualBonesToAdd[VBToAddIndex];
 			const FVirtualBoneRefData& VB = VBRefData[VBIndex];
-			CopyTransformToRawAnimationData(Pose[FCompactPoseBoneIndex(VB.VBRefSkelIndex)], NewRawTracks[VBToAddIndex + VirtualBoneStartIndex], Frame);
+			CopyTransformToRawAnimationData(Pose[FCompactPoseBoneIndex(VB.VBRefSkelIndex)], NewRawTracks[VBToAddIndex + VirtualBoneStartIndex], KeyIndex);
 			UE::Anim::Compression::SanitizeRawAnimSequenceTrack(NewRawTracks[VBToAddIndex + VirtualBoneStartIndex]);
 		}
 	}
@@ -2239,15 +2225,14 @@ void UAnimSequence::TestEvalauteAnimation() const
 
 	FAnimExtractContext ExtractContext;
 
-	for (int32 Frame = 0; Frame < NumberOfSampledKeys; ++Frame)
+	for (int32 KeyIndex = 0; KeyIndex < NumberOfSampledKeys; ++KeyIndex)
 	{
 		// Initialise curve data from Skeleton
 		FBlendedCurve Curve;
 		Curve.InitFrom(EvalContext.RequiredBones);
 
 		//Grab pose for this frame
-		const float CurrentFrameTime = Frame * EvalContext.IntervalTime;
-		ExtractContext.CurrentTime = CurrentFrameTime;
+		ExtractContext.CurrentTime = TargetFrameRate.AsSeconds(KeyIndex);
 
 		UE::Anim::FStackAttributeContainer TempAttributes;
 		FAnimationPoseData AnimPoseData(Pose, Curve, TempAttributes);
@@ -2307,7 +2292,7 @@ void UAnimSequence::BakeOutAdditiveIntoRawData(TArray<FRawAnimSequenceTrack>& Ne
 	const FSmartNameMapping* CurveMapping = MySkeleton->GetSmartNameContainer(USkeleton::AnimCurveMappingName);
 	TArray<FName> SkeletonCurveNames;
 	CurveMapping->FillNameArray(SkeletonCurveNames);
-	for (const FFloatCurve& FloatCurve : DataModel->GetFloatCurves())
+	for (const FFloatCurve& FloatCurve : DataModelInterface->GetFloatCurves())
 	{
 		ensure(SkeletonCurveNames.Contains(FloatCurve.Name.DisplayName));
 	}
@@ -2349,7 +2334,7 @@ void UAnimSequence::BakeOutAdditiveIntoRawData(TArray<FRawAnimSequenceTrack>& Ne
 
 	FAnimExtractContext ExtractContext;
 
-	for (int32 Frame = 0; Frame < NumberOfSampledKeys; ++Frame)
+	for (int32 KeyIndex = 0; KeyIndex < NumberOfSampledKeys; ++KeyIndex)
 	{
 		// Initialise curve data from Skeleton
 		FBlendedCurve Curve;
@@ -2359,9 +2344,11 @@ void UAnimSequence::BakeOutAdditiveIntoRawData(TArray<FRawAnimSequenceTrack>& Ne
 		DummyBaseCurve.InitFrom(EvalContext.RequiredBones);
 
 		//Grab pose for this frame
-		const float PreviousFrameTime = (Frame - 1) * EvalContext.IntervalTime;
-		const float CurrentFrameTime = Frame * EvalContext.IntervalTime;
-		ExtractContext.CurrentTime = CurrentFrameTime;
+		const double PreviousKeyTime = TargetFrameRate.AsSeconds(KeyIndex - 1);
+		const double CurrentKeyTime = TargetFrameRate.AsSeconds(KeyIndex);
+		FFrameTime CurrentFrameTime = KeyIndex;
+
+		ExtractContext.CurrentTime = CurrentKeyTime;
 
 		UE::Anim::FStackAttributeContainer BaseAttributes;
 		FAnimationPoseData AnimPoseData(Pose, Curve, BaseAttributes);
@@ -2374,8 +2361,8 @@ void UAnimSequence::BakeOutAdditiveIntoRawData(TArray<FRawAnimSequenceTrack>& Ne
 		//Write out every track for this frame
 		for (FCompactPoseBoneIndex TrackIndex(0); TrackIndex < NewRawTracks.Num(); ++TrackIndex)
 		{
-			CopyTransformToRawAnimationData(Pose[TrackIndex], NewRawTracks[TrackIndex.GetInt()], Frame);
-			CopyTransformToRawAnimationData(BasePose[TrackIndex], AdditiveBaseAnimationData[TrackIndex.GetInt()], Frame);
+			CopyTransformToRawAnimationData(Pose[TrackIndex], NewRawTracks[TrackIndex.GetInt()], KeyIndex);
+			CopyTransformToRawAnimationData(BasePose[TrackIndex], AdditiveBaseAnimationData[TrackIndex.GetInt()], KeyIndex);
 		}
 
 		//Write out curve data for this frame
@@ -2402,21 +2389,22 @@ void UAnimSequence::BakeOutAdditiveIntoRawData(TArray<FRawAnimSequenceTrack>& Ne
 					if (!bHasKeys)
 					{
 						//Add pre key of 0
-						if (Frame > 0)
+						if (KeyIndex > 0)
 						{
-							RawCurve->UpdateOrAddKey(0.f, PreviousFrameTime);
+							RawCurve->UpdateOrAddKey(0.f, PreviousKeyTime);
 						}
-
 					}
 
 					if (!bHasKeys || IsNewKeyDifferent(RawCurve->FloatCurve.GetLastKey(), CurveWeight))
 					{
-						RawCurve->UpdateOrAddKey(CurveWeight, CurrentFrameTime);
+						RawCurve->UpdateOrAddKey(CurveWeight, CurrentKeyTime);
 						TArray<FRichCurveKey>& CurveKeys = RawCurve->FloatCurve.Keys;
 						if (CurveKeys.Num() > 1)
 						{
 							FRichCurveKey& PrevKey = CurveKeys.Last(1);
-							if (PrevKey.Time < (PreviousFrameTime - UE_SMALL_NUMBER)) // Did we skip a frame, if so need to make previous key const
+							// Round to frame here as it would have been added at a specific frame boundary (though float->double conversion might mean the value is off)
+							const FFrameTime PrevKeyTime = TargetFrameRate.AsFrameTime(PrevKey.Time).RoundToFrame();							
+							if (PrevKeyTime < CurrentFrameTime - 1) // Did we skip a frame, if so need to make previous key const
 							{
 								PrevKey.InterpMode = RCIM_Constant;
 							}
@@ -2454,13 +2442,14 @@ void UAnimSequence::UpdateDependentStreamingAnimations() const
 
 #endif
 
-
 #if WITH_EDITOR
 void UAnimSequence::ResetAnimation()
 {
 	// clear everything. Making new animation, so need to reset all the things that belong here
 	ValidateModel();
 
+	TargetFrameRate = UAnimationSettings::Get()->GetDefaultFrameRate();
+
 	Controller->OpenBracket(LOCTEXT("ResetAnimation_Bracket", "Resetting Animation Sequence"));
 	{
 		Controller->ResetModel();
@@ -2472,15 +2461,12 @@ void UAnimSequence::ResetAnimation()
 	AuthoredSyncMarkers.Empty();
 	UniqueMarkerNames.Empty();
 	AnimNotifyTracks.Empty();
-	RawDataGuid.Invalidate();
 }
 
 void UAnimSequence::DeleteDeprecatedRawAnimationData()
 {
 	PRAGMA_DISABLE_DEPRECATION_WARNINGS
 	RawAnimationData.Empty();
-	AnimationTrackNames.Empty();
-	TrackToSkeletonMapTable.Empty();
 	SourceRawAnimationData_DEPRECATED.Empty();
 	PRAGMA_ENABLE_DEPRECATION_WARNINGS
 }
@@ -2511,9 +2497,9 @@ bool UAnimSequence::IsValidAdditive() const
 		case ABPT_RefPose:
 			return true;
 		case ABPT_AnimScaled:
-			return (RefPoseSeq != NULL);
+			return (RefPoseSeq != nullptr);
 		case ABPT_AnimFrame:
-			return (RefPoseSeq != NULL) && (RefFrameIndex >= 0);
+			return (RefPoseSeq != nullptr) && (RefFrameIndex >= 0);
 		case ABPT_LocalAnimFrame:
 			return (RefFrameIndex >= 0);
 		default:
@@ -2591,7 +2577,7 @@ void FillUpTransformBasedOnRig(USkeleton* Skeleton, TArray<FTransform>& NodeSpac
 
 					const FTransformBase* TransformBase = Rig->GetTransformBaseByNodeName(NodeName);
 
-					if (TransformBase != NULL)
+					if (TransformBase != nullptr)
 					{
 						// orientation constraint			
 						const auto& RotConstraint = TransformBase->Constraints[EControlConstraint::Type::Orientation];
@@ -2661,7 +2647,7 @@ void UAnimSequence::RemapTracksToNewSkeleton( USkeleton* NewSkeleton, bool bConv
 		USkeleton* OldSkeleton = GetSkeleton();
 
 		// first check if both has same rig, if so, we'll retarget using it
-		if (OldSkeleton && OldSkeleton->GetRig() != NULL && NewSkeleton->GetRig() == OldSkeleton->GetRig() && OldSkeleton->GetPreviewMesh() && NewSkeleton->GetPreviewMesh())
+		if (OldSkeleton && OldSkeleton->GetRig() != nullptr && NewSkeleton->GetRig() == OldSkeleton->GetRig() && OldSkeleton->GetPreviewMesh() && NewSkeleton->GetPreviewMesh())
 		{
 			const URig* Rig = OldSkeleton->GetRig();
 
@@ -2746,8 +2732,8 @@ void UAnimSequence::RemapTracksToNewSkeleton( USkeleton* NewSkeleton, bool bConv
 			ConvertedSpaceAnimations.AddZeroed(SrcNumTracks);
 			ConvertedLocalSpaceAnimations.AddZeroed(SrcNumTracks);
 
-			const int32 NumKeys = DataModel->GetNumberOfKeys();
-			float Interval = DataModel->GetFrameRate().AsInterval();
+			const int32 NumKeys = DataModelInterface->GetNumberOfKeys();
+			float Interval = DataModelInterface->GetFrameRate().AsInterval();
 
 			// allocate arrays
 			for(int32 SrcTrackIndex=0; SrcTrackIndex<SrcNumTracks; ++SrcTrackIndex)
@@ -2771,7 +2757,7 @@ void UAnimSequence::RemapTracksToNewSkeleton( USkeleton* NewSkeleton, bool bConv
 				{
 					FTransform AnimatedLocalKey;
 					//ExtractBoneTransform(RiggingAnimationData.AnimationTracks[SrcTrackIndex], AnimatedLocalKey, Interval * Key);
-					FAnimationUtils::ExtractTransformFromTrack(Interval* Key, NumberOfSampledKeys, GetPlayLength(), RiggingAnimationData.AnimationTracks[SrcTrackIndex], Interpolation, AnimatedLocalKey);
+					FAnimationUtils::ExtractTransformFromTrack(RiggingAnimationData.AnimationTracks[SrcTrackIndex], Interval* Key, NumberOfSampledKeys, GetPlayLength(), Interpolation, AnimatedLocalKey);
 
 					AnimatedLocalKey.ScaleTranslation(OldToNewTranslationRatio[NodeIndex]);
 
@@ -2874,7 +2860,7 @@ void UAnimSequence::RemapTracksToNewSkeleton( USkeleton* NewSkeleton, bool bConv
 		{
 			// Validate animation tracks against the new skeleton, any tracks linked to bones that do not exist in the new hierarchy are removed
 			Controller->RemoveBoneTracksMissingFromSkeleton(NewSkeleton);
-			ensure(DataModel->GetNumBoneTracks() != 0);
+			ensure(DataModelInterface->GetNumBoneTracks() != 0);
 
 			// make sure you do update reference pose before coming here
 			
@@ -2937,8 +2923,8 @@ void UAnimSequence::RemapTracksToNewSkeleton( USkeleton* NewSkeleton, bool bConv
 			ConvertedLocalSpaces.AddZeroed(NumBones);
 			ConvertedSpaceBases.AddZeroed(NumBones);
 
-			const int32 NumKeys = DataModel->GetNumberOfKeys();
-			const float Interval = DataModel->GetFrameRate().AsInterval();
+			const int32 NumKeys = DataModelInterface->GetNumberOfKeys();
+			const float Interval = DataModelInterface->GetFrameRate().AsInterval();
 
 			// allocate arrays
 			for(int32 BoneIndex=0; BoneIndex<NumBones; ++BoneIndex)
@@ -2954,7 +2940,7 @@ void UAnimSequence::RemapTracksToNewSkeleton( USkeleton* NewSkeleton, bool bConv
 			{
 				FName BoneName = NewSkeleton->GetReferenceSkeleton().GetBoneName(BoneIndex);
 				const int32 OldBoneIndex = OldSkeleton->GetReferenceSkeleton().FindBoneIndex(BoneName);
-				const int32 TrackIndex = DataModel->GetBoneTrackIndexByName(BoneName);
+				const int32 TrackIndex = DataModelInterface->GetBoneTrackIndexByName(BoneName);
 				const int32 ParentBoneIndex = NewSkeleton->GetReferenceSkeleton().GetParentIndex(BoneIndex);
 
 				if(TrackIndex != INDEX_NONE)
@@ -2963,7 +2949,7 @@ void UAnimSequence::RemapTracksToNewSkeleton( USkeleton* NewSkeleton, bool bConv
 					for(int32 Key=0; Key<NumKeys; ++Key)
 					{
 						FTransform AnimatedLocalKey;
-						UE::Anim::GetBoneTransformFromModel(DataModel, AnimatedLocalKey, TrackIndex, Interval * Key, Interpolation);
+						UE::Anim::GetBoneTransformFromModel(DataModelInterface.GetInterface(), AnimatedLocalKey, TrackIndex, Interval * Key, Interpolation);
 
 						// note that we apply scale in the animated space
 						// at this point, you should have scaled version of animated skeleton
@@ -3006,7 +2992,7 @@ void UAnimSequence::RemapTracksToNewSkeleton( USkeleton* NewSkeleton, bool bConv
 			for(int32 BoneIndex=0; BoneIndex<NumBones; ++BoneIndex)
 			{
 				const FName BoneName = NewSkeleton->GetReferenceSkeleton().GetBoneName(BoneIndex);
-				const int32 TrackIndex = DataModel->GetBoneTrackIndexByName(BoneName);
+				const int32 TrackIndex = DataModelInterface->GetBoneTrackIndexByName(BoneName);
 				int32 ParentBoneIndex = NewSkeleton->GetReferenceSkeleton().GetParentIndex(BoneIndex);
 
 				for(int32 Key=0; Key<NumKeys; ++Key)
@@ -3058,7 +3044,7 @@ void UAnimSequence::RemapTracksToNewSkeleton( USkeleton* NewSkeleton, bool bConv
 		}
 		else
 		{
-			Controller->RemoveBoneTracksMissingFromSkeleton(NewSkeleton);
+			Controller->UpdateWithSkeleton(NewSkeleton);
 		}
 
 		// I have to set this here in order for compression
@@ -3067,62 +3053,13 @@ void UAnimSequence::RemapTracksToNewSkeleton( USkeleton* NewSkeleton, bool bConv
 	}
 	else
 	{
-		Controller->RemoveBoneTracksMissingFromSkeleton(NewSkeleton);
+		SetSkeleton(NewSkeleton);
+		Controller->UpdateWithSkeleton(NewSkeleton, false);
 	}
 
 	Super::RemapTracksToNewSkeleton(NewSkeleton, bConvertSpaces);
 }
 
-
-void UAnimSequence::RemoveNaNTracks()
-{
-	ValidateModel();
-
-	TArray<FName> TracksToRemove;
-
-	const TArray<FBoneAnimationTrack>& BoneAnimationTracks = DataModel->GetBoneAnimationTracks();		
-	for (const FBoneAnimationTrack& AnimationTrack : BoneAnimationTracks)
-	{
-		const FRawAnimSequenceTrack& RawTrack = AnimationTrack.InternalTrackData;
-
-		bool bContainsNaN = false;
-		for ( auto Key : RawTrack.PosKeys )
-		{
-			bContainsNaN |= Key.ContainsNaN();
-		}
-
-		if (!bContainsNaN)
-		{
-			for(auto Key : RawTrack.RotKeys)
-			{
-				bContainsNaN |= Key.ContainsNaN();
-			}
-		}
-
-		if (!bContainsNaN)
-		{
-			for(auto Key : RawTrack.ScaleKeys)
-			{
-				bContainsNaN |= Key.ContainsNaN();
-			}
-		}
-
-		if (bContainsNaN)
-		{
-			UE_LOG(LogAnimation, Warning, TEXT("Animation raw data contains NaNs - Removing the following track [%s Track (%s)]"), (GetOuter() ? *GetOuter()->GetFullName() : *GetFullName()), *AnimationTrack.Name.ToString());
-
-			TracksToRemove.Add(AnimationTrack.Name);
-		}
-	}
-
-	IAnimationDataController::FScopedBracket ScopedBracket(Controller, LOCTEXT("RemoveNaNTracks_Bracket", "Removing track(s) containing NaN key data"));	
-	for (const FName& TrackName : TracksToRemove)
-	{
-		Controller->RemoveBoneTrack(TrackName);
-	}	
-}
-
-
 int32 FindFirstChildTrack(const USkeleton* MySkeleton, const FReferenceSkeleton& RefSkeleton, const TArray<FName>& AnimationTrackNames, FName BoneName)
 {
 	const int32 BoneIndex = RefSkeleton.FindBoneIndex(BoneName);
@@ -3224,8 +3161,8 @@ int32 UAnimSequence::GetSpaceBasedAnimationData(TArray< TArray<FTransform> >& An
 	ValidateModel();
 
 	// 2d array of animated time [boneindex][time key]
-	const int32 NumKeys = DataModel->GetNumberOfKeys();
-	const float Interval = DataModel->GetFrameRate().AsInterval();
+	const int32 NumKeys = DataModelInterface->GetNumberOfKeys();
+	const float Interval = DataModelInterface->GetFrameRate().AsInterval();
 
 	// allocate arrays
 	for (int32 BoneIndex = 0; BoneIndex < NumBones; ++BoneIndex)
@@ -3391,7 +3328,7 @@ int32 UAnimSequence::GetSpaceBasedAnimationData(TArray< TArray<FTransform> >& An
 		for (int32 BoneIndex = 0; BoneIndex < NumBones; ++BoneIndex)
 		{
 			const FName BoneName = MySkeleton->GetReferenceSkeleton().GetBoneName(BoneIndex);
-			const int32 TrackIndex = DataModel->GetBoneTrackIndexByName(BoneName);
+			const int32 TrackIndex = DataModelInterface->GetBoneTrackIndexByName(BoneName);
 			const int32 ParentBoneIndex = MySkeleton->GetReferenceSkeleton().GetParentIndex(BoneIndex);
 
 			if (TrackIndex != INDEX_NONE)
@@ -3400,7 +3337,7 @@ int32 UAnimSequence::GetSpaceBasedAnimationData(TArray< TArray<FTransform> >& An
 				for (int32 Key = 0; Key < NumKeys; ++Key)
 				{
 					FTransform AnimatedLocalKey;
-					UE::Anim::GetBoneTransformFromModel(DataModel, AnimatedLocalKey, TrackIndex, Interval * Key, Interpolation);
+					UE::Anim::GetBoneTransformFromModel(DataModelInterface.GetInterface(), AnimatedLocalKey, TrackIndex, Interval * Key, Interpolation);
 
 					if (ParentBoneIndex != INDEX_NONE)
 					{
@@ -3445,10 +3382,10 @@ bool UAnimSequence::ConvertAnimationDataToRiggingData(FAnimSequenceTrackContaine
 		TArray<FName> ValidNodeNames;
 		int32 NumNodes = MySkeleton->GetMappedValidNodes(ValidNodeNames);
 		TArray< TArray<FTransform> > AnimationDataInComponentSpace;
-		int32 NumBones = GetSpaceBasedAnimationData(AnimationDataInComponentSpace, NULL);
+		int32 NumBones = GetSpaceBasedAnimationData(AnimationDataInComponentSpace, nullptr);
 
 		ValidateModel();
-		const int32 NumModelKeys = DataModel->GetNumberOfKeys();
+		const int32 NumModelKeys = DataModelInterface->GetNumberOfKeys();
 
 		if (NumBones > 0)
 		{
@@ -3605,7 +3542,7 @@ bool UAnimSequence::ConvertRiggingDataToAnimationData(FAnimSequenceTrackContaine
 		TArray<FQuat4f> RotKeys;
 		TArray<FVector3f> ScaleKeys;
 
-		const int32 NumModelKeys = DataModel->GetNumberOfKeys();
+		const int32 NumModelKeys = DataModelInterface->GetNumberOfKeys();
 		const FReferenceSkeleton& RefSkeleton = MySkeleton->GetReferenceSkeleton();
 		const URig* Rig = MySkeleton->GetRig();
 		for (int32 NodeIndex = 0; NodeIndex < ValidNumNodes; ++NodeIndex)
@@ -3662,7 +3599,7 @@ void UAnimSequence::BakeTrackCurvesToRawAnimationTracks(TArray<FRawAnimSequenceT
 {
 	ValidateModel();
 
-	if (DataModel->GetNumberOfTransformCurves() > 0)
+	if (DataModelInterface->GetNumberOfTransformCurves() > 0)
 	{
 		USkeleton* CurSkeleton = GetSkeleton();
 		check(CurSkeleton);
@@ -3710,15 +3647,14 @@ void UAnimSequence::BakeTrackCurvesToRawAnimationTracks(TArray<FRawAnimSequenceT
 
 
 		// NumFrames can't be zero (filtered earlier)
-		const float Interval = GetIntervalPerKey(NumberOfSampledKeys, GetPlayLength());
 
 		const FSmartNameMapping* NameMapping = CurSkeleton->GetSmartNameContainer(USkeleton::AnimTrackCurveMappingName);
-		const FAnimationCurveData& AnimationCurveData = DataModel->GetCurveData();
+		const FAnimationCurveData& AnimationCurveData = DataModelInterface->GetCurveData();
 		for (const FTransformCurve& Curve : AnimationCurveData.TransformCurves)
 		{
 			// find curves first, and then see what is index of this curve
 			FName BoneName;
-
+			NameMapping->GetName(Curve.Name.UID, BoneName);
 			if (Curve.GetCurveTypeFlag(AACF_Disabled) == false &&
 				ensureAlways(NameMapping->GetName(Curve.Name.UID, BoneName)))
 			{
@@ -3779,11 +3715,13 @@ void UAnimSequence::BakeTrackCurvesToRawAnimationTracks(TArray<FRawAnimSequenceT
 				for (int32 KeyIndex = 0; KeyIndex < NumberOfSampledKeys; ++KeyIndex)
 				{
 					// now evaluate
-					FTransform AdditiveTransform = Curve.Evaluate(KeyIndex * Interval, 1.0);
-						FTransform LocalTransform(FQuat(RawTrack.RotKeys[KeyIndex]), FVector(RawTrack.PosKeys[KeyIndex]), FVector(RawTrack.ScaleKeys[KeyIndex]));
-						RawTrack.RotKeys[KeyIndex] = FQuat4f(LocalTransform.GetRotation() * AdditiveTransform.GetRotation());
-						RawTrack.PosKeys[KeyIndex] = FVector3f(LocalTransform.TransformPosition(AdditiveTransform.GetTranslation()));
-						RawTrack.ScaleKeys[KeyIndex] = FVector3f(LocalTransform.GetScale3D() * AdditiveTransform.GetScale3D());
+					const double EvaluationTime = DataModelInterface->GetFrameRate().AsSeconds(KeyIndex);
+					FTransform AdditiveTransform = Curve.Evaluate(EvaluationTime, 1.0);
+					FTransform LocalTransform(FQuat(RawTrack.RotKeys[KeyIndex]), FVector(RawTrack.PosKeys[KeyIndex]), FVector(RawTrack.ScaleKeys[KeyIndex]));
+
+					RawTrack.RotKeys[KeyIndex] = FQuat4f(LocalTransform.GetRotation() * AdditiveTransform.GetRotation());
+					RawTrack.PosKeys[KeyIndex] = FVector3f(LocalTransform.TransformPosition(AdditiveTransform.GetTranslation()));
+					RawTrack.ScaleKeys[KeyIndex] = FVector3f(LocalTransform.GetScale3D() * AdditiveTransform.GetScale3D());
 				}
 
 				// Apply PostProcess behaviour in-place
@@ -3809,9 +3747,7 @@ void UAnimSequence::AddKeyToSequence(float Time, const FName& BoneName, const FT
 	FAnimationCurveIdentifier TransformCurveId(NewCurveName, ERawCurveTrackTypes::RCT_Transform);
 	Controller->AddCurve(TransformCurveId, AACF_DriveTrack | AACF_Editable);
 
-	Controller->RemoveBoneTracksMissingFromSkeleton(CurrentSkeleton);
-
-	const FTransformCurve* TransformCurve = DataModel->FindTransformCurve(TransformCurveId);
+	const FTransformCurve* TransformCurve = DataModelInterface->FindTransformCurve(TransformCurveId);
 	check(TransformCurve);
 
 	Controller->SetTransformCurveKey(TransformCurveId, Time, AdditiveTransform);
@@ -3819,13 +3755,13 @@ void UAnimSequence::AddKeyToSequence(float Time, const FName& BoneName, const FT
 
 uint8* UAnimSequence::FindSyncMarkerPropertyData(int32 SyncMarkerIndex, FArrayProperty*& ArrayProperty)
 {
-	ArrayProperty = NULL;
+	ArrayProperty = nullptr;
 
 	if (AuthoredSyncMarkers.IsValidIndex(SyncMarkerIndex))
 	{
 		return FindArrayProperty(TEXT("AuthoredSyncMarkers"), ArrayProperty, SyncMarkerIndex);
 	}
-	return NULL;
+	return nullptr;
 }
 
 bool UAnimSequence::CreateAnimation(USkeletalMesh* Mesh)
@@ -3837,10 +3773,12 @@ bool UAnimSequence::CreateAnimation(USkeletalMesh* Mesh)
 		const FReferenceSkeleton& RefSkeleton = Mesh->GetRefSkeleton();
 
 		IAnimationDataController::FScopedBracket ScopedBracket(Controller, LOCTEXT("CreateAnimationSkeletalMesh_Bracket", "Creating Animation Sequence based up Skeletal Mesh"));
+
+		Controller->InitializeModel();
 		ResetAnimation();
 
-		Controller->SetPlayLength(MINIMUM_ANIMATION_LENGTH);
-		Controller->SetFrameRate(FFrameRate(1.f / MINIMUM_ANIMATION_LENGTH, 1));
+		Controller->SetFrameRate(UAnimationSettings::Get()->GetDefaultFrameRate());
+		Controller->SetNumberOfFrames(1);
 
 		const int32 NumBones = RefSkeleton.GetRawBoneNum();
 
@@ -3874,10 +3812,12 @@ bool UAnimSequence::CreateAnimation(USkeletalMeshComponent* MeshComponent)
 		const FReferenceSkeleton& RefSkeleton = Mesh->GetRefSkeleton();
 
 		IAnimationDataController::FScopedBracket ScopedBracket(Controller, LOCTEXT("CreateAnimationSkeletalMeshComponent_Bracket", "Creating Animation Sequence based up Skeletal Mesh Component"));
+
+		Controller->InitializeModel();
 		ResetAnimation();
 
-		Controller->SetPlayLength(MINIMUM_ANIMATION_LENGTH);
-		Controller->SetFrameRate(FFrameRate(1.f / MINIMUM_ANIMATION_LENGTH, 1));
+		Controller->SetFrameRate(UAnimationSettings::Get()->GetDefaultFrameRate());
+		Controller->SetNumberOfFrames(1);
 
 		const int32 NumBones = RefSkeleton.GetRawBoneNum();
 
@@ -3905,8 +3845,8 @@ bool UAnimSequence::CreateAnimation(UAnimSequence* Sequence)
 	{
 		ResetAnimation();
 
-		CopyDataModel(Sequence->GetDataModel());
-		Controller->SetModel(DataModel);
+		CopyDataModel(Sequence->GetDataModelInterface());
+		Controller->SetModel(DataModelInterface);
 		Notifies = Sequence->Notifies;
 		AnimNotifyTracks = Sequence->AnimNotifyTracks;
 
@@ -3973,7 +3913,7 @@ void UAnimSequence::EvaluateCurveData(FBlendedCurve& OutCurve, float CurrentTime
 		// Evaluate float curves from the AnimationData Model
 		if (OutCurve.NumValidCurveCount > 0)
 		{
-			for (auto CurveIter = DataModel->GetFloatCurves().CreateConstIterator(); CurveIter; ++CurveIter)
+			for (auto CurveIter = DataModelInterface->GetFloatCurves().CreateConstIterator(); CurveIter; ++CurveIter)
 			{
 				const FFloatCurve& Curve = *CurveIter;
 				if (OutCurve.IsEnabled(Curve.Name.UID))
@@ -4328,7 +4268,11 @@ void UAnimSequence::ValidateCurrentPosition(const FMarkerSyncAnimPosition& Posit
 
 bool UAnimSequence::UseRawDataForPoseExtraction(const FBoneContainer& RequiredBones) const
 {
-	return CanEvaluateRawAnimationData() && (bUseRawDataOnly || (GetSkeletonVirtualBoneGuid() != GetSkeleton()->GetVirtualBoneGuid()) || RequiredBones.GetDisableRetargeting() || RequiredBones.ShouldUseRawData() || RequiredBones.ShouldUseSourceData());
+	return CanEvaluateRawAnimationData() && (bUseRawDataOnly || (GetSkeletonVirtualBoneGuid() != GetSkeleton()->GetVirtualBoneGuid()) || RequiredBones.GetDisableRetargeting() || RequiredBones.ShouldUseRawData() ||
+ #if WITH_EDITOR
+  	GForceRawData == 1 ||
+ #endif // WITH_EDITOR
+	RequiredBones.ShouldUseSourceData());
 }
 
 #if WITH_EDITOR
@@ -4359,13 +4303,13 @@ void UAnimSequence::EvaluateAttributes(FAnimationPoseData& OutAnimationPoseData,
 	if (bUseRawData)
 	{
 		ValidateModel();
-		for (const FAnimatedBoneAttribute& Attribute : DataModel->GetAttributes())
+		for (const FAnimatedBoneAttribute& Attribute : DataModelInterface->GetAttributes())
 		{
 			const FCompactPoseBoneIndex PoseBoneIndex = RequiredBones.GetCompactPoseIndexFromSkeletonIndex(Attribute.Identifier.GetBoneIndex());
 			// Only add attribute if the bone its tied to exists in the currently evaluated set of bones
 			if(PoseBoneIndex.IsValid())
 			{
-				UE::Anim::Attributes::GetAttributeValue(OutAttributes, PoseBoneIndex, Attribute, ExtractionContext);
+				UE::Anim::Attributes::GetAttributeValue(OutAttributes, PoseBoneIndex, Attribute, ExtractionContext.CurrentTime);
 			}
 		}
 	}
@@ -4393,7 +4337,7 @@ void UAnimSequence::RemoveCustomAttribute(const FName& BoneName, const FName& At
 
 	TArray<FAnimationAttributeIdentifier> IdentifiersToRemove;
 
-	Algo::TransformIf(DataModel->GetAttributes(), IdentifiersToRemove, 
+	Algo::TransformIf(DataModelInterface->GetAttributes(), IdentifiersToRemove, 
 		[BoneName, AttributeName](const FAnimatedBoneAttribute& Attribute)
 		{
 			return Attribute.Identifier.GetBoneName() == BoneName && Attribute.Identifier.GetName() == AttributeName;
@@ -4466,7 +4410,7 @@ void UAnimSequence::SynchronousAnimatedBoneAttributesCompression()
 		// Helper struct to match sample timings with regular additive baking
 		FByFramePoseEvalContext EvalContext(this);
 
-		for (const FAnimatedBoneAttribute& AdditiveAttribute : DataModel->GetAttributes())
+		for (const FAnimatedBoneAttribute& AdditiveAttribute : DataModelInterface->GetAttributes())
 		{
 			ensure(!AttributeCurves.Contains(AdditiveAttribute.Identifier));
 			FAttributeCurve& AttributeCurve = AttributeCurves.Add(AdditiveAttribute.Identifier);
@@ -4475,10 +4419,9 @@ void UAnimSequence::SynchronousAnimatedBoneAttributesCompression()
 			{
 				AttributeCurve.SetScriptStruct(AdditiveAttribute.Identifier.GetType());
 				const FAnimatedBoneAttribute& RefAttribute = *RefPoseAttributePtr;
-				for (int32 Frame = 0; Frame < NumberOfSampledKeys; ++Frame)
+				for (int32 KeyIndex = 0; KeyIndex < NumberOfSampledKeys; ++KeyIndex)
 				{
-					const float CurrentFrameTime = Frame * EvalContext.IntervalTime;
-
+					const double CurrentFrameTime = TargetFrameRate.AsSeconds(KeyIndex);
 					UE::Anim::FStackAttributeContainer AdditiveAttributes;
 					uint8* AdditivePtr = AdditiveAttributes.FindOrAdd(AdditiveAttribute.Identifier.GetType(), UE::Anim::FAttributeId(NAME_None, 0,NAME_None));
 					AdditiveAttribute.Curve.EvaluateToPtr(AdditiveAttribute.Identifier.GetType(), CurrentFrameTime, AdditivePtr);
@@ -4500,7 +4443,7 @@ void UAnimSequence::SynchronousAnimatedBoneAttributesCompression()
 	}
 	else
 	{
-		for (const FAnimatedBoneAttribute& Attribute : DataModel->GetAttributes())
+		for (const FAnimatedBoneAttribute& Attribute : DataModelInterface->GetAttributes())
 		{
 			// Do something with the attributes
 			ensure(!AttributeCurves.Contains(Attribute.Identifier));
@@ -5094,7 +5037,7 @@ bool UAnimSequence::IsCompressedDataValid() const
 	// if we have no raw data. This can happen with sequences that only has curves.
 #if WITH_EDITOR
 	PRAGMA_DISABLE_DEPRECATION_WARNINGS
-	if (IsDataModelValid() && DataModel->GetNumBoneTracks() == 0 && RawAnimationData.Num() == 0)
+	if (IsDataModelValid() && DataModelInterface->GetNumBoneTracks() == 0 && RawAnimationData.Num() == 0)
 	PRAGMA_ENABLE_DEPRECATION_WARNINGS
 	{
 		return true;
@@ -5111,7 +5054,7 @@ bool UAnimSequence::IsCurveCompressedDataValid() const
 	// if we have no raw data. This can happen with sequences that only has bones.
 
 	const bool bNoCompressedCurveData = CompressedData.CompressedCurveByteStream.Num() == 0;
-	const bool bHasRawCurveData = IsDataModelValid() && DataModel->GetNumberOfFloatCurves() != 0;
+	const bool bHasRawCurveData = IsDataModelValid() && DataModelInterface->GetNumberOfFloatCurves() != 0;
 
 	if ( bNoCompressedCurveData && bHasRawCurveData )
 	{
@@ -5145,7 +5088,7 @@ void UAnimSequence::EnsureValidRawDataGuid()
 PRAGMA_DISABLE_DEPRECATION_WARNINGS
 	if (IsDataModelValid() && !RawDataGuid.IsValid())
 	{
-		RawDataGuid = GenerateGuidFromRawData();
+		RawDataGuid = DataModelInterface->GenerateGuid();
 	}
 PRAGMA_ENABLE_DEPRECATION_WARNINGS
 }
@@ -5170,7 +5113,7 @@ void UAnimSequence::PopulateModel()
 	const bool bValidNumberOfKeys = NumberOfKeys >= 1 || bCurveOnlySequence;
 	const float ExpectedPlayLength = (float)SamplingFrameRate.AsSeconds(FMath::Max(NumberOfKeys - 1, 1));
 	const bool bValidSequenceLength = FMath::IsNearlyEqual(ExpectedPlayLength, SequenceLength, UE_KINDA_SMALL_NUMBER) || bCurveOnlySequence;
-	const bool bValidNumberOfTrackArrayEntries = AnimationTrackNames.Num() == RawAnimationData.Num() && RawAnimationData.Num() == TrackToSkeletonMapTable.Num();
+	const bool bValidNumberOfTrackArrayEntries = AnimationTrackNames.Num() == RawAnimationData.Num();
 	const bool bValidAnimationData = SourceRawAnimationData_DEPRECATED.Num() == 0 || SourceRawAnimationData_DEPRECATED.Num() == RawAnimationData.Num();
 	const bool bValidFrameRate = SamplingFrameRate.IsValid() && (SamplingFrameRate.AsInterval() >= 0.f);
 	const bool bValidCurveData = RawCurveData.FloatCurves.Num() >= 1;
@@ -5188,68 +5131,84 @@ void UAnimSequence::PopulateModel()
 
 	Controller->OpenBracket(LOCTEXT("UAnimSequence::PopulateModel_Bracket", "Generating Animation Model Data from Animation Sequence"));
 
-	Controller->SetPlayLength(bCurveOnlySequence ? PlayLength : ExpectedPlayLength);
 	Controller->SetFrameRate(FrameRate);
 
-	USkeleton* TargetSkeleton = GetSkeleton();
+	const FFrameTime FrameTime = FrameRate.AsFrameTime(bCurveOnlySequence ? PlayLength : ExpectedPlayLength);
+	const FFrameNumber RoundedFrameNumber = FrameTime.RoundToFrame();
+	ensureAlways(RoundedFrameNumber.Value == FMath::Max((NumKeys - 1), 1));
+	Controller->SetNumberOfFrames(RoundedFrameNumber);	
+
+	const USkeleton* TargetSkeleton = GetSkeleton();
 	UE::Anim::CopyCurveDataToModel(CurveData, TargetSkeleton, *Controller);
 	
 	const int32 NumTracks = SequenceTracks.Num();
 	for (int32 TrackIndex = 0; TrackIndex < NumTracks; ++TrackIndex)
 	{
-		Controller->AddBoneTrack(TempAnimationTrackNames[TrackIndex]);
-
-		auto GenerateUniformKeys = [NumKeys](auto& Keys, auto IdentityKey)
+		if (SequenceTracks[TrackIndex].ContainsNaN())
 		{
-			// Convert track keys to be uniform			
-			if (Keys.Num() == 0)
+			continue;
+		}
+		
+		const int32 BoneTrackIndex = Controller->AddBoneTrack(TempAnimationTrackNames[TrackIndex]);
+		if (BoneTrackIndex != INDEX_NONE)
+		{
+			auto GenerateUniformKeys = [NumKeys](auto& Keys, auto IdentityKey)
 			{
-				// set all to identity
-				for (int32 Index = 0; Index < NumKeys; ++Index)
+				// Convert track keys to be uniform
+				if (Keys.Num() == 0)
 				{
-					Keys.Add(IdentityKey);
+					// set all to identity
+					for (int32 Index = 0; Index < NumKeys; ++Index)
+					{
+						Keys.Add(IdentityKey);
+					}
 				}
-			}
-			else if (Keys.Num() == 1)
-			{
-				// set all to single key
-				const auto KeyZero = Keys[0];
-				for (int32 Index = 0; Index < (NumKeys - 1); ++Index)
+				else if (Keys.Num() == 1)
 				{
-					Keys.Add(KeyZero);
+					// set all to single key
+					const auto KeyZero = Keys[0];
+					for (int32 Index = 0; Index < (NumKeys - 1); ++Index)
+					{
+						Keys.Add(KeyZero);
+					}
 				}
-			}
 
-			ensure(Keys.Num() == NumKeys);
-		};
+				ensure(Keys.Num() == NumKeys);
+			};
 
-		TArray<FVector3f> PosKeys = SequenceTracks[TrackIndex].PosKeys;
-		GenerateUniformKeys(PosKeys, FVector3f::ZeroVector);
+			TArray<FVector3f> PosKeys = SequenceTracks[TrackIndex].PosKeys;
+			GenerateUniformKeys(PosKeys, FVector3f::ZeroVector);
 
-		TArray<FQuat4f> RotKeys = SequenceTracks[TrackIndex].RotKeys;
-		GenerateUniformKeys(RotKeys, FQuat4f::Identity);
+			TArray<FQuat4f> RotKeys = SequenceTracks[TrackIndex].RotKeys;
+			GenerateUniformKeys(RotKeys, FQuat4f::Identity);
 
-		TArray<FVector3f> ScaleKeys = SequenceTracks[TrackIndex].ScaleKeys;
-		GenerateUniformKeys(ScaleKeys, FVector3f::OneVector);
+			TArray<FVector3f> ScaleKeys = SequenceTracks[TrackIndex].ScaleKeys;
+			GenerateUniformKeys(ScaleKeys, FVector3f::OneVector);
 
-		Controller->SetBoneTrackKeys(TempAnimationTrackNames[TrackIndex], PosKeys, RotKeys, ScaleKeys);
+			Controller->SetBoneTrackKeys(TempAnimationTrackNames[TrackIndex], PosKeys, RotKeys, ScaleKeys);
+		}
 	}
 
 	DeleteDeprecatedRawAnimationData();
 	
+	Controller->NotifyPopulated();
+	
 	Controller->CloseBracket();
 }
 
-void UAnimSequence::OnModelModified(const EAnimDataModelNotifyType& NotifyType, UAnimDataModel* Model, const FAnimDataModelNotifPayload& Payload)
+void UAnimSequence::OnModelModified(const EAnimDataModelNotifyType& NotifyType, IAnimationDataModel* Model, const FAnimDataModelNotifPayload& Payload)
 {
+	QUICK_SCOPE_CYCLE_COUNTER(STAT_AnimSequence_OnModelModified);
 	UAnimSequenceBase::OnModelModified(NotifyType, Model, Payload);	
 
 	static const bool GenerateNewGUID = true;
 	static const bool RegenerateGUID = false;
 
-	auto UpdateRawDataGuid = [this](bool bForceNewRawDataGuid)
+	auto UpdateRawDataGuid = [this, Model](bool bForceNewRawDataGuid)
 	{
-		RawDataGuid = bForceNewRawDataGuid ? FGuid::NewGuid() : GenerateGuidFromRawData();
+		PRAGMA_DISABLE_DEPRECATION_WARNINGS
+		RawDataGuid = bForceNewRawDataGuid ? FGuid::NewGuid() : Model->GenerateGuid();
+		PRAGMA_ENABLE_DEPRECATION_WARNINGS
 		bUseRawDataOnly = true;
 
 		FlagDependentAnimationsAsRawDataOnly();
@@ -5266,14 +5225,19 @@ void UAnimSequence::OnModelModified(const EAnimDataModelNotifyType& NotifyType,
 		}
 	};
 
-	auto HandleLengthChanged = [this](float NewLength, float OldLength, float T0, float T1)
+	auto HandleNumberOfFramesChanged = [this](FFrameNumber NewNumberOfFrames, FFrameNumber OldNumberOfFrames, FFrameNumber Frame0, FFrameNumber Frame1)
 	{
 		if (bPopulatingDataModel)
 		{
 			return;
 		}
+		const FFrameRate& ModelFrameRate = DataModelInterface->GetFrameRate();
+
+		const float NewLength = ModelFrameRate.AsSeconds(NewNumberOfFrames);
+		const float T0 = ModelFrameRate.AsSeconds(Frame0);
+		const float T1 = ModelFrameRate.AsSeconds(Frame1);
 
-		if (NewLength > OldLength)
+		if (NewNumberOfFrames > OldNumberOfFrames)
 		{
 			const float InsertTime = T0;
 			const float Duration = T1 - T0;
@@ -5292,7 +5256,7 @@ void UAnimSequence::OnModelModified(const EAnimDataModelNotifyType& NotifyType,
 				Marker.Time = FMath::Clamp(CurrentTime, 0.f, NewLength);
 			}
 		}
-		else if (NewLength < OldLength)
+		else if (NewNumberOfFrames < OldNumberOfFrames)
 		{
 			const float StartRemoveTime = T0;
 			const float EndRemoveTime = T1;
@@ -5315,23 +5279,20 @@ void UAnimSequence::OnModelModified(const EAnimDataModelNotifyType& NotifyType,
 
 				Marker.Time = FMath::Clamp(CurrentTime, 0.f, NewLength);
 			}
-
 		}
 	};
 
 	bool bShouldMarkPackageDirty = true;
 
-	PRAGMA_DISABLE_DEPRECATION_WARNINGS
 	switch (NotifyType)
 	{
-		case EAnimDataModelNotifyType::SequenceLengthChanged:		
+		case EAnimDataModelNotifyType::SequenceLengthChanged:
 		{
 			const FSequenceLengthChangedPayload& TypedPayload = Payload.GetPayload<FSequenceLengthChangedPayload>();
 
-			const float OldSequenceLength = TypedPayload.PreviousLength;
-			const float CurrentSequenceLength = Model->GetPlayLength();
-			
-			HandleLengthChanged(CurrentSequenceLength, OldSequenceLength, TypedPayload.T0, TypedPayload.T1);
+			const FFrameNumber OldNumberOfFrames = TypedPayload.PreviousNumberOfFrames;
+			const FFrameNumber CurrentNumberOfFrames = Model->GetNumberOfFrames();
+			HandleNumberOfFramesChanged(CurrentNumberOfFrames, OldNumberOfFrames, TypedPayload.Frame0, TypedPayload.Frame1);
 						
 			if (NotifyCollector.IsNotWithinBracket())
 			{
@@ -5468,9 +5429,8 @@ void UAnimSequence::OnModelModified(const EAnimDataModelNotifyType& NotifyType,
 			break;
 		}
 	}
-	PRAGMA_ENABLE_DEPRECATION_WARNINGS
 
-	if (bShouldMarkPackageDirty)
+	if (NotifyCollector.IsNotWithinBracket() && NotifyType != EAnimDataModelNotifyType::BracketClosed && bShouldMarkPackageDirty)
 	{
 		MarkPackageDirty();
 	}
@@ -5479,52 +5439,86 @@ void UAnimSequence::OnModelModified(const EAnimDataModelNotifyType& NotifyType,
 
 void UAnimSequence::ResampleAnimationTrackData()
 {
+	QUICK_SCOPE_CYCLE_COUNTER(STAT_ResampleAnimationTrackData);
+	
 	if (ShouldDataModelBeValid())
 	{
-	ValidateModel();
+		ValidateModel();
 
-	// Make a copy, deals with bone name and index
-	ResampledAnimationTrackData = DataModel->GetBoneAnimationTracks();
+		FAnimationUtils::EnsureAnimSequenceLoaded(*this);
 
-	if (TargetFrameRate == DataModel->GetFrameRate())
-	{
-		// Can directly copy
-		NumberOfSampledKeys = DataModel->GetNumberOfKeys();
-		NumberOfSampledFrames = DataModel->GetNumberOfFrames();
-	}
-	else
-	{
-		// We actually need to resample bone transforms
-		NumberOfSampledFrames = TargetFrameRate.AsFrameNumber(DataModel->GetPlayLength()).Value;
-		NumberOfSampledKeys = NumberOfSampledFrames + 1;
-				
-		for (int32 TrackIndex = 0; TrackIndex < ResampledAnimationTrackData.Num(); ++TrackIndex)	
+		// Make a copy, deals with bone name and index
+		ResampledAnimationTrackData = DataModelInterface->GetBoneAnimationTracks();
+	
+		ensure(DataModelInterface != nullptr);
 		{
-			FBoneAnimationTrack& TrackData = ResampledAnimationTrackData[TrackIndex];
+			// Can directly copy (for now, as we still maintain GetBoneAnimationTracks)
+			if (TargetFrameRate == DataModelInterface->GetFrameRate())
+			{
+				NumberOfSampledKeys = DataModelInterface->GetNumberOfKeys();
+				NumberOfSampledFrames = DataModelInterface->GetNumberOfFrames();
+			}
+			else
+			{
+				FMemMark Mark(FMemStack::Get());
+				FByFramePoseEvalContext EvalContext(this);
+				EvalContext.RequiredBones.SetDisableRetargeting(true);
+				EvalContext.RequiredBones.SetUseRAWData(true);
+				EvalContext.RequiredBones.SetUseSourceData(true);
 
-			TrackData.InternalTrackData.PosKeys.SetNumUninitialized(NumberOfSampledKeys);
-			TrackData.InternalTrackData.RotKeys.SetNumUninitialized(NumberOfSampledKeys);
-			TrackData.InternalTrackData.ScaleKeys.SetNumUninitialized(NumberOfSampledKeys);
+				// Make a copy, deals with bone name and index
+				ResampledAnimationTrackData = DataModelInterface->GetBoneAnimationTracks();
 
-			for (int32 FrameIndex = 0; FrameIndex < NumberOfSampledKeys; ++FrameIndex)
-			{
-				FTransform Transform;
+				// We actually need to resample bone transforms
+				const FFrameNumber ModelNumberOfFrames = DataModelInterface->GetNumberOfFrames();
+				const FFrameTime ResampledFrameTime = FFrameRate::TransformTime(ModelNumberOfFrames, DataModelInterface->GetFrameRate(), TargetFrameRate);
+				ensureMsgf(FMath::IsNearlyZero(ResampledFrameTime.GetSubFrame()), TEXT("Incompatible resampling frame rate for animation sequence %s, frame remainder of %1.8f"), *GetName(), ResampledFrameTime.GetSubFrame());
+
+				NumberOfSampledFrames = ResampledFrameTime.FloorToFrame().Value;
+				NumberOfSampledKeys = NumberOfSampledFrames + 1;
 
-				const double SampleTime = TargetFrameRate.AsSeconds(FrameIndex);
+				for (int32 TrackIndex = 0; TrackIndex < ResampledAnimationTrackData.Num(); ++TrackIndex)
+				{
+					FBoneAnimationTrack& TrackData = ResampledAnimationTrackData[TrackIndex];
+					TrackData.InternalTrackData.PosKeys.SetNumUninitialized(NumberOfSampledKeys);
+					TrackData.InternalTrackData.RotKeys.SetNumUninitialized(NumberOfSampledKeys);
+					TrackData.InternalTrackData.ScaleKeys.SetNumUninitialized(NumberOfSampledKeys);
+				}
+				
+				FBlendedCurve Curve;
+				Curve.InitFrom(EvalContext.RequiredBones);
+				UE::Anim::FStackAttributeContainer AttributeContainer;
+				
+				for (int32 FrameIndex = 0; FrameIndex < NumberOfSampledKeys; ++FrameIndex)
+				{
+					UE::Anim::DataModel::FEvaluationContext EvaluationContext(FFrameTime(FrameIndex), TargetFrameRate, GetRetargetTransformsSourceName(), GetRetargetTransforms());
+					
+					FCompactPose Pose;
+					Pose.SetBoneContainer(&EvalContext.RequiredBones);
 
-				UE::Anim::GetBoneTransformFromModel(DataModel, Transform, TrackIndex, SampleTime, EAnimInterpolationType::Linear);
+					FAnimationPoseData PoseData(Pose, Curve, AttributeContainer);
+					DataModelInterface->Evaluate(PoseData, EvaluationContext);
 
-				TrackData.InternalTrackData.PosKeys[FrameIndex] = FVector3f(Transform.GetLocation());
-				TrackData.InternalTrackData.RotKeys[FrameIndex] = FQuat4f(Transform.GetRotation());
-				TrackData.InternalTrackData.ScaleKeys[FrameIndex] = FVector3f(Transform.GetScale3D());
-			}
+					for (int32 TrackIndex = 0; TrackIndex < ResampledAnimationTrackData.Num(); ++TrackIndex)
+					{
+						FBoneAnimationTrack& TrackData = ResampledAnimationTrackData[TrackIndex];
+						const FTransform& Transform = Pose[EvalContext.RequiredBones.GetCompactPoseIndexFromSkeletonIndex(TrackData.BoneTreeIndex)];
+
+						TrackData.InternalTrackData.PosKeys[FrameIndex] = FVector3f(Transform.GetLocation());
+						TrackData.InternalTrackData.RotKeys[FrameIndex] = FQuat4f(Transform.GetRotation());
+						TrackData.InternalTrackData.ScaleKeys[FrameIndex] = FVector3f(Transform.GetScale3D());
+					}
+				}
+			}	
+			
+			
 		}
 	}
 }
-}
 
 void UAnimSequence::RecompressAnimationData()
 {
+	QUICK_SCOPE_CYCLE_COUNTER(STAT_RecompressAnimationData);
 	PRAGMA_DISABLE_DEPRECATION_WARNINGS
 	ClearCompressedBoneData();
 	PRAGMA_ENABLE_DEPRECATION_WARNINGS
diff --git a/Engine/Source/Runtime/Engine/Private/Animation/AnimSequenceBase.cpp b/Engine/Source/Runtime/Engine/Private/Animation/AnimSequenceBase.cpp
--- a/Engine/Source/Runtime/Engine/Private/Animation/AnimSequenceBase.cpp
+++ b/Engine/Source/Runtime/Engine/Private/Animation/AnimSequenceBase.cpp
@@ -13,6 +13,7 @@
 #include "Animation/AnimationPoseData.h"
 #include "Animation/AnimSequenceHelpers.h"
 #include "Animation/MirrorDataTable.h"
+#include "Animation/AnimData/IAnimationDataModel.h"
 #include "Animation/AnimData/AnimDataModel.h"
 #include "Modules/ModuleManager.h"
 
@@ -21,6 +22,7 @@
 #include "Modules/ModuleManager.h"
 #endif // WITH_EDITOR
 
+#include "Animation/AnimationSettings.h"
 #include "UObject/UE5MainStreamObjectVersion.h"
 
 DEFINE_LOG_CATEGORY(LogAnimMarkerSync);
@@ -48,19 +50,6 @@ UAnimSequenceBase::UAnimSequenceBase(const FObjectInitializer& ObjectInitializer
 #endif // WITH_EDITOR
 }
 
-void UAnimSequenceBase::SetSequenceLength(float NewLength)
-{
-#if WITH_EDITOR
-	// Editor time we update the source data
-	Controller->SetPlayLength(NewLength);
-#else 
-	// In case this is called during runtime, update the sequence length directly. Current use-case is runtime created Montages.
-	PRAGMA_DISABLE_DEPRECATION_WARNINGS
-	SequenceLength = NewLength;
-	PRAGMA_ENABLE_DEPRECATION_WARNINGS
-#endif
-}
-
 bool UAnimSequenceBase::IsPostLoadThreadSafe() const
 {
 	return false;	// PostLoad is not thread safe because of the call to VerifyCurveNames() (calling USkeleton::VerifySmartName) that can mutate a shared map in the skeleton.
@@ -68,11 +57,15 @@ bool UAnimSequenceBase::IsPostLoadThreadSafe() const
 
 void UAnimSequenceBase::PostLoad()
 {
+	QUICK_SCOPE_CYCLE_COUNTER(STAT_AnimSequenceBase_PostLoad);
 #if WITH_EDITORONLY_DATA
 	if (!HasAnyFlags(EObjectFlags::RF_ClassDefaultObject))
 	{
+		LLM_SCOPE(ELLMTag::Animation);
+		
 		auto PreloadSkeletonAndVerifyCurves = [this]()
 		{
+			QUICK_SCOPE_CYCLE_COUNTER(STAT_PreloadSkeletonAndVerifyCurves);
 			if (USkeleton* MySkeleton = GetSkeleton())
 			{
 				if (FLinkerLoad* SkeletonLinker = MySkeleton->GetLinker())
@@ -87,31 +80,80 @@ void UAnimSequenceBase::PostLoad()
 				PRAGMA_ENABLE_DEPRECATION_WARNINGS
 			}
 		};
-		
+
 		if(ShouldDataModelBeValid())
 		{
-			const bool bRequiresModelCreation = DataModel == nullptr;
-			const bool bRequiresModelPopulation = GetLinkerCustomVersion(FUE5MainStreamObjectVersion::GUID) < FUE5MainStreamObjectVersion::IntroducingAnimationDataModel;
-			checkf(bRequiresModelPopulation || DataModel != nullptr, TEXT("Invalid Animation Sequence base state, no data model found past upgrade object version"));
-
-			// Construct a new UAnimDataModel instance
-			if(bRequiresModelCreation)
-			{
-				CreateModel();
-			}
-
-			ValidateModel();
-			GetController();
-			BindToModelModificationEvent();
-
-			PreloadSkeletonAndVerifyCurves();
-
-			GetController();
-			if (bRequiresModelPopulation)
-			{
-				bPopulatingDataModel = true;
-				PopulateModel();
-				bPopulatingDataModel = false;
+		    const UClass* TargetDataModelClass = UE::Anim::DataModel::IAnimationDataModels::FindClassForAnimationAsset(this);
+		    const bool bRequiresModelCreation = DataModelInterface == nullptr || DataModelInterface.GetObject()->GetClass() != TargetDataModelClass;
+
+		    TScriptInterface<IAnimationDataModel> CachedDataModelInterface = DataModelInterface;
+		    
+		    const bool bRequiresModelPopulation = GetLinkerCustomVersion(FUE5MainStreamObjectVersion::GUID) < FUE5MainStreamObjectVersion::IntroducingAnimationDataModel;
+		    PRAGMA_DISABLE_DEPRECATION_WARNINGS
+		    checkf(bRequiresModelPopulation || DataModel != nullptr || DataModelInterface != nullptr, TEXT("Invalid Animation Sequence base state, no data model found past upgrade object version"));
+		    PRAGMA_ENABLE_DEPRECATION_WARNINGS
+    
+		    // Construct a new IAnimationDataModel instance
+		    if(bRequiresModelCreation)
+		    {
+		    	QUICK_SCOPE_CYCLE_COUNTER(STAT_CreateDataModel);
+			    CreateModel();
+		    }
+
+		    ValidateModel();
+		    GetController();
+		    BindToModelModificationEvent();
+
+		    PreloadSkeletonAndVerifyCurves();
+
+		    if (bRequiresModelPopulation || bRequiresModelCreation)
+		    {
+		    	QUICK_SCOPE_CYCLE_COUNTER(STAT_Populate);
+			    const bool bDoNotTransactAction = false;
+    
+		    	Controller->OpenBracket(LOCTEXT("UAnimSequenceBase::PostLoad_PopulatingModelInterface","Populating Animation Data Model Interface"), bDoNotTransactAction);
+			    
+		    	bPopulatingDataModel = true;
+		    	
+		    	Controller->InitializeModel();
+    
+		    	PRAGMA_DISABLE_DEPRECATION_WARNINGS
+				// In case a data model has already been created populate the new one with its data
+				if (DataModel != nullptr)
+				{					
+					QUICK_SCOPE_CYCLE_COUNTER(STAT_PopulateWithExistingModel);
+					DataModel->PostLoad();
+					DataModel->ConditionalPostLoadSubobjects();
+					PopulateWithExistingModel(DataModel.Get());
+					DataModel = nullptr;
+					Controller->NotifyPopulated();
+				}
+		    	PRAGMA_ENABLE_DEPRECATION_WARNINGS
+				// If switching to a different DataModelInterface implementation, copy the data from the existing one
+				else if (CachedDataModelInterface != nullptr)
+				{
+					QUICK_SCOPE_CYCLE_COUNTER(STAT_PopulateWithExistingModel);
+					PopulateWithExistingModel(CachedDataModelInterface);
+					Controller->NotifyPopulated();
+				}
+		    	// Otherwise upgrade this animation asset to be model-based
+		    	else
+		    	{
+		    		QUICK_SCOPE_CYCLE_COUNTER(STAT_PopulateModel);
+					PopulateModel();
+		    		Controller->NotifyPopulated();
+		    	}
+		    	bPopulatingDataModel = false;
+		    	Controller->CloseBracket();
+		    }
+		    else
+		    {
+		    	// Fix-up to ensure correct curves are used for compression
+		    	PRAGMA_DISABLE_DEPRECATION_WARNINGS
+				RawCurveData.FloatCurves = DataModelInterface->GetCurveData().FloatCurves;
+		    	RawCurveData.TransformCurves = DataModelInterface->GetCurveData().TransformCurves;
+		    	RawCurveData.RemoveRedundantKeys(0.f);
+		    	PRAGMA_ENABLE_DEPRECATION_WARNINGS
 			}
 		}
 		else
@@ -144,18 +186,30 @@ void UAnimSequenceBase::PostLoad()
 #endif
 	RefreshCacheData();
 
+#if WITH_EDITOR
+	if (!GetPackage()->GetHasBeenEndLoaded())
+	{
+		FCoreUObjectDelegates::OnEndLoadPackage.AddUObject(this, &UAnimSequenceBase::OnEndLoadPackage);
+	}
+	else
+	{
+		OnAnimModelLoaded();
+	}	
+#endif
+
 	if(USkeleton* MySkeleton = GetSkeleton())
 	{
-		const bool bDoNotTransactAction = false;
 #if WITH_EDITOR
-		if (IsDataModelValid())
+		if (ensure(IsDataModelValid()))
 		{
+			const bool bDoNotTransactAction = false;
 			if (GetLinkerCustomVersion(FFortniteMainBranchObjectVersion::GUID) < FFortniteMainBranchObjectVersion::FixUpNoneNameAnimationCurves)
 			{
-				
+				QUICK_SCOPE_CYCLE_COUNTER(STAT_FixUpNoneNameAnimationCurves);
+				QUICK_SCOPE_CYCLE_COUNTER(STAT_OnAnimModelLoaded);
 				Controller->OpenBracket(LOCTEXT("FFortniteMainBranchObjectVersion::FixUpNoneNameAnimationCurves_Bracket","FFortniteMainBranchObjectVersion::FixUpNoneNameAnimationCurves"), bDoNotTransactAction);
 				{
-					const TArray<FFloatCurve>& FloatCurves = DataModel->GetFloatCurves();
+					const TArray<FFloatCurve>& FloatCurves = DataModelInterface->GetFloatCurves();
 					for (int32 Index = 0; Index < FloatCurves.Num(); ++Index)
 					{
 						const FFloatCurve& Curve = FloatCurves[Index];
@@ -174,23 +228,7 @@ void UAnimSequenceBase::PostLoad()
 				}
 				Controller->CloseBracket(bDoNotTransactAction);
 			}
-			else
-			{
-				const TArray<FFloatCurve>& FloatCurves = DataModel->GetFloatCurves();
-				for (int32 Index = 0; Index < FloatCurves.Num(); ++Index)
-				{
-					const FFloatCurve& Curve = FloatCurves[Index];
-					ensureMsgf(Curve.Name.DisplayName != NAME_None, TEXT("[AnimSequencer %s] has invalid curve name."), *GetFullName());
-				}
-			}
-
-			Controller->FindOrAddCurveNamesOnSkeleton(MySkeleton, ERawCurveTrackTypes::RCT_Float, bDoNotTransactAction);
-			Controller->FindOrAddCurveNamesOnSkeleton(MySkeleton, ERawCurveTrackTypes::RCT_Transform, bDoNotTransactAction);
 		}
-#else
-		PRAGMA_DISABLE_DEPRECATION_WARNINGS
-		VerifyCurveNames<FFloatCurve>(*MySkeleton, USkeleton::AnimCurveMappingName, RawCurveData.FloatCurves);
-		PRAGMA_ENABLE_DEPRECATION_WARNINGS
 #endif
 
 		// this should continue to add if skeleton hasn't been saved either 
@@ -198,9 +236,10 @@ void UAnimSequenceBase::PostLoad()
 		if (GetLinkerCustomVersion(FFrameworkObjectVersion::GUID) < FFrameworkObjectVersion::MoveCurveTypesToSkeleton
 			|| MySkeleton->GetLinkerCustomVersion(FFrameworkObjectVersion::GUID) < FFrameworkObjectVersion::MoveCurveTypesToSkeleton)
 		{
-
+			QUICK_SCOPE_CYCLE_COUNTER(STAT_MoveCurveTypesToSkeleton);
 #if WITH_EDITOR
-			const TArray<FFloatCurve>& FloatCurves = DataModel->GetFloatCurves();
+			// This is safe as the data model will have been created during this PostLoad call
+			const TArray<FFloatCurve>& FloatCurves = DataModelInterface->GetFloatCurves();
 #else
 			PRAGMA_DISABLE_DEPRECATION_WARNINGS
 			const TArray<FFloatCurve>& FloatCurves = RawCurveData.FloatCurves;
@@ -227,10 +266,8 @@ void UAnimSequenceBase::PostDuplicate(EDuplicateMode::Type DuplicateMode)
 	Super::PostDuplicate(DuplicateMode);
 
 #if WITH_EDITOR
-	if (IsDataModelValid())
-	{
-		BindToModelModificationEvent();
-	}
+	checkf(DataModelInterface.GetObject()->GetOuter() == this, TEXT("Animation Data Model interface has incorrect outer, expected %s - found %f"), *this->GetName(), *DataModelInterface.GetObject()->GetOuter()->GetName());
+	BindToModelModificationEvent();
 #endif // WITH_EDITOR
 }
 
@@ -699,9 +736,9 @@ int32 UAnimSequenceBase::GetNumberOfSampledKeys() const
 	return GetSamplingFrameRate().AsFrameTime(GetPlayLength()).RoundToFrame().Value;
 }
 
-const FFrameRate& UAnimSequenceBase::GetSamplingFrameRate() const
+FFrameRate UAnimSequenceBase::GetSamplingFrameRate() const
 {
-	static const FFrameRate DefaultFrameRate = FFrameRate(DEFAULT_SAMPLERATE, 1);
+	static const FFrameRate DefaultFrameRate = UAnimationSettings::Get()->GetDefaultFrameRate();
 	return DefaultFrameRate;
 }
 
@@ -788,9 +825,9 @@ void UAnimSequenceBase::GetAssetRegistryTags(TArray<FAssetRegistryTag>& OutTags)
 	// as possible.
 	FString CurveNameList;
 
-	if (DataModel)
+	if (DataModelInterface.GetObject())
 	{
-		for(const FFloatCurve& Curve : DataModel->GetFloatCurves())
+		for(const FFloatCurve& Curve : DataModelInterface->GetFloatCurves())
 		{
 			CurveNameList += FString::Printf(TEXT("%s%s"), *Curve.Name.DisplayName.ToString(), *USkeleton::CurveTagDelimiter);
 		}
@@ -800,14 +837,14 @@ void UAnimSequenceBase::GetAssetRegistryTags(TArray<FAssetRegistryTag>& OutTags)
 
 uint8* UAnimSequenceBase::FindNotifyPropertyData(int32 NotifyIndex, FArrayProperty*& ArrayProperty)
 {
-	// initialize to NULL
-	ArrayProperty = NULL;
+	// initialize to nullptr
+	ArrayProperty = nullptr;
 
 	if(Notifies.IsValidIndex(NotifyIndex))
 	{
 		return FindArrayProperty(TEXT("Notifies"), ArrayProperty, NotifyIndex);
 	}
-	return NULL;
+	return nullptr;
 }
 
 uint8* UAnimSequenceBase::FindArrayProperty(const TCHAR* PropName, FArrayProperty*& ArrayProperty, int32 ArrayIndex)
@@ -831,7 +868,7 @@ uint8* UAnimSequenceBase::FindArrayProperty(const TCHAR* PropName, FArrayPropert
 			return ArrayHelper.GetRawPtr(ArrayIndex);
 		}
 	}
-	return NULL;
+	return nullptr;
 }
 
 void UAnimSequenceBase::RefreshParentAssetData()
@@ -869,17 +906,17 @@ void UAnimSequenceBase::RefreshParentAssetData()
 	ValidateModel();
 	Controller->OpenBracket(LOCTEXT("RefreshParentAssetData_Bracket", "Refreshing Parent Asset Data"));
 	{
-		const UAnimDataModel* ParentDataModel = ParentSeqBase->GetDataModel();
+		const IAnimationDataModel* ParentDataModel = ParentSeqBase->GetDataModel();
 
-		if (!FMath::IsNearlyEqual(DataModel->GetPlayLength(), ParentDataModel->GetPlayLength()))
+		if (DataModelInterface->GetNumberOfFrames() != ParentDataModel->GetNumberOfFrames())
 		{
-			Controller->SetPlayLength(ParentDataModel->GetPlayLength());
+			Controller->SetNumberOfFrames(ParentDataModel->GetNumberOfFrames());
 		}
 		
 		Controller->RemoveAllCurvesOfType(ERawCurveTrackTypes::RCT_Float);
 		for (const FFloatCurve& FloatCurve : ParentDataModel->GetFloatCurves())
 		{
-			const FAnimationCurveIdentifier CurveId(FloatCurve.Name, ERawCurveTrackTypes::RCT_Float);
+			const FAnimationCurveIdentifier CurveId = UAnimationCurveIdentifierExtensions::FindCurveIdentifier(GetSkeleton(), FloatCurve.Name.DisplayName, ERawCurveTrackTypes::RCT_Float);
 			Controller->AddCurve(CurveId, FloatCurve.GetCurveTypeFlags());
 			Controller->SetCurveKeys(CurveId, FloatCurve.FloatCurve.GetConstRefOfKeys());
 		}
@@ -890,11 +927,19 @@ void UAnimSequenceBase::RefreshParentAssetData()
 			const FAnimationCurveIdentifier CurveId(TransformCurve.Name, ERawCurveTrackTypes::RCT_Transform);
 			Controller->AddCurve(CurveId, TransformCurve.GetCurveTypeFlags());
 
-			TArray<FTransform> Values;
-			TArray<float> Times;
-			TransformCurve.GetKeys(Times, Values);
-
-			Controller->SetTransformCurveKeys(CurveId, Values, Times);
+			// Set each individual channel rich curve keys, to account for any custom tangents etc.
+			for (int32 SubCurveIndex = 0; SubCurveIndex < 3; ++SubCurveIndex)
+			{
+				const ETransformCurveChannel Channel = static_cast<ETransformCurveChannel>(SubCurveIndex);
+				const FVectorCurve* VectorCurve = TransformCurve.GetVectorCurveByIndex(SubCurveIndex);
+				for (int32 ChannelIndex = 0; ChannelIndex < 3; ++ChannelIndex)
+				{
+					const EVectorCurveChannel Axis = static_cast<EVectorCurveChannel>(ChannelIndex);
+					FAnimationCurveIdentifier TargetCurveIdentifier = CurveId;
+					UAnimationCurveIdentifierExtensions::GetTransformChildCurveIdentifier(TargetCurveIdentifier, Channel, Axis);
+					Controller->SetCurveKeys(TargetCurveIdentifier, VectorCurve->FloatCurves[ChannelIndex].GetConstRefOfKeys());
+				}
+			}
 		}
 	}
 	Controller->CloseBracket();
@@ -955,7 +1000,7 @@ bool UAnimSequenceBase::HasCurveData(SmartName::UID_Type CurveUID, bool bForceUs
 #if WITH_EDITOR
 	if (IsDataModelValid())
 	{
-		return DataModel->FindFloatCurve(FAnimationCurveIdentifier(CurveUID, ERawCurveTrackTypes::RCT_Float)) != nullptr;
+		return DataModelInterface->FindFloatCurve(FAnimationCurveIdentifier(CurveUID, ERawCurveTrackTypes::RCT_Float)) != nullptr;
 	}
 #endif
 	
@@ -994,7 +1039,7 @@ void UAnimSequenceBase::HandleAssetPlayerTickedInternal(FAnimAssetTickContext &C
 }
 
 #if WITH_EDITOR
-void UAnimSequenceBase::OnModelModified(const EAnimDataModelNotifyType& NotifyType, UAnimDataModel* Model, const FAnimDataModelNotifPayload& Payload)
+void UAnimSequenceBase::OnModelModified(const EAnimDataModelNotifyType& NotifyType, IAnimationDataModel* Model, const FAnimDataModelNotifPayload& Payload)
 {
 	NotifyCollector.Handle(NotifyType);
 
@@ -1002,20 +1047,25 @@ void UAnimSequenceBase::OnModelModified(const EAnimDataModelNotifyType& NotifyTy
 	InitializeNotifyTrack();
 	PRAGMA_ENABLE_DEPRECATION_WARNINGS
 
-	auto HandleLengthChange = [&](float NewLength, float OldLength, float T0, float T1)
+	auto HandleNumberOfFramesChange = [&](int32 NewNumFrames, int32 OldNumFrames, int32 Frame0, int32 Frame1)
 	{
 		if (bPopulatingDataModel)
 		{
 			return;
 		}
+
+		const FFrameRate& ModelFrameRate = DataModelInterface->GetFrameRate();
+		const float T0 = ModelFrameRate.AsSeconds(Frame0);
+		const float T1 = ModelFrameRate.AsSeconds(Frame1);
+		const float NewLength = ModelFrameRate.AsSeconds(NewNumFrames);
 		
-		const float StartRemoveTime = T0;
-		const float EndRemoveTime = T1;
+		const float StartRemoveTime = ModelFrameRate.AsSeconds(Frame0);
+		const float EndRemoveTime = ModelFrameRate.AsSeconds(Frame1);
 
 		// Total time value for frames that were removed
 		const float Duration = T1 - T0;
 
-		if (NewLength > OldLength)
+		if (NewNumFrames > OldNumFrames)
 		{
 			const float InsertTime = T0;
 			for (FAnimNotifyEvent& Notify : Notifies)
@@ -1032,8 +1082,8 @@ void UAnimSequenceBase::OnModelModified(const EAnimDataModelNotifyType& NotifyTy
 				// if state, make sure to adjust end time
 				if (Notify.NotifyStateClass)
 				{
-					float NotifyDuration = Notify.GetDuration();
-					float NotifyEnd = CurrentTime + NotifyDuration;
+					const float NotifyDuration = Notify.GetDuration();
+					const float NotifyEnd = CurrentTime + NotifyDuration;
 					if (NotifyEnd >= InsertTime)
 					{
 						NewDuration = NotifyDuration + Duration;
@@ -1044,7 +1094,7 @@ void UAnimSequenceBase::OnModelModified(const EAnimDataModelNotifyType& NotifyTy
 					}
 				}
 
-				// Shift out notify by the time alloted for the inserted keys
+				// Shift out notify by the time allotted for the inserted keys
 				if (CurrentTime >= InsertTime)
 				{
 					CurrentTime += Duration;
@@ -1067,7 +1117,7 @@ void UAnimSequenceBase::OnModelModified(const EAnimDataModelNotifyType& NotifyTy
 				}
 			}
 		}
-		else if (NewLength < OldLength)
+		else if (NewNumFrames < OldNumFrames)
 		{
 			// re-locate notifies
 			for (FAnimNotifyEvent& Notify : Notifies)
@@ -1084,15 +1134,15 @@ void UAnimSequenceBase::OnModelModified(const EAnimDataModelNotifyType& NotifyTy
 				// if state, make sure to adjust end time
 				if (Notify.NotifyStateClass)
 				{
-					float NotifyDuration = Notify.GetDuration();
-					float NotifyEnd = CurrentTime + NotifyDuration;
+					const float NotifyDuration = Notify.GetDuration();
+					const float NotifyEnd = CurrentTime + NotifyDuration;
 					NewDuration = NotifyDuration;
 
 					// If Notify is inside of the trimmed time frame, zero out the duration
 					if (CurrentTime >= StartRemoveTime && NotifyEnd <= EndRemoveTime)
 					{
 						// small number @todo see if there is define for this
-						NewDuration = DataModel->GetFrameRate().AsInterval();
+						NewDuration = DataModelInterface->GetFrameRate().AsInterval();
 					}
 					// If Notify overlaps trimmed time frame, remove trimmed duration
 					else if (CurrentTime < EndRemoveTime && NotifyEnd > EndRemoveTime)
@@ -1100,7 +1150,7 @@ void UAnimSequenceBase::OnModelModified(const EAnimDataModelNotifyType& NotifyTy
 						NewDuration = NotifyEnd - Duration - CurrentTime;
 					}
 
-					NewDuration = FMath::Max(NewDuration, (float)DataModel->GetFrameRate().AsInterval());
+					NewDuration = FMath::Max(NewDuration, (float)DataModelInterface->GetFrameRate().AsInterval());
 				}
 
 				if (CurrentTime >= StartRemoveTime && CurrentTime <= EndRemoveTime)
@@ -1126,16 +1176,15 @@ void UAnimSequenceBase::OnModelModified(const EAnimDataModelNotifyType& NotifyTy
 					Notify.TriggerTimeOffset = GetTriggerTimeOffsetForType(EAnimEventTriggerOffsets::OffsetBefore);
 				}
 			}
-		}	
+		}
 	};
 
-	// Source animation curve data can be coarse, so copy the data down and remove any redundant keys
+	// Copy any float/transform curves from the model into RawCurveData, as it is used at runtime for AnimMontage/Composite(s)
 	auto CopyCurvesFromModel = [this]()
 	{
 		PRAGMA_DISABLE_DEPRECATION_WARNINGS
-		RawCurveData.FloatCurves = DataModel->GetFloatCurves();
-		RawCurveData.TransformCurves = DataModel->GetTransformCurves();
-		RawCurveData.RemoveRedundantKeys();
+		RawCurveData.FloatCurves = DataModelInterface->GetCurveData().FloatCurves;
+		RawCurveData.TransformCurves = DataModelInterface->GetCurveData().TransformCurves;
 		PRAGMA_ENABLE_DEPRECATION_WARNINGS
 	};
 
@@ -1144,14 +1193,14 @@ void UAnimSequenceBase::OnModelModified(const EAnimDataModelNotifyType& NotifyTy
 		case EAnimDataModelNotifyType::SequenceLengthChanged:
 		{
 			const FSequenceLengthChangedPayload& TypedPayload = Payload.GetPayload<FSequenceLengthChangedPayload>();
-			const float PreviousSequenceLength = TypedPayload.PreviousLength;
-			const float CurrentSequenceLength = Model->GetPlayLength();
+			const FFrameNumber PreviousNumberOfFrames = TypedPayload.PreviousNumberOfFrames;
+			const int32 CurrentNumberOfFrames = Model->GetNumberOfFrames();
 
 			PRAGMA_DISABLE_DEPRECATION_WARNINGS
-			SequenceLength = CurrentSequenceLength;
+			SequenceLength = Model->GetFrameRate().AsSeconds(CurrentNumberOfFrames);
 			PRAGMA_ENABLE_DEPRECATION_WARNINGS
 
-			HandleLengthChange(CurrentSequenceLength, PreviousSequenceLength, TypedPayload.T0, TypedPayload.T1);
+			HandleNumberOfFramesChange(CurrentNumberOfFrames, PreviousNumberOfFrames.Value, TypedPayload.Frame0.Value, TypedPayload.Frame1.Value);
 
 			// Ensure that we only clamp the notifies at the end of a bracket
 			if (NotifyCollector.IsNotWithinBracket())
@@ -1250,9 +1299,14 @@ void UAnimSequenceBase::OnModelModified(const EAnimDataModelNotifyType& NotifyTy
 	}
 }
 
-UAnimDataModel* UAnimSequenceBase::GetDataModel() const
+IAnimationDataModel* UAnimSequenceBase::GetDataModel() const
+{
+	return DataModelInterface.GetInterface();
+}
+
+TScriptInterface<IAnimationDataModel> UAnimSequenceBase::GetDataModelInterface() const
 {
-	return DataModel;
+	return DataModelInterface;
 }
 
 IAnimationDataController& UAnimSequenceBase::GetController()
@@ -1261,12 +1315,12 @@ IAnimationDataController& UAnimSequenceBase::GetController()
 
 	if(Controller == nullptr)
 	{
-		IAnimationDataControllerModule& ControllerModule = FModuleManager::Get().GetModuleChecked<IAnimationDataControllerModule>("AnimationDataController");
-		Controller = ControllerModule.GetController();
-		Controller->SetModel(DataModel);
+		Controller = DataModelInterface->GetController();
+		checkf(Controller, TEXT("Failed to create AnimationDataController"));
+		Controller->SetModel(DataModelInterface);
 	}
 
-	ensureAlways(Controller->GetModel() == DataModel);
+	ensureAlways(Controller->GetModelInterface() == DataModelInterface);
 
 	return *Controller;
 }
@@ -1277,51 +1331,68 @@ void UAnimSequenceBase::PopulateModel()
 	
 	// Make a copy of the current data, to mitigate any changes due to notify callbacks
 	const FFrameRate FrameRate = GetSamplingFrameRate();
-	const float PlayLength = GetPlayLength();
 	PRAGMA_DISABLE_DEPRECATION_WARNINGS
+	const float PlayLength = SequenceLength;
 	const FRawCurveTracks CurveData = RawCurveData;
 	PRAGMA_ENABLE_DEPRECATION_WARNINGS
 		
 	IAnimationDataController::FScopedBracket ScopedBracket(Controller, LOCTEXT("UAnimSequenceBase::PopulateModel_Bracket", "Generating Animation Model Data for Animation Sequence Base"));
-	Controller->SetPlayLength(PlayLength);
 	Controller->SetFrameRate(FrameRate);
+	Controller->SetNumberOfFrames(Controller->ConvertSecondsToFrameNumber(PlayLength));
 
 	USkeleton* TargetSkeleton = GetSkeleton();	
 	UE::Anim::CopyCurveDataToModel(CurveData, TargetSkeleton,  *Controller);	
 }
 
+void UAnimSequenceBase::PopulateWithExistingModel(TScriptInterface<IAnimationDataModel> ExistingDataModel)
+{
+	Controller->PopulateWithExistingModel(ExistingDataModel);
+}
+
 void UAnimSequenceBase::BindToModelModificationEvent()
 {
 	ValidateModel();
-	DataModel->GetModifiedEvent().RemoveAll(this);
-	DataModel->GetModifiedEvent().AddUObject(this, &UAnimSequenceBase::OnModelModified);
+	DataModelInterface->GetModifiedEvent().RemoveAll(this);
+	DataModelInterface->GetModifiedEvent().AddUObject(this, &UAnimSequenceBase::OnModelModified);
 }
 
-void UAnimSequenceBase::CopyDataModel(const UAnimDataModel* ModelToDuplicate)
+void UAnimSequenceBase::CopyDataModel(const TScriptInterface<IAnimationDataModel>& ModelToDuplicate)
 {
 	checkf(ModelToDuplicate != nullptr, TEXT("Invalidate data model %s"), *GetFullName());
-	if (DataModel)
+	if (ModelToDuplicate)
 	{
-		DataModel->GetModifiedEvent().RemoveAll(this);
+		ModelToDuplicate->GetModifiedEvent().RemoveAll(this);
 	}
 
-	DataModel = DuplicateObject(ModelToDuplicate, this);
-	Controller->SetModel(DataModel);
+	DataModelInterface = DuplicateObject(ModelToDuplicate.GetObject(), this);
+
+	if(Controller)
+	{
+		Controller->SetModel(DataModelInterface);
+	}
+	else
+	{
+		GetController();
+	}
 
 	BindToModelModificationEvent();
 }
 
 void UAnimSequenceBase::CreateModel()
 {
-	checkf(DataModel == nullptr, TEXT("Invalid attempt to override the existing data model %s"), *GetFullName());
-	DataModel = NewObject<UAnimDataModel>(this, FName(TEXT("AnimationDataModel")));
-		
+	const UClass* TargetClass = UE::Anim::DataModel::IAnimationDataModels::FindClassForAnimationAsset(this);
+	checkf(TargetClass != nullptr, TEXT("Unable to find valid AnimationDataModel class"));
+
+	checkf(!DataModelInterface || DataModelInterface.GetObject()->GetClass() != TargetClass, TEXT("Invalid attempt to override the existing data model %s"), *GetFullName());
+	UObject* ClassDataModel = NewObject<UObject>(this, TargetClass, TargetClass->GetFName(), RF_Standalone);
+	DataModelInterface = ClassDataModel;
+
 	BindToModelModificationEvent();
 }
 
 void UAnimSequenceBase::ValidateModel() const
 {
-	checkf(DataModel != nullptr, TEXT("Invalid AnimSequenceBase state (%s), no data model found"), *GetPathName());
+	checkf(DataModelInterface != nullptr, TEXT("Invalid AnimSequenceBase state (%s), no data model found"), *GetPathName());
 }
 
 bool UAnimSequenceBase::ShouldDataModelBeValid() const
@@ -1337,3 +1408,37 @@ bool UAnimSequenceBase::ShouldDataModelBeValid() const
 
 #undef LOCTEXT_NAMESPACE 
 
+#if WITH_EDITOR
+void UAnimSequenceBase::OnEndLoadPackage(const FEndLoadPackageContext& Context)
+{
+	if (!GetPackage()->GetHasBeenEndLoaded())
+	{
+		return;
+	}
+
+	FCoreUObjectDelegates::OnEndLoadPackage.RemoveAll(this);
+
+	OnAnimModelLoaded();
+}
+
+void UAnimSequenceBase::OnAnimModelLoaded()
+{
+	QUICK_SCOPE_CYCLE_COUNTER(STAT_OnAnimModelLoaded);
+	if(USkeleton* MySkeleton = GetSkeleton())
+	{		
+		if (IsDataModelValid())
+		{
+			const TArray<FFloatCurve>& FloatCurves = DataModelInterface->GetFloatCurves();
+			for (int32 Index = 0; Index < FloatCurves.Num(); ++Index)
+			{
+				const FFloatCurve& Curve = FloatCurves[Index];
+				ensureMsgf(Curve.Name.DisplayName != NAME_None, TEXT("[AnimSequencer %s] has invalid curve name."), *GetFullName());
+			}
+
+			constexpr bool bDoNotTransactAction = false;
+			Controller->FindOrAddCurveNamesOnSkeleton(MySkeleton, ERawCurveTrackTypes::RCT_Float, bDoNotTransactAction);
+			Controller->FindOrAddCurveNamesOnSkeleton(MySkeleton, ERawCurveTrackTypes::RCT_Transform, bDoNotTransactAction);
+		}
+	}
+}
+#endif // WITH_EDITOR
\ No newline at end of file
