**New**

diff --git a/Engine/Source/Runtime/Engine/Private/Animation/AnimSequenceHelpers.cpp b/Engine/Source/Runtime/Engine/Private/Animation/AnimSequenceHelpers.cpp
--- a/Engine/Source/Runtime/Engine/Private/Animation/AnimSequenceHelpers.cpp
+++ b/Engine/Source/Runtime/Engine/Private/Animation/AnimSequenceHelpers.cpp
@@ -1,7 +1,7 @@
 // Copyright Epic Games, Inc. All Rights Reserved.
 
 #include "Animation/AnimSequenceHelpers.h"
-#include "Animation/AnimData/AnimDataModel.h"
+#include "Animation/AnimData/IAnimationDataModel.h"
 #include "Animation/AnimData/IAnimationDataController.h"
 
 #include "UObject/NameTypes.h"
@@ -13,11 +13,11 @@
 #include "AnimationRuntime.h"
 #include "AnimEncoding.h"
 #include "AnimationUtils.h"
-#include "Animation/AnimTypes.h"
+#include "Animation/AnimationPoseData.h"
 #include "Animation/AnimNotifies/AnimNotify.h"
 #include "Animation/AnimNotifies/AnimNotifyState.h"
 #include "Misc/MessageDialog.h"
-#include "Animation/AnimSequenceHelpers.h"
+#include "Animation/AttributesRuntime.h"
 
 #define LOCTEXT_NAMESPACE "AnimSequenceHelpers"
 
@@ -27,222 +27,32 @@ namespace UE
 namespace Anim
 {
 
-struct FRetargetTracking
-{
-	const FCompactPoseBoneIndex PoseBoneIndex;
-	const int32 SkeletonBoneIndex;
-
-	FRetargetTracking(const FCompactPoseBoneIndex InPoseBoneIndex, const int32 InSkeletonBoneIndex)
-		: PoseBoneIndex(InPoseBoneIndex), SkeletonBoneIndex(InSkeletonBoneIndex)
-	{
-	}
-};
-
-struct FBuildRawPoseScratchArea : public TThreadSingleton<FBuildRawPoseScratchArea>
-{
-	TArray<FRetargetTracking> RetargetTracking;
-	TArray<FVirtualBoneCompactPoseData> VirtualBoneCompactPoseData;
-};
-
-
 #if WITH_EDITOR
-FTransform ExtractTransformForKey(int32 Key, const FRawAnimSequenceTrack& TrackToExtract)
+void BuildPoseFromModel(const IAnimationDataModel* Model, FCompactPose& OutPose, const double Time, const EAnimInterpolationType& InterpolationType, const FName& RetargetSource, const TArray<FTransform>& RetargetTransforms)
 {
-	static const FVector DefaultScale3D = FVector(1.f);
-	const bool bHasScaleKey = TrackToExtract.ScaleKeys.Num() > 0;
+	FBlendedCurve TempCurve;
+	UE::Anim::FStackAttributeContainer TempAttributes;
+	TempCurve.InitFrom(OutPose.GetBoneContainer());
 
-	const int32 PosKeyIndex = FMath::Min(Key, TrackToExtract.PosKeys.Num() - 1);
-	const int32 RotKeyIndex = FMath::Min(Key, TrackToExtract.RotKeys.Num() - 1);
-	if (bHasScaleKey)
-	{
-		const int32 ScaleKeyIndex = FMath::Min(Key, TrackToExtract.ScaleKeys.Num() - 1);
-		return FTransform(FQuat(TrackToExtract.RotKeys[RotKeyIndex]), FVector(TrackToExtract.PosKeys[PosKeyIndex]), FVector(TrackToExtract.ScaleKeys[ScaleKeyIndex]));
-	}
-	else
-	{
-		return FTransform(FQuat(TrackToExtract.RotKeys[RotKeyIndex]), FVector(TrackToExtract.PosKeys[PosKeyIndex]), FVector(DefaultScale3D));
-	}
+	FAnimationPoseData PoseData(OutPose, TempCurve, TempAttributes);
+	BuildPoseFromModel(Model, PoseData, Time, InterpolationType, RetargetSource, RetargetTransforms);
 }
-
-template<bool bInterpolateT>
-void ExtractPoseFromModel(const TArray<FBoneAnimationTrack>& BoneAnimationTracks, const TMap<FName, FTransform>& OverrideBoneTransforms, FCompactPose& InOutPose, const int32 KeyIndex1, const int32 KeyIndex2, float Alpha, const USkeleton* SourceSkeleton)
-{
-	const int32 NumAnimationTracks = BoneAnimationTracks.Num();
-	const FBoneContainer& RequiredBones = InOutPose.GetBoneContainer();
-
-	TArray<UE::Anim::FRetargetTracking>& RetargetTracking = UE::Anim::FBuildRawPoseScratchArea::Get().RetargetTracking;
-	RetargetTracking.Reset(NumAnimationTracks);
-
-	TArray<FVirtualBoneCompactPoseData>& VBCompactPoseData = UE::Anim::FBuildRawPoseScratchArea::Get().VirtualBoneCompactPoseData;
-	VBCompactPoseData = RequiredBones.GetVirtualBoneCompactPoseData();
-
-	FCompactPose Key2Pose;
-	if (bInterpolateT)
-	{
-		Key2Pose.CopyBonesFrom(InOutPose);
-	}
-
-	const USkeleton* TargetSkeleton = RequiredBones.GetSkeletonAsset();
-	const FSkeletonRemapping* SkeletonRemapping = TargetSkeleton->GetSkeletonRemapping(SourceSkeleton);
-
-	for (const FBoneAnimationTrack& AnimationTrack : BoneAnimationTracks)
-	{
-		const int32 SourceSkeletonBoneIndex = AnimationTrack.BoneTreeIndex;
-		const int32 SkeletonBoneIndex = SkeletonRemapping ? SkeletonRemapping->GetTargetSkeletonBoneIndex(SourceSkeletonBoneIndex) : SourceSkeletonBoneIndex;
-
-		// not sure it's safe to assume that SkeletonBoneIndex can never be INDEX_NONE
-		if ((SkeletonBoneIndex != INDEX_NONE) && (SkeletonBoneIndex < MAX_BONES))
-		{
-			const FCompactPoseBoneIndex PoseBoneIndex = RequiredBones.GetCompactPoseIndexFromSkeletonIndex(SkeletonBoneIndex);
-			if (PoseBoneIndex != INDEX_NONE)
-			{
-				for (int32 Idx = 0; Idx < VBCompactPoseData.Num(); ++Idx)
-				{
-					FVirtualBoneCompactPoseData& VB = VBCompactPoseData[Idx];
-					if (PoseBoneIndex == VB.VBIndex)
-					{
-						// Remove this bone as we have written data for it (false so we dont resize allocation)
-						VBCompactPoseData.RemoveAtSwap(Idx, 1, false);
-						break; //Modified TArray so must break here
-					}
-				}
 				
-				// extract animation
-				const FRawAnimSequenceTrack& TrackToExtract = AnimationTrack.InternalTrackData;
-				const FTransform* OverrideTransform = OverrideBoneTransforms.Find(AnimationTrack.Name);
-				{
-					// Bail out (with rather wacky data) if data is empty for some reason.
-					if (TrackToExtract.PosKeys.Num() == 0 || TrackToExtract.RotKeys.Num() == 0)
-					{
-						InOutPose[PoseBoneIndex].SetIdentity();
-
-						if (bInterpolateT)
-						{
-							Key2Pose[PoseBoneIndex].SetIdentity();
-						}
-					}
-					else
-					{
-						InOutPose[PoseBoneIndex] = UE::Anim::ExtractTransformForKey(KeyIndex1, TrackToExtract);
-
-						if (bInterpolateT)
-						{
-							Key2Pose[PoseBoneIndex] = UE::Anim::ExtractTransformForKey(KeyIndex2, TrackToExtract);
-						}
-					}
-				}
-
-				if (OverrideTransform)
-				{
-					InOutPose[PoseBoneIndex].SetRotation(InOutPose[PoseBoneIndex].GetRotation() * OverrideTransform->GetRotation());
-					InOutPose[PoseBoneIndex].TransformPosition(OverrideTransform->GetTranslation());
-					InOutPose[PoseBoneIndex].SetScale3D(InOutPose[PoseBoneIndex].GetScale3D() * OverrideTransform->GetScale3D());
-
-					if (bInterpolateT)
-					{
-						Key2Pose[PoseBoneIndex].SetRotation(Key2Pose[PoseBoneIndex].GetRotation() * OverrideTransform->GetRotation());
-						Key2Pose[PoseBoneIndex].TransformPosition(OverrideTransform->GetTranslation());
-						Key2Pose[PoseBoneIndex].SetScale3D(Key2Pose[PoseBoneIndex].GetScale3D() * OverrideTransform->GetScale3D());
-					}
-				}
-
-				RetargetTracking.Add(UE::Anim::FRetargetTracking(PoseBoneIndex, SkeletonBoneIndex));
-			}
-		}
-	}
-
-	//Build Virtual Bones
-	if (VBCompactPoseData.Num() > 0)
-	{
-		FCSPose<FCompactPose> CSPose1;
-		CSPose1.InitPose(InOutPose);
-
-		FCSPose<FCompactPose> CSPose2;
-		if (bInterpolateT)
-		{
-			CSPose2.InitPose(Key2Pose);
-		}
-
-		for (FVirtualBoneCompactPoseData& VB : VBCompactPoseData)
-		{
-			FTransform Source = CSPose1.GetComponentSpaceTransform(VB.SourceIndex);
-			FTransform Target = CSPose1.GetComponentSpaceTransform(VB.TargetIndex);
-			InOutPose[VB.VBIndex] = Target.GetRelativeTransform(Source);
-
-			if (bInterpolateT)
-			{
-				FTransform Source2 = CSPose2.GetComponentSpaceTransform(VB.SourceIndex);
-				FTransform Target2 = CSPose2.GetComponentSpaceTransform(VB.TargetIndex);
-				Key2Pose[VB.VBIndex] = Target2.GetRelativeTransform(Source2);
-			}
-		}
-	}
-
-	if (bInterpolateT)
-	{
-		for (FCompactPoseBoneIndex BoneIndex : InOutPose.ForEachBoneIndex())
-		{
-			InOutPose[BoneIndex].Blend(InOutPose[BoneIndex], Key2Pose[BoneIndex], Alpha);
-		}
-	}
-}
-
-void BuildPoseFromModel(const UAnimDataModel* Model, FCompactPose& OutPose, const float Time, const EAnimInterpolationType& InterpolationType, const FName& RetargetSource, const TArray<FTransform>& RetargetTransforms)
+void BuildPoseFromModel(const IAnimationDataModel* Model, FAnimationPoseData& OutPoseData, const double Time, const EAnimInterpolationType& InterpolationType, const FName& RetargetSource, const TArray<FTransform>& RetargetTransforms)
 {
+	FCompactPose& OutPose = OutPoseData.GetPose();
 	check(Model);
 
-	const int32 NumberOfKeys = Model->GetNumberOfKeys();
-	const float SequenceLength = Model->GetPlayLength();
-	const double FramesPerSecond = Model->GetFrameRate().AsDecimal();
-
-	// Generate keys to interpolate between
-	int32 KeyIndex1, KeyIndex2;
-	float Alpha;
-	FAnimationRuntime::GetKeyIndicesFromTime(KeyIndex1, KeyIndex2, Alpha, Time, NumberOfKeys, SequenceLength, FramesPerSecond);
-
-	if (InterpolationType == EAnimInterpolationType::Step)
-	{
-		// Force stepping between keys
-		Alpha = 0.f;
-	}
-
-	bool bShouldInterpolate = true;
-
-	if (Alpha < UE_KINDA_SMALL_NUMBER)
-	{
-		Alpha = 0.f;
-		bShouldInterpolate = false;
-	}
-	else if (Alpha > 1.f - UE_KINDA_SMALL_NUMBER)
-	{
-		bShouldInterpolate = false;
-		KeyIndex1 = KeyIndex2;
-	}
-
-	TMap<FName, FTransform> ActiveCurves;
-	if (!OutPose.GetBoneContainer().ShouldUseSourceData())
-	{
-		EvaluateTransformCurvesFromModel(Model, ActiveCurves, Time, 1.f);
-	}
-
-	const USkeleton* SourceSkeleton = Model->GetAnimationSequence()->GetSkeleton();
-	if (bShouldInterpolate)
-	{
-		ExtractPoseFromModel<true>(Model->GetBoneAnimationTracks(), ActiveCurves, OutPose, KeyIndex1, KeyIndex2, Alpha, SourceSkeleton);
-	}
-	else
-	{
-		ExtractPoseFromModel<false>(Model->GetBoneAnimationTracks(), ActiveCurves, OutPose, KeyIndex1, KeyIndex2, Alpha, SourceSkeleton);
-	}
+	OutPose.ResetToRefPose();
 
-	Retargeting::RetargetPose(OutPose, RetargetSource, RetargetTransforms);
+	const UE::Anim::DataModel::FEvaluationContext EvaluationContext(Time, Model->GetFrameRate(), RetargetSource, RetargetTransforms, InterpolationType);
+	Model->Evaluate(OutPoseData, EvaluationContext);
 }
 
-void EvaluateFloatCurvesFromModel(const UAnimDataModel* Model, FBlendedCurve& OutCurves, float Time)
+void EvaluateFloatCurvesFromModel(const IAnimationDataModel* Model, FBlendedCurve& OutCurves, double Time)
 {
 	check(Model);
 
-	//SCOPE_CYCLE_COUNTER(STAT_EvalRawCurveData);
 	if (OutCurves.NumValidCurveCount > 0)
 	{
 		// evaluate the curve data at the Time and add to Instance
@@ -257,7 +67,7 @@ void EvaluateFloatCurvesFromModel(const UAnimDataModel* Model, FBlendedCurve& Ou
 	}
 }
 
-void EvaluateTransformCurvesFromModel(const UAnimDataModel* Model, TMap<FName, FTransform>& OutCurves, float Time, float BlendWeight)
+void EvaluateTransformCurvesFromModel(const IAnimationDataModel* Model, TMap<FName, FTransform>& OutCurves, double Time, float BlendWeight)
 {
 	if (Model)
 	{
@@ -279,15 +89,15 @@ void EvaluateTransformCurvesFromModel(const UAnimDataModel* Model, TMap<FName, F
 	}
 }
 
-void GetBoneTransformFromModel(const UAnimDataModel* Model, FTransform& OutTransform, int32 TrackIndex, float Time, const EAnimInterpolationType& Interpolation)
+void GetBoneTransformFromModel(const IAnimationDataModel* Model, FTransform& OutTransform, int32 TrackIndex, double Time, const EAnimInterpolationType& Interpolation)
 {
 	const FBoneAnimationTrack& TrackData = Model->GetBoneTrackByIndex(TrackIndex);
 	const FRawAnimSequenceTrack& RawTrack = TrackData.InternalTrackData;
 	
-	FAnimationUtils::ExtractTransformFromTrack(Time, Model->GetNumberOfKeys(), Model->GetPlayLength(), RawTrack, Interpolation, OutTransform);
+	FAnimationUtils::ExtractTransformFromTrack(RawTrack, Time, Model->GetNumberOfKeys(), Model->GetPlayLength(), Interpolation, OutTransform);
 }
 
-void GetBoneTransformFromModel(const UAnimDataModel* Model, FTransform& OutTransform, int32 TrackIndex, int32 KeyIndex)
+void GetBoneTransformFromModel(const IAnimationDataModel* Model, FTransform& OutTransform, int32 TrackIndex, int32 KeyIndex)
 {
 	const FBoneAnimationTrack& TrackData = Model->GetBoneTrackByIndex(TrackIndex);
 	const FRawAnimSequenceTrack& RawTrack = TrackData.InternalTrackData;
@@ -395,6 +205,36 @@ void CopyCurveDataToModel(const FRawCurveTracks& CurveData, const USkeleton* Ske
 	else
 	{
 		OutTransform.SetScale3D(DefaultScale3D);
+	}	
+}
+	
+Retargeting::FRetargetingScope::FRetargetingScope(FCompactPose& ToRetargetPose, const DataModel::FEvaluationContext& InEvaluationContext)
+	: RetargetPose(ToRetargetPose),
+	EvaluationContext(InEvaluationContext),
+	RetargetTracking(FBuildRawPoseScratchArea::Get().RetargetTracking),
+	bShouldRetarget(!ToRetargetPose.GetBoneContainer().GetDisableRetargeting() && EvaluationContext.RetargetTransforms.Num())
+{
+	RetargetTracking.Reset();
+}
+
+void Retargeting::FRetargetingScope::AddTrackedBone(FCompactPoseBoneIndex CompactBoneIndex, int32 SkeletonBoneIndex) const
+{
+	if (bShouldRetarget)
+	{
+		RetargetTracking.Add({CompactBoneIndex, SkeletonBoneIndex});
+	}
+}
+
+Retargeting::FRetargetingScope::~FRetargetingScope()
+{
+	if (bShouldRetarget)
+	{
+		const FBoneContainer& RequiredBones = RetargetPose.GetBoneContainer();
+		const USkeleton* Skeleton = RequiredBones.GetSkeletonAsset();
+		for (const FRetargetTracking& RT : RetargetTracking)
+		{
+			FAnimationRuntime::RetargetBoneTransform(Skeleton, EvaluationContext.RetargetSource, EvaluationContext.RetargetTransforms, RetargetPose[RT.PoseBoneIndex], RT.SkeletonBoneIndex, RT.PoseBoneIndex, RequiredBones, false);
+		}
 	}
 }
 
@@ -405,11 +245,11 @@ void Retargeting::RetargetPose(FCompactPose& InOutPose, const FName& RetargetSou
 
 	if (!bDisableRetargeting && RetargetTransforms.Num())
 	{
-		const TArray<UE::Anim::FRetargetTracking>& RetargetTracking = UE::Anim::FBuildRawPoseScratchArea::Get().RetargetTracking;
+		const TArray<UE::Anim::Retargeting::FRetargetTracking>& RetargetTracking = UE::Anim::FBuildRawPoseScratchArea::Get().RetargetTracking;
 
 		USkeleton* Skeleton = RequiredBones.GetSkeletonAsset();
 
-		for (const UE::Anim::FRetargetTracking& RT : RetargetTracking)
+		for (const UE::Anim::Retargeting::FRetargetTracking& RT : RetargetTracking)
 		{
 			FAnimationRuntime::RetargetBoneTransform(Skeleton, RetargetSource, RetargetTransforms, InOutPose[RT.PoseBoneIndex], RT.SkeletonBoneIndex, RT.PoseBoneIndex, RequiredBones, false);
 		}
@@ -702,9 +542,8 @@ void Compression::SanitizeRawAnimSequenceTrack(FRawAnimSequenceTrack& RawTrack)
 	// if scale is too small, zero it out. Cause it hard to retarget when compress
 	// inverse scale is applied to translation, and causing translation to be huge to retarget, but
 	// compression can't handle that much precision. 
-	for (auto ScaleIter = RawTrack.ScaleKeys.CreateIterator(); ScaleIter; ++ScaleIter)
+	for (auto& Scale3D : RawTrack.ScaleKeys)
 	{
-		FVector3f& Scale3D = *ScaleIter;
 		if (FMath::IsNearlyZero(Scale3D.X))
 		{
 			Scale3D.X = 0.f;
@@ -718,15 +557,43 @@ void Compression::SanitizeRawAnimSequenceTrack(FRawAnimSequenceTrack& RawTrack)
 			Scale3D.Z = 0.f;
 		}
 	}
+	
+	for (auto& Position : RawTrack.PosKeys)
+	{		
+		if (FMath::IsNearlyZero(Position.X))
+		{
+			Position.X = 0.f;
+		}
+		if (FMath::IsNearlyZero(Position.Y))
+		{
+			Position.Y = 0.f;
+		}
+		if (FMath::IsNearlyZero(Position.Z))
+		{
+			Position.Z = 0.f;
+		}
+	}
 
 	// make sure Rotation part is normalized before compress
-	for (auto RotIter = RawTrack.RotKeys.CreateIterator(); RotIter; ++RotIter)
+	for (auto& Rotation : RawTrack.RotKeys)
 	{
-		FQuat4f& Rotation = *RotIter;
 		if (!Rotation.IsNormalized())
 		{
 			Rotation.Normalize();
 		}
+
+		if (FMath::IsNearlyZero(Rotation.X))
+		{
+			Rotation.X = 0.f;
+		}
+		if (FMath::IsNearlyZero(Rotation.Y))
+		{
+			Rotation.Y = 0.f;
+		}
+		if (FMath::IsNearlyZero(Rotation.Z))
+		{
+			Rotation.Z = 0.f;
+		}
 	}
 }
 
@@ -749,7 +616,7 @@ Compression::FScopedCompressionGuard::~FScopedCompressionGuard()
 
 bool AnimationData::AddLoopingInterpolation(UAnimSequence* InSequence)
 {
-	const UAnimDataModel* DataModel = InSequence->GetDataModel();
+	const IAnimationDataModel* DataModel = InSequence->GetDataModel();
 	IAnimationDataController& Controller = InSequence->GetController();
 
 	const int32 NumTracks = DataModel->GetNumBoneTracks();
@@ -785,8 +652,8 @@ bool AnimationData::AddLoopingInterpolation(UAnimSequence* InSequence)
 			Controller.SetBoneTrackKeys(AnimationTrack.Name, PositionalKeys, RotationalKeys, ScaleKeys);
 		}
 
-		const float Interval = DataModel->GetFrameRate().AsInterval();
-		Controller.SetPlayLength(Interval + DataModel->GetPlayLength());
+		// New number of frames is equal to current number of keys, as we'll be adding one frame (Number of Frames + 1 == Number of Keys)
+		Controller.SetNumberOfFrames(NumKeys);
 
 		return true;
 	}
@@ -796,7 +663,7 @@ bool AnimationData::AddLoopingInterpolation(UAnimSequence* InSequence)
 
 bool AnimationData::Trim(UAnimSequence* InSequence, float TrimStart, float TrimEnd)
 {
-	const UAnimDataModel* DataModel = InSequence->GetDataModel();
+	const IAnimationDataModel* DataModel = InSequence->GetDataModel();
 	IAnimationDataController& Controller = InSequence->GetController();
 
 	const int32 NumTracks = DataModel->GetNumBoneTracks();
@@ -842,7 +709,7 @@ bool AnimationData::Trim(UAnimSequence* InSequence, float TrimStart, float TrimE
 
 void AnimationData::DuplicateKeys(UAnimSequence* InSequence, int32 StartKeyIndex, int32 NumDuplicates, int32 SourceKeyIndex /*= INDEX_NONE */)
 {
-	const UAnimDataModel* Model = InSequence->GetDataModel();
+	const IAnimationDataModel* Model = InSequence->GetDataModel();
 	IAnimationDataController& Controller = InSequence->GetController();
 
 	const int32 NumberOfKeys = Model->GetNumberOfKeys();
@@ -894,14 +761,14 @@ void AnimationData::DuplicateKeys(UAnimSequence* InSequence, int32 StartKeyIndex
 			const float InsertedTime = FrameRate.AsInterval() * NumDuplicates;
 
 			// Notify will happen with time slice that was inserted
-			Controller.Resize(NewSequenceLength, StartTime, StartTime + InsertedTime);
+			Controller.ResizeInFrames(NewNumFrames, StartKeyIndex, StartKeyIndex + NumDuplicates);
 		}
 	}
 }
 
 void AnimationData::RemoveKeys(UAnimSequence* InSequence, int32 StartKeyIndex, int32 NumKeysToRemove)
 {
-	const UAnimDataModel* Model = InSequence->GetDataModel();
+	const IAnimationDataModel* Model = InSequence->GetDataModel();
 	IAnimationDataController& Controller = InSequence->GetController();
 
 	const int32 NumberOfKeys = Model->GetNumberOfKeys();
@@ -937,20 +804,17 @@ void AnimationData::RemoveKeys(UAnimSequence* InSequence, int32 StartKeyIndex, i
 		}
 
 		const int32 NewNumberOfFrames = FMath::Max(NewNumberOfKeys - 1, 0);
-		const float NewSequenceLength = FrameRate.AsSeconds(NewNumberOfFrames);
 
-		const float StartFrameTime = FrameRate.AsSeconds(FMath::Max(StartKeyIndex, 0));
-		const float RemovedTime = NumKeysToRemove * FrameRate.AsInterval();
-		const float EndTime = StartFrameTime + RemovedTime;
+		const int32 StartKey = FMath::Max(StartKeyIndex, 0);
 
 		// Notify will happen with time slice that was removed
-		Controller.Resize(NewSequenceLength, StartFrameTime, EndTime);
+		Controller.ResizeInFrames(NewNumberOfFrames, StartKey, StartKey + NumberOfKeys);
 	}
 }
 
 int32 AnimationData::FindFirstChildTrackIndex(const UAnimSequence* InSequence, const USkeleton* Skeleton, const FName& BoneName)
 {
-	const UAnimDataModel* DataModel = InSequence->GetDataModel();
+	const IAnimationDataModel* DataModel = InSequence->GetDataModel();
 	const FReferenceSkeleton& RefSkeleton = Skeleton->GetReferenceSkeleton();
 
 	const int32 BoneIndex = RefSkeleton.FindBoneIndex(BoneName);
diff --git a/Engine/Source/Runtime/Engine/Private/Animation/AnimSingleNodeInstanceProxy.cpp b/Engine/Source/Runtime/Engine/Private/Animation/AnimSingleNodeInstanceProxy.cpp
--- a/Engine/Source/Runtime/Engine/Private/Animation/AnimSingleNodeInstanceProxy.cpp
+++ b/Engine/Source/Runtime/Engine/Private/Animation/AnimSingleNodeInstanceProxy.cpp
@@ -174,7 +174,7 @@ void FAnimSingleNodeInstanceProxy::InternalBlendSpaceEvaluatePose(class UBlendSp
 {
 	FAnimationPoseData AnimationPoseData = { OutContext.Pose, OutContext.Curve, OutContext.CustomAttributes };
 
-	FAnimExtractContext ExtractionContext(CurrentTime, ShouldExtractRootMotion(), DeltaTimeRecord, bLooping);
+	FAnimExtractContext ExtractionContext(static_cast<double>(CurrentTime), ShouldExtractRootMotion(), DeltaTimeRecord, bLooping);
 
 	if (BlendSpace->IsValidAdditive())
 	{
@@ -188,7 +188,7 @@ void FAnimSingleNodeInstanceProxy::InternalBlendSpaceEvaluatePose(class UBlendSp
 #if WITH_EDITORONLY_DATA
 		if (BlendSpace->PreviewBasePose)
 		{
-			BlendSpace->PreviewBasePose->GetBonePose(AnimationPoseData, FAnimExtractContext(PreviewPoseCurrentTime));
+			BlendSpace->PreviewBasePose->GetBonePose(AnimationPoseData, FAnimExtractContext(static_cast<double>(PreviewPoseCurrentTime)));
 		}
 		else
 #endif // WITH_EDITORONLY_DATA
@@ -308,7 +308,7 @@ void FAnimNode_SingleNode::Evaluate_AnyThread(FPoseContext& Output)
 		= false;
 #endif
 
-	if (Proxy->CurrentAsset != NULL && !Proxy->CurrentAsset->HasAnyFlags(RF_BeginDestroyed))
+	if (Proxy->CurrentAsset != nullptr && !Proxy->CurrentAsset->HasAnyFlags(RF_BeginDestroyed))
 	{
 		FAnimationPoseData OutputAnimationPoseData(Output);
 
@@ -319,7 +319,7 @@ void FAnimNode_SingleNode::Evaluate_AnyThread(FPoseContext& Output)
 		}
 		else if (UAnimSequence* Sequence = Cast<UAnimSequence>(Proxy->CurrentAsset))
 		{
-			FAnimExtractContext ExtractionContext(Proxy->CurrentTime, Sequence->bEnableRootMotion, Proxy->DeltaTimeRecord, Proxy->bLooping);
+			FAnimExtractContext ExtractionContext(static_cast<double>(Proxy->CurrentTime), Sequence->bEnableRootMotion, Proxy->DeltaTimeRecord, Proxy->bLooping);
 
 			if (Sequence->IsValidAdditive())
 			{
@@ -378,12 +378,12 @@ void FAnimNode_SingleNode::Evaluate_AnyThread(FPoseContext& Output)
 			else*/
 			{
 				// if SkeletalMesh isn't there, we'll need to use skeleton
-				Streamable->GetAnimationPose(OutputAnimationPoseData, FAnimExtractContext(Proxy->CurrentTime, Streamable->bEnableRootMotion, Proxy->DeltaTimeRecord, Proxy->bLooping));
+				Streamable->GetAnimationPose(OutputAnimationPoseData, FAnimExtractContext(static_cast<double>(Proxy->CurrentTime), Streamable->bEnableRootMotion, Proxy->DeltaTimeRecord, Proxy->bLooping));
 			}
 		}
 		else if (UAnimComposite* Composite = Cast<UAnimComposite>(Proxy->CurrentAsset))
 		{
-			FAnimExtractContext ExtractionContext(Proxy->CurrentTime, Proxy->ShouldExtractRootMotion(), Proxy->DeltaTimeRecord, Proxy->bLooping);
+			FAnimExtractContext ExtractionContext(static_cast<double>(Proxy->CurrentTime), Proxy->ShouldExtractRootMotion(), Proxy->DeltaTimeRecord, Proxy->bLooping);
 			const FAnimTrack& AnimTrack = Composite->AnimationTrack;
 
 			// find out if this is additive animation
@@ -441,7 +441,7 @@ void FAnimNode_SingleNode::Evaluate_AnyThread(FPoseContext& Output)
 					if (bCanProcessAdditiveAnimationsLocal && Montage->PreviewBasePose && Montage->GetPlayLength() > 0.f)
 					{
 						FAnimationPoseData LocalAnimationPoseData = { LocalSourcePose, LocalSourceCurve, LocalSourceAttributes };
-						Montage->PreviewBasePose->GetBonePose(LocalAnimationPoseData, FAnimExtractContext(Proxy->CurrentTime));
+						Montage->PreviewBasePose->GetBonePose(LocalAnimationPoseData, FAnimExtractContext(static_cast<double>(Proxy->CurrentTime)));
 					}
 					else
 #endif // WITH_EDITORONLY_DATA
@@ -561,7 +561,7 @@ void FAnimNode_SingleNode::Evaluate_AnyThread(FPoseContext& Output)
 void FAnimNode_SingleNode::Update_AnyThread(const FAnimationUpdateContext& Context)
 {
 	float NewPlayRate = Proxy->PlayRate;
-	UAnimSequence* PreviewBasePose = NULL;
+	UAnimSequence* PreviewBasePose = nullptr;
 
 	if (Proxy->bPlaying == false)
 	{
@@ -569,7 +569,7 @@ void FAnimNode_SingleNode::Update_AnyThread(const FAnimationUpdateContext& Conte
 		NewPlayRate = 0.f;
 	}
 
-	if(Proxy->CurrentAsset != NULL)
+	if(Proxy->CurrentAsset != nullptr)
 	{
 		UE::Anim::FAnimSyncGroupScope& SyncScope = Context.GetMessageChecked<UE::Anim::FAnimSyncGroupScope>();
 
diff --git a/Engine/Source/Runtime/Engine/Private/Animation/AnimStreamable.cpp b/Engine/Source/Runtime/Engine/Private/Animation/AnimStreamable.cpp
--- a/Engine/Source/Runtime/Engine/Private/Animation/AnimStreamable.cpp
+++ b/Engine/Source/Runtime/Engine/Private/Animation/AnimStreamable.cpp
@@ -19,14 +19,17 @@
 #include "Animation/AnimBoneCompressionCodec.h"
 #include "Animation/AnimCurveCompressionCodec.h"
 #include "BonePose.h"
+#include "CommonFrameRates.h"
 #include "ContentStreaming.h"
+#include "ITimeManagementModule.h"
 #include "ProfilingDebugging/CookStats.h"
 #include "Animation/AnimationPoseData.h"
 #include "Animation/AttributesRuntime.h"
 #include "UObject/UE5MainStreamObjectVersion.h"
 
-#include "Animation/AnimData/AnimDataModel.h"
 #include "Animation/AnimSequenceHelpers.h"
+#include "Animation/AnimSequence.h"
+#include "Animation/AnimData/IAnimationDataModel.h"
 
 CSV_DECLARE_CATEGORY_MODULE_EXTERN(ENGINE_API, Animation);
 
@@ -168,6 +171,41 @@ void UAnimStreamable::Serialize(FArchive& Ar)
 		UE_LOG(LogAnimation, Fatal, TEXT("This platform requires cooked packages, and animation data was not cooked into %s."), *GetFullName());
 	}
 
+#if WITH_EDITORONLY_DATA
+	if (Ar.IsLoading() && Ar.CustomVer(FUE5MainStreamObjectVersion::GUID) < FUE5MainStreamObjectVersion::AnimationDataModelInterface)
+	{
+		// Figure out correct SamplingFrameRate value 
+		if(SourceSequence && !SourceSequence->HasAnyFlags(RF_NeedPostLoad))
+		{
+			SamplingFrameRate = SourceSequence->GetSamplingFrameRate();
+		}
+		else
+		{
+			const int32 NumberOfFrames = FMath::Max(NumberOfKeys - 1, 1);
+
+			// Generate the frame-rate according to the number of frames and sequence length
+			const double DecimalFrameRate = (double)NumberOfFrames / ((double)GetPlayLength() > 0.0 ? (double)GetPlayLength() : 1.0);
+
+			// Account for non-whole number frame rates using large denominator
+			const double Denominator = 1000000.0;
+			SamplingFrameRate = FFrameRate(DecimalFrameRate * Denominator, Denominator);	 
+
+			// Try to simplifiy the frame rate, in case it is a multiple of the commonly used frame rates e.g. 10000/300000 -> 1/30
+			TArrayView<const FCommonFrameRateInfo> CommonFrameRates = FModuleManager::LoadModulePtr<ITimeManagementModule>("TimeManagement")->GetAllCommonFrameRates();
+			for (const FCommonFrameRateInfo& Info : CommonFrameRates)
+			{
+				const bool bDoesNotAlreadyMatch = Info.FrameRate.Denominator != SamplingFrameRate.Denominator && Info.FrameRate.Numerator != SamplingFrameRate.Numerator;
+		
+				if (bDoesNotAlreadyMatch && FMath::IsNearlyEqual(SamplingFrameRate.AsInterval(), Info.FrameRate.AsInterval()))
+				{
+					SamplingFrameRate = Info.FrameRate;
+					break;
+				}
+			}
+		}
+	}
+#endif // WITH_EDITORONLY_DATA
+
 	if (bCooked)
 	{
 		if (Ar.IsLoading())
@@ -191,19 +229,6 @@ void UAnimStreamable::Serialize(FArchive& Ar)
 #endif // WITH_EDITORONLY_DATA
 }
 
-#if WITH_EDITOR
-/*bool UAnimComposite::GetAllAnimationSequencesReferred(TArray<UAnimationAsset*>& AnimationAssets, bool bRecursive) 
-{
-	return AnimationTrack.GetAllAnimationSequencesReferred(AnimationAssets, bRecursive);
-}
-
-void UAnimComposite::ReplaceReferredAnimations(const TMap<UAnimationAsset*, UAnimationAsset*>& ReplacementMap)
-{
-	AnimationTrack.ReplaceReferredAnimations(ReplacementMap);
-}*/
-#endif
-
-
 void UAnimStreamable::HandleAssetPlayerTickedInternal(FAnimAssetTickContext &Context, const float PreviousTime, const float MoveDelta, const FAnimTickRecord &Instance, struct FAnimNotifyQueue& NotifyQueue) const
 {
 	Super::HandleAssetPlayerTickedInternal(Context, PreviousTime, MoveDelta, Instance, NotifyQueue);
@@ -290,10 +315,8 @@ void UAnimStreamable::GetAnimationPose(FAnimationPoseData& OutAnimationPoseData,
 		//Need to evaluate raw data
 		ValidateModel();
 
-		UE::Anim::EvaluateFloatCurvesFromModel(DataModel, OutCurve, ExtractionContext.CurrentTime);
-
 		// Warning if we have invalid data
-		for (const FBoneAnimationTrack& AnimationTrack : DataModel->GetBoneAnimationTracks())
+		for (const FBoneAnimationTrack& AnimationTrack : DataModelInterface->GetBoneAnimationTracks())
 		{
 			const FRawAnimSequenceTrack& TrackToExtract = AnimationTrack.InternalTrackData;
 			// Bail out (with rather wacky data) if data is empty for some reason.
@@ -303,7 +326,8 @@ void UAnimStreamable::GetAnimationPose(FAnimationPoseData& OutAnimationPoseData,
 			}
 		}
 
-		UE::Anim::BuildPoseFromModel(DataModel, OutPose, ExtractionContext.CurrentTime, Interpolation, RetargetSource, MySkeleton->GetRefLocalPoses(RetargetSource));
+		const UE::Anim::DataModel::FEvaluationContext EvaluationContext(ExtractionContext.CurrentTime, DataModelInterface->GetFrameRate(), RetargetSource, MySkeleton->GetRefLocalPoses(RetargetSource), Interpolation);
+		DataModelInterface->Evaluate(OutAnimationPoseData, EvaluationContext);
 
 		if ((ExtractionContext.bExtractRootMotion && RootMotionReset.bEnableRootMotion) || RootMotionReset.bForceRootLock)
 		{
@@ -313,19 +337,22 @@ void UAnimStreamable::GetAnimationPose(FAnimationPoseData& OutAnimationPoseData,
 	}
 #endif
 
-	const int32 ChunkIndex = GetChunkIndexForTime(GetRunningPlatformData().Chunks, ExtractionContext.CurrentTime);
+	const FStreamableAnimPlatformData& StreamableAnimData = GetRunningPlatformData();
+
+	const int32 ChunkIndex = GetChunkIndexForTime(StreamableAnimData.Chunks, ExtractionContext.CurrentTime);
 	if(ensureMsgf(ChunkIndex != INDEX_NONE, TEXT("Could not get valid chunk with Time %.2f for Streaming Anim %s"), ExtractionContext.CurrentTime, *GetFullName()))
 	{
-		IAnimationStreamingManager& StreamingManager = IStreamingManager::Get().GetAnimationStreamingManager();
+		const FAnimStreamableChunk& CurrentChunk = StreamableAnimData.Chunks[ChunkIndex];
+		const IAnimationStreamingManager& StreamingManager = IStreamingManager::Get().GetAnimationStreamingManager();
 
-		bool bUsingFirstChunk = (ChunkIndex == 0);
+		const bool bUsingFirstChunk = (ChunkIndex == 0);
 
 		const FCompressedAnimSequence* CurveCompressedDataChunk = StreamingManager.GetLoadedChunk(this, 0, bUsingFirstChunk); //Curve Data stored in chunk 0 till it is properly cropped
 
 		if (!CurveCompressedDataChunk)
 		{
 #if WITH_EDITOR
-			CurveCompressedDataChunk = GetRunningPlatformData().Chunks[0].CompressedAnimSequence;
+			CurveCompressedDataChunk = StreamableAnimData.Chunks[0].CompressedAnimSequence;
 #else
 			UE_LOG(LogAnimation, Warning, TEXT("Failed to get streamed compressed data Time: %.2f, ChunkIndex:%i, Anim: %s"), ExtractionContext.CurrentTime, 0, *GetFullName());
 			return;
@@ -336,14 +363,13 @@ void UAnimStreamable::GetAnimationPose(FAnimationPoseData& OutAnimationPoseData,
 
 		const FCompressedAnimSequence* CompressedData = bUsingFirstChunk ? CurveCompressedDataChunk : StreamingManager.GetLoadedChunk(this, ChunkIndex, true);
 
-		float ChunkCurrentTime = ExtractionContext.CurrentTime - GetRunningPlatformData().Chunks[ChunkIndex].StartTime;
-
+		double ChunkCurrentTime = ExtractionContext.CurrentTime - CurrentChunk.StartTime;
 		if (!CompressedData)
 		{
 #if WITH_EDITOR
-			CompressedData = GetRunningPlatformData().Chunks[ChunkIndex].CompressedAnimSequence;
+			CompressedData = CurrentChunk.CompressedAnimSequence;
 #else
-			const int32 NumChunks = GetRunningPlatformData().Chunks.Num();
+			const int32 NumChunks = StreamableAnimData.Chunks.Num();
 
 			int32 FallbackChunkIndex = ChunkIndex;
 			while (!CompressedData)
@@ -356,7 +382,7 @@ void UAnimStreamable::GetAnimationPose(FAnimationPoseData& OutAnimationPoseData,
 					return;
 				}
 				CompressedData = StreamingManager.GetLoadedChunk(this, FallbackChunkIndex, false);
-				ChunkCurrentTime = GetRunningPlatformData().Chunks[FallbackChunkIndex].SequenceLength;
+				ChunkCurrentTime = StreamableAnimData.Chunks[FallbackChunkIndex].SequenceLength;
 			}
 
 			UE_LOG(LogAnimation, Warning, TEXT("Failed to get streamed compressed data Time: %.2f, ChunkIndex:%i - Using Chunk %i Anim: %s"), ExtractionContext.CurrentTime, ChunkIndex, FallbackChunkIndex, *GetFullName());
@@ -374,7 +400,8 @@ void UAnimStreamable::GetAnimationPose(FAnimationPoseData& OutAnimationPoseData,
 		ChunkExtractionContext.PoseCurves = ExtractionContext.PoseCurves;
 
 		//FPlatformMisc::LowLevelOutputDebugStringf(TEXT("Playing Streaming Anim %s Time: %.2f Chunk:%i\n"), *GetName(), ExtractionContext.CurrentTime, ChunkIndex);
-		DecompressPose(OutPose, *CompressedData, ChunkExtractionContext, GetSkeleton(), GetRunningPlatformData().Chunks[ChunkIndex].SequenceLength, Interpolation, bIsBakedAdditive, RetargetSource, GetFName(), RootMotionReset);
+		FAnimSequenceDecompressionContext Context(SamplingFrameRate, CurrentChunk.NumFrames, Interpolation, GetFName(), *CompressedData->CompressedDataStructure.Get(), GetSkeleton(), bIsBakedAdditive);
+		UE::Anim::Decompression::DecompressPose(OutPose, *CompressedData, ChunkExtractionContext, Context, RetargetSource, RootMotionReset);
 	}
 }
 
@@ -462,16 +489,22 @@ void UAnimStreamable::InitFrom(const UAnimSequence* InSourceSequence)
 
 	BoneCompressionSettings = InSourceSequence->BoneCompressionSettings;
 	CurveCompressionSettings = InSourceSequence->CurveCompressionSettings;
+	
+	DataModelInterface = StaticDuplicateObject(InSourceSequence->GetDataModelInterface().GetObject(), this);
 
-	DataModel = DuplicateObject(InSourceSequence->GetDataModel(), this);
-	Controller->SetModel(DataModel);
+	// Far from ideal (retrieving controller to ensure it matches the DataModelInterface type)
+	Controller = DataModelInterface->GetController();
+	Controller->SetModel(DataModelInterface);
 
 	Notifies = InSourceSequence->Notifies;
-	Controller->SetPlayLength(InSourceSequence->GetDataModel()->GetPlayLength());
+	Controller->SetNumberOfFrames(InSourceSequence->GetDataModel()->GetNumberOfFrames());
 	RateScale = InSourceSequence->RateScale;
 	Interpolation = InSourceSequence->Interpolation;
 	RetargetSource = InSourceSequence->RetargetSource;
-	NumberOfKeys = DataModel->GetNumberOfKeys();
+	NumberOfKeys = DataModelInterface->GetNumberOfKeys();
+	SamplingFrameRate = SourceSequence->GetDataModel()->GetFrameRate();
+
+	Controller->NotifyPopulated();
 
 	bEnableRootMotion = InSourceSequence->bEnableRootMotion;
 	RootMotionRootLock = InSourceSequence->RootMotionRootLock;
@@ -561,8 +594,8 @@ void UAnimStreamable::RequestCompressedData(const ITargetPlatform* Platform)
 		NumChunks = FMath::Max(InitialNumChunks, 1);
 	}
 
-	int32 NumFramesToChunk = NumberOfKeys - 1;
-	int32 FramesPerChunk = NumberOfKeys / NumChunks;
+	const int32 NumFramesToChunk = NumberOfKeys - 1;
+	const int32 FramesPerChunk = NumFramesToChunk / NumChunks;
 
 	PlatformData.Chunks.AddDefaulted(NumChunks);
 
@@ -608,7 +641,7 @@ float UAnimStreamable::GetChunkSizeSeconds(const ITargetPlatform* Platform) cons
 }
 
 template<typename KeyType>
-void MakeKeyChunk(const TArray<KeyType>& SrcKeys, TArray<KeyType>& DestKeys, int32 NumFrames, const uint32 FrameStart, const uint32 FrameEnd)
+void MakeKeyChunk(const TArray<KeyType>& SrcKeys, TArray<KeyType>& DestKeys, int32 NumberOfKeys, const uint32 FrameStart, const uint32 FrameEnd)
 {
 	if (SrcKeys.Num() == 1)
 	{
@@ -616,7 +649,7 @@ void MakeKeyChunk(const TArray<KeyType>& SrcKeys, TArray<KeyType>& DestKeys, int
 	}
 	else
 	{
-		check(SrcKeys.Num() == NumFrames); // Invalid data otherwise
+		check(SrcKeys.Num() == NumberOfKeys); // Invalid data otherwise
 
 		DestKeys.Reset((FrameEnd - FrameStart) + 1);
 		for (uint32 FrameIndex = FrameStart; FrameIndex <= FrameEnd; ++FrameIndex)
@@ -643,10 +676,9 @@ void UAnimStreamable::RequestCompressedDataForChunk(const FString& ChunkDDCKey,
 		const bool bSkipDDC = false;
 
 		const int32 ChunkNumFrames = FrameEnd - FrameStart;
-		const int32 NumberOfFrames = NumberOfKeys - 1;
-		const float FrameLength = GetPlayLength() / (float)NumberOfFrames;
-		Chunk.StartTime = FrameStart * FrameLength;
-		Chunk.SequenceLength = ChunkNumFrames * FrameLength;
+		Chunk.StartTime = SamplingFrameRate.AsSeconds(static_cast<int32>(FrameStart));
+		Chunk.SequenceLength = SamplingFrameRate.AsSeconds(ChunkNumFrames);
+		Chunk.NumFrames = ChunkNumFrames;
 
 		if (!bSkipDDC && GetDerivedDataCacheRef().GetSynchronous(*FinalDDCKey, OutData, GetPathName()))
 		{
@@ -656,7 +688,7 @@ void UAnimStreamable::RequestCompressedDataForChunk(const FString& ChunkDDCKey,
 		{
 			FCompressibleAnimRef CompressibleData = MakeShared<FCompressibleAnimData, ESPMode::ThreadSafe>(BoneCompressionSettings, CurveCompressionSettings, GetSkeleton(), Interpolation, Chunk.SequenceLength, ChunkNumFrames+1);
 
-			const TArray<FBoneAnimationTrack>& BoneAnimationTracks = DataModel->GetBoneAnimationTracks();
+			const TArray<FBoneAnimationTrack>& BoneAnimationTracks = DataModelInterface->GetBoneAnimationTracks();
 			CompressibleData->TrackToSkeletonMapTable.Empty();
 
 			for (const FBoneAnimationTrack& AnimTrack : BoneAnimationTracks)
@@ -677,22 +709,9 @@ void UAnimStreamable::RequestCompressedDataForChunk(const FString& ChunkDDCKey,
 			if (FrameStart == 0)
 			{
 				//Crop curve logic broken, for the moment store curve data in always loaded chunk 0
-				CompressibleData->RawFloatCurves = DataModel->GetFloatCurves();
+				CompressibleData->RawFloatCurves = DataModelInterface->GetFloatCurves();
 			}
 
-			/*if (SourceSequence && ChunkIndex == 7)
-			{
-				UAnimSequence* Seq = const_cast<UAnimSequence*>(SourceSequence); //Stop update loop
-				SourceSequence = nullptr;
-				Seq->RawAnimationData = CompressibleData->RawAnimationData;
-				Seq->SequenceLength = CompressibleData->SequenceLength;
-				Seq->NumFrames = CompressibleData->NumFrames;
-				Seq->MarkRawDataAsModified(false);
-				Seq->OnRawDataChanged();
-
-				SourceSequence = Seq;
-			}*/
-
 			AnimCompressor->SetCompressibleData(CompressibleData);
 
 			if (bSkipDDC)
@@ -736,7 +755,7 @@ void UAnimStreamable::UpdateRawData()
 {
 	if (IsDataModelValid())
 	{
-		RawDataGuid = DataModel->GenerateGuid();
+		RawDataGuid = DataModelInterface->GenerateGuid();
 		RequestCompressedData();
 	}
 }
@@ -783,29 +802,4 @@ FStreamableAnimPlatformData& UAnimStreamable::GetStreamingAnimPlatformData(const
 #else
 	return RunningAnimPlatformData;
 #endif
-}
-
-/*EAdditiveAnimationType UAnimComposite::GetAdditiveAnimType() const
-{
-	int32 AdditiveType = AnimationTrack.GetTrackAdditiveType();
-
-	if (AdditiveType != -1)
-	{
-		return (EAdditiveAnimationType)AdditiveType;
-	}
-
-	return AAT_None;
-}
-
-bool UAnimComposite::HasRootMotion() const
-{
-	return AnimationTrack.HasRootMotion();
-}*/
-
-#if WITH_EDITOR
-/*class UAnimSequence* UAnimComposite::GetAdditiveBasePose() const
-{
-	// @todo : for now it just picks up the first sequence
-	return AnimationTrack.GetAdditiveBasePose();
-}*/
-#endif 
+}
\ No newline at end of file
diff --git a/Engine/Source/Runtime/Engine/Private/Animation/AnimationDecompression.cpp b/Engine/Source/Runtime/Engine/Private/Animation/AnimationDecompression.cpp
new file mode 100644
--- /dev/null
+++ b/Engine/Source/Runtime/Engine/Private/Animation/AnimationDecompression.cpp
@@ -0,0 +1,291 @@
+// Copyright Epic Games, Inc. All Rights Reserved.
+
+#include "Animation/AnimationDecompression.h"
+#include "Animation/AnimSequenceDecompressionContext.h"
+#include "Animation/AnimCompressionTypes.h"
+#include "Animation/AnimBoneCompressionCodec.h"
+
+#include "AnimationRuntime.h"
+#include "AnimEncoding.h"
+
+
+CSV_DECLARE_CATEGORY_MODULE_EXTERN(ENGINE_API, Animation);
+DECLARE_CYCLE_STAT(TEXT("Build Anim Track Pairs"), STAT_BuildAnimTrackPairs, STATGROUP_Anim);
+DECLARE_CYCLE_STAT(TEXT("Extract Pose From Anim Data"), STAT_ExtractPoseFromAnimData, STATGROUP_Anim);
+
+namespace UE::Anim::Decompression {
+	
+struct FGetBonePoseScratchArea : public TThreadSingleton<FGetBonePoseScratchArea>
+{
+	BoneTrackArray RotationScalePairs;
+	BoneTrackArray TranslationPairs;
+	BoneTrackArray AnimScaleRetargetingPairs;
+	BoneTrackArray AnimRelativeRetargetingPairs;
+	BoneTrackArray OrientAndScaleRetargetingPairs;
+};
+
+void DecompressPose(FCompactPose& OutPose,
+								const FCompressedAnimSequence& CompressedData,
+								const FAnimExtractContext& ExtractionContext,
+								FAnimSequenceDecompressionContext& DecompressionContext,
+								FName RetargetSource,
+								const FRootMotionReset& RootMotionReset)
+{
+	const TArray<FTransform>& RetargetTransforms = DecompressionContext.GetSourceSkeleton()->GetRefLocalPoses(RetargetSource);
+	DecompressPose(OutPose, CompressedData, ExtractionContext, DecompressionContext, RetargetTransforms, RootMotionReset);
+}
+
+void DecompressPose(FCompactPose& OutPose,
+	const FCompressedAnimSequence& CompressedData,
+	const FAnimExtractContext& ExtractionContext,
+	FAnimSequenceDecompressionContext& DecompressionContext,
+	const TArray<FTransform>& RetargetTransforms,
+	const FRootMotionReset& RootMotionReset)
+{
+	const FBoneContainer& RequiredBones = OutPose.GetBoneContainer();
+	const int32 NumTracks = CompressedData.CompressedTrackToSkeletonMapTable.Num();
+
+	const USkeleton* TargetSkeleton = RequiredBones.GetSkeletonAsset();
+	const FSkeletonRemapping* SkeletonRemapping = TargetSkeleton->GetSkeletonRemapping(DecompressionContext.GetSourceSkeleton());
+
+	BoneTrackArray& RotationScalePairs = FGetBonePoseScratchArea::Get().RotationScalePairs;
+	BoneTrackArray& TranslationPairs = FGetBonePoseScratchArea::Get().TranslationPairs;
+	BoneTrackArray& AnimScaleRetargetingPairs = FGetBonePoseScratchArea::Get().AnimScaleRetargetingPairs;
+	BoneTrackArray& AnimRelativeRetargetingPairs = FGetBonePoseScratchArea::Get().AnimRelativeRetargetingPairs;
+	BoneTrackArray& OrientAndScaleRetargetingPairs = FGetBonePoseScratchArea::Get().OrientAndScaleRetargetingPairs;
+
+	// build a list of desired bones
+	RotationScalePairs.Reset();
+	TranslationPairs.Reset();
+	AnimScaleRetargetingPairs.Reset();
+	AnimRelativeRetargetingPairs.Reset();
+	OrientAndScaleRetargetingPairs.Reset();
+
+	// Optimization: assuming first index is root bone. That should always be the case in Skeletons.
+	checkSlow((RequiredBones.GetMeshPoseIndexFromSkeletonPoseIndex(FSkeletonPoseBoneIndex(0)) == FMeshPoseBoneIndex(0)));
+	// this is not guaranteed for AnimSequences though... If Root is not animated, Track will not exist.
+	const bool bFirstTrackIsRootBone = (CompressedData.GetSkeletonIndexFromTrackIndex(0) == 0);
+
+	{
+		SCOPE_CYCLE_COUNTER(STAT_BuildAnimTrackPairs);
+
+		// Handle root bone separately if it is track 0. so we start w/ Index 1.
+		for (int32 TrackIndex = (bFirstTrackIsRootBone ? 1 : 0); TrackIndex < NumTracks; TrackIndex++)
+		{
+			const int32 SourceSkeletonBoneIndex = CompressedData.GetSkeletonIndexFromTrackIndex(TrackIndex);
+			const int32 TargetSkeletonBoneIndex = (SkeletonRemapping) ? SkeletonRemapping->GetTargetSkeletonBoneIndex(SourceSkeletonBoneIndex) : SourceSkeletonBoneIndex;
+
+			if (TargetSkeletonBoneIndex != INDEX_NONE)
+			{
+				const FCompactPoseBoneIndex BoneIndex = RequiredBones.GetCompactPoseIndexFromSkeletonIndex(TargetSkeletonBoneIndex);
+				//Nasty, we break our type safety, code in the lower levels should be adjusted for this
+				const int32 CompactPoseBoneIndex = BoneIndex.GetInt();
+				if (CompactPoseBoneIndex != INDEX_NONE)
+				{
+					RotationScalePairs.Add(BoneTrackPair(CompactPoseBoneIndex, TrackIndex));
+
+					// Skip extracting translation component for EBoneTranslationRetargetingMode::Skeleton.
+					switch (TargetSkeleton->GetBoneTranslationRetargetingMode(TargetSkeletonBoneIndex))
+					{
+					case EBoneTranslationRetargetingMode::Animation:
+						TranslationPairs.Add(BoneTrackPair(CompactPoseBoneIndex, TrackIndex));
+						break;
+					case EBoneTranslationRetargetingMode::AnimationScaled:
+						TranslationPairs.Add(BoneTrackPair(CompactPoseBoneIndex, TrackIndex));
+						AnimScaleRetargetingPairs.Add(BoneTrackPair(CompactPoseBoneIndex, SourceSkeletonBoneIndex));
+						break;
+					case EBoneTranslationRetargetingMode::AnimationRelative:
+						TranslationPairs.Add(BoneTrackPair(CompactPoseBoneIndex, TrackIndex));
+
+						// With baked additives, we can skip 'AnimationRelative' tracks, as the relative transform gets canceled out.
+						// (A1 + Rel) - (A2 + Rel) = A1 - A2.
+						if (!DecompressionContext.IsAdditiveAnimation())
+						{
+							AnimRelativeRetargetingPairs.Add(BoneTrackPair(CompactPoseBoneIndex, SourceSkeletonBoneIndex));
+						}
+						break;
+					case EBoneTranslationRetargetingMode::OrientAndScale:
+						TranslationPairs.Add(BoneTrackPair(CompactPoseBoneIndex, TrackIndex));
+
+						// Additives remain additives, they're not retargeted.
+						if (!DecompressionContext.IsAdditiveAnimation())
+						{
+							OrientAndScaleRetargetingPairs.Add(BoneTrackPair(CompactPoseBoneIndex, SourceSkeletonBoneIndex));
+						}
+						break;
+					}
+				}
+			}
+		}
+	}
+
+	{
+		SCOPE_CYCLE_COUNTER(STAT_ExtractPoseFromAnimData);
+		CSV_SCOPED_TIMING_STAT(Animation, ExtractPoseFromAnimData);
+		CSV_CUSTOM_STAT(Animation, NumberOfExtractedAnimations, 1, ECsvCustomStatOp::Accumulate);
+
+		DecompressionContext.Seek(ExtractionContext.CurrentTime);
+
+		// Handle Root Bone separately
+		if (bFirstTrackIsRootBone)
+		{
+			const int32 TrackIndex = 0;
+			FCompactPoseBoneIndex RootBone(0);
+			FTransform& RootAtom = OutPose[RootBone];
+
+			CompressedData.BoneCompressionCodec->DecompressBone(DecompressionContext, TrackIndex, RootAtom);
+
+			// Retarget the root onto the target skeleton (correcting for differences in rest poses)
+			if (SkeletonRemapping)
+			{
+				const int32 TargetSkeletonBoneIndex = 0;
+
+				if (DecompressionContext.IsAdditiveAnimation())
+				{
+					RootAtom.SetRotation(SkeletonRemapping->RetargetAdditiveRotationToTargetSkeleton(TargetSkeletonBoneIndex, RootAtom.GetRotation()));
+					if (TargetSkeleton->GetBoneTranslationRetargetingMode(TargetSkeletonBoneIndex) != EBoneTranslationRetargetingMode::Skeleton)
+					{
+						RootAtom.SetTranslation(SkeletonRemapping->RetargetAdditiveTranslationToTargetSkeleton(TargetSkeletonBoneIndex, RootAtom.GetTranslation()));
+					}
+				}
+				else
+				{
+					RootAtom.SetRotation(SkeletonRemapping->RetargetBoneRotationToTargetSkeleton(TargetSkeletonBoneIndex, RootAtom.GetRotation()));
+					if (TargetSkeleton->GetBoneTranslationRetargetingMode(TargetSkeletonBoneIndex) != EBoneTranslationRetargetingMode::Skeleton)
+					{
+						RootAtom.SetTranslation(SkeletonRemapping->RetargetBoneTranslationToTargetSkeleton(TargetSkeletonBoneIndex, RootAtom.GetTranslation()));
+					}
+				}
+			}
+
+			// @laurent - we should look into splitting rotation and translation tracks, so we don't have to process translation twice.
+			FAnimationRuntime::RetargetBoneTransform(DecompressionContext.GetSourceSkeleton(), DecompressionContext.AnimName, RetargetTransforms, RootAtom, 0, RootBone, RequiredBones, DecompressionContext.IsAdditiveAnimation());
+		}
+
+		if (RotationScalePairs.Num() > 0)
+		{
+			// get the remaining bone atoms
+			TArrayView<FTransform> OutPoseBones = OutPose.GetMutableBones();
+			CompressedData.BoneCompressionCodec->DecompressPose(DecompressionContext, RotationScalePairs, TranslationPairs, RotationScalePairs, OutPoseBones);
+		}
+	}
+
+	// Retarget the pose onto the target skeleton (correcting for differences in rest poses)
+	if (SkeletonRemapping)
+	{
+		if (DecompressionContext.IsAdditiveAnimation())
+		{
+			for (FCompactPoseBoneIndex BoneIndex(bFirstTrackIsRootBone ? 1 : 0); BoneIndex < OutPose.GetNumBones(); ++BoneIndex)
+			{
+				const int32 TargetSkeletonBoneIndex = RequiredBones.GetSkeletonIndex(BoneIndex);
+				OutPose[BoneIndex].SetRotation(SkeletonRemapping->RetargetAdditiveRotationToTargetSkeleton(TargetSkeletonBoneIndex, OutPose[BoneIndex].GetRotation()));
+				if (TargetSkeleton->GetBoneTranslationRetargetingMode(TargetSkeletonBoneIndex) != EBoneTranslationRetargetingMode::Skeleton)
+				{
+					OutPose[BoneIndex].SetTranslation(SkeletonRemapping->RetargetAdditiveTranslationToTargetSkeleton(TargetSkeletonBoneIndex, OutPose[BoneIndex].GetTranslation()));
+				}
+			}
+		}
+		else
+		{
+			for (FCompactPoseBoneIndex BoneIndex(bFirstTrackIsRootBone ? 1 : 0); BoneIndex < OutPose.GetNumBones(); ++BoneIndex)
+			{
+				const int32 TargetSkeletonBoneIndex = RequiredBones.GetSkeletonIndex(BoneIndex);
+				OutPose[BoneIndex].SetRotation(SkeletonRemapping->RetargetBoneRotationToTargetSkeleton(TargetSkeletonBoneIndex, OutPose[BoneIndex].GetRotation()));
+				if (TargetSkeleton->GetBoneTranslationRetargetingMode(TargetSkeletonBoneIndex) != EBoneTranslationRetargetingMode::Skeleton)
+				{
+					OutPose[BoneIndex].SetTranslation(SkeletonRemapping->RetargetBoneTranslationToTargetSkeleton(TargetSkeletonBoneIndex, OutPose[BoneIndex].GetTranslation()));
+				}
+			}
+		}
+	}
+
+	// Once pose has been extracted, snap root bone back to first frame if we are extracting root motion.
+	if ((ExtractionContext.bExtractRootMotion && RootMotionReset.bEnableRootMotion) || RootMotionReset.bForceRootLock)
+	{
+		RootMotionReset.ResetRootBoneForRootMotion(OutPose[FCompactPoseBoneIndex(0)], RequiredBones);
+	}
+
+	// Anim Scale Retargeting
+	int32 const NumBonesToScaleRetarget = AnimScaleRetargetingPairs.Num();
+	if (NumBonesToScaleRetarget > 0)
+	{
+		TArray<FTransform> const& AuthoredOnRefSkeleton = RetargetTransforms;
+
+		for (const BoneTrackPair& BonePair : AnimScaleRetargetingPairs)
+		{
+			const FCompactPoseBoneIndex BoneIndex(BonePair.AtomIndex); //Nasty, we break our type safety, code in the lower levels should be adjusted for this
+			const int32 SourceSkeletonBoneIndex = BonePair.TrackIndex;
+
+			// @todo - precache that in FBoneContainer when we have SkeletonIndex->TrackIndex mapping. So we can just apply scale right away.
+			float const SourceTranslationLength = AuthoredOnRefSkeleton[SourceSkeletonBoneIndex].GetTranslation().Size();
+			if (SourceTranslationLength > UE_KINDA_SMALL_NUMBER)
+			{
+				float const TargetTranslationLength = RequiredBones.GetRefPoseTransform(BoneIndex).GetTranslation().Size();
+				OutPose[BoneIndex].ScaleTranslation(TargetTranslationLength / SourceTranslationLength);
+			}
+		}
+	}
+
+	// Anim Relative Retargeting
+	int32 const NumBonesToRelativeRetarget = AnimRelativeRetargetingPairs.Num();
+	if (NumBonesToRelativeRetarget > 0)
+	{
+		TArray<FTransform> const& AuthoredOnRefSkeleton = RetargetTransforms;
+
+		for (const BoneTrackPair& BonePair : AnimRelativeRetargetingPairs)
+		{
+			const FCompactPoseBoneIndex BoneIndex(BonePair.AtomIndex); //Nasty, we break our type safety, code in the lower levels should be adjusted for this
+			const int32 SourceSkeletonBoneIndex = BonePair.TrackIndex;
+
+			const FTransform& RefPoseTransform = RequiredBones.GetRefPoseTransform(BoneIndex);
+
+			// Remap the base pose onto the target skeleton so that we are working entirely in target space
+			FTransform BaseTransform = AuthoredOnRefSkeleton[SourceSkeletonBoneIndex];
+			if (SkeletonRemapping)
+			{
+				const int32 TargetSkeletonBoneIndex = SkeletonRemapping->GetTargetSkeletonBoneIndex(SourceSkeletonBoneIndex);
+				BaseTransform = SkeletonRemapping->RetargetBoneTransformToTargetSkeleton(TargetSkeletonBoneIndex, BaseTransform);
+			}
+
+			// Apply the retargeting as if it were an additive difference between the current skeleton and the retarget skeleton. 
+			OutPose[BoneIndex].SetRotation(OutPose[BoneIndex].GetRotation() * BaseTransform.GetRotation().Inverse() * RefPoseTransform.GetRotation());
+			OutPose[BoneIndex].SetTranslation(OutPose[BoneIndex].GetTranslation() + (RefPoseTransform.GetTranslation() - BaseTransform.GetTranslation()));
+			OutPose[BoneIndex].SetScale3D(OutPose[BoneIndex].GetScale3D() * (RefPoseTransform.GetScale3D() * BaseTransform.GetSafeScaleReciprocal(BaseTransform.GetScale3D())));
+			OutPose[BoneIndex].NormalizeRotation();
+		}
+	}
+
+	// Translation 'Orient and Scale' Translation Retargeting
+	const int32 NumBonesToOrientAndScaleRetarget = OrientAndScaleRetargetingPairs.Num();
+	if (NumBonesToOrientAndScaleRetarget > 0)
+	{
+		const FRetargetSourceCachedData& RetargetSourceCachedData = RequiredBones.GetRetargetSourceCachedData(DecompressionContext.AnimName, RetargetTransforms);
+		const TArray<FOrientAndScaleRetargetingCachedData>& OrientAndScaleDataArray = RetargetSourceCachedData.OrientAndScaleData;
+		const TArray<int32>& CompactPoseIndexToOrientAndScaleIndex = RetargetSourceCachedData.CompactPoseIndexToOrientAndScaleIndex;
+
+		// If we have any cached retargeting data.
+		if ((OrientAndScaleDataArray.Num() > 0) && (CompactPoseIndexToOrientAndScaleIndex.Num() == RequiredBones.GetCompactPoseNumBones()))
+		{
+			for (int32 Index = 0; Index < NumBonesToOrientAndScaleRetarget; Index++)
+			{
+				const BoneTrackPair& BonePair = OrientAndScaleRetargetingPairs[Index];
+				const FCompactPoseBoneIndex CompactPoseBoneIndex(BonePair.AtomIndex);
+				const int32 OrientAndScaleIndex = CompactPoseIndexToOrientAndScaleIndex[CompactPoseBoneIndex.GetInt()];
+				if (OrientAndScaleIndex != INDEX_NONE)
+				{
+					const FOrientAndScaleRetargetingCachedData& OrientAndScaleData = OrientAndScaleDataArray[OrientAndScaleIndex];
+					FTransform& BoneTransform = OutPose[CompactPoseBoneIndex];
+					const FVector AnimatedTranslation = BoneTransform.GetTranslation();
+
+					// If Translation is not animated, we can just copy the TargetTranslation. No retargeting needs to be done.
+					const FVector NewTranslation = (AnimatedTranslation - OrientAndScaleData.SourceTranslation).IsNearlyZero(BONE_TRANS_RT_ORIENT_AND_SCALE_PRECISION) ?
+						OrientAndScaleData.TargetTranslation :
+						OrientAndScaleData.TranslationDeltaOrient.RotateVector(AnimatedTranslation) * OrientAndScaleData.TranslationScale;
+
+					BoneTransform.SetTranslation(NewTranslation);
+				}
+			}
+		}
+	}
+}
+}
diff --git a/Engine/Source/Runtime/Engine/Private/Animation/AnimationRuntime.cpp b/Engine/Source/Runtime/Engine/Private/Animation/AnimationRuntime.cpp
--- a/Engine/Source/Runtime/Engine/Private/Animation/AnimationRuntime.cpp
+++ b/Engine/Source/Runtime/Engine/Private/Animation/AnimationRuntime.cpp
@@ -1447,6 +1447,46 @@ void FAnimationRuntime::GetKeyIndicesFromTime(int32& OutKeyIndex1, int32& OutKey
 	OutAlpha = (float)Alpha;
 }
 
+void FAnimationRuntime::GetKeyIndicesFromTime(int32& OutKeyIndex1, int32& OutKeyIndex2, float& OutAlpha, const double Time, const FFrameRate& FrameRate, const int32 NumberOfKeys)
+{
+	// Check for 1-frame, before-first-frame and after-last-frame cases.
+	if (Time <= 0.0 || NumberOfKeys == 1)
+	{
+		OutKeyIndex1 = 0;
+		OutKeyIndex2 = 0;
+		OutAlpha = 0.0f;
+		return;
+	}
+
+	const FFrameTime FrameTime = FrameRate.AsFrameTime(Time);
+	const FFrameTime LastFrameTimeIndex = FFrameTime(NumberOfKeys - 1);
+	if (FrameTime >= LastFrameTimeIndex)
+	{
+		OutKeyIndex1 = LastFrameTimeIndex.FrameNumber.Value;
+		OutKeyIndex2 = 0;
+		OutAlpha = 0.0f;
+		return;
+	}
+
+	// Find the integer part (ensuring within range) and that gives us the 'starting' key index.
+	const int32 KeyIndex1 = FMath::Clamp<int32>(FrameTime.GetFrame().Value, 0, NumberOfKeys - 1); 
+
+	// The alpha (fractional part) is then just the remainder.
+	const float Alpha = FrameTime.GetSubFrame();
+
+	int32 KeyIndex2 = KeyIndex1 + 1;
+
+	// If we have gone over the end, do different things in case of looping
+	if (KeyIndex2 == NumberOfKeys)
+	{
+		KeyIndex2 = KeyIndex1;
+	}
+
+	OutKeyIndex1 = KeyIndex1;
+	OutKeyIndex2 = KeyIndex2;
+	OutAlpha = Alpha;
+}
+
 FTransform FAnimationRuntime::GetComponentSpaceRefPose(const FCompactPoseBoneIndex& CompactPoseBoneIndex, const FBoneContainer& BoneContainer)
 {
 	FCompactPoseBoneIndex CurrentIndex = CompactPoseBoneIndex;
diff --git a/Engine/Source/Runtime/Engine/Private/Animation/AnimationSettings.cpp b/Engine/Source/Runtime/Engine/Private/Animation/AnimationSettings.cpp
--- a/Engine/Source/Runtime/Engine/Private/Animation/AnimationSettings.cpp
+++ b/Engine/Source/Runtime/Engine/Private/Animation/AnimationSettings.cpp
@@ -1,7 +1,6 @@
 // Copyright Epic Games, Inc. All Rights Reserved.
 
 #include "Animation/AnimationSettings.h"
-#include "Animation/AnimCompress_BitwiseCompressOnly.h"
 
 UAnimationSettings::UAnimationSettings(const FObjectInitializer& ObjectInitializer)
 	: Super(ObjectInitializer)
@@ -32,6 +31,8 @@ UAnimationSettings::UAnimationSettings(const FObjectInitializer& ObjectInitializ
 		FMirrorFindReplaceExpression("Right", "Left", EMirrorFindReplaceMethod::Prefix), FMirrorFindReplaceExpression("Left", "Right",  EMirrorFindReplaceMethod::Prefix),
 		FMirrorFindReplaceExpression("((?:^[sS]pine|^[rR]oot|^[pP]elvis|^[nN]eck|^[hH]ead|^ik_hand_gun).*)", "$1", EMirrorFindReplaceMethod::RegularExpression)
 	};
+	
+	DefaultFrameRate = FFrameRate(30,1);
 }
 
 TArray<FString> UAnimationSettings::GetBoneCustomAttributeNamesToImport() const
@@ -54,6 +55,11 @@ TArray<FString> UAnimationSettings::GetBoneCustomAttributeNamesToImport() const
 	return AttributeNames;
 }
 
+const FFrameRate& UAnimationSettings::GetDefaultFrameRate() const
+{
+	checkf(DefaultFrameRate.IsValid(), TEXT("Invalid default frame-rate set in Project Settings"));
+	return DefaultFrameRate;
+}
 
 #if WITH_EDITOR
 void UAnimationSettings::PostEditChangeProperty(struct FPropertyChangedEvent& PropertyChangedEvent)
diff --git a/Engine/Source/Runtime/Engine/Private/Animation/AnimationUtils.cpp b/Engine/Source/Runtime/Engine/Private/Animation/AnimationUtils.cpp
--- a/Engine/Source/Runtime/Engine/Private/Animation/AnimationUtils.cpp
+++ b/Engine/Source/Runtime/Engine/Private/Animation/AnimationUtils.cpp
@@ -278,16 +278,14 @@ void FAnimationUtils::ComputeCompressionError(const FCompressibleAnimData& Compr
 
 		const FTransform EndEffectorDummyBoneSocket(FQuat::Identity, FVector(END_EFFECTOR_DUMMY_BONE_LENGTH_SOCKET));
 		const FTransform EndEffectorDummyBone(FQuat::Identity, FVector(END_EFFECTOR_DUMMY_BONE_LENGTH));
-		const FAnimKeyHelper Helper(CompressibleAnimData.SequenceLength, CompressedData.AnimData->CompressedNumberOfKeys);
-		const float KeyLength = Helper.TimePerKey() + UE_SMALL_NUMBER;
 
-		FAnimSequenceDecompressionContext DecompContext(CompressibleAnimData.SequenceLength, CompressibleAnimData.Interpolation, CompressibleAnimData.AnimFName, *CompressedData.AnimData);
+		FAnimSequenceDecompressionContext DecompContext(CompressibleAnimData.SampledFrameRate, CompressibleAnimData.GetNumberOfFrames(), CompressibleAnimData.Interpolation, CompressibleAnimData.AnimFName, *CompressedData.AnimData, nullptr, CompressibleAnimData.bIsValidAdditive);
 
 		const TArray<FBoneData>& BoneData = CompressibleAnimData.BoneData;
 
 		for (int32 KeyIndex = 0; KeyIndex< CompressedData.AnimData->CompressedNumberOfKeys; KeyIndex++)
 		{
-			const float Time = (float)KeyIndex * KeyLength;
+			const double Time = CompressibleAnimData.SampledFrameRate.AsSeconds(KeyIndex);
 			DecompContext.Seek(Time);
 
 			// get the raw and compressed atom for each bone
@@ -315,7 +313,7 @@ void FAnimationUtils::ComputeCompressionError(const FCompressibleAnimData& Compr
 
 						FTransform AdditiveRawTransform;
 						FTransform AdditiveNewTransform;
-						FAnimationUtils::ExtractTransformFromTrack(Time, CompressibleAnimData.NumberOfKeys, CompressibleAnimData.SequenceLength, CompressibleAnimData.RawAnimationData[BoneIndexData.TrackIndex], CompressibleAnimData.Interpolation, AdditiveRawTransform);
+						FAnimationUtils::ExtractTransformFromTrack(CompressibleAnimData.RawAnimationData[BoneIndexData.TrackIndex], Time, CompressibleAnimData.NumberOfKeys, CompressibleAnimData.SequenceLength, CompressibleAnimData.Interpolation, AdditiveRawTransform);
 						CompressedData.Codec->DecompressBone(DecompContext, BoneIndexData.TrackIndex, AdditiveNewTransform);
 
 						const ScalarRegister VBlendWeight(1.f);
@@ -324,7 +322,7 @@ void FAnimationUtils::ComputeCompressionError(const FCompressibleAnimData& Compr
 					}
 					else
 					{
-						FAnimationUtils::ExtractTransformFromTrack(Time, CompressibleAnimData.NumberOfKeys, CompressibleAnimData.SequenceLength, CompressibleAnimData.RawAnimationData[BoneIndexData.TrackIndex], CompressibleAnimData.Interpolation, RawTransforms[BoneIndex]);
+						FAnimationUtils::ExtractTransformFromTrack(CompressibleAnimData.RawAnimationData[BoneIndexData.TrackIndex], Time, CompressibleAnimData.NumberOfKeys, CompressibleAnimData.SequenceLength,  CompressibleAnimData.Interpolation, RawTransforms[BoneIndex]);
 						CompressedData.Codec->DecompressBone(DecompContext, BoneIndexData.TrackIndex, NewTransforms[BoneIndex]);
 					}
 				}
@@ -447,7 +445,7 @@ UObject* FAnimationUtils::GetDefaultAnimSequenceOuter(UAnimSet* InAnimSet, bool
 		return NewPackage;
 	}
 
-	return NULL;
+	return nullptr;
 }
 
 
@@ -933,10 +931,12 @@ UAnimCurveCompressionSettings* FAnimationUtils::GetDefaultAnimationCurveCompress
 
 void FAnimationUtils::EnsureAnimSequenceLoaded(UAnimSequence& AnimSeq)
 {
+	QUICK_SCOPE_CYCLE_COUNTER(STAT_EnsureAnimSequenceLoaded);
 	EnsureDependenciesAreLoaded(&AnimSeq);
 	EnsureDependenciesAreLoaded(AnimSeq.GetSkeleton());
 	EnsureDependenciesAreLoaded(AnimSeq.BoneCompressionSettings);
 	EnsureDependenciesAreLoaded(AnimSeq.CurveCompressionSettings);
+	EnsureDependenciesAreLoaded(AnimSeq.RefPoseSeq);
 }
 
 void FAnimationUtils::ExtractTransformForFrameFromTrackSafe(const FRawAnimSequenceTrack& RawTrack, int32 Frame, FTransform& OutAtom)
@@ -970,6 +970,11 @@ void FAnimationUtils::ExtractTransformForFrameFromTrack(const FRawAnimSequenceTr
 }
 
 void FAnimationUtils::ExtractTransformFromTrack(float Time, int32 NumFrames, float SequenceLength, const FRawAnimSequenceTrack& RawTrack, EAnimInterpolationType Interpolation, FTransform &OutAtom)
+{
+	ExtractTransformFromTrack(RawTrack, static_cast<double>(Time), NumFrames, static_cast<double>(SequenceLength), Interpolation, OutAtom);
+}
+
+void FAnimationUtils::ExtractTransformFromTrack(const FRawAnimSequenceTrack& RawTrack, double Time, int32 NumFrames, double SequenceLength, EAnimInterpolationType Interpolation, FTransform &OutAtom)
 {
 	// Bail out (with rather wacky data) if data is empty for some reason.
 	if (RawTrack.PosKeys.Num() == 0 || RawTrack.RotKeys.Num() == 0)
@@ -1037,6 +1042,11 @@ void FAnimationUtils::ExtractTransformFromTrack(float Time, int32 NumFrames, flo
 
 #if WITH_EDITOR
 void FAnimationUtils::ExtractTransformFromCompressionData(const FCompressibleAnimData& CompressibleAnimData, FCompressibleAnimDataResult& CompressedAnimData, float Time, int32 TrackIndex, bool bUseRawData, FTransform& OutBoneTransform)
+{
+	ExtractTransformFromCompressionData(CompressibleAnimData, CompressedAnimData, (double)Time, TrackIndex, bUseRawData, OutBoneTransform);
+}
+
+void FAnimationUtils::ExtractTransformFromCompressionData(const FCompressibleAnimData& CompressibleAnimData, FCompressibleAnimDataResult& CompressedAnimData, double Time, int32 TrackIndex, bool bUseRawData, FTransform& OutBoneTransform)
 {
 	FUECompressedAnimDataMutable& AnimDataMutable = static_cast<FUECompressedAnimDataMutable&>(*CompressedAnimData.AnimData);
 
@@ -1046,13 +1056,13 @@ void FAnimationUtils::ExtractTransformFromCompressionData(const FCompressibleAni
 		// Build our read-only version from the mutable source
 		FUECompressedAnimData AnimData(AnimDataMutable);
 
-		FAnimSequenceDecompressionContext DecompContext(CompressibleAnimData.SequenceLength, CompressibleAnimData.Interpolation, CompressibleAnimData.AnimFName, AnimData);
+		FAnimSequenceDecompressionContext DecompContext(CompressibleAnimData.SampledFrameRate, CompressibleAnimData.GetNumberOfFrames(), CompressibleAnimData.Interpolation, CompressibleAnimData.AnimFName, AnimData, nullptr, CompressibleAnimData.bIsValidAdditive);
 		DecompContext.Seek(Time);
 		CompressedAnimData.Codec->DecompressBone(DecompContext, TrackIndex, OutBoneTransform);
 		return;
 	}
 
-	FAnimationUtils::ExtractTransformFromTrack(Time, CompressibleAnimData.NumberOfKeys, CompressibleAnimData.SequenceLength, CompressibleAnimData.RawAnimationData[TrackIndex], CompressibleAnimData.Interpolation, OutBoneTransform);
+	FAnimationUtils::ExtractTransformFromTrack(CompressibleAnimData.RawAnimationData[TrackIndex], Time, CompressibleAnimData.NumberOfKeys, CompressibleAnimData.SequenceLength, CompressibleAnimData.Interpolation, OutBoneTransform);
 }
 
 bool FAnimationUtils::CompressAnimBones(FCompressibleAnimData& AnimSeq, FCompressibleAnimDataResult& Target)
diff --git a/Engine/Source/Runtime/Engine/Private/Animation/AttributesRuntime.cpp b/Engine/Source/Runtime/Engine/Private/Animation/AttributesRuntime.cpp
--- a/Engine/Source/Runtime/Engine/Private/Animation/AttributesRuntime.cpp
+++ b/Engine/Source/Runtime/Engine/Private/Animation/AttributesRuntime.cpp
@@ -27,13 +27,18 @@ namespace UE { namespace Anim {
 
 #if WITH_EDITOR
 void Attributes::GetAttributeValue(FStackAttributeContainer& OutAttributes, const FCompactPoseBoneIndex& PoseBoneIndex, const FAnimatedBoneAttribute& Attribute, const FAnimExtractContext& ExtractionContext)
+{
+	GetAttributeValue(OutAttributes, PoseBoneIndex, Attribute, ExtractionContext.CurrentTime);
+}
+
+void Attributes::GetAttributeValue(FStackAttributeContainer& OutAttributes, const FCompactPoseBoneIndex& PoseBoneIndex, const FAnimatedBoneAttribute& Attribute, double CurrentTime)
 {
 	// Evaluating a single attribute into a stack-based container
 	if (Attribute.Identifier.IsValid())
 	{
 		uint8* AttributeDataPtr = OutAttributes.FindOrAdd(Attribute.Identifier.GetType(), FAttributeId(Attribute.Identifier.GetName(), PoseBoneIndex));
 		check(AttributeDataPtr)
-		Attribute.Curve.EvaluateToPtr(Attribute.Identifier.GetType(), ExtractionContext.CurrentTime, AttributeDataPtr);
+		Attribute.Curve.EvaluateToPtr(Attribute.Identifier.GetType(), CurrentTime, AttributeDataPtr);
 	}
 }
 #endif // WITH_EDITOR
diff --git a/Engine/Source/Runtime/Engine/Private/Animation/BlendSpace.cpp b/Engine/Source/Runtime/Engine/Private/Animation/BlendSpace.cpp
--- a/Engine/Source/Runtime/Engine/Private/Animation/BlendSpace.cpp
+++ b/Engine/Source/Runtime/Engine/Private/Animation/BlendSpace.cpp
@@ -772,7 +772,7 @@ void UBlendSpace::ResetToRefPose(FCompactPose& OutPose) const
 	}
 }
 
-static const FAnimExtractContext DefaultBlendSpaceExtractionContext = { 0.f, true, {}, false };
+static const FAnimExtractContext DefaultBlendSpaceExtractionContext = { 0.0, true, {}, false };
 
 void UBlendSpace::GetAnimationPose(TArray<FBlendSampleData>& BlendSampleDataCache, /*out*/ FCompactPose& OutPose, /*out*/ FBlendedCurve& OutCurve) const
 {
@@ -874,7 +874,7 @@ void UBlendSpace::GetAnimationPose_Internal(TArray<FBlendSampleData>& BlendSampl
 
 					FAnimationPoseData ChildAnimationPoseData = { Pose, ChildrenCurves[I], ChildrenAttributes[I] };
 					// first one always fills up the source one
-					Sample.Animation->GetAnimationPose(ChildAnimationPoseData, FAnimExtractContext(Time, ExtractionContext.bExtractRootMotion, BlendSampleDataCache[I].DeltaTimeRecord, ExtractionContext.bLooping));
+					Sample.Animation->GetAnimationPose(ChildAnimationPoseData, FAnimExtractContext(static_cast<double>(Time), ExtractionContext.bExtractRootMotion, BlendSampleDataCache[I].DeltaTimeRecord, ExtractionContext.bLooping));
 				}
 				else
 				{
@@ -1594,7 +1594,7 @@ FVector UBlendSpace::GetNormalizedBlendInput(const FVector& BlendInput) const
 
 const FEditorElement* UBlendSpace::GetGridSampleInternal(int32 Index) const
 {
-	return GridSamples.IsValidIndex(Index) ? &GridSamples[Index] : NULL;
+	return GridSamples.IsValidIndex(Index) ? &GridSamples[Index] : nullptr;
 }
 
 // When using CriticallyDampedSmoothing, how to go from the interpolation speed to the smooth
diff --git a/Engine/Source/Runtime/Engine/Private/Animation/BonePose.cpp b/Engine/Source/Runtime/Engine/Private/Animation/BonePose.cpp
--- a/Engine/Source/Runtime/Engine/Private/Animation/BonePose.cpp
+++ b/Engine/Source/Runtime/Engine/Private/Animation/BonePose.cpp
@@ -4,6 +4,7 @@
 #include "AnimationRuntime.h"
 #include "AnimEncoding.h"
 #include "HAL/ThreadSingleton.h"
+#include "Animation/AnimSequenceHelpers.h"
 #if INTEL_ISPC
 #include "BonePose.ispc.generated.h"
 
@@ -136,17 +137,6 @@ bool FMeshPose::IsNormalized() const
 	return true;
 }
 
-struct FRetargetTracking
-{
-	const FCompactPoseBoneIndex PoseBoneIndex;
-	const int32 SkeletonBoneIndex;
-
-	FRetargetTracking(const FCompactPoseBoneIndex InPoseBoneIndex, const int32 InSkeletonBoneIndex)
-		: PoseBoneIndex(InPoseBoneIndex), SkeletonBoneIndex(InSkeletonBoneIndex)
-	{
-	}
-};
-
 FTransform ExtractTransformForKey(int32 Key, const FRawAnimSequenceTrack &TrackToExtract)
 {
 	static const FVector DefaultScale3D = FVector(1.f);
@@ -165,23 +155,16 @@ FTransform ExtractTransformForKey(int32 Key, const FRawAnimSequenceTrack &TrackT
 	}
 }
 
-struct FBuildRawPoseScratchArea : public TThreadSingleton<FBuildRawPoseScratchArea>
-{
-	TArray<FRetargetTracking> RetargetTracking;
-	TArray<FVirtualBoneCompactPoseData> VirtualBoneCompactPoseData;
-};
-
-
 template<bool bInterpolateT>
 void BuildPoseFromRawDataInternal(const TArray<FRawAnimSequenceTrack>& InAnimationData, const TArray<struct FTrackToSkeletonMap>& TrackToSkeletonMapTable, FCompactPose& InOutPose, int32 KeyIndex1, int32 KeyIndex2, float Alpha, float TimePerKey, const TMap<int32, const FTransformCurve*>* AdditiveBoneTransformCurves)
 {
 	const int32 NumTracks = InAnimationData.Num();
 	const FBoneContainer& RequiredBones = InOutPose.GetBoneContainer();
 
-	TArray<FRetargetTracking>& RetargetTracking = FBuildRawPoseScratchArea::Get().RetargetTracking;
+	TArray<UE::Anim::Retargeting::FRetargetTracking>& RetargetTracking = UE::Anim::FBuildRawPoseScratchArea::Get().RetargetTracking;
 	RetargetTracking.Reset(NumTracks);
 
-	TArray<FVirtualBoneCompactPoseData>& VBCompactPoseData = FBuildRawPoseScratchArea::Get().VirtualBoneCompactPoseData;
+	TArray<FVirtualBoneCompactPoseData>& VBCompactPoseData =  UE::Anim::FBuildRawPoseScratchArea::Get().VirtualBoneCompactPoseData;
 	VBCompactPoseData = RequiredBones.GetVirtualBoneCompactPoseData();
 
 	FCompactPose Key2Pose;
@@ -254,7 +237,7 @@ void BuildPoseFromRawDataInternal(const TArray<FRawAnimSequenceTrack>& InAnimati
 					}
 				}
 
-				RetargetTracking.Add(FRetargetTracking(PoseBoneIndex, SkeletonBoneIndex));
+				RetargetTracking.Add(UE::Anim::Retargeting::FRetargetTracking(PoseBoneIndex, SkeletonBoneIndex));
 			}
 		}
 	}
@@ -332,7 +315,9 @@ void BuildPoseFromRawData(
 {
 	int32 KeyIndex1, KeyIndex2;
 	float Alpha;
+	PRAGMA_DISABLE_DEPRECATION_WARNINGS
 	FAnimationRuntime::GetKeyIndicesFromTime(KeyIndex1, KeyIndex2, Alpha, InTime, NumFrames, SequenceLength);
+	PRAGMA_ENABLE_DEPRECATION_WARNINGS
 
 	if (Interpolation == EAnimInterpolationType::Step)
 	{
@@ -368,11 +353,10 @@ void BuildPoseFromRawData(
 
 	if (!bDisableRetargeting)
 	{
-		const TArray<FRetargetTracking>& RetargetTracking = FBuildRawPoseScratchArea::Get().RetargetTracking;
-
-		USkeleton* Skeleton = RequiredBones.GetSkeletonAsset();
+		const TArray<UE::Anim::Retargeting::FRetargetTracking>& RetargetTracking = UE::Anim::FBuildRawPoseScratchArea::Get().RetargetTracking;
 
-		for (const FRetargetTracking& RT : RetargetTracking)
+		const USkeleton* Skeleton = RequiredBones.GetSkeletonAsset();
+		for (const UE::Anim::Retargeting::FRetargetTracking& RT : RetargetTracking)
 		{
 			FAnimationRuntime::RetargetBoneTransform(Skeleton, SourceName, RetargetTransforms, InOutPose[RT.PoseBoneIndex], RT.SkeletonBoneIndex, RT.PoseBoneIndex, RequiredBones, false);
 		}
diff --git a/Engine/Source/Runtime/Engine/Private/Animation/PoseAsset.cpp b/Engine/Source/Runtime/Engine/Private/Animation/PoseAsset.cpp
--- a/Engine/Source/Runtime/Engine/Private/Animation/PoseAsset.cpp
+++ b/Engine/Source/Runtime/Engine/Private/Animation/PoseAsset.cpp
@@ -9,7 +9,7 @@
 #include "Animation/AnimSequence.h"
 #include "Animation/AnimSequenceBase.h"
 #include "Animation/AnimationPoseData.h"
-#include "Animation/AnimData/AnimDataModel.h"
+#include "Animation/AnimData/IAnimationDataModel.h"
 #include "Animation/AnimSequenceHelpers.h"
 #include "Framework/Notifications/NotificationManager.h"
 #include "UObject/ObjectSaveContext.h"
@@ -925,12 +925,16 @@ void UPoseAsset::PostLoad()
 			}
 			SourceAnimation->ConditionalPostLoad();
 			
-			if (!SourceAnimationRawDataGUID.IsValid() || SourceAnimationRawDataGUID != SourceAnimation->GetRawDataGuid())
+			if (!SourceAnimationRawDataGUID.IsValid() || SourceAnimationRawDataGUID != SourceAnimation->GetDataModel()->GenerateGuid())
 			{
 				FFormatNamedArguments Args;
 				Args.Add(TEXT("AssetName"), FText::FromString(GetPathName()));
 				Args.Add(TEXT("SourceAsset"), FText::FromString(SourceAnimation->GetPathName()));
-				const FText ResultText = FText::Format(LOCTEXT("PoseAssetSourceOutOfDate", "PoseAsset {AssetName} is out-of-date with its source animation {SourceAsset}"), Args);
+
+				Args.Add(TEXT("Stored"), FText::FromString(SourceAnimationRawDataGUID.ToString()));
+				Args.Add(TEXT("Found"), FText::FromString(SourceAnimation->GetDataModel()->GenerateGuid().ToString()));
+				
+				const FText ResultText = FText::Format(LOCTEXT("PoseAssetSourceOutOfDate", "PoseAsset {AssetName} is out-of-date with its source animation {SourceAsset} {Stored} vs {Found}"), Args);
 				UE_LOG(LogAnimation, Warning,TEXT("%s"), *ResultText.ToString());
 			}
 		}	
@@ -1380,7 +1384,7 @@ void UPoseAsset::CreatePoseFromAnimation(class UAnimSequence* AnimSequence, cons
 		{
 			SetSkeleton(TargetSkeleton);
 			SourceAnimation = AnimSequence;
-			SourceAnimationRawDataGUID = AnimSequence->GetRawDataGuid();
+			SourceAnimationRawDataGUID = AnimSequence->GetDataModel()->GenerateGuid();
 
 			// reinitialize, now we're making new pose from this animation
 			Reinitialize();
@@ -1394,7 +1398,7 @@ void UPoseAsset::CreatePoseFromAnimation(class UAnimSequence* AnimSequence, cons
 				FMemMark Mark(FMemStack::Get());
 
 				// set up track data - @todo: add revaliation code when checked
-				UAnimDataModel* DataModel = AnimSequence->GetDataModel();
+				IAnimationDataModel* DataModel = AnimSequence->GetDataModel();
 
 				TArray<FName> TrackNames;
 				DataModel->GetBoneTrackNames(TrackNames);
diff --git a/Engine/Source/Runtime/Engine/Private/Curves/KeyHandle.cpp b/Engine/Source/Runtime/Engine/Private/Curves/KeyHandle.cpp
--- a/Engine/Source/Runtime/Engine/Private/Curves/KeyHandle.cpp
+++ b/Engine/Source/Runtime/Engine/Private/Curves/KeyHandle.cpp
@@ -8,7 +8,18 @@ FKeyHandle::FKeyHandle()
 	static uint32 LastKeyHandleIndex = 1;
 	Index = ++LastKeyHandleIndex;
 
-	check(LastKeyHandleIndex != 0); // check in the unlikely event that this overflows
+	if (LastKeyHandleIndex == 0)
+	{
+		// If we are cooking, allow wrap-around
+		if (IsRunningCookCommandlet())
+		{
+			Index = LastKeyHandleIndex = 1;
+		}
+		else
+		{
+			check(LastKeyHandleIndex != 0); // check in the unlikely event that this overflows
+		}
+	}
 }
 
 FKeyHandle::FKeyHandle(uint32 SpecificIndex)
diff --git a/Engine/Source/Runtime/Engine/Private/Curves/RichCurve.cpp b/Engine/Source/Runtime/Engine/Private/Curves/RichCurve.cpp
--- a/Engine/Source/Runtime/Engine/Private/Curves/RichCurve.cpp
+++ b/Engine/Source/Runtime/Engine/Private/Curves/RichCurve.cpp
@@ -964,7 +964,9 @@ void FRichCurve::RemoveRedundantKeysInternal(float Tolerance, int32 InStartKeepK
 	TArray<FKeyHandle> AllHandlesByIndex;
 	TArray<FKeyHandle> KeepHandles;
 
-	if (KeyHandlesToIndices.Num() != 0)
+	// If we have no handles to start with do not bother generating any
+	const bool bHasKeyHandles = KeyHandlesToIndices.Num() != 0;
+	if (bHasKeyHandles)
 	{
 		check(KeyHandlesToIndices.Num() == Keys.Num());
 		AllHandlesByIndex.AddZeroed(Keys.Num());
@@ -975,11 +977,6 @@ void FRichCurve::RemoveRedundantKeysInternal(float Tolerance, int32 InStartKeepK
 			AllHandlesByIndex[HandleIndexPair.Value] = HandleIndexPair.Key;
 		}
 	}
-	else
-	{
-		AllHandlesByIndex.AddDefaulted(Keys.Num());
-	}
-	
 
 	{
 		TArray<FRichCurveKey> NewKeys;
@@ -989,7 +986,11 @@ void FRichCurve::RemoveRedundantKeysInternal(float Tolerance, int32 InStartKeepK
 		for(int32 StartKeepIndex = 0; StartKeepIndex <= ActualStartKeepKey; ++StartKeepIndex)
 		{
 			NewKeys.Add(Keys[StartKeepIndex]);
-			KeepHandles.Add(AllHandlesByIndex[StartKeepIndex]);
+			
+			if (bHasKeyHandles)
+			{
+				KeepHandles.Add(AllHandlesByIndex[StartKeepIndex]);
+			}
 		}
 
 		//Add keys up to the first end keep key if they are not redundant
@@ -1050,7 +1051,11 @@ void FRichCurve::RemoveRedundantKeysInternal(float Tolerance, int32 InStartKeepK
 			{
 				MostRecentKeepKeyIndex = TestIndex;
 				NewKeys.Add(Keys[TestIndex]);
-				KeepHandles.Add(AllHandlesByIndex[TestIndex]);
+
+				if(bHasKeyHandles)
+				{
+					KeepHandles.Add(AllHandlesByIndex[TestIndex]);
+				}
 			}
 		}
 
@@ -1058,16 +1063,22 @@ void FRichCurve::RemoveRedundantKeysInternal(float Tolerance, int32 InStartKeepK
 		for (int32 EndKeepIndex = ActualEndKeepKey; EndKeepIndex < Keys.Num(); ++EndKeepIndex)
 		{
 			NewKeys.Add(Keys[EndKeepIndex]);
-			KeepHandles.Add(AllHandlesByIndex[EndKeepIndex]);
+			if (bHasKeyHandles)
+			{
+				KeepHandles.Add(AllHandlesByIndex[EndKeepIndex]);
+			}
 		}
 		Keys = MoveTemp(NewKeys); //Do this at the end of scope, guaranteed that NewKeys is going away
 	}
 
 	AutoSetTangents();
 
-	// Rebuild KeyHandlesToIndices
-	check(Keys.Num() == KeepHandles.Num());
-	KeyHandlesToIndices.Initialize(KeepHandles);
+	if (bHasKeyHandles)
+	{
+		// Rebuild KeyHandlesToIndices
+		check(Keys.Num() == KeepHandles.Num());
+		KeyHandlesToIndices.Initialize(KeepHandles);
+	}
 }
 
 void FRichCurve::RemapTimeValue(float& InTime, float& CycleValueOffset) const
@@ -1144,7 +1155,7 @@ float FRichCurve::Eval(float InTime, float InDefaultValue) const
 	{
 		// If no keys in curve, return the Default value.
 	} 
-	else if (NumKeys < 2 || (InTime <= Keys[0].Time))
+	else if (NumKeys < 2 || (InTime < Keys[0].Time))
 	{
 		if (PreInfinityExtrap == RCCE_Linear && NumKeys > 1)
 		{
@@ -1168,7 +1179,7 @@ float FRichCurve::Eval(float InTime, float InDefaultValue) const
 			InterpVal = Keys[0].Value;
 		}
 	}
-	else if (InTime < Keys[NumKeys - 1].Time)
+	else if (InTime <= Keys[NumKeys - 1].Time)
 	{
 		// perform a lower bound to get the second of the interpolation nodes
 		int32 first = 1;
@@ -1177,8 +1188,8 @@ float FRichCurve::Eval(float InTime, float InDefaultValue) const
 
 		while (count > 0)
 		{
-			int32 step = count / 2;
-			int32 middle = first + step;
+			const int32 step = count / 2;
+			const int32 middle = first + step;
 
 			if (InTime >= Keys[middle].Time)
 			{
@@ -1264,13 +1275,10 @@ static ERichCurveCompressionFormat FindRichCurveCompressionFormat(const FRichCur
 	{
 		if (Key.InterpMode == RCIM_Cubic)
 		{
-			if (Key.TangentMode != RCTM_Auto && Key.TangentMode != RCTM_None)
+			if (Key.TangentWeightMode != RCTWM_WeightedNone)
 			{
-				if (Key.TangentWeightMode != RCTWM_WeightedNone)
-				{
-					bIsWeighted = true;
-					break;
-				}
+				bIsWeighted = true;
+				break;
 			}
 		}
 		if (Key.InterpMode != RefKey.InterpMode)
@@ -1448,10 +1456,10 @@ static ERichCurveKeyTimeCompressionFormat FindRichCurveKeyFormat(const FRichCurv
 
 void FRichCurve::CompressCurve(FCompressedRichCurve& OutCurve, float ErrorThreshold, float SampleRate) const
 {
-	ERichCurveCompressionFormat CompressionFormat = FindRichCurveCompressionFormat(*this);
+	const ERichCurveCompressionFormat CompressionFormat = FindRichCurveCompressionFormat(*this);
 	OutCurve.CompressionFormat = CompressionFormat;
 
-	ERichCurveKeyTimeCompressionFormat KeyFormat = FindRichCurveKeyFormat(*this, ErrorThreshold, SampleRate, CompressionFormat);
+	const ERichCurveKeyTimeCompressionFormat KeyFormat = FindRichCurveKeyFormat(*this, ErrorThreshold, SampleRate, CompressionFormat);
 	OutCurve.KeyTimeCompressionFormat = KeyFormat;
 
 	OutCurve.PreInfinityExtrap = PreInfinityExtrap;
@@ -1977,7 +1985,7 @@ FORCEINLINE_DEBUGGABLE static float InterpEval(float InTime, const KeyTimeAdapte
 	float CycleValueOffset = 0.0;
 
 	const float FirstKeyTime = KeyTimeAdapter.GetTime(0);
-	if (InTime <= FirstKeyTime)
+	if (InTime < FirstKeyTime)
 	{
 		if (PreInfinityExtrap != RCCE_Linear && PreInfinityExtrap != RCCE_Constant)
 		{
@@ -1990,7 +1998,7 @@ FORCEINLINE_DEBUGGABLE static float InterpEval(float InTime, const KeyTimeAdapte
 	}
 
 	const float LastKeyTime = KeyTimeAdapter.GetTime(NumKeys - 1);
-	if (InTime >= LastKeyTime)
+	if (InTime > LastKeyTime)
 	{
 		if (PostInfinityExtrap != RCCE_Linear && PostInfinityExtrap != RCCE_Constant)
 		{
@@ -2032,6 +2040,14 @@ FORCEINLINE_DEBUGGABLE static float InterpEval(float InTime, const KeyTimeAdapte
 	const KeyDataHandle KeyValueHandle0 = KeyDataAdapter.GetKeyDataHandle(First - 1);
 	const float KeyValue0 = KeyDataAdapter.GetKeyValue(KeyValueHandle0);
 
+	auto IsItNotWeighted = [KeyDataAdapter](KeyDataHandle Handle0, KeyDataHandle Handle1) -> bool
+	{
+		const ERichCurveTangentWeightMode Mode0 = KeyDataAdapter.GetKeyTangentWeightMode(Handle0);
+		const ERichCurveTangentWeightMode Mode1 = KeyDataAdapter.GetKeyTangentWeightMode(Handle1);
+		
+		return (Mode0 == RCTWM_WeightedNone || Mode0 == RCTWM_WeightedArrive) && (Mode1 == RCTWM_WeightedNone || Mode1 == RCTWM_WeightedLeave);
+	};
+
 	// const value here allows the code to be stripped statically if the data is uniform
 	// which it is most of the time
 	const ERichCurveCompressionFormat KeyInterpMode0 = KeyDataAdapter.GetKeyInterpMode(First - 1);
@@ -2049,7 +2065,7 @@ FORCEINLINE_DEBUGGABLE static float InterpEval(float InTime, const KeyTimeAdapte
 		{
 			InterpolatedValue = FMath::Lerp(P0, P3, Alpha);
 		}
-		else if (KeyInterpMode0 != RCCF_Weighted || KeyDataAdapter.GetKeyInterpMode(First) != RCCF_Weighted)
+		else if (IsItNotWeighted(KeyValueHandle0, KeyValueHandle1))
 		{
 			const float OneThird = 1.0f / 3.0f;
 			const float ScaledDiff = Diff * OneThird;
diff --git a/Engine/Source/Runtime/Engine/Public/Animation/AnimCompressionTypes.h b/Engine/Source/Runtime/Engine/Public/Animation/AnimCompressionTypes.h
--- a/Engine/Source/Runtime/Engine/Public/Animation/AnimCompressionTypes.h
+++ b/Engine/Source/Runtime/Engine/Public/Animation/AnimCompressionTypes.h
@@ -21,6 +21,7 @@
 #include "Serialization/MemoryWriter.h"
 
 #include "ProfilingDebugging/CsvProfiler.h"
+#include "Animation/AnimationDecompression.h"
 
 #include "AnimCompressionTypes.generated.h"
 
@@ -222,10 +223,12 @@ public:
 
 	/** Number of keys within the (non-uniform) RawAnimationData tracks */
 	
-	UE_DEPRECATED(5.0, "NumberOfFrames has been replaced with NumberOfKeys")
+	UE_DEPRECATED(5.0, "NumberOfFrames has been replaced with NumberOfKeys and GetNumberOfFrames")
 	int32 NumberOfFrames;
 	int32 NumberOfKeys;
 
+	int32 GetNumberOfFrames() const { return FMath::Max(NumberOfKeys - 1, 1); }
+
 	bool bIsValidAdditive;
 
 	float ErrorThresholdScale;
@@ -235,6 +238,8 @@ public:
 	FString FullName;
 	FName   AnimFName;
 
+	FFrameRate SampledFrameRate;
+
 	FCancelCompressionSignal IsCancelledSignal;
 
 	static int32 GetApproxRawDataArraySize(const TArray<FRawAnimSequenceTrack>& AnimData)
@@ -806,6 +811,7 @@ struct FRootMotionReset
 	}
 };
 
+UE_DEPRECATED(5.1, "Signature of DecompressPose has been deprecated, use UE::Anim::Decompression::DecompressPose instead")
 extern void DecompressPose(FCompactPose& OutPose,
 							const FCompressedAnimSequence& CompressedData,
 							const FAnimExtractContext& ExtractionContext,
@@ -817,6 +823,7 @@ extern void DecompressPose(FCompactPose& OutPose,
 							FName SourceName,
 							const FRootMotionReset& RootMotionReset);
 
+UE_DEPRECATED(5.1, "Signature of DecompressPose has been deprecated, use UE::Anim::Decompression::DecompressPose instead")
 extern void DecompressPose(	FCompactPose& OutPose,
 							const FCompressedAnimSequence& CompressedData,
 							const FAnimExtractContext& ExtractionContext,
diff --git a/Engine/Source/Runtime/Engine/Public/Animation/AnimSequenceDecompressionContext.h b/Engine/Source/Runtime/Engine/Public/Animation/AnimSequenceDecompressionContext.h
--- a/Engine/Source/Runtime/Engine/Public/Animation/AnimSequenceDecompressionContext.h
+++ b/Engine/Source/Runtime/Engine/Public/Animation/AnimSequenceDecompressionContext.h
@@ -4,24 +4,45 @@
 
 #include "CoreMinimal.h"
 #include "Animation/AnimTypes.h"
+#include "Misc/FrameRate.h"
 
 struct ICompressedAnimData;
+class USkeleton;
 
 /* Encapsulates decompression related data used by bone compression codecs. */
 struct FAnimSequenceDecompressionContext
 {
-	FAnimSequenceDecompressionContext(float SequenceLength_, EAnimInterpolationType Interpolation_, const FName& AnimName_, const ICompressedAnimData& CompressedAnimData_)
-			:
+	UE_DEPRECATED(5.1, "FAnimSequenceDecompressionContext signature has been deprecated, use different one")
+	FAnimSequenceDecompressionContext(float SequenceLength_, EAnimInterpolationType Interpolation_, const FName& AnimName_, const ICompressedAnimData& CompressedAnimData_,  const USkeleton* InSourceSkeleton, bool bIsBakedAdditive)
+		:
 		PRAGMA_DISABLE_DEPRECATION_WARNINGS
 		SequenceLength(SequenceLength_),
 		PRAGMA_ENABLE_DEPRECATION_WARNINGS
 		Interpolation(Interpolation_), AnimName(AnimName_), CompressedAnimData(CompressedAnimData_),
 		PRAGMA_DISABLE_DEPRECATION_WARNINGS
-		Time(0.f), RelativePos(0.f)
+		Time(0.f), RelativePos(0.f),
 		PRAGMA_ENABLE_DEPRECATION_WARNINGS
+		SourceSkeleton(InSourceSkeleton),
+		bAdditiveAnimation(bIsBakedAdditive)
 	{
 	}
 
+	FAnimSequenceDecompressionContext(const FFrameRate& InSamplingRate, const int32 InNumberOfFrames, EAnimInterpolationType Interpolation_, const FName& AnimName_, const ICompressedAnimData& CompressedAnimData_, const USkeleton* InSourceSkeleton, bool bIsBakedAdditive)
+		:
+		PRAGMA_DISABLE_DEPRECATION_WARNINGS
+		SequenceLength(InSamplingRate.AsSeconds(InNumberOfFrames)),
+		PRAGMA_ENABLE_DEPRECATION_WARNINGS
+		Interpolation(Interpolation_), AnimName(AnimName_), CompressedAnimData(CompressedAnimData_),
+		PRAGMA_DISABLE_DEPRECATION_WARNINGS
+		Time(0.f), RelativePos(0.f),
+		PRAGMA_ENABLE_DEPRECATION_WARNINGS
+		SamplingRate(InSamplingRate),
+		SamplingTime(0),
+		NumberOfFrames(InNumberOfFrames),
+		SourceSkeleton(InSourceSkeleton),
+		bAdditiveAnimation(bIsBakedAdditive)
+	{}
+
 	// Anim info
 	UE_DEPRECATED(5.1, "Direct access to SequenceLength has been deprecated use GetPlayableLength instead")
 	float SequenceLength;
@@ -36,33 +57,49 @@ struct FAnimSequenceDecompressionContext
 	UE_DEPRECATED(5.1, "Direct access to RelativePos has been deprecated use GetRelativePosition instead")
 	float RelativePos;
 
-	float GetPlayableLength() const
+	double GetPlayableLength() const
 	{
-		PRAGMA_DISABLE_DEPRECATION_WARNINGS
-		return SequenceLength;
-		PRAGMA_ENABLE_DEPRECATION_WARNINGS
+		return SamplingRate.AsSeconds(NumberOfFrames);
 	}
-
-	float GetEvaluationTime() const
+	double GetEvaluationTime() const
 	{
-		PRAGMA_DISABLE_DEPRECATION_WARNINGS
-		return Time;
-		PRAGMA_ENABLE_DEPRECATION_WARNINGS
+		return SamplingRate.AsSeconds(SamplingTime);
+	}
+	double GetRelativePosition() const
+	{
+		return FMath::Min(SamplingTime.AsDecimal() / static_cast<double>(NumberOfFrames), 1.0);
 	}
 
-	float GetRelativePosition() const
+	double GetInterpolatedEvaluationTime() const
 	{
-		PRAGMA_DISABLE_DEPRECATION_WARNINGS
-		return RelativePos;
-		PRAGMA_ENABLE_DEPRECATION_WARNINGS
+		const FFrameTime InterpolationTime = Interpolation == EAnimInterpolationType::Step ? SamplingTime.FloorToFrame() : SamplingTime;
+		return SamplingRate.AsSeconds(InterpolationTime);
 	}
 
+	UE_DEPRECATED(5.1, "Seek with float precision sampling time has been deprecated use double signature instead")
 	void Seek(float SampleAtTime)
+	{
+		Seek(static_cast<double>(SampleAtTime));
+	}
+
+	void Seek(double SampleAtTime)
 	{
 		PRAGMA_DISABLE_DEPRECATION_WARNINGS
 		Time = SampleAtTime;
 		RelativePos = SequenceLength != 0.f ? SampleAtTime / SequenceLength : 0.f;
 		PRAGMA_ENABLE_DEPRECATION_WARNINGS
+		
+		SamplingTime = SamplingRate.AsFrameTime(SampleAtTime);
 	}
+
+	const USkeleton* GetSourceSkeleton() const { return SourceSkeleton; }
+	bool IsAdditiveAnimation() const { return bAdditiveAnimation; }	
+
+protected:
+	FFrameRate SamplingRate;
+	FFrameTime SamplingTime;
+	int32 NumberOfFrames;
+	const USkeleton* SourceSkeleton;
+	bool bAdditiveAnimation;
 };
 
diff --git a/Engine/Source/Runtime/Engine/Public/Animation/AnimTypes.h b/Engine/Source/Runtime/Engine/Public/Animation/AnimTypes.h
--- a/Engine/Source/Runtime/Engine/Public/Animation/AnimTypes.h
+++ b/Engine/Source/Runtime/Engine/Public/Animation/AnimTypes.h
@@ -374,8 +374,8 @@ public:
 		, TriggerTimeOffset(0)
 		, EndTriggerTimeOffset(0)
 		, TriggerWeightThreshold(ZERO_ANIMWEIGHT_THRESH)
-		, Notify(NULL)
-		, NotifyStateClass(NULL)
+		, Notify(nullptr)
+		, NotifyStateClass(nullptr)
 		, Duration(0)
 		, bConvertedFromBranchingPoint(false)
 		, MontageTickType(EMontageNotifyTickType::Queued)
@@ -424,7 +424,7 @@ public:
 	/** Returns true if this is blueprint derived notifies **/
 	bool IsBlueprintNotify() const
 	{
-		return Notify != NULL || NotifyStateClass != NULL;
+		return Notify != nullptr || NotifyStateClass != nullptr;
 	}
 
 	bool operator ==(const FAnimNotifyEvent& Other) const
@@ -880,6 +880,33 @@ struct ENGINE_API FRawAnimSequenceTrack
 		return Ar;
 	}
 
+	bool ContainsNaN() const
+	{
+		bool bContainsNaN = false;
+
+		auto CheckForNan = [&bContainsNaN](const auto& Keys) -> bool
+		{
+			if (!bContainsNaN)
+			{
+				for (const auto& Key : Keys)
+				{
+					if (Key.ContainsNaN())
+						return true;
+				}
+
+				return false;
+			}
+		
+			return true;
+		};
+
+		bContainsNaN = CheckForNan(PosKeys);
+		bContainsNaN = CheckForNan(RotKeys);
+		bContainsNaN = CheckForNan(ScaleKeys);
+
+		return bContainsNaN;
+	}	
+
 	static const uint32 SingleKeySize = sizeof(FVector3f) + sizeof(FQuat4f) + sizeof(FVector3f);
 };
 
diff --git a/Engine/Source/Runtime/Engine/Public/Animation/AttributesRuntime.h b/Engine/Source/Runtime/Engine/Public/Animation/AttributesRuntime.h
--- a/Engine/Source/Runtime/Engine/Public/Animation/AttributesRuntime.h
+++ b/Engine/Source/Runtime/Engine/Public/Animation/AttributesRuntime.h
@@ -56,7 +56,10 @@ namespace UE
 		struct ENGINE_API Attributes
 		{
 #if WITH_EDITOR
-			static void GetAttributeValue(FStackAttributeContainer& OutAttributes, const FCompactPoseBoneIndex& PoseBoneIndex, const  FAnimatedBoneAttribute& Attribute, const FAnimExtractContext& ExtractionContext);
+			UE_DEPRECATED(5.1, "GetAttributeValue with signature using FAnimExtractContext is deprecated use other version instead")
+			static void GetAttributeValue(FStackAttributeContainer& OutAttributes, const FCompactPoseBoneIndex& PoseBoneIndex, const FAnimatedBoneAttribute& Attribute, const FAnimExtractContext& ExtractionContext);
+
+			static void GetAttributeValue(FStackAttributeContainer& OutAttributes, const FCompactPoseBoneIndex& PoseBoneIndex, const FAnimatedBoneAttribute& Attribute, double CurrentTime);
 #endif // WITH_EDITOR
 
 			/** Blend custom attribute values from N set of inputs */
diff --git a/Engine/Source/Runtime/Engine/Public/AnimationRuntime.h b/Engine/Source/Runtime/Engine/Public/AnimationRuntime.h
--- a/Engine/Source/Runtime/Engine/Public/AnimationRuntime.h
+++ b/Engine/Source/Runtime/Engine/Public/AnimationRuntime.h
@@ -531,6 +531,9 @@ public:
 	 * The reason why you can provide a FramesPerSecond value is because this can be slightly more accurate than calculating it, in case super high precision is needed.
 	 **/
 	static void GetKeyIndicesFromTime(int32& OutKeyIndex1, int32& OutKeyIndex2, float& OutAlpha, const double Time, const int32 NumKeys, const double SequenceLength, double FramesPerSecond=-1.0);
+	
+	/** Get KeyIndices using FFrameRate::AsFrameTime to calculate the keys and alpha value **/
+	static void GetKeyIndicesFromTime(int32& OutKeyIndex1, int32& OutKeyIndex2, float& OutAlpha, const double Time, const FFrameRate& FrameRate, const int32 NumberOfKeys);
 
 	/** 
 	 *	Utility for taking an array of bone indices and ensuring that all parents are present 
diff --git a/Engine/Source/Runtime/Engine/Public/AnimationUtils.h b/Engine/Source/Runtime/Engine/Public/AnimationUtils.h
--- a/Engine/Source/Runtime/Engine/Public/AnimationUtils.h
+++ b/Engine/Source/Runtime/Engine/Public/AnimationUtils.h
@@ -164,10 +164,16 @@ public:
 	// Extract specific frame from raw track and place in OutAtom
 	ENGINE_API static void ExtractTransformForFrameFromTrack(const FRawAnimSequenceTrack& RawTrack, int32 Frame, FTransform& OutAtom);
 
+	UE_DEPRECATED(5.1, "ExtractTransformFromTrack has been deprecated, use different signature")
 	ENGINE_API static void ExtractTransformFromTrack(float Time, int32 NumFrames, float SequenceLength, const struct FRawAnimSequenceTrack& RawTrack, EAnimInterpolationType Interpolation, FTransform &OutAtom);
 
+	ENGINE_API static void ExtractTransformFromTrack(const struct FRawAnimSequenceTrack& RawTrack, double Time, int32 NumFrames, double SequenceLength, EAnimInterpolationType Interpolation, FTransform &OutAtom);
+
 #if WITH_EDITOR
+	UE_DEPRECATED(5.1, "ExtractTransformFromCompressionData has been deprecated, use different signature")
 	static void ExtractTransformFromCompressionData(const FCompressibleAnimData& CompressibleAnimData, FCompressibleAnimDataResult& CompressedAnimData, float Time, int32 TrackIndex, bool bUseRawData, FTransform& OutBoneTransform);
+	
+	static void ExtractTransformFromCompressionData(const FCompressibleAnimData& CompressibleAnimData, FCompressibleAnimDataResult& CompressedAnimData, double Time, int32 TrackIndex, bool bUseRawData, FTransform& OutBoneTransform);
 
 	/**
 	 * Compresses the animation bones within a sequence with the chosen settings.
diff --git a/Engine/Source/Runtime/MovieSceneTracks/Private/Sections/MovieSceneSkeletalAnimationSection.cpp b/Engine/Source/Runtime/MovieSceneTracks/Private/Sections/MovieSceneSkeletalAnimationSection.cpp
--- a/Engine/Source/Runtime/MovieSceneTracks/Private/Sections/MovieSceneSkeletalAnimationSection.cpp
+++ b/Engine/Source/Runtime/MovieSceneTracks/Private/Sections/MovieSceneSkeletalAnimationSection.cpp
@@ -482,7 +482,7 @@ int32 UMovieSceneSkeletalAnimationSection::SetBoneIndexForRootMotionCalculations
 		//but if not first find first
 		int32 RootIndex = INDEX_NONE;
 #if WITH_EDITOR
-		const UAnimDataModel* DataModel = AnimSequence->GetDataModel();
+		const IAnimationDataModel* DataModel = AnimSequence->GetDataModel();
 		const TArray<FBoneAnimationTrack>& BoneAnimationTracks = DataModel->GetBoneAnimationTracks();
 		const int32 NumTracks = BoneAnimationTracks.Num();
 		for (int32 TrackIndex = 0; TrackIndex < NumTracks; ++TrackIndex)
@@ -542,7 +542,7 @@ bool UMovieSceneSkeletalAnimationSection::GetRootMotionTransform(FFrameTime Curr
 		Params.Weight.Evaluate(CurrentTime, ManualWeight);
 		OutWeight = ManualWeight * EvaluateEasing(CurrentTime);
 		bIsAdditive = false;
-		float CurrentTimeSeconds = static_cast<float>(MapTimeToAnimation(CurrentTime, FrameRate));
+		const double CurrentTimeSeconds = MapTimeToAnimation(CurrentTime, FrameRate);
 		bIsAdditive = AnimSequence->GetAdditiveAnimType() != EAdditiveAnimationType::AAT_None;
 
 		if (TempRootBoneIndex.IsSet())
diff --git a/Engine/Source/Runtime/MovieSceneTracks/Private/Tracks/MovieSceneSkeletalAnimationTrack.cpp b/Engine/Source/Runtime/MovieSceneTracks/Private/Tracks/MovieSceneSkeletalAnimationTrack.cpp
--- a/Engine/Source/Runtime/MovieSceneTracks/Private/Tracks/MovieSceneSkeletalAnimationTrack.cpp
+++ b/Engine/Source/Runtime/MovieSceneTracks/Private/Tracks/MovieSceneSkeletalAnimationTrack.cpp
@@ -383,7 +383,7 @@ static void CalculateDistanceMap(USkeletalMeshComponent* SkelMeshComp, UAnimSequ
 		FloatArray.SetNum(SecondAnimNumFrames);
 		float FirstAnimTime = FirstAnimIndex * FrameRateDiff + StartFirstAnimTime;
 		FirstAnimIndex += 1.0f;
-		FAnimExtractContext FirstExtractionContext(FirstAnimTime, false);
+		FAnimExtractContext FirstExtractionContext(static_cast<double>(FirstAnimTime), false);
 		FirstAnimSeq->GetAnimationPose(FirstAnimPoseData, FirstExtractionContext);
 		FirstMeshPoses.InitPose(FirstAnimPoseData.GetPose());
 		float SecondAnimIndex = 0.0f;
@@ -392,7 +392,7 @@ static void CalculateDistanceMap(USkeletalMeshComponent* SkelMeshComp, UAnimSequ
 			DistVal = 0.0f;
 			float SecondAnimTime = SecondAnimIndex * FrameRateDiff;
 			SecondAnimIndex += 1.0f;
-			FAnimExtractContext SecondExtractionContext(SecondAnimTime, false);
+			FAnimExtractContext SecondExtractionContext(static_cast<double>(SecondAnimTime), false);
 			SecondAnimSeq->GetAnimationPose(SecondAnimPoseData, SecondExtractionContext);
 			SecondMeshPoses.InitPose(SecondAnimPoseData.GetPose());
 
@@ -641,7 +641,7 @@ void UMovieSceneSkeletalAnimationTrack::SetUpRootMotions(bool bForce)
 
 	}
 }
-static FTransform GetWorldTransformForBone(UAnimSequence* AnimSequence, USkeletalMeshComponent* MeshComponent,const FName& InBoneName, float Seconds)
+static FTransform GetWorldTransformForBone(UAnimSequence* AnimSequence, USkeletalMeshComponent* MeshComponent,const FName& InBoneName, double Seconds)
 {
 	FName BoneName = InBoneName;
 	FTransform  WorldTransform = FTransform::Identity;
@@ -653,7 +653,7 @@ static FTransform GetWorldTransformForBone(UAnimSequence* AnimSequence, USkeleta
 		int32 TrackIndex = INDEX_NONE;
 
 #if WITH_EDITOR
-		const UAnimDataModel* Model = AnimSequence->GetDataModel();
+		const IAnimationDataModel* Model = AnimSequence->GetDataModel();
 		if (const FBoneAnimationTrack* TrackData = Model->FindBoneTrackByName(BoneName))
 		{
 			TrackIndex = TrackData->BoneTreeIndex;
@@ -936,9 +936,9 @@ void UMovieSceneSkeletalAnimationTrack::MatchSectionByBoneTransform(bool bMatchW
 
 		if (FirstAnimSequence && SecondAnimSequence)
 		{
-			float FirstSectionTime = static_cast<float>(FirstSection->MapTimeToAnimation(CurrentFrame, FrameRate));
+			double FirstSectionTime = FirstSection->MapTimeToAnimation(CurrentFrame, FrameRate);
 			FTransform  FirstTransform = GetWorldTransformForBone(FirstAnimSequence, SkelMeshComp, BoneName, FirstSectionTime);
-			float SecondSectionTime = static_cast<float>(CurrentSection->MapTimeToAnimation(CurrentFrame, FrameRate));
+			double SecondSectionTime = CurrentSection->MapTimeToAnimation(CurrentFrame, FrameRate);
 			FTransform  SecondTransform = GetWorldTransformForBone(SecondAnimSequence, SkelMeshComp, BoneName, SecondSectionTime);
 
 			//Need to match the translations and rotations here 
