Deleted broken but unused FLightMapCoefficients::operator==, where comparing the Coefficients array would always return false, unless comparing the object with itself (compares pointers into the array, not the array contents - deprecated under C++20).

diff --git a/Engine/Source/Runtime/Engine/Public/LightMap.h b/Engine/Source/Runtime/Engine/Public/LightMap.h
--- a/Engine/Source/Runtime/Engine/Public/LightMap.h
+++ b/Engine/Source/Runtime/Engine/Public/LightMap.h
@@ -515,18 +515,6 @@ struct FLightMapCoefficients
 	uint8 Coefficients[NUM_STORED_LIGHTMAP_COEF][4];
 	uint8 SkyOcclusion[4];
 	uint8 AOMaterialMask;
-
-	/** Equality operator */
-	bool operator==( const FLightMapCoefficients& RHS ) const
-	{
-		return Coverage == RHS.Coverage &&
-			   Coefficients == RHS.Coefficients &&
-			   SkyOcclusion[0] == RHS.SkyOcclusion[0] &&
-			   SkyOcclusion[1] == RHS.SkyOcclusion[1] &&
-			   SkyOcclusion[2] == RHS.SkyOcclusion[2] &&
-			   SkyOcclusion[3] == RHS.SkyOcclusion[3] &&
-			   AOMaterialMask == RHS.AOMaterialMask;
-	}
 };
 
 struct FQuantizedLightmapData
