World Partition: set transient actor descriptors for modified and unsaved actors to the original container, even if they are not present in that container (they're essentially proxy descriptors on top existing ones).

#rnx

diff --git a/Engine/Source/Runtime/Engine/Private/WorldPartition/LevelInstance/LevelInstanceActorDesc.cpp b/Engine/Source/Runtime/Engine/Private/WorldPartition/LevelInstance/LevelInstanceActorDesc.cpp
--- a/Engine/Source/Runtime/Engine/Private/WorldPartition/LevelInstance/LevelInstanceActorDesc.cpp
+++ b/Engine/Source/Runtime/Engine/Private/WorldPartition/LevelInstance/LevelInstanceActorDesc.cpp
@@ -158,12 +158,6 @@ bool FLevelInstanceActorDesc::IsContainerInstance() const
 
 bool FLevelInstanceActorDesc::GetContainerInstance(const UActorDescContainer*& OutLevelContainer, FTransform& OutLevelTransform, EContainerClusterMode& OutClusterMode) const
 {
-	if (!LevelInstanceContainer.IsValid() && IsLoaded())
-	{
-		// Lazy initialization of LevelInstanceContainer (used by ModifiedActorsDescList)
-		const_cast<FLevelInstanceActorDesc*>(this)->RegisterContainerInstance(GetActor()->GetWorld());
-	}
-
 	if (LevelInstanceContainer.IsValid())
 	{
 		OutLevelContainer = LevelInstanceContainer.Get();
diff --git a/Engine/Source/Runtime/Engine/Private/WorldPartition/WorldPartitionStreamingGeneration.cpp b/Engine/Source/Runtime/Engine/Private/WorldPartition/WorldPartitionStreamingGeneration.cpp
--- a/Engine/Source/Runtime/Engine/Private/WorldPartition/WorldPartitionStreamingGeneration.cpp
+++ b/Engine/Source/Runtime/Engine/Private/WorldPartition/WorldPartitionStreamingGeneration.cpp
@@ -121,7 +121,7 @@ class FWorldPartitionStreamingGenerator
 		const bool bIsTempContainerPackage = FPackageName::IsTempPackage(InContainer->GetPackage()->GetName());
 		
 		// Test whether an actor is editor only. Will fallback to the actor descriptor only if the actor is not loaded
-		auto IsActorEditorOnly = [](const FWorldPartitionActorDesc* ActorDesc, const FActorContainerID& ContainerID)
+		auto IsActorEditorOnly = [](const FWorldPartitionActorDesc* ActorDesc, const FActorContainerID& ContainerID) -> bool
 		{
 			if (ActorDesc->IsRuntimeRelevant(ContainerID))
 			{
@@ -138,9 +138,15 @@ class FWorldPartitionStreamingGenerator
 		};
 
 		// Create an actor descriptor view for the specified actor (modified or unsaved actors)
-		auto GetModifiedActorDesc = [this](AActor* InActor)
+		auto GetModifiedActorDesc = [this](AActor* InActor, const UActorDescContainer* InContainer) -> FWorldPartitionActorDesc*
 		{
-			return ModifiedActorsDescList->AddActor(InActor);
+			FWorldPartitionActorDesc* ModifiedActorDesc = ModifiedActorsDescList->AddActor(InActor);
+
+			// Pretend that this actor descriptor belongs to the original container, even if it's not present. It's essentially a proxy
+			// descriptor on top an existing one and at this point no code should require to access the container to resolve it anyways.
+			ModifiedActorDesc->SetContainer(const_cast<UActorDescContainer*>(InContainer));
+
+			return ModifiedActorDesc;
 		};
 
 		// Register the actor descriptor view
@@ -174,7 +180,7 @@ class FWorldPartitionStreamingGenerator
 					if (bHandleUnsavedActors && (bIsTempContainerPackage || Actor->GetPackage()->IsDirty()))
 					{
 						// Dirty, unsaved actor for PIE
-						FWorldPartitionActorDescView ModifiedActorDescView = GetModifiedActorDesc(Actor);
+						FWorldPartitionActorDescView ModifiedActorDescView = GetModifiedActorDesc(Actor, InContainer);
 						RegisterActorDescView(ActorDescIt->GetGuid(), ModifiedActorDescView);
 						continue;
 					}
@@ -193,7 +199,7 @@ class FWorldPartitionStreamingGenerator
 			{
 				if (IsValid(Actor) && Actor->IsPackageExternal() && Actor->IsMainPackageActor() && !Actor->IsEditorOnly() && !InContainer->GetActorDesc(Actor->GetActorGuid()))
 				{
-					FWorldPartitionActorDescView ModifiedActorDescView = GetModifiedActorDesc(Actor);
+					FWorldPartitionActorDescView ModifiedActorDescView = GetModifiedActorDesc(Actor, InContainer);
 					RegisterActorDescView(Actor->GetActorGuid(), ModifiedActorDescView);
 				}
 			}
