[VisualLogger] added support for some commands in non-editor builds
- 'vislog record' now supported and will start recording to file
- 'vislog stop' now supported and will stop recording to file
- 'vislog disableallbut' now supported

diff --git a/Engine/Source/Runtime/Engine/Private/VisualLogger/VisualLogger.cpp b/Engine/Source/Runtime/Engine/Private/VisualLogger/VisualLogger.cpp
--- a/Engine/Source/Runtime/Engine/Private/VisualLogger/VisualLogger.cpp
+++ b/Engine/Source/Runtime/Engine/Private/VisualLogger/VisualLogger.cpp
@@ -866,7 +866,7 @@ bool FVisualLogger::IsCategoryLogged(const FLogCategoryBase& Category) const
 const FGuid EVisualLoggerVersion::GUID = FGuid(0xA4237A36, 0xCAEA41C9, 0x8FA218F8, 0x58681BF3);
 FCustomVersionRegistration GVisualLoggerVersion(EVisualLoggerVersion::GUID, EVisualLoggerVersion::LatestVersion, TEXT("VisualLogger"));
 
-#if !(UE_BUILD_SHIPPING || UE_BUILD_TEST)
+#if ENABLE_VISUAL_LOG
 
 class FLogVisualizerExec : private FSelfRegisteringExec
 {
@@ -876,47 +876,53 @@ public:
 	{
 		if (FParse::Command(&Cmd, TEXT("VISLOG")))
 		{
-			if (FModuleManager::Get().LoadModulePtr<IModuleInterface>("LogVisualizer") != nullptr)
+			const FString Command = FParse::Token(Cmd, /*UseEscape*/ false);
+			if (Command == TEXT("record"))
 			{
-#if ENABLE_VISUAL_LOG
-				const FString Command = FParse::Token(Cmd, /*UseEscape*/ false);
-				if (Command == TEXT("record"))
+				if (GIsEditor)
 				{
 					FVisualLogger::Get().SetIsRecording(true);
-					return true;
 				}
-				else if (Command == TEXT("stop"))
+				else
 				{
-					FVisualLogger::Get().SetIsRecording(false);
-					return true;
+					FVisualLogger::Get().SetIsRecordingToFile(true);
 				}
-				else if (Command == TEXT("disableallbut"))
+				return true;
+			}
+			else if (Command == TEXT("stop"))
+			{
+				if (GIsEditor)
 				{
-					const FString Category = FParse::Token(Cmd, /*UseEscape*/ true);
-					FVisualLogger::Get().BlockAllCategories(true);
-					FVisualLogger::Get().AddCategoryToAllowList(*Category);
-					return true;
+					FVisualLogger::Get().SetIsRecording(false);
 				}
-#if WITH_EDITOR
 				else
 				{
-					FGlobalTabmanager::Get()->TryInvokeTab(FName(TEXT("VisualLogger")));
-					return true;
+					FVisualLogger::Get().SetIsRecordingToFile(false);
 				}
-#endif
-#else
-			UE_LOG(LogVisual, Warning, TEXT("Unable to open LogVisualizer - logs are disabled"));
-#endif
+				return true;
+			}
+			else if (Command == TEXT("disableallbut"))
+			{
+				const FString Category = FParse::Token(Cmd, /*UseEscape*/ true);
+				FVisualLogger::Get().BlockAllCategories(true);
+				FVisualLogger::Get().AddCategoryToAllowList(*Category);
+				return true;
 			}
+#if WITH_EDITOR
+			else if (GIsEditor)
+			{
+				FGlobalTabmanager::Get()->TryInvokeTab(FName(TEXT("VisualLogger")));
+				return true;
+			}
+#endif // WITH_EDITOR
 		}
-#if ENABLE_VISUAL_LOG
 		else if (FParse::Command(&Cmd, TEXT("LogNavOctree")))
 		{
 			FVisualLogger::NavigationDataDump(GetWorldForVisualLogger(InWorld), LogNavigation, ELogVerbosity::Log, FBox());
 		}
-#endif
+
 		return false;
 	}
 } LogVisualizerExec;
 
-#endif // !(UE_BUILD_SHIPPING || UE_BUILD_TEST)
+#endif // ENABLE_VISUAL_LOG
