Chaos Cloth - Add Cloth Asset Builder.
- Add ConstCloth adapters.
- Add basic Cloth Adapter weight map API.
- Fix the reindexation of the Start/End indices and better deal with empty ranges in the Cloth Collection using new SetNum functions.
- Add a new function to copy the sim mesh to the render mesh part of the cloth patterns to help with debugging the rendering code.

diff --git a/Engine/Source/Runtime/Engine/Public/Rendering/SkeletalMeshRenderData.h b/Engine/Source/Runtime/Engine/Public/Rendering/SkeletalMeshRenderData.h
--- a/Engine/Source/Runtime/Engine/Public/Rendering/SkeletalMeshRenderData.h
+++ b/Engine/Source/Runtime/Engine/Public/Rendering/SkeletalMeshRenderData.h
@@ -50,8 +50,8 @@ public:
 	TUniquePtr<class FSkeletalMeshRenderData> NextCachedRenderData;
 #endif
 
-	FSkeletalMeshRenderData();
-	~FSkeletalMeshRenderData();
+	ENGINE_API FSkeletalMeshRenderData();
+	ENGINE_API ~FSkeletalMeshRenderData();
 
 #if WITH_EDITOR
 	void Cache(const ITargetPlatform* TargetPlatform, USkeletalMesh* Owner, class FSkeletalMeshCompilationContext* ContextPtr);
@@ -64,7 +64,7 @@ public:
 	void Serialize(FArchive& Ar, USkeletalMesh* Owner);
 
 	/** Initializes rendering resources. */
-	void InitResources(bool bNeedsVertexColors, TArray<UMorphTarget*>& InMorphTargets, USkeletalMesh* Owner);
+	ENGINE_API void InitResources(bool bNeedsVertexColors, TArray<UMorphTarget*>& InMorphTargets, USkeletalMesh* Owner);
 
 	/** Releases rendering resources. */
 	ENGINE_API void ReleaseResources();
