- Fix transition error with Cascade GPU sims as sorting is done on compute

#rnx

diff --git a/Engine/Source/Runtime/Engine/Private/Particles/ParticleGpuSimulation.cpp b/Engine/Source/Runtime/Engine/Private/Particles/ParticleGpuSimulation.cpp
--- a/Engine/Source/Runtime/Engine/Private/Particles/ParticleGpuSimulation.cpp
+++ b/Engine/Source/Runtime/Engine/Private/Particles/ParticleGpuSimulation.cpp
@@ -385,12 +385,12 @@ public:
 			ParticleResources->SimulationAttributesTexture.InitResource();
 
 			RHICmdList.Transition({ 
-				FRHITransitionInfo(ParticleResources->RenderAttributesTexture.TextureRHI, ERHIAccess::Unknown, ERHIAccess::SRVGraphics),
-				FRHITransitionInfo(ParticleResources->SimulationAttributesTexture.TextureRHI, ERHIAccess::Unknown, ERHIAccess::SRVGraphics),
-				FRHITransitionInfo(ParticleResources->StateTextures[0].PositionTextureRHI, ERHIAccess::Unknown, ERHIAccess::SRVGraphics),
-				FRHITransitionInfo(ParticleResources->StateTextures[0].VelocityTextureRHI, ERHIAccess::Unknown, ERHIAccess::SRVGraphics),
-				FRHITransitionInfo(ParticleResources->StateTextures[1].PositionTextureRHI, ERHIAccess::Unknown, ERHIAccess::SRVGraphics),
-				FRHITransitionInfo(ParticleResources->StateTextures[1].VelocityTextureRHI, ERHIAccess::Unknown, ERHIAccess::SRVGraphics)
+				FRHITransitionInfo(ParticleResources->RenderAttributesTexture.TextureRHI, ERHIAccess::Unknown, ERHIAccess::SRVMask),
+				FRHITransitionInfo(ParticleResources->SimulationAttributesTexture.TextureRHI, ERHIAccess::Unknown, ERHIAccess::SRVMask),
+				FRHITransitionInfo(ParticleResources->StateTextures[0].PositionTextureRHI, ERHIAccess::Unknown, ERHIAccess::SRVMask),
+				FRHITransitionInfo(ParticleResources->StateTextures[0].VelocityTextureRHI, ERHIAccess::Unknown, ERHIAccess::SRVMask),
+				FRHITransitionInfo(ParticleResources->StateTextures[1].PositionTextureRHI, ERHIAccess::Unknown, ERHIAccess::SRVMask),
+				FRHITransitionInfo(ParticleResources->StateTextures[1].VelocityTextureRHI, ERHIAccess::Unknown, ERHIAccess::SRVMask)
 			});
 		});
 	}
@@ -4658,8 +4658,8 @@ void FFXSystem::PrepareGPUSimulation(FRHICommandListImmediate& RHICmdList)
 	{
 		// Setup render states.
 		RHICmdList.Transition({
-			FRHITransitionInfo(CurrentStateTextures.PositionTextureRHI, ERHIAccess::SRVGraphics, ERHIAccess::RTV),
-			FRHITransitionInfo(CurrentStateTextures.VelocityTextureRHI, ERHIAccess::SRVGraphics, ERHIAccess::RTV)
+			FRHITransitionInfo(CurrentStateTextures.PositionTextureRHI, ERHIAccess::SRVMask, ERHIAccess::RTV),
+			FRHITransitionInfo(CurrentStateTextures.VelocityTextureRHI, ERHIAccess::SRVMask, ERHIAccess::RTV)
 		});
 	}
 }
@@ -4765,16 +4765,16 @@ void FFXSystem::SimulateGPUParticles(
 
 			{
 				RHICmdList.Transition({
-					FRHITransitionInfo(PreviousStateRenderTargets[0], ERHIAccess::SRVGraphics, ERHIAccess::RTV),
-					FRHITransitionInfo(PreviousStateRenderTargets[1], ERHIAccess::SRVGraphics, ERHIAccess::RTV)});
+					FRHITransitionInfo(PreviousStateRenderTargets[0], ERHIAccess::SRVMask, ERHIAccess::RTV),
+					FRHITransitionInfo(PreviousStateRenderTargets[1], ERHIAccess::SRVMask, ERHIAccess::RTV)});
 
 				FRHIRenderPassInfo RPInfo(2, PreviousStateRenderTargets, ERenderTargetActions::Clear_Store);
 				RHICmdList.BeginRenderPass(RPInfo, TEXT("GPUParticlesClearPreviousStateTextures"));
 				RHICmdList.EndRenderPass();
 
 				RHICmdList.Transition({
-					FRHITransitionInfo(PreviousStateRenderTargets[0], ERHIAccess::RTV, ERHIAccess::SRVGraphics),
-					FRHITransitionInfo(PreviousStateRenderTargets[1], ERHIAccess::RTV, ERHIAccess::SRVGraphics) });
+					FRHITransitionInfo(PreviousStateRenderTargets[0], ERHIAccess::RTV, ERHIAccess::SRVMask),
+					FRHITransitionInfo(PreviousStateRenderTargets[1], ERHIAccess::RTV, ERHIAccess::SRVMask) });
 			}
 			
 			PrevStateTextures.bTexturesCleared = true;
@@ -5005,8 +5005,8 @@ void FFXSystem::SimulateGPUParticles(
 		{
 			// Transition attribute textures to writeble, particle state texture are in writeble state already
 			RHICmdList.Transition({ 
-				FRHITransitionInfo(InjectRenderTargets[2], ERHIAccess::SRVGraphics, ERHIAccess::RTV),
-				FRHITransitionInfo(InjectRenderTargets[3], ERHIAccess::SRVGraphics, ERHIAccess::RTV)
+				FRHITransitionInfo(InjectRenderTargets[2], ERHIAccess::SRVMask, ERHIAccess::RTV),
+				FRHITransitionInfo(InjectRenderTargets[3], ERHIAccess::SRVMask, ERHIAccess::RTV)
 			});
 
 			RHICmdList.BeginRenderPass(RPInfo, TEXT("ParticleInjection"));
@@ -5025,8 +5025,8 @@ void FFXSystem::SimulateGPUParticles(
 			RHICmdList.EndRenderPass();
 
 			RHICmdList.Transition({
-				FRHITransitionInfo(InjectRenderTargets[2], ERHIAccess::RTV, ERHIAccess::SRVGraphics),
-				FRHITransitionInfo(InjectRenderTargets[3], ERHIAccess::RTV, ERHIAccess::SRVGraphics)
+				FRHITransitionInfo(InjectRenderTargets[2], ERHIAccess::RTV, ERHIAccess::SRVMask),
+				FRHITransitionInfo(InjectRenderTargets[3], ERHIAccess::RTV, ERHIAccess::SRVMask)
 			});
 		}
 
@@ -5060,8 +5060,8 @@ void FFXSystem::SimulateGPUParticles(
 	
 	// finish current state render
 	RHICmdList.Transition({ 
-		FRHITransitionInfo(CurrentStateRenderTargets[0], ERHIAccess::RTV, ERHIAccess::SRVGraphics),
-		FRHITransitionInfo(CurrentStateRenderTargets[1], ERHIAccess::RTV, ERHIAccess::SRVGraphics)});
+		FRHITransitionInfo(CurrentStateRenderTargets[0], ERHIAccess::RTV, ERHIAccess::SRVMask),
+		FRHITransitionInfo(CurrentStateRenderTargets[1], ERHIAccess::RTV, ERHIAccess::SRVMask)});
 
 	RHICmdList.EndUpdateMultiFrameResource(CurrentStateRenderTargets[0]);
 	RHICmdList.EndUpdateMultiFrameResource(CurrentStateRenderTargets[1]);
@@ -5073,8 +5073,8 @@ void FFXSystem::SimulateGPUParticles(
 		FParticleStateTextures& VisualizeStateTextures = ParticleSimulationResources->GetPreviousStateTextures();
 		
 		RHICmdList.Transition({ 
-			FRHITransitionInfo(VisualizeStateTextures.PositionTextureRHI, ERHIAccess::SRVGraphics, ERHIAccess::RTV),
-			FRHITransitionInfo(VisualizeStateTextures.VelocityTextureRHI, ERHIAccess::SRVGraphics, ERHIAccess::RTV)
+			FRHITransitionInfo(VisualizeStateTextures.PositionTextureRHI, ERHIAccess::SRVMask, ERHIAccess::RTV),
+			FRHITransitionInfo(VisualizeStateTextures.VelocityTextureRHI, ERHIAccess::SRVMask, ERHIAccess::RTV)
 		});
 				
 		FRHITexture* VisualizeStateRHIs[2] = { VisualizeStateTextures.PositionTextureRHI, VisualizeStateTextures.VelocityTextureRHI };
@@ -5105,8 +5105,8 @@ void FFXSystem::SimulateGPUParticles(
 		}
 
 		RHICmdList.Transition({ 
-			FRHITransitionInfo(VisualizeStateTextures.PositionTextureRHI, ERHIAccess::RTV, ERHIAccess::SRVGraphics),
-			FRHITransitionInfo(VisualizeStateTextures.VelocityTextureRHI, ERHIAccess::RTV, ERHIAccess::SRVGraphics)
+			FRHITransitionInfo(VisualizeStateTextures.PositionTextureRHI, ERHIAccess::RTV, ERHIAccess::SRVMask),
+			FRHITransitionInfo(VisualizeStateTextures.VelocityTextureRHI, ERHIAccess::RTV, ERHIAccess::SRVMask)
 		});
 	}
 
