Fix an issue where the Input Device ID could be uninitzalized, causing the Common Game viewport client to ignore any rerouted input. This resulted in some in game UI not being triggered. Initalize the DeviceId to the old controller ID by default, so that anyone who has subclasses viewports wony break with these updates like we saw happen here iwth the common viewport.

#rnx

diff --git a/Engine/Source/Runtime/Engine/Private/Slate/SceneViewport.cpp b/Engine/Source/Runtime/Engine/Private/Slate/SceneViewport.cpp
--- a/Engine/Source/Runtime/Engine/Private/Slate/SceneViewport.cpp
+++ b/Engine/Source/Runtime/Engine/Private/Slate/SceneViewport.cpp
@@ -17,7 +17,7 @@
 #include "Framework/Application/SlateUser.h"
 #include "Slate/SlateTextures.h"
 #include "Slate/DebugCanvas.h"
-
+#include "GenericPlatform/GenericPlatformInputDeviceMapper.h"
 #include "IHeadMountedDisplay.h"
 #include "IXRTrackingSystem.h"
 #include "StereoRenderTargetManager.h"
@@ -219,7 +219,8 @@ void FSceneViewport::ProcessInput( float DeltaTime )
 
 void FSceneViewport::UpdateCachedCursorPos( const FGeometry& InGeometry, const FPointerEvent& InMouseEvent )
 {
-	if (InMouseEvent.GetUserIndex() == FSlateApplication::CursorUserIndex)
+	const FPlatformUserId UserId = IPlatformInputDeviceMapper::Get().GetUserForInputDevice(InMouseEvent.GetInputDeviceId());
+	if (UserId == FSlateApplication::SlateAppPrimaryPlatformUser)
 	{
 		FVector2D LocalPixelMousePos = InGeometry.AbsoluteToLocal(InMouseEvent.GetScreenSpacePosition());
 		LocalPixelMousePos.X *= CachedGeometry.Scale;
@@ -1030,7 +1031,7 @@ FReply FSceneViewport::OnKeyDown( const FGeometry& InGeometry, const FKeyEvent&
 			// Switch to the viewport clients world before processing input
 			FScopedConditionalWorldSwitcher WorldSwitcher(ViewportClient);
 
-			if (!ViewportClient->InputKey(FInputKeyEventArgs(this, InKeyEvent.GetUserIndex(), Key, InKeyEvent.IsRepeat() ? IE_Repeat : IE_Pressed, 1.0f, false)))
+			if (!ViewportClient->InputKey(FInputKeyEventArgs(this, InKeyEvent.GetInputDeviceId(), Key, InKeyEvent.IsRepeat() ? IE_Repeat : IE_Pressed, 1.0f, false)))
 			{
 				CurrentReplyState = FReply::Unhandled();
 			}
@@ -1058,7 +1059,7 @@ FReply FSceneViewport::OnKeyUp( const FGeometry& InGeometry, const FKeyEvent& In
 			// Switch to the viewport clients world before processing input
 			FScopedConditionalWorldSwitcher WorldSwitcher(ViewportClient);
 
-			if (!ViewportClient->InputKey(FInputKeyEventArgs(this, InKeyEvent.GetUserIndex(), Key, IE_Released, 1.0f, false)))
+			if (!ViewportClient->InputKey(FInputKeyEventArgs(this, InKeyEvent.GetInputDeviceId(), Key, IE_Released, 1.0f, false)))
 			{
 				CurrentReplyState = FReply::Unhandled();
 			}
@@ -1087,7 +1088,7 @@ FReply FSceneViewport::OnAnalogValueChanged(const FGeometry& MyGeometry, const F
 			// Switch to the viewport clients world before processing input
 			FScopedConditionalWorldSwitcher WorldSwitcher(ViewportClient);
 
-			if (!ViewportClient->InputAxis(this, InAnalogInputEvent.GetUserIndex(), Key, Key == EKeys::Gamepad_RightY ? -InAnalogInputEvent.GetAnalogValue() : InAnalogInputEvent.GetAnalogValue(), FApp::GetDeltaTime(), 1, Key.IsGamepadKey()))
+			if (!ViewportClient->InputAxis(this, InAnalogInputEvent.GetInputDeviceId(), Key, Key == EKeys::Gamepad_RightY ? -InAnalogInputEvent.GetAnalogValue() : InAnalogInputEvent.GetAnalogValue(), FApp::GetDeltaTime(), 1, Key.IsGamepadKey()))
 			{
 				CurrentReplyState = FReply::Unhandled();
 			}
diff --git a/Engine/Source/Runtime/SlateCore/Public/Input/Events.h b/Engine/Source/Runtime/SlateCore/Public/Input/Events.h
--- a/Engine/Source/Runtime/SlateCore/Public/Input/Events.h
+++ b/Engine/Source/Runtime/SlateCore/Public/Input/Events.h
@@ -177,6 +177,7 @@
 		: ModifierKeys(InModifierKeys)
 		, bIsRepeat(bInIsRepeat)
 		, UserIndex(InUserIndex)
+		, InputDeviceId(FInputDeviceId::CreateFromInternalId(InUserIndex))
 		, EventPath(nullptr)
 	{ }
 
