Unity build fix

diff --git a/Engine/Source/Runtime/Engine/Public/WorldPartition/WorldPartitionMiniMapHelper.h b/Engine/Source/Runtime/Engine/Public/WorldPartition/WorldPartitionMiniMapHelper.h
--- a/Engine/Source/Runtime/Engine/Public/WorldPartition/WorldPartitionMiniMapHelper.h
+++ b/Engine/Source/Runtime/Engine/Public/WorldPartition/WorldPartitionMiniMapHelper.h
@@ -11,6 +11,7 @@
 #if WITH_EDITOR
 #include "Math/Box.h"
 #include "Math/Matrix.h"
+#include "Engine/EngineTypes.h"
 
 class AWorldPartitionMiniMap;
 class UWorld;
