Static mesh
- Fix runtime build from meshdescription when rebuilding the same mesh twice. Made sure we don't recreate render resources when not needed since we can call build from outside of the game thread.
- When we release resources, the renderdata ptr doesn't get nullified so we check the initialized flag instead.

diff --git a/Engine/Source/Runtime/Engine/Private/StaticMesh.cpp b/Engine/Source/Runtime/Engine/Private/StaticMesh.cpp
--- a/Engine/Source/Runtime/Engine/Private/StaticMesh.cpp
+++ b/Engine/Source/Runtime/Engine/Private/StaticMesh.cpp
@@ -6396,7 +6396,7 @@ bool UStaticMesh::BuildFromMeshDescriptions(const TArray<const FMeshDescription*
 	TOptional<FStaticMeshComponentRecreateRenderStateContext> RecreateRenderStateContext;
 	
 	bool bNewMesh = true;
-	if (GetRenderData())
+	if (AreRenderingResourcesInitialized())
 	{
 		bNewMesh = false;
 		const bool bInvalidateLighting = true;
