Made Character OnLanded multicast delegate BlueprintAssignable, consistent with the other multicast delegates

diff --git a/Engine/Source/Runtime/Engine/Classes/GameFramework/Character.h b/Engine/Source/Runtime/Engine/Classes/GameFramework/Character.h
--- a/Engine/Source/Runtime/Engine/Classes/GameFramework/Character.h
+++ b/Engine/Source/Runtime/Engine/Classes/GameFramework/Character.h
@@ -740,6 +740,7 @@ public:
 	 * @param Hit Result describing the landing that resulted in a valid landing spot.
 	 * @see OnMovementModeChanged()
 	 */
+	UPROPERTY(BlueprintAssignable, Category=Character)
 	FLandedSignature LandedDelegate;
 
 	/**
