Tentative fix for non deterministic Texture serialize call

diff --git a/Engine/Source/Runtime/Engine/Private/Texture.cpp b/Engine/Source/Runtime/Engine/Private/Texture.cpp
--- a/Engine/Source/Runtime/Engine/Private/Texture.cpp
+++ b/Engine/Source/Runtime/Engine/Private/Texture.cpp
@@ -807,11 +807,11 @@ void UTexture::Serialize(FArchive& Ar)
 	{
 		// Could potentially guard this with a new custom version, but overhead of just checking on every load should be very small
 		Source.EnsureBlocksAreSorted();
-	}
 
-	if ( GetLinkerUEVersion() < VER_UE4_TEXTURE_LEGACY_GAMMA )
-	{
-		bUseLegacyGamma = true;
+		if (Ar.UEVer() < VER_UE4_TEXTURE_LEGACY_GAMMA)
+		{
+			bUseLegacyGamma = true;
+		}
 	}
 
 	if (Ar.IsCooking() && VirtualTextureStreaming)
