Fixes for recent changes to Niagara's rigid body DI
-make sure transform is applied for canvas rendering of boxes
-we can't assume that the DIs match in exposed/override parameters in the debug hud rendering
-make sure that ActorTags are properly ignored if none are supplied

diff --git a/Engine/Plugins/FX/Niagara/Source/Niagara/Private/NiagaraDataInterfaceRigidMeshCollisionQuery.cpp b/Engine/Plugins/FX/Niagara/Source/Niagara/Private/NiagaraDataInterfaceRigidMeshCollisionQuery.cpp
--- a/Engine/Plugins/FX/Niagara/Source/Niagara/Private/NiagaraDataInterfaceRigidMeshCollisionQuery.cpp
+++ b/Engine/Plugins/FX/Niagara/Source/Niagara/Private/NiagaraDataInterfaceRigidMeshCollisionQuery.cpp
@@ -1724,7 +1724,7 @@ bool UNiagaraDataInterfaceRigidMeshCollisionQuery::FilterActor(const AActor* Act
 		return true;
 	}
 
-	if (!ActorTags.ContainsByPredicate([&](const FName& Tag) { return Tag == NAME_None || Actor->Tags.Contains(Tag); }))
+	if (!ActorTags.IsEmpty() && !ActorTags.ContainsByPredicate([&](const FName& Tag) { return Tag == NAME_None || Actor->Tags.Contains(Tag); }))
 	{
 		return true;
 	}
diff --git a/Engine/Plugins/FX/Niagara/Source/Niagara/Private/NiagaraDebugHud.cpp b/Engine/Plugins/FX/Niagara/Source/Niagara/Private/NiagaraDebugHud.cpp
--- a/Engine/Plugins/FX/Niagara/Source/Niagara/Private/NiagaraDebugHud.cpp
+++ b/Engine/Plugins/FX/Niagara/Source/Niagara/Private/NiagaraDebugHud.cpp
@@ -2634,7 +2634,7 @@ void FNiagaraDebugHud::DrawComponents(FNiagaraWorldManager* WorldManager, UCanva
 								{
 									if (UserVariable.IsDataInterface())
 									{
-										UNiagaraDataInterface* DataInterface = ParameterStore->GetDataInterfaces()[UserVariable.Offset];
+										UNiagaraDataInterface* DataInterface = ParameterStore->GetDataInterface(UserVariable);
 										FString DataInterfaceString;
 										if ( DataInterface != nullptr && Settings.DataInterfaceVerbosity != ENiagaraDebugHudVerbosity::None )
 										{
diff --git a/Engine/Source/Runtime/Engine/Private/DrawDebugHelpers.cpp b/Engine/Source/Runtime/Engine/Private/DrawDebugHelpers.cpp
--- a/Engine/Source/Runtime/Engine/Private/DrawDebugHelpers.cpp
+++ b/Engine/Source/Runtime/Engine/Private/DrawDebugHelpers.cpp
@@ -1180,14 +1180,14 @@ void DrawDebugCanvasWireBox(UCanvas* Canvas, const FMatrix& Transform, const FBo
 {
 	const FVector Vertices[] =
 	{
-		FVector(Box.Min.X, Box.Min.Y, Box.Min.Z),
-		FVector(Box.Min.X, Box.Min.Y, Box.Max.Z),
-		FVector(Box.Min.X, Box.Max.Y, Box.Min.Z),
-		FVector(Box.Min.X, Box.Max.Y, Box.Max.Z),
-		FVector(Box.Max.X, Box.Min.Y, Box.Min.Z),
-		FVector(Box.Max.X, Box.Min.Y, Box.Max.Z),
-		FVector(Box.Max.X, Box.Max.Y, Box.Min.Z),
-		FVector(Box.Max.X, Box.Max.Y, Box.Max.Z)
+		Transform.TransformPosition(FVector(Box.Min.X, Box.Min.Y, Box.Min.Z)),
+		Transform.TransformPosition(FVector(Box.Min.X, Box.Min.Y, Box.Max.Z)),
+		Transform.TransformPosition(FVector(Box.Min.X, Box.Max.Y, Box.Min.Z)),
+		Transform.TransformPosition(FVector(Box.Min.X, Box.Max.Y, Box.Max.Z)),
+		Transform.TransformPosition(FVector(Box.Max.X, Box.Min.Y, Box.Min.Z)),
+		Transform.TransformPosition(FVector(Box.Max.X, Box.Min.Y, Box.Max.Z)),
+		Transform.TransformPosition(FVector(Box.Max.X, Box.Max.Y, Box.Min.Z)),
+		Transform.TransformPosition(FVector(Box.Max.X, Box.Max.Y, Box.Max.Z))
 	};
 
 	const FIntVector2 Edges[] =
