Fix stale TypedElement crash when uploading a project FWorldPartitionHelpers

diff --git a/Engine/Source/Runtime/Engine/Private/WorldPartition/WorldPartitionHelpers.cpp b/Engine/Source/Runtime/Engine/Private/WorldPartition/WorldPartitionHelpers.cpp
--- a/Engine/Source/Runtime/Engine/Private/WorldPartition/WorldPartitionHelpers.cpp
+++ b/Engine/Source/Runtime/Engine/Private/WorldPartition/WorldPartitionHelpers.cpp
@@ -178,7 +178,7 @@ bool FWorldPartitionHelpers::HasExceededMaxMemory()
 void FWorldPartitionHelpers::DoCollectGarbage()
 {
 	const FPlatformMemoryStats MemStatsBefore = FPlatformMemory::GetStats();
-	CollectGarbage(RF_NoFlags, true);
+	CollectGarbage(GARBAGE_COLLECTION_KEEPFLAGS, true);
 	const FPlatformMemoryStats MemStatsAfter = FPlatformMemory::GetStats();
 
 	UE_LOG(LogWorldPartition, Log, TEXT("GC Performed - Available Physical: %.2fGB, Available Virtual: %.2fGB"),
