Ensure that editor-only subobjects are instanced on load for cooked packages when running on platform targets that include editor-only data.

Change summary:
- Modified UObject::PostLoadSubobjects() to allow subobject instancing on targets that do not require cooked data (i.e. targets that include editor-only data).
- Reverted 20563066 (temp fix that also needed to be backed out in order for editor-only subobject instancing paths to function properly).

diff --git a/Engine/Source/Runtime/CoreUObject/Private/UObject/Obj.cpp b/Engine/Source/Runtime/CoreUObject/Private/UObject/Obj.cpp
--- a/Engine/Source/Runtime/CoreUObject/Private/UObject/Obj.cpp
+++ b/Engine/Source/Runtime/CoreUObject/Private/UObject/Obj.cpp
@@ -1160,8 +1160,12 @@ void UObject::PostLoadSubobjects( FObjectInstancingGraph* OuterInstanceGraph/*=N
 		ClearFlags(RF_NeedPostLoadSubobjects);
 
 		// Cooked data will already have its subobjects fully instanced as uninstanced subobjects are only due to newly introduced subobjects in
-		// an archetype that an instance of that object hasn't been saved with
-		if (!FPlatformProperties::RequiresCookedData() && !GetPackage()->HasAnyPackageFlags(PKG_Cooked))
+		// an archetype that an instance of that object hasn't been saved with. Platforms that include editor-only data still require this step
+		// if the outer package is cooked in order to properly instance any editor-only subobjects that will not have otherwise been serialized.
+		// 
+		// @todo_LoadPerf - Consider modifying to bypass this for cooked packages if the class doesn't include any instanced editor-only fields.
+		// Currently we do not have a flag that indicates this scenario, and it might be expensive to iterate over the linked property chain here.
+		if (!FPlatformProperties::RequiresCookedData())
 		{
 			FObjectInstancingGraph CurrentInstanceGraph;
 
diff --git a/Engine/Source/Runtime/Engine/Classes/Engine/BlueprintGeneratedClass.h b/Engine/Source/Runtime/Engine/Classes/Engine/BlueprintGeneratedClass.h
--- a/Engine/Source/Runtime/Engine/Classes/Engine/BlueprintGeneratedClass.h
+++ b/Engine/Source/Runtime/Engine/Classes/Engine/BlueprintGeneratedClass.h
@@ -799,9 +799,8 @@ protected:
 	* @param	InStruct			the current scope for which the given property list applies
 	* @param	DataPtr				destination address (where to start copying values to)
 	* @param	DefaultDataPtr		source address (where to start copying the defaults data from)
-	* @param	bSkipEditorOnly		whether to skip editor-only properties
 	*/
-	static void InitPropertiesFromCustomList(const FCustomPropertyListNode* InPropertyList, UStruct* InStruct, uint8* DataPtr, const uint8* DefaultDataPtr, bool bSkipEditorOnly = false);
+	static void InitPropertiesFromCustomList(const FCustomPropertyListNode* InPropertyList, UStruct* InStruct, uint8* DataPtr, const uint8* DefaultDataPtr);
 
 	/**
 	* Helper method to assist with initializing from an array property with an explicit item list.
@@ -810,9 +809,8 @@ protected:
 	* @param	InPropertyList		only these properties (indices) will be copied from defaults
 	* @param	DataPtr				destination address (where to start copying values to)
 	* @param	DefaultDataPtr		source address (where to start copying the defaults data from)
-	* @param	bSkipEditorOnly		whether to skip editor-only properties
 	*/
-	static void InitArrayPropertyFromCustomList(const FArrayProperty* ArrayProperty, const FCustomPropertyListNode* InPropertyList, uint8* DataPtr, const uint8* DefaultDataPtr, bool bSkipEditorOnly = false);
+	static void InitArrayPropertyFromCustomList(const FArrayProperty* ArrayProperty, const FCustomPropertyListNode* InPropertyList, uint8* DataPtr, const uint8* DefaultDataPtr);
 
 
 	/**
diff --git a/Engine/Source/Runtime/Engine/Private/BlueprintGeneratedClass.cpp b/Engine/Source/Runtime/Engine/Private/BlueprintGeneratedClass.cpp
--- a/Engine/Source/Runtime/Engine/Private/BlueprintGeneratedClass.cpp
+++ b/Engine/Source/Runtime/Engine/Private/BlueprintGeneratedClass.cpp
@@ -836,32 +836,18 @@ void UBlueprintGeneratedClass::InitPropertiesFromCustomList(uint8* DataPtr, cons
 
 	if (const FCustomPropertyListNode* CustomPropertyList = GetCustomPropertyListForPostConstruction())
 	{
-#if WITH_EDITOR
-		// Don't inherit CDO editor-only properties if it comes from a cooked package without editor data
-		const bool bSkipEditorOnly = RootPackageHasAnyFlags(PKG_FilterEditorOnly);
-#else
-		const bool bSkipEditorOnly = false;
-#endif //WITH_EDITOR
-
-		InitPropertiesFromCustomList(CustomPropertyList, this, DataPtr, DefaultDataPtr, bSkipEditorOnly);
+		InitPropertiesFromCustomList(CustomPropertyList, this, DataPtr, DefaultDataPtr);
 	}
 }
 
-void UBlueprintGeneratedClass::InitPropertiesFromCustomList(const FCustomPropertyListNode* InPropertyList, UStruct* InStruct, uint8* DataPtr, const uint8* DefaultDataPtr, bool bSkipEditorOnly)
+void UBlueprintGeneratedClass::InitPropertiesFromCustomList(const FCustomPropertyListNode* InPropertyList, UStruct* InStruct, uint8* DataPtr, const uint8* DefaultDataPtr)
 {
 	for (const FCustomPropertyListNode* CustomPropertyListNode = InPropertyList; CustomPropertyListNode; CustomPropertyListNode = CustomPropertyListNode->PropertyListNext)
 	{
-#if WITH_EDITOR
-		if (bSkipEditorOnly && CustomPropertyListNode->Property->IsEditorOnlyProperty())
-		{
-			continue;
-		}
-#endif //WITH_EDITOR
-
 		uint8* PropertyValue = CustomPropertyListNode->Property->ContainerPtrToValuePtr<uint8>(DataPtr, CustomPropertyListNode->ArrayIndex);
 		const uint8* DefaultPropertyValue = CustomPropertyListNode->Property->ContainerPtrToValuePtr<uint8>(DefaultDataPtr, CustomPropertyListNode->ArrayIndex);
 
-		if (!InitPropertyFromSubPropertyList(CustomPropertyListNode->Property, CustomPropertyListNode->SubPropertyList, PropertyValue, DefaultPropertyValue, bSkipEditorOnly))
+		if (!InitPropertyFromSubPropertyList(CustomPropertyListNode->Property, CustomPropertyListNode->SubPropertyList, PropertyValue, DefaultPropertyValue))
 		{
 			// Unable to init properties from sub custom property list, fall back to the default copy value behavior
 			CustomPropertyListNode->Property->CopySingleValue(PropertyValue, DefaultPropertyValue);
@@ -869,7 +855,7 @@ void UBlueprintGeneratedClass::InitPropertiesFromCustomList(const FCustomPropert
 	}
 }
 
-void UBlueprintGeneratedClass::InitArrayPropertyFromCustomList(const FArrayProperty* ArrayProperty, const FCustomPropertyListNode* InPropertyList, uint8* DataPtr, const uint8* DefaultDataPtr, bool bSkipEditorOnly)
+void UBlueprintGeneratedClass::InitArrayPropertyFromCustomList(const FArrayProperty* ArrayProperty, const FCustomPropertyListNode* InPropertyList, uint8* DataPtr, const uint8* DefaultDataPtr)
 {
 	FScriptArrayHelper DstArrayValueHelper(ArrayProperty, DataPtr);
 	FScriptArrayHelper SrcArrayValueHelper(ArrayProperty, DefaultDataPtr);
@@ -898,7 +884,7 @@ void UBlueprintGeneratedClass::InitArrayPropertyFromCustomList(const FArrayPrope
 			continue;
 		}
 
-		if (!InitPropertyFromSubPropertyList(ArrayProperty->Inner, CustomArrayPropertyListNode->SubPropertyList, DstArrayItemValue, SrcArrayItemValue, bSkipEditorOnly))
+		if (!InitPropertyFromSubPropertyList(ArrayProperty->Inner, CustomArrayPropertyListNode->SubPropertyList, DstArrayItemValue, SrcArrayItemValue))
 		{
 			// Unable to init properties from sub custom property list, fall back to the default copy value behavior
 			ArrayProperty->Inner->CopyCompleteValue(DstArrayItemValue, SrcArrayItemValue);
@@ -906,18 +892,18 @@ void UBlueprintGeneratedClass::InitArrayPropertyFromCustomList(const FArrayPrope
 	}
 }
 
-bool UBlueprintGeneratedClass::InitPropertyFromSubPropertyList(const FProperty* Property, const FCustomPropertyListNode* SubPropertyList, uint8* PropertyValue, const uint8* DefaultPropertyValue, bool bSkipEditorOnly)
+bool UBlueprintGeneratedClass::InitPropertyFromSubPropertyList(const FProperty* Property, const FCustomPropertyListNode* SubPropertyList, uint8* PropertyValue, const uint8* DefaultPropertyValue)
 {
 	if (SubPropertyList != nullptr)
 	{
 		if (const FStructProperty* StructProperty = CastField<FStructProperty>(Property))
 		{
-			InitPropertiesFromCustomList(SubPropertyList, StructProperty->Struct, PropertyValue, DefaultPropertyValue, bSkipEditorOnly);
+			InitPropertiesFromCustomList(SubPropertyList, StructProperty->Struct, PropertyValue, DefaultPropertyValue);
 			return true;
 		}
 		else if (const FArrayProperty* ArrayProperty = CastField<FArrayProperty>(Property))
 		{
-			InitArrayPropertyFromCustomList(ArrayProperty, SubPropertyList, PropertyValue, DefaultPropertyValue, bSkipEditorOnly);
+			InitArrayPropertyFromCustomList(ArrayProperty, SubPropertyList, PropertyValue, DefaultPropertyValue);
 			return true;
 		}
 		// No need to handle Sets and Maps as they are not optimized through the custom property list
