Fix for an overlooked engine-side only API signature during last manual merge (20606273).

#rnx

diff --git a/Engine/Source/Runtime/Engine/Classes/Engine/BlueprintGeneratedClass.h b/Engine/Source/Runtime/Engine/Classes/Engine/BlueprintGeneratedClass.h
--- a/Engine/Source/Runtime/Engine/Classes/Engine/BlueprintGeneratedClass.h
+++ b/Engine/Source/Runtime/Engine/Classes/Engine/BlueprintGeneratedClass.h
@@ -820,10 +820,9 @@ protected:
 	* @param	CustomPropertySubListNode	the list node of the sub property to copy
 	* @param	PropertyValue				destination address (where to start copying values to)
 	* @param	DefaultPropertyValue		source address (where to start copying the defaults data from)
-	* @param	bSkipEditorOnly		whether to skip editor-only properties
 	* @return	true if the method was able to copy the property successfully via the sub custom property list
 	*/
-	static bool InitPropertyFromSubPropertyList(const FProperty* Property, const FCustomPropertyListNode* SubPropertyList, uint8* PropertyValue, const uint8* DefaultPropertyValue, bool bSkipEditorOnly = false);
+	static bool InitPropertyFromSubPropertyList(const FProperty* Property, const FCustomPropertyListNode* SubPropertyList, uint8* PropertyValue, const uint8* DefaultPropertyValue);
 
 	// @todo: BP2CPP_remove
 	/** Check for and handle manual application of default value overrides to component subobjects that were inherited from a nativized parent class */
