Fix a crash that some people had when opening some anim sequence in the editor, related to skeleton remapping/compatibility.
This was reported by someone on twitter, and the suggested fix was confirmed to work.

diff --git a/Engine/Source/Runtime/Engine/Private/Animation/Skeleton.cpp b/Engine/Source/Runtime/Engine/Private/Animation/Skeleton.cpp
--- a/Engine/Source/Runtime/Engine/Private/Animation/Skeleton.cpp
+++ b/Engine/Source/Runtime/Engine/Private/Animation/Skeleton.cpp
@@ -148,10 +148,13 @@ bool USkeleton::IsCompatibleSkeletonByAssetString(const FString& SkeletonAssetSt
 	FString CompatibleName;
 	for (const auto& CompatibleSkeleton : CompatibleSkeletons)
 	{
-		CompatibleName = FObjectPropertyBase::GetExportPath(CompatibleSkeleton.Get());
-		if (CompatibleName == SkeletonAssetString)
+		if (CompatibleSkeleton.IsValid())
 		{
-			return true;
+			CompatibleName = FObjectPropertyBase::GetExportPath(CompatibleSkeleton.Get());
+			if (CompatibleName == SkeletonAssetString)
+			{
+				return true;
+			}
 		}
 	}
 
