World Partition Actor Descriptor Properties:
- Wrapped properties management into engine class FPropertyPairsMap.
- Added engine unit-tests for new FPropertyPairsMap class.
- Don't write the RVT tag in descriptor properties when the array only contain null entries.

#rnx

diff --git a/Engine/Source/Runtime/Engine/Classes/Components/ActorComponent.h b/Engine/Source/Runtime/Engine/Classes/Components/ActorComponent.h
--- a/Engine/Source/Runtime/Engine/Classes/Components/ActorComponent.h
+++ b/Engine/Source/Runtime/Engine/Classes/Components/ActorComponent.h
@@ -14,6 +14,7 @@
 #include "EdGraph/EdGraphPin.h"
 #include "Interfaces/Interface_AssetUserData.h"
 #include "UObject/StructOnScope.h"
+#include "PropertyPairsMap.h"
 #include "ComponentInstanceDataCache.h"
 #include "ActorComponent.generated.h"
 
@@ -843,7 +844,7 @@ public:
 	/**
 	 * Add properties to the component owner's actor desc.
 	 */
-	virtual void GetActorDescProperties(TMap<FName, FName>& Properties) const { }
+	virtual void GetActorDescProperties(FPropertyPairsMap& PropertyPairsMap) const {}
 #endif // WITH_EDITOR
 
 	/**
diff --git a/Engine/Source/Runtime/Engine/Classes/Components/PrimitiveComponent.h b/Engine/Source/Runtime/Engine/Classes/Components/PrimitiveComponent.h
--- a/Engine/Source/Runtime/Engine/Classes/Components/PrimitiveComponent.h
+++ b/Engine/Source/Runtime/Engine/Classes/Components/PrimitiveComponent.h
@@ -2165,7 +2165,7 @@ public:
 	virtual void OnComponentDestroyed(bool bDestroyingHierarchy) override;
 #if WITH_EDITOR
 	virtual void CheckForErrors() override;
-	virtual void GetActorDescProperties(TMap<FName, FName>& Properties) const;
+	virtual void GetActorDescProperties(FPropertyPairsMap& PropertyPairsMap) const;
 #endif // WITH_EDITOR	
 	//~ End UActorComponent Interface
 
diff --git a/Engine/Source/Runtime/Engine/Classes/GameFramework/Actor.h b/Engine/Source/Runtime/Engine/Classes/GameFramework/Actor.h
--- a/Engine/Source/Runtime/Engine/Classes/GameFramework/Actor.h
+++ b/Engine/Source/Runtime/Engine/Classes/GameFramework/Actor.h
@@ -12,6 +12,7 @@
 #include "UObject/CoreNet.h"
 #include "Engine/EngineTypes.h"
 #include "Engine/EngineBaseTypes.h"
+#include "PropertyPairsMap.h"
 #include "ComponentInstanceDataCache.h"
 #include "Components/ChildActorComponent.h"
 #include "RenderCommandFence.h"
@@ -936,7 +937,7 @@ public:
 	/**
 	 * Add properties to the actor desc.
 	 */
-	virtual void GetActorDescProperties(TMap<FName, FName>& Properties) const;
+	virtual void GetActorDescProperties(FPropertyPairsMap& PropertyPairsMap) const;
 
 	/**
 	 * Creates an uninitialized actor descriptor from a specific class.
diff --git a/Engine/Source/Runtime/Engine/Private/ActorEditor.cpp b/Engine/Source/Runtime/Engine/Private/ActorEditor.cpp
--- a/Engine/Source/Runtime/Engine/Private/ActorEditor.cpp
+++ b/Engine/Source/Runtime/Engine/Private/ActorEditor.cpp
@@ -874,11 +874,11 @@ TUniquePtr<FWorldPartitionActorDesc> AActor::CreateActorDesc() const
 	return ActorDesc;
 }
 
-void AActor::GetActorDescProperties(TMap<FName, FName>& Properties) const
+void AActor::GetActorDescProperties(FPropertyPairsMap& PropertyPairsMap) const
 {
-	ForEachComponent<UActorComponent>(false, [&Properties](UActorComponent* Component)
+	ForEachComponent<UActorComponent>(false, [&PropertyPairsMap](UActorComponent* Component)
 	{
-		Component->GetActorDescProperties(Properties);
+		Component->GetActorDescProperties(PropertyPairsMap);
 	});
 }
 
diff --git a/Engine/Source/Runtime/Engine/Private/Components/PrimitiveComponent.cpp b/Engine/Source/Runtime/Engine/Private/Components/PrimitiveComponent.cpp
--- a/Engine/Source/Runtime/Engine/Private/Components/PrimitiveComponent.cpp
+++ b/Engine/Source/Runtime/Engine/Private/Components/PrimitiveComponent.cpp
@@ -1208,11 +1208,17 @@ void UPrimitiveComponent::CheckForErrors()
 	}
 }
 
-void UPrimitiveComponent::GetActorDescProperties(TMap<FName, FName>& Properties) const
+void UPrimitiveComponent::GetActorDescProperties(FPropertyPairsMap& PropertyPairsMap) const
 {
-	if (RuntimeVirtualTextures.Num())
+	Super::GetActorDescProperties(PropertyPairsMap);
+
+	for (URuntimeVirtualTexture* RuntimeVirtualTexture : RuntimeVirtualTextures)
 	{
-		Properties.Add(UPrimitiveComponent::RVTActorDescProperty, NAME_None);
+		if (RuntimeVirtualTexture)
+		{
+			PropertyPairsMap.AddProperty(UPrimitiveComponent::RVTActorDescProperty);
+			return;
+		}
 	}
 }
 
diff --git a/Engine/Source/Runtime/Engine/Private/Tests/PropertyPairsMapTests.cpp b/Engine/Source/Runtime/Engine/Private/Tests/PropertyPairsMapTests.cpp
new file mode 100644
--- /dev/null
+++ b/Engine/Source/Runtime/Engine/Private/Tests/PropertyPairsMapTests.cpp
@@ -0,0 +1,62 @@
+// Copyright Epic Games, Inc. All Rights Reserved.
+
+#include "Misc/AutomationTest.h"
+#include "PropertyPairsMap.h"
+
+#if WITH_DEV_AUTOMATION_TESTS
+
+#define TEST_NAME_ROOT "System.Engine.PropertyPairsMap"
+
+namespace PropertyPairsMapTests
+{
+	constexpr const uint32 TestFlags = EAutomationTestFlags::ApplicationContextMask | EAutomationTestFlags::EngineFilter;
+
+	IMPLEMENT_SIMPLE_AUTOMATION_TEST(FPropertyPairsMapTests, TEST_NAME_ROOT, TestFlags)
+	bool FPropertyPairsMapTests::RunTest(const FString& Parameters)
+	{
+		FPropertyPairsMap PropertyPairsMap;
+
+		TestFalse(TEXT("PropertyPairsMap.HasProperty(Invalid)"), PropertyPairsMap.HasProperty(TEXT("Invalid")));
+
+		FName WithoutValueValue(TEXT("Invalid"));
+		PropertyPairsMap.AddProperty(TEXT("WithoutValue"));
+		TestTrue(TEXT("PropertyPairsMap.HasProperty(WithoutValue)"), PropertyPairsMap.HasProperty(TEXT("WithoutValue")));
+		TestTrue(TEXT("PropertyPairsMap.GetProperty(WithoutValue)"), PropertyPairsMap.GetProperty(TEXT("WithoutValue")));
+		TestTrue(TEXT("WithoutValueValue is Invalid"), WithoutValueValue.ToString() == TEXT("Invalid"));
+		TestTrue(TEXT("PropertyPairsMap.GetProperty(WithVoutValue)"), PropertyPairsMap.GetProperty(TEXT("WithoutValue"), &WithoutValueValue));
+		TestTrue(TEXT("WithoutValueValue is None"), WithoutValueValue.IsNone());
+
+		FName WithValueValue(TEXT("Invalid"));
+		PropertyPairsMap.AddProperty(TEXT("WithValue"), TEXT("Value1"));
+		TestTrue(TEXT("PropertyPairsMap.HasProperty(WithValue)"), PropertyPairsMap.HasProperty(TEXT("WithValue")));
+		TestTrue(TEXT("PropertyPairsMap.GetProperty(WithValue)"), PropertyPairsMap.GetProperty(TEXT("WithValue")));
+		TestTrue(TEXT("WithValueValue is Invalid"), WithValueValue.ToString() == TEXT("Invalid"));
+		TestTrue(TEXT("PropertyPairsMap.GetProperty(WithValue)"), PropertyPairsMap.GetProperty(TEXT("WithValue"), &WithValueValue));
+		TestTrue(TEXT("WithValueValue is Value1"), WithValueValue.ToString() == TEXT("Value1"));
+		PropertyPairsMap.AddProperty(TEXT("WithValue"), TEXT("Value2"));
+		TestTrue(TEXT("PropertyPairsMap.HasProperty(WithValue)"), PropertyPairsMap.HasProperty(TEXT("WithValue")));
+		TestTrue(TEXT("PropertyPairsMap.GetProperty(WithValue)"), PropertyPairsMap.GetProperty(TEXT("WithValue")));
+		TestTrue(TEXT("WithValueValue is Value1"), WithValueValue.ToString() == TEXT("Value1"));
+		TestTrue(TEXT("PropertyPairsMap.GetProperty(WithValue)"), PropertyPairsMap.GetProperty(TEXT("WithValue"), &WithValueValue));
+		TestTrue(TEXT("WithValueValue is Value2"), WithValueValue.ToString() == TEXT("Value2"));
+
+		PropertyPairsMap.ForEachProperty([this](FName Name, FName Value)
+		{
+			TestTrue(TEXT("Unbreakable iterator"), (Name == TEXT("WithoutValue")) || (Name == TEXT("WithtValue")));
+			TestTrue(TEXT("Unbreakable iterator"), Value.IsNone() || (Name == TEXT("Value2")));
+		});
+
+		PropertyPairsMap.ForEachProperty([this](FName Name, FName Value)
+		{
+			TestTrue(TEXT("Breakable iterator"), (Name == TEXT("WithoutValue")) || (Name == TEXT("WithtValue")));
+			TestTrue(TEXT("Breakable iterator"), Value.IsNone() || (Name == TEXT("Value2")));
+			return true;
+		});
+
+		return true;
+	}
+}
+
+#undef TEST_NAME_ROOT
+
+#endif 
diff --git a/Engine/Source/Runtime/Engine/Private/WorldPartition/WorldPartitionActorDesc.cpp b/Engine/Source/Runtime/Engine/Private/WorldPartition/WorldPartitionActorDesc.cpp
--- a/Engine/Source/Runtime/Engine/Private/WorldPartition/WorldPartitionActorDesc.cpp
+++ b/Engine/Source/Runtime/Engine/Private/WorldPartition/WorldPartitionActorDesc.cpp
@@ -98,7 +98,7 @@ void FWorldPartitionActorDesc::Init(const AActor* InActor)
 
 	Tags = InActor->Tags;
 
-	Properties.Empty();
+	check(Properties.IsEmpty());
 	InActor->GetActorDescProperties(Properties);
 
 	ActorPackage = InActor->GetPackage()->GetFName();
@@ -172,7 +172,7 @@ bool FWorldPartitionActorDesc::Equals(const FWorldPartitionActorDesc* Other) con
 		CompareUnsortedArrays(DataLayers, Other->DataLayers) &&
 		CompareUnsortedArrays(References, Other->References) &&
 		CompareUnsortedArrays(Tags, Other->Tags) &&
-		Properties.OrderIndependentCompareEqual(Other->Properties);
+		Properties == Other->Properties;
 }
 
 void FWorldPartitionActorDesc::SerializeTo(TArray<uint8>& OutData)
@@ -483,4 +483,4 @@ void FWorldPartitionActorDesc::Unload()
 		ActorPtr = nullptr;
 	}
 }
-#endif
+#endif
\ No newline at end of file
diff --git a/Engine/Source/Runtime/Engine/Public/PropertyPairsMap.h b/Engine/Source/Runtime/Engine/Public/PropertyPairsMap.h
new file mode 100644
--- /dev/null
+++ b/Engine/Source/Runtime/Engine/Public/PropertyPairsMap.h
@@ -0,0 +1,107 @@
+// Copyright Epic Games, Inc. All Rights Reserved.
+
+#pragma once
+
+#include "CoreTypes.h"
+
+/** A simple name/value pairs map */
+class FPropertyPairsMap
+{
+public:
+	/**
+	 * Add a property/value pair to the map. This will replace any existing property with the new value.
+	 *
+	 * @param InName Name of the property.
+	 * @param InValue Optional value to be associated with the property.
+	 */
+	FORCEINLINE void AddProperty(FName InName, FName InValue = NAME_None)
+	{
+		Properties.Add(InName, InValue);
+	}
+
+	/**
+	 * Test if the provided property exists in the map.
+	 *
+	 * @param InName Name of the property.
+	 * @returns True if the property exist in the map.
+	 */
+	FORCEINLINE bool HasProperty(FName InName) const
+	{
+		return GetProperty(InName);
+	}
+
+	FORCEINLINE bool HasProperty(const TCHAR* InName) const
+	{
+		FName PropertyName(InName, FNAME_Find);
+		return !PropertyName.IsNone() && GetProperty(PropertyName);
+	}
+
+	/**
+	 * Gets the value associated with the provided property.
+	 *
+	 * @param InName Name of the property.
+	 * @out OutValue Value associated with the provided property.
+	 * @returns True if the property exist in the map.
+	 */
+	FORCEINLINE bool GetProperty(FName InName, FName* OutValue = nullptr) const
+	{
+		if (const FName* ValuePtr = Properties.Find(InName))
+		{
+			if (OutValue)
+			{
+				*OutValue = *ValuePtr;
+			}
+			return true;
+		}
+		return false;
+	}
+
+	FORCEINLINE bool GetProperty(const TCHAR* InName, FName* OutValue = nullptr) const
+	{
+		FName PropertyName(InName, FNAME_Find);
+		return !PropertyName.IsNone() && GetProperty(PropertyName, OutValue);
+	}
+
+	/**
+	 * Compare this map with another for equality, ignoring ordering.
+	 *
+	 * @param InOther Other property map to compare against.
+	 * @returns True if the property maps contains the same properties and values.
+	 */
+	bool operator==(const FPropertyPairsMap& InOther) const
+	{
+		return Properties.OrderIndependentCompareEqual(InOther.Properties);
+	}
+
+	bool operator!=(const FPropertyPairsMap& InOther) const
+	{
+		return !(*this == InOther);
+	}
+
+	/**
+	 * Serialize this property map to the provided archive.
+	 *
+	 * @param Ar Archive to serialize to.
+	 * @param InPropertyPairsMap Property map to serialize.
+	 * @returns Archive used to serialize the property map.
+	 */
+	FORCEINLINE friend FArchive& operator<<(FArchive& InAr, FPropertyPairsMap& InPropertyPairsMap)
+	{
+		return InAr << InPropertyPairsMap.Properties;
+	}
+
+	/** @returns True if the property map is empty. */
+	FORCEINLINE bool IsEmpty() const
+	{
+		return Properties.IsEmpty();
+	}
+
+	/** @return The number of properties in the map. */
+	FORCEINLINE int32 Num() const
+	{
+		return Properties.Num();
+	}
+
+protected:
+	TMap<FName, FName> Properties;
+};
\ No newline at end of file
diff --git a/Engine/Source/Runtime/Engine/Public/WorldPartition/WorldPartitionActorDesc.h b/Engine/Source/Runtime/Engine/Public/WorldPartition/WorldPartitionActorDesc.h
--- a/Engine/Source/Runtime/Engine/Public/WorldPartition/WorldPartitionActorDesc.h
+++ b/Engine/Source/Runtime/Engine/Public/WorldPartition/WorldPartitionActorDesc.h
@@ -2,11 +2,11 @@
 #pragma once
 
 #include "CoreMinimal.h"
-#include "AssetRegistry/AssetData.h"
+#include "PropertyPairsMap.h"
 #include "UObject/NameTypes.h"
 #include "UObject/ObjectMacros.h"
-#include "GameFramework/Actor.h"
-#include "Containers/Set.h"
+#include "UObject/WeakObjectPtrTemplates.h"
+#include "UObject/WeakObjectPtr.h"
 #include "Misc/Guid.h"
 
 #if WITH_EDITOR
@@ -19,6 +19,7 @@ struct FWorldPartitionActorDescInitData
 	TArray<uint8> SerializedData;
 };
 
+class AActor;
 class UActorDescContainer;
 struct FActorContainerID;
 
@@ -57,7 +58,7 @@ inline bool CompareUnsortedArrays(const TArray<FName>& Array1, const TArray<FNam
 /**
  * Represents a potentially unloaded actor (editor-only)
  */
-class ENGINE_API FWorldPartitionActorDesc 
+class ENGINE_API FWorldPartitionActorDesc
 {
 #if WITH_EDITOR
 	friend class AActor;
@@ -100,8 +101,8 @@ public:
 	FBox GetBounds() const;
 	inline const FGuid& GetParentActor() const { return ParentActor; }
 	inline bool IsUsingDataLayerAsset() const { return bIsUsingDataLayerAsset; }
-	inline FName GetProperty(FName PropertyName) const { return Properties.FindRef(PropertyName); }
-	inline bool HasProperty(FName PropertyName) const { return Properties.Contains(PropertyName); }
+	inline bool GetProperty(FName PropertyName, FName* PropertyValue) const { return Properties.GetProperty(PropertyName, PropertyValue); }
+	inline bool HasProperty(FName PropertyName) const { return Properties.HasProperty(PropertyName); }
 
 	FName GetActorName() const;
 	FName GetActorLabelOrName() const;
@@ -232,7 +233,7 @@ protected:
 	TArray<FName>					DataLayers;
 	TArray<FGuid>					References;
 	TArray<FName>					Tags;
-	TMap<FName, FName>				Properties;
+	FPropertyPairsMap				Properties;
 	FName							FolderPath;
 	FGuid							FolderGuid;
 
