Revert work-in-progress code from previous CL that wasn't supposed to be submitted.

diff --git a/Engine/Source/Runtime/Engine/Private/Tests/PropertyPairsMapTests.cpp b/Engine/Source/Runtime/Engine/Private/Tests/PropertyPairsMapTests.cpp
--- a/Engine/Source/Runtime/Engine/Private/Tests/PropertyPairsMapTests.cpp
+++ b/Engine/Source/Runtime/Engine/Private/Tests/PropertyPairsMapTests.cpp
@@ -40,19 +40,6 @@ namespace PropertyPairsMapTests
 		TestTrue(TEXT("PropertyPairsMap.GetProperty(WithValue)"), PropertyPairsMap.GetProperty(TEXT("WithValue"), &WithValueValue));
 		TestTrue(TEXT("WithValueValue is Value2"), WithValueValue.ToString() == TEXT("Value2"));
 
-		PropertyPairsMap.ForEachProperty([this](FName Name, FName Value)
-		{
-			TestTrue(TEXT("Unbreakable iterator"), (Name == TEXT("WithoutValue")) || (Name == TEXT("WithtValue")));
-			TestTrue(TEXT("Unbreakable iterator"), Value.IsNone() || (Name == TEXT("Value2")));
-		});
-
-		PropertyPairsMap.ForEachProperty([this](FName Name, FName Value)
-		{
-			TestTrue(TEXT("Breakable iterator"), (Name == TEXT("WithoutValue")) || (Name == TEXT("WithtValue")));
-			TestTrue(TEXT("Breakable iterator"), Value.IsNone() || (Name == TEXT("Value2")));
-			return true;
-		});
-
 		return true;
 	}
 }
