CookPackageSplitter: Add missing function in previous change.
#rnx

diff --git a/Engine/Source/Runtime/Engine/Private/WorldPartition/WorldPartitionCookPackageSplitter.cpp b/Engine/Source/Runtime/Engine/Private/WorldPartition/WorldPartitionCookPackageSplitter.cpp
--- a/Engine/Source/Runtime/Engine/Private/WorldPartition/WorldPartitionCookPackageSplitter.cpp
+++ b/Engine/Source/Runtime/Engine/Private/WorldPartition/WorldPartitionCookPackageSplitter.cpp
@@ -137,7 +137,12 @@ void FWorldPartitionCookPackageSplitter::PreSaveGeneratorPackage(UPackage* Owner
 {
 	UWorld* PartitionedWorld = ValidateDataObject(OwnerObject);
 	UWorldPartition* WorldPartition = PartitionedWorld->PersistentLevel->GetWorldPartition();
-	WorldPartition->FinalizeGeneratorPackageForCook(GeneratedPackages);
+
+	TArray<UObject*> LoadedObjects;
+	if (WorldPartition->LoadGeneratorPackageObjectsForCook(LoadedObjects))
+	{
+		WorldPartition->FinalizeGeneratorPackageForCook(GeneratedPackages);
+	}
 }
 
 void FWorldPartitionCookPackageSplitter::OnOwnerReloaded(UPackage* OwnerPackage, UObject* OwnerObject)
