Add base for async-aware savegame systems to extend.

#rnx

diff --git a/Engine/Source/Runtime/Engine/Private/SaveGameSystem.cpp b/Engine/Source/Runtime/Engine/Private/SaveGameSystem.cpp
--- a/Engine/Source/Runtime/Engine/Private/SaveGameSystem.cpp
+++ b/Engine/Source/Runtime/Engine/Private/SaveGameSystem.cpp
@@ -2,6 +2,7 @@
 
 #include "SaveGameSystem.h"
 #include "HAL/PlatformMisc.h"
+#include "HAL/ThreadHeartBeat.h"
 #include "Containers/Ticker.h"
 
 
@@ -170,3 +171,124 @@ void ISaveGameSystem::OnAsyncComplete(TFunction<void()> Callback)
 		}
 	));
 }
+
+
+
+
+
+
+
+
+ISaveGameSystem::ESaveExistsResult FBaseAsyncSaveGameSystem::DoesSaveGameExistWithResult(const TCHAR* Name, const int32 UserIndex)
+{
+	TSharedPtr<ESaveExistsResult> Result = MakeShared<ESaveExistsResult>(ESaveExistsResult::UnspecifiedError);
+
+	const FPlatformUserId PlatformUserId = FPlatformMisc::GetPlatformUserForUserIndex(UserIndex);
+
+	UE::Tasks::FTask Op = InternalDoesSaveGameExistAsync(Name, PlatformUserId, nullptr, Result);
+	WaitForAsyncTask(Op);
+
+	return *Result;
+}
+
+bool FBaseAsyncSaveGameSystem::SaveGame(bool bAttemptToUseUI, const TCHAR* Name, const int32 UserIndex, const TArray<uint8>& Data)
+{
+	TSharedPtr<bool> Result = MakeShared<bool>(false);
+	TSharedRef<const TArray<uint8>> DataPtr = MakeShared<const TArray<uint8>>(Data); // have to take a copy
+
+	const FPlatformUserId PlatformUserId = FPlatformMisc::GetPlatformUserForUserIndex(UserIndex);
+	UE::Tasks::FTask Op = InternalSaveGameAsync(bAttemptToUseUI, Name, PlatformUserId, DataPtr, nullptr, Result);
+	WaitForAsyncTask(Op);
+
+	return *Result;
+}
+
+bool FBaseAsyncSaveGameSystem::LoadGame(bool bAttemptToUseUI, const TCHAR* Name, const int32 UserIndex, TArray<uint8>& Data)
+{
+	TSharedPtr<bool> Result = MakeShared<bool>(false);
+	TSharedRef<TArray<uint8>> DataPtr = MakeShared<TArray<uint8>>(Data); // have to take a copy
+
+	const FPlatformUserId PlatformUserId = FPlatformMisc::GetPlatformUserForUserIndex(UserIndex);
+	UE::Tasks::FTask Op = InternalLoadGameAsync(bAttemptToUseUI, Name, PlatformUserId, DataPtr, nullptr, Result);
+	WaitForAsyncTask(Op);
+
+	Data = MoveTemp(*DataPtr);
+	return  *Result;
+}
+
+bool FBaseAsyncSaveGameSystem::DeleteGame(bool bAttemptToUseUI, const TCHAR* Name, const int32 UserIndex)
+{
+	TSharedPtr<bool> Result = MakeShared<bool>(false);
+
+	const FPlatformUserId PlatformUserId = FPlatformMisc::GetPlatformUserForUserIndex(UserIndex);
+	UE::Tasks::FTask Op = InternalDeleteGameAsync(bAttemptToUseUI, Name, PlatformUserId, nullptr, Result);
+	WaitForAsyncTask(Op);
+
+	return *Result;
+}
+
+bool FBaseAsyncSaveGameSystem::GetSaveGameNames(TArray<FString>& FoundSaves, const int32 UserIndex)
+{
+	TSharedPtr<bool> Result = MakeShared<bool>(false);
+	TSharedRef<TArray<FString>> FoundSavesPtr = MakeShared<TArray<FString>>();
+
+	const FPlatformUserId PlatformUserId = FPlatformMisc::GetPlatformUserForUserIndex(UserIndex);
+	UE::Tasks::FTask Op = InternalGetSaveGameNamesAsync(PlatformUserId, FoundSavesPtr, nullptr, Result);
+	WaitForAsyncTask(Op);
+
+	FoundSaves = MoveTemp(*FoundSavesPtr);
+	return *Result;
+}
+
+
+
+void FBaseAsyncSaveGameSystem::DoesSaveGameExistAsync(const TCHAR* Name, FPlatformUserId PlatformUserId, FSaveGameAsyncExistsCallback Callback)
+{
+	InternalDoesSaveGameExistAsync(Name, PlatformUserId, Callback);
+}
+
+void FBaseAsyncSaveGameSystem::SaveGameAsync(bool bAttemptToUseUI, const TCHAR* Name, FPlatformUserId PlatformUserId, TSharedRef<const TArray<uint8>> Data, FSaveGameAsyncOpCompleteCallback Callback)
+{
+	InternalSaveGameAsync(bAttemptToUseUI, Name, PlatformUserId, Data, Callback);
+}
+
+void FBaseAsyncSaveGameSystem::LoadGameAsync(bool bAttemptToUseUI, const TCHAR* Name, FPlatformUserId PlatformUserId, FSaveGameAsyncLoadCompleteCallback Callback)
+{
+	TSharedRef<TArray<uint8>> Data = MakeShared<TArray<uint8>>();
+	InternalLoadGameAsync(bAttemptToUseUI, Name, PlatformUserId, Data, Callback);
+}
+
+void FBaseAsyncSaveGameSystem::DeleteGameAsync(bool bAttemptToUseUI, const TCHAR* Name, FPlatformUserId PlatformUserId, FSaveGameAsyncOpCompleteCallback Callback)
+{
+	InternalDeleteGameAsync(bAttemptToUseUI, Name, PlatformUserId, Callback);
+}
+
+void FBaseAsyncSaveGameSystem::GetSaveGameNamesAsync(FPlatformUserId PlatformUserId, FSaveGameAsyncGetNamesCallback Callback)
+{
+	TSharedRef<TArray<FString>> FoundSavesPtr = MakeShared<TArray<FString>>();
+	InternalGetSaveGameNamesAsync(PlatformUserId, FoundSavesPtr, Callback);
+}
+
+
+
+
+void FBaseAsyncSaveGameSystem::WaitForAsyncTask(UE::Tasks::FTask AsyncSaveTask)
+{
+	// need to pump messages on the game thread
+	if (IsInGameThread())
+	{
+		// Suspend the hang and hitch heartbeats, as this is a long running task.
+		FSlowHeartBeatScope SuspendHeartBeat;
+		FDisableHitchDetectorScope SuspendGameThreadHitch;
+
+		while (!AsyncSaveTask.IsCompleted())
+		{
+			FPlatformMisc::PumpMessagesOutsideMainLoop();
+		}
+	}
+	else
+	{
+		// not running on the game thread, so just block until the async operation comes back
+		AsyncSaveTask.BusyWait();
+	}
+}
diff --git a/Engine/Source/Runtime/Engine/Public/SaveGameSystem.h b/Engine/Source/Runtime/Engine/Public/SaveGameSystem.h
--- a/Engine/Source/Runtime/Engine/Public/SaveGameSystem.h
+++ b/Engine/Source/Runtime/Engine/Public/SaveGameSystem.h
@@ -167,3 +167,43 @@ protected:
 		return FString::Printf(TEXT("%sSaveGames/%s.sav"), *FPaths::ProjectSavedDir(), Name);
 	}
 };
+
+
+
+/** helper base class for async-compatible save games */
+class ENGINE_API FBaseAsyncSaveGameSystem : public ISaveGameSystem
+{
+public:
+
+	virtual bool DoesSaveGameExist(const TCHAR* Name, const int32 UserIndex) override
+	{
+		return ESaveExistsResult::OK == DoesSaveGameExistWithResult(Name, UserIndex);
+	}
+
+	// syncronous save functions
+	virtual ESaveExistsResult DoesSaveGameExistWithResult(const TCHAR* Name, const int32 UserIndex) override;
+	virtual bool GetSaveGameNames(TArray<FString>& FoundSaves, const int32 UserIndex) override;
+	virtual bool SaveGame(bool bAttemptToUseUI, const TCHAR* Name, const int32 UserIndex, const TArray<uint8>& Data) override;
+	virtual bool LoadGame(bool bAttemptToUseUI, const TCHAR* Name, const int32 UserIndex, TArray<uint8>& Data) override;
+	virtual bool DeleteGame(bool bAttemptToUseUI, const TCHAR* Name, const int32 UserIndex) override;
+
+	// asyncronous save functions
+	virtual void DoesSaveGameExistAsync(const TCHAR* Name, FPlatformUserId PlatformUserId, FSaveGameAsyncExistsCallback Callback) override;
+	virtual void SaveGameAsync(bool bAttemptToUseUI, const TCHAR* Name, FPlatformUserId PlatformUserId, TSharedRef<const TArray<uint8>> Data, FSaveGameAsyncOpCompleteCallback Callback) override;
+	virtual void LoadGameAsync(bool bAttemptToUseUI, const TCHAR* Name, FPlatformUserId PlatformUserId, FSaveGameAsyncLoadCompleteCallback Callback) override;
+	virtual void DeleteGameAsync(bool bAttemptToUseUI, const TCHAR* Name, FPlatformUserId PlatformUserId, FSaveGameAsyncOpCompleteCallback Callback) override;
+	virtual void GetSaveGameNamesAsync(FPlatformUserId PlatformUserId, FSaveGameAsyncGetNamesCallback Callback) override;
+
+protected:
+	// internal async helpers that require implementation. NB. OutResult will be null for async calls
+	virtual UE::Tasks::FTask InternalDoesSaveGameExistAsync(const TCHAR* Name, FPlatformUserId PlatformUserId, FSaveGameAsyncExistsCallback Callback, TSharedPtr<ESaveExistsResult> OutResult = nullptr) = 0;
+	virtual UE::Tasks::FTask InternalSaveGameAsync(bool bAttemptToUseUI, const TCHAR* Name, FPlatformUserId PlatformUserId, TSharedRef<const TArray<uint8>> Data, FSaveGameAsyncOpCompleteCallback Callback, TSharedPtr<bool> OutResult = nullptr) = 0;
+	virtual UE::Tasks::FTask InternalLoadGameAsync(bool bAttemptToUseUI, const TCHAR* Name, FPlatformUserId PlatformUserId, TSharedRef<TArray<uint8>> Data, FSaveGameAsyncLoadCompleteCallback Callback, TSharedPtr<bool> OutResult = nullptr) = 0;
+	virtual UE::Tasks::FTask InternalDeleteGameAsync(bool bAttemptToUseUI, const TCHAR* Name, FPlatformUserId PlatformUserId, FSaveGameAsyncOpCompleteCallback Callback, TSharedPtr<bool> OutResult = nullptr) = 0;
+	virtual UE::Tasks::FTask InternalGetSaveGameNamesAsync(FPlatformUserId PlatformUserId, TSharedRef<TArray<FString>> FoundSaves, FSaveGameAsyncGetNamesCallback Callback, TSharedPtr<bool> OutResult = nullptr) = 0;
+
+	// specializations can override this
+	virtual void WaitForAsyncTask(UE::Tasks::FTask AsyncSaveTask);
+
+};
+
