Fixed typo in comment.

#rnx

diff --git a/Engine/Source/Runtime/Engine/Public/ReferenceCluster.h b/Engine/Source/Runtime/Engine/Public/ReferenceCluster.h
--- a/Engine/Source/Runtime/Engine/Public/ReferenceCluster.h
+++ b/Engine/Source/Runtime/Engine/Public/ReferenceCluster.h
@@ -12,6 +12,6 @@
  *			list, forming a closed set of n clusters.
  *
  *	@return	Array of reference clusters (e.g. group of objects referencing each other). For example, an
- *			input of n objects without any references will generate be n clusters of a single object each.
+ *			input of n objects without any references will generate n clusters of a single object each.
  */
 ENGINE_API TArray<TArray<FGuid>> GenerateObjectsClusters(TArray<TPair<FGuid, TArray<FGuid>>> InObjects);
\ No newline at end of file
