WorldPartition - Runtime LWC support

diff --git a/Engine/Source/Runtime/Engine/Private/WorldPartition/RuntimeSpatialHash/RuntimeSpatialHashGridHelper.cpp b/Engine/Source/Runtime/Engine/Private/WorldPartition/RuntimeSpatialHash/RuntimeSpatialHashGridHelper.cpp
--- a/Engine/Source/Runtime/Engine/Private/WorldPartition/RuntimeSpatialHash/RuntimeSpatialHashGridHelper.cpp
+++ b/Engine/Source/Runtime/Engine/Private/WorldPartition/RuntimeSpatialHash/RuntimeSpatialHashGridHelper.cpp
@@ -23,40 +23,40 @@ static FAutoConsoleVariableRef CVarRuntimeSpatialHashPlaceSmallActorsUsingLocati
 	GRuntimeSpatialHashPlaceSmallActorsUsingLocation,
 	TEXT("Set RuntimeSpatialHashPlaceSmallActorsUsingLocation to true to place actors smaller than a cell size into their corresponding cell using their location instead of their bounding box."));
 
-FSquare2DGridHelper::FSquare2DGridHelper(const FBox& InWorldBounds, const FVector& InOrigin, int32 InCellSize)
+FSquare2DGridHelper::FSquare2DGridHelper(const FBox& InWorldBounds, const FVector& InOrigin, int64 InCellSize)
 	: WorldBounds(InWorldBounds)
 	, Origin(InOrigin)
 	, CellSize(InCellSize)
 {
 	// Compute Grid's size and level count based on World bounds
-	int32 GridSize = 1;
+	int64 GridSize = 1;
 	int32 GridLevelCount = 1;
 
 	if (WorldBounds.IsValid)
 	{
 		const FVector2D DistMin = FMath::Abs(FVector2D(WorldBounds.Min - Origin));
 		const FVector2D DistMax = FMath::Abs(FVector2D(WorldBounds.Max - Origin));
-		const float WorldBoundsMaxExtent = FMath::Max(DistMin.GetMax(), DistMax.GetMax());
+		const double WorldBoundsMaxExtent = FMath::Max(DistMin.GetMax(), DistMax.GetMax());
 
-		if (WorldBoundsMaxExtent > 0.f)
+		if (WorldBoundsMaxExtent > 0)
 		{
-			GridSize = 2.f * FMath::CeilToFloat(WorldBoundsMaxExtent / CellSize); 
+			GridSize = 2 * FMath::CeilToDouble(WorldBoundsMaxExtent / CellSize); 
 			if (!FMath::IsPowerOfTwo(GridSize))
 			{
-				GridSize = FMath::Pow(2.f, FMath::CeilToFloat(FMath::Log2(static_cast<float>(GridSize))));
+				GridSize = FMath::Pow(2, FMath::CeilToDouble(FMath::Log2(static_cast<double>(GridSize))));
 			}
-			GridLevelCount = FMath::FloorLog2(GridSize) + 1;
+			GridLevelCount = FMath::FloorLog2_64(GridSize) + 1;
 		}
 	}
 
 	check(FMath::IsPowerOfTwo(GridSize));
 
 	Levels.Reserve(GridLevelCount);
-	int32 CurrentCellSize = CellSize;
-	int32 CurrentGridSize = GridSize;
+	int64 CurrentCellSize = CellSize;
+	int64 CurrentGridSize = GridSize;
 	for (int32 Level = 0; Level < GridLevelCount; ++Level)
 	{
-		int32 LevelGridSize = CurrentGridSize;
+		int64 LevelGridSize = CurrentGridSize;
 
 		if (!GRuntimeSpatialHashUseAlignedGridLevels)
 		{
@@ -92,37 +92,37 @@ void FSquare2DGridHelper::ForEachCells(TFunctionRef<void(const FSquare2DGridHelp
 }
 #endif
 
-int32 FSquare2DGridHelper::ForEachIntersectingCells(const FBox& InBox, TFunctionRef<void(const FIntVector&)> InOperation, int32 InStartLevel) const
+int32 FSquare2DGridHelper::ForEachIntersectingCells(const FBox& InBox, TFunctionRef<void(const FGridCellCoord&)> InOperation, int32 InStartLevel) const
 {
 	int32 NumCells = 0;
 
 	for (int32 Level = InStartLevel; Level < Levels.Num(); Level++)
 	{
-		NumCells += Levels[Level].ForEachIntersectingCells(InBox, [InBox, Level, InOperation](const FIntVector2& Coord) { InOperation(FIntVector(Coord.X, Coord.Y, Level)); });
+		NumCells += Levels[Level].ForEachIntersectingCells(InBox, [InBox, Level, InOperation](const FGridCellCoord2& Coord) { InOperation(FGridCellCoord(Coord.X, Coord.Y, Level)); });
 	}
 
 	return NumCells;
 }
 
-int32 FSquare2DGridHelper::ForEachIntersectingCells(const FSphere& InSphere, TFunctionRef<void(const FIntVector&)> InOperation, int32 InStartLevel) const
+int32 FSquare2DGridHelper::ForEachIntersectingCells(const FSphere& InSphere, TFunctionRef<void(const FGridCellCoord&)> InOperation, int32 InStartLevel) const
 {
 	int32 NumCells = 0;
 
 	for (int32 Level = InStartLevel; Level < Levels.Num(); Level++)
 	{
-		NumCells += Levels[Level].ForEachIntersectingCells(InSphere, [InSphere, Level, InOperation](const FIntVector2& Coord) { InOperation(FIntVector(Coord.X, Coord.Y, Level)); });
+		NumCells += Levels[Level].ForEachIntersectingCells(InSphere, [InSphere, Level, InOperation](const FGridCellCoord2& Coord) { InOperation(FGridCellCoord(Coord.X, Coord.Y, Level)); });
 	}
 
 	return NumCells;
 }
 
-int32 FSquare2DGridHelper::ForEachIntersectingCells(const FSphericalSector& InShape, TFunctionRef<void(const FIntVector&)> InOperation, int32 InStartLevel) const
+int32 FSquare2DGridHelper::ForEachIntersectingCells(const FSphericalSector& InShape, TFunctionRef<void(const FGridCellCoord&)> InOperation, int32 InStartLevel) const
 {
 	int32 NumCells = 0;
 
 	for (int32 Level = InStartLevel; Level < Levels.Num(); Level++)
 	{
-		NumCells += Levels[Level].ForEachIntersectingCells(InShape, [Level, InOperation](const FIntVector2& Coord) { InOperation(FIntVector(Coord.X, Coord.Y, Level)); });
+		NumCells += Levels[Level].ForEachIntersectingCells(InShape, [Level, InOperation](const FGridCellCoord2& Coord) { InOperation(FGridCellCoord(Coord.X, Coord.Y, Level)); });
 	}
 
 	return NumCells;
@@ -154,7 +154,7 @@ void FSquare2DGridHelper::ValidateSingleActorReferer()
 
 #if WITH_EDITOR
 
-FSquare2DGridHelper GetGridHelper(const FBox& WorldBounds, int32 GridCellSize)
+FSquare2DGridHelper GetGridHelper(const FBox& WorldBounds, int64 GridCellSize)
 {
 	// Default grid to a minimum of 1 level and 1 cell, for always loaded actors
 	FVector GridOrigin = FVector::ZeroVector;
@@ -174,7 +174,7 @@ FSquare2DGridHelper GetPartitionedActors(const UWorldPartition* WorldPartition,
 	{
 		int32 IntersectingCellCount = 0;
 		FSquare2DGridHelper::FGridLevel& LastGridLevel = PartitionedActors.Levels.Last();
-		LastGridLevel.ForEachIntersectingCells(WorldBounds, [&IntersectingCellCount](const FIntVector2& Coords) { ++IntersectingCellCount; });
+		LastGridLevel.ForEachIntersectingCells(WorldBounds, [&IntersectingCellCount](const FGridCellCoord2& Coords) { ++IntersectingCellCount; });
 		if (!ensure(IntersectingCellCount == 1))
 		{
 			UE_LOG(LogWorldPartition, Warning, TEXT("Can't find grid cell that encompasses world bounds."));
@@ -194,7 +194,7 @@ FSquare2DGridHelper GetPartitionedActors(const UWorldPartition* WorldPartition,
 			if (GRuntimeSpatialHashPlaceSmallActorsUsingLocation && (ClusterBounds.GetArea() <= CellArea))
 			{
 				// Find grid level cell that contains the actor cluster pivot and put actors in it.
-				FIntVector2 CellCoords;
+				FGridCellCoord2 CellCoords;
 				if (PartitionedActors.GetLowestLevel().GetCellCoords(ClusterBounds.GetCenter(), CellCoords))
 				{
 					GridCell = &PartitionedActors.GetLowestLevel().GetCell(CellCoords);
@@ -204,8 +204,8 @@ FSquare2DGridHelper GetPartitionedActors(const UWorldPartition* WorldPartition,
 			{
 				// Find grid level cell that encompasses the actor cluster bounding box and put actors in it.
 				const FVector2D ClusterSize = ClusterBounds.GetSize();
-				const float MinRequiredCellExtent = FMath::Max(ClusterSize.X, ClusterSize.Y);
-				const int32 FirstPotentialGridLevel = FMath::Max(FMath::CeilToFloat(FMath::Log2(MinRequiredCellExtent / (float)PartitionedActors.CellSize)), 0);
+				const double MinRequiredCellExtent = FMath::Max(ClusterSize.X, ClusterSize.Y);
+				const int32 FirstPotentialGridLevel = FMath::Max(FMath::CeilToDouble(FMath::Log2(MinRequiredCellExtent / (double)PartitionedActors.CellSize)), 0);
 
 				for (int32 GridLevelIndex = FirstPotentialGridLevel; GridLevelIndex < PartitionedActors.Levels.Num(); GridLevelIndex++)
 				{
@@ -213,7 +213,7 @@ FSquare2DGridHelper GetPartitionedActors(const UWorldPartition* WorldPartition,
 
 					if (GridLevel.GetNumIntersectingCells(ClusterInstance->Bounds) == 1)
 					{
-						GridLevel.ForEachIntersectingCells(ClusterInstance->Bounds, [&GridLevel, ActorCluster, ClusterInstance, &GridCell](const FIntVector2& Coords)
+						GridLevel.ForEachIntersectingCells(ClusterInstance->Bounds, [&GridLevel, ActorCluster, ClusterInstance, &GridCell](const FGridCellCoord2& Coords)
 						{
 							check(!GridCell);
 							GridCell = &GridLevel.GetCell(Coords);
@@ -231,10 +231,10 @@ FSquare2DGridHelper GetPartitionedActors(const UWorldPartition* WorldPartition,
 
 		if (!GridCell)
 		{
-			UE_LOG(LogWorldPartition, Warning, TEXT("Cluster of %d actors was promoted always loaded, BV of [%d x %d] (meters)"), 
+			UE_LOG(LogWorldPartition, Warning, TEXT("Cluster of %d actors was promoted always loaded, BV of [%lld x %lld] (meters)"), 
 				ActorCluster->Actors.Num(),
-				(int)(0.01f * ClusterInstance->Bounds.GetSize().X),
-				(int)(0.01f * ClusterInstance->Bounds.GetSize().Y));
+				(int64)(0.01f * ClusterInstance->Bounds.GetSize().X),
+				(int64)(0.01f * ClusterInstance->Bounds.GetSize().Y));
 			for (const FGuid& ActorGuid : ActorCluster->Actors)
 			{
 				const FWorldPartitionActorDescView& ActorDescView = ClusterInstance->ContainerInstance->GetActorDescView(ActorGuid);
@@ -253,10 +253,10 @@ FSquare2DGridHelper GetPartitionedActors(const UWorldPartition* WorldPartition,
 		{
 			if (ActorCluster->Actors.Num() > 1)
 			{
-				UE_LOG(LogWorldPartition, Verbose, TEXT("Clustered %d actors, generated shared BV of [%d x %d] (meters)"),
+				UE_LOG(LogWorldPartition, Verbose, TEXT("Clustered %d actors, generated shared BV of [%lld x %lld] (meters)"),
 					ActorCluster->Actors.Num(),
-					(int)(0.01f * ClusterInstance->Bounds.GetSize().X),
-					(int)(0.01f * ClusterInstance->Bounds.GetSize().Y));
+					(int64)(0.01f * ClusterInstance->Bounds.GetSize().X),
+					(int64)(0.01f * ClusterInstance->Bounds.GetSize().Y));
 
 				if (ClusterInstance->DataLayers.Num())
 				{
@@ -312,9 +312,9 @@ FORCEINLINE static bool IsPointInsideSector(const FVector2D& TestPoint, const FV
 	}
 }
 
-bool FSquare2DGridHelper::FGrid2D::DoesCircleSectorIntersectsCell(const FIntVector2& Coords, const FVector2D& SectorCenter, float SectorRadiusSquared, const FVector2D& SectorStartVector, const FVector2D& SectorEndVector, float SectorAngle) const
+bool FSquare2DGridHelper::FGrid2D::DoesCircleSectorIntersectsCell(const FGridCellCoord2& Coords, const FVector2D& SectorCenter, float SectorRadiusSquared, const FVector2D& SectorStartVector, const FVector2D& SectorEndVector, float SectorAngle) const
 {
-	const int64 CellIndex = (int64)Coords.Y * (int64)GridSize + (int64)Coords.X;
+	const int64 CellIndex = Coords.Y * GridSize + Coords.X;
 	FBox2D CellBounds;
 	GetCellBounds(CellIndex, CellBounds);
 
@@ -344,7 +344,7 @@ bool FSquare2DGridHelper::FGrid2D::DoesCircleSectorIntersectsCell(const FIntVect
 	return false;
 }
 
-int32 FSquare2DGridHelper::FGrid2D::ForEachIntersectingCells(const FSphericalSector& InShape, TFunctionRef<void(const FIntVector2&)> InOperation) const
+int32 FSquare2DGridHelper::FGrid2D::ForEachIntersectingCells(const FSphericalSector& InShape, TFunctionRef<void(const FGridCellCoord2&)> InOperation) const
 {
 	check(!InShape.IsNearlyZero());
 	const bool bIsSphere = InShape.IsSphere();
@@ -359,7 +359,7 @@ int32 FSquare2DGridHelper::FGrid2D::ForEachIntersectingCells(const FSphericalSec
 
 	int32 NumCells = 0;
 	const FSphere Sphere(InShape.GetCenter(), InShape.GetRadius());
-	ForEachIntersectingCells(Sphere, [this, bIsSphere, Center2D, SectorRadiusSquared, SectorStart2D, SectorEnd2D, SectorAngle = InShape.GetAngle(), InOperation, &NumCells](const FIntVector2& Coords)
+	ForEachIntersectingCells(Sphere, [this, bIsSphere, Center2D, SectorRadiusSquared, SectorStart2D, SectorEnd2D, SectorAngle = InShape.GetAngle(), InOperation, &NumCells](const FGridCellCoord2& Coords)
 	{
 		// If sector, filter coords that intersect sector
 		if (bIsSphere || DoesCircleSectorIntersectsCell(Coords, Center2D, SectorRadiusSquared, SectorStart2D, SectorEnd2D, SectorAngle))
diff --git a/Engine/Source/Runtime/Engine/Private/WorldPartition/RuntimeSpatialHash/RuntimeSpatialHashGridHelper.h b/Engine/Source/Runtime/Engine/Private/WorldPartition/RuntimeSpatialHash/RuntimeSpatialHashGridHelper.h
--- a/Engine/Source/Runtime/Engine/Private/WorldPartition/RuntimeSpatialHash/RuntimeSpatialHashGridHelper.h
+++ b/Engine/Source/Runtime/Engine/Private/WorldPartition/RuntimeSpatialHash/RuntimeSpatialHashGridHelper.h
@@ -22,10 +22,10 @@ struct FSquare2DGridHelper
 	struct FGrid2D
 	{
 		FVector2D Origin;
-		int32 CellSize;
-		int32 GridSize;
+		int64 CellSize;
+		int64 GridSize;
 
-		inline FGrid2D(const FVector2D& InOrigin, int32 InCellSize, int32 InGridSize)
+		inline FGrid2D(const FVector2D& InOrigin, int64 InCellSize, int64 InGridSize)
 			: Origin(InOrigin)
 			, CellSize(InCellSize)
 			, GridSize(InGridSize)
@@ -36,9 +36,17 @@ struct FSquare2DGridHelper
 		 *
 		 * @return true if the specified coordinates are valid
 		 */
-		inline bool IsValidCoords(const FIntVector2& InCoords) const
+		inline bool IsValidCoords(const FGridCellCoord2& InCoords) const
 		{
-			return (InCoords.X >= 0) && (InCoords.X < GridSize) && (InCoords.Y >= 0) && (InCoords.Y < GridSize);
+			if ((InCoords.X >= 0) && (InCoords.X < GridSize) && (InCoords.Y >= 0) && (InCoords.Y < GridSize))
+			{
+				int64 Index = (InCoords.Y * GridSize) + InCoords.X;
+				if (ensureMsgf(Index >= 0, TEXT("World Partition reached the current limit of large world coordinates.")))
+				{
+					return true;
+				}
+			}
+			return false;
 		}
 
 		/**
@@ -48,9 +56,9 @@ struct FSquare2DGridHelper
 		 */
 		inline bool GetCellBounds(int64 InIndex, FBox2D& OutBounds) const
 		{
-			if (InIndex >= 0 && InIndex <= ((int64)GridSize * (int64)GridSize))
+			if (InIndex >= 0 && InIndex <= (GridSize * GridSize))
 			{
-				const FIntVector2 Coords(InIndex % (int64)GridSize, InIndex / (int64)GridSize);
+				const FGridCellCoord2 Coords(InIndex % GridSize, InIndex / GridSize);
 				return GetCellBounds(Coords, OutBounds);
 			}
 
@@ -62,11 +70,11 @@ struct FSquare2DGridHelper
 		 *
 		 * @return true if the specified coord was valid
 		 */
-		inline bool GetCellBounds(const FIntVector2& InCoords, FBox2D& OutBounds) const
+		inline bool GetCellBounds(const FGridCellCoord2& InCoords, FBox2D& OutBounds) const
 		{
 			if (IsValidCoords(InCoords))
 			{
-				const FVector2D Min = (FVector2D(Origin) - FVector2D(GridSize * CellSize * 0.5f, GridSize * CellSize * 0.5f)) + FVector2D(InCoords.X * CellSize, InCoords.Y * CellSize);
+				const FVector2D Min = (FVector2D(Origin) - FVector2D(GridSize * CellSize * 0.5, GridSize * CellSize * 0.5)) + FVector2D(InCoords.X * CellSize, InCoords.Y * CellSize);
 				const FVector2D Max = Min + FVector2D(CellSize, CellSize);
 				OutBounds = FBox2D(Min, Max);
 				return true;
@@ -80,11 +88,11 @@ struct FSquare2DGridHelper
 		 *
 		 * @return true if the position was inside the grid
 		 */
-		inline bool GetCellCoords(const FVector2D& InPos, FIntVector2& OutCoords) const
+		inline bool GetCellCoords(const FVector2D& InPos, FGridCellCoord2& OutCoords) const
 		{
-			OutCoords = FIntVector2(
-				FMath::FloorToInt(((InPos.X - Origin.X) / CellSize) + GridSize * 0.5f),
-				FMath::FloorToInt(((InPos.Y - Origin.Y) / CellSize) + GridSize * 0.5f)
+			OutCoords = FGridCellCoord2(
+				FMath::FloorToInt(((InPos.X - Origin.X) / CellSize) + GridSize * 0.5),
+				FMath::FloorToInt(((InPos.Y - Origin.Y) / CellSize) + GridSize * 0.5)
 			);
 
 			return IsValidCoords(OutCoords);
@@ -95,7 +103,7 @@ struct FSquare2DGridHelper
 		 *
 		 * @return true if the bounds was intersecting with the grid
 		 */
-		inline bool GetCellCoords(const FBox2D& InBounds2D, FIntVector2& OutMinCellCoords, FIntVector2& OutMaxCellCoords) const
+		inline bool GetCellCoords(const FBox2D& InBounds2D, FGridCellCoord2& OutMinCellCoords, FGridCellCoord2& OutMaxCellCoords) const
 		{
 			GetCellCoords(InBounds2D.Min, OutMinCellCoords);
 			if ((OutMinCellCoords.X >= GridSize) || (OutMinCellCoords.Y >= GridSize))
@@ -109,11 +117,10 @@ struct FSquare2DGridHelper
 				return false;
 			}
 
-			OutMinCellCoords.X = FMath::Clamp(OutMinCellCoords.X, 0, GridSize - 1);
-			OutMinCellCoords.Y = FMath::Clamp(OutMinCellCoords.Y, 0, GridSize - 1);
-			OutMaxCellCoords.X = FMath::Clamp(OutMaxCellCoords.X, 0, GridSize - 1);
-			OutMaxCellCoords.Y = FMath::Clamp(OutMaxCellCoords.Y, 0, GridSize - 1);
-
+			OutMinCellCoords.X = FMath::Clamp(OutMinCellCoords.X, 0LL, GridSize - 1);
+			OutMinCellCoords.Y = FMath::Clamp(OutMinCellCoords.Y, 0LL, GridSize - 1);
+			OutMaxCellCoords.X = FMath::Clamp(OutMaxCellCoords.X, 0LL, GridSize - 1);
+			OutMaxCellCoords.Y = FMath::Clamp(OutMaxCellCoords.Y, 0LL, GridSize - 1);
 			return true;
 		}
 
@@ -122,14 +129,13 @@ struct FSquare2DGridHelper
 		 *
 		 * @return true if the coords was inside the grid
 		 */
-		inline bool GetCellIndex(const FIntVector2& InCoords, uint64& OutIndex) const
+		inline bool GetCellIndex(const FGridCellCoord2& InCoords, uint64& OutIndex) const
 		{
 			if (IsValidCoords(InCoords))
 			{
-				OutIndex = ((int64)InCoords.Y * (int64)GridSize) + (int64)InCoords.X;
+				OutIndex = (InCoords.Y * GridSize) + InCoords.X;
 				return true;
 			}
-
 			return false;
 		}
 
@@ -140,9 +146,9 @@ struct FSquare2DGridHelper
 		 */
 		inline bool GetCellIndex(const FVector& InPos, uint64& OutIndex) const
 		{
-			FIntVector2 Coords = FIntVector2(
-				FMath::FloorToInt(((InPos.X - Origin.X) / CellSize) + GridSize * 0.5f),
-				FMath::FloorToInt(((InPos.Y - Origin.Y) / CellSize) + GridSize * 0.5f)
+			FGridCellCoord2 Coords = FGridCellCoord2(
+				FMath::FloorToInt(((InPos.X - Origin.X) / CellSize) + GridSize * 0.5),
+				FMath::FloorToInt(((InPos.Y - Origin.Y) / CellSize) + GridSize * 0.5)
 			);
 
 			return GetCellIndex(Coords, OutIndex);
@@ -155,8 +161,8 @@ struct FSquare2DGridHelper
 		 */
 		int32 GetNumIntersectingCells(const FBox& InBox) const
 		{
-			FIntVector2 MinCellCoords;
-			FIntVector2 MaxCellCoords;
+			FGridCellCoord2 MinCellCoords;
+			FGridCellCoord2 MaxCellCoords;
 			const FBox2D Bounds2D(FVector2D(InBox.Min), FVector2D(InBox.Max));
 
 			if (GetCellCoords(Bounds2D, MinCellCoords, MaxCellCoords))
@@ -167,42 +173,35 @@ struct FSquare2DGridHelper
 			return 0;
 		}
 
-		// Runs a function on all cells
-		void ForEachCells(TFunctionRef<void(const FIntVector2&)> InOperation) const
-		{
-			for (int32 y = 0; y < GridSize; y++)
-			{
-				for (int32 x = 0; x < GridSize; x++)
-				{
-					InOperation(FIntVector2(x, y));
-				}
-			}
-		}
-
 		/**
 		 * Runs a function on all intersecting cells for the provided box
 		 *
 		 * @return the number of intersecting cells
 		 */
-		int32 ForEachIntersectingCellsBreakable(const FBox& InBox, TFunctionRef<bool(const FIntVector2&)> InOperation) const
+		int32 ForEachIntersectingCellsBreakable(const FBox& InBox, TFunctionRef<bool(const FGridCellCoord2&)> InOperation) const
 		{
 			int32 NumCells = 0;
 
-			FIntVector2 MinCellCoords;
-			FIntVector2 MaxCellCoords;
+			FGridCellCoord2 MinCellCoords;
+			FGridCellCoord2 MaxCellCoords;
 			const FBox2D Bounds2D(FVector2D(InBox.Min), FVector2D(InBox.Max));
 
 			if (GetCellCoords(Bounds2D, MinCellCoords, MaxCellCoords))
 			{
-				for (int32 y = MinCellCoords.Y; y <= MaxCellCoords.Y; y++)
+				for (int64 y = MinCellCoords.Y; y <= MaxCellCoords.Y; y++)
 				{
-					for (int32 x = MinCellCoords.X; x <= MaxCellCoords.X; x++)
+					for (int64 x = MinCellCoords.X; x <= MaxCellCoords.X; x++)
 					{
-						if (!InOperation(FIntVector2(x, y)))
+						const FGridCellCoord2 Coord(x, y);
+						// Validate that generated coordinate is valid (in case we reached the 64-bit limit of cell index)
+						if (IsValidCoords(Coord))
 						{
-							return NumCells;
+							if (!InOperation(FGridCellCoord2(x, y)))
+							{
+								return NumCells;
+							}
+							++NumCells;
 						}
-						++NumCells;
 					}
 				}
 			}
@@ -210,9 +209,9 @@ struct FSquare2DGridHelper
 			return NumCells;
 		}
 
-		int32 ForEachIntersectingCells(const FBox& InBox, TFunctionRef<void(const FIntVector2&)> InOperation) const
+		int32 ForEachIntersectingCells(const FBox& InBox, TFunctionRef<void(const FGridCellCoord2&)> InOperation) const
 		{
-			return ForEachIntersectingCellsBreakable(InBox, [InOperation](const FIntVector2& Vector) { InOperation(Vector); return true; });
+			return ForEachIntersectingCellsBreakable(InBox, [InOperation](const FGridCellCoord2& Vector) { InOperation(Vector); return true; });
 		}
 
 		/**
@@ -220,16 +219,16 @@ struct FSquare2DGridHelper
 		 *
 		 * @return the number of intersecting cells
 		 */
-		int32 ForEachIntersectingCells(const FSphere& InSphere, TFunctionRef<void(const FIntVector2&)> InOperation) const
+		int32 ForEachIntersectingCells(const FSphere& InSphere, TFunctionRef<void(const FGridCellCoord2&)> InOperation) const
 		{
 			int32 NumCells = 0;
 
 			// @todo_ow: rasterize circle instead?
 			const FBox Box(InSphere.Center - FVector(InSphere.W), InSphere.Center + FVector(InSphere.W));
 
-			ForEachIntersectingCells(Box, [this, &InSphere, &InOperation, &NumCells](const FIntVector2& Coords)
+			ForEachIntersectingCells(Box, [this, &InSphere, &InOperation, &NumCells](const FGridCellCoord2& Coords)
 			{
-				const int64 CellIndex = (int64)Coords.Y * (int64)GridSize + (int64)Coords.X;
+				const int64 CellIndex = Coords.Y * GridSize + Coords.X;
 
 				FBox2D CellBounds;
 				GetCellBounds(CellIndex, CellBounds);
@@ -250,11 +249,11 @@ struct FSquare2DGridHelper
 		 *
 		 * @return the number of intersecting cells
 		 */
-		int32 ForEachIntersectingCells(const FSphericalSector& InShape, TFunctionRef<void(const FIntVector2&)> InOperation) const;
+		int32 ForEachIntersectingCells(const FSphericalSector& InShape, TFunctionRef<void(const FGridCellCoord2&)> InOperation) const;
 
 	private:
 
-		bool DoesCircleSectorIntersectsCell(const FIntVector2& Coords, const FVector2D& SectorCenter, float SectorRadiusSquared, const FVector2D& SectorStartVector, const FVector2D& SectorEndVector, float SectorAngle) const;
+		bool DoesCircleSectorIntersectsCell(const FGridCellCoord2& Coords, const FVector2D& SectorCenter, float SectorRadiusSquared, const FVector2D& SectorStartVector, const FVector2D& SectorEndVector, float SectorAngle) const;
 	};
 
 	struct FGridLevel : public FGrid2D
@@ -282,7 +281,7 @@ struct FSquare2DGridHelper
 
 		struct FGridCell
 		{
-			FGridCell(const FIntVector& InCoords)
+			FGridCell(const FGridCellCoord& InCoords)
 				: Coords(InCoords)
 			{}
 
@@ -316,18 +315,18 @@ struct FSquare2DGridHelper
 				return nullptr;
 			}
 
-			FIntVector GetCoords() const
+			FGridCellCoord GetCoords() const
 			{
 				return Coords;
 			}
 
 		private:
-			FIntVector Coords;
+			FGridCellCoord Coords;
 			TArray<FGridCellDataChunk> DataChunks;
 		};
 #endif
 
-		inline FGridLevel(const FVector2D& InOrigin, int32 InCellSize, int32 InGridSize, int32 InLevel)
+		inline FGridLevel(const FVector2D& InOrigin, int64 InCellSize, int64 InGridSize, int32 InLevel)
 			: FGrid2D(InOrigin, InCellSize, InGridSize)
 #if WITH_EDITOR
 			, Level(InLevel)
@@ -340,12 +339,12 @@ struct FSquare2DGridHelper
 		 *
 		 * @return the cell at the specified grid coordinate
 		 */
-		inline FGridCell& GetCell(const FIntVector2& InCoords)
+		inline FGridCell& GetCell(const FGridCellCoord2& InCoords)
 		{
 			check(IsValidCoords(InCoords));
 
 			uint64 CellIndex;
-			GetCellIndex(InCoords, CellIndex);
+			verify(GetCellIndex(InCoords, CellIndex));
 
 			int64 CellIndexMapping;
 			int64* CellIndexMappingPtr = CellsMapping.Find(CellIndex);
@@ -355,7 +354,7 @@ struct FSquare2DGridHelper
 			}
 			else
 			{
-				CellIndexMapping = Cells.Emplace(FIntVector(InCoords.X, InCoords.Y, Level));
+				CellIndexMapping = Cells.Emplace(FGridCellCoord(InCoords.X, InCoords.Y, Level));
 				CellsMapping.Add(CellIndex, CellIndexMapping);
 			}
 
@@ -367,17 +366,17 @@ struct FSquare2DGridHelper
 		 *
 		 * @return the cell at the specified grid coordinate
 		 */
-		inline const FGridCell& GetCell(const FIntVector2& InCoords) const
+		inline const FGridCell& GetCell(const FGridCellCoord2& InCoords) const
 		{
 			check(IsValidCoords(InCoords));
 
 			uint64 CellIndex;
-			GetCellIndex(InCoords, CellIndex);
+			verify(GetCellIndex(InCoords, CellIndex));
 
 			int64 CellIndexMapping = CellsMapping.FindChecked(CellIndex);
 
 			const FGridCell& Cell = Cells[CellIndexMapping];
-			check(Cell.GetCoords() == FIntVector(InCoords.X, InCoords.Y, Level));
+			check(Cell.GetCoords() == FGridCellCoord(InCoords.X, InCoords.Y, Level));
 			return Cell;
 		}
 
@@ -387,20 +386,20 @@ struct FSquare2DGridHelper
 #endif
 	};
 
-	FSquare2DGridHelper(const FBox& InWorldBounds, const FVector& InOrigin, int32 InCellSize);
+	FSquare2DGridHelper(const FBox& InWorldBounds, const FVector& InOrigin, int64 InCellSize);
 
 #if WITH_EDITOR
 	// Returns the lowest grid level
 	inline FGridLevel& GetLowestLevel() { return Levels[0]; }
 
 	// Returns the always loaded (top level) cell
-	inline FGridLevel::FGridCell& GetAlwaysLoadedCell() { return Levels.Last().GetCell(FIntVector2(0,0)); }
+	inline FGridLevel::FGridCell& GetAlwaysLoadedCell() { return Levels.Last().GetCell(FGridCellCoord2(0,0)); }
 
 	// Returns the always loaded (top level) cell
-	inline const FGridLevel::FGridCell& GetAlwaysLoadedCell() const { return Levels.Last().GetCell(FIntVector2(0,0)); }
+	inline const FGridLevel::FGridCell& GetAlwaysLoadedCell() const { return Levels.Last().GetCell(FGridCellCoord2(0,0)); }
 
 	// Returns the cell at the given coord
-	inline const FGridLevel::FGridCell& GetCell(const FIntVector& InCoords) const { return Levels[InCoords.Z].GetCell(FIntVector2(InCoords.X, InCoords.Y)); }
+	inline const FGridLevel::FGridCell& GetCell(const FGridCellCoord& InCoords) const { return Levels[InCoords.Z].GetCell(FGridCellCoord2(InCoords.X, InCoords.Y)); }
 #endif
 
 	/**
@@ -408,11 +407,11 @@ struct FSquare2DGridHelper
 	 *
 	 * @return true if the specified coord was valid
 	 */
-	inline bool GetCellBounds(const FIntVector& InCoords, FBox2D& OutBounds) const
+	inline bool GetCellBounds(const FGridCellCoord& InCoords, FBox2D& OutBounds) const
 	{
 		if (Levels.IsValidIndex(InCoords.Z))
 		{
-			return Levels[InCoords.Z].GetCellBounds(FIntVector2(InCoords.X, InCoords.Y), OutBounds);
+			return Levels[InCoords.Z].GetCellBounds(FGridCellCoord2(InCoords.X, InCoords.Y), OutBounds);
 		}
 		return false;
 	}
@@ -422,14 +421,14 @@ struct FSquare2DGridHelper
 	 *
 	 * @return true if the specified coord was valid
 	 */
-	inline bool GetCellGlobalCoords(const FIntVector& InCoords, FIntVector& OutGlobalCoords) const
+	inline bool GetCellGlobalCoords(const FGridCellCoord& InCoords, FGridCellCoord& OutGlobalCoords) const
 	{
 		if (Levels.IsValidIndex(InCoords.Z))
 		{
 			const FGridLevel& GridLevel = Levels[InCoords.Z];
-			if (GridLevel.IsValidCoords(FIntVector2(InCoords.X, InCoords.Y)))
+			if (GridLevel.IsValidCoords(FGridCellCoord2(InCoords.X, InCoords.Y)))
 			{
-				int32 CoordOffset = Levels[InCoords.Z].GridSize >> 1;
+				int64 CoordOffset = Levels[InCoords.Z].GridSize >> 1;
 				OutGlobalCoords = InCoords;
 				OutGlobalCoords.X -= CoordOffset;
 				OutGlobalCoords.Y -= CoordOffset;
@@ -449,21 +448,21 @@ struct FSquare2DGridHelper
 	 *
 	 * @return the number of intersecting cells
 	 */
-	int32 ForEachIntersectingCells(const FBox& InBox, TFunctionRef<void(const FIntVector&)> InOperation, int32 InStartLevel = 0) const;
+	int32 ForEachIntersectingCells(const FBox& InBox, TFunctionRef<void(const FGridCellCoord&)> InOperation, int32 InStartLevel = 0) const;
 
 	/**
 	 * Runs a function on all intersecting cells for the provided sphere
 	 *
 	 * @return the number of intersecting cells
 	 */
-	int32 ForEachIntersectingCells(const FSphere& InSphere, TFunctionRef<void(const FIntVector&)> InOperation, int32 InStartLevel = 0) const;
+	int32 ForEachIntersectingCells(const FSphere& InSphere, TFunctionRef<void(const FGridCellCoord&)> InOperation, int32 InStartLevel = 0) const;
 
 	/**
 	 * Runs a function on all intersecting cells for the provided spherical sector
 	 *
 	 * @return the number of intersecting cells
 	 */
-	int32 ForEachIntersectingCells(const FSphericalSector& InShape, TFunctionRef<void(const FIntVector&)> InOperation, int32 InStartLevel = 0) const;
+	int32 ForEachIntersectingCells(const FSphericalSector& InShape, TFunctionRef<void(const FGridCellCoord&)> InOperation, int32 InStartLevel = 0) const;
 
 #if WITH_EDITOR
 	// Validates that actor is not referenced by multiple cells
@@ -473,11 +472,11 @@ struct FSquare2DGridHelper
 public:
 	FBox WorldBounds;
 	FVector Origin;
-	int32 CellSize;
+	int64 CellSize;
 	TArray<FGridLevel> Levels;
 };
 
 #if WITH_EDITOR
-FSquare2DGridHelper GetGridHelper(const FBox& WorldBounds, int32 GridCellSize);
+FSquare2DGridHelper GetGridHelper(const FBox& WorldBounds, int64 GridCellSize);
 FSquare2DGridHelper GetPartitionedActors(const UWorldPartition* WorldPartition, const FBox& WorldBounds, const FSpatialHashRuntimeGrid& Grid, const TArray<const FActorClusterInstance*>& GridActors);
-#endif // #if WITH_EDITOR
+#endif // #if WITH_EDITOR
\ No newline at end of file
diff --git a/Engine/Source/Runtime/Engine/Private/WorldPartition/RuntimeSpatialHash/RuntimeSpatialHashHLOD.cpp b/Engine/Source/Runtime/Engine/Private/WorldPartition/RuntimeSpatialHash/RuntimeSpatialHashHLOD.cpp
--- a/Engine/Source/Runtime/Engine/Private/WorldPartition/RuntimeSpatialHash/RuntimeSpatialHashHLOD.cpp
+++ b/Engine/Source/Runtime/Engine/Private/WorldPartition/RuntimeSpatialHash/RuntimeSpatialHashHLOD.cpp
@@ -152,13 +152,13 @@ static TArray<FGuid> GenerateHLODsForGrid(UWorldPartition* WorldPartition, const
 	TArray<FGuid> GridHLODActors;
 	PartitionedActors.ForEachCells([&](const FSquare2DGridHelper::FGridLevel::FGridCell& GridCell)
 	{
-		const FIntVector CellCoord = GridCell.GetCoords();
+		const FGridCellCoord CellCoord = GridCell.GetCoords();
 
 		FBox2D CellBounds2D;
 		PartitionedActors.GetCellBounds(CellCoord, CellBounds2D);
 		FBox CellBounds = FBox(FVector(CellBounds2D.Min, WorldBounds.Min.Z), FVector(CellBounds2D.Max, WorldBounds.Max.Z));
 
-		FIntVector CellGlobalCoord;
+		FGridCellCoord CellGlobalCoord;
 		verify(PartitionedActors.GetCellGlobalCoords(CellCoord, CellGlobalCoord));
 
 		for (const FSquare2DGridHelper::FGridLevel::FGridCellDataChunk& GridCellDataChunk : GridCell.GetDataChunks())
diff --git a/Engine/Source/Runtime/Engine/Private/WorldPartition/WorldPartitionRuntimeCell.cpp b/Engine/Source/Runtime/Engine/Private/WorldPartition/WorldPartitionRuntimeCell.cpp
--- a/Engine/Source/Runtime/Engine/Private/WorldPartition/WorldPartitionRuntimeCell.cpp
+++ b/Engine/Source/Runtime/Engine/Private/WorldPartition/WorldPartitionRuntimeCell.cpp
@@ -48,19 +48,21 @@ void UWorldPartitionRuntimeCell::SetMinMaxZ(const FVector2D& InMinMaxZ)
 	MinMaxZ = InMinMaxZ;
 }
 
-void UWorldPartitionRuntimeCell::SetDebugInfo(FIntVector InCoords, FName InGridName)
+void UWorldPartitionRuntimeCell::SetDebugInfo(int64 InCoordX, int64 InCoordY, int64 InCoordZ, FName InGridName)
 {
-	Coords = InCoords;
-	GridName = InGridName;
+	DebugInfo.CoordX = InCoordX;
+	DebugInfo.CoordY = InCoordY;
+	DebugInfo.CoordZ = InCoordZ;
+	DebugInfo.GridName = InGridName;
 	UpdateDebugName();
 }
 
 void UWorldPartitionRuntimeCell::UpdateDebugName()
 {
 	TStringBuilder<512> Builder;
-	Builder += GridName.ToString();
+	Builder += DebugInfo.GridName.ToString();
 	Builder += TEXT("_");
-	Builder += FString::Printf(TEXT("L%d_X%d_Y%d"), Coords.Z, Coords.X, Coords.Y);
+	Builder += FString::Printf(TEXT("L%d_X%d_Y%d"), DebugInfo.CoordZ, DebugInfo.CoordX, DebugInfo.CoordY);
 	int32 DataLayerCount = DataLayers.Num();
 
 	const UDataLayerSubsystem* DataLayerSubsystem = UWorld::GetSubsystem<UDataLayerSubsystem>(GetOuterUWorldPartition()->GetWorld());
@@ -77,7 +79,7 @@ void UWorldPartitionRuntimeCell::UpdateDebugName()
 		}
 		Builder += FString::Printf(TEXT("ID:%X]"), FDataLayersID(DataLayerObjects).GetHash());
 	}
-	DebugName = Builder.ToString();
+	DebugInfo.Name = Builder.ToString();
 }
 
 void UWorldPartitionRuntimeCell::DumpStateLog(FHierarchicalLogArchive& Ar)
@@ -116,10 +118,10 @@ int32 UWorldPartitionRuntimeCell::SortCompare(const UWorldPartitionRuntimeCell*
 
 bool UWorldPartitionRuntimeCell::IsDebugShown() const
 {
-	return FWorldPartitionDebugHelper::IsDebugRuntimeHashGridShown(GridName) &&
+	return FWorldPartitionDebugHelper::IsDebugRuntimeHashGridShown(DebugInfo.GridName) &&
 		   FWorldPartitionDebugHelper::IsDebugStreamingStatusShown(GetStreamingStatus()) &&
 		   FWorldPartitionDebugHelper::AreDebugDataLayersShown(DataLayers) &&
-		   FWorldPartitionDebugHelper::IsDebugCellNameShown(DebugName);
+		   FWorldPartitionDebugHelper::IsDebugCellNameShown(DebugInfo.Name);
 }
 
 FLinearColor UWorldPartitionRuntimeCell::GetDebugStreamingPriorityColor() const
diff --git a/Engine/Source/Runtime/Engine/Private/WorldPartition/WorldPartitionRuntimeSpatialHash.cpp b/Engine/Source/Runtime/Engine/Private/WorldPartition/WorldPartitionRuntimeSpatialHash.cpp
--- a/Engine/Source/Runtime/Engine/Private/WorldPartition/WorldPartitionRuntimeSpatialHash.cpp
+++ b/Engine/Source/Runtime/Engine/Private/WorldPartition/WorldPartitionRuntimeSpatialHash.cpp
@@ -143,7 +143,7 @@ const FSquare2DGridHelper& FSpatialHashStreamingGrid::GetGridHelper() const
 	return *GridHelper;
 }
 
-int32 FSpatialHashStreamingGrid::GetCellSize(int32 Level) const
+int64 FSpatialHashStreamingGrid::GetCellSize(int32 Level) const
 {
 	return GetGridHelper().Levels[Level].CellSize;
 }
@@ -176,7 +176,7 @@ void FSpatialHashStreamingGrid::GetCells(const FWorldPartitionStreamingQuerySour
 	{
 		QuerySource.ForEachShape(GetLoadingRange(), GridName, HLODLayer, /*bProjectIn2D*/ true, [&](const FSphericalSector& Shape)
 		{
-			Helper.ForEachIntersectingCells(Shape, [&](const FIntVector& Coords)
+			Helper.ForEachIntersectingCells(Shape, [&](const FGridCellCoord& Coords)
 			{
 				if (const FSpatialHashStreamingGridLayerCell* LayerCell = GetLayerCell(Coords))
 				{
@@ -213,7 +213,7 @@ void FSpatialHashStreamingGrid::GetCells(const TArray<FWorldPartitionStreamingSo
 {
 	TRACE_CPUPROFILER_EVENT_SCOPE(FSpatialHashStreamingGrid::GetCells);
 	
-	typedef TMap<FIntVector, TArray<UWorldPartitionRuntimeCell::FStreamingSourceInfo>> FIntersectingCells;
+	typedef TMap<FGridCellCoord, TArray<UWorldPartitionRuntimeCell::FStreamingSourceInfo>> FIntersectingCells;
 	FIntersectingCells AllActivatedCells;
 
 	const float GridLoadingRange = GetLoadingRange();
@@ -224,7 +224,7 @@ void FSpatialHashStreamingGrid::GetCells(const TArray<FWorldPartitionStreamingSo
 		{
 			UWorldPartitionRuntimeCell::FStreamingSourceInfo Info(Source, Shape);
 
-			Helper.ForEachIntersectingCells(Shape, [&](const FIntVector& Coords)
+			Helper.ForEachIntersectingCells(Shape, [&](const FGridCellCoord& Coords)
 			{
 				bool bAddedActivatedCell = false;
 #if !UE_BUILD_SHIPPING
@@ -276,17 +276,17 @@ void FSpatialHashStreamingGrid::GetCells(const TArray<FWorldPartitionStreamingSo
 			const int32 AlwaysLoadedLevel = Helper.Levels.Num() - 1;
 			for (const auto& InTestCell : InTestCells)
 			{
-				const FIntVector& TestCell = InTestCell.Key;
+				const FGridCellCoord& TestCell = InTestCell.Key;
 				int32 CurrentLevelIndex = TestCell.Z;
 				int32 ParentLevelIndex = CurrentLevelIndex + 1;
 				// Only test with Parent Level if it's below the AlwaysLoaded Level
 				if (ParentLevelIndex < AlwaysLoadedLevel)
 				{
 					FBox2D CurrentLevelCellBounds;
-					Helper.Levels[CurrentLevelIndex].GetCellBounds(FIntVector2(TestCell.X, TestCell.Y), CurrentLevelCellBounds);
+					Helper.Levels[CurrentLevelIndex].GetCellBounds(FGridCellCoord2(TestCell.X, TestCell.Y), CurrentLevelCellBounds);
 					FBox Box(FVector(CurrentLevelCellBounds.Min, 0), FVector(CurrentLevelCellBounds.Max, 0));
 
-					Helper.ForEachIntersectingCells(Box, [&](const FIntVector& IntersectingCoords)
+					Helper.ForEachIntersectingCells(Box, [&](const FGridCellCoord& IntersectingCoords)
 					{
 						check(IntersectingCoords.Z >= ParentLevelIndex);
 						if (!InAllCells.Contains(IntersectingCoords))
@@ -338,7 +338,7 @@ void FSpatialHashStreamingGrid::GetCells(const TArray<FWorldPartitionStreamingSo
 	}
 }
 
-const FSpatialHashStreamingGridLayerCell* FSpatialHashStreamingGrid::GetLayerCell(const FIntVector& Coords) const
+const FSpatialHashStreamingGridLayerCell* FSpatialHashStreamingGrid::GetLayerCell(const FGridCellCoord& Coords) const
 {
 	check(GridLevels.IsValidIndex(Coords.Z));
 	if (const int32* LayerCellIndexPtr = GridLevels[Coords.Z].LayerCellsMapping.Find(Coords.Y * GetGridHelper().Levels[Coords.Z].GridSize + Coords.X))
@@ -428,7 +428,7 @@ void FSpatialHashStreamingGrid::Draw3D(UWorld* World, const TArray<FWorldPartiti
 	{
 		FVector StartTrace = Source.Location + FVector(0.f, 0.f, 100.f);
 		FVector EndTrace = StartTrace - FVector(0.f, 0.f, 1000000.f);
-		float Z = Source.Location.Z;
+		double Z = Source.Location.Z;
 		FHitResult Hit;
 		if (World->LineTraceSingleByObjectType(Hit, StartTrace, EndTrace, FCollisionObjectQueryParams(ECC_WorldStatic), FCollisionQueryParams(SCENE_QUERY_STAT(DebugWorldPartitionTrace), true)))
 		{
@@ -440,7 +440,7 @@ void FSpatialHashStreamingGrid::Draw3D(UWorld* World, const TArray<FWorldPartiti
 
 		for (int32 GridLevel = MinGridLevel; GridLevel <= MaxGridLevel; ++GridLevel)
 		{
-			Helper.Levels[GridLevel].ForEachIntersectingCells(Region, [&](const FIntVector2& Coords)
+			Helper.Levels[GridLevel].ForEachIntersectingCells(Region, [&](const FGridCellCoord2& Coords)
 			{
 				const int32* LayerCellIndexPtr = GridLevels[GridLevel].LayerCellsMapping.Find(Coords.Y * Helper.Levels[GridLevel].GridSize + Coords.X);
 				const FSpatialHashStreamingGridLayerCell* LayerCell = LayerCellIndexPtr ? &GridLevels[GridLevel].LayerCells[*LayerCellIndexPtr] : nullptr;
@@ -448,8 +448,8 @@ void FSpatialHashStreamingGrid::Draw3D(UWorld* World, const TArray<FWorldPartiti
 				check(FilteredCells.Num());
 
 				FBox2D CellWorldBounds;
-				Helper.Levels[GridLevel].GetCellBounds(FIntVector2(Coords.X, Coords.Y), CellWorldBounds);
-				float CellSizeY = CellWorldBounds.GetSize().Y / FilteredCells.Num();
+				Helper.Levels[GridLevel].GetCellBounds(FGridCellCoord2(Coords.X, Coords.Y), CellWorldBounds);
+				double CellSizeY = CellWorldBounds.GetSize().Y / FilteredCells.Num();
 				CellWorldBounds.Max.Y = CellWorldBounds.Min.Y + CellSizeY;
 				FVector BoundsExtent(CellWorldBounds.GetExtent(), 100.f);
 				FVector BoundsOrigin(CellWorldBounds.GetCenter(), Z);
@@ -475,7 +475,7 @@ void FSpatialHashStreamingGrid::Draw3D(UWorld* World, const TArray<FWorldPartiti
 					if (DataLayerDebugColors.Num() && Cell->GetDataLayers().Num() > 0)
 					{
 						FBox DataLayerColoredBox = CellBox;
-						float DataLayerSizeX = DataLayerColoredBox.GetSize().X / (5 * Cell->GetDataLayers().Num()); // Use 20% of cell's width
+						double DataLayerSizeX = DataLayerColoredBox.GetSize().X / (5 * Cell->GetDataLayers().Num()); // Use 20% of cell's width
 						DataLayerColoredBox.Max.X = DataLayerColoredBox.Min.X + DataLayerSizeX; 
 						FTranslationMatrix DataLayerOffsetMatrix(FVector(DataLayerSizeX, 0, 0));
 						for (const FName& DataLayer : Cell->GetDataLayers())
@@ -547,7 +547,7 @@ void FSpatialHashStreamingGrid::Draw2D(UCanvas* Canvas, UWorld* World, const TAr
 
 	// Precompute a cell coordinate text with/height using a generic coordinate
 	// This will be used later to filter out drawing of cell coordinates (avoid doing expesive calls to Canvas->StrLen)
-	const FString CellCoordString = UWorldPartitionRuntimeSpatialHash::GetCellCoordString(FIntVector(88,88,88));
+	const FString CellCoordString = UWorldPartitionRuntimeSpatialHash::GetCellCoordString(FGridCellCoord(88,88,88));
 	float MaxCellCoordTextWidth, MaxCellCoordTextHeight;
 	Canvas->StrLen(GEngine->GetTinyFont(), CellCoordString, MaxCellCoordTextWidth, MaxCellCoordTextHeight);
 
@@ -556,10 +556,10 @@ void FSpatialHashStreamingGrid::Draw2D(UCanvas* Canvas, UWorld* World, const TAr
 	{
 		// Draw Grid cells at desired grid level
 		const FSquare2DGridHelper& Helper = GetGridHelper();
-		Helper.Levels[GridLevel].ForEachIntersectingCells(Region, [&](const FIntVector2& Coords)
+		Helper.Levels[GridLevel].ForEachIntersectingCells(Region, [&](const FGridCellCoord2& Coords)
 		{
 			FBox2D CellWorldBounds;
-			Helper.Levels[GridLevel].GetCellBounds(FIntVector2(Coords.X, Coords.Y), CellWorldBounds);
+			Helper.Levels[GridLevel].GetCellBounds(FGridCellCoord2(Coords.X, Coords.Y), CellWorldBounds);
 			FBox2D CellScreenBounds = FBox2D(WorldToScreen(CellWorldBounds.Min), WorldToScreen(CellWorldBounds.Max));
 			// Clamp inside grid bounds
 			if (!GridScreenBounds.IsInside(CellScreenBounds))
@@ -571,8 +571,8 @@ void FSpatialHashStreamingGrid::Draw2D(UCanvas* Canvas, UWorld* World, const TAr
 			}
 			else
 			{
-				FIntVector CellGlobalCoords;
-				verify(Helper.GetCellGlobalCoords(FIntVector(Coords.X, Coords.Y, GridLevel), CellGlobalCoords));
+				FGridCellCoord CellGlobalCoords;
+				verify(Helper.GetCellGlobalCoords(FGridCellCoord(Coords.X, Coords.Y, GridLevel), CellGlobalCoords));
 				const FString CellCoordString = UWorldPartitionRuntimeSpatialHash::GetCellCoordString(CellGlobalCoords);
 				const FVector2D CellBoundsSize = CellScreenBounds.GetSize();
 				if (MaxCellCoordTextWidth < CellBoundsSize.X && MaxCellCoordTextHeight < CellBoundsSize.Y)
@@ -755,7 +755,7 @@ void UWorldPartitionRuntimeSpatialHash::PreSave(FObjectPreSaveContext ObjectSave
 	}
 }
 
-FString UWorldPartitionRuntimeSpatialHash::GetCellCoordString(const FIntVector& InCellGlobalCoord)
+FString UWorldPartitionRuntimeSpatialHash::GetCellCoordString(const FGridCellCoord& InCellGlobalCoord)
 {
 	return FString::Printf(TEXT("L%d_X%d_Y%d"), InCellGlobalCoord.Z, InCellGlobalCoord.X, InCellGlobalCoord.Y);
 }
@@ -876,7 +876,7 @@ void UWorldPartitionRuntimeSpatialHash::DumpStateLog(FHierarchicalLogArchive& Ar
 		struct FGridLevelStats
 		{
 			int32 CellCount;
-			int32 CellSize;
+			int64 CellSize;
 			int32 ActorCount;
 		};
 
@@ -908,7 +908,7 @@ void UWorldPartitionRuntimeSpatialHash::DumpStateLog(FHierarchicalLogArchive& Ar
 			FHierarchicalLogArchive::FIndentScope IndentScope = Ar.PrintfIndent(TEXT("Grid Levels: %d"), StreamingGrid.GridLevels.Num());
 			for (int Level=0; Level<LevelsStats.Num(); ++Level)
 			{
-				Ar.Printf(TEXT("Level %2d: Cell Count %4d | Cell Size %7d | Actor Count %4d (%3.1f%%)"), Level, LevelsStats[Level].CellCount, LevelsStats[Level].CellSize, LevelsStats[Level].ActorCount, (100.f*LevelsStats[Level].ActorCount)/TotalActorCount);
+				Ar.Printf(TEXT("Level %2d: Cell Count %4d | Cell Size %7lld | Actor Count %4d (%3.1f%%)"), Level, LevelsStats[Level].CellCount, LevelsStats[Level].CellSize, LevelsStats[Level].ActorCount, (100.f*LevelsStats[Level].ActorCount)/TotalActorCount);
 			}
 		}
 
@@ -933,7 +933,7 @@ void UWorldPartitionRuntimeSpatialHash::DumpStateLog(FHierarchicalLogArchive& Ar
 	}	
 }
 
-FName UWorldPartitionRuntimeSpatialHash::GetCellName(UWorldPartition* WorldPartition, FName InGridName, const FIntVector& InCellGlobalCoord, const FDataLayersID& InDataLayerID)
+FName UWorldPartitionRuntimeSpatialHash::GetCellName(UWorldPartition* WorldPartition, FName InGridName, const FGridCellCoord& InCellGlobalCoord, const FDataLayersID& InDataLayerID)
 {
 	const FString PackageName = FPackageName::GetShortName(WorldPartition->GetPackage());
 	const FString PackageNameNoPIEPrefix = UWorld::RemovePIEPrefix(PackageName);
@@ -952,7 +952,7 @@ void UWorldPartitionRuntimeSpatialHash::SetPreviewGrids(bool bInPreviewGrids)
 	bPreviewGrids = bInPreviewGrids;
 }
 
-FName UWorldPartitionRuntimeSpatialHash::GetCellName(FName InGridName, const FIntVector& InCellGlobalCoord, const FDataLayersID& InDataLayerID) const
+FName UWorldPartitionRuntimeSpatialHash::GetCellName(FName InGridName, const FGridCellCoord& InCellGlobalCoord, const FDataLayersID& InDataLayerID) const
 {
 	UWorldPartition* WorldPartition = GetOuterUWorldPartition();
 	return UWorldPartitionRuntimeSpatialHash::GetCellName(WorldPartition, InGridName, InCellGlobalCoord, InDataLayerID);
@@ -990,9 +990,9 @@ bool UWorldPartitionRuntimeSpatialHash::CreateStreamingGrid(const FSpatialHashRu
 
 		for (const auto& TempCellMapping : TempLevel.CellsMapping)
 		{
-			const uint32 CellIndex = TempCellMapping.Key;
-			const int32 CellCoordX = CellIndex % TempLevel.GridSize;
-			const int32 CellCoordY = CellIndex / TempLevel.GridSize;
+			const int64 CellIndex = TempCellMapping.Key;
+			const int64 CellCoordX = CellIndex % TempLevel.GridSize;
+			const int64 CellCoordY = CellIndex / TempLevel.GridSize;
 
 			const FSquare2DGridHelper::FGridLevel::FGridCell& TempCell = TempLevel.Cells[TempCellMapping.Value];
 
@@ -1050,8 +1050,8 @@ bool UWorldPartitionRuntimeSpatialHash::CreateStreamingGrid(const FSpatialHashRu
 					continue;
 				}
 				
-				FIntVector CellGlobalCoords;
-				verify(PartionedActors.GetCellGlobalCoords(FIntVector(CellCoordX, CellCoordY, Level), CellGlobalCoords));
+				FGridCellCoord CellGlobalCoords;
+				verify(PartionedActors.GetCellGlobalCoords(FGridCellCoord(CellCoordX, CellCoordY, Level), CellGlobalCoords));
 				FName CellName = GetCellName(CurrentStreamingGrid.GridName, CellGlobalCoords, GridCellDataChunk.GetDataLayersID());
 
 				UWorldPartitionRuntimeSpatialHashCell* StreamingCell = NewObject<UWorldPartitionRuntimeSpatialHashCell>(WorldPartition, StreamingPolicy->GetRuntimeCellClass(), CellName);
@@ -1074,10 +1074,12 @@ bool UWorldPartitionRuntimeSpatialHash::CreateStreamingGrid(const FSpatialHashRu
 				StreamingCell->Level = Level;
 				StreamingCell->SetPriority(RuntimeGrid.Priority);
 				FBox2D Bounds;
-				verify(TempLevel.GetCellBounds(FIntVector2(CellCoordX, CellCoordY), Bounds));
+				verify(TempLevel.GetCellBounds(FGridCellCoord2(CellCoordX, CellCoordY), Bounds));
 				StreamingCell->Position = FVector(Bounds.GetCenter(), 0.f);
-				StreamingCell->Extent = Bounds.GetExtent().X;
-				StreamingCell->SetDebugInfo(CellGlobalCoords, CurrentStreamingGrid.GridName);
+				const double CellExtent = Bounds.GetExtent().X;
+				check(CellExtent < MAX_flt);
+				StreamingCell->Extent = CellExtent;
+				StreamingCell->SetDebugInfo(CellGlobalCoords.X, CellGlobalCoords.Y, CellGlobalCoords.Z, CurrentStreamingGrid.GridName);
 				StreamingCell->SetClientOnlyVisible(CurrentStreamingGrid.bClientOnlyVisible);
 				StreamingCell->SetBlockOnSlowLoading(CurrentStreamingGrid.bBlockOnSlowStreaming);
 				StreamingCell->SetIsHLOD(RuntimeGrid.HLODLayer ? true : false);
@@ -1392,7 +1394,7 @@ EWorldPartitionStreamingPerformance UWorldPartitionRuntimeSpatialHash::GetStream
 	}
 
 	check(Cell->GetBlockOnSlowLoading());
-	const float BlockOnSlowStreamingWarningRatio = GBlockOnSlowStreamingRatio * GBlockOnSlowStreamingWarningFactor;
+	const double BlockOnSlowStreamingWarningRatio = GBlockOnSlowStreamingRatio * GBlockOnSlowStreamingWarningFactor;
 	
 	const UWorldPartitionRuntimeSpatialHashCell* StreamingCell = Cast<const UWorldPartitionRuntimeSpatialHashCell>(Cell);
 	check(StreamingCell);
@@ -1404,9 +1406,9 @@ EWorldPartitionStreamingPerformance UWorldPartitionRuntimeSpatialHash::GetStream
 
 	if (StreamingCell->CachedIsBlockingSource)
 	{
-		const float Distance = FMath::Sqrt(StreamingCell->CachedMinSquareDistanceToBlockingSource) - ((float)(*StreamingGrid)->GetCellSize(StreamingCell->Level) / 2.f);
+		const double Distance = FMath::Sqrt(StreamingCell->CachedMinSquareDistanceToBlockingSource) - ((double)(*StreamingGrid)->GetCellSize(StreamingCell->Level) / 2);
 
-		const float Ratio = Distance / LoadingRange;
+		const double Ratio = Distance / LoadingRange;
 
 		if (Ratio < GBlockOnSlowStreamingRatio)
 		{
@@ -1457,7 +1459,7 @@ void UWorldPartitionRuntimeSpatialHash::Draw2D(UCanvas* Canvas, const TArray<FWo
 		// Display view sides based on extended grid loading range (minimum of N cells)
 		// Take into consideration GShowRuntimeSpatialHashGridLevel when using CellSize
 		int32 MinGridLevel = FMath::Clamp<int32>(GShowRuntimeSpatialHashGridLevel, 0, StreamingGrid->GridLevels.Num() - 1);
-		int32 CellSize = (1 << MinGridLevel) * StreamingGrid->CellSize;
+		int64 CellSize = (1LL << MinGridLevel) * (int64)StreamingGrid->CellSize;
 		const FVector MinExtent(CellSize * GridViewMinimumSizeInCellCount);
 		FBox Region(ForceInit);
 		for (const FWorldPartitionStreamingSource& Source : Sources)
diff --git a/Engine/Source/Runtime/Engine/Private/WorldPartition/WorldPartitionRuntimeSpatialHashCell.cpp b/Engine/Source/Runtime/Engine/Private/WorldPartition/WorldPartitionRuntimeSpatialHashCell.cpp
--- a/Engine/Source/Runtime/Engine/Private/WorldPartition/WorldPartitionRuntimeSpatialHashCell.cpp
+++ b/Engine/Source/Runtime/Engine/Private/WorldPartition/WorldPartitionRuntimeSpatialHashCell.cpp
@@ -18,9 +18,9 @@ UWorldPartitionRuntimeSpatialHashCell::UWorldPartitionRuntimeSpatialHashCell(con
 : Super(ObjectInitializer)
 , Level(0)
 , CachedIsBlockingSource(false)
-, CachedMinSquareDistanceToBlockingSource(UE_MAX_FLT)
-, CachedMinSquareDistanceToSource(UE_MAX_FLT)
-, CachedSourceSortingDistance(0.f)
+, CachedMinSquareDistanceToBlockingSource(MAX_dbl)
+, CachedMinSquareDistanceToSource(MAX_dbl)
+, CachedSourceSortingDistance(0)
 {}
 
 #if WITH_EDITOR
@@ -48,13 +48,13 @@ bool UWorldPartitionRuntimeSpatialHashCell::CacheStreamingSourceInfo(const UWorl
 	if (bWasCacheDirtied)
 	{
 		CachedIsBlockingSource = false;
-		CachedMinSquareDistanceToBlockingSource = UE_MAX_FLT;
-		CachedMinSquareDistanceToSource = UE_MAX_FLT;
+		CachedMinSquareDistanceToBlockingSource = MAX_dbl;
+		CachedMinSquareDistanceToSource = MAX_dbl;
 		CachedSourceModulatedDistances.Reset();
 	}
 
 	float AngleContribution = FMath::Clamp(GRuntimeSpatialHashCellToSourceAngleContributionToCellImportance, 0.f, 1.f);
-	const float SquareDistance = FVector::DistSquared2D(Info.SourceShape.GetCenter(), Position);
+	const double SquareDistance = FVector::DistSquared2D(Info.SourceShape.GetCenter(), Position);
 	float AngleFactor = 1.f;
 	if (!FMath::IsNearlyZero(AngleContribution))
 	{
@@ -89,8 +89,8 @@ bool UWorldPartitionRuntimeSpatialHashCell::CacheStreamingSourceInfo(const UWorl
 		AngleFactor = FMath::Pow(NormalizedAngle, AngleContribution);
 	}
 	// Modulate distance to cell by angle relative to source forward vector (to prioritize cells in front)
-	const float SquareAngleFactor = AngleFactor * AngleFactor;
-	const float ModulatedSquareDistance = SquareDistance * SquareAngleFactor;
+	const double SquareAngleFactor = AngleFactor * AngleFactor;
+	const double ModulatedSquareDistance = SquareDistance * SquareAngleFactor;
 	CachedSourceModulatedDistances.Add(ModulatedSquareDistance);
 	int32 Count = CachedSourceModulatedDistances.Num();
 	check(Count == CachedSourcePriorityWeights.Num());
@@ -100,13 +100,13 @@ bool UWorldPartitionRuntimeSpatialHashCell::CacheStreamingSourceInfo(const UWorl
 	}
 	else
 	{
-		float TotalSourcePriorityWeight = 0.f;
+		double TotalSourcePriorityWeight = 0.f;
 		for (int32 i = 0; i < Count; ++i)
 		{
 			TotalSourcePriorityWeight += CachedSourcePriorityWeights[i];
 		}
 		int32 HighestPrioMinDistIndex = 0;
-		float WeightedModulatedDistance = 0.f;
+		double WeightedModulatedDistance = 0.f;
 		for (int32 i = 0; i < Count; ++i)
 		{
 			WeightedModulatedDistance += CachedSourceModulatedDistances[i] * CachedSourcePriorityWeights[i] / TotalSourcePriorityWeight;
@@ -149,11 +149,11 @@ int32 UWorldPartitionRuntimeSpatialHashCell::SortCompare(const UWorldPartitionRu
 		if (Result == 0)
 		{
 			// Closest distance (lower value is higher prio)
-			const float Diff = CachedSourceSortingDistance - Other->CachedSourceSortingDistance;
+			const double Diff = CachedSourceSortingDistance - Other->CachedSourceSortingDistance;
 			if (FMath::IsNearlyZero(Diff))
 			{
-				const float RawDistanceDiff = CachedMinSquareDistanceToSource - Other->CachedMinSquareDistanceToSource;
-				Result = RawDistanceDiff < 0.f ? -1 : (RawDistanceDiff > 0.f ? 1 : 0);
+				const double RawDistanceDiff = CachedMinSquareDistanceToSource - Other->CachedMinSquareDistanceToSource;
+				Result = RawDistanceDiff < 0 ? -1 : (RawDistanceDiff > 0.f ? 1 : 0);
 			}
 			else
 			{
diff --git a/Engine/Source/Runtime/Engine/Private/WorldPartition/WorldPartitionStreamingPolicy.cpp b/Engine/Source/Runtime/Engine/Private/WorldPartition/WorldPartitionStreamingPolicy.cpp
--- a/Engine/Source/Runtime/Engine/Private/WorldPartition/WorldPartitionStreamingPolicy.cpp
+++ b/Engine/Source/Runtime/Engine/Private/WorldPartition/WorldPartitionStreamingPolicy.cpp
@@ -731,11 +731,11 @@ FStreamingSourceVelocity::FStreamingSourceVelocity(const FName& InSourceName)
 
 float FStreamingSourceVelocity::GetAverageVelocity(const FVector& NewPosition, const float CurrentTime)
 {
-	const float TeleportDistance = 100.f;
+	const double TeleportDistance = 100;
 	const float MaxDeltaSeconds = 5.f;
 	const bool bIsFirstCall = (LastIndex == INDEX_NONE);
 	const float DeltaSeconds = bIsFirstCall ? 0.f : (CurrentTime - LastUpdateTime);
-	const float Distance = bIsFirstCall ? 0.f : ((NewPosition - LastPosition) * 0.01f).Size();
+	const double Distance = bIsFirstCall ? 0.f : ((NewPosition - LastPosition) * 0.01).Size();
 	if (bIsFirstCall)
 	{
 		UE_LOG(LogWorldPartition, Log, TEXT("New Streaming Source: %s -> Position: %s"), *SourceName.ToString(), *NewPosition.ToString());
@@ -756,7 +756,8 @@ float FStreamingSourceVelocity::GetAverageVelocity(const FVector& NewPosition, c
 	}
 
 	// Compute velocity (m/s)
-	const float Velocity = Distance / DeltaSeconds;
+	check(Distance < MAX_flt);
+	const float Velocity = (float)Distance / DeltaSeconds;
 	// Update velocities history buffer and sum
 	LastIndex = (LastIndex + 1) % VELOCITY_HISTORY_SAMPLE_COUNT;
 	VelocitiesHistorySum = FMath::Max<float>(0.f, (VelocitiesHistorySum + Velocity - VelocitiesHistory[LastIndex]));
diff --git a/Engine/Source/Runtime/Engine/Public/WorldPartition/WorldPartitionRuntimeCell.h b/Engine/Source/Runtime/Engine/Public/WorldPartition/WorldPartitionRuntimeCell.h
--- a/Engine/Source/Runtime/Engine/Public/WorldPartition/WorldPartitionRuntimeCell.h
+++ b/Engine/Source/Runtime/Engine/Public/WorldPartition/WorldPartitionRuntimeCell.h
@@ -106,6 +106,27 @@ enum class EWorldPartitionRuntimeCellState : uint8
 	Activated
 };
 
+USTRUCT()
+struct FWorldPartitionRuntimeCellDebugInfo
+{
+	GENERATED_BODY()
+
+	UPROPERTY()
+	FString Name;
+
+	UPROPERTY()
+	FName GridName;
+
+	UPROPERTY()
+	int64 CoordX;
+
+	UPROPERTY()
+	int64 CoordY;
+
+	UPROPERTY()
+	int64 CoordZ;
+};
+
 static_assert(EWorldPartitionRuntimeCellState::Unloaded < EWorldPartitionRuntimeCellState::Loaded && EWorldPartitionRuntimeCellState::Loaded < EWorldPartitionRuntimeCellState::Activated, "Streaming Query code is dependent on this being true");
 
 /**
@@ -142,10 +163,10 @@ class UWorldPartitionRuntimeCell : public UObject, public IWorldPartitionCell
 	virtual void SetStreamingPriority(int32 InStreamingPriority) const PURE_VIRTUAL(UWorldPartitionRuntimeCell::SetStreamingPriority,);
 	virtual EStreamingStatus GetStreamingStatus() const { return LEVEL_Unloaded; }
 	virtual bool IsLoading() const { return false; }
-	virtual const FString& GetDebugName() const { return DebugName; }
+	virtual const FString& GetDebugName() const { return DebugInfo.Name; }
 	virtual bool IsDebugShown() const;
 	virtual int32 SortCompare(const UWorldPartitionRuntimeCell* Other) const;
-	virtual FName GetGridName() const { return GridName; }
+	virtual FName GetGridName() const { return DebugInfo.GridName; }
 	/** Caches information on streaming source that will be used later on to sort cell. Returns true if cache was reset, else returns false. */
 	virtual bool CacheStreamingSourceInfo(const UWorldPartitionRuntimeCell::FStreamingSourceInfo& Info) const;
 
@@ -175,7 +196,7 @@ class UWorldPartitionRuntimeCell : public UObject, public IWorldPartitionCell
 
 	void SetDataLayers(const TArray<const UDataLayerInstance*>& InDataLayerInstances);
 	void SetMinMaxZ(const FVector2D& InMinMaxZ);
-	void SetDebugInfo(FIntVector InCoords, FName InGridName);
+	void SetDebugInfo(int64 InCoordX, int64 InCoordY, int64 InCoordZ, FName InGridName);
 	virtual void AddActorToCell(const FWorldPartitionActorDescView& ActorDescView, const FActorContainerID& InContainerID, const FTransform& InContainerTransform, const UActorDescContainer* InContainer) PURE_VIRTUAL(UWorldPartitionRuntimeCell::AddActorToCell,);
 	virtual int32 GetActorCount() const PURE_VIRTUAL(UWorldPartitionRuntimeCell::GetActorCount, return 0;);
 
@@ -216,13 +237,7 @@ private:
 
 	// Debug Info
 	UPROPERTY()
-	FIntVector Coords;
-
-	UPROPERTY()
-	FName GridName;
-
-	UPROPERTY()
-	FString DebugName;
+	FWorldPartitionRuntimeCellDebugInfo DebugInfo;
 
 	// Custom Priority
 	UPROPERTY()
diff --git a/Engine/Source/Runtime/Engine/Public/WorldPartition/WorldPartitionRuntimeSpatialHash.h b/Engine/Source/Runtime/Engine/Public/WorldPartition/WorldPartitionRuntimeSpatialHash.h
--- a/Engine/Source/Runtime/Engine/Public/WorldPartition/WorldPartitionRuntimeSpatialHash.h
+++ b/Engine/Source/Runtime/Engine/Public/WorldPartition/WorldPartitionRuntimeSpatialHash.h
@@ -15,6 +15,9 @@
 #include "WorldPartitionRuntimeSpatialHashGridPreviewer.h"
 #include "WorldPartitionRuntimeSpatialHash.generated.h"
 
+typedef UE::Math::TIntVector3<int64> FGridCellCoord;
+typedef UE::Math::TIntVector2<int64> FGridCellCoord2;
+
 USTRUCT()
 struct FSpatialHashStreamingGridLayerCell
 {
@@ -65,7 +68,7 @@ struct FSpatialHashStreamingGridLevel
 	 * Represents the grid cells index mapping
 	 */
 	UPROPERTY()
-	TMap<int32, int32> LayerCellsMapping;
+	TMap<int64, int32> LayerCellsMapping;
 };
 
 // Disable copy of this USTRUCT by the autogenerated code.
@@ -134,7 +137,7 @@ struct FSpatialHashStreamingGrid
 	}
 
 	// Used by PIE/Game
-	int32 GetCellSize(int32 Level) const;
+	int64 GetCellSize(int32 Level) const;
 	void GetCells(const FWorldPartitionStreamingQuerySource& QuerySource, TSet<const UWorldPartitionRuntimeCell*>& OutCells, bool bEnableZCulling) const;
 	void GetCells(const TArray<FWorldPartitionStreamingSource>& Sources, const class UDataLayerSubsystem* DataLayerSubsystem, UWorldPartitionRuntimeHash::FStreamingSourceCells& OutActivateCells, UWorldPartitionRuntimeHash::FStreamingSourceCells& OutLoadCells, bool bEnableZCulling) const;
 	void GetAlwaysLoadedCells(const UDataLayerSubsystem* DataLayerSubsystem, TSet<const UWorldPartitionRuntimeCell*>& OutActivateCells, TSet<const UWorldPartitionRuntimeCell*>& OutLoadCells) const;
@@ -142,7 +145,7 @@ struct FSpatialHashStreamingGrid
 	void Draw3D(UWorld* World, const TArray<FWorldPartitionStreamingSource>& Sources, const FTransform& Transform) const;
 
 private:
-	const FSpatialHashStreamingGridLayerCell* GetLayerCell(const FIntVector& Coords) const;
+	const FSpatialHashStreamingGridLayerCell* GetLayerCell(const FGridCellCoord& Coords) const;
 	void DrawStreamingSource2D(UCanvas* Canvas, const FSphericalSector& Shape, TFunctionRef<FVector2D(const FVector2D&)> WorldToScreen, const FColor& Color) const;
 	void DrawStreamingSource3D(UWorld* World, const FSphericalSector& Shape, const FTransform& Transform, const FColor& Color) const;
 	void GetFilteredCellsForDebugDraw(const FSpatialHashStreamingGridLayerCell* LayerCell, const UDataLayerSubsystem* DataLayerSubsystem, TArray<const UWorldPartitionRuntimeCell*>& FilteredCells) const;
@@ -237,8 +240,8 @@ public:
 	virtual bool GenerateHLOD(ISourceControlHelper* SourceControlHelper, FActorClusterContext& ActorClusterContext, bool bCreateActorsOnly) override;
 	virtual void DrawPreview() const override;
 
-	FName GetCellName(FName InGridName, const FIntVector& InCellGlobalCoord, const FDataLayersID& InDataLayerID) const;
-	static FName GetCellName(UWorldPartition* WorldPartition, FName InGridName, const FIntVector& InCellGlobalCoord, const FDataLayersID& InDataLayerID);
+	FName GetCellName(FName InGridName, const FGridCellCoord& InCellGlobalCoord, const FDataLayersID& InDataLayerID) const;
+	static FName GetCellName(UWorldPartition* WorldPartition, FName InGridName, const FGridCellCoord& InCellGlobalCoord, const FDataLayersID& InDataLayerID);
 
 	bool GetPreviewGrids() const;
 	void SetPreviewGrids(bool bInPreviewGrids);
@@ -249,7 +252,7 @@ public:
 	virtual bool GetStreamingCells(const FWorldPartitionStreamingQuerySource& QuerySource, TSet<const UWorldPartitionRuntimeCell*>& OutCells) const override;
 	virtual bool GetStreamingCells(const TArray<FWorldPartitionStreamingSource>& Sources, UWorldPartitionRuntimeHash::FStreamingSourceCells& OutActivateCells, UWorldPartitionRuntimeHash::FStreamingSourceCells& OutLoadCells) const override;
 		
-	static FString GetCellCoordString(const FIntVector& InCellGlobalCoord);
+	static FString GetCellCoordString(const FGridCellCoord& InCellGlobalCoord);
 
 protected:
 	bool ShouldConsiderClientOnlyVisibleCells() const;
diff --git a/Engine/Source/Runtime/Engine/Public/WorldPartition/WorldPartitionRuntimeSpatialHashCell.h b/Engine/Source/Runtime/Engine/Public/WorldPartition/WorldPartitionRuntimeSpatialHashCell.h
--- a/Engine/Source/Runtime/Engine/Public/WorldPartition/WorldPartitionRuntimeSpatialHashCell.h
+++ b/Engine/Source/Runtime/Engine/Public/WorldPartition/WorldPartitionRuntimeSpatialHashCell.h
@@ -31,10 +31,10 @@ class UWorldPartitionRuntimeSpatialHashCell : public UWorldPartitionRuntimeCell
 	mutable bool CachedIsBlockingSource;
 
 	// Represents the square distance from cell to the closest blocking streaming source
-	mutable float CachedMinSquareDistanceToBlockingSource;
+	mutable double CachedMinSquareDistanceToBlockingSource;
 
 	// Represents the square distance from cell to the closest streaming source
-	mutable float CachedMinSquareDistanceToSource;
+	mutable double CachedMinSquareDistanceToSource;
 
 	// Modulated distance to the different streaming sources used to sort relative priority amongst streaming cells
 	// The value is affected by :
@@ -42,9 +42,9 @@ class UWorldPartitionRuntimeSpatialHashCell : public UWorldPartitionRuntimeCell
 	// - The priority of each source
 	// - The distance between the cell and each source
 	// - The angle between the cell and each source orientation
-	mutable float CachedSourceSortingDistance;
+	mutable double CachedSourceSortingDistance;
 
-	mutable TArray<float> CachedSourceModulatedDistances;
+	mutable TArray<double> CachedSourceModulatedDistances;
 
 #if WITH_EDITORONLY_DATA
 	UPROPERTY()
diff --git a/Engine/Source/Runtime/Engine/Public/WorldPartition/WorldPartitionStreamingSource.h b/Engine/Source/Runtime/Engine/Public/WorldPartition/WorldPartitionStreamingSource.h
--- a/Engine/Source/Runtime/Engine/Public/WorldPartition/WorldPartitionStreamingSource.h
+++ b/Engine/Source/Runtime/Engine/Public/WorldPartition/WorldPartitionStreamingSource.h
@@ -410,11 +410,11 @@ struct ENGINE_API FWorldPartitionStreamingSource
 	{
 		const FVector Direction = Rotation.Euler();
 		return FString::Printf(
-			TEXT("Priority: %d | %s | %s | Pos: X=%d,Y=%d,Z=%d | Rot: X=%d,Y=%d,Z=%d | Vel: %3.2f m/s (%d mph)"), 
+			TEXT("Priority: %d | %s | %s | Pos: X=%lld,Y=%lld,Z=%lld | Rot: X=%d,Y=%d,Z=%d | Vel: %3.2f m/s (%d mph)"), 
 			Priority, 
 			GetStreamingSourceTargetStateName(TargetState),
 			bBlockOnSlowLoading ? TEXT("Blocking") : TEXT("NonBlocking"),
-			(int32)Location.X, (int32)Location.Y, (int32)Location.Z, 
+			(int64)Location.X, (int64)Location.Y, (int64)Location.Z, 
 			(int32)Direction.X, (int32)Direction.Y, (int32)Direction.Z, 
 			Velocity, 
 			(int32)(Velocity*2.23694f)
