Fix uninitialized members of FWorldPartitionRuntimeCellDebugInfo.

#rnx

diff --git a/Engine/Source/Runtime/Engine/Public/WorldPartition/WorldPartitionRuntimeCell.h b/Engine/Source/Runtime/Engine/Public/WorldPartition/WorldPartitionRuntimeCell.h
--- a/Engine/Source/Runtime/Engine/Public/WorldPartition/WorldPartitionRuntimeCell.h
+++ b/Engine/Source/Runtime/Engine/Public/WorldPartition/WorldPartitionRuntimeCell.h
@@ -118,13 +118,13 @@ struct FWorldPartitionRuntimeCellDebugInfo
 	FName GridName;
 
 	UPROPERTY()
-	int64 CoordX;
+	int64 CoordX = 0;
 
 	UPROPERTY()
-	int64 CoordY;
+	int64 CoordY = 0;
 
 	UPROPERTY()
-	int64 CoordZ;
+	int64 CoordZ = 0;
 };
 
 static_assert(EWorldPartitionRuntimeCellState::Unloaded < EWorldPartitionRuntimeCellState::Loaded && EWorldPartitionRuntimeCellState::Loaded < EWorldPartitionRuntimeCellState::Activated, "Streaming Query code is dependent on this being true");
