Actor: Fix missing check for NAME_NO_NUMBER_INTERNAL when assigning default label

This was causing WorldDataLayers to be labeled WorldDataLayers-1

diff --git a/Engine/Source/Runtime/Engine/Private/ActorEditor.cpp b/Engine/Source/Runtime/Engine/Private/ActorEditor.cpp
--- a/Engine/Source/Runtime/Engine/Private/ActorEditor.cpp
+++ b/Engine/Source/Runtime/Engine/Private/ActorEditor.cpp
@@ -925,10 +925,9 @@ const FString& AActor::GetActorLabel(bool bCreateIfNone) const
 		if (!FActorSpawnUtils::IsGloballyUniqueName(GetFName()))
 		{
 			// Don't bother adding a suffix for number '0'
-			const int32 NameNumber = NAME_INTERNAL_TO_EXTERNAL( GetFName().GetNumber() );
-			if( NameNumber != 0 )
+			if (const int32 NameNumber = GetFName().GetNumber(); NameNumber != NAME_NO_NUMBER_INTERNAL)
 			{
-				DefaultActorLabel.AppendInt(NameNumber);
+				DefaultActorLabel.AppendInt(NAME_INTERNAL_TO_EXTERNAL(NameNumber));
 			}
 		}
 
