[KillPendingKill] Release object references from dormant replicators when levels are streamed out or actors are destroyed.

diff --git a/Engine/Source/Runtime/Engine/Classes/Engine/NetConnection.h b/Engine/Source/Runtime/Engine/Classes/Engine/NetConnection.h
--- a/Engine/Source/Runtime/Engine/Classes/Engine/NetConnection.h
+++ b/Engine/Source/Runtime/Engine/Classes/Engine/NetConnection.h
@@ -763,9 +763,7 @@ public:
 	TMap<FNetworkGUID, TArray<class UActorChannel*>> KeepProcessingActorChannelBunchesMap;
 
 	/** A list of replicators that belong to recently dormant actors/objects */
-	TMap< UObject*, TSharedRef< FObjectReplicator > > DormantReplicatorMap;
-
-	
+	TMap<FObjectKey, TSharedRef<FObjectReplicator>> DormantReplicatorMap;	
 
 	ENGINE_API FName GetClientWorldPackageName() const { return ClientWorldPackageName; }
 
@@ -1296,7 +1294,12 @@ public:
 	/** Returns the OutgoingBunches array, only to be used by UChannel::SendBunch */
 	TArray<FOutBunch *>& GetOutgoingBunches() { return OutgoingBunches; }
 
-	/** Removes Actor and its replicated components from DormantReplicatorMap. */
+	/** Add a replicator to the dormancy map and release its strong pointer to its object */
+	void AddDormantReplicator(UObject* Object, const TSharedRef<FObjectReplicator>& Replicator);
+	
+	/** Find a dormant replicator for the channel actor or one of its subobjects. Removes it from the map if found. */
+	TSharedPtr<FObjectReplicator> FindAndRemoveDormantReplicator(UObject* Object);
+
 	void CleanupDormantReplicatorsForActor(AActor* Actor);
 
 	/** Removes stale entries from DormantReplicatorMap. */
diff --git a/Engine/Source/Runtime/Engine/Private/DataChannel.cpp b/Engine/Source/Runtime/Engine/Private/DataChannel.cpp
--- a/Engine/Source/Runtime/Engine/Private/DataChannel.cpp
+++ b/Engine/Source/Runtime/Engine/Private/DataChannel.cpp
@@ -2234,7 +2234,7 @@ void UActorChannel::CleanupReplicators(const bool bKeepReplicators)
 			//		KeepProcessingActorChannelBunchesMap will get in here, then when the channel closes a second time, we'll hit this assert
 			//		It should be okay to just set the most recent replicator
 			//check( Connection->DormantReplicatorMap.Find( CompIt.Value()->GetObject() ) == NULL );
-			Connection->DormantReplicatorMap.Add(CompIt.Value()->GetObject(), CompIt.Value());
+			Connection->AddDormantReplicator(CompIt.Value()->GetObject(), CompIt.Value());
 			CompIt.Value()->StopReplicating(this);		// Stop replicating on this channel
 		}
 		else
@@ -4739,13 +4739,19 @@ TSharedRef<FObjectReplicator>& UActorChannel::CreateReplicator(UObject* Obj, boo
 
 	// Try to find in the dormancy map if desired
 	TSharedPtr<FObjectReplicator> NewReplicator;
-	TSharedRef<FObjectReplicator>* ReplicatorRefPtr = (bCheckDormantReplicators ? Connection->DormantReplicatorMap.Find(Obj) : nullptr);
+	if (bCheckDormantReplicators)
+	{
+		NewReplicator = Connection->FindAndRemoveDormantReplicator(Obj);
+	}
+	else
+	{
+		Connection->FindAndRemoveDormantReplicator(Obj); // Discard dormant replicator if there was one that we're not using;
+	}
 
 	// Check if we found it and that it is has a valid object
-	if ( ReplicatorRefPtr != nullptr && ReplicatorRefPtr->Get().GetWeakObjectPtr().IsValid() )
+	if(NewReplicator.IsValid())
 	{
 		UE_LOG( LogNetTraffic, Log, TEXT( "Found existing replicator for %s" ), *Obj->GetName() );
-		NewReplicator = *ReplicatorRefPtr;
 	}
 	else
 	{
@@ -4757,9 +4763,6 @@ TSharedRef<FObjectReplicator>& UActorChannel::CreateReplicator(UObject* Obj, boo
 	// Add to the replication map
 	TSharedRef<FObjectReplicator>& NewRef = ReplicationMap.Add(Obj, NewReplicator.ToSharedRef());
 
-	// Remove from dormancy map in case we found it there
-	Connection->DormantReplicatorMap.Remove(Obj);
-
 	// Start replicating with this replicator
 	NewRef->StartReplicating(this);
 	return NewRef;
diff --git a/Engine/Source/Runtime/Engine/Private/NetConnection.cpp b/Engine/Source/Runtime/Engine/Private/NetConnection.cpp
--- a/Engine/Source/Runtime/Engine/Private/NetConnection.cpp
+++ b/Engine/Source/Runtime/Engine/Private/NetConnection.cpp
@@ -4618,6 +4618,7 @@ void UNetConnection::FlushDormancy(class AActor* Actor)
 	{
 		FlushDormancyForObject( Actor );
 
+		// TODO: Is this set of objects sufficient? Should we query the dormancy map from the connection instead?
 		for ( UActorComponent* ActorComp : Actor->GetReplicatedComponents() )
 		{
 			if ( ActorComp && ActorComp->GetIsReplicated() )
@@ -4848,6 +4849,31 @@ void UNetConnection::DestroyIgnoredActor(AActor* Actor)
 	}
 }
 
+void UNetConnection::AddDormantReplicator(UObject* Object, const TSharedRef<FObjectReplicator>& Replicator)
+{
+	Replicator->ReleaseStrongReference();
+	DormantReplicatorMap.Add(Object, Replicator);
+}
+
+TSharedPtr<FObjectReplicator> UNetConnection::FindAndRemoveDormantReplicator(UObject* Object)
+{
+	FObjectKey Key(Object);
+	
+	if (DormantReplicatorMap.Contains(Key))
+	{
+		TSharedRef<FObjectReplicator> Ref = DormantReplicatorMap.FindAndRemoveChecked(Key);
+
+		// Only return the replicator if the object is still valid, otherwise just remove it from the cache and allow the caller to create a new one
+		if (UObject* StrongPtr = Ref->GetWeakObjectPtr().Get())
+		{
+			// Reassign the strong pointer for GC/faster resolve
+			Ref->SetObject(StrongPtr);
+			return Ref;
+		}
+	}
+	return {};
+}
+
 void UNetConnection::CleanupDormantReplicatorsForActor(AActor* Actor)
 {
 	if (Actor)
diff --git a/Engine/Source/Runtime/Engine/Public/Net/DataReplication.h b/Engine/Source/Runtime/Engine/Public/Net/DataReplication.h
--- a/Engine/Source/Runtime/Engine/Public/Net/DataReplication.h
+++ b/Engine/Source/Runtime/Engine/Public/Net/DataReplication.h
@@ -205,7 +205,8 @@ public:
 
 	FORCEINLINE UObject* GetObject() const
 	{
-		return ObjectPtr;
+		// If this replicator is dormant we have released our strong ref but the object may still be alive.
+		return ObjectPtr ? ObjectPtr : WeakObjectPtr.Get();
 	}
 
 	FORCEINLINE void SetObject(UObject* NewObj)
@@ -214,6 +215,11 @@ public:
 		WeakObjectPtr = NewObj;
 	}
 
+	void ReleaseStrongReference()
+	{
+		ObjectPtr = nullptr;
+	}
+
 	FORCEINLINE void PreNetReceive()
 	{ 
 		UObject* Object = GetObject();
