Deprecate ControllerID verserions of the console in favor of InputDeviceId

#rnx

diff --git a/Engine/Plugins/Experimental/CommonUI/Source/CommonUI/Private/CommonGameViewportClient.cpp b/Engine/Plugins/Experimental/CommonUI/Source/CommonUI/Private/CommonGameViewportClient.cpp
--- a/Engine/Plugins/Experimental/CommonUI/Source/CommonUI/Private/CommonGameViewportClient.cpp
+++ b/Engine/Plugins/Experimental/CommonUI/Source/CommonUI/Private/CommonGameViewportClient.cpp
@@ -204,7 +204,7 @@ bool UCommonGameViewportClient::IsKeyPriorityAboveUI(const FInputKeyEventArgs& E
 {
 #if !UE_BUILD_SHIPPING
 	// First priority goes to the viewport console regardless any state or setting
-	if (ViewportConsole && ViewportConsole->InputKey(EventArgs.ControllerId, EventArgs.Key, EventArgs.Event, EventArgs.AmountDepressed, EventArgs.IsGamepad()))
+	if (ViewportConsole && ViewportConsole->InputKey(EventArgs.InputDevice, EventArgs.Key, EventArgs.Event, EventArgs.AmountDepressed, EventArgs.IsGamepad()))
 	{
 		return true;
 	}
diff --git a/Engine/Source/Runtime/Engine/Classes/Engine/Console.h b/Engine/Source/Runtime/Engine/Classes/Engine/Console.h
--- a/Engine/Source/Runtime/Engine/Classes/Engine/Console.h
+++ b/Engine/Source/Runtime/Engine/Classes/Engine/Console.h
@@ -205,7 +205,19 @@ class ENGINE_API UConsole
 	 *
 	 * @return	True to consume the character, false to pass it on.
 	 */
+	UE_DEPRECATED(5.1, "This version of InputChar_Typing is deprecated. Please use the version that takes a DeviceId instead.")
 	virtual bool InputChar_Typing( int32 ControllerId, const FString& Unicode );
+
+	/**
+	 * Process a character input event (typing) routed through unrealscript from another object. This method is assigned as the value for the
+	 * OnReceivedNativeInputKey delegate so that native input events are routed to this unrealscript function.
+	 *
+	 * @param	DeviceId		the input device that generated this character input event
+	 * @param	Unicode			the character that was typed
+	 *
+	 * @return	True to consume the character, false to pass it on.
+	 */
+	virtual bool InputChar_Typing(FInputDeviceId DeviceId, const FString& Unicode);
 	
 	/**
 	 * perform rendering of the console on the canvas
@@ -221,7 +233,13 @@ class ENGINE_API UConsole
 	/**
 	 * This state is used when the console is open.
 	 */
+	UE_DEPRECATED(5.1, "This version of InputChar_Open has been deprecated. Please use the version that takes a DeviceId instead.")
 	virtual bool InputChar_Open( int32 ControllerId, const FString& Unicode );
+
+	/**
+	 * This state is used when the console is open.
+	 */
+	virtual bool InputChar_Open( FInputDeviceId DeviceId, const FString& Unicode );
 	
 	/**
 	 * Process an input key event routed through unrealscript from another object. This method is assigned as the value for the
@@ -234,7 +252,21 @@ class ENGINE_API UConsole
 	 *
 	 * @return	true to consume the key event, false to pass it on.
 	 */
+	UE_DEPRECATED(5.1, "This version of InputKey_Open has been deprecated. Please use the version that takes a DeviceId instead.")
 	virtual bool InputKey_Open( int32 ControllerId, FKey Key, EInputEvent Event, float AmountDepressed = 1.f, bool bGamepad = false );
+
+	/**
+	 * Process an input key event routed through unrealscript from another object. This method is assigned as the value for the
+	 * OnReceivedNativeInputKey delegate so that native input events are routed to this unrealscript function.
+	 *
+	 * @param	ControllerId	the controller that generated this input key event
+	 * @param	Key				the name of the key which an event occured for (KEY_Up, KEY_Down, etc.)
+	 * @param	EventType		the type of event which occured (pressed, released, etc.)
+	 * @param	AmountDepressed	for analog keys, the depression percent.
+	 *
+	 * @return	true to consume the key event, false to pass it on.
+	 */
+	virtual bool InputKey_Open(FInputDeviceId DeviceId, FKey Key, EInputEvent Event, float AmountDepressed = 1.f, bool bGamepad = false);
 	
 	/** 
 	 * perform rendering of the console on the canvas
@@ -247,10 +279,21 @@ class ENGINE_API UConsole
 	/** Perform actions on transition from Open state */
 	virtual void EndState_Open(FName NextStateName);
 
+	UE_DEPRECATED(5.1, "This version of InputChar has been deprecated. Please use the version that takes an FInputDeviceId instead")
 	virtual bool InputChar(int32 ControllerId, const FString& Unicode);
+	virtual bool InputChar(FInputDeviceId DeviceId, const FString& Unicode);
+	
+	UE_DEPRECATED(5.1, "This version of InputKey has been deprecated. Please use the version that takes an FInputDeviceId instead")
 	virtual bool InputKey(int32 ControllerId, FKey Key, EInputEvent Event, float AmountDepressed=1.f, bool bGamepad=false);
+	virtual bool InputKey(FInputDeviceId DeviceId, FKey Key, EInputEvent Event, float AmountDepressed=1.f, bool bGamepad=false);
+
+	UE_DEPRECATED(5.1, "This version of InputAxis has been deprecated. Please use the version that takes an FInputDeviceId instead")
 	virtual bool InputAxis(int32 ControllerId, FKey Key, float Delta, float DeltaTime, int32 NumSamples=1, bool bGamepad=false) { return false; };
+	virtual bool InputAxis(FInputDeviceId DevideId, FKey Key, float Delta, float DeltaTime, int32 NumSamples=1, bool bGamepad=false) { return false; };
+
+	UE_DEPRECATED(5.1, "This version of InputTouch has been deprecated. Please use the version that takes an FInputDeviceId instead")
 	virtual bool InputTouch(int32 ControllerId, uint32 Handle, ETouchType::Type Type, const FVector2D& TouchLocation, float Force, FDateTime DeviceTimestamp, uint32 TouchpadIndex) { return false; }
+	virtual bool InputTouch(FInputDeviceId DevideId, uint32 Handle, ETouchType::Type Type, const FVector2D& TouchLocation, float Force, FDateTime DeviceTimestamp, uint32 TouchpadIndex) { return false; }
 
 	/** render to the canvas based on the console state */
 	virtual void PostRender_Console(class UCanvas* Canvas);
@@ -272,7 +315,7 @@ class ENGINE_API UConsole
 
 private:
 
-	bool InputKey_InputLine( int32 ControllerId, FKey Key, EInputEvent Event, float AmountDepressed = 1.f, bool bGamepad = false );
+	bool InputKey_InputLine(FInputDeviceId DeviceId, FKey Key, EInputEvent Event, float AmountDepressed = 1.f, bool bGamepad = false);
 
 	// interface FOutputDevice
 	virtual void Serialize( const TCHAR* V, ELogVerbosity::Type Verbosity, const class FName& Category ) override;
diff --git a/Engine/Source/Runtime/Engine/Classes/Engine/GameViewportClient.h b/Engine/Source/Runtime/Engine/Classes/Engine/GameViewportClient.h
--- a/Engine/Source/Runtime/Engine/Classes/Engine/GameViewportClient.h
+++ b/Engine/Source/Runtime/Engine/Classes/Engine/GameViewportClient.h
@@ -169,7 +169,8 @@ public:
 	//~ Begin FViewportClient Interface.
 	virtual void RedrawRequested(FViewport* InViewport) override {}
 	virtual bool InputKey(const FInputKeyEventArgs& EventArgs) override;
-	// TODO: Deprecate this one
+	
+	UE_DEPRECATED(5.1, "This version of InputAxis has been deprecated. Please use the version that takes an FInputDeviceId instead")
 	virtual bool InputAxis(FViewport* Viewport, int32 ControllerId, FKey Key, float Delta, float DeltaTime, int32 NumSamples=1, bool bGamepad=false) override;
 	virtual bool InputAxis(FViewport* Viewport, FInputDeviceId InputDevice, FKey Key, float Delta, float DeltaTime, int32 NumSamples = 1, bool bGamepad = false) override;
 	
diff --git a/Engine/Source/Runtime/Engine/Private/GameViewportClient.cpp b/Engine/Source/Runtime/Engine/Private/GameViewportClient.cpp
--- a/Engine/Source/Runtime/Engine/Private/GameViewportClient.cpp
+++ b/Engine/Source/Runtime/Engine/Private/GameViewportClient.cpp
@@ -634,7 +634,7 @@ bool UGameViewportClient::InputKey(const FInputKeyEventArgs& InEventArgs)
 
 	if (IgnoreInput())
 	{
-		return ViewportConsole ? ViewportConsole->InputKey(EventArgs.ControllerId, EventArgs.Key, EventArgs.Event, EventArgs.AmountDepressed, EventArgs.IsGamepad()) : false;
+		return ViewportConsole ? ViewportConsole->InputKey(EventArgs.InputDevice, EventArgs.Key, EventArgs.Event, EventArgs.AmountDepressed, EventArgs.IsGamepad()) : false;
 	}
 
 	OnInputKeyEvent.Broadcast(EventArgs);
@@ -651,7 +651,7 @@ bool UGameViewportClient::InputKey(const FInputKeyEventArgs& InEventArgs)
 #endif
 
 	// route to subsystems that care
-	bool bResult = ( ViewportConsole ? ViewportConsole->InputKey(EventArgs.ControllerId, EventArgs.Key, EventArgs.Event, EventArgs.AmountDepressed, EventArgs.IsGamepad()) : false );
+	bool bResult = ( ViewportConsole ? ViewportConsole->InputKey(EventArgs.InputDevice, EventArgs.Key, EventArgs.Event, EventArgs.AmountDepressed, EventArgs.IsGamepad()) : false );
 
 	// Try the override callback, this may modify event args
 	if (!bResult && OnOverrideInputKeyEvent.IsBound())
@@ -693,7 +693,7 @@ bool UGameViewportClient::InputKey(const FInputKeyEventArgs& InEventArgs)
 	return bResult;
 }
 
-
+PRAGMA_DISABLE_DEPRECATION_WARNINGS
 bool UGameViewportClient::InputAxis(FViewport* InViewport, int32 ControllerId, FKey Key, float Delta, float DeltaTime, int32 NumSamples, bool bGamepad)
 {
 	// Remap the old int32 ControllerId value to the new InputDeviceId
@@ -704,6 +704,7 @@ bool UGameViewportClient::InputAxis(FViewport* InViewport, int32 ControllerId, F
 	
 	return InputAxis(InViewport, DeviceId, Key, Delta, DeltaTime, NumSamples, bGamepad);
 }
+PRAGMA_ENABLE_DEPRECATION_WARNINGS
 
 bool UGameViewportClient::InputAxis(FViewport* InViewport, FInputDeviceId InputDevice, FKey Key, float Delta, float DeltaTime, int32 NumSamples, bool bGamepad)
 {
@@ -727,7 +728,7 @@ bool UGameViewportClient::InputAxis(FViewport* InViewport, FInputDeviceId InputD
 		// route to subsystems that care
 		if (ViewportConsole != nullptr)
 		{
-			bResult = ViewportConsole->InputAxis(EventArgs.ControllerId, EventArgs.Key, Delta, DeltaTime, NumSamples, EventArgs.IsGamepad());
+			bResult = ViewportConsole->InputAxis(EventArgs.InputDevice, EventArgs.Key, Delta, DeltaTime, NumSamples, EventArgs.IsGamepad());
 		}
 		
 		// Try the override callback, this may modify event args
@@ -746,15 +747,13 @@ bool UGameViewportClient::InputAxis(FViewport* InViewport, FInputDeviceId InputD
 		}
 
 #if WITH_EDITOR
-		const int32 NumLocalPlayers = World && World->GetGameInstance() ? World->GetGameInstance()->GetNumLocalPlayers() : 0;
-		// For PIE, let the next PIE window handle the input if none of our players did
-		// (this allows people to use multiple controllers to control each window)
-		if (!bResult && EventArgs.ControllerId > NumLocalPlayers - 1 && InViewport->IsPlayInEditorViewport())
+		// If there are multiple local players and the input has not been handled yet, we should pass it along
+		// to the next viewport so that it can check if any of its Player Controllers need it
+		if (!bResult && InViewport->IsPlayInEditorViewport())
 		{
-			UGameViewportClient *NextViewport = GEngine->GetNextPIEViewport(this);
-			if (NextViewport)
+			if (UGameViewportClient* NextViewport = GEngine->GetNextPIEViewport(this))
 			{
-				bResult = NextViewport->InputAxis(InViewport, EventArgs.ControllerId - NumLocalPlayers, EventArgs.Key, Delta, DeltaTime, NumSamples, EventArgs.IsGamepad());
+				bResult = NextViewport->InputAxis(InViewport, EventArgs.InputDevice, EventArgs.Key, Delta, DeltaTime, NumSamples, EventArgs.IsGamepad());
 			}
 		}
 #endif
@@ -777,7 +776,7 @@ bool UGameViewportClient::InputChar(FViewport* InViewport, int32 ControllerId, T
 	CharacterString += Character;
 
 	//Always route to the console
-	bool bResult = (ViewportConsole ? ViewportConsole->InputChar(ControllerId, CharacterString) : false);
+	bool bResult = (ViewportConsole ? ViewportConsole->InputChar(FInputDeviceId::CreateFromInternalId(ControllerId), CharacterString) : false);
 
 	if (IgnoreInput())
 	{
@@ -802,7 +801,7 @@ bool UGameViewportClient::InputTouch(FViewport* InViewport, int32 ControllerId,
 	}
 
 	// route to subsystems that care
-	bool bResult = (ViewportConsole ? ViewportConsole->InputTouch(ControllerId, Handle, Type, TouchLocation, Force, DeviceTimestamp, TouchpadIndex) : false);
+	bool bResult = (ViewportConsole ? ViewportConsole->InputTouch(FInputDeviceId::CreateFromInternalId(ControllerId), Handle, Type, TouchLocation, Force, DeviceTimestamp, TouchpadIndex) : false);
 	if (!bResult)
 	{
 		ULocalPlayer* const TargetPlayer = GEngine->GetLocalPlayerFromControllerId(this, ControllerId);
diff --git a/Engine/Source/Runtime/Engine/Private/PlayerController.cpp b/Engine/Source/Runtime/Engine/Private/PlayerController.cpp
--- a/Engine/Source/Runtime/Engine/Private/PlayerController.cpp
+++ b/Engine/Source/Runtime/Engine/Private/PlayerController.cpp
@@ -3748,7 +3748,7 @@ void APlayerController::ConsoleKey(FKey Key)
 	{
 		if (LocalPlayer->ViewportClient && LocalPlayer->ViewportClient->ViewportConsole)
 		{
-			LocalPlayer->ViewportClient->ViewportConsole->InputKey(0, Key, IE_Pressed);
+			LocalPlayer->ViewportClient->ViewportConsole->InputKey(IPlatformInputDeviceMapper::Get().GetDefaultInputDevice(), Key, IE_Pressed);
 		}
 	}
 #endif // ALLOW_CONSOLE
diff --git a/Engine/Source/Runtime/Engine/Private/UserInterface/Console.cpp b/Engine/Source/Runtime/Engine/Private/UserInterface/Console.cpp
--- a/Engine/Source/Runtime/Engine/Private/UserInterface/Console.cpp
+++ b/Engine/Source/Runtime/Engine/Private/UserInterface/Console.cpp
@@ -799,7 +799,15 @@ void UConsole::AppendInputText(const FString& Text)
 	UpdatePrecompletedInputLine();
 }
 
+PRAGMA_DISABLE_DEPRECATION_WARNINGS
 bool UConsole::InputChar_Typing(int32 ControllerId, const FString& Unicode)
+{
+	FInputDeviceId DeviceId = FInputDeviceId::CreateFromInternalId(ControllerId);
+	return InputChar_Typing(DeviceId, Unicode);
+}
+PRAGMA_ENABLE_DEPRECATION_WARNINGS
+
+bool UConsole::InputChar_Typing(FInputDeviceId DeviceId, const FString& Unicode)
 {
 	if (bCaptureKeyInput)
 	{
@@ -811,7 +819,7 @@ bool UConsole::InputChar_Typing(int32 ControllerId, const FString& Unicode)
 	return true;
 }
 
-bool UConsole::InputKey_InputLine(int32 ControllerId, FKey Key, EInputEvent Event, float AmountDepressed, bool bGamepad)
+bool UConsole::InputKey_InputLine(FInputDeviceId DeviceId, FKey Key, EInputEvent Event, float AmountDepressed, bool bGamepad)
 {
 	if (Event == IE_Pressed)
 	{
@@ -1194,8 +1202,15 @@ void UConsole::EndState_Typing(FName NextStateName)
 	bAutoCompleteLocked = false;
 }
 
-
+PRAGMA_DISABLE_DEPRECATION_WARNINGS
 bool UConsole::InputChar_Open(int32 ControllerId, const FString& Unicode)
+{
+	FInputDeviceId DeviceId = FInputDeviceId::CreateFromInternalId(ControllerId);
+	return InputChar_Open(DeviceId, Unicode);
+}
+PRAGMA_ENABLE_DEPRECATION_WARNINGS
+
+bool UConsole::InputChar_Open(FInputDeviceId DeviceId, const FString& Unicode)
 {
 	if (bCaptureKeyInput)
 	{
@@ -1207,8 +1222,15 @@ bool UConsole::InputChar_Open(int32 ControllerId, const FString& Unicode)
 	return true;
 }
 
-
+PRAGMA_DISABLE_DEPRECATION_WARNINGS
 bool UConsole::InputKey_Open(int32 ControllerId, FKey Key, EInputEvent Event, float AmountDepressed, bool bGamepad)
+{
+	FInputDeviceId DeviceId = FInputDeviceId::CreateFromInternalId(ControllerId);
+	return InputKey_Open(DeviceId, Key, Event, AmountDepressed, bGamepad);
+}
+PRAGMA_ENABLE_DEPRECATION_WARNINGS
+
+bool UConsole::InputKey_Open(FInputDeviceId DeviceId, FKey Key, EInputEvent Event, float AmountDepressed, bool bGamepad)
 {
 	if (Key == EKeys::PageUp || Key == EKeys::MouseScrollUp)
 	{
@@ -1346,22 +1368,38 @@ void UConsole::EndState_Open(FName NextStateName)
 {
 }
 
+PRAGMA_DISABLE_DEPRECATION_WARNINGS
 bool UConsole::InputChar(int32 ControllerId, const FString& Unicode)
+{
+	FInputDeviceId DeviceId = FInputDeviceId::CreateFromInternalId(ControllerId);
+	return InputChar(DeviceId, Unicode);
+}
+PRAGMA_ENABLE_DEPRECATION_WARNINGS
+
+bool UConsole::InputChar(FInputDeviceId DeviceId, const FString& Unicode)
 {
 	if (ConsoleState == NAME_Typing)
 	{
-		return InputChar_Typing(ControllerId, Unicode);
+		return InputChar_Typing(DeviceId, Unicode);
 	}
 	if (ConsoleState == NAME_Open)
 	{
-		return InputChar_Open(ControllerId, Unicode);
+		return InputChar_Open(DeviceId, Unicode);
 	}
 	return bCaptureKeyInput;
 }
 
+PRAGMA_DISABLE_DEPRECATION_WARNINGS
 bool UConsole::InputKey(int32 ControllerId, FKey Key, EInputEvent Event, float AmountDepressed, bool bGamepad)
 {
-	bool bWasConsumed = InputKey_InputLine(ControllerId, Key, Event, AmountDepressed, bGamepad);
+	FInputDeviceId DeviceId = FInputDeviceId::CreateFromInternalId(ControllerId);
+	return InputKey(DeviceId, Key, Event, AmountDepressed, bGamepad);
+}
+PRAGMA_ENABLE_DEPRECATION_WARNINGS
+
+bool UConsole::InputKey(FInputDeviceId DeviceId, FKey Key, EInputEvent Event, float AmountDepressed, bool bGamepad)
+{
+	bool bWasConsumed = InputKey_InputLine(DeviceId, Key, Event, AmountDepressed, bGamepad);
 
 	if (!bWasConsumed)
 	{
@@ -1372,7 +1410,7 @@ bool UConsole::InputKey(int32 ControllerId, FKey Key, EInputEvent Event, float A
 		}
 		if (ConsoleState == NAME_Open)
 		{
-			bWasConsumed = InputKey_Open(ControllerId, Key, Event, AmountDepressed, bGamepad);
+			bWasConsumed = InputKey_Open(DeviceId, Key, Event, AmountDepressed, bGamepad);
 			// if the console is open we don't want any other one to consume the input
 			return true;
 		}
