Various important camera shake fixes:

- Correctly teardown camera shakes that have been forcibly stopped.
- Correctly teardown the sequencer shake pattern wrapped inside the old legacy matinee shake class: we were previously leaking ECS data for sequencer shakes because of this. Although none of it was actually run, it was taking up space, possibly running unnecessary instantiation phases, and after a long while it could force the ECS eval to go into threaded mode.
- Don't clear all of the linker when we finish a sequencer shake: they're now all sharing the same linker, so other ongoing shakes would be blown away

diff --git a/Engine/Plugins/Cameras/GameplayCameras/Source/GameplayCameras/Private/MatineeCameraShake.cpp b/Engine/Plugins/Cameras/GameplayCameras/Source/GameplayCameras/Private/MatineeCameraShake.cpp
--- a/Engine/Plugins/Cameras/GameplayCameras/Source/GameplayCameras/Private/MatineeCameraShake.cpp
+++ b/Engine/Plugins/Cameras/GameplayCameras/Source/GameplayCameras/Private/MatineeCameraShake.cpp
@@ -129,6 +129,13 @@ void UMatineeCameraShake::DoStopShake(bool bImmediately)
 			}
 		}
 	}
+	
+	if (SequenceShakePattern)
+	{
+		FCameraShakeStopParams StopParams;
+		StopParams.bImmediately = bImmediately;
+		SequenceShakePattern->StopShakePattern(StopParams);
+	}
 
 	UE_LOG(LogMatineeCameraShake, Verbose, TEXT("UMatineeCameraShake::DoStopShake %s"), *GetNameSafe(this));
 
@@ -499,6 +506,14 @@ bool UMatineeCameraShake::DoGetIsFinished() const
 		);
 }
 
+void UMatineeCameraShake::DoTeardownShake()
+{
+	if (SequenceShakePattern)
+	{
+		SequenceShakePattern->TeardownShakePattern();
+	}
+}
+
 /// @cond DOXYGEN_WARNINGS
 
 bool UMatineeCameraShake::ReceiveIsFinished_Implementation() const
@@ -576,6 +591,12 @@ bool UMatineeCameraShakePattern::IsFinishedImpl() const
 	return Shake->DoGetIsFinished();
 }
 
+void UMatineeCameraShakePattern::TeardownShakePatternImpl()
+{
+	UMatineeCameraShake* Shake = GetShakeInstance<UMatineeCameraShake>();
+	Shake->DoTeardownShake();
+}
+
 UMovieSceneMatineeCameraShakeEvaluator::UMovieSceneMatineeCameraShakeEvaluator(const FObjectInitializer& ObjInit)
 	: Super(ObjInit)
 {
diff --git a/Engine/Plugins/Cameras/GameplayCameras/Source/GameplayCameras/Public/MatineeCameraShake.h b/Engine/Plugins/Cameras/GameplayCameras/Source/GameplayCameras/Public/MatineeCameraShake.h
--- a/Engine/Plugins/Cameras/GameplayCameras/Source/GameplayCameras/Public/MatineeCameraShake.h
+++ b/Engine/Plugins/Cameras/GameplayCameras/Source/GameplayCameras/Public/MatineeCameraShake.h
@@ -260,6 +260,7 @@ private:
 	void DoScrubShake(const FCameraShakeScrubParams& Params, FCameraShakeUpdateResult& OutResult);
 	void DoStopShake(bool bImmediately);
 	bool DoGetIsFinished() const;
+	void DoTeardownShake();
 
 protected:
 
@@ -322,6 +323,7 @@ private:
 	virtual void ScrubShakePatternImpl(const FCameraShakeScrubParams& Params, FCameraShakeUpdateResult& OutResult) override;
 	virtual bool IsFinishedImpl() const override;
 	virtual void StopShakePatternImpl(const FCameraShakeStopParams& Params) override;
+	virtual void TeardownShakePatternImpl() override;
 };
 
 /** Backwards compatible name for the Matinee camera shake, for C++ code. */
diff --git a/Engine/Plugins/MovieScene/TemplateSequence/Source/TemplateSequence/Private/SequenceCameraShake.cpp b/Engine/Plugins/MovieScene/TemplateSequence/Source/TemplateSequence/Private/SequenceCameraShake.cpp
--- a/Engine/Plugins/MovieScene/TemplateSequence/Source/TemplateSequence/Private/SequenceCameraShake.cpp
+++ b/Engine/Plugins/MovieScene/TemplateSequence/Source/TemplateSequence/Private/SequenceCameraShake.cpp
@@ -33,7 +33,7 @@ USequenceCameraShakePattern::USequenceCameraShakePattern(const FObjectInitialize
 	, bRandomSegment(false)
 {
 	CameraStandIn = CreateDefaultSubobject<UCameraAnimationSequenceCameraStandIn>(TEXT("CameraStandIn"), true);
-	Player = CreateDefaultSubobject<UCameraAnimationSequencePlayer>(TEXT("Player"), true);
+	Player = CreateDefaultSubobject<UCameraAnimationSequencePlayer>(TEXT("CameraShakePlayer"), true);
 
 	// Make sure we have our custom accessors registered for our stand-in class.
 	UCameraAnimationSequenceCameraStandIn::RegisterCameraStandIn();
@@ -132,12 +132,9 @@ void USequenceCameraShakePattern::TeardownShakePatternImpl()
 {
 	using namespace UE::MovieScene;
 
+	// Stop if we had reached the end of the animation and the sequence needs finishing.
+	// If the shake had been stopped explicitly, this basically won't do anything.
 	Player->Stop();
-	
-	if (UMovieSceneEntitySystemLinker* Linker = Player->GetEvaluationTemplate().GetEntitySystemLinker())
-	{
-		Linker->Reset();
-	}
 }
 
 void USequenceCameraShakePattern::UpdateCamera(FFrameTime NewPosition, const FMinimalViewInfo& InPOV, FCameraShakeUpdateResult& OutResult)
diff --git a/Engine/Source/Runtime/Engine/Private/Camera/CameraModifier_CameraShake.cpp b/Engine/Source/Runtime/Engine/Private/Camera/CameraModifier_CameraShake.cpp
--- a/Engine/Source/Runtime/Engine/Private/Camera/CameraModifier_CameraShake.cpp
+++ b/Engine/Source/Runtime/Engine/Private/Camera/CameraModifier_CameraShake.cpp
@@ -135,6 +135,7 @@ UCameraShakeBase* UCameraModifier_CameraShake::AddCameraShake(TSubclassOf<UCamer
 						// reliably restart the existing shake and expect it to be the same as what the caller wants. 
 						// So we forcibly stop the existing shake immediately and will create a brand new one.
 						ShakeInst->StopShake(true);
+						ShakeInst->TeardownShake();
 						// Discard it right away so the spot is free in the active shakes array.
 						ShakeInfo.ShakeInstance = nullptr;
 					}
@@ -239,9 +240,9 @@ void UCameraModifier_CameraShake::RemoveCameraShake(UCameraShakeBase* ShakeInst,
 		if (ShakeInfo.ShakeInstance == ShakeInst)
 		{
 			ShakeInst->StopShake(bImmediately);
-
 			if (bImmediately)
 			{
+				ShakeInst->TeardownShake();
 				SaveShakeInExpiredPoolIfPossible(ShakeInfo);
 				ActiveShakes.RemoveAt(i, 1);
 				UE_LOG(LogCameraShake, Verbose, TEXT("UCameraModifier_CameraShake::RemoveCameraShake %s"), *GetNameSafe(ShakeInfo.ShakeInstance));
@@ -262,6 +263,7 @@ void UCameraModifier_CameraShake::RemoveAllCameraShakesOfClass(TSubclassOf<UCame
 			ShakeInst->StopShake(bImmediately);
 			if (bImmediately)
 			{
+				ShakeInst->TeardownShake();
 				SaveShakeInExpiredPoolIfPossible(ShakeInfo);
 				ActiveShakes.RemoveAt(i, 1);
 				UE_LOG(LogCameraShake, Verbose, TEXT("UCameraModifier_CameraShake::RemoveAllCameraShakesOfClass %s"), *GetNameSafe(ShakeInfo.ShakeInstance));
@@ -280,6 +282,7 @@ void UCameraModifier_CameraShake::RemoveAllCameraShakesFromSource(const UCameraS
 			ShakeInfo.ShakeInstance->StopShake(bImmediately);
 			if (bImmediately)
 			{
+				ShakeInfo.ShakeInstance->TeardownShake();
 				SaveShakeInExpiredPoolIfPossible(ShakeInfo);
 				ActiveShakes.RemoveAt(i, 1);
 				UE_LOG(LogCameraShake, Verbose, TEXT("UCameraModifier_CameraShake::RemoveAllCameraShakesFromSource %s"), *GetNameSafe(ShakeInfo.ShakeInstance));
@@ -300,6 +303,7 @@ void UCameraModifier_CameraShake::RemoveAllCameraShakesOfClassFromSource(TSubcla
 			ShakeInfo.ShakeInstance->StopShake(bImmediately);
 			if (bImmediately)
 			{
+				ShakeInfo.ShakeInstance->TeardownShake();
 				SaveShakeInExpiredPoolIfPossible(ShakeInfo);
 				ActiveShakes.RemoveAt(i, 1);
 				UE_LOG(LogCameraShake, Verbose, TEXT("UCameraModifier_CameraShake::RemoveAllCameraShakesOfClassFromSource %s"), *GetNameSafe(ShakeInfo.ShakeInstance));
@@ -320,6 +324,7 @@ void UCameraModifier_CameraShake::RemoveAllCameraShakes(bool bImmediately)
 	{
 		for (FActiveCameraShakeInfo& ShakeInfo : ActiveShakes)
 		{
+			ShakeInfo.ShakeInstance->TeardownShake();
 			SaveShakeInExpiredPoolIfPossible(ShakeInfo);
 		}
 
