Remove some editor only define from material UStruct to re-enable proper Optional class UHT checks
Those inclusion should be minimal if they make their way in cooked data.

diff --git a/Engine/Source/Runtime/Engine/Classes/Materials/MaterialLayersFunctions.h b/Engine/Source/Runtime/Engine/Classes/Materials/MaterialLayersFunctions.h
--- a/Engine/Source/Runtime/Engine/Classes/Materials/MaterialLayersFunctions.h
+++ b/Engine/Source/Runtime/Engine/Classes/Materials/MaterialLayersFunctions.h
@@ -357,7 +357,6 @@ struct ENGINE_API FMaterialLayersFunctions : public FMaterialLayersFunctionsRunt
 		return !operator==(Other);
 	}
 
-#if WITH_EDITORONLY_DATA
 private:
 	UPROPERTY()
 	TArray<bool> LayerStates_DEPRECATED;
@@ -379,7 +378,6 @@ private:
 
 	UPROPERTY()
 	TArray<FGuid> DeletedParentLayerGuids_DEPRECATED;
-#endif // WITH_EDITORONLY_DATA
 };
 
 // Don't allowing comparing a full FMaterialLayersFunctions against partial RuntimeData
diff --git a/Engine/Source/Runtime/Engine/Public/StaticParameterSet.h b/Engine/Source/Runtime/Engine/Public/StaticParameterSet.h
--- a/Engine/Source/Runtime/Engine/Public/StaticParameterSet.h
+++ b/Engine/Source/Runtime/Engine/Public/StaticParameterSet.h
@@ -247,7 +247,6 @@ struct FStaticTerrainLayerWeightParameter
 {
 	GENERATED_USTRUCT_BODY();
 
-#if WITH_EDITORONLY_DATA
 	UPROPERTY()
 	FMaterialParameterInfo ParameterInfo_DEPRECATED;
 
@@ -256,7 +255,6 @@ struct FStaticTerrainLayerWeightParameter
 
 	UPROPERTY()
 	bool bOverride_DEPRECATED = true;
-#endif // WITH_EDITORONLY_DATA
 
 	UPROPERTY()
 	FName LayerName;
