[KillPendingKill] Invalidate replicators on the client as soon as their object becomes garbage.
Also clean up the replicator so it doesn't leak through AllOwnedReplicators.
Handles the case where subobjects are marked as garbage during the channel update, or where the ability system marks a replicated subobject as garbage on the client.

diff --git a/Engine/Source/Runtime/Engine/Private/DataChannel.cpp b/Engine/Source/Runtime/Engine/Private/DataChannel.cpp
--- a/Engine/Source/Runtime/Engine/Private/DataChannel.cpp
+++ b/Engine/Source/Runtime/Engine/Private/DataChannel.cpp
@@ -3121,11 +3121,11 @@ void UActorChannel::ProcessBunch( FInBunch & Bunch )
 	for (auto RepComp = ReplicationMap.CreateIterator(); RepComp; ++RepComp)
 	{
 		TSharedRef<FObjectReplicator>& ObjectReplicator = RepComp.Value();
-		if (ObjectReplicator->GetObject() == nullptr)
+		if (!IsValid(ObjectReplicator->GetObject()))
 		{
 			if (!Connection->Driver->IsServer())
 			{
-				RepComp.RemoveCurrent();
+				RepComp.RemoveCurrent(); // This should cause the replicator to be cleaned up as there should be no outstandings refs. 
 			}
 			continue;
 		}
