[KillPendingKill] Try to use an actor as the owner for new audio components in gameplay statics if the world context object is a component or belongs to an actor.

diff --git a/Engine/Source/Runtime/Engine/Private/GameplayStatics.cpp b/Engine/Source/Runtime/Engine/Private/GameplayStatics.cpp
--- a/Engine/Source/Runtime/Engine/Private/GameplayStatics.cpp
+++ b/Engine/Source/Runtime/Engine/Private/GameplayStatics.cpp
@@ -60,6 +60,26 @@
 
 #define LOCTEXT_NAMESPACE "GameplayStatics"
 
+namespace GameplayStatics
+{
+	AActor* GetActorOwnerFromWorldContextObject(UObject* WorldContextObject)
+	{
+		if (AActor* Actor = Cast<AActor>(WorldContextObject))
+		{
+			return Actor;
+		}
+		return WorldContextObject->GetTypedOuter<AActor>();
+	}
+	const AActor* GetActorOwnerFromWorldContextObject(const UObject* WorldContextObject)
+	{
+		if (const AActor* Actor = Cast<const AActor>(WorldContextObject))
+		{
+			return Actor;
+		}
+		return WorldContextObject->GetTypedOuter<AActor>();
+	}
+}
+
 static const int UE_SAVEGAME_FILE_TYPE_TAG = 0x53415647;		// "SAVG"
 
 struct FSaveGameFileVersion
@@ -1578,8 +1598,7 @@ void UGameplayStatics::PlaySound2D(const UObject* WorldContextObject, USoundBase
 		NewActiveSound.SubtitlePriority = Sound->GetSubtitlePriority();
 
 		// If OwningActor isn't supplied to this function, derive an owner from the WorldContextObject
-		const AActor* WorldContextOwner = Cast<const AActor>(WorldContextObject);
-		const AActor* ActiveSoundOwner = OwningActor ? OwningActor : WorldContextOwner;
+		const AActor* ActiveSoundOwner = OwningActor ? OwningActor : GameplayStatics::GetActorOwnerFromWorldContextObject(WorldContextObject);
 
 		NewActiveSound.SetOwner(ActiveSoundOwner);
 
@@ -1604,8 +1623,7 @@ UAudioComponent* UGameplayStatics::CreateSound2D(const UObject* WorldContextObje
 	}
 
 	// Derive an owner from the WorldContextObject
-	UObject* MutableWorldContext = const_cast<UObject*>(WorldContextObject);
-	AActor* WorldContextOwner = Cast<AActor>(MutableWorldContext);
+	AActor* WorldContextOwner = GameplayStatics::GetActorOwnerFromWorldContextObject(const_cast<UObject*>(WorldContextObject));
 
 	FAudioDevice::FCreateComponentParams Params = bPersistAcrossLevelTransition
 		? FAudioDevice::FCreateComponentParams(ThisWorld->GetAudioDeviceRaw())
@@ -1663,9 +1681,7 @@ void UGameplayStatics::PlaySoundAtLocation(const UObject* WorldContextObject, cl
 		}
 
 		// If OwningActor isn't supplied to this function, derive an owner from the WorldContextObject
-		const AActor* WorldContextOwner = Cast<const AActor>(WorldContextObject);
-		const AActor* ActiveSoundOwner = OwningActor ? OwningActor : WorldContextOwner;
-
+		const AActor* ActiveSoundOwner = OwningActor ? OwningActor : GameplayStatics::GetActorOwnerFromWorldContextObject(WorldContextObject);
 		UActorSoundParameterInterface::Fill(ActiveSoundOwner, Params);
 
 		AudioDevice->PlaySoundAtLocation(Sound, ThisWorld, VolumeMultiplier, PitchMultiplier, StartTime, Location, Rotation, AttenuationSettings, ConcurrencySettings, &Params, ActiveSoundOwner);
@@ -1688,8 +1704,7 @@ UAudioComponent* UGameplayStatics::SpawnSoundAtLocation(const UObject* WorldCont
 	const bool bIsInGameWorld = ThisWorld->IsGameWorld();
 
 	// Derive an owner from the WorldContextObject
-	UObject* MutableWorldContext = const_cast<UObject*>(WorldContextObject);
-	AActor* WorldContextOwner = Cast<AActor>(MutableWorldContext);
+	AActor* WorldContextOwner = GameplayStatics::GetActorOwnerFromWorldContextObject(const_cast<UObject*>(WorldContextObject));
 
 	FAudioDevice::FCreateComponentParams Params(ThisWorld, WorldContextOwner);
 	Params.SetLocation(Location);
