World Partition: fix for deferred loading context when loaded actors would load actors during the registration phase.

#rnx

diff --git a/Engine/Source/Runtime/Engine/Private/WorldPartition/WorldPartitionHandle.cpp b/Engine/Source/Runtime/Engine/Private/WorldPartition/WorldPartitionHandle.cpp
--- a/Engine/Source/Runtime/Engine/Private/WorldPartition/WorldPartitionHandle.cpp
+++ b/Engine/Source/Runtime/Engine/Private/WorldPartition/WorldPartitionHandle.cpp
@@ -100,56 +100,64 @@ void FWorldPartitionLoadingContext::FImmediate::UnregisterActor(FWorldPartitionA
 */
 FWorldPartitionLoadingContext::FDeferred::~FDeferred()
 {
-	for (auto& [Container, ContainerOp] : ContainerOps)
-	{
-		const FTransform& ContainerTransform = Container->GetInstanceTransform();
-		const FTransform* ContainerTransformPtr = ContainerTransform.Equals(FTransform::Identity) ? nullptr : &ContainerTransform;
+	TMap<UActorDescContainer*, FContainerOps> LocalContainerOps = MoveTemp(ContainerOps);
 
-		auto CreateActorList = [](TArray<AActor*>& ActorList, const TSet<FWorldPartitionActorDesc*>& SourceList) -> ULevel*
+	while(LocalContainerOps.Num())
+	{
+		for (auto& [Container, ContainerOp] : LocalContainerOps)
 		{
-			ActorList.Empty(SourceList.Num());
+			const FTransform& ContainerTransform = Container->GetInstanceTransform();
+			const FTransform* ContainerTransformPtr = ContainerTransform.Equals(FTransform::Identity) ? nullptr : &ContainerTransform;
 
-			ULevel* Level = nullptr;
-			for (FWorldPartitionActorDesc* ActorDesc : SourceList)
+			auto CreateActorList = [](TArray<AActor*>& ActorList, const TSet<FWorldPartitionActorDesc*>& SourceList) -> ULevel*
 			{
-				// When cleaning up worlds, actors are already marked as garbage at this point, so no need to remove them from the world
-				if (AActor* Actor = ActorDesc->GetActor(); IsValid(Actor))
-				{
-					ActorList.Add(Actor);
+				ActorList.Empty(SourceList.Num());
 
-					ULevel* ActorLevel = Actor->GetLevel();
-					check(!Level || (Level == ActorLevel));
-					Level = ActorLevel;
+				ULevel* Level = nullptr;
+				for (FWorldPartitionActorDesc* ActorDesc : SourceList)
+				{
+					// When cleaning up worlds, actors are already marked as garbage at this point, so no need to remove them from the world
+					if (AActor* Actor = ActorDesc->GetActor(); IsValid(Actor))
+					{
+						ActorList.Add(Actor);
+
+						ULevel* ActorLevel = Actor->GetLevel();
+						check(!Level || (Level == ActorLevel));
+						Level = ActorLevel;
+					}
 				}
-			}
-
-			return Level;
-		};
 
-		if (ContainerOp.Registrations.Num())
-		{
-			TGuardValue<bool> IsEditorLoadingPackageGuard(GIsEditorLoadingPackage, true);
+				return Level;
+			};
 
-			TArray<AActor*> ActorList;
-			if (ULevel* Level = CreateActorList(ActorList, ContainerOp.Registrations))
+			if (ContainerOp.Registrations.Num())
 			{
-				Level->AddLoadedActors(ActorList, ContainerTransformPtr);
-			}
-		}
+				TGuardValue<bool> IsEditorLoadingPackageGuard(GIsEditorLoadingPackage, true);
 
-		if (ContainerOp.Unregistrations.Num())
-		{
-			TArray<AActor*> ActorList;
-			if (ULevel* Level = CreateActorList(ActorList, ContainerOp.Unregistrations))
-			{
-				Level->RemoveLoadedActors(ActorList, ContainerTransformPtr);
+				TArray<AActor*> ActorList;
+				if (ULevel* Level = CreateActorList(ActorList, ContainerOp.Registrations))
+				{
+					Level->AddLoadedActors(ActorList, ContainerTransformPtr);
+				}
 			}
 
-			for (FWorldPartitionActorDesc* ActorDesc : ContainerOp.Unregistrations)
+			if (ContainerOp.Unregistrations.Num())
 			{
-				ActorDesc->Unload();
+				TArray<AActor*> ActorList;
+				if (ULevel* Level = CreateActorList(ActorList, ContainerOp.Unregistrations))
+				{
+					Level->RemoveLoadedActors(ActorList, ContainerTransformPtr);
+				}
+
+				for (FWorldPartitionActorDesc* ActorDesc : ContainerOp.Unregistrations)
+				{
+					ActorDesc->Unload();
+				}
 			}
 		}
+
+		// Continue with potentially new registrations/unregistrations that may have happenned during previous cycle
+		LocalContainerOps = MoveTemp(ContainerOps);
 	}
 }
 
