Added missing isValid check in hitresult ToString method

diff --git a/Engine/Source/Runtime/Engine/Private/Engine/HitResult.cpp b/Engine/Source/Runtime/Engine/Private/Engine/HitResult.cpp
--- a/Engine/Source/Runtime/Engine/Private/Engine/HitResult.cpp
+++ b/Engine/Source/Runtime/Engine/Private/Engine/HitResult.cpp
@@ -156,7 +156,7 @@ FString FHitResult::ToString() const
 		PenetrationDepth,
 		Item,
 		PhysMaterial.IsValid() ? *PhysMaterial->GetName() : TEXT("None"),
-		*FLightWeightInstanceSubsystem::Get().GetName(HitObjectHandle),
+		HitObjectHandle.IsValid() ? *FLightWeightInstanceSubsystem::Get().GetName(HitObjectHandle) : TEXT("None"),
 		Component.IsValid() ? *Component->GetName() : TEXT("None"),
 		BoneName.IsValid() ? *BoneName.ToString() : TEXT("None"),
 		FaceIndex);
