Animation Data Model should not be flagged as StandAlone, as it prevents deleting its outer (Animation Sequence)

diff --git a/Engine/Source/Runtime/Engine/Private/Animation/AnimSequenceBase.cpp b/Engine/Source/Runtime/Engine/Private/Animation/AnimSequenceBase.cpp
--- a/Engine/Source/Runtime/Engine/Private/Animation/AnimSequenceBase.cpp
+++ b/Engine/Source/Runtime/Engine/Private/Animation/AnimSequenceBase.cpp
@@ -1384,7 +1384,7 @@ void UAnimSequenceBase::CreateModel()
 	checkf(TargetClass != nullptr, TEXT("Unable to find valid AnimationDataModel class"));
 
 	checkf(!DataModelInterface || DataModelInterface.GetObject()->GetClass() != TargetClass, TEXT("Invalid attempt to override the existing data model %s"), *GetFullName());
-	UObject* ClassDataModel = NewObject<UObject>(this, TargetClass, TargetClass->GetFName(), RF_Standalone);
+	UObject* ClassDataModel = NewObject<UObject>(this, TargetClass, TargetClass->GetFName());
 	DataModelInterface = ClassDataModel;
 
 	BindToModelModificationEvent();
