HDR support for XR devices that support HDR but are connected to a SDR TV

diff --git a/Engine/Plugins/Runtime/OpenXR/Source/OpenXRHMD/Private/OpenXRHMD.cpp b/Engine/Plugins/Runtime/OpenXR/Source/OpenXRHMD/Private/OpenXRHMD.cpp
--- a/Engine/Plugins/Runtime/OpenXR/Source/OpenXRHMD/Private/OpenXRHMD.cpp
+++ b/Engine/Plugins/Runtime/OpenXR/Source/OpenXRHMD/Private/OpenXRHMD.cpp
@@ -2563,6 +2563,16 @@ FXRRenderBridge* FOpenXRHMD::GetActiveRenderBridge_GameThread(bool /* bUseSepara
 	return RenderBridge;
 }
 
+bool FOpenXRHMD::HDRGetMetaDataForStereo(EDisplayOutputFormat& OutDisplayOutputFormat, EDisplayColorGamut& OutDisplayColorGamut, bool& OutbHDRSupported)
+{
+	if (RenderBridge == nullptr)
+	{
+		return false;
+	}
+
+	return RenderBridge->HDRGetMetaDataForStereo(OutDisplayOutputFormat, OutDisplayColorGamut, OutbHDRSupported);
+}
+
 FIntPoint FOpenXRHMD::GetIdealRenderTargetSize() const
 {
 	const FPipelinedFrameState& PipelineState = GetPipelinedFrameStateForThread();
diff --git a/Engine/Plugins/Runtime/OpenXR/Source/OpenXRHMD/Private/OpenXRHMD.h b/Engine/Plugins/Runtime/OpenXR/Source/OpenXRHMD/Private/OpenXRHMD.h
--- a/Engine/Plugins/Runtime/OpenXR/Source/OpenXRHMD/Private/OpenXRHMD.h
+++ b/Engine/Plugins/Runtime/OpenXR/Source/OpenXRHMD/Private/OpenXRHMD.h
@@ -200,6 +200,8 @@ public:
 
 	virtual FVector2D GetPlayAreaBounds(EHMDTrackingOrigin::Type Origin) const override;
 
+	virtual bool HDRGetMetaDataForStereo(EDisplayOutputFormat& OutDisplayOutputFormat, EDisplayColorGamut& OutDisplayColorGamut, bool& OutbHDRSupported) override;
+
 protected:
 
 	bool StartSession();
diff --git a/Engine/Plugins/Runtime/OpenXR/Source/OpenXRHMD/Private/OpenXRHMD_RenderBridge.h b/Engine/Plugins/Runtime/OpenXR/Source/OpenXRHMD/Private/OpenXRHMD_RenderBridge.h
--- a/Engine/Plugins/Runtime/OpenXR/Source/OpenXRHMD/Private/OpenXRHMD_RenderBridge.h
+++ b/Engine/Plugins/Runtime/OpenXR/Source/OpenXRHMD/Private/OpenXRHMD_RenderBridge.h
@@ -50,6 +50,8 @@ public:
 
 	virtual bool Support10BitSwapchain() const { return false; }
 
+	virtual bool HDRGetMetaDataForStereo(EDisplayOutputFormat& OutDisplayOutputFormat, EDisplayColorGamut& OutDisplayColorGamut, bool& OutbHDRSupported) { return false; }
+
 protected:
 	uint64 AdapterLuid;
 
diff --git a/Engine/Source/Runtime/Engine/Private/Slate/SceneViewport.cpp b/Engine/Source/Runtime/Engine/Private/Slate/SceneViewport.cpp
--- a/Engine/Source/Runtime/Engine/Private/Slate/SceneViewport.cpp
+++ b/Engine/Source/Runtime/Engine/Private/Slate/SceneViewport.cpp
@@ -356,6 +356,29 @@ FIntPoint FSceneViewport::ViewportToVirtualDesktopPixel(FVector2D ViewportCoordi
 	return FIntPoint( FMath::TruncToInt(TransformedPoint.X / CachedGeometry.Scale), FMath::TruncToInt(TransformedPoint.Y / CachedGeometry.Scale) );
 }
 
+IStereoRenderTargetManager* RetrieveStereoRenderTargetManager(bool bIsStereoRenderingAllowed)
+{
+	return (bIsStereoRenderingAllowed && GEngine->StereoRenderingDevice.IsValid() && GEngine->StereoRenderingDevice->IsStereoEnabledOnNextFrame())
+		   ? GEngine->StereoRenderingDevice->GetRenderTargetManager()
+		   : nullptr;
+
+}
+
+void ComputeSceneViewportHDRMetaData(EDisplayOutputFormat& OutDisplayOutputFormat, EDisplayColorGamut& OutDisplayColorGamut, bool& OutbHDRSupported, const FVector2D& WindowTopLeft, const FVector2D& WindowBottomRight, void* OSWindow, bool bIsStereoRenderingAllowed)
+{
+	// @todo vreditor switch: This code needs to be called when switching between stereo/non when going immersive.  Seems to always work out that way anyway though? (Probably due to resize)
+	IStereoRenderTargetManager* const StereoRenderTargetManager = RetrieveStereoRenderTargetManager(bIsStereoRenderingAllowed);
+	if (StereoRenderTargetManager != nullptr)
+	{
+		if (StereoRenderTargetManager->HDRGetMetaDataForStereo(OutDisplayOutputFormat, OutDisplayColorGamut, OutbHDRSupported))
+		{
+			return;
+		}
+	}
+	
+	HDRGetMetaData(OutDisplayOutputFormat, OutDisplayColorGamut, OutbHDRSupported, WindowTopLeft, WindowBottomRight, OSWindow);
+}
+
 void FSceneViewport::OnDrawViewport( const FGeometry& AllottedGeometry, const FSlateRect& MyCullingRect, FSlateWindowElementList& OutDrawElements, int32 LayerId, const FWidgetStyle& InWidgetStyle, bool bParentEnabled )
 {
 	// Switch to the viewport clients world before resizing
@@ -374,7 +397,7 @@ void FSceneViewport::OnDrawViewport( const FGeometry& AllottedGeometry, const FS
 			bool bNewHDREnabled;
 			EDisplayColorGamut NewDisplayColorGamut;
 			EDisplayOutputFormat NewDisplayOutputFormat;
-			HDRGetMetaData(NewDisplayOutputFormat, NewDisplayColorGamut, bNewHDREnabled, PaintWindow->GetPositionInScreen(), PaintWindow->GetPositionInScreen() + PaintWindow->GetSizeInScreen(), PaintWindow->GetNativeWindow()->GetOSWindowHandle());
+			ComputeSceneViewportHDRMetaData(NewDisplayOutputFormat, NewDisplayColorGamut, bNewHDREnabled, PaintWindow->GetPositionInScreen(), PaintWindow->GetPositionInScreen() + PaintWindow->GetSizeInScreen(), PaintWindow->GetNativeWindow()->GetOSWindowHandle(), IsStereoRenderingAllowed());
 			// if we manage to get data for the window, we can ignore the global toggle IsHDREnabled since HDRGetMetaData will take both the global flag and the monitor properties
 			bHDRStale = DisplayOutputFormat != NewDisplayOutputFormat;
 			bHDRStale |= DisplayColorGamut != NewDisplayColorGamut;
@@ -1245,6 +1268,11 @@ EDisplayOutputFormat FSceneViewport::GetDisplayOutputFormat() const
 	return DisplayOutputFormat;
 }
 
+bool FSceneViewport::GetSceneHDREnabled() const
+{
+	return bHDRViewport;
+}
+
 void FSceneViewport::OnViewportDeactivated(const FWindowActivateEvent& InActivateEvent)
 {
 	// We backup if we have capture for us on activation, however we also maintain "true" if it's already true!
@@ -1695,7 +1723,7 @@ void FSceneViewport::UpdateViewportRHI(bool bDestroyed, uint32 NewSizeX, uint32
 					WindowBottomRight = Window->GetPositionInScreen() + Window->GetSizeInScreen();
 				}
 			}
-			HDRGetMetaData(DisplayOutputFormat, DisplayColorGamut, bHDRViewport, WindowTopLeft, WindowBottomRight, OSWindow);
+			ComputeSceneViewportHDRMetaData(DisplayOutputFormat, DisplayColorGamut, bHDRViewport, WindowTopLeft, WindowBottomRight, OSWindow, IsStereoRenderingAllowed());
 
 			BeginInitResource(this);
 			
@@ -1714,10 +1742,7 @@ void FSceneViewport::UpdateViewportRHI(bool bDestroyed, uint32 NewSizeX, uint32
 				{
 					uint32 TexSizeX = SizeX, TexSizeY = SizeY;
 					{
-						IStereoRenderTargetManager* const StereoRenderTargetManager =
-							(IsStereoRenderingAllowed() && GEngine->StereoRenderingDevice.IsValid() && GEngine->StereoRenderingDevice->IsStereoEnabledOnNextFrame())
-							? GEngine->StereoRenderingDevice->GetRenderTargetManager()
-							: nullptr;
+						IStereoRenderTargetManager* const StereoRenderTargetManager = RetrieveStereoRenderTargetManager(IsStereoRenderingAllowed());
 						if (StereoRenderTargetManager)
 						{
 							StereoRenderTargetManager->CalculateRenderTargetSize(*this, TexSizeX, TexSizeY);
@@ -2028,11 +2053,7 @@ void FSceneViewport::InitDynamicRHI()
 		NumBufferedFrames = 1;
 		
 		// @todo vreditor switch: This code needs to be called when switching between stereo/non when going immersive.  Seems to always work out that way anyway though? (Probably due to resize)
-		IStereoRenderTargetManager * const StereoRenderTargetManager = 
-			(IsStereoRenderingAllowed() && GEngine->StereoRenderingDevice.IsValid() && GEngine->StereoRenderingDevice->IsStereoEnabledOnNextFrame())
-				? GEngine->StereoRenderingDevice->GetRenderTargetManager() 
-				: nullptr;
-
+		IStereoRenderTargetManager* const StereoRenderTargetManager = RetrieveStereoRenderTargetManager(IsStereoRenderingAllowed());
 		if (StereoRenderTargetManager != nullptr)
 		{
 			StereoRenderTargetManager->CalculateRenderTargetSize(*this, TexSizeX, TexSizeY);
diff --git a/Engine/Source/Runtime/Engine/Public/Slate/SceneViewport.h b/Engine/Source/Runtime/Engine/Public/Slate/SceneViewport.h
--- a/Engine/Source/Runtime/Engine/Public/Slate/SceneViewport.h
+++ b/Engine/Source/Runtime/Engine/Public/Slate/SceneViewport.h
@@ -244,6 +244,7 @@ public:
 	virtual FIntPoint GetSize() const override { return GetSizeXY(); }
 	virtual EDisplayColorGamut GetDisplayColorGamut() const override;
 	virtual EDisplayOutputFormat GetDisplayOutputFormat() const override;
+	virtual bool GetSceneHDREnabled() const override;
 
 	void SetViewportSize(uint32 NewSizeX,uint32 NewSizeY);
 	void SetFixedViewportSize(uint32 NewSizeX, uint32 NewSizeY);
diff --git a/Engine/Source/Runtime/Engine/Public/StereoRenderTargetManager.h b/Engine/Source/Runtime/Engine/Public/StereoRenderTargetManager.h
--- a/Engine/Source/Runtime/Engine/Public/StereoRenderTargetManager.h
+++ b/Engine/Source/Runtime/Engine/Public/StereoRenderTargetManager.h
@@ -84,5 +84,16 @@ public:
 	 */
 	virtual bool AllocateShadingRateTexture(uint32 Index, uint32 RenderSizeX, uint32 RenderSizeY, uint8 Format, uint32 NumMips, ETextureCreateFlags Flags, ETextureCreateFlags TargetableTextureFlags, FTexture2DRHIRef& OutTexture, FIntPoint& OutTextureSize) { return false; }
 
+	/**
+	 * Retrieves HDR information about the stereo device, if any is available.
+	 * The default implementation always returns false to indicate that the information from the monitor can be used instead.
+	 *
+	 * @param OutDisplayOutputFormat	(out) encoding used by the stereo device
+	 * @param OutDisplayColorGamut		(out) color space used by the stereo device
+	 * @param OutbHDRSupported			(out) whether HDR is supported by the stereo device
+	 * @return							true, if HDR information is available for the stereo device
+	 */
+	virtual bool HDRGetMetaDataForStereo(EDisplayOutputFormat& OutDisplayOutputFormat, EDisplayColorGamut& OutDisplayColorGamut, bool& OutbHDRSupported) { return false; }
+
 	static EPixelFormat GetStereoLayerPixelFormat() { return PF_B8G8R8A8; }
 };
\ No newline at end of file
diff --git a/Engine/Source/Runtime/Engine/Public/UnrealClient.h b/Engine/Source/Runtime/Engine/Public/UnrealClient.h
--- a/Engine/Source/Runtime/Engine/Public/UnrealClient.h
+++ b/Engine/Source/Runtime/Engine/Public/UnrealClient.h
@@ -63,6 +63,7 @@ public:
 
 	virtual EDisplayColorGamut GetDisplayColorGamut() const { return EDisplayColorGamut::sRGB_D65; }
 	virtual EDisplayOutputFormat GetDisplayOutputFormat() const { return EDisplayOutputFormat::SDR_sRGB; }
+	virtual bool GetSceneHDREnabled() const { return false; }
 
 	/**
 	* Handles freezing/unfreezing of rendering
diff --git a/Engine/Source/Runtime/Renderer/Private/PostProcess/PostProcessDeviceEncodingOnly.cpp b/Engine/Source/Runtime/Renderer/Private/PostProcess/PostProcessDeviceEncodingOnly.cpp
--- a/Engine/Source/Runtime/Renderer/Private/PostProcess/PostProcessDeviceEncodingOnly.cpp
+++ b/Engine/Source/Runtime/Renderer/Private/PostProcess/PostProcessDeviceEncodingOnly.cpp
@@ -57,7 +57,8 @@ FDeviceEncodingOnlyOutputDeviceParameters GetDeviceEncodingOnlyOutputDeviceParam
 
 	float Gamma = CVarOutputGamma->GetValueOnRenderThread();
 
-	if (PLATFORM_APPLE && Gamma == 0.0f)
+	// In case gamma is unspecified, fall back to 2.2 which is the most common case
+	if ((PLATFORM_APPLE || OutputDeviceValue == EDisplayOutputFormat::SDR_ExplicitGammaMapping) && Gamma == 0.0f)
 	{
 		Gamma = 2.2f;
 	}
diff --git a/Engine/Source/Runtime/Renderer/Private/PostProcess/PostProcessTonemap.cpp b/Engine/Source/Runtime/Renderer/Private/PostProcess/PostProcessTonemap.cpp
--- a/Engine/Source/Runtime/Renderer/Private/PostProcess/PostProcessTonemap.cpp
+++ b/Engine/Source/Runtime/Renderer/Private/PostProcess/PostProcessTonemap.cpp
@@ -247,7 +247,8 @@ FTonemapperOutputDeviceParameters GetTonemapperOutputDeviceParameters(const FSce
 
 	float Gamma = CVarOutputGamma->GetValueOnRenderThread();
 
-	if (PLATFORM_APPLE && Gamma == 0.0f)
+    // In case gamma is unspecified, fall back to 2.2 which is the most common case
+	if ((PLATFORM_APPLE || OutputDeviceValue == EDisplayOutputFormat::SDR_ExplicitGammaMapping) && Gamma == 0.0f)
 	{
 		Gamma = 2.2f;
 	}
diff --git a/Engine/Source/Runtime/Renderer/Private/PostProcess/PostProcessing.cpp b/Engine/Source/Runtime/Renderer/Private/PostProcess/PostProcessing.cpp
--- a/Engine/Source/Runtime/Renderer/Private/PostProcess/PostProcessing.cpp
+++ b/Engine/Source/Runtime/Renderer/Private/PostProcess/PostProcessing.cpp
@@ -246,7 +246,7 @@ void AddPostProcessingPasses(
 
 	const FEngineShowFlags& EngineShowFlags = View.Family->EngineShowFlags;
 	const bool bVisualizeHDR = EngineShowFlags.VisualizeHDR;
-	const bool bViewFamilyOutputInHDR = GRHISupportsHDROutput && IsHDREnabled();
+	const bool bViewFamilyOutputInHDR = View.Family->RenderTarget->GetSceneHDREnabled();
 	const bool bVisualizeGBufferOverview = IsVisualizeGBufferOverviewEnabled(View);
 	const bool bVisualizeGBufferDumpToFile = IsVisualizeGBufferDumpToFileEnabled(View);
 	const bool bVisualizeGBufferDumpToPIpe = IsVisualizeGBufferDumpToPipeEnabled(View);
@@ -1439,7 +1439,7 @@ void AddDebugViewPostProcessingPasses(FRDGBuilder& GraphBuilder, const FViewInfo
 	// Some view modes do not actually output a color so they should not be tonemapped.
 	const bool bTonemapAfter = View.Family->EngineShowFlags.RayTracingDebug || View.Family->EngineShowFlags.VisualizeGPUSkinCache;
 	const bool bTonemapBefore = !bTonemapAfter && !View.Family->EngineShowFlags.ShaderComplexity;
-	const bool bViewFamilyOutputInHDR = GRHISupportsHDROutput && IsHDREnabled();
+	const bool bViewFamilyOutputInHDR = View.Family->RenderTarget->GetSceneHDREnabled();
 
 	enum class EPass : uint32
 	{
@@ -1679,7 +1679,7 @@ void AddVisualizeCalibrationMaterialPostProcessingPasses(FRDGBuilder& GraphBuild
 
 	const FEngineShowFlags& EngineShowFlags = View.Family->EngineShowFlags;
 	const bool bVisualizeHDR = EngineShowFlags.VisualizeHDR;
-	const bool bViewFamilyOutputInHDR = GRHISupportsHDROutput && IsHDREnabled();
+	const bool bViewFamilyOutputInHDR = View.Family->RenderTarget->GetSceneHDREnabled();
 	const bool bOutputInHDR = IsPostProcessingOutputInHDR();
 
 	// Post Process Material - Before Color Correction
diff --git a/Engine/Source/Runtime/SlateRHIRenderer/Private/SlateRHIRenderer.cpp b/Engine/Source/Runtime/SlateRHIRenderer/Private/SlateRHIRenderer.cpp
--- a/Engine/Source/Runtime/SlateRHIRenderer/Private/SlateRHIRenderer.cpp
+++ b/Engine/Source/Runtime/SlateRHIRenderer/Private/SlateRHIRenderer.cpp
@@ -593,7 +593,8 @@ public:
 		int32 OutputGamutValue = (int32)DisplayColorGamut;
 		float Gamma = CVarOutputGamma->GetValueOnRenderThread();
 
-		if (PLATFORM_APPLE && Gamma == 0.0f)
+        // In case gamma is unspecified, fall back to 2.2 which is the most common case
+		if ((PLATFORM_APPLE || DisplayOutputFormat == EDisplayOutputFormat::SDR_ExplicitGammaMapping) && Gamma == 0.0f)
 		{
 			Gamma = 2.2f;
 		}
