Adding command line parameters to specify a specific shader model and material quality level to cook for.

This is used to limit what cooking shadermaps to a single shader model and/or material quality level.  Useful for locally cooking or specific situations where you want to cook a minimal set of shader maps.  Each command line parameter can be used indpendently.

- Adding -CacheShaderFormat.  When specifying this during a cook the cooker will only compile this shader model if supported for each material.
- Adding -CacheMaterialQuality.  When specifying this during a cook the cooker will only compile this quality level and default for each material.

# Usage
* -CacheShaderFormat=PCD3D_SM5
* -CacheMaterialQuality=[0,1,2,3]

diff --git a/Engine/Source/Runtime/Engine/Private/Materials/Material.cpp b/Engine/Source/Runtime/Engine/Private/Materials/Material.cpp
--- a/Engine/Source/Runtime/Engine/Private/Materials/Material.cpp
+++ b/Engine/Source/Runtime/Engine/Private/Materials/Material.cpp
@@ -3831,6 +3831,8 @@ void UMaterial::BeginCacheForCookedPlatformData( const ITargetPlatform *TargetPl
 	TArray<FName> DesiredShaderFormats;
 	TargetPlatform->GetAllTargetedShaderFormats(DesiredShaderFormats);
 
+	GetCmdLineFilterShaderFormats(DesiredShaderFormats);
+
 	TArray<FMaterialResource*> *CachedMaterialResourcesForPlatform = CachedMaterialResourcesForCooking.Find( TargetPlatform );
 
 	if ( CachedMaterialResourcesForPlatform == NULL )
diff --git a/Engine/Source/Runtime/Engine/Private/Materials/MaterialInstance.cpp b/Engine/Source/Runtime/Engine/Private/Materials/MaterialInstance.cpp
--- a/Engine/Source/Runtime/Engine/Private/Materials/MaterialInstance.cpp
+++ b/Engine/Source/Runtime/Engine/Private/Materials/MaterialInstance.cpp
@@ -2407,6 +2407,8 @@ void UMaterialInstance::BeginCacheForCookedPlatformData( const ITargetPlatform *
 		TArray<FName> DesiredShaderFormats;
 		TargetPlatform->GetAllTargetedShaderFormats(DesiredShaderFormats);
 
+		GetCmdLineFilterShaderFormats(DesiredShaderFormats);
+
 		// Cache shaders for each shader format, storing the results in CachedMaterialResourcesForCooking so they will be available during saving
 		for (int32 FormatIndex = 0; FormatIndex < DesiredShaderFormats.Num(); FormatIndex++)
 		{
diff --git a/Engine/Source/Runtime/Engine/Private/Materials/MaterialInterface.cpp b/Engine/Source/Runtime/Engine/Private/Materials/MaterialInterface.cpp
--- a/Engine/Source/Runtime/Engine/Private/Materials/MaterialInterface.cpp
+++ b/Engine/Source/Runtime/Engine/Private/Materials/MaterialInterface.cpp
@@ -272,6 +272,21 @@ void UMaterialInterface::GetQualityLevelUsage(TArray<bool, TInlineAllocator<EMat
 			}
 		}
 	}
+
+#if WITH_EDITOR
+	// if we specified -CacheMaterialQuality= on the cmd line we should only cook that quality.
+	static const int32 MaterialQualityLevel = GetCmdLineMaterialQualityToCache();
+	if (MaterialQualityLevel != INDEX_NONE)
+	{
+		// the format is only valid if it is a desired format for this platform.
+		if (OutQualityLevelsUsed.IsValidIndex(MaterialQualityLevel))
+		{
+			// only cache the format specified on the command line.
+			OutQualityLevelsUsed.SetNumZeroed(OutQualityLevelsUsed.Num());
+			OutQualityLevelsUsed[MaterialQualityLevel] = true;
+		}
+	}
+#endif
 }
 
 TArrayView<const TObjectPtr<UObject>> UMaterialInterface::GetReferencedTextures() const
diff --git a/Engine/Source/Runtime/Engine/Private/Materials/MaterialShared.cpp b/Engine/Source/Runtime/Engine/Private/Materials/MaterialShared.cpp
--- a/Engine/Source/Runtime/Engine/Private/Materials/MaterialShared.cpp
+++ b/Engine/Source/Runtime/Engine/Private/Materials/MaterialShared.cpp
@@ -166,6 +166,45 @@ FName GetMaterialQualityLevelFName(EMaterialQualityLevel::Type InQualityLevel)
 	return MaterialQualityLevelNames[(int32)InQualityLevel];
 }
 
+#if WITH_EDITOR
+
+/**
+* What shader format should we explicitly cook for?
+* @returns shader format name or NAME_None if the switch was not specified.
+*
+* @note: -CacheShaderFormat=
+*/
+FName GetCmdLineShaderFormatToCache()
+{
+	FString ShaderFormat;
+	FParse::Value(FCommandLine::Get(), TEXT("-CacheShaderFormat="), ShaderFormat);
+	return ShaderFormat.Len() ? FName(ShaderFormat) : NAME_None;
+}
+
+void GetCmdLineFilterShaderFormats(TArray<FName>& InOutShderFormats)
+{
+	// if we specified -CacheShaderFormat= on the cmd line we should only cook that format.
+	static const FName CommandLineShaderFormat = GetCmdLineShaderFormatToCache();
+	if (CommandLineShaderFormat != NAME_None)
+	{
+		// the format is only valid if it is a desired format for this platform.
+		if (InOutShderFormats.Contains(CommandLineShaderFormat))
+		{
+			// only cache the format specified on the command line.
+			InOutShderFormats.Reset(1);
+			InOutShderFormats.Add(CommandLineShaderFormat);
+		}
+	}
+}
+
+int32 GetCmdLineMaterialQualityToCache()
+{
+	int32 MaterialQuality = INDEX_NONE;
+	FParse::Value(FCommandLine::Get(), TEXT("-CacheMaterialQuality="), MaterialQuality);
+	return MaterialQuality;
+}
+#endif
+
 #if STORE_ONLY_ACTIVE_SHADERMAPS
 const FMaterialResourceLocOnDisk* FindMaterialResourceLocOnDisk(
 	const TArray<FMaterialResourceLocOnDisk>& DiskLocations,
diff --git a/Engine/Source/Runtime/Engine/Public/MaterialShared.h b/Engine/Source/Runtime/Engine/Public/MaterialShared.h
--- a/Engine/Source/Runtime/Engine/Public/MaterialShared.h
+++ b/Engine/Source/Runtime/Engine/Public/MaterialShared.h
@@ -129,6 +129,20 @@ inline uint32 GetUseSubsurfaceProfileShadingModelMask()
 /** Whether to allow dithered LOD transitions for a specific feature level. */
 ENGINE_API bool AllowDitheredLODTransition(ERHIFeatureLevel::Type FeatureLevel);
 
+#if WITH_EDITOR
+
+void GetCmdLineFilterShaderFormats(TArray<FName>& InOutShderFormats);
+
+/** 
+* What material quality level should we explicitly cook for? 
+* @returns Quality level or INDEX_NONE if the switch was not specified.
+* 
+* @note: -CacheMaterialQuality=
+*/
+int32 GetCmdLineMaterialQualityToCache();
+
+#endif
+
 /**
  * The types which can be used by materials.
  */
