Sequencer Crash - No Movie Scene found for SequencerDataModel
#fix ensure that all paths initialise the DataModel (after skeleton has been set)

diff --git a/Engine/Source/Editor/UnrealEd/Private/AnimationEditorUtils.cpp b/Engine/Source/Editor/UnrealEd/Private/AnimationEditorUtils.cpp
--- a/Engine/Source/Editor/UnrealEd/Private/AnimationEditorUtils.cpp
+++ b/Engine/Source/Editor/UnrealEd/Private/AnimationEditorUtils.cpp
@@ -503,6 +503,12 @@ namespace AnimationEditorUtils
 				if(NewAsset)
 				{
 					NewAsset->SetSkeleton(Skeleton);
+
+					if (UAnimSequenceBase* SequenceBase = Cast<UAnimSequenceBase>(NewAsset))
+					{
+						SequenceBase->GetController().InitializeModel();
+					}
+					
 					if (SkeletalMesh)
 					{
 						NewAsset->SetPreviewMesh(SkeletalMesh);
diff --git a/Engine/Source/Editor/UnrealEd/Private/Factories/AnimSequenceFactory.cpp b/Engine/Source/Editor/UnrealEd/Private/Factories/AnimSequenceFactory.cpp
--- a/Engine/Source/Editor/UnrealEd/Private/Factories/AnimSequenceFactory.cpp
+++ b/Engine/Source/Editor/UnrealEd/Private/Factories/AnimSequenceFactory.cpp
@@ -72,6 +72,9 @@ UObject* UAnimSequenceFactory::FactoryCreateNew(UClass* Class, UObject* InParent
 	{
 		AnimSequence->SetSkeleton(TargetSkeleton);
 	}
+
+	AnimSequence->GetController().InitializeModel();
+	
 	if (PreviewSkeletalMesh)
 	{
 		AnimSequence->SetPreviewMesh(PreviewSkeletalMesh);
diff --git a/Engine/Source/Runtime/Engine/Private/Animation/AnimSequence.cpp b/Engine/Source/Runtime/Engine/Private/Animation/AnimSequence.cpp
--- a/Engine/Source/Runtime/Engine/Private/Animation/AnimSequence.cpp
+++ b/Engine/Source/Runtime/Engine/Private/Animation/AnimSequence.cpp
@@ -3774,7 +3774,6 @@ bool UAnimSequence::CreateAnimation(USkeletalMesh* Mesh)
 
 		IAnimationDataController::FScopedBracket ScopedBracket(Controller, LOCTEXT("CreateAnimationSkeletalMesh_Bracket", "Creating Animation Sequence based up Skeletal Mesh"));
 
-		Controller->InitializeModel();
 		ResetAnimation();
 
 		Controller->SetFrameRate(UAnimationSettings::Get()->GetDefaultFrameRate());
@@ -3813,7 +3812,6 @@ bool UAnimSequence::CreateAnimation(USkeletalMeshComponent* MeshComponent)
 
 		IAnimationDataController::FScopedBracket ScopedBracket(Controller, LOCTEXT("CreateAnimationSkeletalMeshComponent_Bracket", "Creating Animation Sequence based up Skeletal Mesh Component"));
 
-		Controller->InitializeModel();
 		ResetAnimation();
 
 		Controller->SetFrameRate(UAnimationSettings::Get()->GetDefaultFrameRate());
