CIS Build Error: Invalid play length value provided: 0 frames
#fix ensure that DataModel is actually loaded before use in upgrade path (this can get hit through UAnimationAsset::ReplaceSkeleton)

diff --git a/Engine/Source/Runtime/Engine/Private/Animation/AnimSequenceBase.cpp b/Engine/Source/Runtime/Engine/Private/Animation/AnimSequenceBase.cpp
--- a/Engine/Source/Runtime/Engine/Private/Animation/AnimSequenceBase.cpp
+++ b/Engine/Source/Runtime/Engine/Private/Animation/AnimSequenceBase.cpp
@@ -122,6 +122,10 @@ void UAnimSequenceBase::PostLoad()
 				if (DataModel != nullptr)
 				{					
 					QUICK_SCOPE_CYCLE_COUNTER(STAT_PopulateWithExistingModel);
+					if (FLinkerLoad* DataModelLinker = DataModel->GetLinker())
+					{
+						DataModelLinker->Preload(DataModel);
+					}					
 					DataModel->PostLoad();
 					DataModel->ConditionalPostLoadSubobjects();
 					PopulateWithExistingModel(DataModel.Get());
