WorldPartition: ActorDesc deprecation support

diff --git a/Engine/Source/Runtime/Engine/Private/WorldPartition/WorldPartitionActorDesc.cpp b/Engine/Source/Runtime/Engine/Private/WorldPartition/WorldPartitionActorDesc.cpp
--- a/Engine/Source/Runtime/Engine/Private/WorldPartition/WorldPartitionActorDesc.cpp
+++ b/Engine/Source/Runtime/Engine/Private/WorldPartition/WorldPartitionActorDesc.cpp
@@ -25,6 +25,8 @@
 #endif
 
 #if WITH_EDITOR
+TMap<TSubclassOf<AActor>, FWorldPartitionActorDesc::FActorDescDeprecator> FWorldPartitionActorDesc::Deprecators;
+
 FWorldPartitionActorDesc::FWorldPartitionActorDesc()
 	: bIsUsingDataLayerAsset(false)
 	, SoftRefCount(0)
@@ -147,6 +149,18 @@ void FWorldPartitionActorDesc::Init(const FWorldPartitionActorDescInitData& Desc
 	// Serialize metadata payload
 	Serialize(MetadataAr);
 
+	// Call registered deprecator
+	TSubclassOf<AActor> DeprecatedClass = ActorNativeClass;
+	while (DeprecatedClass)
+	{
+		if (FActorDescDeprecator* Deprecator = Deprecators.Find(DeprecatedClass))
+		{
+			(*Deprecator)(MetadataAr, this);
+			break;
+		}
+		DeprecatedClass = DeprecatedClass->GetSuperClass();
+	}
+
 	Container = nullptr;
 }
 
@@ -193,6 +207,12 @@ void FWorldPartitionActorDesc::SerializeTo(TArray<uint8>& OutData)
 	OutData.Append(PayloadData);
 }
 
+void FWorldPartitionActorDesc::RegisterActorDescDeprecator(TSubclassOf<AActor> ActorClass, const FActorDescDeprecator& Deprecator)
+{
+	check(!Deprecators.Contains(ActorClass));
+	Deprecators.Add(ActorClass, Deprecator);
+}
+
 void FWorldPartitionActorDesc::TransformInstance(const FString& From, const FString& To, const FTransform& InstanceTransform)
 {
 	check(!HardRefCount);
diff --git a/Engine/Source/Runtime/Engine/Public/WorldPartition/WorldPartitionActorDesc.h b/Engine/Source/Runtime/Engine/Public/WorldPartition/WorldPartitionActorDesc.h
--- a/Engine/Source/Runtime/Engine/Public/WorldPartition/WorldPartitionActorDesc.h
+++ b/Engine/Source/Runtime/Engine/Public/WorldPartition/WorldPartitionActorDesc.h
@@ -7,6 +7,7 @@
 #include "UObject/ObjectMacros.h"
 #include "UObject/WeakObjectPtrTemplates.h"
 #include "UObject/WeakObjectPtr.h"
+#include "Templates/SubclassOf.h"
 #include "Misc/Guid.h"
 
 #if WITH_EDITOR
@@ -101,6 +102,7 @@ public:
 	FBox GetBounds() const;
 	inline const FGuid& GetParentActor() const { return ParentActor; }
 	inline bool IsUsingDataLayerAsset() const { return bIsUsingDataLayerAsset; }
+	inline void AddProperty(FName PropertyName, FName PropertyValue = NAME_None) { Properties.AddProperty(PropertyName, PropertyValue); }
 	inline bool GetProperty(FName PropertyName, FName* PropertyValue) const { return Properties.GetProperty(PropertyName, PropertyValue); }
 	inline bool HasProperty(FName PropertyName) const { return Properties.HasProperty(PropertyName); }
 
@@ -196,6 +198,9 @@ public:
 
 	void TransformInstance(const FString& From, const FString& To, const FTransform& InstanceTransform);
 
+	using FActorDescDeprecator = TFunction<void(FArchive&, FWorldPartitionActorDesc*)>;
+	static void RegisterActorDescDeprecator(TSubclassOf<AActor> ActorClass, const FActorDescDeprecator& Deprecator);
+
 protected:
 	FWorldPartitionActorDesc();
 
@@ -247,5 +252,7 @@ protected:
 	UActorDescContainer*			Container;
 	TArray<FName>					DataLayerInstanceNames;
 	bool							bIsForcedNonSpatiallyLoaded;
+
+	static TMap<TSubclassOf<AActor>, FActorDescDeprecator> Deprecators;
 #endif
 };
