Bind InstancedViewUB independent of active stereo rendering

The shader compiler will enable ISR or Mobile Multi-View in shaders, independent of whether they are actively engaged in stereo rendering.
But we _bind_ the instanced view UB when we think stereo is active. We can bind the instanced view UB even while doing mono rendering.

diff --git a/Engine/Source/Runtime/Engine/Private/SceneView.cpp b/Engine/Source/Runtime/Engine/Private/SceneView.cpp
--- a/Engine/Source/Runtime/Engine/Private/SceneView.cpp
+++ b/Engine/Source/Runtime/Engine/Private/SceneView.cpp
@@ -21,7 +21,6 @@
 #include "Interfaces/Interface_PostProcessVolume.h"
 #include "Engine/TextureCube.h"
 #include "StereoRendering.h"
-#include "StereoRenderTargetManager.h"
 #include "IHeadMountedDisplay.h"
 #include "IXRTrackingSystem.h"
 #include "Engine/RendererSettings.h"
@@ -29,6 +28,7 @@
 #include "HighResScreenshot.h"
 #include "Slate/SceneViewport.h"
 #include "RenderUtils.h"
+#include "StereoRenderUtils.h"
 #include "SceneRelativeViewMatrices.h"
 #include "NaniteDefinitions.h"
 
@@ -887,7 +887,10 @@ FSceneView::FSceneView(const FSceneViewInitOptions& InitOptions)
 		bIsMobileMultiViewEnabled = bIsInstancedStereoEnabled = RHISupportsInstancedStereo(ShaderPlatform);
 	}
 
-	bShouldBindInstancedViewUB = bIsInstancedStereoEnabled || bIsMobileMultiViewEnabled;
+	{
+		const UE::StereoRenderUtils::FStereoShaderAspects Aspects(ShaderPlatform);
+		bShouldBindInstancedViewUB = Aspects.IsInstancedStereoEnabled() || Aspects.IsMobileMultiViewEnabled();
+	}
 
 	SetupAntiAliasingMethod();
 
diff --git a/Engine/Source/Runtime/Engine/Private/ShaderCompiler/ShaderCompiler.cpp b/Engine/Source/Runtime/Engine/Private/ShaderCompiler/ShaderCompiler.cpp
--- a/Engine/Source/Runtime/Engine/Private/ShaderCompiler/ShaderCompiler.cpp
+++ b/Engine/Source/Runtime/Engine/Private/ShaderCompiler/ShaderCompiler.cpp
@@ -62,6 +62,7 @@
 #include "ObjectCacheContext.h"
 #include "ProfilingDebugging/StallDetector.h"
 #include "RenderUtils.h"
+#include "StereoRenderUtils.h"
 #include "ProfilingDebugging/CountersTrace.h"
 #include "ClearReplacementShaders.h"
 #include "Engine/RendererSettings.h"
@@ -5514,47 +5515,35 @@ void GlobalBeginCompileShader(
 		}
 	}
 
-	// Set VR definitions
+	if (ShaderPlatform == SP_VULKAN_ES3_1_ANDROID || ShaderPlatform == SP_VULKAN_SM5_ANDROID)
 	{
-		static const auto CVarInstancedStereo = IConsoleManager::Get().FindTConsoleVariableDataInt(TEXT("vr.InstancedStereo"));
-		static const auto CVarMobileMultiView = IConsoleManager::Get().FindTConsoleVariableDataInt(TEXT("vr.MobileMultiView"));
-		static const auto CVarODSCapture = IConsoleManager::Get().FindTConsoleVariableDataInt(TEXT("vr.ODSCapture"));
-		static const auto CVarMobileHDR = IConsoleManager::Get().FindTConsoleVariableDataInt(TEXT("r.MobileHDR"));
-
-		const bool bIsInstancedStereoCVar = CVarInstancedStereo ? (CVarInstancedStereo->GetValueOnAnyThread() != 0) : false;
-		const bool bIsMobileMultiViewCVar = CVarMobileMultiView && CVarMobileHDR ?
-			(CVarMobileMultiView->GetValueOnAnyThread() != 0 && CVarMobileHDR->GetValueOnAnyThread() == 0) : false;
-		const bool bIsODSCapture = CVarODSCapture && (CVarODSCapture->GetValueOnAnyThread() != 0);
-		
-		if (ShaderPlatform == SP_VULKAN_ES3_1_ANDROID || ShaderPlatform == SP_VULKAN_SM5_ANDROID)
+		bool bIsStripReflect = true;
+		GConfig->GetBool(TEXT("/Script/AndroidRuntimeSettings.AndroidRuntimeSettings"), TEXT("bStripShaderReflection"), bIsStripReflect, GEngineIni);
+		if (!bIsStripReflect)
 		{
-			bool bIsStripReflect = true;
-			GConfig->GetBool(TEXT("/Script/AndroidRuntimeSettings.AndroidRuntimeSettings"), TEXT("bStripShaderReflection"), bIsStripReflect, GEngineIni);
-			if (!bIsStripReflect)
-			{
-				Input.Environment.SetDefine(TEXT("STRIP_REFLECT_ANDROID"), false);
-			}
+			Input.Environment.SetDefine(TEXT("STRIP_REFLECT_ANDROID"), false);
 		}
+	}
 
-		bool bIsInstancedStereo = !bUsingMobileRenderer && bIsInstancedStereoCVar && RHISupportsInstancedStereo(ShaderPlatform);
-		bool bIsMobileMultiview = bUsingMobileRenderer && bIsMobileMultiViewCVar;
-		if (bIsMobileMultiview && !RHISupportsMobileMultiView(ShaderPlatform))
-		{
-			// Native mobile multi-view is not supported, fall back to instancing if available
-			bIsMobileMultiview = bIsInstancedStereo = RHISupportsInstancedStereo(ShaderPlatform);
-		}
+	// Set VR definitions
+	{
+		const UE::StereoRenderUtils::FStereoShaderAspects Aspects(ShaderPlatform);
 
-		Input.Environment.SetDefine(TEXT("INSTANCED_STEREO"), bIsInstancedStereo);
-		Input.Environment.SetDefine(TEXT("MULTI_VIEW"), bIsInstancedStereo && RHISupportsMultiView(ShaderPlatform));
-		Input.Environment.SetDefine(TEXT("MOBILE_MULTI_VIEW"), bIsMobileMultiview);
+		Input.Environment.SetDefine(TEXT("INSTANCED_STEREO"), Aspects.IsInstancedStereoEnabled());
+		Input.Environment.SetDefine(TEXT("MULTI_VIEW"), Aspects.IsInstancedMultiViewportEnabled());
+		Input.Environment.SetDefine(TEXT("MOBILE_MULTI_VIEW"), Aspects.IsMobileMultiViewEnabled());
 
 		// Throw a warning if we are silently disabling ISR due to missing platform support.
-		if (bIsInstancedStereoCVar && !bIsInstancedStereo && !GShaderCompilingManager->AreWarningsSuppressed(ShaderPlatform))
+		static const auto CVarInstancedStereo = IConsoleManager::Get().FindTConsoleVariableDataInt(TEXT("vr.InstancedStereo"));
+		const bool bIsInstancedStereoCVar = CVarInstancedStereo ? (CVarInstancedStereo->GetValueOnAnyThread() != 0) : false;
+		if (bIsInstancedStereoCVar && !Aspects.IsInstancedStereoEnabled() && !GShaderCompilingManager->AreWarningsSuppressed(ShaderPlatform))
 		{
 			UE_LOG(LogShaderCompilers, Log, TEXT("Instanced stereo rendering is not supported for the %s shader platform."), *ShaderFormatName.ToString());
 			GShaderCompilingManager->SuppressWarnings(ShaderPlatform);
 		}
 
+		static const auto CVarODSCapture = IConsoleManager::Get().FindTConsoleVariableDataInt(TEXT("vr.ODSCapture"));
+		const bool bIsODSCapture = CVarODSCapture && (CVarODSCapture->GetValueOnAnyThread() != 0);
 		Input.Environment.SetDefine(TEXT("ODS_CAPTURE"), bIsODSCapture);
 	}
 
diff --git a/Engine/Source/Runtime/RenderCore/Private/Shader.cpp b/Engine/Source/Runtime/RenderCore/Private/Shader.cpp
--- a/Engine/Source/Runtime/RenderCore/Private/Shader.cpp
+++ b/Engine/Source/Runtime/RenderCore/Private/Shader.cpp
@@ -18,6 +18,7 @@
 #include "ShaderCore.h"
 #include "ShaderCompilerCore.h"
 #include "RenderUtils.h"
+#include "StereoRenderUtils.h"
 #include "Misc/ConfigCacheIni.h"
 #include "Misc/ScopeLock.h"
 #include "UObject/RenderingObjectVersion.h"
@@ -1390,37 +1391,27 @@ void ShaderMapAppendKeyString(EShaderPlatform Platform, FString& KeyString)
 	}
 
 	{
-		static const auto CVarInstancedStereo = IConsoleManager::Get().FindTConsoleVariableDataInt(TEXT("vr.InstancedStereo"));
-		static const auto CVarMobileMultiView = IConsoleManager::Get().FindTConsoleVariableDataInt(TEXT("vr.MobileMultiView"));
-		static const auto CVarODSCapture = IConsoleManager::Get().FindTConsoleVariableDataInt(TEXT("vr.ODSCapture"));
-
-		bool bIsInstancedStereo = (RHISupportsInstancedStereo(Platform) && (CVarInstancedStereo && CVarInstancedStereo->GetValueOnGameThread() != 0));
-		const bool bIsMultiView = (RHISupportsMultiView(Platform) && bIsInstancedStereo);
-
-		bool bIsMobileMultiView = IsMobilePlatform(Platform) && (CVarMobileMultiView && CVarMobileMultiView->GetValueOnGameThread() != 0);
-		if (bIsMobileMultiView && !RHISupportsMobileMultiView(Platform))
-		{
-			// Native mobile multi-view is not supported, fall back to instancing if available (even if disabled in settings)
-			bIsMobileMultiView = bIsInstancedStereo = RHISupportsInstancedStereo(Platform);
-		}
-
-		const bool bIsODSCapture = CVarODSCapture && (CVarODSCapture->GetValueOnGameThread() != 0);
+		const UE::StereoRenderUtils::FStereoShaderAspects Aspects(Platform);
 
-		if (bIsInstancedStereo)
+		if (Aspects.IsInstancedStereoEnabled())
 		{
 			KeyString += TEXT("_VRIS");
-			
-			if (bIsMultiView)
+
+			if (Aspects.IsInstancedMultiViewportEnabled())
 			{
 				KeyString += TEXT("_MVIEW");
 			}
 		}
 
-		if (bIsMobileMultiView)
+		if (Aspects.IsMobileMultiViewEnabled())
 		{
 			KeyString += TEXT("_MMVIEW");
 		}
+	}
 
+	{
+		static const auto CVarODSCapture = IConsoleManager::Get().FindTConsoleVariableDataInt(TEXT("vr.ODSCapture"));
+		const bool bIsODSCapture = CVarODSCapture && (CVarODSCapture->GetValueOnGameThread() != 0);
 		if (bIsODSCapture)
 		{
 			KeyString += TEXT("_ODSC");
diff --git a/Engine/Source/Runtime/RenderCore/Private/StereoRenderUtils.cpp b/Engine/Source/Runtime/RenderCore/Private/StereoRenderUtils.cpp
new file mode 100644
--- /dev/null
+++ b/Engine/Source/Runtime/RenderCore/Private/StereoRenderUtils.cpp
@@ -0,0 +1,52 @@
+#include "StereoRenderUtils.h"
+#include "RenderUtils.h"
+
+namespace UE::StereoRenderUtils
+{
+	
+RENDERCORE_API FStereoShaderAspects::FStereoShaderAspects(EShaderPlatform Platform) :
+	bInstancedStereoEnabled(false)
+	, bMobileMultiViewEnabled(false)
+	, bInstancedMultiViewportEnabled(false)
+	, bInstancedStereoNative(false)
+	, bMobileMultiViewNative(false)
+	, bMobileMultiViewFallback(false)
+{
+	check(Platform < EShaderPlatform::SP_NumPlatforms);
+	
+	// Would be nice to use URendererSettings, but not accessible in RenderCore
+	static FShaderPlatformCachedIniValue<bool> CVarInstancedStereo(TEXT("vr.InstancedStereo"));
+	static FShaderPlatformCachedIniValue<bool> CVarMobileMultiView(TEXT("vr.MobileMultiView"));
+	static FShaderPlatformCachedIniValue<bool> CVarMobileHDR(TEXT("r.MobileHDR"));
+	
+	const bool bInstancedStereo = CVarInstancedStereo.Get(Platform);
+	
+	const bool bMobilePlatform = IsMobilePlatform(Platform);
+	const bool bMobilePostprocessing = CVarMobileHDR.Get(Platform);
+	const bool bMobileMultiView = CVarMobileMultiView.Get(Platform);
+	
+	bInstancedStereoNative = !bMobilePlatform && bInstancedStereo && RHISupportsInstancedStereo(Platform);
+	
+	const bool bMobileMultiViewCoreSupport = bMobilePlatform && bMobileMultiView && !bMobilePostprocessing;
+	if (bMobileMultiViewCoreSupport)
+	{
+		if (RHISupportsMobileMultiView(Platform))
+		{
+			bMobileMultiViewNative = true;
+		}
+		else if (RHISupportsInstancedStereo(Platform))
+		{
+			bMobileMultiViewFallback = true;
+		}
+	}
+	
+	bInstancedStereoEnabled = bInstancedStereoNative || bMobileMultiViewFallback;
+	bInstancedMultiViewportEnabled = bInstancedStereoNative && RHISupportsMultiView(Platform);
+	bMobileMultiViewEnabled = bMobileMultiViewNative || bMobileMultiViewFallback;
+}
+	
+RENDERCORE_API bool FStereoShaderAspects::IsInstancedStereoEnabled() const { return bInstancedStereoEnabled; }
+RENDERCORE_API bool FStereoShaderAspects::IsMobileMultiViewEnabled() const { return bMobileMultiViewEnabled; }
+RENDERCORE_API bool FStereoShaderAspects::IsInstancedMultiViewportEnabled() const { return bInstancedMultiViewportEnabled; }
+	
+} // namespace UE::RenderUtils
diff --git a/Engine/Source/Runtime/RenderCore/Public/StereoRenderUtils.h b/Engine/Source/Runtime/RenderCore/Public/StereoRenderUtils.h
new file mode 100644
--- /dev/null
+++ b/Engine/Source/Runtime/RenderCore/Public/StereoRenderUtils.h
@@ -0,0 +1,51 @@
+// Copyright Epic Games, Inc. All Rights Reserved.
+
+#pragma once
+
+#include "RHIDefinitions.h"
+
+namespace UE::StereoRenderUtils
+{
+
+	/*
+	 * Detect the single-draw stereo shader variant, in order to support usage across different platforms
+	 */
+	class FStereoShaderAspects
+	{
+	public:
+		/**
+		* Determines the stereo aspects of the shader pipeline based on the input shader platform
+		* @param Platform	Target shader platform used to determine stereo shader variant
+		*/
+		RENDERCORE_API FStereoShaderAspects(EShaderPlatform Platform);
+
+		/**
+		 * Whether instanced stereo rendering is enabled - i.e. using a single instanced drawcall to render to both stereo views.
+		 * Can redirect output either via viewport index or X coordinate shift + clip planes.
+		 */
+		RENDERCORE_API bool IsInstancedStereoEnabled() const;
+
+		/**
+		 * Whether mobile multiview is enabled - i.e. using VK_KHR_multiview. Another drawcall reduction technique, independent of instanced stereo.
+		 * Mobile multiview generates view indices to index into texture arrays.
+		 * Can be internally emulated using instanced stereo if native support is unavailable, by using ISR-generated view indices to index into texture arrays.
+		 */
+		RENDERCORE_API bool IsMobileMultiViewEnabled() const;
+
+		/**
+		 * Whether multiviewport rendering is enabled - i.e. using ViewportIndex to index into viewport.
+		 * Relies on instanced stereo rendering being enabled.
+		 */
+		RENDERCORE_API bool IsInstancedMultiViewportEnabled() const;
+
+	private:
+		bool bInstancedStereoEnabled : 1;
+		bool bMobileMultiViewEnabled : 1;
+		bool bInstancedMultiViewportEnabled : 1;
+
+		bool bInstancedStereoNative : 1;
+		bool bMobileMultiViewNative : 1;
+		bool bMobileMultiViewFallback : 1;
+	};
+
+}
