Fix Static Analysis Win64 (MSVC) - Arithmetic overflow

#rnx

diff --git a/Engine/Source/Runtime/Engine/Private/WorldPartition/WorldPartitionRuntimeSpatialHash.cpp b/Engine/Source/Runtime/Engine/Private/WorldPartition/WorldPartitionRuntimeSpatialHash.cpp
--- a/Engine/Source/Runtime/Engine/Private/WorldPartition/WorldPartitionRuntimeSpatialHash.cpp
+++ b/Engine/Source/Runtime/Engine/Private/WorldPartition/WorldPartitionRuntimeSpatialHash.cpp
@@ -897,7 +897,7 @@ void UWorldPartitionRuntimeSpatialHash::DumpStateLog(FHierarchicalLogArchive& Ar
 						LevelActorCount += Cell->GetActorCount();
 					}
 				}
-				LevelsStats.Add({ LevelCellCount, (StreamingGrid.CellSize << Level), LevelActorCount });
+				LevelsStats.Add({ LevelCellCount, ((int64)StreamingGrid.CellSize << (int64)Level), LevelActorCount });
 				TotalActorCount += LevelActorCount;
 				++Level;
 			}
