UE-151338 - Implement refactored implicit conversion for math types.

Adds support for most of the remaining LWC math types (eg: FVector4). Note that FVector2D and FBox2D are currently disabled. Their corresponding double variants have FName collisions with the primary variant names, which needs to be addressed separately.

Some cleanup was also added that reroutes all conversion through a common function, ConvertType. This alleviates the need for separate conversion functions for each type.

diff --git a/Engine/Source/Runtime/CoreUObject/Public/UObject/NoExportTypes.h b/Engine/Source/Runtime/CoreUObject/Public/UObject/NoExportTypes.h
--- a/Engine/Source/Runtime/CoreUObject/Public/UObject/NoExportTypes.h
+++ b/Engine/Source/Runtime/CoreUObject/Public/UObject/NoExportTypes.h
@@ -471,7 +471,7 @@ struct FGuid
  * A point or direction FVector in 3d space.
  * @note The full C++ class is located here: Engine\Source\Runtime\Core\Public\Math\Vector.h
  */
-USTRUCT(immutable, noexport, BlueprintType, IsAlwaysAccessible, HasDefaults, HasNoOpConstructor, IsCoreType, meta = (HasNativeBreak = "/Script/Engine.KismetMathLibrary.BreakVector3f"))
+USTRUCT(immutable, noexport, BlueprintType, BlueprintInternalUseOnly, IsAlwaysAccessible, HasDefaults, HasNoOpConstructor, IsCoreType, meta = (HasNativeBreak = "/Script/Engine.KismetMathLibrary.BreakVector3f"))
 struct FVector3f
 {
 	UPROPERTY(EditAnywhere, Category = Vector, SaveGame)
@@ -523,7 +523,7 @@ struct FVector
 * A 4-D homogeneous vector.
 * @note The full C++ class is located here: Engine\Source\Runtime\Core\Public\Math\Vector4.h
 */
-USTRUCT(immutable, noexport, IsAlwaysAccessible, HasDefaults, HasNoOpConstructor, IsCoreType)
+USTRUCT(immutable, noexport, BlueprintType, BlueprintInternalUseOnly, IsAlwaysAccessible, HasDefaults, HasNoOpConstructor, IsCoreType)
 struct FVector4f
 {
 	UPROPERTY(EditAnywhere, Category = Vector4, SaveGame)
@@ -585,7 +585,7 @@ struct FVector4
 * A vector in 2-D space composed of components (X, Y) with floating point precision.
 * @note The full C++ class is located here: Engine\Source\Runtime\Core\Public\Math\Vector2D.h
 */
-USTRUCT(immutable, noexport, IsAlwaysAccessible, HasDefaults, HasNoOpConstructor, IsCoreType)
+USTRUCT(immutable, noexport, BlueprintType, BlueprintInternalUseOnly, IsAlwaysAccessible, HasDefaults, HasNoOpConstructor, IsCoreType)
 struct FVector2f
 {
 	UPROPERTY(EditAnywhere, Category=Vector2D, SaveGame)
@@ -639,7 +639,7 @@ struct FTwoVectors
  * A plane definition in 3D space.
  * @note The full C++ class is located here: Engine\Source\Runtime\Core\Public\Math\Plane.h
  */
-USTRUCT(immutable, noexport, IsAlwaysAccessible, HasDefaults, HasNoOpConstructor, IsCoreType)
+USTRUCT(immutable, noexport, BlueprintType, BlueprintInternalUseOnly, IsAlwaysAccessible, HasDefaults, HasNoOpConstructor, IsCoreType)
 struct FPlane4f : public FVector3f
 {
 	UPROPERTY(EditAnywhere, Category=Plane, SaveGame)
@@ -673,7 +673,7 @@ struct FPlane : public FVector
  * An orthogonal rotation in 3d space.
  * @note The full C++ class is located here: Engine\Source\Runtime\Core\Public\Math\Rotator.h
  */
-USTRUCT(immutable, noexport, IsAlwaysAccessible, HasDefaults, HasNoOpConstructor, IsCoreType)
+USTRUCT(immutable, noexport, BlueprintType, BlueprintInternalUseOnly, IsAlwaysAccessible, HasDefaults, HasNoOpConstructor, IsCoreType)
 struct FRotator3f
 {
 	/** Pitch (degrees) around Y axis */
@@ -733,7 +733,7 @@ struct FRotator
  * Quaternion.
  * @note The full C++ class is located here: Engine\Source\Runtime\Core\Public\Math\Quat.h
  */
-USTRUCT(immutable, noexport, IsAlwaysAccessible, HasDefaults, HasNoOpConstructor, IsCoreType)
+USTRUCT(immutable, noexport, BlueprintType, BlueprintInternalUseOnly, IsAlwaysAccessible, HasDefaults, HasNoOpConstructor, IsCoreType)
 struct FQuat4f
 {
 	UPROPERTY(EditAnywhere, Category=Quat, SaveGame)
@@ -938,7 +938,7 @@ struct FLinearColor
  * A point or direction FVector in 3d space.
  * @note The full C++ class is located here: Engine\Source\Runtime\Core\Public\Math\Box.h
  */
-USTRUCT(immutable, noexport, IsAlwaysAccessible, HasDefaults, HasNoOpConstructor, IsCoreType)
+USTRUCT(immutable, noexport, BlueprintType, BlueprintInternalUseOnly, IsAlwaysAccessible, HasDefaults, HasNoOpConstructor, IsCoreType)
 struct FBox3f
 {
 	UPROPERTY(EditAnywhere, Category = Box, SaveGame)
@@ -990,7 +990,7 @@ struct FBox
  * A rectangular 2D Box.
  * @note The full C++ class is located here: Engine\Source\Runtime\Core\Public\Math\Box2D.h
  */
-USTRUCT(immutable, noexport, IsAlwaysAccessible, HasDefaults, HasNoOpConstructor, IsCoreType)
+USTRUCT(immutable, noexport, BlueprintType, BlueprintInternalUseOnly, IsAlwaysAccessible, HasDefaults, HasNoOpConstructor, IsCoreType)
 struct FBox2f
 {
 	UPROPERTY(EditAnywhere, Category=Box2D, SaveGame)
@@ -1138,7 +1138,7 @@ struct FOrientedBox
  * A 4x4 matrix.
  * @note The full C++ class is located here: Engine\Source\Runtime\Core\Public\Math\Matrix.h
  */
-USTRUCT(immutable, noexport, HasDefaults, HasNoOpConstructor, IsCoreType)
+USTRUCT(immutable, noexport, BlueprintType, BlueprintInternalUseOnly, HasDefaults, HasNoOpConstructor, IsCoreType)
 struct FMatrix44f
 {
 	UPROPERTY(EditAnywhere, Category=Matrix, SaveGame)
@@ -1490,7 +1490,7 @@ struct FInterpCurveLinearColor
  * Transform composed of Quat/Translation/Scale.
  * @note This is implemented in either TransformVectorized.h or TransformNonVectorized.h depending on the platform.
  */
-USTRUCT(immutable, noexport, IsAlwaysAccessible, HasDefaults, IsCoreType)
+USTRUCT(immutable, noexport, BlueprintType, BlueprintInternalUseOnly, IsAlwaysAccessible, HasDefaults, IsCoreType)
 struct FTransform3f
 {
 	/** Rotation of this transformation, as a quaternion. */
diff --git a/Engine/Source/Runtime/Engine/Internal/Kismet/BlueprintTypeConversions.h b/Engine/Source/Runtime/Engine/Internal/Kismet/BlueprintTypeConversions.h
--- a/Engine/Source/Runtime/Engine/Internal/Kismet/BlueprintTypeConversions.h
+++ b/Engine/Source/Runtime/Engine/Internal/Kismet/BlueprintTypeConversions.h
@@ -62,6 +62,7 @@ public:
 	static TMap<int, int> ConvertMapType(const TMap<int, int>& InMap);
 	DECLARE_FUNCTION(execConvertMapType);
 
+
 	// Custom struct conversions
 
 	DECLARE_CONVERSION_FUNCTIONS(FVector, FVector3f, FVector3d);
@@ -71,4 +72,79 @@ public:
 
 	UFUNCTION(CustomThunk, meta = (BlueprintInternalUseOnly = "true"))
 	static int32 ConvertFVector3fToFVector3d(int32 InFromData);
+
+	// LWC_TODO - FVector2D/FVector2d have FName collisions.
+#if 0
+	DECLARE_CONVERSION_FUNCTIONS(FVector2D, FVector2f, FVector2d);
+
+	UFUNCTION(CustomThunk, meta = (BlueprintInternalUseOnly = "true"))
+	static int32 ConvertFVector2dToFVector2f(int32 InFromData);
+
+	UFUNCTION(CustomThunk, meta = (BlueprintInternalUseOnly = "true"))
+	static int32 ConvertFVector2fToFVector2d(int32 InFromData);
+#endif
+
+	DECLARE_CONVERSION_FUNCTIONS(FVector4, FVector4f, FVector4d);
+
+	UFUNCTION(CustomThunk, meta = (BlueprintInternalUseOnly = "true"))
+	static int32 ConvertFVector4dToFVector4f(int32 InFromData);
+
+	UFUNCTION(CustomThunk, meta = (BlueprintInternalUseOnly = "true"))
+	static int32 ConvertFVector4fToFVector4d(int32 InFromData);
+
+
+	DECLARE_CONVERSION_FUNCTIONS(FPlane, FPlane4f, FPlane4d);
+
+	UFUNCTION(CustomThunk, meta = (BlueprintInternalUseOnly = "true"))
+	static int32 ConvertFPlane4dToFPlane4f(int32 InFromData);
+
+	UFUNCTION(CustomThunk, meta = (BlueprintInternalUseOnly = "true"))
+	static int32 ConvertFPlane4fToFPlane4d(int32 InFromData);
+
+
+	DECLARE_CONVERSION_FUNCTIONS(FQuat, FQuat4f, FQuat4d);
+
+	UFUNCTION(CustomThunk, meta = (BlueprintInternalUseOnly = "true"))
+	static int32 ConvertFQuat4dToFQuat4f(int32 InFromData);
+
+	UFUNCTION(CustomThunk, meta = (BlueprintInternalUseOnly = "true"))
+	static int32 ConvertFQuat4fToFQuat4d(int32 InFromData);
+	
+
+	DECLARE_CONVERSION_FUNCTIONS(FRotator, FRotator3f, FRotator3d);
+	
+	UFUNCTION(CustomThunk, meta = (BlueprintInternalUseOnly = "true"))
+	static int32 ConvertFRotator3dToFRotator3f(int32 InFromData);
+
+	UFUNCTION(CustomThunk, meta = (BlueprintInternalUseOnly = "true"))
+	static int32 ConvertFRotator3fToFRotator3d(int32 InFromData);
+
+
+	DECLARE_CONVERSION_FUNCTIONS(FTransform, FTransform3f, FTransform3d);
+	
+	UFUNCTION(CustomThunk, meta = (BlueprintInternalUseOnly = "true"))
+	static int32 ConvertFTransform3dToFTransform3f(int32 InFromData);
+
+	UFUNCTION(CustomThunk, meta = (BlueprintInternalUseOnly = "true"))
+	static int32 ConvertFTransform3fToFTransform3d(int32 InFromData);
+
+
+	DECLARE_CONVERSION_FUNCTIONS(FMatrix, FMatrix44f, FMatrix44d);
+	
+	UFUNCTION(CustomThunk, meta = (BlueprintInternalUseOnly = "true"))
+	static int32 ConvertFMatrix44dToFMatrix44f(int32 InFromData);
+
+	UFUNCTION(CustomThunk, meta = (BlueprintInternalUseOnly = "true"))
+	static int32 ConvertFMatrix44fToFMatrix44d(int32 InFromData);
+
+	// LWC_TODO - FBox2D/FBox2d have FName collisions.
+#if 0
+	DECLARE_CONVERSION_FUNCTIONS(FBox2D, FBox2f, FBox2d);
+
+	UFUNCTION(CustomThunk, meta = (BlueprintInternalUseOnly = "true"))
+	static int32 ConvertFBox2dToFBox2f(int32 InFromData);
+
+	UFUNCTION(CustomThunk, meta = (BlueprintInternalUseOnly = "true"))
+	static int32 ConvertFBox2fToFBox2d(int32 InFromData);
+#endif
 };
diff --git a/Engine/Source/Runtime/Engine/Private/Kismet/BlueprintTypeConversions.cpp b/Engine/Source/Runtime/Engine/Private/Kismet/BlueprintTypeConversions.cpp
--- a/Engine/Source/Runtime/Engine/Private/Kismet/BlueprintTypeConversions.cpp
+++ b/Engine/Source/Runtime/Engine/Private/Kismet/BlueprintTypeConversions.cpp
@@ -5,24 +5,28 @@
 #define MAKE_CONVERSION_FUNCTION_NAME(SourceType, DestType)													\
 	Convert##SourceType##To##DestType
 
-#define DEFINE_CONVERSION_FUNCTIONS(BASETYPE, VARIANTTYPE1, VARIANTTYPE2, IMPL1TO2, IMPL2TO1)				\
-	DEFINE_FUNCTION(UBlueprintTypeConversions::MAKE_CONVERSION_EXEC_FUNCTION_NAME(VARIANTTYPE1, VARIANTTYPE2))	\
+#define DEFINE_CONVERSION_FUNCTIONS(BASETYPE, VARIANTTYPE1, VARIANTTYPE2)									\
+	DEFINE_FUNCTION(UBlueprintTypeConversions::MAKE_CONVERSION_EXEC_FUNCTION_NAME(VARIANTTYPE1, VARIANTTYPE2)) \
 	{																										\
+		using namespace UE::Kismet::BlueprintTypeConversions::Internal;										\
+																											\
 		void* DestAddr = Stack.MostRecentPropertyAddress;													\
 		Stack.MostRecentProperty = nullptr;																	\
 		Stack.StepCompiledIn(nullptr, nullptr);																\
 		const void* SourceAddr = Stack.MostRecentPropertyAddress;											\
 		P_FINISH;																							\
-		IMPL1TO2(SourceAddr, DestAddr);																		\
+		ConvertType<VARIANTTYPE1, VARIANTTYPE2>(SourceAddr, DestAddr);										\
 	}																										\
-	DEFINE_FUNCTION(UBlueprintTypeConversions::MAKE_CONVERSION_EXEC_FUNCTION_NAME(VARIANTTYPE2, VARIANTTYPE1))	\
+	DEFINE_FUNCTION(UBlueprintTypeConversions::MAKE_CONVERSION_EXEC_FUNCTION_NAME(VARIANTTYPE2, VARIANTTYPE1)) \
 	{																										\
+		using namespace UE::Kismet::BlueprintTypeConversions::Internal;										\
+																											\
 		void* DestAddr = Stack.MostRecentPropertyAddress;													\
 		Stack.MostRecentProperty = nullptr;																	\
 		Stack.StepCompiledIn(nullptr, nullptr);																\
 		const void* SourceAddr = Stack.MostRecentPropertyAddress;											\
 		P_FINISH;																							\
-		IMPL2TO1(SourceAddr, DestAddr);																		\
+		ConvertType<VARIANTTYPE2, VARIANTTYPE1>(SourceAddr, DestAddr);										\
 	}																										\
 	UE::Kismet::BlueprintTypeConversions::Internal::FStructConversionEntry BASETYPE##Entry(					\
 		&TBaseStructure<BASETYPE>::Get,																		\
@@ -31,44 +35,37 @@
 		&TVariantStructure<VARIANTTYPE2>::Get,																\
 		TEXT(PREPROCESSOR_TO_STRING(MAKE_CONVERSION_FUNCTION_NAME(VARIANTTYPE1, VARIANTTYPE2))),			\
 		TEXT(PREPROCESSOR_TO_STRING(MAKE_CONVERSION_FUNCTION_NAME(VARIANTTYPE2, VARIANTTYPE1))),			\
-		&IMPL1TO2,																							\
-		&IMPL2TO1																							\
+		&UE::Kismet::BlueprintTypeConversions::Internal::ConvertType<VARIANTTYPE1, VARIANTTYPE2>,			\
+		&UE::Kismet::BlueprintTypeConversions::Internal::ConvertType<VARIANTTYPE2, VARIANTTYPE1>			\
 	);
 
 namespace UE::Kismet::BlueprintTypeConversions::Internal
 {
 
-ConversionFunctionT FindConversionFunction(const FProperty* InFromProperty, const FProperty* InToProperty)
+template <typename TFrom, typename TTo>
+FORCEINLINE_DEBUGGABLE void ConvertType(const void* InFromData, void* InToData)
 {
-	ConversionFunctionT Result = nullptr;
-
-	auto ConvertFloatToDouble = [](const void* InFromData, void* InToData)
-	{
-		const float* FromFloat = reinterpret_cast<const float*>(InFromData);
-		check(FromFloat);
-		double* ToDouble = reinterpret_cast<double*>(InToData);
-		check(ToDouble);
-
-		*ToDouble = *FromFloat;
-	};
+	const TFrom* From = reinterpret_cast<const TFrom*>(InFromData);
+	check(From);
+	TTo* To = reinterpret_cast<TTo*>(InToData);
+	check(To);
 
-	auto ConvertDoubleToFloat = [](const void* InFromData, void* InToData)
-	{
-		const double* FromDouble = reinterpret_cast<const double*>(InFromData);
-		check(FromDouble);
-		float* ToFloat = reinterpret_cast<float*>(InToData);
-		check(ToFloat);
+	// LWC_TODO - Ensure that we're using the correct constructor here.
+	// Specifically, we need to use a constructor or function that calls a narrowing cast for safe conversions.
+	*To = TTo(*From);
+}
 
-		*ToFloat = static_cast<float>(*FromDouble);
-	};
+ConversionFunctionT FindConversionFunction(const FProperty* InFromProperty, const FProperty* InToProperty)
+{
+	ConversionFunctionT Result = nullptr;
 
 	if (InFromProperty->IsA<FFloatProperty>() && InToProperty->IsA<FDoubleProperty>())
 	{
-		Result = ConvertFloatToDouble;
+		Result = &ConvertType<float, double>;
 	}
 	else if (InFromProperty->IsA<FDoubleProperty>() && InToProperty->IsA<FFloatProperty>())
 	{
-		Result = ConvertDoubleToFloat;
+		Result = &ConvertType<double, float>;
 	}
 	else if (InFromProperty->IsA<FStructProperty>() && InToProperty->IsA<FStructProperty>())
 	{
@@ -219,6 +216,8 @@ UBlueprintTypeConversions::UBlueprintTypeConversions(const FObjectInitializer& O
 {
 }
 
+// Container conversions
+
 DEFINE_FUNCTION(UBlueprintTypeConversions::execConvertArrayType)
 {
 	using namespace UE::Kismet;
@@ -344,30 +343,14 @@ DEFINE_FUNCTION(UBlueprintTypeConversions::execConvertMapType)
 	DestMap.Rehash();
 }
 
-// Container conversions
-
-FORCEINLINE_DEBUGGABLE void ConvertFVector3dToFVector3fImpl(const void* InFromData, void* InToData)
-{
-	const FVector3d* From = reinterpret_cast<const FVector3d*>(InFromData);
-	check(From);
-	FVector3f* To = reinterpret_cast<FVector3f*>(InToData);
-	check(To);
-
-	float x = static_cast<float>(From->X);
-	float y = static_cast<float>(From->Y);
-	float z = static_cast<float>(From->Z);
-
-	*To = FVector3f(x, y, z);
-}
-
-FORCEINLINE_DEBUGGABLE void ConvertFVector3fToFVector3dImpl(const void* InFromData, void* InToData)
-{
-	const FVector3f* From = reinterpret_cast<const FVector3f*>(InFromData);
-	check(From);
-	FVector3d* To = reinterpret_cast<FVector3d*>(InToData);
-	check(To);
-
-	*To = FVector3d(From->X, From->Y, From->Z);
-}
-
-DEFINE_CONVERSION_FUNCTIONS(FVector, FVector3d, FVector3f, ConvertFVector3dToFVector3fImpl, ConvertFVector3fToFVector3dImpl)
+// Custom struct conversions
+
+DEFINE_CONVERSION_FUNCTIONS(FVector, FVector3d, FVector3f)
+// DEFINE_CONVERSION_FUNCTIONS(FVector2D, FVector2d, FVector2f)
+DEFINE_CONVERSION_FUNCTIONS(FVector4, FVector4d, FVector4f)
+DEFINE_CONVERSION_FUNCTIONS(FPlane, FPlane4d, FPlane4f)
+DEFINE_CONVERSION_FUNCTIONS(FQuat, FQuat4d, FQuat4f)
+DEFINE_CONVERSION_FUNCTIONS(FRotator, FRotator3d, FRotator3f)
+DEFINE_CONVERSION_FUNCTIONS(FTransform, FTransform3d, FTransform3f)
+DEFINE_CONVERSION_FUNCTIONS(FMatrix, FMatrix44d, FMatrix44f)
+// DEFINE_CONVERSION_FUNCTIONS(FBox2D, FBox2d, FBox2f)
