Fix crash when saving submix to file
  - Do not try to release or regenerate platformdata when on worker threads as it is forbidden and wasteful work not required for saving wav data

#rnx

diff --git a/Engine/Source/Runtime/Engine/Private/SampleBufferIO.cpp b/Engine/Source/Runtime/Engine/Private/SampleBufferIO.cpp
--- a/Engine/Source/Runtime/Engine/Private/SampleBufferIO.cpp
+++ b/Engine/Source/Runtime/Engine/Private/SampleBufferIO.cpp
@@ -447,7 +447,8 @@ namespace Audio
 
 	void FSoundWavePCMWriter::ApplyBufferToSoundWave()
 	{
-		CurrentSoundWave->InvalidateCompressedData();
+		// Since we just want to replace the PCM data to save it to disk. We don't need to compute anything platformdata related.
+		CurrentSoundWave->InvalidateCompressedData(false /* bFreeResources */, false /* bRebuildStreamingChunk */);
 
 		CurrentSoundWave->SetSampleRate(CurrentBuffer.GetSampleRate());
 		CurrentSoundWave->NumChannels = CurrentBuffer.GetNumChannels();
