Fixed the DebugExecBindings commands being run on Pressed and on Release of the associated key.

diff --git a/Engine/Source/Runtime/Engine/Private/UserInterface/PlayerInput.cpp b/Engine/Source/Runtime/Engine/Private/UserInterface/PlayerInput.cpp
--- a/Engine/Source/Runtime/Engine/Private/UserInterface/PlayerInput.cpp
+++ b/Engine/Source/Runtime/Engine/Private/UserInterface/PlayerInput.cpp
@@ -292,13 +292,16 @@ bool UPlayerInput::InputKey(const FInputKeyParams& Params)
 	#if !UE_BUILD_SHIPPING
 		CurrentEvent		= IE_Axis;
 
-		const FString Command = GetBind(Params.Key);
-		if(Command.Len())
+		if(Params.Event == IE_Pressed)
 		{
-			UWorld* World = GetWorld();
-			check(World);
-			ExecInputCommands( World, *Command,*GLog);
-			return true;
+			const FString Command = GetBind(Params.Key);
+			if(Command.Len())
+			{
+				UWorld* World = GetWorld();
+				check(World);
+				ExecInputCommands( World, *Command,*GLog);
+				return true;
+			}
 		}
 	#endif
 
