World Partition Actor Descriptors serialization: added custom archive trapping unsupported types serialization through unimplemented virtuals.

#rnx

diff --git a/Engine/Source/Runtime/Engine/Private/WorldPartition/WorldPartitionActorDesc.cpp b/Engine/Source/Runtime/Engine/Private/WorldPartition/WorldPartitionActorDesc.cpp
--- a/Engine/Source/Runtime/Engine/Private/WorldPartition/WorldPartitionActorDesc.cpp
+++ b/Engine/Source/Runtime/Engine/Private/WorldPartition/WorldPartitionActorDesc.cpp
@@ -1,8 +1,9 @@
 // Copyright Epic Games, Inc. All Rights Reserved.
 
+#if WITH_EDITOR
 #include "WorldPartition/WorldPartitionActorDesc.h"
+#include "WorldPartition/WorldPartitionActorDescArchive.h"
 
-#if WITH_EDITOR
 #include "Misc/HashBuilder.h"
 #include "UObject/LinkerInstancingContext.h"
 #include "UObject/UObjectHash.h"
@@ -22,9 +23,7 @@
 #include "WorldPartition/DataLayer/DataLayerAsset.h"
 #include "WorldPartition/DataLayer/DataLayerUtils.h"
 #include "Engine/Public/ActorReferencesUtils.h"
-#endif
 
-#if WITH_EDITOR
 TMap<TSubclassOf<AActor>, FWorldPartitionActorDesc::FActorDescDeprecator> FWorldPartitionActorDesc::Deprecators;
 
 FWorldPartitionActorDesc::FWorldPartitionActorDesc()
@@ -147,7 +146,8 @@ void FWorldPartitionActorDesc::Init(const FWorldPartitionActorDescInitData& Desc
 	MetadataAr.SetCustomVersions(CustomVersions);
 	
 	// Serialize metadata payload
-	Serialize(MetadataAr);
+	FActorDescArchive ActorDescAr(MetadataAr);
+	Serialize(ActorDescAr);
 
 	// Call registered deprecator
 	TSubclassOf<AActor> DeprecatedClass = ActorNativeClass;
@@ -194,12 +194,15 @@ void FWorldPartitionActorDesc::SerializeTo(TArray<uint8>& OutData)
 	// Serialize to archive and gather custom versions
 	TArray<uint8> PayloadData;
 	FMemoryWriter PayloadAr(PayloadData, true);
-	Serialize(PayloadAr);
+	FActorDescArchive ActorDescAr(PayloadAr);
+
+	// Serialize actor descriptor
+	Serialize(ActorDescAr);
 
 	// Serialize custom versions
 	TArray<uint8> HeaderData;
 	FMemoryWriter HeaderAr(HeaderData);
-	FCustomVersionContainer CustomVersions = PayloadAr.GetCustomVersions();
+	FCustomVersionContainer CustomVersions = ActorDescAr.GetCustomVersions();
 	CustomVersions.Serialize(HeaderAr);
 
 	// Append data
diff --git a/Engine/Source/Runtime/Engine/Private/WorldPartition/WorldPartitionActorDescArchive.cpp b/Engine/Source/Runtime/Engine/Private/WorldPartition/WorldPartitionActorDescArchive.cpp
new file mode 100644
--- /dev/null
+++ b/Engine/Source/Runtime/Engine/Private/WorldPartition/WorldPartitionActorDescArchive.cpp
@@ -0,0 +1,14 @@
+// Copyright Epic Games, Inc. All Rights Reserved.
+
+#if WITH_EDITOR
+#include "WorldPartition/WorldPartitionActorDescArchive.h"
+
+FActorDescArchive::FActorDescArchive(FArchive& InArchive)
+	: FArchiveProxy(InArchive)
+{
+	check(InArchive.IsPersistent());		
+
+	SetIsPersistent(true);
+	SetIsLoading(InArchive.IsLoading());		
+}
+#endif
\ No newline at end of file
diff --git a/Engine/Source/Runtime/Engine/Public/WorldPartition/DataLayer/WorldDataLayersActorDesc.h b/Engine/Source/Runtime/Engine/Public/WorldPartition/DataLayer/WorldDataLayersActorDesc.h
--- a/Engine/Source/Runtime/Engine/Public/WorldPartition/DataLayer/WorldDataLayersActorDesc.h
+++ b/Engine/Source/Runtime/Engine/Public/WorldPartition/DataLayer/WorldDataLayersActorDesc.h
@@ -2,11 +2,11 @@
 
 #pragma once
 
+#if WITH_EDITOR
+
 #include "WorldPartition/WorldPartitionActorDesc.h"
 #include "WorldPartition/DataLayer/DataLayerType.h"
 
-#if WITH_EDITOR
-
 class ENGINE_API FDataLayerInstanceDesc
 {
 public:
@@ -45,14 +45,11 @@ private:
 	//~ End UDeprecatedDataLayerInstance
 };
 
-#endif
-
 /**
  * ActorDesc for AWorldDataLayers actors.
  */
 class ENGINE_API FWorldDataLayersActorDesc : public FWorldPartitionActorDesc
 {
-#if WITH_EDITOR
 public:
 	FWorldDataLayersActorDesc();
 	bool IsValid() const { return bIsValid; }
@@ -70,5 +67,5 @@ protected:
 private:
 	TArray<FDataLayerInstanceDesc> DataLayerInstances;
 	bool bIsValid;
-#endif
 };
+#endif
diff --git a/Engine/Source/Runtime/Engine/Public/WorldPartition/WorldPartitionActorDesc.h b/Engine/Source/Runtime/Engine/Public/WorldPartition/WorldPartitionActorDesc.h
--- a/Engine/Source/Runtime/Engine/Public/WorldPartition/WorldPartitionActorDesc.h
+++ b/Engine/Source/Runtime/Engine/Public/WorldPartition/WorldPartitionActorDesc.h
@@ -1,6 +1,7 @@
 // Copyright Epic Games, Inc. All Rights Reserved.
 #pragma once
 
+#if WITH_EDITOR
 #include "CoreMinimal.h"
 #include "PropertyPairsMap.h"
 #include "UObject/NameTypes.h"
@@ -10,7 +11,6 @@
 #include "Templates/SubclassOf.h"
 #include "Misc/Guid.h"
 
-#if WITH_EDITOR
 // Struct used to create actor descriptor
 struct FWorldPartitionActorDescInitData
 {
@@ -54,14 +54,12 @@ inline bool CompareUnsortedArrays(const TArray<FName>& Array1, const TArray<FNam
 {
 	return CompareUnsortedArrays(Array1, Array2, [](const FName& A, const FName& B) { return A.LexicalLess(B); });
 }
-#endif
 
 /**
  * Represents a potentially unloaded actor (editor-only)
  */
 class ENGINE_API FWorldPartitionActorDesc
 {
-#if WITH_EDITOR
 	friend class AActor;
 	friend class UActorDescContainer;
 	friend class UWorldPartition;
@@ -254,5 +252,5 @@ protected:
 	bool							bIsForcedNonSpatiallyLoaded;
 
 	static TMap<TSubclassOf<AActor>, FActorDescDeprecator> Deprecators;
-#endif
 };
+#endif
\ No newline at end of file
diff --git a/Engine/Source/Runtime/Engine/Public/WorldPartition/WorldPartitionActorDescArchive.h b/Engine/Source/Runtime/Engine/Public/WorldPartition/WorldPartitionActorDescArchive.h
new file mode 100644
--- /dev/null
+++ b/Engine/Source/Runtime/Engine/Public/WorldPartition/WorldPartitionActorDescArchive.h
@@ -0,0 +1,23 @@
+// Copyright Epic Games, Inc. All Rights Reserved.
+#pragma once
+
+#if WITH_EDITOR
+#include "CoreMinimal.h"
+#include "Serialization/Archive.h"
+
+class FActorDescArchive : public FArchiveProxy
+{
+public:
+	FActorDescArchive(FArchive& InArchive);
+
+	//~ Begin FArchive Interface
+	virtual FArchive& operator<<(FText& Value) override { unimplemented(); return *this; }
+	virtual FArchive& operator<<(UObject*& Value) override { unimplemented(); return *this; }
+	virtual FArchive& operator<<(FLazyObjectPtr& Value) override { unimplemented(); return *this; }
+	virtual FArchive& operator<<(FObjectPtr& Value) override { unimplemented(); return *this; }
+	virtual FArchive& operator<<(FSoftObjectPtr& Value) override { unimplemented(); return *this; }
+	virtual FArchive& operator<<(FSoftObjectPath& Value) override { unimplemented(); return *this; }
+	virtual FArchive& operator<<(FWeakObjectPtr& Value) override { unimplemented(); return *this; }
+	//~ End FArchive Interface
+};
+#endif
\ No newline at end of file
