Remove unnecessary accessor function on gameplay statics

#rnx

diff --git a/Engine/Source/Runtime/Engine/Classes/Kismet/GameplayStatics.h b/Engine/Source/Runtime/Engine/Classes/Kismet/GameplayStatics.h
--- a/Engine/Source/Runtime/Engine/Classes/Kismet/GameplayStatics.h
+++ b/Engine/Source/Runtime/Engine/Classes/Kismet/GameplayStatics.h
@@ -271,15 +271,6 @@ class ENGINE_API UGameplayStatics : public UBlueprintFunctionLibrary
 	UFUNCTION(BlueprintPure, Category="Game", meta=(DisplayName="Get Local Player Controller ID", UnsafeDuringActorConstruction="true"))
 	static int32 GetPlayerControllerID(APlayerController* Player);
 
-	/**
-	 * Gets what Platform User Id a player is using. This only works for local player controllers.
-	 *
-	 * @param Player	The player controller of the player to get the ID of
-	 * @return			The Platform User Id of the passed in player. PLATFORMUSERID_NONE if there is no physical controller assigned to the passed in player
-	 */
-	UFUNCTION(BlueprintPure, Category="Game", meta=(DisplayName="Get Local Player Platform User Id", UnsafeDuringActorConstruction="true"))
-	static FPlatformUserId GetPlayerPlatformUserId(const APlayerController* PlayerController);
-
 	/**
 	 * Sets what physical controller ID a player should be using. This only works for local player controllers.
 	 *
diff --git a/Engine/Source/Runtime/Engine/Private/GameplayStatics.cpp b/Engine/Source/Runtime/Engine/Private/GameplayStatics.cpp
--- a/Engine/Source/Runtime/Engine/Private/GameplayStatics.cpp
+++ b/Engine/Source/Runtime/Engine/Private/GameplayStatics.cpp
@@ -394,7 +394,7 @@ APlayerController* UGameplayStatics::GetPlayerControllerFromPlatformUser(const U
 		for (FConstPlayerControllerIterator Iterator = World->GetPlayerControllerIterator(); Iterator; ++Iterator)
 		{
 			APlayerController* PlayerController = Iterator->Get();
-			FPlatformUserId PlayerControllerUserID = GetPlayerPlatformUserId(PlayerController);
+			FPlatformUserId PlayerControllerUserID = PlayerController->GetPlatformUserId();
 			if (PlayerControllerUserID.IsValid() && PlayerControllerUserID == UserId)
 			{
 				return PlayerController;
@@ -462,23 +462,10 @@ void UGameplayStatics::RemovePlayer(APlayerController* PlayerController, bool bD
 
 int32 UGameplayStatics::GetPlayerControllerID(APlayerController* PlayerController)
 {
-	FPlatformUserId UserID = GetPlayerPlatformUserId(PlayerController);
+	FPlatformUserId UserID = PlayerController ? PlayerController->GetPlatformUserId() : PLATFORMUSERID_NONE;
 	return FGenericPlatformMisc::GetUserIndexForPlatformUser(UserID);
 }
 
-FPlatformUserId UGameplayStatics::GetPlayerPlatformUserId(const APlayerController* PlayerController)
-{
-	if (PlayerController)
-	{
-		if (const ULocalPlayer* LocalPlayer = PlayerController->GetLocalPlayer())
-		{
-			return LocalPlayer->GetPlatformUserId();
-		}
-	}
-
-	return PLATFORMUSERID_NONE;
-}
-
 void UGameplayStatics::SetPlayerControllerID(APlayerController* PlayerController, int32 ControllerId)
 {
 	SetPlayerPlatformUserId(PlayerController, FPlatformMisc::GetPlatformUserForUserIndex(ControllerId));
