Fixed not being able to summon VisualLog while in `-game` mode. Still requires ENABLE_VISUAL_LOG being true

diff --git a/Engine/Source/Runtime/Engine/Private/VisualLogger/VisualLogger.cpp b/Engine/Source/Runtime/Engine/Private/VisualLogger/VisualLogger.cpp
--- a/Engine/Source/Runtime/Engine/Private/VisualLogger/VisualLogger.cpp
+++ b/Engine/Source/Runtime/Engine/Private/VisualLogger/VisualLogger.cpp
@@ -908,13 +908,11 @@ public:
 				FVisualLogger::Get().AddCategoryToAllowList(*Category);
 				return true;
 			}
-#if WITH_EDITOR
-			else if (GIsEditor)
+			else if (FModuleManager::Get().LoadModulePtr<IModuleInterface>("LogVisualizer"))
 			{
 				FGlobalTabmanager::Get()->TryInvokeTab(FName(TEXT("VisualLogger")));
 				return true;
 			}
-#endif // WITH_EDITOR
 		}
 		else if (FParse::Command(&Cmd, TEXT("LogNavOctree")))
 		{
