Properly check for platform data availability when releasing texture resource.

diff --git a/Engine/Source/Runtime/Engine/Private/Texture.cpp b/Engine/Source/Runtime/Engine/Private/Texture.cpp
--- a/Engine/Source/Runtime/Engine/Private/Texture.cpp
+++ b/Engine/Source/Runtime/Engine/Private/Texture.cpp
@@ -211,7 +211,7 @@ void UTexture::ReleaseResource()
 
 		// When using PlatformData, the resource shouldn't be released before it is initialized to prevent threading issues
 		// where the platform data could be updated at the same time InitRHI is reading it on the renderthread.
-		if (GetRunningPlatformData())
+		if (GetRunningPlatformData() && *GetRunningPlatformData())
 		{
 			WaitForPendingInitOrStreaming();
 		}
