Provide access to mip counts prior to texture encode launch.

diff --git a/Engine/Source/Developer/TextureCompressor/Private/TextureCompressorModule.cpp b/Engine/Source/Developer/TextureCompressor/Private/TextureCompressorModule.cpp
--- a/Engine/Source/Developer/TextureCompressor/Private/TextureCompressorModule.cpp
+++ b/Engine/Source/Developer/TextureCompressor/Private/TextureCompressorModule.cpp
@@ -1773,9 +1773,9 @@ static FVector ComputeWSCubeDirectionAtTexelCenter(uint32 CubemapFace, uint32 x,
 	return DirectionWS;
 }
 
-static uint32 ComputeLongLatCubemapExtents(const FImage& SrcImage, const uint32 MaxCubemapTextureResolution)
+static uint32 ComputeLongLatCubemapExtents(int32 SrcImageSizeX, const uint32 MaxCubemapTextureResolution)
 {
-	return FMath::Clamp(1U << FMath::FloorLog2(SrcImage.SizeX / 2), 32U, MaxCubemapTextureResolution);
+	return FMath::Clamp(1U << FMath::FloorLog2(SrcImageSizeX / 2), 32U, MaxCubemapTextureResolution);
 }
 
 void ITextureCompressorModule::GenerateBaseCubeMipFromLongitudeLatitude2D(FImage* OutMip, const FImage& SrcImage, const uint32 MaxCubemapTextureResolution, uint8 SourceEncodingOverride)
@@ -1786,7 +1786,7 @@ void ITextureCompressorModule::GenerateBaseCubeMipFromLongitudeLatitude2D(FImage
 	SrcImage.Linearize(SourceEncodingOverride, LongLatImage);
 
 	// TODO_TEXTURE: Expose target size to user.
-	uint32 Extent = ComputeLongLatCubemapExtents(LongLatImage, MaxCubemapTextureResolution);
+	uint32 Extent = ComputeLongLatCubemapExtents(LongLatImage.SizeX, MaxCubemapTextureResolution);
 	float InvExtent = 1.0f / Extent;
 	OutMip->Init(Extent, Extent, SrcImage.NumSlices * 6, ERawImageFormat::RGBA32F, EGammaSpace::Linear);
 
@@ -3026,6 +3026,79 @@ static void NormalizeMip(FImage& InOutMip)
 	}
 }
 
+
+// Special case for TMGS_LeaveExistingMips
+static int32 GetMipCountForLeaveExistingMips(int32 InMip0SizeX, int32 InMip0SizeY, int32 InExistingMipCount, uint32 InMaxTexture2DResolution)
+{
+	int32 i = 0;
+	for (; i < InExistingMipCount; i++)
+	{
+		uint32 MipSizeX = FMath::Max<uint32>(1, InMip0SizeX >> i);
+		uint32 MipSizeY = FMath::Max<uint32>(1, InMip0SizeY >> i);
+
+		if (MipSizeX <= InMaxTexture2DResolution &&
+			MipSizeY <= InMaxTexture2DResolution)
+		{
+			return InExistingMipCount - i;
+		}
+	}
+
+	// Couldn't find a fit, texture build will fail.
+	check(0);
+	return 0;
+}
+
+
+// Returns true if the target texture size is different and padding/stretching is required.
+static bool GetPowerOfTwoTargetTextureSize(int32 InMip0SizeX, int32 InMip0SizeY, int32 InMip0NumSlices, bool bInIsVolume, ETexturePowerOfTwoSetting::Type InPow2Setting, int32& OutTargetSizeX, int32& OutTargetSizeY, int32& OutTargetSizeZ)
+{
+	check(InPow2Setting != ETexturePowerOfTwoSetting::None);
+
+	int32 TargetTextureSizeX = InMip0SizeX;
+	int32 TargetTextureSizeY = InMip0SizeY;
+	int32 TargetTextureSizeZ = bInIsVolume ? InMip0NumSlices : 1; // Only used for volume texture.
+
+	const int32 PowerOfTwoTextureSizeX = FMath::RoundUpToPowerOfTwo(TargetTextureSizeX);
+	const int32 PowerOfTwoTextureSizeY = FMath::RoundUpToPowerOfTwo(TargetTextureSizeY);
+	const int32 PowerOfTwoTextureSizeZ = FMath::RoundUpToPowerOfTwo(TargetTextureSizeZ);
+
+	switch (InPow2Setting)
+	{
+	// None should not get here
+
+	case ETexturePowerOfTwoSetting::PadToPowerOfTwo:
+		TargetTextureSizeX = PowerOfTwoTextureSizeX;
+		TargetTextureSizeY = PowerOfTwoTextureSizeY;
+		TargetTextureSizeZ = PowerOfTwoTextureSizeZ;
+		break;
+
+	case ETexturePowerOfTwoSetting::PadToSquarePowerOfTwo:
+		TargetTextureSizeX = TargetTextureSizeY = TargetTextureSizeZ =
+			FMath::Max3<int32>(PowerOfTwoTextureSizeX, PowerOfTwoTextureSizeY, PowerOfTwoTextureSizeZ);
+		break;
+
+	default:
+		checkf(false, TEXT("Unknown entry in ETexturePowerOfTwoSetting::Type"));
+		break;
+	}
+
+	// Z only matters as a sampling dimension if we are a volume texture.
+	if (bInIsVolume == false)
+	{
+		TargetTextureSizeZ = InMip0NumSlices;
+	}
+
+	OutTargetSizeX = TargetTextureSizeX;
+	OutTargetSizeY = TargetTextureSizeY;
+	OutTargetSizeZ = TargetTextureSizeZ;
+
+	return (TargetTextureSizeX != InMip0SizeX) ||
+		(TargetTextureSizeY != InMip0SizeY) ||
+		(bInIsVolume && TargetTextureSizeZ != InMip0NumSlices);
+}
+
+
+
 /**
  * Texture compression module
  */
@@ -3036,6 +3109,85 @@ public:
 	{
 	}
 
+	virtual int32 GetMipCountForBuildSettings(int32 InMip0SizeX, int32 InMip0SizeY, int32 InMip0NumSlices, int32 InExistingMipCount, const FTextureBuildSettings& BuildSettings) const override
+	{
+		if (BuildSettings.MipGenSettings == TMGS_LeaveExistingMips)
+		{
+			// Since we can't generate, we only have to limit to MaxTextureSize
+			return GetMipCountForLeaveExistingMips(InMip0SizeX, InMip0SizeY, InExistingMipCount, BuildSettings.MaxTextureResolution);
+		}
+		else if (BuildSettings.MipGenSettings == TMGS_NoMipmaps)
+		{
+			return 1;
+		}
+
+		// AFAICT LatLongCubeMaps don't do any of this - pow2 is broken with them but it runs, and max texture stuff
+		// is handled internally in the extents function.
+
+		int32 BaseSizeX = InMip0SizeX;
+		int32 BaseSizeY = InMip0SizeY;
+		int32 BaseSizeZ = BuildSettings.bVolume ? InMip0NumSlices : 1; // Volume textures are the only type that mip their Z, arrays and cubes are fixed.
+
+		ETexturePowerOfTwoSetting::Type PowerOfTwoMode = (ETexturePowerOfTwoSetting::Type)BuildSettings.PowerOfTwoMode;
+		if (PowerOfTwoMode != ETexturePowerOfTwoSetting::None)
+		{
+			int32 TargetSizeX, TargetSizeY, TargetSizeZ;
+			bool NeedsAdjustment = GetPowerOfTwoTargetTextureSize(BaseSizeX, BaseSizeY, BaseSizeY, BuildSettings.bVolume, PowerOfTwoMode, TargetSizeX, TargetSizeY, TargetSizeZ);
+			if (NeedsAdjustment)
+			{
+				// In this case we are regenerating the entire mip chain.
+				InExistingMipCount = 1;
+				BaseSizeX = TargetSizeX;
+				BaseSizeY = TargetSizeY;
+				BaseSizeZ = TargetSizeZ; // volume textures already accounted for
+			}
+			// Otherwise we have valid pow2 so we can reuse any existing mips and regenerate
+			// any missing tail mips.
+		}
+
+		// LatLong sources are clamped in ComputeLongLatCubemapExtents
+		if (BuildSettings.bLongLatSource == false)
+		{
+			// Max texture resolution strips off mips that are above the limit.
+			int64 MaxTextureResolution = BuildSettings.MaxTextureResolution;
+
+			uint32 GeneratedMaxMipDimension = FMath::Max3(BaseSizeX, BaseSizeY, BaseSizeZ);
+			int32 GeneratedMipCount = 1 + FMath::FloorLog2(GeneratedMaxMipDimension);
+			int32 i = 0;
+			for (; i < GeneratedMipCount; i++)
+			{
+				// The code in BuildTextureMips doesn't worry about fitting Z in volume textures...
+				// \todo volume texture MaxTextureSize. The old code ignored size Z, so we do to. I'm not sure
+				// there's ever a case where volume textures have a Z that's bigger than X/Y.
+				int32 MipSizeX = FMath::Max<uint32>(1, BaseSizeX >> i);
+				int32 MipSizeY = FMath::Max<uint32>(1, BaseSizeY >> i);
+				int32 MipSizeZ = BuildSettings.bVolume ? FMath::Max<uint32>(1, BaseSizeZ >> i) : BaseSizeZ;
+
+				if (MipSizeX <= MaxTextureResolution &&
+					MipSizeY <= MaxTextureResolution)
+				{
+					BaseSizeX = MipSizeX;
+					BaseSizeY = MipSizeY;
+					BaseSizeZ = MipSizeZ;
+					break;
+				}
+			}
+		}
+
+		// At this point we have a base mip size that is valid.
+
+		// Relevant BuildSettings.MipGenSettings have been handled at top of function.
+
+		// NumOutputMips is the number of mips that would be made if you made a full mip chain
+		//  eg. 256 makes 9 mips , 300 also makes 9 mips
+		uint32 MaxMipDimension = FMath::Max3(BaseSizeX, BaseSizeY, BaseSizeZ);
+		if (BuildSettings.bLongLatSource)
+		{
+			MaxMipDimension = ComputeLongLatCubemapExtents(BaseSizeX, BuildSettings.MaxTextureResolution);
+		}
+		return 1 + FMath::FloorLog2(MaxMipDimension);
+	}
+
 	virtual bool BuildTexture(
 		const TArray<FImage>& SourceMips,
 		const TArray<FImage>& AssociatedNormalSourceMips,
@@ -3102,7 +3254,7 @@ public:
 		// allow to leave texture in sRGB in case compressor accepts other than non-F32 input source
 		// otherwise linearizing will force format to be RGBA32F
 		const bool bNeedLinearize = !TextureFormat->CanAcceptNonF32Source() || AssociatedNormalSourceMips.Num() != 0;
-		if (!BuildTextureMips(SourceMips, BuildSettings, CompressorCaps, bNeedLinearize, IntermediateMipChain))
+		if (!BuildTextureMips(SourceMips, BuildSettings, CompressorCaps, bNeedLinearize, IntermediateMipChain, DebugTexturePathName))
 		{
 			return false;
 		}
@@ -3135,7 +3287,7 @@ public:
 			//  we should instead compute the roughness scalar first on the original normap map
 			//  then filter on the roughness scalar
 
-			if (!BuildTextureMips(AssociatedNormalSourceMips, DefaultSettings, CompressorCaps, true, IntermediateAssociatedNormalSourceMipChain))
+			if (!BuildTextureMips(AssociatedNormalSourceMips, DefaultSettings, CompressorCaps, true, IntermediateAssociatedNormalSourceMipChain, DebugTexturePathName))
 			{
 				UE_LOG(LogTextureCompressor, Warning, TEXT("Failed to generate texture mips for composite texture"));
 			}
@@ -3196,12 +3348,14 @@ public:
 
 private:
 
+
 	bool BuildTextureMips(
 		const TArray<FImage>& InSourceMipChain,
 		const FTextureBuildSettings& BuildSettings,
 		const FTextureFormatCompressorCaps& CompressorCaps,
 		const bool bNeedLinearize,
-		TArray<FImage>& OutMipChain)
+		TArray<FImage>& OutMipChain,
+		FStringView DebugTexturePathName)
 	{
 		TRACE_CPUPROFILER_EVENT_SCOPE(Texture.BuildTextureMips);
 		
@@ -3242,38 +3396,13 @@ private:
 		{
 			const FImage& FirstSourceMipImage = (*pSourceMips)[0];
 
-			int32 TargetTextureSizeX = FirstSourceMipImage.SizeX;
-			int32 TargetTextureSizeY = FirstSourceMipImage.SizeY;
-			int32 TargetTextureSizeZ = BuildSettings.bVolume ? FirstSourceMipImage.NumSlices : 1; // Only used for volume texture.
-
-			const int32 PowerOfTwoTextureSizeX = FMath::RoundUpToPowerOfTwo(TargetTextureSizeX);
-			const int32 PowerOfTwoTextureSizeY = FMath::RoundUpToPowerOfTwo(TargetTextureSizeY);
-			const int32 PowerOfTwoTextureSizeZ = FMath::RoundUpToPowerOfTwo(TargetTextureSizeZ);
-
-			switch (PowerOfTwoMode)
-			{
-			// None should not get here
-
-			case ETexturePowerOfTwoSetting::PadToPowerOfTwo:
-				TargetTextureSizeX = PowerOfTwoTextureSizeX;
-				TargetTextureSizeY = PowerOfTwoTextureSizeY;
-				TargetTextureSizeZ = PowerOfTwoTextureSizeZ;
-				break;
-
-			case ETexturePowerOfTwoSetting::PadToSquarePowerOfTwo:
-				TargetTextureSizeX = TargetTextureSizeY = TargetTextureSizeZ =
-					FMath::Max3<int32>(PowerOfTwoTextureSizeX, PowerOfTwoTextureSizeY, PowerOfTwoTextureSizeZ);
-				break;
-
-			default:
-				checkf(false, TEXT("Unknown entry in ETexturePowerOfTwoSetting::Type"));
-				break;
-			}
-			
-			bool bPadOrStretchTexture = 
-				(TargetTextureSizeX != FirstSourceMipImage.SizeX) ||
-				(TargetTextureSizeY != FirstSourceMipImage.SizeY) ||
-				( BuildSettings.bVolume && TargetTextureSizeZ != FirstSourceMipImage.NumSlices );
+			int32 TargetTextureSizeX = 0;
+			int32 TargetTextureSizeY = 0;
+			int32 TargetTextureSizeZ = 0;			
+			bool bPadOrStretchTexture = GetPowerOfTwoTargetTextureSize(
+				FirstSourceMipImage.SizeX, FirstSourceMipImage.SizeY, FirstSourceMipImage.NumSlices,
+				BuildSettings.bVolume, PowerOfTwoMode,
+				TargetTextureSizeX, TargetTextureSizeY, TargetTextureSizeZ);
 
 			if (bPadOrStretchTexture)
 			{
@@ -3466,7 +3595,7 @@ private:
 			//  eg. 256 makes 9 mips , 300 also makes 9 mips
 			NumOutputMips = 1 + FMath::FloorLog2(
 				bLongLatCubemap ?
-				ComputeLongLatCubemapExtents(TopMip, BuildSettings.MaxTextureResolution) :
+				ComputeLongLatCubemapExtents(TopMip.SizeX, BuildSettings.MaxTextureResolution) :
 				FMath::Max3(TopMip.SizeX, TopMip.SizeY, TopMipSizeZ) );
 				
 			// unless LeaveExistingMips, we only copy 1 
@@ -3607,6 +3736,12 @@ private:
 		}
 		check(OutMipChain.Num() == NumOutputMips);
 
+		int32 CalculatedMipCount = GetMipCountForBuildSettings(InSourceMipChain[0].SizeX, InSourceMipChain[0].SizeY, InSourceMipChain[0].NumSlices, InSourceMipChain.Num(), BuildSettings);
+		if (CalculatedMipCount != NumOutputMips)
+		{
+			UE_LOG(LogTextureCompressor, Error, TEXT("Texture %.*s generated %d mips when GetMipCountForBuildSettings expected %d!"), DebugTexturePathName.Len(), DebugTexturePathName.GetData(), NumOutputMips, CalculatedMipCount);
+		}
+
 		// Apply post-mip generation adjustments.
 		if (BuildSettings.bReplicateRed)
 		{
diff --git a/Engine/Source/Developer/TextureCompressor/Public/TextureCompressorModule.h b/Engine/Source/Developer/TextureCompressor/Public/TextureCompressorModule.h
--- a/Engine/Source/Developer/TextureCompressor/Public/TextureCompressorModule.h
+++ b/Engine/Source/Developer/TextureCompressor/Public/TextureCompressorModule.h
@@ -367,4 +367,9 @@ public:
 	 * @param DiffuseConvolveMipLevel - The mip level that contains the diffuse convolution.
 	 */
 	TEXTURECOMPRESSOR_API static void GenerateAngularFilteredMips(TArray<FImage>& InOutMipChain, int32 NumMips, uint32 DiffuseConvolveMipLevel);
+
+	/**
+	* Returns the number of mips that the given texture will generate with the given build settings.
+	*/
+	virtual int32 GetMipCountForBuildSettings(int32 InMip0SizeX, int32 InMip0SizeY, int32 InMip0NumSlices, int32 InExistingMipCount, const FTextureBuildSettings& InBuildSettings) const =0;
 };
diff --git a/Engine/Source/Runtime/Engine/Private/TextureDerivedDataTask.cpp b/Engine/Source/Runtime/Engine/Private/TextureDerivedDataTask.cpp
--- a/Engine/Source/Runtime/Engine/Private/TextureDerivedDataTask.cpp
+++ b/Engine/Source/Runtime/Engine/Private/TextureDerivedDataTask.cpp
@@ -51,11 +51,13 @@
 #include "Misc/ScopeExit.h"
 #include "ProfilingDebugging/CookStats.h"
 #include "Serialization/BulkDataRegistry.h"
+#include "TextureCompressorModule.h"
 #include "TextureDerivedDataBuildUtils.h"
+#include "TextureEncodingSettings.h"
 #include "VT/VirtualTextureChunkDDCCache.h"
 #include "VT/VirtualTextureDataBuilder.h"
-#include "TextureEncodingSettings.h"
 
+
 static TAutoConsoleVariable<int32> CVarVTValidateCompressionOnLoad(
 	TEXT("r.VT.ValidateCompressionOnLoad"),
 	0,
@@ -429,6 +431,13 @@
 		DerivedData->PixelFormat = DerivedData->VTData->LayerTypes[0];
 		DerivedData->SetNumSlices(1);
 
+		{
+			// Verify our predicted count matches.
+			FVirtualTextureBuilderDerivedInfo Info;
+			Info.InitializeFromBuildSettings(TextureData, InBuildSettingsPerLayer.GetData());
+			check(Info.NumMips == DerivedData->VTData->GetNumMips());
+		}
+
 		bool bCompressionValid = true;
 		if (CVarVTValidateCompressionOnSave.GetValueOnAnyThread())
 		{
@@ -561,7 +570,15 @@
 			}
 
 			DerivedData->SetOptData(OptData);
-				
+
+			{
+				// Verify our predicted count matches.
+				const TArray<FImage>& SourceMips = TextureData.Blocks[0].MipsPerLayer[0];
+				int32 NumMips = Compressor->GetMipCountForBuildSettings(SourceMips[0].SizeX, SourceMips[0].SizeY, SourceMips[0].NumSlices, SourceMips.Num(), InBuildSettingsPerLayer[0]);
+				check (NumMips == CompressedMips.Num());
+			}
+
+
 			// Store it in the cache.
 			// @todo: This will remove the streaming bulk data, which we immediately reload below!
 			// Should ideally avoid this redundant work, but it only happens when we actually have 
diff --git a/Engine/Source/Runtime/Engine/Private/VT/VirtualTextureDataBuilder.cpp b/Engine/Source/Runtime/Engine/Private/VT/VirtualTextureDataBuilder.cpp
--- a/Engine/Source/Runtime/Engine/Private/VT/VirtualTextureDataBuilder.cpp
+++ b/Engine/Source/Runtime/Engine/Private/VT/VirtualTextureDataBuilder.cpp
@@ -222,13 +222,6 @@ struct FPixelDataRectangle
 
 FVirtualTextureDataBuilder::FVirtualTextureDataBuilder(FVirtualTextureBuiltData &SetOutData, const FString& InDebugTexturePathName, ITextureCompressorModule *InCompressor, IImageWrapperModule* InImageWrapper)
 	: OutData(SetOutData)
-	, SizeInBlocksX(0)
-	, SizeInBlocksY(0)
-	, BlockSizeX(0)
-	, BlockSizeY(0)
-	, BlockSizeScale(1)
-	, SizeX(0)
-	, SizeY(0)
 	, DebugTexturePathName(InDebugTexturePathName)
 {
 	Compressor = InCompressor ? InCompressor : &FModuleManager::LoadModuleChecked<ITextureCompressorModule>(TEXTURE_COMPRESSOR_MODULENAME);
@@ -240,22 +233,20 @@ FVirtualTextureDataBuilder::~FVirtualTextureDataBuilder()
 	FreeSourcePixels();
 }
 
-void FVirtualTextureDataBuilder::Build(const FTextureSourceData& InSourceData, const FTextureSourceData& InCompositeSourceData, const FTextureBuildSettings* InSettingsPerLayer, bool bAllowAsync)
+bool FVirtualTextureBuilderDerivedInfo::InitializeFromBuildSettings(const FTextureSourceData& InSourceData, const FTextureBuildSettings* InSettingsPerLayer)
 {
 	const int32 NumLayers = InSourceData.Layers.Num();
 	checkf(NumLayers <= (int32)VIRTUALTEXTURE_DATA_MAXLAYERS, TEXT("The maximum amount of layers is exceeded."));
 	checkf(NumLayers > 0, TEXT("No layers to build."));
 
-	SettingsPerLayer.AddUninitialized(NumLayers);
-	FMemory::Memcpy(&SettingsPerLayer[0], InSettingsPerLayer, sizeof(FTextureBuildSettings) * NumLayers);
-	const FTextureBuildSettings& BuildSettingsLayer0 = SettingsPerLayer[0];
+	const FTextureBuildSettings& BuildSettingsLayer0 = InSettingsPerLayer[0];
 	const int32 TileSize = BuildSettingsLayer0.VirtualTextureTileSize;
 
 	BlockSizeX = InSourceData.BlockSizeX;
 	BlockSizeY = InSourceData.BlockSizeY;
 
 	// BlockSize is potentially adjusted by rounding to power of 2
-	switch (SettingsPerLayer[0].PowerOfTwoMode)
+	switch (BuildSettingsLayer0.PowerOfTwoMode)
 	{
 	case ETexturePowerOfTwoSetting::None:
 		break;
@@ -274,10 +265,10 @@ void FVirtualTextureDataBuilder::Build(const FTextureSourceData& InSourceData, c
 		break;
 	}
 
-	check(SettingsPerLayer[0].MaxTextureResolution >= (uint32)TileSize);
-	
+	check(InSettingsPerLayer[0].MaxTextureResolution >= (uint32)TileSize);
+
 	// Clamp BlockSizeX and BlockSizeY to MaxTextureResolution, but don't change aspect ratio
-	const uint32 ClampBlockSize = SettingsPerLayer[0].MaxTextureResolution;
+	const uint32 ClampBlockSize = InSettingsPerLayer[0].MaxTextureResolution;
 	if (FMath::Max<uint32>(BlockSizeX, BlockSizeY) > ClampBlockSize)
 	{
 		const int32 ClampedBlockSizeX = BlockSizeX >= BlockSizeY ? ClampBlockSize : FMath::Max(ClampBlockSize * BlockSizeX / BlockSizeY, 1u);
@@ -304,14 +295,39 @@ void FVirtualTextureDataBuilder::Build(const FTextureSourceData& InSourceData, c
 	SizeX = BlockSizeX * SizeInBlocksX;
 	SizeY = BlockSizeY * SizeInBlocksY;
 
+	const uint32 Size = FMath::Max(SizeX, SizeY);
+
+	// Mip down to 1x1 pixels, but don't create more mips than can fit in max page table texture
+
+	// This is incorrect - should be floor not ceil. Testing _so far_ has found there to be no difference noted
+	// in projects, however saving this work to a separate CL since this is old and presumably stable code, due to
+	// the fact that the Min(x, VIRTUALTEXTURE_LOG2_MAX_PAGETABLE_SIZE) is load bearing. In order to get an incorrect
+	// mip count, you need Size to be non pow2 and the result to be larger than VIRTUALTEXTURE_LOG2_MAX_PAGETABLE_SIZE.
+	NumMips = FMath::Min<uint32>(FMath::CeilLogTwo(Size) + 1, VIRTUALTEXTURE_LOG2_MAX_PAGETABLE_SIZE);
+	return true;
+}
+
+void FVirtualTextureDataBuilder::Build(const FTextureSourceData& InSourceData, const FTextureSourceData& InCompositeSourceData, const FTextureBuildSettings* InSettingsPerLayer, bool bAllowAsync)
+{
+	const int32 NumLayers = InSourceData.Layers.Num();
+	checkf(NumLayers <= (int32)VIRTUALTEXTURE_DATA_MAXLAYERS, TEXT("The maximum amount of layers is exceeded."));
+	checkf(NumLayers > 0, TEXT("No layers to build."));
+
+	SettingsPerLayer.AddUninitialized(NumLayers);
+	FMemory::Memcpy(&SettingsPerLayer[0], InSettingsPerLayer, sizeof(FTextureBuildSettings) * NumLayers);
+	const FTextureBuildSettings& BuildSettingsLayer0 = SettingsPerLayer[0];
+	const int32 TileSize = BuildSettingsLayer0.VirtualTextureTileSize;
+
+	DerivedInfo.InitializeFromBuildSettings(InSourceData, InSettingsPerLayer);
+
 	//NOTE: OutData may point to a previously build data so it is important to
 	//properly initialize all fields and not assume this is a freshly constructed object
 
 	OutData.TileBorderSize = BuildSettingsLayer0.VirtualTextureBorderSize;
 	OutData.TileSize = TileSize;
 	OutData.NumLayers = NumLayers;
-	OutData.Width = SizeX;
-	OutData.Height = SizeY;
+	OutData.Width = DerivedInfo.SizeX;
+	OutData.Height = DerivedInfo.SizeY;
 	OutData.WidthInBlocks = InSourceData.SizeInBlocksX;
 	OutData.HeightInBlocks = InSourceData.SizeInBlocksY;
 
@@ -325,14 +341,7 @@ void FVirtualTextureDataBuilder::Build(const FTextureSourceData& InSourceData, c
 	OutData.TileOffsetInChunk.Empty();
 
 	OutData.Chunks.Empty();
-
-	const uint32 Size = FMath::Max(SizeX, SizeY);
-	const uint32 SizeInTiles = FMath::DivideAndRoundUp<uint32>(Size, TileSize);
-	const uint32 BlockSize = FMath::Max(BlockSizeX, BlockSizeY);
-	const uint32 BlockSizeInTiles = FMath::DivideAndRoundUp<uint32>(BlockSize, TileSize);
-
-	// Mip down to 1x1 pixels, but don't create more mips than can fit in max page table texture
-	OutData.NumMips = FMath::Min<uint32>(FMath::CeilLogTwo(Size) + 1, VIRTUALTEXTURE_LOG2_MAX_PAGETABLE_SIZE);
+	OutData.NumMips = DerivedInfo.NumMips;
 
 	BuildSourcePixels(InSourceData, InCompositeSourceData);
 
@@ -388,11 +397,11 @@ void FVirtualTextureDataBuilder::BuildPagesMacroBlocks(bool bAllowAsync)
 	const int32 TileSize = SettingsPerLayer[0].VirtualTextureTileSize;
 	const uint32 MinSizePerChunkInTiles = FMath::DivideAndRoundUp<uint32>(MinSizePerChunk, TileSize);
 	const uint32 MinTilesPerChunk = MinSizePerChunkInTiles * MinSizePerChunkInTiles;
-	const int32 BlockSizeInTilesX = FMath::DivideAndRoundUp(BlockSizeX, TileSize);
-	const int32 BlockSizeInTilesY = FMath::DivideAndRoundUp(BlockSizeY, TileSize);
+	const int32 BlockSizeInTilesX = FMath::DivideAndRoundUp(DerivedInfo.BlockSizeX, TileSize);
+	const int32 BlockSizeInTilesY = FMath::DivideAndRoundUp(DerivedInfo.BlockSizeY, TileSize);
 
-	uint32 MipWidthInTiles = FMath::DivideAndRoundUp(SizeX, TileSize);
-	uint32 MipHeightInTiles = FMath::DivideAndRoundUp(SizeY, TileSize);
+	uint32 MipWidthInTiles = FMath::DivideAndRoundUp(DerivedInfo.SizeX, TileSize);
+	uint32 MipHeightInTiles = FMath::DivideAndRoundUp(DerivedInfo.SizeY, TileSize);
 	uint32 NumTiles = 0u;
 
 	for (uint32 Mip = 0; Mip < OutData.NumMips; ++Mip)
@@ -423,8 +432,8 @@ void FVirtualTextureDataBuilder::BuildPagesMacroBlocks(bool bAllowAsync)
 
 		OutData.TileIndexPerChunk.Add(TileIndex);
 
-		MipWidthInTiles = FMath::DivideAndRoundUp(SizeX, TileSize);
-		MipHeightInTiles = FMath::DivideAndRoundUp(SizeY, TileSize);
+		MipWidthInTiles = FMath::DivideAndRoundUp(DerivedInfo.SizeX, TileSize);
+		MipHeightInTiles = FMath::DivideAndRoundUp(DerivedInfo.SizeY, TileSize);
 		for (uint32 Mip = 0; Mip < OutData.NumMips; ++Mip)
 		{
 			FVirtualTextureTileOffsetData& OffsetData = OutData.TileOffsetData.AddDefaulted_GetRef();
@@ -886,7 +895,7 @@ void FVirtualTextureDataBuilder::BuildSourcePixels(const FTextureSourceData& Sou
 		BlockData.BlockX = SourceBlockData.BlockX;
 		// UE applies a (1-y) transform to imported UVs, so apply a similar transform to UDIM block locations here
 		// This ensures that UDIM tiles will appear in the correct location when sampled with transformed UVs
-		BlockData.BlockY = (SizeInBlocksY - SourceBlockData.BlockY) % SizeInBlocksY;
+		BlockData.BlockY = (DerivedInfo.SizeInBlocksY - SourceBlockData.BlockY) % DerivedInfo.SizeInBlocksY;
 		BlockData.NumMips = SourceBlockData.NumMips;
 		BlockData.NumSlices = SourceBlockData.NumSlices;
 		BlockData.MipBias = SourceBlockData.MipBias;
@@ -945,7 +954,7 @@ void FVirtualTextureDataBuilder::BuildSourcePixels(const FTextureSourceData& Sou
 
 			// For multi-block images, we may have scaled the max block size to be tile-sized, but individual blocks may still be smaller than 1 tile
 			// These need to be scaled up as well (scaling up individual blocks has the effect of reducing the block's mip-bias)
-			int32 LocalBlockSizeScale = BlockSizeScale;
+			int32 LocalBlockSizeScale = DerivedInfo.BlockSizeScale;
 			while (SourceMips[0].SizeX * LocalBlockSizeScale < TileSize || SourceMips[0].SizeY * LocalBlockSizeScale < TileSize)
 			{
 				check(BlockData.MipBias > 0u);
@@ -997,14 +1006,14 @@ void FVirtualTextureDataBuilder::BuildSourcePixels(const FTextureSourceData& Sou
 				BlockData.SizeY = CompressedMips[0].SizeY;
 
 				// re-compute mip bias to account for any resizing of this block (typically due to clamped max size)
-				const int32 MipBiasX = FMath::CeilLogTwo(BlockSizeX / BlockData.SizeX);
-				const int32 MipBiasY = FMath::CeilLogTwo(BlockSizeY / BlockData.SizeY);
-				checkf(MipBiasX == MipBiasY, TEXT("Mismatched aspect ratio (%d x %d), (%d x %d)"), BlockSizeX, BlockSizeY, BlockData.SizeX, BlockData.SizeY);
+				const int32 MipBiasX = FMath::CeilLogTwo(DerivedInfo.BlockSizeX / BlockData.SizeX);
+				const int32 MipBiasY = FMath::CeilLogTwo(DerivedInfo.BlockSizeY / BlockData.SizeY);
+				checkf(MipBiasX == MipBiasY, TEXT("Mismatched aspect ratio (%d x %d), (%d x %d)"), DerivedInfo.BlockSizeX, DerivedInfo.BlockSizeY, BlockData.SizeX, BlockData.SizeY);
 				BlockData.MipBias = MipBiasX;
 			}
 
-			check(BlockData.SizeX << BlockData.MipBias == BlockSizeX);
-			check(BlockData.SizeY << BlockData.MipBias == BlockSizeY);
+			check(BlockData.SizeX << BlockData.MipBias == DerivedInfo.BlockSizeX);
+			check(BlockData.SizeY << BlockData.MipBias == DerivedInfo.BlockSizeY);
 
 			const uint32 BlockSize = FMath::Max(BlockData.SizeX, BlockData.SizeY);
 			if (NumBlocks == 1u)
@@ -1045,21 +1054,21 @@ void FVirtualTextureDataBuilder::BuildSourcePixels(const FTextureSourceData& Sou
 	// If we have more than 1 block, need to create miptail that contains mips made from multiple blocks
 	if (NumBlocks > 1)
 	{
-		const uint32 BlockSize = FMath::Max(BlockSizeX, BlockSizeY);
+		const uint32 BlockSize = FMath::Max(DerivedInfo.BlockSizeX, DerivedInfo.BlockSizeY);
 		const uint32 BlockSizeInTiles = FMath::DivideAndRoundUp<uint32>(BlockSize, TileSize);
 		const uint32 MaxMipInBlock = FMath::CeilLogTwo(BlockSizeInTiles);
-		const uint32 MipWidthInBlock = FMath::Max<uint32>(BlockSizeX >> MaxMipInBlock, 1);
-		const uint32 MipHeightInBlock = FMath::Max<uint32>(BlockSizeY >> MaxMipInBlock, 1);
-		const uint32 MipInputSizeX = FMath::RoundUpToPowerOfTwo(SizeInBlocksX * MipWidthInBlock);
-		const uint32 MipInputSizeY = FMath::RoundUpToPowerOfTwo(SizeInBlocksY * MipHeightInBlock);
+		const uint32 MipWidthInBlock = FMath::Max<uint32>(DerivedInfo.BlockSizeX >> MaxMipInBlock, 1);
+		const uint32 MipHeightInBlock = FMath::Max<uint32>(DerivedInfo.BlockSizeY >> MaxMipInBlock, 1);
+		const uint32 MipInputSizeX = FMath::RoundUpToPowerOfTwo(DerivedInfo.SizeInBlocksX * MipWidthInBlock);
+		const uint32 MipInputSizeY = FMath::RoundUpToPowerOfTwo(DerivedInfo.SizeInBlocksY * MipHeightInBlock);
 		const uint32 MipInputSize = FMath::Max(MipInputSizeX, MipInputSizeY);
 		//const uint32 MipInputSizeInTiles = FMath::DivideAndRoundUp<uint32>(MipInputSize, TileSize);
 
 		FTextureSourceBlockData& SourceMiptailBlock = SourceBlocks.AddDefaulted_GetRef();
 		SourceMiptailBlock.BlockX = 0;
 		SourceMiptailBlock.BlockY = 0;
-		SourceMiptailBlock.SizeInBlocksX = SizeInBlocksX; // miptail block covers the entire logical source texture
-		SourceMiptailBlock.SizeInBlocksY = SizeInBlocksY;
+		SourceMiptailBlock.SizeInBlocksX = DerivedInfo.SizeInBlocksX; // miptail block covers the entire logical source texture
+		SourceMiptailBlock.SizeInBlocksY = DerivedInfo.SizeInBlocksY;
 		SourceMiptailBlock.SizeX = FMath::Max(MipInputSizeX >> 1, 1u);
 		SourceMiptailBlock.SizeY = FMath::Max(MipInputSizeY >> 1, 1u);
 		SourceMiptailBlock.NumMips = OutData.NumMips - MaxMipInBlock - 1;//   FMath::CeilLogTwo(MipInputSize); // Don't add 1, since 'MipInputSize' is one mip larger
diff --git a/Engine/Source/Runtime/Engine/Private/VT/VirtualTextureDataBuilder.h b/Engine/Source/Runtime/Engine/Private/VT/VirtualTextureDataBuilder.h
--- a/Engine/Source/Runtime/Engine/Private/VT/VirtualTextureDataBuilder.h
+++ b/Engine/Source/Runtime/Engine/Private/VT/VirtualTextureDataBuilder.h
@@ -42,6 +42,21 @@ struct FVirtualTextureSourceLayerData
 	bool bHasAlpha;
 };
 
+// Holds a bunch of stuff we derive from the input data that we use during the build.
+struct FVirtualTextureBuilderDerivedInfo
+{
+	int32 SizeInBlocksX = 0;
+	int32 SizeInBlocksY = 0;
+	int32 BlockSizeX = 0;
+	int32 BlockSizeY = 0;
+	int32 BlockSizeScale = 1;
+	int32 SizeX = 0;
+	int32 SizeY = 0;
+	int32 NumMips = 0;
+
+	bool InitializeFromBuildSettings(const FTextureSourceData& InSourceData, const FTextureBuildSettings* InSettingsPerLayer);
+};
+
 /**
  * Helper class for building virtual texture data. This works on a set of FTextureSource objects. The idea is that if needed we can create
  * FTextureSource without creating actual UTextures. This is why the builder should stay independent of UTexture. Things it does:
@@ -86,13 +101,7 @@ private:
 	FVirtualTextureBuiltData &OutData;
 
 	// Some convenience variables (mostly derived from the passed in build settings)
-	int32 SizeInBlocksX;
-	int32 SizeInBlocksY;
-	int32 BlockSizeX;
-	int32 BlockSizeY;
-	int32 BlockSizeScale;
-	int32 SizeX;
-	int32 SizeY;
+	FVirtualTextureBuilderDerivedInfo DerivedInfo;
 
 	TArray<FVirtualTextureSourceLayerData> SourceLayers;
 	TArray<FTextureSourceBlockData> SourceBlocks;
