Fix ErrMsg when calling non-COM-centered overload of FSimpleSuspensionHelpers::ComputeSprungMasses

diff --git a/Engine/Source/Runtime/Engine/Private/PhysicsEngine/SimpleSuspension.cpp b/Engine/Source/Runtime/Engine/Private/PhysicsEngine/SimpleSuspension.cpp
--- a/Engine/Source/Runtime/Engine/Private/PhysicsEngine/SimpleSuspension.cpp
+++ b/Engine/Source/Runtime/Engine/Private/PhysicsEngine/SimpleSuspension.cpp
@@ -355,7 +355,7 @@ bool FSimpleSuspensionHelpers::ComputeSprungMasses(const TArray<FVector>& LocalS
 	}
 
 	// Do the calculation
-	return ComputeSprungMasses(MassSpringPositions, TotalMass, OutSprungMasses);
+	return ComputeSprungMasses(MassSpringPositions, TotalMass, OutSprungMasses, ErrMsg);
 }
 
 void FSimpleSuspensionHelpers::ComputeSpringNaturalFrequencyAndDampingRatio(const float SprungMass, const float SpringStiffness, const float SpringDamping, float& OutNaturalFrequency, float& OutDampingRatio)
