Updated categories of animation nodes & functions

Reorganized and consolidated for easier navigation

diff --git a/Engine/Source/Editor/AnimGraph/Private/AnimGraphNode_AnimDynamics.cpp b/Engine/Source/Editor/AnimGraph/Private/AnimGraphNode_AnimDynamics.cpp
--- a/Engine/Source/Editor/AnimGraph/Private/AnimGraphNode_AnimDynamics.cpp
+++ b/Engine/Source/Editor/AnimGraph/Private/AnimGraphNode_AnimDynamics.cpp
@@ -30,6 +30,12 @@ FText UAnimGraphNode_AnimDynamics::GetTooltipText() const
 	return LOCTEXT("NodeTooltip", "Anim Dynamics");
 }
 
+
+FString UAnimGraphNode_AnimDynamics::GetNodeCategory() const
+{
+	return TEXT("Animation|Dynamics");
+}
+
 void UAnimGraphNode_AnimDynamics::GetOnScreenDebugInfo(TArray<FText>& DebugInfo, FAnimNode_Base* RuntimeAnimNode, USkeletalMeshComponent* PreviewSkelMeshComp) const
 {
 	if(RuntimeAnimNode)
diff --git a/Engine/Source/Editor/AnimGraph/Private/AnimGraphNode_ApplyAdditive.cpp b/Engine/Source/Editor/AnimGraph/Private/AnimGraphNode_ApplyAdditive.cpp
--- a/Engine/Source/Editor/AnimGraph/Private/AnimGraphNode_ApplyAdditive.cpp
+++ b/Engine/Source/Editor/AnimGraph/Private/AnimGraphNode_ApplyAdditive.cpp
@@ -159,7 +159,7 @@ void UAnimGraphNode_ApplyAdditive::CustomizeDetails(IDetailLayoutBuilder& Detail
 
 FString UAnimGraphNode_ApplyAdditive::GetNodeCategory() const
 {
-	return TEXT("Blends");
+	return TEXT("Animation|Blends");
 	//@TODO: TEXT("Apply additive to normal pose"), TEXT("Apply additive pose"));
 }
 
diff --git a/Engine/Source/Editor/AnimGraph/Private/AnimGraphNode_ApplyMeshSpaceAdditive.cpp b/Engine/Source/Editor/AnimGraph/Private/AnimGraphNode_ApplyMeshSpaceAdditive.cpp
--- a/Engine/Source/Editor/AnimGraph/Private/AnimGraphNode_ApplyMeshSpaceAdditive.cpp
+++ b/Engine/Source/Editor/AnimGraph/Private/AnimGraphNode_ApplyMeshSpaceAdditive.cpp
@@ -159,7 +159,7 @@ void UAnimGraphNode_ApplyMeshSpaceAdditive::PostEditChangeProperty(struct FPrope
 
 FString UAnimGraphNode_ApplyMeshSpaceAdditive::GetNodeCategory() const
 {
-	return TEXT("Blends");
+	return TEXT("Animation|Blends");
 }
 
 void UAnimGraphNode_ApplyMeshSpaceAdditive::ValidateAnimNodeDuringCompilation(class USkeleton* ForSkeleton, class FCompilerResultsLog& MessageLog)
diff --git a/Engine/Source/Editor/AnimGraph/Private/AnimGraphNode_BlendBoneByChannel.cpp b/Engine/Source/Editor/AnimGraph/Private/AnimGraphNode_BlendBoneByChannel.cpp
--- a/Engine/Source/Editor/AnimGraph/Private/AnimGraphNode_BlendBoneByChannel.cpp
+++ b/Engine/Source/Editor/AnimGraph/Private/AnimGraphNode_BlendBoneByChannel.cpp
@@ -14,7 +14,7 @@ UAnimGraphNode_BlendBoneByChannel::UAnimGraphNode_BlendBoneByChannel(const FObje
 
 FString UAnimGraphNode_BlendBoneByChannel::GetNodeCategory() const
 {
-	return TEXT("Blends");
+	return TEXT("Animation|Blends");
 }
 
 FLinearColor UAnimGraphNode_BlendBoneByChannel::GetNodeTitleColor() const
diff --git a/Engine/Source/Editor/AnimGraph/Private/AnimGraphNode_BlendListBase.cpp b/Engine/Source/Editor/AnimGraph/Private/AnimGraphNode_BlendListBase.cpp
--- a/Engine/Source/Editor/AnimGraph/Private/AnimGraphNode_BlendListBase.cpp
+++ b/Engine/Source/Editor/AnimGraph/Private/AnimGraphNode_BlendListBase.cpp
@@ -32,7 +32,7 @@ void UAnimGraphNode_BlendListBase::PostEditChangeProperty(struct FPropertyChange
 
 FString UAnimGraphNode_BlendListBase::GetNodeCategory() const
 {
-	return TEXT("Blends");
+	return TEXT("Animation|Blends");
 }
 
 void UAnimGraphNode_BlendListBase::RemovePinsFromOldPins(TArray<UEdGraphPin*>& OldPins, int32 RemovedArrayIndex)
diff --git a/Engine/Source/Editor/AnimGraph/Private/AnimGraphNode_BlendListByEnum.cpp b/Engine/Source/Editor/AnimGraph/Private/AnimGraphNode_BlendListByEnum.cpp
--- a/Engine/Source/Editor/AnimGraph/Private/AnimGraphNode_BlendListByEnum.cpp
+++ b/Engine/Source/Editor/AnimGraph/Private/AnimGraphNode_BlendListByEnum.cpp
@@ -23,9 +23,9 @@ UAnimGraphNode_BlendListByEnum::UAnimGraphNode_BlendListByEnum(const FObjectInit
 	Node.AddPose();
 }
 
-FString UAnimGraphNode_BlendListByEnum::GetNodeCategory() const
+FText UAnimGraphNode_BlendListByEnum::GetMenuCategory() const
 {
-	return FString::Printf(TEXT("%s, Blend List by enum"), *Super::GetNodeCategory() );
+	return LOCTEXT("AnimGraphNode_BlendListByEnum_GetMenuCategory", "Animation|Blends|Blend by Enum");
 }
 
 FText UAnimGraphNode_BlendListByEnum::GetTooltipText() const
diff --git a/Engine/Source/Editor/AnimGraph/Private/AnimGraphNode_BlendSpaceBase.cpp b/Engine/Source/Editor/AnimGraph/Private/AnimGraphNode_BlendSpaceBase.cpp
--- a/Engine/Source/Editor/AnimGraph/Private/AnimGraphNode_BlendSpaceBase.cpp
+++ b/Engine/Source/Editor/AnimGraph/Private/AnimGraphNode_BlendSpaceBase.cpp
@@ -119,7 +119,7 @@ void UAnimGraphNode_BlendSpaceBase::CustomizeDetails(IDetailLayoutBuilder& InDet
 
 FText UAnimGraphNode_BlendSpaceBase::GetMenuCategory() const
 {
-	return LOCTEXT("BlendSpaceCategory_Label", "BlendSpaces");
+	return LOCTEXT("BlendSpaceCategory_Label", "Animation|Blend Spaces");
 }
 
 bool UAnimGraphNode_BlendSpaceBase::IsAimOffsetBlendSpace(const UClass* BlendSpaceClass)
diff --git a/Engine/Source/Editor/AnimGraph/Private/AnimGraphNode_BlendSpaceGraphBase.cpp b/Engine/Source/Editor/AnimGraph/Private/AnimGraphNode_BlendSpaceGraphBase.cpp
--- a/Engine/Source/Editor/AnimGraph/Private/AnimGraphNode_BlendSpaceGraphBase.cpp
+++ b/Engine/Source/Editor/AnimGraph/Private/AnimGraphNode_BlendSpaceGraphBase.cpp
@@ -39,7 +39,7 @@ UAnimGraphNode_BlendSpaceGraphBase::UAnimGraphNode_BlendSpaceGraphBase()
 
 FText UAnimGraphNode_BlendSpaceGraphBase::GetMenuCategory() const
 {
-	return LOCTEXT("BlendSpaceCategory_Label", "BlendSpaces");
+	return LOCTEXT("BlendSpaceCategory_Label", "Animation|Blend Spaces");
 }
 
 FLinearColor UAnimGraphNode_BlendSpaceGraphBase::GetNodeTitleColor() const
diff --git a/Engine/Source/Editor/AnimGraph/Private/AnimGraphNode_ComponentToLocalSpace.cpp b/Engine/Source/Editor/AnimGraph/Private/AnimGraphNode_ComponentToLocalSpace.cpp
--- a/Engine/Source/Editor/AnimGraph/Private/AnimGraphNode_ComponentToLocalSpace.cpp
+++ b/Engine/Source/Editor/AnimGraph/Private/AnimGraphNode_ComponentToLocalSpace.cpp
@@ -27,9 +27,9 @@ FText UAnimGraphNode_ComponentToLocalSpace::GetNodeTitle(ENodeTitleType::Type Ti
 	return LOCTEXT("AnimGraphNode_ComponentToLocalSpace_Title", "Component To Local");
 }
 
-FString UAnimGraphNode_ComponentToLocalSpace::GetNodeCategory() const
+FText UAnimGraphNode_ComponentToLocalSpace::GetMenuCategory() const
 {
-	return TEXT("Convert Spaces");
+	return LOCTEXT("AnimGraphNode_LocalToComponentSpace_Category", "Animation|Convert Spaces");
 }
 
 void UAnimGraphNode_ComponentToLocalSpace::PostProcessPinName(const UEdGraphPin* Pin, FString& DisplayName) const
diff --git a/Engine/Source/Editor/AnimGraph/Private/AnimGraphNode_CopyPoseFromMesh.cpp b/Engine/Source/Editor/AnimGraph/Private/AnimGraphNode_CopyPoseFromMesh.cpp
--- a/Engine/Source/Editor/AnimGraph/Private/AnimGraphNode_CopyPoseFromMesh.cpp
+++ b/Engine/Source/Editor/AnimGraph/Private/AnimGraphNode_CopyPoseFromMesh.cpp
@@ -23,6 +23,11 @@ FText UAnimGraphNode_CopyPoseFromMesh::GetNodeTitle(ENodeTitleType::Type TitleTy
 	return LOCTEXT("CopyPoseFromMesh", "Copy Pose From Mesh");
 }
 
+FText UAnimGraphNode_CopyPoseFromMesh::GetMenuCategory() const
+{
+	return LOCTEXT("AnimGraphNode_CopyPoseFromMesh_Category", "Animation|Misc.");
+}
+
 void UAnimGraphNode_CopyPoseFromMesh::GetOutputLinkAttributes(FNodeAttributeArray& OutAttributes) const
 {
 	OutAttributes.Add(UE::Anim::FAttributes::Curves);
diff --git a/Engine/Source/Editor/AnimGraph/Private/AnimGraphNode_CurveSource.cpp b/Engine/Source/Editor/AnimGraph/Private/AnimGraphNode_CurveSource.cpp
--- a/Engine/Source/Editor/AnimGraph/Private/AnimGraphNode_CurveSource.cpp
+++ b/Engine/Source/Editor/AnimGraph/Private/AnimGraphNode_CurveSource.cpp
@@ -7,7 +7,7 @@
 
 FString UAnimGraphNode_CurveSource::GetNodeCategory() const
 {
-	return TEXT("Curves");
+	return TEXT("Animation|Curves");
 }
 
 FText UAnimGraphNode_CurveSource::GetTooltipText() const
diff --git a/Engine/Source/Editor/AnimGraph/Private/AnimGraphNode_Inertialization.cpp b/Engine/Source/Editor/AnimGraph/Private/AnimGraphNode_Inertialization.cpp
--- a/Engine/Source/Editor/AnimGraph/Private/AnimGraphNode_Inertialization.cpp
+++ b/Engine/Source/Editor/AnimGraph/Private/AnimGraphNode_Inertialization.cpp
@@ -22,7 +22,7 @@ FText UAnimGraphNode_Inertialization::GetNodeTitle(ENodeTitleType::Type TitleTyp
 
 FText UAnimGraphNode_Inertialization::GetMenuCategory() const
 {
-	return LOCTEXT("NodeCategory", "Inertialization");
+	return LOCTEXT("NodeCategory", "Animation|Misc.");
 }
 
 void UAnimGraphNode_Inertialization::GetInputLinkAttributes(FNodeAttributeArray& OutAttributes) const
diff --git a/Engine/Source/Editor/AnimGraph/Private/AnimGraphNode_LayeredBoneBlend.cpp b/Engine/Source/Editor/AnimGraph/Private/AnimGraphNode_LayeredBoneBlend.cpp
--- a/Engine/Source/Editor/AnimGraph/Private/AnimGraphNode_LayeredBoneBlend.cpp
+++ b/Engine/Source/Editor/AnimGraph/Private/AnimGraphNode_LayeredBoneBlend.cpp
@@ -66,7 +66,7 @@ void UAnimGraphNode_LayeredBoneBlend::PostEditChangeProperty(struct FPropertyCha
 
 FString UAnimGraphNode_LayeredBoneBlend::GetNodeCategory() const
 {
-	return TEXT("Blends");
+	return TEXT("Animation|Blends");
 }
 
 void UAnimGraphNode_LayeredBoneBlend::CustomizeDetails(IDetailLayoutBuilder& DetailBuilder)
diff --git a/Engine/Source/Editor/AnimGraph/Private/AnimGraphNode_LinkedAnimGraphBase.cpp b/Engine/Source/Editor/AnimGraph/Private/AnimGraphNode_LinkedAnimGraphBase.cpp
--- a/Engine/Source/Editor/AnimGraph/Private/AnimGraphNode_LinkedAnimGraphBase.cpp
+++ b/Engine/Source/Editor/AnimGraph/Private/AnimGraphNode_LinkedAnimGraphBase.cpp
@@ -114,7 +114,7 @@ FText UAnimGraphNode_LinkedAnimGraphBase::GetMenuCategory() const
 
 	if(!Category.IsSet() || Category.GetValue().IsEmpty())
 	{
-		return LOCTEXT("LinkedAnimGraphCategory", "Linked Anim Blueprints");
+		return LOCTEXT("LinkedAnimGraphCategory", "Animation|Linked Anim Graphs");
 	}
 
 	return Category.GetValue();
diff --git a/Engine/Source/Editor/AnimGraph/Private/AnimGraphNode_LinkedInputPose.cpp b/Engine/Source/Editor/AnimGraph/Private/AnimGraphNode_LinkedInputPose.cpp
--- a/Engine/Source/Editor/AnimGraph/Private/AnimGraphNode_LinkedInputPose.cpp
+++ b/Engine/Source/Editor/AnimGraph/Private/AnimGraphNode_LinkedInputPose.cpp
@@ -159,6 +159,11 @@ FText UAnimGraphNode_LinkedInputPose::GetNodeTitle(ENodeTitleType::Type TitleTyp
 	}
 }
 
+FText UAnimGraphNode_LinkedInputPose::GetMenuCategory() const
+{
+	return LOCTEXT("LinkedAnimGraphCategory", "Animation|Linked Anim Graphs");
+}
+
 bool UAnimGraphNode_LinkedInputPose::CanUserDeleteNode() const
 {
 	// Only allow linked input poses to be deleted if their parent graph is mutable
diff --git a/Engine/Source/Editor/AnimGraph/Private/AnimGraphNode_LocalToComponentSpace.cpp b/Engine/Source/Editor/AnimGraph/Private/AnimGraphNode_LocalToComponentSpace.cpp
--- a/Engine/Source/Editor/AnimGraph/Private/AnimGraphNode_LocalToComponentSpace.cpp
+++ b/Engine/Source/Editor/AnimGraph/Private/AnimGraphNode_LocalToComponentSpace.cpp
@@ -28,9 +28,9 @@ FText UAnimGraphNode_LocalToComponentSpace::GetNodeTitle(ENodeTitleType::Type Ti
 	return LOCTEXT("AnimGraphNode_LocalToComponentSpace_Title", "Local To Component");
 }
 
-FString UAnimGraphNode_LocalToComponentSpace::GetNodeCategory() const
+FText UAnimGraphNode_LocalToComponentSpace::GetMenuCategory() const
 {
-	return TEXT("Convert Spaces");
+	return LOCTEXT("AnimGraphNode_LocalToComponentSpace_Category", "Animation|Convert Spaces");
 }
 
 void UAnimGraphNode_LocalToComponentSpace::CreateOutputPins()
diff --git a/Engine/Source/Editor/AnimGraph/Private/AnimGraphNode_MakeDynamicAdditive.cpp b/Engine/Source/Editor/AnimGraph/Private/AnimGraphNode_MakeDynamicAdditive.cpp
--- a/Engine/Source/Editor/AnimGraph/Private/AnimGraphNode_MakeDynamicAdditive.cpp
+++ b/Engine/Source/Editor/AnimGraph/Private/AnimGraphNode_MakeDynamicAdditive.cpp
@@ -32,7 +32,7 @@ FText UAnimGraphNode_MakeDynamicAdditive::GetNodeTitle(ENodeTitleType::Type Titl
 
 FString UAnimGraphNode_MakeDynamicAdditive::GetNodeCategory() const
 {
-	return TEXT("Blends");
+	return TEXT("Animation|Blends");
 }
 
 void UAnimGraphNode_MakeDynamicAdditive::ValidateAnimNodeDuringCompilation(class USkeleton* ForSkeleton, class FCompilerResultsLog& MessageLog)
diff --git a/Engine/Source/Editor/AnimGraph/Private/AnimGraphNode_MeshRefPose.cpp b/Engine/Source/Editor/AnimGraph/Private/AnimGraphNode_MeshRefPose.cpp
--- a/Engine/Source/Editor/AnimGraph/Private/AnimGraphNode_MeshRefPose.cpp
+++ b/Engine/Source/Editor/AnimGraph/Private/AnimGraphNode_MeshRefPose.cpp
@@ -15,7 +15,7 @@ UAnimGraphNode_MeshRefPose::UAnimGraphNode_MeshRefPose(const FObjectInitializer&
 
 FString UAnimGraphNode_MeshRefPose::GetNodeCategory() const
 {
-	return TEXT("Identity");
+	return TEXT("Animation|Poses");
 }
 
 FLinearColor UAnimGraphNode_MeshRefPose::GetNodeTitleColor() const
diff --git a/Engine/Source/Editor/AnimGraph/Private/AnimGraphNode_Mirror.cpp b/Engine/Source/Editor/AnimGraph/Private/AnimGraphNode_Mirror.cpp
--- a/Engine/Source/Editor/AnimGraph/Private/AnimGraphNode_Mirror.cpp
+++ b/Engine/Source/Editor/AnimGraph/Private/AnimGraphNode_Mirror.cpp
@@ -40,7 +40,7 @@ FText UAnimGraphNode_Mirror::GetNodeTitle(ENodeTitleType::Type TitleType) const
 
 FText UAnimGraphNode_Mirror::GetMenuCategory() const
 {
-	return LOCTEXT("NodeCategory", "Mirroring");
+	return LOCTEXT("NodeCategory", "Animation|Misc.");
 }
 
 void UAnimGraphNode_Mirror::GetOutputLinkAttributes(FNodeAttributeArray& OutAttributes) const
diff --git a/Engine/Source/Editor/AnimGraph/Private/AnimGraphNode_ModifyCurve.cpp b/Engine/Source/Editor/AnimGraph/Private/AnimGraphNode_ModifyCurve.cpp
--- a/Engine/Source/Editor/AnimGraph/Private/AnimGraphNode_ModifyCurve.cpp
+++ b/Engine/Source/Editor/AnimGraph/Private/AnimGraphNode_ModifyCurve.cpp
@@ -18,9 +18,9 @@ UAnimGraphNode_ModifyCurve::UAnimGraphNode_ModifyCurve()
 {
 }
 
-FString UAnimGraphNode_ModifyCurve::GetNodeCategory() const
+FText UAnimGraphNode_ModifyCurve::GetMenuCategory() const
 {
-	return TEXT("Skeletal Control Nodes");
+	return LOCTEXT("AnimGraphNode_ModifyCurve_Category", "Animation|Curves");
 }
 
 FText UAnimGraphNode_ModifyCurve::GetTooltipText() const
diff --git a/Engine/Source/Editor/AnimGraph/Private/AnimGraphNode_MultiWayBlend.cpp b/Engine/Source/Editor/AnimGraph/Private/AnimGraphNode_MultiWayBlend.cpp
--- a/Engine/Source/Editor/AnimGraph/Private/AnimGraphNode_MultiWayBlend.cpp
+++ b/Engine/Source/Editor/AnimGraph/Private/AnimGraphNode_MultiWayBlend.cpp
@@ -19,7 +19,7 @@ UAnimGraphNode_MultiWayBlend::UAnimGraphNode_MultiWayBlend(const FObjectInitiali
 
 FString UAnimGraphNode_MultiWayBlend::GetNodeCategory() const
 {
-	return TEXT("Blends");
+	return TEXT("Animation|Blends");
 }
 
 FLinearColor UAnimGraphNode_MultiWayBlend::GetNodeTitleColor() const
diff --git a/Engine/Source/Editor/AnimGraph/Private/AnimGraphNode_PoseBlendNode.cpp b/Engine/Source/Editor/AnimGraph/Private/AnimGraphNode_PoseBlendNode.cpp
--- a/Engine/Source/Editor/AnimGraph/Private/AnimGraphNode_PoseBlendNode.cpp
+++ b/Engine/Source/Editor/AnimGraph/Private/AnimGraphNode_PoseBlendNode.cpp
@@ -142,7 +142,7 @@ FText UAnimGraphNode_PoseBlendNode::GetTitleGivenAssetInfo(const FText& AssetNam
 
 FText UAnimGraphNode_PoseBlendNode::GetMenuCategory() const
 {
-	return LOCTEXT("PoseAssetCategory_Label", "Poses");
+	return LOCTEXT("PoseAssetCategory_Label", "Animation|Poses");
 }
 
 bool UAnimGraphNode_PoseBlendNode::DoesSupportTimeForTransitionGetter() const
diff --git a/Engine/Source/Editor/AnimGraph/Private/AnimGraphNode_PoseByName.cpp b/Engine/Source/Editor/AnimGraph/Private/AnimGraphNode_PoseByName.cpp
--- a/Engine/Source/Editor/AnimGraph/Private/AnimGraphNode_PoseByName.cpp
+++ b/Engine/Source/Editor/AnimGraph/Private/AnimGraphNode_PoseByName.cpp
@@ -45,7 +45,7 @@ FText UAnimGraphNode_PoseByName::GetTooltipText() const
 
 FText UAnimGraphNode_PoseByName::GetMenuCategory() const
 {
-	return LOCTEXT("PoseAssetCategory_Label", "Poses");
+	return LOCTEXT("PoseAssetCategory_Label", "Animation|Poses");
 }
 
 FText UAnimGraphNode_PoseByName::GetNodeTitleForPoseAsset(ENodeTitleType::Type TitleType, UPoseAsset* InPoseAsset) const
diff --git a/Engine/Source/Editor/AnimGraph/Private/AnimGraphNode_PoseDriver.cpp b/Engine/Source/Editor/AnimGraph/Private/AnimGraphNode_PoseDriver.cpp
--- a/Engine/Source/Editor/AnimGraph/Private/AnimGraphNode_PoseDriver.cpp
+++ b/Engine/Source/Editor/AnimGraph/Private/AnimGraphNode_PoseDriver.cpp
@@ -95,7 +95,7 @@ FText UAnimGraphNode_PoseDriver::GetNodeTitle(ENodeTitleType::Type TitleType) co
 
 FText UAnimGraphNode_PoseDriver::GetMenuCategory() const
 {
-	return LOCTEXT("PoseAssetCategory_Label", "Poses");
+	return LOCTEXT("PoseAssetCategory_Label", "Animation|Poses");
 }
 
 
diff --git a/Engine/Source/Editor/AnimGraph/Private/AnimGraphNode_PoseSnapshot.cpp b/Engine/Source/Editor/AnimGraph/Private/AnimGraphNode_PoseSnapshot.cpp
--- a/Engine/Source/Editor/AnimGraph/Private/AnimGraphNode_PoseSnapshot.cpp
+++ b/Engine/Source/Editor/AnimGraph/Private/AnimGraphNode_PoseSnapshot.cpp
@@ -22,4 +22,9 @@ FText UAnimGraphNode_PoseSnapshot::GetNodeTitle(ENodeTitleType::Type TitleType)
 	return LOCTEXT("UAnimGraphNode_PoseSnapshot_Title", "Pose Snapshot");
 }
 
+FText UAnimGraphNode_PoseSnapshot::GetMenuCategory() const
+{
+	return LOCTEXT("UAnimGraphNode_PoseSnapshot_Category", "Animation|Poses");
+}
+
 #undef LOCTEXT_NAMESPACE
diff --git a/Engine/Source/Editor/AnimGraph/Private/AnimGraphNode_RandomPlayer.cpp b/Engine/Source/Editor/AnimGraph/Private/AnimGraphNode_RandomPlayer.cpp
--- a/Engine/Source/Editor/AnimGraph/Private/AnimGraphNode_RandomPlayer.cpp
+++ b/Engine/Source/Editor/AnimGraph/Private/AnimGraphNode_RandomPlayer.cpp
@@ -25,7 +25,7 @@ FText UAnimGraphNode_RandomPlayer::GetNodeTitle(ENodeTitleType::Type TitleType)
 
 FText UAnimGraphNode_RandomPlayer::GetMenuCategory() const
 {
-	return FEditorCategoryUtils::GetCommonCategory(FCommonEditorCategory::Animation);
+	return LOCTEXT("NodeCategory", "Animation|Sequences");
 }
 
 void UAnimGraphNode_RandomPlayer::GetOutputLinkAttributes(FNodeAttributeArray& OutAttributes) const
diff --git a/Engine/Source/Editor/AnimGraph/Private/AnimGraphNode_RefPoseBase.cpp b/Engine/Source/Editor/AnimGraph/Private/AnimGraphNode_RefPoseBase.cpp
--- a/Engine/Source/Editor/AnimGraph/Private/AnimGraphNode_RefPoseBase.cpp
+++ b/Engine/Source/Editor/AnimGraph/Private/AnimGraphNode_RefPoseBase.cpp
@@ -12,7 +12,7 @@ UAnimGraphNode_RefPoseBase::UAnimGraphNode_RefPoseBase(const FObjectInitializer&
 
 FString UAnimGraphNode_RefPoseBase::GetNodeCategory() const
 {
-	return TEXT("Identity");
+	return TEXT("Animation|Poses");
 }
 
 FLinearColor UAnimGraphNode_RefPoseBase::GetNodeTitleColor() const
diff --git a/Engine/Source/Editor/AnimGraph/Private/AnimGraphNode_RigidBody.cpp b/Engine/Source/Editor/AnimGraph/Private/AnimGraphNode_RigidBody.cpp
--- a/Engine/Source/Editor/AnimGraph/Private/AnimGraphNode_RigidBody.cpp
+++ b/Engine/Source/Editor/AnimGraph/Private/AnimGraphNode_RigidBody.cpp
@@ -38,6 +38,11 @@ FText UAnimGraphNode_RigidBody::GetNodeTitle(ENodeTitleType::Type TitleType) con
 	return FText(LOCTEXT("AnimGraphNode_RigidBody_NodeTitle", "RigidBody"));
 }
 
+FString UAnimGraphNode_RigidBody::GetNodeCategory() const
+{
+	return TEXT("Animation|Dynamics");
+}
+
 void UAnimGraphNode_RigidBody::ValidateAnimNodeDuringCompilation(USkeleton* ForSkeleton, FCompilerResultsLog& MessageLog)
 {
 	Super::ValidateAnimNodeDuringCompilation(ForSkeleton, MessageLog);
diff --git a/Engine/Source/Editor/AnimGraph/Private/AnimGraphNode_RotateRootBone.cpp b/Engine/Source/Editor/AnimGraph/Private/AnimGraphNode_RotateRootBone.cpp
--- a/Engine/Source/Editor/AnimGraph/Private/AnimGraphNode_RotateRootBone.cpp
+++ b/Engine/Source/Editor/AnimGraph/Private/AnimGraphNode_RotateRootBone.cpp
@@ -28,9 +28,9 @@ FText UAnimGraphNode_RotateRootBone::GetNodeTitle(ENodeTitleType::Type TitleType
 	return LOCTEXT("RotateRootBone", "Rotate Root Bone");
 }
 
-FString UAnimGraphNode_RotateRootBone::GetNodeCategory() const
+FText UAnimGraphNode_RotateRootBone::GetMenuCategory() const
 {
-	return TEXT("Tools");
+	return LOCTEXT("RotateRootBoneCateogory", "Animation|Misc.");
 }
 
 void UAnimGraphNode_RotateRootBone::CustomizePinData(UEdGraphPin* Pin, FName SourcePropertyName, int32 ArrayIndex) const
diff --git a/Engine/Source/Editor/AnimGraph/Private/AnimGraphNode_SaveCachedPose.cpp b/Engine/Source/Editor/AnimGraph/Private/AnimGraphNode_SaveCachedPose.cpp
--- a/Engine/Source/Editor/AnimGraph/Private/AnimGraphNode_SaveCachedPose.cpp
+++ b/Engine/Source/Editor/AnimGraph/Private/AnimGraphNode_SaveCachedPose.cpp
@@ -63,7 +63,7 @@ FText UAnimGraphNode_SaveCachedPose::GetNodeTitle(ENodeTitleType::Type TitleType
 
 FString UAnimGraphNode_SaveCachedPose::GetNodeCategory() const
 {
-	return TEXT("Cached Poses");
+	return TEXT("Animation|Cached Poses");
 }
 
 void UAnimGraphNode_SaveCachedPose::GetMenuActions(FBlueprintActionDatabaseRegistrar& ActionRegistrar) const
diff --git a/Engine/Source/Editor/AnimGraph/Private/AnimGraphNode_ScaleChainLength.cpp b/Engine/Source/Editor/AnimGraph/Private/AnimGraphNode_ScaleChainLength.cpp
--- a/Engine/Source/Editor/AnimGraph/Private/AnimGraphNode_ScaleChainLength.cpp
+++ b/Engine/Source/Editor/AnimGraph/Private/AnimGraphNode_ScaleChainLength.cpp
@@ -30,7 +30,7 @@ FText UAnimGraphNode_ScaleChainLength::GetNodeTitle(ENodeTitleType::Type TitleTy
 
 FString UAnimGraphNode_ScaleChainLength::GetNodeCategory() const
 {
-	return TEXT("AnimNode");
+	return TEXT("Animation|Misc.");
 }
 
 #undef LOCTEXT_NAMESPACE
diff --git a/Engine/Source/Editor/AnimGraph/Private/AnimGraphNode_SequenceEvaluator.cpp b/Engine/Source/Editor/AnimGraph/Private/AnimGraphNode_SequenceEvaluator.cpp
--- a/Engine/Source/Editor/AnimGraph/Private/AnimGraphNode_SequenceEvaluator.cpp
+++ b/Engine/Source/Editor/AnimGraph/Private/AnimGraphNode_SequenceEvaluator.cpp
@@ -55,7 +55,7 @@ void UAnimGraphNode_SequenceEvaluator::ReplaceReferredAnimations(const TMap<UAni
 
 FText UAnimGraphNode_SequenceEvaluator::GetMenuCategory() const
 {
-	return FEditorCategoryUtils::GetCommonCategory(FCommonEditorCategory::Animation);
+	return LOCTEXT("MenuCategory", "Animation|Sequences");
 }
 
 FText UAnimGraphNode_SequenceEvaluator::GetNodeTitle(ENodeTitleType::Type TitleType) const
diff --git a/Engine/Source/Editor/AnimGraph/Private/AnimGraphNode_SequencePlayer.cpp b/Engine/Source/Editor/AnimGraph/Private/AnimGraphNode_SequencePlayer.cpp
--- a/Engine/Source/Editor/AnimGraph/Private/AnimGraphNode_SequencePlayer.cpp
+++ b/Engine/Source/Editor/AnimGraph/Private/AnimGraphNode_SequencePlayer.cpp
@@ -82,7 +82,7 @@ FText UAnimGraphNode_SequencePlayer::GetNodeTitle(ENodeTitleType::Type TitleType
 
 FText UAnimGraphNode_SequencePlayer::GetMenuCategory() const
 {
-	return FEditorCategoryUtils::GetCommonCategory(FCommonEditorCategory::Animation);
+	return LOCTEXT("MenuCategory", "Animation|Sequences");
 }
 
 void UAnimGraphNode_SequencePlayer::GetMenuActions(FBlueprintActionDatabaseRegistrar& InActionRegistrar) const
diff --git a/Engine/Source/Editor/AnimGraph/Private/AnimGraphNode_SkeletalControlBase.cpp b/Engine/Source/Editor/AnimGraph/Private/AnimGraphNode_SkeletalControlBase.cpp
--- a/Engine/Source/Editor/AnimGraph/Private/AnimGraphNode_SkeletalControlBase.cpp
+++ b/Engine/Source/Editor/AnimGraph/Private/AnimGraphNode_SkeletalControlBase.cpp
@@ -60,7 +60,7 @@ FLinearColor UAnimGraphNode_SkeletalControlBase::GetNodeTitleColor() const
 
 FString UAnimGraphNode_SkeletalControlBase::GetNodeCategory() const
 {
-	return TEXT("Skeletal Control Nodes");
+	return TEXT("Animation|Skeletal Controls");
 }
 
 FText UAnimGraphNode_SkeletalControlBase::GetControllerDescription() const
diff --git a/Engine/Source/Editor/AnimGraph/Private/AnimGraphNode_Slot.cpp b/Engine/Source/Editor/AnimGraph/Private/AnimGraphNode_Slot.cpp
--- a/Engine/Source/Editor/AnimGraph/Private/AnimGraphNode_Slot.cpp
+++ b/Engine/Source/Editor/AnimGraph/Private/AnimGraphNode_Slot.cpp
@@ -62,7 +62,7 @@ FText UAnimGraphNode_Slot::GetNodeTitle(ENodeTitleType::Type TitleType) const
 
 FString UAnimGraphNode_Slot::GetNodeCategory() const
 {
-	return TEXT("Montage");
+	return TEXT("Animation|Montage");
 }
 
 void UAnimGraphNode_Slot::BakeDataDuringCompilation(class FCompilerResultsLog& MessageLog)
diff --git a/Engine/Source/Editor/AnimGraph/Private/AnimGraphNode_StateMachineBase.cpp b/Engine/Source/Editor/AnimGraph/Private/AnimGraphNode_StateMachineBase.cpp
--- a/Engine/Source/Editor/AnimGraph/Private/AnimGraphNode_StateMachineBase.cpp
+++ b/Engine/Source/Editor/AnimGraph/Private/AnimGraphNode_StateMachineBase.cpp
@@ -126,9 +126,9 @@ FText UAnimGraphNode_StateMachineBase::GetNodeTitle(ENodeTitleType::Type TitleTy
 	return FText::FromName(EditorStateMachineGraph->GetFName());
 }
 
-FString UAnimGraphNode_StateMachineBase::GetNodeCategory() const
+FText UAnimGraphNode_StateMachineBase::GetMenuCategory() const
 {
-	return TEXT("State Machines");
+	return LOCTEXT("StateMachineCategory", "Animation|State Machines");
 }
 
 void UAnimGraphNode_StateMachineBase::PostPlacedNewNode()
diff --git a/Engine/Source/Editor/AnimGraph/Private/AnimGraphNode_Sync.cpp b/Engine/Source/Editor/AnimGraph/Private/AnimGraphNode_Sync.cpp
--- a/Engine/Source/Editor/AnimGraph/Private/AnimGraphNode_Sync.cpp
+++ b/Engine/Source/Editor/AnimGraph/Private/AnimGraphNode_Sync.cpp
@@ -32,6 +32,11 @@ FText UAnimGraphNode_Sync::GetTooltipText() const
 	return LOCTEXT("NodeTooltip", "Synchronizes asset players and blendspaces that are children of this node.");
 }
 
+FText UAnimGraphNode_Sync::GetMenuCategory() const
+{
+	return LOCTEXT("Category", "Animation|Synchronization");
+}
+
 void UAnimGraphNode_Sync::GetInputLinkAttributes(FNodeAttributeArray& OutAttributes) const
 {
 	if(Node.GetGroupName() != NAME_None)
diff --git a/Engine/Source/Editor/AnimGraph/Private/AnimGraphNode_Sync.h b/Engine/Source/Editor/AnimGraph/Private/AnimGraphNode_Sync.h
--- a/Engine/Source/Editor/AnimGraph/Private/AnimGraphNode_Sync.h
+++ b/Engine/Source/Editor/AnimGraph/Private/AnimGraphNode_Sync.h
@@ -21,6 +21,7 @@ class UAnimGraphNode_Sync : public UAnimGraphNode_Base
 	// UAnimGraphNode_Base interface
 	virtual FText GetNodeTitle(ENodeTitleType::Type TitleType) const override;
 	virtual FText GetTooltipText() const override;
+	virtual FText GetMenuCategory() const override;
 	virtual void ValidateAnimNodeDuringCompilation(USkeleton* ForSkeleton, FCompilerResultsLog& MessageLog) override;
 	virtual void BakeDataDuringCompilation(class FCompilerResultsLog& MessageLog) override;
 	virtual void GetInputLinkAttributes(FNodeAttributeArray& OutAttributes) const override;
diff --git a/Engine/Source/Editor/AnimGraph/Private/AnimGraphNode_TransitionPoseEvaluator.cpp b/Engine/Source/Editor/AnimGraph/Private/AnimGraphNode_TransitionPoseEvaluator.cpp
--- a/Engine/Source/Editor/AnimGraph/Private/AnimGraphNode_TransitionPoseEvaluator.cpp
+++ b/Engine/Source/Editor/AnimGraph/Private/AnimGraphNode_TransitionPoseEvaluator.cpp
@@ -50,9 +50,9 @@ void UAnimGraphNode_TransitionPoseEvaluator::ValidateAnimNodeDuringCompilation(c
 	Super::ValidateAnimNodeDuringCompilation(ForSkeleton, MessageLog);
 }
 
-FString UAnimGraphNode_TransitionPoseEvaluator::GetNodeCategory() const 
+FText UAnimGraphNode_TransitionPoseEvaluator::GetMenuCategory() const
 {
-	return TEXT("Transition");
+	return LOCTEXT("InputPoseEvaluatorCategory", "Animation|State Machines");
 }
 
 bool UAnimGraphNode_TransitionPoseEvaluator::CanUserDeleteNode() const
diff --git a/Engine/Source/Editor/AnimGraph/Private/AnimGraphNode_TwoWayBlend.cpp b/Engine/Source/Editor/AnimGraph/Private/AnimGraphNode_TwoWayBlend.cpp
--- a/Engine/Source/Editor/AnimGraph/Private/AnimGraphNode_TwoWayBlend.cpp
+++ b/Engine/Source/Editor/AnimGraph/Private/AnimGraphNode_TwoWayBlend.cpp
@@ -15,9 +15,9 @@ UAnimGraphNode_TwoWayBlend::UAnimGraphNode_TwoWayBlend(const FObjectInitializer&
 {
 }
 
-FString UAnimGraphNode_TwoWayBlend::GetNodeCategory() const
+FText UAnimGraphNode_TwoWayBlend::GetMenuCategory() const
 {
-	return TEXT("Blends");
+	return LOCTEXT("TwoWayBlendCategory", "Animation|Blends");
 }
 
 FLinearColor UAnimGraphNode_TwoWayBlend::GetNodeTitleColor() const
diff --git a/Engine/Source/Editor/AnimGraph/Private/AnimGraphNode_UseCachedPose.cpp b/Engine/Source/Editor/AnimGraph/Private/AnimGraphNode_UseCachedPose.cpp
--- a/Engine/Source/Editor/AnimGraph/Private/AnimGraphNode_UseCachedPose.cpp
+++ b/Engine/Source/Editor/AnimGraph/Private/AnimGraphNode_UseCachedPose.cpp
@@ -109,7 +109,7 @@ UObject* UAnimGraphNode_UseCachedPose::GetJumpTargetForDoubleClick() const
 
 FString UAnimGraphNode_UseCachedPose::GetNodeCategory() const
 {
-	return TEXT("Cached Poses");
+	return TEXT("Animation|Cached Poses");
 }
 
 void UAnimGraphNode_UseCachedPose::GetMenuActions(FBlueprintActionDatabaseRegistrar& ActionRegistrar) const
diff --git a/Engine/Source/Editor/AnimGraph/Public/AnimGraphNode_AnimDynamics.h b/Engine/Source/Editor/AnimGraph/Public/AnimGraphNode_AnimDynamics.h
--- a/Engine/Source/Editor/AnimGraph/Public/AnimGraphNode_AnimDynamics.h
+++ b/Engine/Source/Editor/AnimGraph/Public/AnimGraphNode_AnimDynamics.h
@@ -74,6 +74,7 @@ public:
 	// UEdGraphNode_Base
 	virtual FText GetNodeTitle(ENodeTitleType::Type TitleType) const override;
 	virtual FText GetTooltipText() const override;
+	virtual FString GetNodeCategory() const override;
 
 	// UAnimGraphNode_Base
 	virtual void ValidateAnimNodeDuringCompilation(USkeleton* ForSkeleton, FCompilerResultsLog& MessageLog) override;
diff --git a/Engine/Source/Editor/AnimGraph/Public/AnimGraphNode_BlendListByEnum.h b/Engine/Source/Editor/AnimGraph/Public/AnimGraphNode_BlendListByEnum.h
--- a/Engine/Source/Editor/AnimGraph/Public/AnimGraphNode_BlendListByEnum.h
+++ b/Engine/Source/Editor/AnimGraph/Public/AnimGraphNode_BlendListByEnum.h
@@ -32,6 +32,7 @@ public:
 	virtual FText GetTooltipText() const override;
 	virtual FText GetNodeTitle(ENodeTitleType::Type TitleType) const override;
 	virtual void Serialize(FArchive& Ar) override;
+	virtual FText GetMenuCategory() const override;
 	// End of UEdGraphNode interface
 
 	// UK2Node interface
@@ -39,7 +40,6 @@ public:
 	// End of UK2Node interface
 
 	// UAnimGraphNode_Base interface
-	virtual FString GetNodeCategory() const override;
 	virtual void GetMenuActions(FBlueprintActionDatabaseRegistrar& ActionRegistrar) const override;
 	virtual void CustomizePinData(UEdGraphPin* Pin, FName SourcePropertyName, int32 ArrayIndex) const override;
 	virtual void ValidateAnimNodeDuringCompilation(class USkeleton* ForSkeleton, class FCompilerResultsLog& MessageLog) override;
diff --git a/Engine/Source/Editor/AnimGraph/Public/AnimGraphNode_ComponentToLocalSpace.h b/Engine/Source/Editor/AnimGraph/Public/AnimGraphNode_ComponentToLocalSpace.h
--- a/Engine/Source/Editor/AnimGraph/Public/AnimGraphNode_ComponentToLocalSpace.h
+++ b/Engine/Source/Editor/AnimGraph/Public/AnimGraphNode_ComponentToLocalSpace.h
@@ -20,10 +20,10 @@ class UAnimGraphNode_ComponentToLocalSpace : public UAnimGraphNode_Base
 	virtual FLinearColor GetNodeTitleColor() const override;
 	virtual FText GetTooltipText() const override;
 	virtual FText GetNodeTitle(ENodeTitleType::Type TitleType) const override;
+	virtual FText GetMenuCategory() const override;
 	// End of UEdGraphNode interface
 
 	// UAnimGraphNode_Base interface
-	virtual FString GetNodeCategory() const override;
 	virtual void PostProcessPinName(const UEdGraphPin* Pin, FString& DisplayName) const override;
 	// End of UAnimGraphNode_Base interface
 };
diff --git a/Engine/Source/Editor/AnimGraph/Public/AnimGraphNode_CopyPoseFromMesh.h b/Engine/Source/Editor/AnimGraph/Public/AnimGraphNode_CopyPoseFromMesh.h
--- a/Engine/Source/Editor/AnimGraph/Public/AnimGraphNode_CopyPoseFromMesh.h
+++ b/Engine/Source/Editor/AnimGraph/Public/AnimGraphNode_CopyPoseFromMesh.h
@@ -20,6 +20,7 @@ public:
 	// UEdGraphNode interface
 	virtual FText GetNodeTitle(ENodeTitleType::Type TitleType) const override;
 	virtual FText GetTooltipText() const override;
+	virtual FText GetMenuCategory() const override;
 	// End of UEdGraphNode interface
 
 	// UAnimGraphNode_Base interface
diff --git a/Engine/Source/Editor/AnimGraph/Public/AnimGraphNode_LinkedInputPose.h b/Engine/Source/Editor/AnimGraph/Public/AnimGraphNode_LinkedInputPose.h
--- a/Engine/Source/Editor/AnimGraph/Public/AnimGraphNode_LinkedInputPose.h
+++ b/Engine/Source/Editor/AnimGraph/Public/AnimGraphNode_LinkedInputPose.h
@@ -72,6 +72,7 @@ public:
 	virtual FLinearColor GetNodeTitleColor() const override;
 	virtual FText GetTooltipText() const override;
 	virtual FText GetNodeTitle(ENodeTitleType::Type TitleType) const override;
+	virtual FText GetMenuCategory() const override;
 	virtual bool CanUserDeleteNode() const override;
 	virtual bool CanDuplicateNode() const override;
 	virtual void AllocateDefaultPins() override;
diff --git a/Engine/Source/Editor/AnimGraph/Public/AnimGraphNode_LocalToComponentSpace.h b/Engine/Source/Editor/AnimGraph/Public/AnimGraphNode_LocalToComponentSpace.h
--- a/Engine/Source/Editor/AnimGraph/Public/AnimGraphNode_LocalToComponentSpace.h
+++ b/Engine/Source/Editor/AnimGraph/Public/AnimGraphNode_LocalToComponentSpace.h
@@ -20,10 +20,10 @@ class UAnimGraphNode_LocalToComponentSpace : public UAnimGraphNode_Base
 	virtual FLinearColor GetNodeTitleColor() const override;
 	virtual FText GetTooltipText() const override;
 	virtual FText GetNodeTitle(ENodeTitleType::Type TitleType) const override;
+	virtual FText GetMenuCategory() const override;
 	// End of UEdGraphNode interface
 
 	// UAnimGraphNode_Base interface
-	virtual FString GetNodeCategory() const override;
 	virtual void CreateOutputPins() override;
 	virtual void PostProcessPinName(const UEdGraphPin* Pin, FString& DisplayName) const override;
 	// End of UAnimGraphNode_Base interface
diff --git a/Engine/Source/Editor/AnimGraph/Public/AnimGraphNode_ModifyCurve.h b/Engine/Source/Editor/AnimGraph/Public/AnimGraphNode_ModifyCurve.h
--- a/Engine/Source/Editor/AnimGraph/Public/AnimGraphNode_ModifyCurve.h
+++ b/Engine/Source/Editor/AnimGraph/Public/AnimGraphNode_ModifyCurve.h
@@ -25,10 +25,10 @@ public:
 	// UEdGraphNode interface
 	virtual FText GetTooltipText() const override;
 	virtual FText GetNodeTitle(ENodeTitleType::Type TitleType) const override;
+	virtual FText GetMenuCategory() const override;
 	// End of UEdGraphNode interface
 
 	// UAnimGraphNode_Base interface
-	virtual FString GetNodeCategory() const override;
 	virtual void CustomizePinData(UEdGraphPin* Pin, FName SourcePropertyName, int32 ArrayIndex) const override;
 	// End of UAnimGraphNode_Base interface
 
diff --git a/Engine/Source/Editor/AnimGraph/Public/AnimGraphNode_PoseSnapshot.h b/Engine/Source/Editor/AnimGraph/Public/AnimGraphNode_PoseSnapshot.h
--- a/Engine/Source/Editor/AnimGraph/Public/AnimGraphNode_PoseSnapshot.h
+++ b/Engine/Source/Editor/AnimGraph/Public/AnimGraphNode_PoseSnapshot.h
@@ -20,5 +20,6 @@ class UAnimGraphNode_PoseSnapshot : public UAnimGraphNode_Base
 	// UEdGraphNode interface
 	virtual FText GetTooltipText() const override;
 	virtual FText GetNodeTitle(ENodeTitleType::Type TitleType) const override;
+	virtual FText GetMenuCategory() const override;
 	// End of UEdGraphNode interface
 };
diff --git a/Engine/Source/Editor/AnimGraph/Public/AnimGraphNode_RigidBody.h b/Engine/Source/Editor/AnimGraph/Public/AnimGraphNode_RigidBody.h
--- a/Engine/Source/Editor/AnimGraph/Public/AnimGraphNode_RigidBody.h
+++ b/Engine/Source/Editor/AnimGraph/Public/AnimGraphNode_RigidBody.h
@@ -18,6 +18,7 @@ public:
 	// UEdGraphNode interface
 	virtual FText GetNodeTitle(ENodeTitleType::Type TitleType) const override;
 	virtual FText GetTooltipText() const override;
+	virtual FString GetNodeCategory() const override;
 	virtual void ValidateAnimNodeDuringCompilation(USkeleton* ForSkeleton, FCompilerResultsLog& MessageLog) override;
 	// End of UEdGraphNode interface
 
diff --git a/Engine/Source/Editor/AnimGraph/Public/AnimGraphNode_RotateRootBone.h b/Engine/Source/Editor/AnimGraph/Public/AnimGraphNode_RotateRootBone.h
--- a/Engine/Source/Editor/AnimGraph/Public/AnimGraphNode_RotateRootBone.h
+++ b/Engine/Source/Editor/AnimGraph/Public/AnimGraphNode_RotateRootBone.h
@@ -22,10 +22,10 @@ class UAnimGraphNode_RotateRootBone : public UAnimGraphNode_Base
 	virtual FText GetTooltipText() const override;
 	virtual FText GetNodeTitle(ENodeTitleType::Type TitleType) const override;
 	virtual void PostEditChangeProperty(struct FPropertyChangedEvent& PropertyChangedEvent) override;
+	virtual FText GetMenuCategory() const override;
 	//~ End UEdGraphNode Interface.
 
 	//~ Begin UAnimGraphNode_Base Interface
-	virtual FString GetNodeCategory() const override;
 	virtual void CustomizePinData(UEdGraphPin* Pin, FName SourcePropertyName, int32 ArrayIndex) const override;
 	//~ End UAnimGraphNode_Base Interface
 };
diff --git a/Engine/Source/Editor/AnimGraph/Public/AnimGraphNode_StateMachineBase.h b/Engine/Source/Editor/AnimGraph/Public/AnimGraphNode_StateMachineBase.h
--- a/Engine/Source/Editor/AnimGraph/Public/AnimGraphNode_StateMachineBase.h
+++ b/Engine/Source/Editor/AnimGraph/Public/AnimGraphNode_StateMachineBase.h
@@ -33,10 +33,10 @@ class ANIMGRAPH_API UAnimGraphNode_StateMachineBase : public UAnimGraphNode_Base
 	virtual FString GetDocumentationLink() const override;
 	virtual void OnRenameNode(const FString& NewName) override;
 	virtual TArray<UEdGraph*> GetSubGraphs() const override;
+	virtual FText GetMenuCategory() const override;
 	// End of UEdGraphNode interface
 
 	// UAnimGraphNode_Base interface
-	virtual FString GetNodeCategory() const override;
 	virtual void OnProcessDuringCompilation(IAnimBlueprintCompilationContext& InCompilationContext, IAnimBlueprintGeneratedClassCompiledData& OutCompiledData) override;
 	virtual void GetOutputLinkAttributes(FNodeAttributeArray& OutAttributes) const override;
 	virtual void GetRequiredExtensions(TArray<TSubclassOf<UAnimBlueprintExtension>>& OutExtensions) const override;
diff --git a/Engine/Source/Editor/AnimGraph/Public/AnimGraphNode_TransitionPoseEvaluator.h b/Engine/Source/Editor/AnimGraph/Public/AnimGraphNode_TransitionPoseEvaluator.h
--- a/Engine/Source/Editor/AnimGraph/Public/AnimGraphNode_TransitionPoseEvaluator.h
+++ b/Engine/Source/Editor/AnimGraph/Public/AnimGraphNode_TransitionPoseEvaluator.h
@@ -27,10 +27,10 @@ class UAnimGraphNode_TransitionPoseEvaluator : public UAnimGraphNode_Base
 	virtual void ValidateAnimNodeDuringCompilation(class USkeleton* ForSkeleton, class FCompilerResultsLog& MessageLog) override;
 	virtual bool CanDuplicateNode() const override { return false; }
 	virtual bool CanUserDeleteNode() const override;
+	virtual FText GetMenuCategory() const override;
 	// End of UEdGraphNode interface
 
 	// UAnimGraphNode_Base interface
-	virtual FString GetNodeCategory() const override;
 	virtual void CustomizeDetails(IDetailLayoutBuilder& DetailBuilder) override;
 	virtual void GetOutputLinkAttributes(FNodeAttributeArray& OutAttributes) const override;
 	// End of UAnimGraphNode_Base interface
diff --git a/Engine/Source/Editor/AnimGraph/Public/AnimGraphNode_TwoWayBlend.h b/Engine/Source/Editor/AnimGraph/Public/AnimGraphNode_TwoWayBlend.h
--- a/Engine/Source/Editor/AnimGraph/Public/AnimGraphNode_TwoWayBlend.h
+++ b/Engine/Source/Editor/AnimGraph/Public/AnimGraphNode_TwoWayBlend.h
@@ -20,11 +20,11 @@ class UAnimGraphNode_TwoWayBlend : public UAnimGraphNode_Base
 	virtual FLinearColor GetNodeTitleColor() const override;
 	virtual FText GetTooltipText() const override;
 	virtual FText GetNodeTitle(ENodeTitleType::Type TitleType) const override;
+	virtual FText GetMenuCategory() const override;
 	virtual void PostEditChangeProperty(struct FPropertyChangedEvent& PropertyChangedEvent) override;
 	//~ End UEdGraphNode Interface.
 
 	//~ Begin UAnimGraphNode_Base Interface
-	virtual FString GetNodeCategory() const override;
 	virtual void CustomizePinData(UEdGraphPin* Pin, FName SourcePropertyName, int32 ArrayIndex) const override;
 	virtual void CustomizeDetails(IDetailLayoutBuilder& DetailBuilder) override;
 	//~ End UAnimGraphNode_Base Interface
diff --git a/Engine/Source/Runtime/AnimGraphRuntime/Public/AnimExecutionContextLibrary.h b/Engine/Source/Runtime/AnimGraphRuntime/Public/AnimExecutionContextLibrary.h
--- a/Engine/Source/Runtime/AnimGraphRuntime/Public/AnimExecutionContextLibrary.h
+++ b/Engine/Source/Runtime/AnimGraphRuntime/Public/AnimExecutionContextLibrary.h
@@ -28,7 +28,7 @@ public:
 #endif
 	
 	/** Get the anim instance that hosts this context */
-	UFUNCTION(BlueprintPure, Category = "Execution Context", meta=(BlueprintThreadSafe))
+	UFUNCTION(BlueprintPure, Category = "Animation|Utilities", meta=(BlueprintThreadSafe))
 	static UAnimInstance* GetAnimInstance(const FAnimExecutionContext& Context);
 
 	/** Internal compiler use only - Get a reference to an anim node by index */
@@ -36,26 +36,26 @@ public:
 	static FAnimNodeReference GetAnimNodeReference(UAnimInstance* Instance, int32 Index);
 
 	/** Convert to an initialization context */
-	UFUNCTION(BlueprintCallable, Category = "Initialization Context", meta=(BlueprintThreadSafe, ExpandEnumAsExecs = "Result"))
+	UFUNCTION(BlueprintCallable, Category = "Animation|Utilities", meta=(BlueprintThreadSafe, ExpandEnumAsExecs = "Result"))
 	static FAnimInitializationContext ConvertToInitializationContext(const FAnimExecutionContext& Context, EAnimExecutionContextConversionResult& Result);
 	
 	/** Convert to an update context */
-	UFUNCTION(BlueprintCallable, Category = "Update Context", meta=(BlueprintThreadSafe, ExpandEnumAsExecs = "Result"))
+	UFUNCTION(BlueprintCallable, Category = "Animation|Utilities", meta=(BlueprintThreadSafe, ExpandEnumAsExecs = "Result"))
 	static FAnimUpdateContext ConvertToUpdateContext(const FAnimExecutionContext& Context, EAnimExecutionContextConversionResult& Result);
 
 	/** Get the current delta time in seconds */
-	UFUNCTION(BlueprintPure, Category = "Update Context", meta=(BlueprintThreadSafe))
+	UFUNCTION(BlueprintPure, Category = "Animation|Utilities", meta=(BlueprintThreadSafe))
 	static float GetDeltaTime(const FAnimUpdateContext& Context);
 
 	/** Get the current weight of this branch of the graph */
-	UFUNCTION(BlueprintPure, Category = "Update Context", meta=(BlueprintThreadSafe))
+	UFUNCTION(BlueprintPure, Category = "Animation|Utilities", meta=(BlueprintThreadSafe))
 	static float GetCurrentWeight(const FAnimUpdateContext& Context);
 	
 	/** Convert to a pose context */
-	UFUNCTION(BlueprintCallable, Category = "Pose Context", meta=(BlueprintThreadSafe, ExpandEnumAsExecs = "Result"))
+	UFUNCTION(BlueprintCallable, Category = "Animation|Utilities", meta=(BlueprintThreadSafe, ExpandEnumAsExecs = "Result"))
 	static FAnimPoseContext ConvertToPoseContext(const FAnimExecutionContext& Context, EAnimExecutionContextConversionResult& Result);
 
 	/** Convert to a component space pose context */
-	UFUNCTION(BlueprintCallable, Category = "Component Space Pose Context", meta=(BlueprintThreadSafe, ExpandEnumAsExecs = "Result"))
+	UFUNCTION(BlueprintCallable, Category = "Animation|Utilities", meta=(BlueprintThreadSafe, ExpandEnumAsExecs = "Result"))
 	static FAnimComponentSpacePoseContext ConvertToComponentSpacePoseContext(const FAnimExecutionContext& Context, EAnimExecutionContextConversionResult& Result);
 };
\ No newline at end of file
diff --git a/Engine/Source/Runtime/AnimGraphRuntime/Public/KismetAnimationLibrary.h b/Engine/Source/Runtime/AnimGraphRuntime/Public/KismetAnimationLibrary.h
--- a/Engine/Source/Runtime/AnimGraphRuntime/Public/KismetAnimationLibrary.h
+++ b/Engine/Source/Runtime/AnimGraphRuntime/Public/KismetAnimationLibrary.h
@@ -33,7 +33,7 @@ class ANIMGRAPHRUNTIME_API UKismetAnimationLibrary : public UBlueprintFunctionLi
 	 *	@param StartStretchRatio The ratio at which the bones should start to stretch. The higher the value, the later the stretching wil start.
 	 *	@param MaxStretchScale The maximum multiplier for the stretch to reach.
 	 */
-	UFUNCTION(BlueprintPure, Category = "Utilities|Animation", meta = (DisplayName = "Two Bone IK Function", ScriptName = "TwoBoneIK", bAllowStretching = "false", StartStretchRatio = "1.0", MaxStretchScale = "1.2"))
+	UFUNCTION(BlueprintPure, Category = "Animation|Utilities", meta = (DisplayName = "Two Bone IK Function", ScriptName = "TwoBoneIK", bAllowStretching = "false", StartStretchRatio = "1.0", MaxStretchScale = "1.2"))
 	static void K2_TwoBoneIK(
 		const FVector& RootPos, 
 		const FVector& JointPos, 
@@ -57,7 +57,7 @@ class ANIMGRAPHRUNTIME_API UKismetAnimationLibrary : public UBlueprintFunctionLi
 	 *	@param UpVector The position to use for the upvector target (only used is bUseUpVector is turned on)
 	 *	@param ClampConeInDegree A limit for only allowing the lookat to rotate as much as defined by the float value
 	 */
-	UFUNCTION(BlueprintPure, Category = "Utilities|Animation", meta = (DisplayName = "Look At Function", ScriptName = "LookAt", bUseUpVector = "false"))
+	UFUNCTION(BlueprintPure, Category = "Animation|Utilities", meta = (DisplayName = "Look At Function", ScriptName = "LookAt", bUseUpVector = "false"))
 	static FTransform K2_LookAt(
 		const FTransform& CurrentTransform, 
 		const FVector& TargetPosition, 
@@ -81,7 +81,7 @@ class ANIMGRAPHRUNTIME_API UKismetAnimationLibrary : public UBlueprintFunctionLi
 	 *	@param OutRangeMin The minimum for the output range (commonly == 0.0)
 	 *	@param OutRangeMax The maximum for the output range (commonly == 1.0)
 	 */
-	UFUNCTION(BlueprintPure, Category = "Utilities|Animation", meta = (NotBlueprintThreadSafe, DisplayName = "Get Distance Between Two Sockets", ScriptName = "DistanceBetweenSockets", bRemapRange= "false"))
+	UFUNCTION(BlueprintPure, Category = "Animation|Utilities", meta = (NotBlueprintThreadSafe, DisplayName = "Get Distance Between Two Sockets", ScriptName = "DistanceBetweenSockets", bRemapRange= "false"))
 	static float K2_DistanceBetweenTwoSocketsAndMapRange(
 		const USkeletalMeshComponent* Component, 
 		const FName SocketOrBoneNameA, 
@@ -101,7 +101,7 @@ class ANIMGRAPHRUNTIME_API UKismetAnimationLibrary : public UBlueprintFunctionLi
 	 *	@param SocketOrBoneNameFrom The name of the first socket / bone
 	 *	@param SocketOrBoneNameTo The name of the second socket / bone
 	 */
-	UFUNCTION(BlueprintPure, Category = "Utilities|Animation", meta = (NotBlueprintThreadSafe, DisplayName = "Get Direction Between Sockets", ScriptName = "DirectionBetweenSockets"))
+	UFUNCTION(BlueprintPure, Category = "Animation|Utilities", meta = (NotBlueprintThreadSafe, DisplayName = "Get Direction Between Sockets", ScriptName = "DirectionBetweenSockets"))
 	static FVector K2_DirectionBetweenSockets(const USkeletalMeshComponent* Component, const FName SocketOrBoneNameFrom, const FName SocketOrBoneNameTo);
 
 	/** 
@@ -117,7 +117,7 @@ class ANIMGRAPHRUNTIME_API UKismetAnimationLibrary : public UBlueprintFunctionLi
 	 *	@param RangeOutMinZ The minimum for the output range for the z component
 	 *	@param RangeOutMaxZ The maximum for the output range for the z component
 	 */
-	UFUNCTION(BlueprintPure, Category = "Utilities|Animation", meta = (DisplayName = "Make Perlin Noise Vector and Remap", ScriptName = "MakeVectorFromPerlinNoise", RangeOutMinX = "-1.f", RangeOutMaxX = "1.f", RangeOutMinY = "-1.f", RangeOutMaxY = "1.f", RangeOutMinZ = "-1.f", RangeOutMaxZ = "1.f"))
+	UFUNCTION(BlueprintPure, Category = "Animation|Utilities", meta = (DisplayName = "Make Perlin Noise Vector and Remap", ScriptName = "MakeVectorFromPerlinNoise", RangeOutMinX = "-1.f", RangeOutMaxX = "1.f", RangeOutMinY = "-1.f", RangeOutMaxY = "1.f", RangeOutMinZ = "-1.f", RangeOutMaxZ = "1.f"))
 	static FVector K2_MakePerlinNoiseVectorAndRemap(
 		float X,
 		float Y,
@@ -137,7 +137,7 @@ class ANIMGRAPHRUNTIME_API UKismetAnimationLibrary : public UBlueprintFunctionLi
 	 *	@param RangeOutMin The minimum for the output range
 	 *	@param RangeOutMax The maximum for the output range
 	 */
-	UFUNCTION(BlueprintPure, Category = "Utilities|Animation", meta = (DisplayName = "Make Perlin Noise and Remap", ScriptName = "MakeFloatFromPerlinNoise", RangeOutMinX = "-1.f", RangeOutMaxX = "1.f", RangeOutMinY = "-1.f", RangeOutMaxY = "1.f", RangeOutMinZ = "-1.f", RangeOutMaxZ = "1.f"))
+	UFUNCTION(BlueprintPure, Category = "Animation|Utilities", meta = (DisplayName = "Make Perlin Noise and Remap", ScriptName = "MakeFloatFromPerlinNoise", RangeOutMinX = "-1.f", RangeOutMaxX = "1.f", RangeOutMinY = "-1.f", RangeOutMaxY = "1.f", RangeOutMinZ = "-1.f", RangeOutMaxZ = "1.f"))
 	static float K2_MakePerlinNoiseAndRemap(
 		float Value,
 		float RangeOutMin,
@@ -155,7 +155,7 @@ class ANIMGRAPHRUNTIME_API UKismetAnimationLibrary : public UBlueprintFunctionLi
 	 *	@param VelocityMin The minimum velocity to use for normalization (if both min and max are set to 0, normalization is turned off)
 	 *	@param VelocityMax The maximum velocity to use for normalization (if both min and max are set to 0, normalization is turned off)
 	 */
-	UFUNCTION(BlueprintCallable, Category = "Utilities|Animation", meta = (DisplayName = "Calculate Velocity From Position History", ScriptName = "CalculateVelocityFromPositionHistory", NumberOfSamples = "16", VelocityMin = "0.f", VelocityMax = "128.f"))
+	UFUNCTION(BlueprintCallable, Category = "Animation|Utilities", meta = (DisplayName = "Calculate Velocity From Position History", ScriptName = "CalculateVelocityFromPositionHistory", NumberOfSamples = "16", VelocityMin = "0.f", VelocityMax = "128.f"))
 	static float K2_CalculateVelocityFromPositionHistory(
 		float DeltaSeconds,
 		FVector Position,
@@ -183,7 +183,7 @@ class ANIMGRAPHRUNTIME_API UKismetAnimationLibrary : public UBlueprintFunctionLi
 	 *	@param EasingType The easing function to use
 	 *	@param CustomCurve The curve to use if the easing type is "Custom"
 	 */
-	UFUNCTION(BlueprintCallable, Category = "Utilities|Animation", meta = (NotBlueprintThreadSafe, DisplayName = "Calculate Velocity From Sockets", ScriptName = "CalculateVelocityFromSockets", NumberOfSamples = "16", VelocityMin = "0.f", VelocityMax = "128.f"))
+	UFUNCTION(BlueprintCallable, Category = "Animation|Utilities", meta = (NotBlueprintThreadSafe, DisplayName = "Calculate Velocity From Sockets", ScriptName = "CalculateVelocityFromSockets", NumberOfSamples = "16", VelocityMin = "0.f", VelocityMax = "128.f"))
 	static float K2_CalculateVelocityFromSockets(
 		float DeltaSeconds,
 		USkeletalMeshComponent * Component,
@@ -202,7 +202,7 @@ class ANIMGRAPHRUNTIME_API UKismetAnimationLibrary : public UBlueprintFunctionLi
 	/**
 	 *	This function starts measuring the time for a profiling bracket
 	 */
-	UFUNCTION(BlueprintCallable, Category = "Utilities|Animation", meta = (DisplayName = "Start Profiling Timer"))
+	UFUNCTION(BlueprintCallable, Category = "Animation|Utilities", meta = (DisplayName = "Start Profiling Timer"))
 	static void K2_StartProfilingTimer();
 
 	/**
@@ -212,7 +212,7 @@ class ANIMGRAPHRUNTIME_API UKismetAnimationLibrary : public UBlueprintFunctionLi
 	 *  @param LogPrefix A prefix to use for the log
 	 *  @result The time spent in milliseconds
 	 */
-	UFUNCTION(BlueprintCallable, Category = "Utilities|Animation", meta = (DisplayName = "End Profiling Timer"))
+	UFUNCTION(BlueprintCallable, Category = "Animation|Utilities", meta = (DisplayName = "End Profiling Timer"))
 	static float K2_EndProfilingTimer(bool bLog = true, const FString& LogPrefix = "");
 
 	/** 
@@ -221,7 +221,7 @@ class ANIMGRAPHRUNTIME_API UKismetAnimationLibrary : public UBlueprintFunctionLi
 	 * @param	Velocity		The velocity to use as direction relative to BaseRotation
 	 * @param	BaseRotation	The base rotation, e.g. of a pawn
 	 */
-	UFUNCTION(BlueprintPure, Category="Utilities|Animation")
+	UFUNCTION(BlueprintPure, Category="Animation|Utilities")
 	static float CalculateDirection(const FVector& Velocity, const FRotator& BaseRotation);
 };
 
diff --git a/Engine/Source/Runtime/AnimGraphRuntime/Public/LinkedAnimGraphLibrary.h b/Engine/Source/Runtime/AnimGraphRuntime/Public/LinkedAnimGraphLibrary.h
--- a/Engine/Source/Runtime/AnimGraphRuntime/Public/LinkedAnimGraphLibrary.h
+++ b/Engine/Source/Runtime/AnimGraphRuntime/Public/LinkedAnimGraphLibrary.h
@@ -25,11 +25,11 @@ class ANIMGRAPHRUNTIME_API ULinkedAnimGraphLibrary : public UBlueprintFunctionLi
 
 public:
 	/** Get a linked anim graph reference from an anim node reference */
-	UFUNCTION(BlueprintCallable, Category = "Linked Anim Graph", meta=(BlueprintThreadSafe, ExpandEnumAsExecs = "Result"))
+	UFUNCTION(BlueprintCallable, Category = "Animation|Linked Anim Graphs", meta=(BlueprintThreadSafe, ExpandEnumAsExecs = "Result"))
 	static FLinkedAnimGraphReference ConvertToLinkedAnimGraph(const FAnimNodeReference& Node, EAnimNodeReferenceConversionResult& Result);
 
 	/** Get a linked anim graph reference from an anim node reference (pure) */
-	UFUNCTION(BlueprintPure, Category = "Linked Anim Graph", meta=(BlueprintThreadSafe, DisplayName = "Convert to Linked Anim Graph"))
+	UFUNCTION(BlueprintPure, Category = "Animation|Linked Anim Graphs", meta=(BlueprintThreadSafe, DisplayName = "Convert to Linked Anim Graph"))
 	static void ConvertToLinkedAnimGraphPure(const FAnimNodeReference& Node, FLinkedAnimGraphReference& LinkedAnimGraph, bool& Result)
 	{
 		EAnimNodeReferenceConversionResult ConversionResult;
@@ -38,10 +38,10 @@ public:
 	}
 	
 	/** Returns whether the node hosts an instance (e.g. linked anim graph or layer) */
-	UFUNCTION(BlueprintPure, Category = "Linked Anim Graph", meta=(BlueprintThreadSafe))
+	UFUNCTION(BlueprintPure, Category = "Animation|Linked Anim Graphs", meta=(BlueprintThreadSafe))
 	static bool HasLinkedAnimInstance(const FLinkedAnimGraphReference& Node);
 
 	/** Get the linked instance is hosted by this node. If the node does not host an instance then HasLinkedAnimInstance will return false */
-	UFUNCTION(BlueprintPure, Category = "Linked Anim Graph", meta=(BlueprintThreadSafe))
+	UFUNCTION(BlueprintPure, Category = "Animation|Linked Anim Graphs", meta=(BlueprintThreadSafe))
 	static UAnimInstance* GetLinkedAnimInstance(const FLinkedAnimGraphReference& Node);
 };
\ No newline at end of file
diff --git a/Engine/Source/Runtime/AnimGraphRuntime/Public/SequenceEvaluatorLibrary.h b/Engine/Source/Runtime/AnimGraphRuntime/Public/SequenceEvaluatorLibrary.h
--- a/Engine/Source/Runtime/AnimGraphRuntime/Public/SequenceEvaluatorLibrary.h
+++ b/Engine/Source/Runtime/AnimGraphRuntime/Public/SequenceEvaluatorLibrary.h
@@ -28,11 +28,11 @@ class ANIMGRAPHRUNTIME_API USequenceEvaluatorLibrary : public UBlueprintFunction
 
 public:
 	/** Get a sequence evaluator context from an anim node context */
-	UFUNCTION(BlueprintCallable, Category = "Sequence Evaluator", meta=(BlueprintThreadSafe, ExpandEnumAsExecs = "Result"))
+	UFUNCTION(BlueprintCallable, Category = "Animation|Sequences", meta=(BlueprintThreadSafe, ExpandEnumAsExecs = "Result"))
 	static FSequenceEvaluatorReference ConvertToSequenceEvaluator(const FAnimNodeReference& Node, EAnimNodeReferenceConversionResult& Result);
 
 	/** Get a sequence evaluator context from an anim node context (pure) */
-	UFUNCTION(BlueprintPure, Category = "Sequence Evaluator", meta=(BlueprintThreadSafe, DisplayName = "Convert to Sequence Evaluator"))
+	UFUNCTION(BlueprintPure, Category = "Animation|Sequences", meta=(BlueprintThreadSafe, DisplayName = "Convert to Sequence Evaluator"))
 	static void ConvertToSequenceEvaluatorPure(const FAnimNodeReference& Node, FSequenceEvaluatorReference& SequenceEvaluator, bool& Result)
 	{
 		EAnimNodeReferenceConversionResult ConversionResult;
@@ -41,19 +41,19 @@ public:
 	}
 	
 	/** Set the current accumulated time of the sequence evaluator */
-	UFUNCTION(BlueprintCallable, Category = "Sequence Evaluator", meta=(BlueprintThreadSafe))
+	UFUNCTION(BlueprintCallable, Category = "Animation|Sequences", meta=(BlueprintThreadSafe))
 	static FSequenceEvaluatorReference SetExplicitTime(const FSequenceEvaluatorReference& SequenceEvaluator, float Time);
 
 	/** Advance the current accumulated time of the sequence evaluator */
-	UFUNCTION(BlueprintCallable, Category = "Sequence Evaluator", meta = (BlueprintThreadSafe))
+	UFUNCTION(BlueprintCallable, Category = "Animation|Sequences", meta = (BlueprintThreadSafe))
 	static FSequenceEvaluatorReference AdvanceTime(const FAnimUpdateContext& UpdateContext, const FSequenceEvaluatorReference& SequenceEvaluator, float PlayRate = 1.0f);
 
 	/** Set the current sequence of the sequence evaluator */
-	UFUNCTION(BlueprintCallable, Category = "Sequence Evaluator", meta=(BlueprintThreadSafe))
+	UFUNCTION(BlueprintCallable, Category = "Animation|Sequences", meta=(BlueprintThreadSafe))
 	static FSequenceEvaluatorReference SetSequence(const FSequenceEvaluatorReference& SequenceEvaluator, UAnimSequenceBase* Sequence);
 
 	/** Set the current sequence of the sequence evaluator with an inertial blend time */
-	UFUNCTION(BlueprintCallable, Category = "Sequence Evaluator", meta=(BlueprintThreadSafe))
+	UFUNCTION(BlueprintCallable, Category = "Animation|Sequences", meta=(BlueprintThreadSafe))
 	static FSequenceEvaluatorReference SetSequenceWithInertialBlending(const FAnimUpdateContext& UpdateContext, const FSequenceEvaluatorReference& SequenceEvaluator, UAnimSequenceBase* Sequence, float BlendTime = 0.2f);
 	
 	/** Gets the current accumulated time of the sequence evaluator */
diff --git a/Engine/Source/Runtime/AnimGraphRuntime/Public/SequencePlayerLibrary.h b/Engine/Source/Runtime/AnimGraphRuntime/Public/SequencePlayerLibrary.h
--- a/Engine/Source/Runtime/AnimGraphRuntime/Public/SequencePlayerLibrary.h
+++ b/Engine/Source/Runtime/AnimGraphRuntime/Public/SequencePlayerLibrary.h
@@ -28,11 +28,11 @@ class ANIMGRAPHRUNTIME_API USequencePlayerLibrary : public UBlueprintFunctionLib
 
 public:
 	/** Get a sequence player context from an anim node context */
-	UFUNCTION(BlueprintCallable, Category = "Sequence Player", meta = (BlueprintThreadSafe, ExpandEnumAsExecs = "Result"))
+	UFUNCTION(BlueprintCallable, Category = "Animation|Sequences", meta = (BlueprintThreadSafe, ExpandEnumAsExecs = "Result"))
 	static FSequencePlayerReference ConvertToSequencePlayer(const FAnimNodeReference& Node, EAnimNodeReferenceConversionResult& Result);
 
 	/** Get a sequence player context from an anim node context (pure) */
-	UFUNCTION(BlueprintPure, Category = "Sequence Player", meta = (BlueprintThreadSafe, DisplayName = "Convert to Sequence Player"))
+	UFUNCTION(BlueprintPure, Category = "Animation|Sequences", meta = (BlueprintThreadSafe, DisplayName = "Convert to Sequence Player"))
 	static void ConvertToSequencePlayerPure(const FAnimNodeReference& Node, FSequencePlayerReference& SequencePlayer, bool& Result)
 	{
 		EAnimNodeReferenceConversionResult ConversionResult;
@@ -41,7 +41,7 @@ public:
 	}
 
 	/** Set the current accumulated time of the sequence player */
-	UFUNCTION(BlueprintCallable, Category = "Sequence Player", meta = (BlueprintThreadSafe))
+	UFUNCTION(BlueprintCallable, Category = "Animation|Sequences", meta = (BlueprintThreadSafe))
 	static FSequencePlayerReference SetAccumulatedTime(const FSequencePlayerReference& SequencePlayer, float Time);
 
 	/** 
@@ -49,42 +49,42 @@ public:
 	 * If this is called from On Become Relevant or On Initial Update then it should be accompanied by a call to
 	 * SetAccumulatedTime to achieve the desired effect of resetting the play time of a sequence player.
 	 */
-	UFUNCTION(BlueprintCallable, Category = "Sequence Player", meta = (BlueprintThreadSafe))
+	UFUNCTION(BlueprintCallable, Category = "Animation|Sequences", meta = (BlueprintThreadSafe))
 	static FSequencePlayerReference SetStartPosition(const FSequencePlayerReference& SequencePlayer, float StartPosition);
 
 	/** Set the play rate of the sequence player */
-	UFUNCTION(BlueprintCallable, Category = "Sequence Player", meta = (BlueprintThreadSafe))
+	UFUNCTION(BlueprintCallable, Category = "Animation|Sequences", meta = (BlueprintThreadSafe))
 	static FSequencePlayerReference SetPlayRate(const FSequencePlayerReference& SequencePlayer, float PlayRate);
 
 	/** Set the current sequence of the sequence player */
-	UFUNCTION(BlueprintCallable, Category = "Sequence Player", meta = (BlueprintThreadSafe))
+	UFUNCTION(BlueprintCallable, Category = "Animation|Sequences", meta = (BlueprintThreadSafe))
 	static FSequencePlayerReference SetSequence(const FSequencePlayerReference& SequencePlayer, UAnimSequenceBase* Sequence);
 
 	/** Set the current sequence of the sequence player with an inertial blend time */
-	UFUNCTION(BlueprintCallable, Category = "Sequence Player", meta = (BlueprintThreadSafe))
+	UFUNCTION(BlueprintCallable, Category = "Animation|Sequences", meta = (BlueprintThreadSafe))
 	static FSequencePlayerReference SetSequenceWithInertialBlending(const FAnimUpdateContext& UpdateContext, const FSequencePlayerReference& SequencePlayer, UAnimSequenceBase* Sequence, float BlendTime = 0.2f);
 
 	/** Get the current sequence of the sequence player - DEPRECATED, please use pure version */
-	UFUNCTION(BlueprintCallable, Category = "Sequence Player", meta = (BlueprintThreadSafe, DeprecatedFunction))
+	UFUNCTION(BlueprintCallable, Category = "Animation|Sequences", meta = (BlueprintThreadSafe, DeprecatedFunction))
 	static FSequencePlayerReference GetSequence(const FSequencePlayerReference& SequencePlayer, UPARAM(Ref) UAnimSequenceBase*& SequenceBase);
 	
 	/** Get the current sequence of the sequence player */
-	UFUNCTION(BlueprintPure, Category = "Sequence Player", meta = (BlueprintThreadSafe, DisplayName = "Get Sequence"))
+	UFUNCTION(BlueprintPure, Category = "Animation|Sequences", meta = (BlueprintThreadSafe, DisplayName = "Get Sequence"))
 	static UAnimSequenceBase* GetSequencePure(const FSequencePlayerReference& SequencePlayer);
 
 	/** Gets the current accumulated time of the sequence player */
-	UFUNCTION(BlueprintPure, Category = "Sequence Player", meta = (BlueprintThreadSafe))
+	UFUNCTION(BlueprintPure, Category = "Animation|Sequences", meta = (BlueprintThreadSafe))
 	static float GetAccumulatedTime(const FSequencePlayerReference& SequencePlayer);
 
 	/** Get the start position of the sequence player */
-	UFUNCTION(BlueprintPure, Category = "Sequence Player", meta = (BlueprintThreadSafe))
+	UFUNCTION(BlueprintPure, Category = "Animation|Sequences", meta = (BlueprintThreadSafe))
 	static float GetStartPosition(const FSequencePlayerReference& SequencePlayer);
 
 	/** Get the play rate of the sequence player */
-	UFUNCTION(BlueprintPure, Category = "Sequence Player", meta = (BlueprintThreadSafe))
+	UFUNCTION(BlueprintPure, Category = "Animation|Sequences", meta = (BlueprintThreadSafe))
 	static float GetPlayRate(const FSequencePlayerReference& SequencePlayer);
 
 	/** Get the looping state of the sequence player */
-	UFUNCTION(BlueprintPure, Category = "Sequence Player", meta = (BlueprintThreadSafe))
+	UFUNCTION(BlueprintPure, Category = "Animation|Sequences", meta = (BlueprintThreadSafe))
 	static bool GetLoopAnimation(const FSequencePlayerReference& SequencePlayer);
 };
\ No newline at end of file
diff --git a/Engine/Source/Runtime/AnimGraphRuntime/Public/SkeletalControlLibrary.h b/Engine/Source/Runtime/AnimGraphRuntime/Public/SkeletalControlLibrary.h
--- a/Engine/Source/Runtime/AnimGraphRuntime/Public/SkeletalControlLibrary.h
+++ b/Engine/Source/Runtime/AnimGraphRuntime/Public/SkeletalControlLibrary.h
@@ -27,11 +27,11 @@ class ANIMGRAPHRUNTIME_API USkeletalControlLibrary : public UBlueprintFunctionLi
 
 public:
 	/** Get a skeletal control from an anim node */
-	UFUNCTION(BlueprintCallable, Category = "Skeletal Control", meta=(BlueprintThreadSafe, ExpandEnumAsExecs = "Result"))
+	UFUNCTION(BlueprintCallable, Category = "Animation|Skeletal Controls", meta=(BlueprintThreadSafe, ExpandEnumAsExecs = "Result"))
 	static FSkeletalControlReference ConvertToSkeletalControl(const FAnimNodeReference& Node, EAnimNodeReferenceConversionResult& Result);
 
 	/** Get a skeletal control from an anim node (pure) */
-	UFUNCTION(BlueprintPure, Category = "Skeletal Control", meta=(BlueprintThreadSafe, DisplayName = "Convert to Skeletal Control"))
+	UFUNCTION(BlueprintPure, Category = "Animation|Skeletal Controls", meta=(BlueprintThreadSafe, DisplayName = "Convert to Skeletal Control"))
 	static void ConvertToSkeletalControlPure(const FAnimNodeReference& Node, FSkeletalControlReference& SkeletalControl, bool& Result)
 	{
 		EAnimNodeReferenceConversionResult ConversionResult;
@@ -40,10 +40,10 @@ public:
 	}
 	
 	/** Set the alpha value of this skeletal control */
-	UFUNCTION(BlueprintCallable, Category = "Skeletal Control", meta=(BlueprintThreadSafe))
+	UFUNCTION(BlueprintCallable, Category = "Animation|Skeletal Controls", meta=(BlueprintThreadSafe))
 	static FSkeletalControlReference SetAlpha(const FSkeletalControlReference& SkeletalControl, float Alpha);
 
 	/** Get the alpha value of this skeletal control */
-	UFUNCTION(BlueprintPure, Category = "Skeletal Control", meta=(BlueprintThreadSafe))
+	UFUNCTION(BlueprintPure, Category = "Animation|Skeletal Controls", meta=(BlueprintThreadSafe))
 	static float GetAlpha(const FSkeletalControlReference& SkeletalControl);
 };
\ No newline at end of file
diff --git a/Engine/Source/Runtime/Engine/Classes/Animation/AnimInstance.h b/Engine/Source/Runtime/Engine/Classes/Animation/AnimInstance.h
--- a/Engine/Source/Runtime/Engine/Classes/Animation/AnimInstance.h
+++ b/Engine/Source/Runtime/Engine/Classes/Animation/AnimInstance.h
@@ -480,14 +480,14 @@ public:
 	 * This snapshot can then be retrieved by name in the animation blueprint for blending. 
 	 * The snapshot is taken at the current LOD, so if for example you took the snapshot at LOD1 and then used it at LOD0 any bones not in LOD1 will use the reference pose 
 	 */
-	UFUNCTION(BlueprintCallable, Category = "Pose")
+	UFUNCTION(BlueprintCallable, Category = "Animation|Pose")
 	virtual void SavePoseSnapshot(FName SnapshotName);
 
 	/** Add an empty pose snapshot to the internal snapshot cache (or recycle an existing pose snapshot if the name is already in use) */
 	FPoseSnapshot& AddPoseSnapshot(FName SnapshotName);
 
 	/** Remove a previously saved pose snapshot from the internal snapshot cache */
-	UFUNCTION(BlueprintCallable, Category = "Pose")
+	UFUNCTION(BlueprintCallable, Category = "Animation|Pose")
 	void RemovePoseSnapshot(FName SnapshotName);
 
 	/** Get a cached pose snapshot by name */
@@ -498,7 +498,7 @@ public:
 	 * The snapshot is taken at the current LOD, so if for example you took the snapshot at LOD1 
 	 * and then used it at LOD0 any bones not in LOD1 will use the reference pose 
 	 */
-	UFUNCTION(BlueprintCallable, Category = "Pose")
+	UFUNCTION(BlueprintCallable, Category = "Animation|Pose")
 	virtual void SnapshotPose(UPARAM(ref) FPoseSnapshot& Snapshot);
 
 	// Can this animation instance run Update or Evaluation work in parallel
@@ -511,25 +511,25 @@ public:
 	bool NeedsUpdate() const;
 
 	/** Get whether to process notifies from any linked anim instances */
-	UFUNCTION(BlueprintPure, Category = "Notifies")
+	UFUNCTION(BlueprintPure, Category = "Animation|Notifies")
 	bool GetReceiveNotifiesFromLinkedInstances() const { return bReceiveNotifiesFromLinkedInstances; }
 
 	/** Set whether to process notifies from any linked anim instances */
-	UFUNCTION(BlueprintCallable, Category = "Notifies")
+	UFUNCTION(BlueprintCallable, Category = "Animation|Notifies")
 	void SetReceiveNotifiesFromLinkedInstances(bool bSet) { bReceiveNotifiesFromLinkedInstances = bSet; }
 
 	/** Get whether to propagate notifies to any linked anim instances */
-	UFUNCTION(BlueprintPure, Category = "Notifies")
+	UFUNCTION(BlueprintPure, Category = "Animation|Notifies")
 	bool GetPropagateNotifiesToLinkedInstances() const { return bPropagateNotifiesToLinkedInstances; }
 
 	/** Set whether to propagate notifies to any linked anim instances */
-	UFUNCTION(BlueprintCallable, Category = "Notifies")
+	UFUNCTION(BlueprintCallable, Category = "Animation|Notifies")
 	void SetPropagateNotifiesToLinkedInstances(bool bSet) { bPropagateNotifiesToLinkedInstances = bSet; }
 
-	UFUNCTION(BlueprintCallable, Category = "Animation Blueprint Linking|Montage")
+	UFUNCTION(BlueprintCallable, Category = "Animation|Linked Anim Graphs|Montage")
 	bool IsUsingMainInstanceMontageEvaluationData() const { return bUseMainInstanceMontageEvaluationData; }
 
-	UFUNCTION(BlueprintCallable, Category = "Animation Blueprint Linking|Montage")
+	UFUNCTION(BlueprintCallable, Category = "Animation|Linked Anim Graphs|Montage")
 	void SetUseMainInstanceMontageEvaluationData(bool bSet) { bUseMainInstanceMontageEvaluationData = bSet; }
 
 private:
@@ -579,23 +579,23 @@ public:
 public:
 
 	/** Play normal animation asset on the slot node by creating a dynamic UAnimMontage. You can only play one asset (whether montage or animsequence) at a time per SlotGroup. */
-	UFUNCTION(BlueprintCallable, Category="Animation")
+	UFUNCTION(BlueprintCallable, Category="Animation|Montage")
 	UAnimMontage* PlaySlotAnimationAsDynamicMontage(UAnimSequenceBase* Asset, FName SlotNodeName, float BlendInTime = 0.25f, float BlendOutTime = 0.25f, float InPlayRate = 1.f, int32 LoopCount = 1, float BlendOutTriggerTime = -1.f, float InTimeToStartMontageAt = 0.f);
 
 	/** Play normal animation asset on the slot node by creating a dynamic UAnimMontage with blend in arguments. You can only play one asset (whether montage or animsequence) at a time per SlotGroup. */
-	UFUNCTION(BlueprintCallable, Category = "Animation")
+	UFUNCTION(BlueprintCallable, Category = "Animation|Montage")
 	UAnimMontage* PlaySlotAnimationAsDynamicMontage_WithBlendArgs(UAnimSequenceBase* Asset, FName SlotNodeName, const FAlphaBlendArgs& BlendIn, const FAlphaBlendArgs& BlendOut, float InPlayRate = 1.f, int32 LoopCount = 1, float BlendOutTriggerTime = -1.f, float InTimeToStartMontageAt = 0.f);
 
 	/** Play normal animation asset on the slot node by creating a dynamic UAnimMontage with blend in settings. You can only play one asset (whether montage or animsequence) at a time per SlotGroup. */
-	UFUNCTION(BlueprintCallable, Category = "Animation")
+	UFUNCTION(BlueprintCallable, Category = "Animation|Montage")
 	UAnimMontage* PlaySlotAnimationAsDynamicMontage_WithBlendSettings(UAnimSequenceBase* Asset, FName SlotNodeName, const FMontageBlendSettings& BlendInSettings, const FMontageBlendSettings& BlendOutSettings, float InPlayRate = 1.f, int32 LoopCount = 1, float BlendOutTriggerTime = -1.f, float InTimeToStartMontageAt = 0.f);
 
 	/** Stops currently playing slot animation slot or all*/
-	UFUNCTION(BlueprintCallable, Category="Animation")
+	UFUNCTION(BlueprintCallable, Category="Animation|Montage")
 	void StopSlotAnimation(float InBlendOutTime = 0.25f, FName SlotNodeName = NAME_None);
 
 	/** Return true if it's playing the slot animation */
-	UFUNCTION(BlueprintPure, Category="Animation")
+	UFUNCTION(BlueprintPure, Category="Animation|Montage")
 	bool IsPlayingSlotAnimation(const UAnimSequenceBase* Asset, FName SlotNodeName) const;
 
 	/** Return true if this instance playing the slot animation, also returning the montage it is playing on */
@@ -606,48 +606,48 @@ public:
 	 ********************************************************************************************* */
 public:
 	/** Plays an animation montage. Returns the length of the animation montage in seconds. Returns 0.f if failed to play. */
-	UFUNCTION(BlueprintCallable, Category = "Montage")
+	UFUNCTION(BlueprintCallable, Category = "Animation|Montage")
 	float Montage_Play(UAnimMontage* MontageToPlay, float InPlayRate = 1.f, EMontagePlayReturnType ReturnValueType = EMontagePlayReturnType::MontageLength, float InTimeToStartMontageAt=0.f, bool bStopAllMontages = true);
 
 	/** Plays an animation montage. Same as Montage_Play, but you can specify an AlphaBlend for Blend In settings. */
-	UFUNCTION(BlueprintCallable, Category = "Montage")
+	UFUNCTION(BlueprintCallable, Category = "Animation|Montage")
 	float Montage_PlayWithBlendIn(UAnimMontage* MontageToPlay, const FAlphaBlendArgs& BlendIn, float InPlayRate = 1.f, EMontagePlayReturnType ReturnValueType = EMontagePlayReturnType::MontageLength, float InTimeToStartMontageAt=0.f, bool bStopAllMontages = true);
 
 	/** Plays an animation montage. Same as Montage_Play, but you can overwrite all of the montage's default blend in settings. */
-	UFUNCTION(BlueprintCallable, Category = "Montage")
+	UFUNCTION(BlueprintCallable, Category = "Animation|Montage")
 	float Montage_PlayWithBlendSettings(UAnimMontage* MontageToPlay, const FMontageBlendSettings& BlendInSettings, float InPlayRate = 1.f, EMontagePlayReturnType ReturnValueType = EMontagePlayReturnType::MontageLength, float InTimeToStartMontageAt=0.f, bool bStopAllMontages = true);
 
 	/** Stops the animation montage. If reference is NULL, it will stop ALL active montages. */
 	/** Stopped montages will blend out using their montage asset's BlendOut, with InBlendOutTime as the BlendTime */
-	UFUNCTION(BlueprintCallable, Category = "Montage")
+	UFUNCTION(BlueprintCallable, Category = "Animation|Montage")
 	void Montage_Stop(float InBlendOutTime, const UAnimMontage* Montage = NULL);
 
 	/** Same as Montage_Stop. Uses values from the AlphaBlendArgs. Other settings come from the montage asset*/
-	UFUNCTION(BlueprintCallable, Category = "Montage")
+	UFUNCTION(BlueprintCallable, Category = "Animation|Montage")
 	void Montage_StopWithBlendOut(const FAlphaBlendArgs& BlendOut, const UAnimMontage* Montage = nullptr);
 
 	/** Same as Montage_Stop, but all blend settings are provided instead of using the ones on the montage asset*/
-	UFUNCTION(BlueprintCallable, Category = "Montage")
+	UFUNCTION(BlueprintCallable, Category = "Animation|Montage")
 	void Montage_StopWithBlendSettings(const FMontageBlendSettings& BlendOutSettings, const UAnimMontage* Montage = nullptr);
 
 	/** Stops all active montages belonging to a group. */
-	UFUNCTION(BlueprintCallable, Category = "Montage")
+	UFUNCTION(BlueprintCallable, Category = "Animation|Montage")
 	void Montage_StopGroupByName(float InBlendOutTime, FName GroupName);
 
 	/** Pauses the animation montage. If reference is NULL, it will pause ALL active montages. */
-	UFUNCTION(BlueprintCallable, Category = "Montage")
+	UFUNCTION(BlueprintCallable, Category = "Animation|Montage")
 	void Montage_Pause(const UAnimMontage* Montage = NULL);
 
 	/** Resumes a paused animation montage. If reference is NULL, it will resume ALL active montages. */
-	UFUNCTION(BlueprintCallable, Category = "Montage")
+	UFUNCTION(BlueprintCallable, Category = "Animation|Montage")
 	void Montage_Resume(const UAnimMontage* Montage);
 
 	/** Makes a montage jump to a named section. If Montage reference is NULL, it will do that to all active montages. */
-	UFUNCTION(BlueprintCallable, Category="Montage")
+	UFUNCTION(BlueprintCallable, Category="Animation|Montage")
 	void Montage_JumpToSection(FName SectionName, const UAnimMontage* Montage = NULL);
 
 	/** Makes a montage jump to the end of a named section. If Montage reference is NULL, it will do that to all active montages. */
-	UFUNCTION(BlueprintCallable, Category="Montage")
+	UFUNCTION(BlueprintCallable, Category="Animation|Montage")
 	void Montage_JumpToSectionsEnd(FName SectionName, const UAnimMontage* Montage = NULL);
 
 	/** Relink new next section AFTER SectionNameToChange in run-time
@@ -659,47 +659,47 @@ public:
 	 * @param SectionNameToChange : This should be the name of the Montage Section after which you want to insert a new next section
 	 * @param NextSection	: new next section 
 	 */
-	UFUNCTION(BlueprintCallable, Category="Montage")
+	UFUNCTION(BlueprintCallable, Category="Animation|Montage")
 	void Montage_SetNextSection(FName SectionNameToChange, FName NextSection, const UAnimMontage* Montage = NULL);
 
 	/** Change AnimMontage play rate. NewPlayRate = 1.0 is the default playback rate. */
-	UFUNCTION(BlueprintCallable, Category="Montage")
+	UFUNCTION(BlueprintCallable, Category="Animation|Montage")
 	void Montage_SetPlayRate(const UAnimMontage* Montage, float NewPlayRate = 1.f);
 
 	/** Returns true if the animation montage is active. If the Montage reference is NULL, it will return true if any Montage is active. */
-	UFUNCTION(BlueprintPure, Category="Montage")
+	UFUNCTION(BlueprintPure, Category="Animation|Montage")
 	bool Montage_IsActive(const UAnimMontage* Montage) const;
 
 	/** Returns true if the animation montage is currently active and playing. 
 	If reference is NULL, it will return true is ANY montage is currently active and playing. */
-	UFUNCTION(BlueprintPure, Category="Montage")
+	UFUNCTION(BlueprintPure, Category="Animation|Montage")
 	bool Montage_IsPlaying(const UAnimMontage* Montage) const;
 
 	/** Returns the name of the current animation montage section. */
-	UFUNCTION(BlueprintPure, Category="Montage")
+	UFUNCTION(BlueprintPure, Category="Animation|Montage")
 	FName Montage_GetCurrentSection(const UAnimMontage* Montage = NULL) const;
 
 	/** Get Current Montage Position */
-	UFUNCTION(BlueprintPure, Category = "Montage")
+	UFUNCTION(BlueprintPure, Category = "Animation|Montage")
 	float Montage_GetPosition(const UAnimMontage* Montage) const;
 	
 	/** Set position. */
-	UFUNCTION(BlueprintCallable, Category = "Montage")
+	UFUNCTION(BlueprintCallable, Category = "Animation|Montage")
 	void Montage_SetPosition(const UAnimMontage* Montage, float NewPosition);
 	
 	/** return true if Montage is not currently active. (not valid or blending out) */
-	UFUNCTION(BlueprintPure, Category = "Montage")
+	UFUNCTION(BlueprintPure, Category = "Animation|Montage")
 	bool Montage_GetIsStopped(const UAnimMontage* Montage) const;
 
 	/** Get the current blend time of the Montage.
 	If Montage reference is NULL, it will return the current blend time on the first active Montage found. */
-	UFUNCTION(BlueprintPure, Category = "Montage")
+	UFUNCTION(BlueprintPure, Category = "Animation|Montage")
 	float Montage_GetBlendTime(const UAnimMontage* Montage) const;
 
 	/** Get PlayRate for Montage.
 	If Montage reference is NULL, PlayRate for any Active Montage will be returned.
 	If Montage is not playing, 0 is returned. */
-	UFUNCTION(BlueprintPure, Category = "Montage")
+	UFUNCTION(BlueprintPure, Category = "Animation|Montage")
 	float Montage_GetPlayRate(const UAnimMontage* Montage) const;
 
 	/*********************************************************************************************
@@ -719,12 +719,12 @@ public:
 	void MontageSync_StopFollowing(const UAnimMontage* MontageFollower);
 
 	/** Returns true if any montage is playing currently. Doesn't mean it's active though, it could be blending out. */
-	UFUNCTION(BlueprintPure, Category = "Montage")
+	UFUNCTION(BlueprintPure, Category = "Animation|Montage")
 	bool IsAnyMontagePlaying() const;
 
 	/** Get a current Active Montage in this AnimInstance. 
 		Note that there might be multiple Active at the same time. This will only return the first active one it finds. **/
-	UFUNCTION(BlueprintPure, Category = "Montage")
+	UFUNCTION(BlueprintPure, Category = "Animation|Montage")
 	UAnimMontage* GetCurrentActiveMontage() const;
 
 	/** Called when a montage starts blending out, whether interrupted or finished */
@@ -793,7 +793,7 @@ public:
 	UAnimInstance* GetSubInstanceByTag(FName InTag) const { return GetLinkedAnimGraphInstanceByTag(InTag); }
 
 	/** Runs through all nodes, attempting to find the first linked instance by name/tag */
-	UFUNCTION(BlueprintPure, Category = "Animation Blueprint Linking")
+	UFUNCTION(BlueprintPure, Category = "Animation|Linked Anim Graphs")
 	UAnimInstance* GetLinkedAnimGraphInstanceByTag(FName InTag) const;
 
 	UE_DEPRECATED(4.24, "Function renamed, please use GetLinkedAnimGraphInstancesByTag")
@@ -801,14 +801,14 @@ public:
 
 	/** Runs through all nodes, attempting to find all linked instances that match the name/tag */
 	UE_DEPRECATED(5.0, "Tags are unique so this function is no longer supported. Please use GetLinkedAnimGraphInstanceByTag instead")
-	UFUNCTION(BlueprintPure, Category = "Animation Blueprint Linking", meta=(DeprecatedFunction, DeprecationMessage="Tags are unique so this function is no longer supported. Please use GetLinkedAnimGraphInstanceByTag instead"))
+	UFUNCTION(BlueprintPure, Category = "Animation|Linked Anim Graphs", meta=(DeprecatedFunction, DeprecationMessage="Tags are unique so this function is no longer supported. Please use GetLinkedAnimGraphInstanceByTag instead"))
 	void GetLinkedAnimGraphInstancesByTag(FName InTag, TArray<UAnimInstance*>& OutLinkedInstances) const;
 
 	UE_DEPRECATED(4.24, "Function renamed, please use LinkAnimGraphByTag")
 	void SetSubInstanceClassByTag(FName InTag, TSubclassOf<UAnimInstance> InClass) { LinkAnimGraphByTag(InTag, InClass); }
 
 	/** Runs through all nodes, attempting to find a linked instance by name/tag, then sets the class of each node if the tag matches */
-	UFUNCTION(BlueprintCallable, Category = "Animation Blueprint Linking")
+	UFUNCTION(BlueprintCallable, Category = "Animation|Linked Anim Graphs")
 	void LinkAnimGraphByTag(FName InTag, TSubclassOf<UAnimInstance> InClass);
 
 	UE_DEPRECATED(4.24, "Function renamed, please use LinkAnimClassLayers")
@@ -820,7 +820,7 @@ public:
 	 * and a separate linked instance is allocated for each layer node.
 	 * If InClass is null, then all layers are reset to their defaults.
 	 */
-	UFUNCTION(BlueprintCallable, Category = "Animation Blueprint Linking")
+	UFUNCTION(BlueprintCallable, Category = "Animation|Linked Anim Graphs")
 	virtual void LinkAnimClassLayers(TSubclassOf<UAnimInstance> InClass);
 
 	UE_DEPRECATED(4.24, "Function renamed, please use UnlinkAnimClassLayers")
@@ -831,29 +831,29 @@ public:
 	 * State sharing rules are as with SetLayerOverlay.
 	 * If InClass is null, does nothing.
 	 */
-	UFUNCTION(BlueprintCallable, Category = "Animation Blueprint Linking")
+	UFUNCTION(BlueprintCallable, Category = "Animation|Linked Anim Graphs")
 	virtual void UnlinkAnimClassLayers(TSubclassOf<UAnimInstance> InClass);
 
 	UE_DEPRECATED(4.24, "Function renamed, please use GetLinkedAnimLayerInstanceByGroup")
 	UAnimInstance* GetLayerSubInstanceByGroup(FName InGroup) const { return GetLinkedAnimLayerInstanceByGroup(InGroup); }
 
 	/** Gets the layer linked instance corresponding to the specified group */
-	UFUNCTION(BlueprintPure, Category = "Animation Blueprint Linking")
+	UFUNCTION(BlueprintPure, Category = "Animation|Linked Anim Graphs")
 	UAnimInstance* GetLinkedAnimLayerInstanceByGroup(FName InGroup) const;
 
 	/** Runs through all nodes, attempting to find all distinct layer linked instances in the group */
-	UFUNCTION(BlueprintPure, Category = "Animation Blueprint Linking")
+	UFUNCTION(BlueprintPure, Category = "Animation|Linked Anim Graphs")
 	void GetLinkedAnimLayerInstancesByGroup(FName InGroup, TArray<UAnimInstance*>& OutLinkedInstances) const;
 
 	/** Gets layer linked instance that matches group and class */
-	UFUNCTION(BlueprintPure, Category = "Animation Blueprint Linking")
+	UFUNCTION(BlueprintPure, Category = "Animation|Linked Anim Graphs")
 	UAnimInstance* GetLinkedAnimLayerInstanceByGroupAndClass(FName InGroup, TSubclassOf<UAnimInstance> InClass) const;
 
 	UE_DEPRECATED(4.24, "Function renamed, please use GetLinkedAnimLayerInstanceByClass")
 	UAnimInstance* GetLayerSubInstanceByClass(TSubclassOf<UAnimInstance> InClass) const { return GetLinkedAnimLayerInstanceByClass(InClass); }
 
 	/** Gets the first layer linked instance corresponding to the specified class */
-	UFUNCTION(BlueprintPure, Category = "Animation Blueprint Linking")
+	UFUNCTION(BlueprintPure, Category = "Animation|Linked Anim Graphs")
 	UAnimInstance* GetLinkedAnimLayerInstanceByClass(TSubclassOf<UAnimInstance> InClass) const;
 
 	/** Sets up initial layer groupings */
@@ -917,7 +917,7 @@ public:
 	void RequestMontageInertialization(const UAnimMontage* Montage, float Duration, const UBlendProfile* BlendProfile = nullptr);
 
 	/**  Requests an inertial blend during the next anim graph update. Requires your anim graph to have a slot node belonging to the specified group name */
-	UFUNCTION(BlueprintCallable, Category = "Inertial Blend")
+	UFUNCTION(BlueprintCallable, Category = "Animation|Inertial Blending")
 	void RequestSlotGroupInertialization(FName InSlotGroupName, float Duration, const UBlendProfile* BlendProfile = nullptr);
 
 	/** Queue a Montage BlendingOut Event to be triggered. */
@@ -950,7 +950,7 @@ public:
 #endif
 
 	/** Set RootMotionMode */
-	UFUNCTION(BlueprintCallable, Category = "Animation")
+	UFUNCTION(BlueprintCallable, Category = "Animation|Root Motion")
 	void SetRootMotionMode(TEnumAsByte<ERootMotionMode::Type> Value);
 
 	/** 
@@ -976,103 +976,103 @@ public:
 	 */
 
 	/** Gets the length in seconds of the asset referenced in an asset player node */
-	UFUNCTION(BlueprintPure, Category="Asset Player", meta=(DisplayName="Length", BlueprintInternalUseOnly="true", AnimGetter="true", BlueprintThreadSafe))
+	UFUNCTION(BlueprintPure, Category="Animation|Asset Player", meta=(DisplayName="Length", BlueprintInternalUseOnly="true", AnimGetter="true", BlueprintThreadSafe))
 	float GetInstanceAssetPlayerLength(int32 AssetPlayerIndex);
 
 	/** Get the current accumulated time in seconds for an asset player node */
-	UFUNCTION(BlueprintPure, Category="Asset Player", meta = (DisplayName = "Current Time", BlueprintInternalUseOnly = "true", AnimGetter = "true", BlueprintThreadSafe))
+	UFUNCTION(BlueprintPure, Category="Animation|Asset Player", meta = (DisplayName = "Current Time", BlueprintInternalUseOnly = "true", AnimGetter = "true", BlueprintThreadSafe))
 	float GetInstanceAssetPlayerTime(int32 AssetPlayerIndex);
 
 	/** Get the current accumulated time as a fraction for an asset player node */
-	UFUNCTION(BlueprintPure, Category="Asset Player", meta=(DisplayName="Current Time (ratio)", BlueprintInternalUseOnly="true", AnimGetter="true", BlueprintThreadSafe))
+	UFUNCTION(BlueprintPure, Category="Animation|Asset Player", meta=(DisplayName="Current Time (ratio)", BlueprintInternalUseOnly="true", AnimGetter="true", BlueprintThreadSafe))
 	float GetInstanceAssetPlayerTimeFraction(int32 AssetPlayerIndex);
 
 	/** Get the time in seconds from the end of an animation in an asset player node */
-	UFUNCTION(BlueprintPure, Category="Asset Player", meta=(DisplayName="Time Remaining", BlueprintInternalUseOnly="true", AnimGetter="true", BlueprintThreadSafe))
+	UFUNCTION(BlueprintPure, Category="Animation|Asset Player", meta=(DisplayName="Time Remaining", BlueprintInternalUseOnly="true", AnimGetter="true", BlueprintThreadSafe))
 	float GetInstanceAssetPlayerTimeFromEnd(int32 AssetPlayerIndex);
 
 	/** Get the time as a fraction of the asset length of an animation in an asset player node */
-	UFUNCTION(BlueprintPure, Category="Asset Player", meta=(DisplayName="Time Remaining (ratio)", BlueprintInternalUseOnly="true", AnimGetter="true", BlueprintThreadSafe))
+	UFUNCTION(BlueprintPure, Category="Animation|Asset Player", meta=(DisplayName="Time Remaining (ratio)", BlueprintInternalUseOnly="true", AnimGetter="true", BlueprintThreadSafe))
 	float GetInstanceAssetPlayerTimeFromEndFraction(int32 AssetPlayerIndex);
 
 	/** Get the blend weight of a specified state machine */
-	UFUNCTION(BlueprintPure, Category = "States", meta = (DisplayName = "Machine Weight", BlueprintInternalUseOnly = "true", AnimGetter = "true", BlueprintThreadSafe))
+	UFUNCTION(BlueprintPure, Category = "Animation|State Machines", meta = (DisplayName = "Machine Weight", BlueprintInternalUseOnly = "true", AnimGetter = "true", BlueprintThreadSafe))
 	float GetInstanceMachineWeight(int32 MachineIndex);
 
 	/** Get the blend weight of a specified state */
-	UFUNCTION(BlueprintPure, Category="States", meta = (DisplayName="State Weight", BlueprintInternalUseOnly = "true", AnimGetter="true", BlueprintThreadSafe))
+	UFUNCTION(BlueprintPure, Category="Animation|State Machines", meta = (DisplayName="State Weight", BlueprintInternalUseOnly = "true", AnimGetter="true", BlueprintThreadSafe))
 	float GetInstanceStateWeight(int32 MachineIndex, int32 StateIndex);
 
 	/** Get the current elapsed time of a state within the specified state machine */
-	UFUNCTION(BlueprintPure, Category="States", meta = (DisplayName="Current State Time", BlueprintInternalUseOnly = "true", AnimGetter="true", GetterContext="Transition", BlueprintThreadSafe))
+	UFUNCTION(BlueprintPure, Category="Animation|State Machines", meta = (DisplayName="Current State Time", BlueprintInternalUseOnly = "true", AnimGetter="true", GetterContext="Transition", BlueprintThreadSafe))
 	float GetInstanceCurrentStateElapsedTime(int32 MachineIndex);
 
 	/** Get the crossfade duration of a specified transition */
-	UFUNCTION(BlueprintPure, Category="Transitions", meta = (DisplayName="Get Transition Crossfade Duration", BlueprintInternalUseOnly = "true", AnimGetter="true", BlueprintThreadSafe))
+	UFUNCTION(BlueprintPure, Category="Animation|State Machines", meta = (DisplayName="Get Transition Crossfade Duration", BlueprintInternalUseOnly = "true", AnimGetter="true", BlueprintThreadSafe))
 	float GetInstanceTransitionCrossfadeDuration(int32 MachineIndex, int32 TransitionIndex);
 
 	/** Get the elapsed time in seconds of a specified transition */
-	UFUNCTION(BlueprintPure, Category="Transitions", meta = (DisplayName="Get Transition Time Elapsed", BlueprintInternalUseOnly = "true", AnimGetter="true", GetterContext="CustomBlend", BlueprintThreadSafe))
+	UFUNCTION(BlueprintPure, Category="Animation|State Machines", meta = (DisplayName="Get Transition Time Elapsed", BlueprintInternalUseOnly = "true", AnimGetter="true", GetterContext="CustomBlend", BlueprintThreadSafe))
 	float GetInstanceTransitionTimeElapsed(int32 MachineIndex, int32 TransitionIndex);
 
 	/** Get the elapsed time as a fraction of the crossfade duration of a specified transition */
-	UFUNCTION(BlueprintPure, Category="Transitions", meta = (DisplayName="Get Transition Time Elapsed (ratio)", BlueprintInternalUseOnly = "true", AnimGetter="true", GetterContext="CustomBlend", BlueprintThreadSafe))
+	UFUNCTION(BlueprintPure, Category="Animation|State Machines", meta = (DisplayName="Get Transition Time Elapsed (ratio)", BlueprintInternalUseOnly = "true", AnimGetter="true", GetterContext="CustomBlend", BlueprintThreadSafe))
 	float GetInstanceTransitionTimeElapsedFraction(int32 MachineIndex, int32 TransitionIndex);
 
 	/** Get the time remaining in seconds for the most relevant animation in the source state */
-	UFUNCTION(BlueprintPure, Category="Asset Player", meta = (BlueprintInternalUseOnly = "true", AnimGetter="true", GetterContext="Transition", BlueprintThreadSafe))
+	UFUNCTION(BlueprintPure, Category="Animation|State Machines", meta = (BlueprintInternalUseOnly = "true", AnimGetter="true", GetterContext="Transition", BlueprintThreadSafe))
 	float GetRelevantAnimTimeRemaining(int32 MachineIndex, int32 StateIndex);
 
 	/** Get the time remaining as a fraction of the duration for the most relevant animation in the source state */
-	UFUNCTION(BlueprintPure, Category = "Asset Player", meta = (BlueprintInternalUseOnly = "true", AnimGetter = "true", GetterContext = "Transition", BlueprintThreadSafe))
+	UFUNCTION(BlueprintPure, Category = "Animation|State Machines", meta = (BlueprintInternalUseOnly = "true", AnimGetter = "true", GetterContext = "Transition", BlueprintThreadSafe))
 	float GetRelevantAnimTimeRemainingFraction(int32 MachineIndex, int32 StateIndex);
 
 	/** Get the length in seconds of the most relevant animation in the source state */
-	UFUNCTION(BlueprintPure, Category = "Asset Player", meta = (BlueprintInternalUseOnly = "true", AnimGetter = "true", GetterContext = "Transition", BlueprintThreadSafe))
+	UFUNCTION(BlueprintPure, Category = "Animation|State Machines", meta = (BlueprintInternalUseOnly = "true", AnimGetter = "true", GetterContext = "Transition", BlueprintThreadSafe))
 	float GetRelevantAnimLength(int32 MachineIndex, int32 StateIndex);
 
 	/** Get the current accumulated time in seconds for the most relevant animation in the source state */
-	UFUNCTION(BlueprintPure, Category = "Asset Player", meta = (BlueprintInternalUseOnly = "true", AnimGetter = "true", GetterContext = "Transition", BlueprintThreadSafe))
+	UFUNCTION(BlueprintPure, Category = "Animation|State Machines", meta = (BlueprintInternalUseOnly = "true", AnimGetter = "true", GetterContext = "Transition", BlueprintThreadSafe))
 	float GetRelevantAnimTime(int32 MachineIndex, int32 StateIndex);
 
 	/** Get the current accumulated time as a fraction of the length of the most relevant animation in the source state */
-	UFUNCTION(BlueprintPure, Category = "Asset Player", meta = (BlueprintInternalUseOnly = "true", AnimGetter = "true", GetterContext = "Transition", BlueprintThreadSafe))
+	UFUNCTION(BlueprintPure, Category = "Animation|State Machines", meta = (BlueprintInternalUseOnly = "true", AnimGetter = "true", GetterContext = "Transition", BlueprintThreadSafe))
 	float GetRelevantAnimTimeFraction(int32 MachineIndex, int32 StateIndex);
 
 	/** Get whether a particular notify state was active in any state machine last tick.*/
-	UFUNCTION(BlueprintPure, Category = "Asset Player", meta = (BlueprintThreadSafe, Keywords = "Window,TransitionWindow,NotifyState,NotifyStateTransition"))
+	UFUNCTION(BlueprintPure, Category = "Animation|State Machines", meta = (BlueprintThreadSafe, Keywords = "Window,TransitionWindow,NotifyState,NotifyStateTransition"))
 	bool WasAnimNotifyStateActiveInAnyState(TSubclassOf<UAnimNotifyState> AnimNotifyStateType);
 
 	/** Get whether a particular notify state is active in a specific state machine last tick.  */
-	UFUNCTION(BlueprintPure, Category = "Asset Player", meta = (BlueprintInternalUseOnly = "true", AnimGetter = "true", GetterContext = "Transition", BlueprintThreadSafe, Keywords = "Window,TransitionWindow,NotifyState,NotifyStateTransition"))
+	UFUNCTION(BlueprintPure, Category = "Animation|State Machines", meta = (BlueprintInternalUseOnly = "true", AnimGetter = "true", GetterContext = "Transition", BlueprintThreadSafe, Keywords = "Window,TransitionWindow,NotifyState,NotifyStateTransition"))
 	bool WasAnimNotifyStateActiveInStateMachine(int32 MachineIndex, TSubclassOf<UAnimNotifyState> AnimNotifyStateType);
 
 	/** Get whether a particular notify state is active in a specific state last tick. */
-	UFUNCTION(BlueprintPure, Category = "Asset Player", meta = (BlueprintInternalUseOnly = "true", AnimGetter = "true", GetterContext = "Transition", BlueprintThreadSafe, Keywords = "Window,TransitionWindow,NotifyState,NotifyStateTransition"))
+	UFUNCTION(BlueprintPure, Category = "Animation|State Machines", meta = (BlueprintInternalUseOnly = "true", AnimGetter = "true", GetterContext = "Transition", BlueprintThreadSafe, Keywords = "Window,TransitionWindow,NotifyState,NotifyStateTransition"))
 	bool WasAnimNotifyStateActiveInSourceState(int32 MachineIndex, int32 StateIndex, TSubclassOf<UAnimNotifyState> AnimNotifyStateType);
 
 	/** Get whether the most relevant animation was in a particular notify state last tick. */
-	UFUNCTION(BlueprintPure, Category = "Asset Player", meta = (BlueprintInternalUseOnly = "true", AnimGetter = "true", GetterContext = "Transition", BlueprintThreadSafe, Keywords = "Window,TransitionWindow,NotifyState,NotifyStateTransition"))
+	UFUNCTION(BlueprintPure, Category = "Animation|State Machines", meta = (BlueprintInternalUseOnly = "true", AnimGetter = "true", GetterContext = "Transition", BlueprintThreadSafe, Keywords = "Window,TransitionWindow,NotifyState,NotifyStateTransition"))
     bool WasAnimNotifyTriggeredInSourceState(int32 MachineIndex, int32 StateIndex,  TSubclassOf<UAnimNotify> AnimNotifyType);
 
 	/** Get whether the most relevant animation triggered the animation notify with the specified name last tick.. */
-	UFUNCTION(BlueprintPure, Category = "Asset Player", meta = (BlueprintInternalUseOnly = "true", AnimGetter = "true", GetterContext = "Transition", BlueprintThreadSafe, Keywords = "Window,TransitionWindow,NotifyState,NotifyStateTransition"))
+	UFUNCTION(BlueprintPure, Category = "Animation|State Machines", meta = (BlueprintInternalUseOnly = "true", AnimGetter = "true", GetterContext = "Transition", BlueprintThreadSafe, Keywords = "Window,TransitionWindow,NotifyState,NotifyStateTransition"))
     bool WasAnimNotifyNameTriggeredInSourceState(int32 MachineIndex, int32 StateIndex, FName NotifyName);
 
 	/** Get whether a particular notify type was active in a specific state machine last tick.  */
-	UFUNCTION(BlueprintPure, Category = "Asset Player", meta = (BlueprintInternalUseOnly = "true", AnimGetter = "true", GetterContext = "Transition", BlueprintThreadSafe, Keywords = "Window,TransitionWindow,NotifyState,NotifyStateTransition"))
+	UFUNCTION(BlueprintPure, Category = "Animation|State Machines", meta = (BlueprintInternalUseOnly = "true", AnimGetter = "true", GetterContext = "Transition", BlueprintThreadSafe, Keywords = "Window,TransitionWindow,NotifyState,NotifyStateTransition"))
     bool WasAnimNotifyTriggeredInStateMachine(int32 MachineIndex, TSubclassOf<UAnimNotify> AnimNotifyType);
 
 	/** Get whether the given state machine triggered the animation notify with the specified name last tick. */
-	UFUNCTION(BlueprintPure, Category = "Asset Player", meta = (BlueprintInternalUseOnly = "true", AnimGetter = "true", GetterContext = "Transition", BlueprintThreadSafe, Keywords = "Window,TransitionWindow,NotifyState,NotifyStateTransition"))
+	UFUNCTION(BlueprintPure, Category = "Animation|State Machines", meta = (BlueprintInternalUseOnly = "true", AnimGetter = "true", GetterContext = "Transition", BlueprintThreadSafe, Keywords = "Window,TransitionWindow,NotifyState,NotifyStateTransition"))
     bool WasAnimNotifyNameTriggeredInStateMachine(int32 MachineIndex, FName NotifyName);
 	
 	/**  Get whether an animation notify of a given type was triggered last tick. */
-	UFUNCTION(BlueprintPure, Category = "Asset Player", meta = (BlueprintInternalUseOnly = "true", AnimGetter = "true", GetterContext = "Transition", BlueprintThreadSafe, Keywords = "Window,TransitionWindow,NotifyState,NotifyStateTransition"))
+	UFUNCTION(BlueprintPure, Category = "Animation|State Machines", meta = (BlueprintInternalUseOnly = "true", AnimGetter = "true", GetterContext = "Transition", BlueprintThreadSafe, Keywords = "Window,TransitionWindow,NotifyState,NotifyStateTransition"))
     bool WasAnimNotifyTriggeredInAnyState(TSubclassOf<UAnimNotify> AnimNotifyType);
 
 	/** Get whether the animation notify with the specified name triggered last tick. */
-	UFUNCTION(BlueprintPure, Category = "Asset Player", meta = (BlueprintInternalUseOnly = "true", AnimGetter = "true", GetterContext = "Transition", BlueprintThreadSafe, Keywords = "Window,TransitionWindow,NotifyState,NotifyStateTransition"))
+	UFUNCTION(BlueprintPure, Category = "Animation|State Machines", meta = (BlueprintInternalUseOnly = "true", AnimGetter = "true", GetterContext = "Transition", BlueprintThreadSafe, Keywords = "Window,TransitionWindow,NotifyState,NotifyStateTransition"))
     bool WasAnimNotifyNameTriggeredInAnyState(FName NotifyName);
 	
 	/** Gets the runtime instance of the specified state machine by Name */
@@ -1114,30 +1114,30 @@ public:
 
 public:
 	/** Returns the value of a named curve. */
-	UFUNCTION(BlueprintPure, Category="Animation", meta=(BlueprintThreadSafe))
+	UFUNCTION(BlueprintPure, Category="Animation|Curves", meta=(BlueprintThreadSafe))
 	float GetCurveValue(FName CurveName) const;
 
 	/** This returns last up-to-date list of active curve names */
-	UFUNCTION(BlueprintPure, Category = "Animation", meta=(BlueprintThreadSafe))
+	UFUNCTION(BlueprintPure, Category = "Animation|Curves", meta=(BlueprintThreadSafe))
 	void GetActiveCurveNames(EAnimCurveType CurveType, TArray<FName>& OutNames) const;
 
 	/* This returns all curve names */
-	UFUNCTION(BlueprintPure, Category = "Animation", meta=(BlueprintThreadSafe))
+	UFUNCTION(BlueprintPure, Category = "Animation|Curves", meta=(BlueprintThreadSafe))
 	void GetAllCurveNames(TArray<FName>& OutNames) const;
 
 	/** Returns value of named curved in OutValue, returns whether the curve was actually found or not. */
 	bool GetCurveValue(FName CurveName, float& OutValue) const;
 
 	/** Returns the name of a currently active state in a state machine. */
-	UFUNCTION(BlueprintPure, Category="Animation", meta=(BlueprintInternalUseOnly = "true", AnimGetter = "true", BlueprintThreadSafe))
+	UFUNCTION(BlueprintPure, Category="Animation|State Machines", meta=(BlueprintInternalUseOnly = "true", AnimGetter = "true", BlueprintThreadSafe))
 	FName GetCurrentStateName(int32 MachineIndex);
 
 	/** Sets a morph target to a certain weight. */
-	UFUNCTION(BlueprintCallable, Category="Animation")
+	UFUNCTION(BlueprintCallable, Category="Animation|Morph Targets")
 	void SetMorphTarget(FName MorphTargetName, float Value);
 
 	/** Clears the current morph targets. */
-	UFUNCTION(BlueprintCallable, Category="Animation")
+	UFUNCTION(BlueprintCallable, Category="Animation|Morph Targets")
 	void ClearMorphTargets();
 
 	UE_DEPRECATED(5.0, "Please use UKismetAnimationLibrary::CalculateDirection instead")
@@ -1156,16 +1156,16 @@ public:
 	void UnlockAIResources(bool bUnlockMovement, bool UnlockAILogic);
 	//--- AI communication end ---//
 
-	UFUNCTION(BlueprintCallable, Category = "SyncGroup", meta=(BlueprintThreadSafe))
+	UFUNCTION(BlueprintCallable, Category = "Animation|Synchronization", meta=(BlueprintThreadSafe))
 	bool GetTimeToClosestMarker(FName SyncGroup, FName MarkerName, float& OutMarkerTime) const;
 
-	UFUNCTION(BlueprintCallable, Category = "SyncGroup", meta=(BlueprintThreadSafe))
+	UFUNCTION(BlueprintCallable, Category = "Animation|Synchronization", meta=(BlueprintThreadSafe))
 	bool HasMarkerBeenHitThisFrame(FName SyncGroup, FName MarkerName) const;
 
-	UFUNCTION(BlueprintCallable, Category = "SyncGroup", meta=(BlueprintThreadSafe))
+	UFUNCTION(BlueprintCallable, Category = "Animation|Synchronization", meta=(BlueprintThreadSafe))
 	bool IsSyncGroupBetweenMarkers(FName InSyncGroupName, FName PreviousMarker, FName NextMarker, bool bRespectMarkerOrder = true) const;
 
-	UFUNCTION(BlueprintCallable, Category = "SyncGroup", meta=(BlueprintThreadSafe))
+	UFUNCTION(BlueprintCallable, Category = "Animation|Synchronization", meta=(BlueprintThreadSafe))
 	FMarkerSyncAnimPosition GetSyncGroupPosition(FName InSyncGroupName) const;
 
 	/** Attempts to queue a transition request, returns true if the request was successful */
@@ -1310,7 +1310,7 @@ public:
 	virtual void DisplayDebugInstance(FDisplayDebugManager& DisplayDebugManager, float& Indent);
 
 	/** Reset any dynamics running simulation-style updates (e.g. on teleport, time skip etc.) */
-	UFUNCTION(BlueprintCallable, Category = "Animation")
+	UFUNCTION(BlueprintCallable, Category = "Animation|Dynamics")
 	void ResetDynamics(ETeleportType InTeleportType);
 
 	UE_DEPRECATED(4.20, "Please use ResetDynamics with a ETeleportType argument")
