- Enable SRVs if CPU access is requested

diff --git a/Engine/Source/Runtime/Engine/Private/Rendering/StaticMeshVertexBuffer.cpp b/Engine/Source/Runtime/Engine/Private/Rendering/StaticMeshVertexBuffer.cpp
--- a/Engine/Source/Runtime/Engine/Private/Rendering/StaticMeshVertexBuffer.cpp
+++ b/Engine/Source/Runtime/Engine/Private/Rendering/StaticMeshVertexBuffer.cpp
@@ -331,9 +331,14 @@ void FStaticMeshVertexBuffer::InitRHI()
 	TRACE_CPUPROFILER_EVENT_SCOPE(FStaticMeshVertexBuffer::InitRHI);
 	SCOPED_LOADTIMER(FStaticMeshVertexBuffer_InitRHI);
 
+	// When bAllowCPUAccess is true the meshes is likely going to be used for Niagara to spawn particles on mesh surface.
+	// And it can be the case for CPU *and* GPU access: no differenciation today. That is why we create a SRV in this case.
+	// This also avoid setting lots of states on all the members of all the different buffers used by meshes. Follow up: https://jira.it.epicgames.net/browse/UE-69376.
+
 	const bool bHadTangentsData = TangentsData != nullptr;
+	const bool bCreateTangentsSRV = bHadTangentsData && TangentsData->GetAllowCPUAccess();
 	TangentsVertexBuffer.VertexBufferRHI = CreateTangentsRHIBuffer_RenderThread();
-	if (TangentsVertexBuffer.VertexBufferRHI && (RHISupportsManualVertexFetch(GMaxRHIShaderPlatform) || IsGPUSkinCacheAvailable(GMaxRHIShaderPlatform)))
+	if (TangentsVertexBuffer.VertexBufferRHI && (bCreateTangentsSRV || RHISupportsManualVertexFetch(GMaxRHIShaderPlatform) || IsGPUSkinCacheAvailable(GMaxRHIShaderPlatform)))
 	{
 		// When TangentsData is null, this buffer hasn't been streamed in yet. We still need to create a FRHIShaderResourceView which will be
 		// cached in a vertex factory uniform buffer later. The nullptr tells the RHI that the SRV doesn't view on anything yet.
@@ -343,8 +348,9 @@ void FStaticMeshVertexBuffer::InitRHI()
 	}
 
 	const bool bHadTexCoordData = TexcoordData != nullptr;
+	const bool bCreateTexCoordSRV = bHadTexCoordData && TexcoordData->GetAllowCPUAccess();
 	TexCoordVertexBuffer.VertexBufferRHI = CreateTexCoordRHIBuffer_RenderThread();
-	if (TexCoordVertexBuffer.VertexBufferRHI && RHISupportsManualVertexFetch(GMaxRHIShaderPlatform))
+	if (TexCoordVertexBuffer.VertexBufferRHI && (bCreateTexCoordSRV || RHISupportsManualVertexFetch(GMaxRHIShaderPlatform)))
 	{
 		// When TexcoordData is null, this buffer hasn't been streamed in yet. We still need to create a FRHIShaderResourceView which will be
 		// cached in a vertex factory uniform buffer later. The nullptr tells the RHI that the SRV doesn't view on anything yet.
