Calculate platform data payload size based on the provided mip bias value.

diff --git a/Engine/Source/Editor/TextureEditor/Private/TextureEditorToolkit.cpp b/Engine/Source/Editor/TextureEditor/Private/TextureEditorToolkit.cpp
--- a/Engine/Source/Editor/TextureEditor/Private/TextureEditorToolkit.cpp
+++ b/Engine/Source/Editor/TextureEditor/Private/TextureEditorToolkit.cpp
@@ -707,7 +707,7 @@ void FTextureEditorToolkit::PopulateQuickInfo( )
 	const int32 MaxInGameDepth = SurfaceDepth ? FMath::Max(SurfaceDepth >> MaxResMipBias, 1) : 0;
 
 	// Texture asset size
-	const int64 ResourceSize = PlatformDataPtr && *PlatformDataPtr ? (*PlatformDataPtr)->GetPayloadSize() : Texture->GetResourceSizeBytes(EResourceSizeMode::Exclusive);
+	const int64 ResourceSize = PlatformDataPtr && *PlatformDataPtr ? (*PlatformDataPtr)->GetPayloadSize(ActualMipBias) : Texture->GetResourceSizeBytes(EResourceSizeMode::Exclusive);
 
 	FNumberFormattingOptions SizeOptions;
 	SizeOptions.UseGrouping = false;
diff --git a/Engine/Source/Runtime/Engine/Classes/Engine/Texture.h b/Engine/Source/Runtime/Engine/Classes/Engine/Texture.h
--- a/Engine/Source/Runtime/Engine/Classes/Engine/Texture.h
+++ b/Engine/Source/Runtime/Engine/Classes/Engine/Texture.h
@@ -999,7 +999,7 @@ public:
 	ENGINE_API EPixelFormat GetLayerPixelFormat(uint32 LayerIndex) const;
 
 	/** Return the size of the texture pixel data in bytes, not including headers or alignment. */
-	ENGINE_API int64 GetPayloadSize() const;
+	ENGINE_API int64 GetPayloadSize(int32 MipBias) const;
 
 private:
 
diff --git a/Engine/Source/Runtime/Engine/Private/TextureDerivedData.cpp b/Engine/Source/Runtime/Engine/Private/TextureDerivedData.cpp
--- a/Engine/Source/Runtime/Engine/Private/TextureDerivedData.cpp
+++ b/Engine/Source/Runtime/Engine/Private/TextureDerivedData.cpp
@@ -2134,22 +2134,27 @@ EPixelFormat FTexturePlatformData::GetLayerPixelFormat(uint32 LayerIndex) const
 	return PixelFormat;
 }
 
-int64 FTexturePlatformData::GetPayloadSize() const
+int64 FTexturePlatformData::GetPayloadSize(int32 MipBias) const
 {
 	int64 PayloadSize = 0;
 	if (VTData)
 	{
-		for (int32 ChunkIndex = 0; ChunkIndex < VTData->Chunks.Num(); ChunkIndex++)
+		int32 NumTiles = 0;
+		for (uint32 MipIndex = MipBias; MipIndex < VTData->NumMips; MipIndex++)
 		{
-			PayloadSize += VTData->Chunks[ChunkIndex].SizeInBytes;
+			NumTiles += VTData->TileOffsetData[MipIndex].Width * VTData->TileOffsetData[MipIndex].Height;
 		}
+		int32 TileSizeWithBorder = (int32)(VTData->TileSize + 2 * VTData->TileBorderSize);
+		int32 TileBlockSizeX = FMath::DivideAndRoundUp(TileSizeWithBorder, GPixelFormats[PixelFormat].BlockSizeX);
+		int32 TileBlockSizeY = FMath::DivideAndRoundUp(TileSizeWithBorder, GPixelFormats[PixelFormat].BlockSizeY);
+		PayloadSize += (int64)GPixelFormats[PixelFormat].BlockBytes * TileBlockSizeX * TileBlockSizeY * VTData->NumLayers * NumTiles;
 	}
 	else
 	{
-		for (int32 MipIndex = 0; MipIndex < Mips.Num(); MipIndex++)
+		for (int32 MipIndex = MipBias; MipIndex < Mips.Num(); MipIndex++)
 		{
-			int32 BlockSizeX = FMath::DivideAndRoundUp(FMath::Max(Mips[MipIndex].SizeX, 1), GPixelFormats[PixelFormat].BlockSizeX);
-			int32 BlockSizeY = FMath::DivideAndRoundUp(FMath::Max(Mips[MipIndex].SizeY, 1), GPixelFormats[PixelFormat].BlockSizeY);
+			int32 BlockSizeX = FMath::DivideAndRoundUp(Mips[MipIndex].SizeX, GPixelFormats[PixelFormat].BlockSizeX);
+			int32 BlockSizeY = FMath::DivideAndRoundUp(Mips[MipIndex].SizeY, GPixelFormats[PixelFormat].BlockSizeY);
 			// for TextureCube and TextureCubeArray all the mipmaps contain the same number of slices, which is encoded in the PackedData member
 			// at the same time we can not just use SizeZ of a TextureCube mipmap, because for compatibility reasons it is always set to 1 and not 6 (which is the actual number of slices)
 			int32 BlockSizeZ = FMath::DivideAndRoundUp(FMath::Max(IsCubemap() ? GetNumSlices() : Mips[MipIndex].SizeZ, 1), GPixelFormats[PixelFormat].BlockSizeZ);
