Add support for previewing individual slices of a cubemap in the texture editor.

diff --git a/Engine/Shaders/Private/SimpleElementPixelShader.usf b/Engine/Shaders/Private/SimpleElementPixelShader.usf
--- a/Engine/Shaders/Private/SimpleElementPixelShader.usf
+++ b/Engine/Shaders/Private/SimpleElementPixelShader.usf
@@ -308,7 +308,40 @@ void GammaDistanceFieldMain(
 TextureCube CubeTexture;
 SamplerState CubeTextureSampler;
 
-// .x:MipLevel, .yzw:unused
+float3 GetCubemapVector(float2 ScaledUVs, uint CubeFace)
+{
+	float3 CubeCoordinates;
+	
+	// calculation of the cube coordinates matches the DDS orientaion and face order: positive x, negative x, positive y, negative y, positive z, negative z
+	if (CubeFace == 0)
+	{
+		CubeCoordinates = float3(1, -ScaledUVs.y, -ScaledUVs.x);
+	}
+	else if (CubeFace == 1)
+	{
+		CubeCoordinates = float3(-1, -ScaledUVs.y, ScaledUVs.x);
+	}
+	else if (CubeFace == 2)
+	{
+		CubeCoordinates = float3(ScaledUVs.x, 1, ScaledUVs.y);
+	}
+	else if (CubeFace == 3)
+	{
+		CubeCoordinates = float3(ScaledUVs.x, -1, -ScaledUVs.y);
+	}
+	else if (CubeFace == 4)
+	{
+		CubeCoordinates = float3(ScaledUVs.x, -ScaledUVs.y, 1);
+	}
+	else
+	{
+		CubeCoordinates = float3(-ScaledUVs.x, -ScaledUVs.y, -1);
+	}
+
+	return CubeCoordinates;
+}
+
+// .x:MipLevel, .y: SliceIndex, .zw:unused
 float4 PackedProperties0;
 float4x4 ColorWeights;
 
@@ -322,15 +355,21 @@ void CubemapTextureProperties(
 #endif
 	)
 {
-	float2 Angles = float2(2 * PI * (TextureCoordinate.x + 0.5f), PI * TextureCoordinate.y);
-	
-	float s = sin(Angles.y);
-	float3 Direction = float3(s * sin(Angles.x), cos(Angles.y), -s * cos(Angles.x));
-
 	float MipLevel = PackedProperties0.x;
+	float SliceIndex = PackedProperties0.y;
+	float4 BaseColor;
 
-	// show specified mip level
-	float4 BaseColor = TextureCubeSampleLevel(CubeTexture, CubeTextureSampler, Direction.xzy, MipLevel);
+	if (SliceIndex >= 0.0)
+	{
+		BaseColor = TextureCubeSampleLevel(CubeTexture, CubeTextureSampler, GetCubemapVector(TextureCoordinate * 2 - 1, SliceIndex), MipLevel);
+	}
+	else
+	{
+		float2 Angles = float2(2 * PI * (TextureCoordinate.x + 0.5f), PI * TextureCoordinate.y);
+		float s = sin(Angles.y);
+		float3 Direction = float3(s * sin(Angles.x), cos(Angles.y), -s * cos(Angles.x));
+		BaseColor = TextureCubeSampleLevel(CubeTexture, CubeTextureSampler, Direction.xzy, MipLevel);
+	}
 
 	float4 FinalColor;
 
diff --git a/Engine/Source/Editor/TextureEditor/Private/Models/TextureEditorViewportClient.cpp b/Engine/Source/Editor/TextureEditor/Private/Models/TextureEditorViewportClient.cpp
--- a/Engine/Source/Editor/TextureEditor/Private/Models/TextureEditorViewportClient.cpp
+++ b/Engine/Source/Editor/TextureEditor/Private/Models/TextureEditorViewportClient.cpp
@@ -112,7 +112,7 @@ void FTextureEditorViewportClient::Draw(FViewport* Viewport, FCanvas* Canvas)
 	{
 		if (TextureCube || TextureCubeArray || RTTextureCube)
 		{
-			BatchedElementParameters = new FMipLevelBatchedElementParameters(MipLevel, false);
+			BatchedElementParameters = new FMipLevelBatchedElementParameters(MipLevel, SliceIndex, false);
 		}
 		else if (VolumeTexture)
 		{
diff --git a/Engine/Source/Editor/TextureEditor/Private/TextureEditorToolkit.cpp b/Engine/Source/Editor/TextureEditor/Private/TextureEditorToolkit.cpp
--- a/Engine/Source/Editor/TextureEditor/Private/TextureEditorToolkit.cpp
+++ b/Engine/Source/Editor/TextureEditor/Private/TextureEditorToolkit.cpp
@@ -713,9 +713,12 @@ void FTextureEditorToolkit::PopulateQuickInfo( )
 	SizeOptions.UseGrouping = false;
 	SizeOptions.MaximumFractionalDigits = 0;
 
-	// Cubes are previewed as unwrapped 2D textures.
+	// Cubes are previewed as unwrapped 2D textures in case when slice index is not selected.
 	// These have 2x the width of a cube face.
-	PreviewEffectiveTextureWidth *= IsCubeTexture() ? 2 : 1;
+	if (IsCubeTexture() && GetSlice() < 0)
+	{
+		PreviewEffectiveTextureWidth *= 2;
+	}
 
 	FNumberFormattingOptions Options;
 	Options.UseGrouping = false;
@@ -2140,9 +2143,9 @@ bool FTextureEditorToolkit::HasLayers() const
 
 bool FTextureEditorToolkit::HasSlices() const
 {
-	// currently only supports Texture2DArray, but can be extended for other texture types which are using multiple slices
+	// currently supports Texture2D arrays, cubemaps and cubemap arrays
 	// slice selection should be supported even for a texture array with less than two elements, because array elements can be added dynamically
-	return Texture->IsA(UTexture2DArray::StaticClass());
+	return Is2DArrayTexture() || IsCubeTexture();
 }
 
 int32 FTextureEditorToolkit::GetNumSlices() const
@@ -2151,6 +2154,14 @@ int32 FTextureEditorToolkit::GetNumSlices() const
 	{
 		return Cast<UTexture2DArray>(Texture)->GetArraySize();
 	}
+	else if (Texture->IsA(UTextureRenderTarget2DArray::StaticClass()))
+	{
+		return Cast<UTextureRenderTarget2DArray>(Texture)->GetSurfaceArraySize();
+	}
+	else if (IsCubeTexture())
+	{
+		return 6;
+	}
 	return 1;
 }
 
diff --git a/Engine/Source/Editor/UnrealEd/Private/ThumbnailRendering/TextureThumbnailRenderer.cpp b/Engine/Source/Editor/UnrealEd/Private/ThumbnailRendering/TextureThumbnailRenderer.cpp
--- a/Engine/Source/Editor/UnrealEd/Private/ThumbnailRendering/TextureThumbnailRenderer.cpp
+++ b/Engine/Source/Editor/UnrealEd/Private/ThumbnailRendering/TextureThumbnailRenderer.cpp
@@ -72,7 +72,7 @@ void UTextureThumbnailRenderer::Draw(UObject* Object, int32 X, int32 Y, uint32 W
 		if(TextureCube || TextureCubeArray || RTTextureCube)
 		{
 			// is released by the render thread when it was rendered
-			BatchedElementParameters = new FMipLevelBatchedElementParameters((float)0);
+			BatchedElementParameters = new FMipLevelBatchedElementParameters((float)0, (float)-1);
 			
 			// If the thumbnail is square then make it 2:1 for cubes.
 			if(Width == Height)
diff --git a/Engine/Source/Runtime/Engine/Private/CubemapUnwrapUtils.cpp b/Engine/Source/Runtime/Engine/Private/CubemapUnwrapUtils.cpp
--- a/Engine/Source/Runtime/Engine/Private/CubemapUnwrapUtils.cpp
+++ b/Engine/Source/Runtime/Engine/Private/CubemapUnwrapUtils.cpp
@@ -37,7 +37,7 @@ namespace CubemapHelpers
 	bool GenerateLongLatUnwrap(const FTextureResource* TextureResource, const uint32 AxisDimenion, const EPixelFormat SourcePixelFormat, TArray64<uint8>& BitsOUT, FIntPoint& SizeOUT, EPixelFormat& FormatOUT)
 	{
 		TRefCountPtr<FBatchedElementParameters> BatchedElementParameters;
-		BatchedElementParameters = new FMipLevelBatchedElementParameters((float)0, true);
+		BatchedElementParameters = new FMipLevelBatchedElementParameters((float)0, (float)-1, true);
 		const FIntPoint LongLatDimensions(AxisDimenion * 2, AxisDimenion);
 
 		// If the source format is 8 bit per channel or less then select a LDR target format.
@@ -121,11 +121,11 @@ void FCubemapTexturePropertiesVS::SetParameters( FRHICommandList& RHICmdList, co
 }
 
 template<bool bHDROutput>
-void FCubemapTexturePropertiesPS<bHDROutput>::SetParameters( FRHICommandList& RHICmdList, const FTexture* Texture, const FMatrix& ColorWeightsValue, float MipLevel, float GammaValue )
+void FCubemapTexturePropertiesPS<bHDROutput>::SetParameters(FRHICommandList& RHICmdList, const FTexture* Texture, const FMatrix& ColorWeightsValue, float MipLevel, float SliceIndex, float GammaValue)
 {
 	SetTextureParameter(RHICmdList, RHICmdList.GetBoundPixelShader(),CubeTexture,CubeTextureSampler,Texture);
 
-	FVector4f PackedProperties0Value(MipLevel, 0, 0, 0);
+	FVector4f PackedProperties0Value(MipLevel, SliceIndex, 0, 0);
 	SetShaderValue(RHICmdList, RHICmdList.GetBoundPixelShader(), PackedProperties0, PackedProperties0Value);
 	SetShaderValue(RHICmdList, RHICmdList.GetBoundPixelShader(), ColorWeights, (FMatrix44f)ColorWeightsValue);
 	SetShaderValue(RHICmdList, RHICmdList.GetBoundPixelShader(), Gamma, GammaValue);
@@ -160,7 +160,7 @@ void FMipLevelBatchedElementParameters::BindShaders(FRHICommandList& RHICmdList,
 	SetGraphicsPipelineState(RHICmdList, GraphicsPSOInit, 0);
 
 	VertexShader->SetParameters(RHICmdList, InTransform);
-	PixelShader->SetParameters(RHICmdList, Texture, ColorWeights, MipLevel, InGamma);
+	PixelShader->SetParameters(RHICmdList, Texture, ColorWeights, MipLevel, SliceIndex, InGamma);
 }
 
 void FIESLightProfilePS::SetParameters( FRHICommandList& RHICmdList, const FTexture* Texture, float InBrightnessInLumens )
diff --git a/Engine/Source/Runtime/Engine/Public/CubemapUnwrapUtils.h b/Engine/Source/Runtime/Engine/Public/CubemapUnwrapUtils.h
--- a/Engine/Source/Runtime/Engine/Public/CubemapUnwrapUtils.h
+++ b/Engine/Source/Runtime/Engine/Public/CubemapUnwrapUtils.h
@@ -85,7 +85,7 @@ public:
 	}
 	FCubemapTexturePropertiesPS() {}
 
-	void SetParameters(FRHICommandList& RHICmdList, const FTexture* Texture, const FMatrix& ColorWeightsValue, float MipLevel, float GammaValue);
+	void SetParameters(FRHICommandList& RHICmdList, const FTexture* Texture, const FMatrix& ColorWeightsValue, float MipLevel, float SliceIndex, float GammaValue);
 
 	static void ModifyCompilationEnvironment(const FGlobalShaderPermutationParameters& Parameters, FShaderCompilerEnvironment& OutEnvironment)
 	{
@@ -104,9 +104,10 @@ private:
 class ENGINE_API FMipLevelBatchedElementParameters : public FBatchedElementParameters
 {
 public:
-	FMipLevelBatchedElementParameters(float InMipLevel, bool bInHDROutput = false)
+	FMipLevelBatchedElementParameters(float InMipLevel, float InSliceIndex, bool bInHDROutput = false)
 		: bHDROutput(bInHDROutput)
 		, MipLevel(InMipLevel)
+		, SliceIndex(InSliceIndex)
 	{
 	}
 
@@ -119,6 +120,7 @@ private:
 	bool bHDROutput;
 	/** Parameters that need to be passed to the shader */
 	float MipLevel;
+	float SliceIndex;
 };
 
 
