[Backout] - CL20593017, 20594056, 20625634, 20645670, 20649179, 20649223, 20649380, 20658535

diff --git a/Engine/Source/Editor/UnrealEd/Private/Fbx/FbxSkeletalMeshExport.cpp b/Engine/Source/Editor/UnrealEd/Private/Fbx/FbxSkeletalMeshExport.cpp
--- a/Engine/Source/Editor/UnrealEd/Private/Fbx/FbxSkeletalMeshExport.cpp
+++ b/Engine/Source/Editor/UnrealEd/Private/Fbx/FbxSkeletalMeshExport.cpp
@@ -29,7 +29,7 @@ FbxNode* FFbxExporter::CreateSkeleton(const USkeletalMesh* SkelMesh, TArray<FbxN
 
 	if(RefSkeleton.GetRawBoneNum() == 0)
 	{
-		return nullptr;
+		return NULL;
 	}
 
 	// Create a list of the nodes we create for each bone, so that children can 
@@ -116,13 +116,13 @@ FbxNode* FFbxExporter::CreateMesh(const USkeletalMesh* SkelMesh, const TCHAR* Me
 	const int32 VertexCount = SourceModel.GetNumNonClothingVertices();
 
 	// Verify the integrity of the mesh.
-	if (VertexCount == 0) return nullptr;
+	if (VertexCount == 0) return NULL;
 
 	// Copy all the vertex data from the various chunks to a single buffer.
 	// Makes the rest of the code in this function cleaner and easier to maintain.  
 	TArray<FSoftSkinVertex> Vertices;
 	SourceModel.GetNonClothVertices(Vertices);
-	if (Vertices.Num() != VertexCount) return nullptr;
+	if (Vertices.Num() != VertexCount) return NULL;
 
 	FbxMesh* Mesh = FbxMesh::Create(Scene, TCHAR_TO_UTF8(MeshName));
 
@@ -137,7 +137,7 @@ FbxNode* FFbxExporter::CreateMesh(const USkeletalMesh* SkelMesh, const TCHAR* Me
 
 	// Create Layer 0 to hold the normals
 	FbxLayer* LayerZero = Mesh->GetLayer(0);
-	if (LayerZero == nullptr)
+	if (LayerZero == NULL)
 	{
 		Mesh->CreateLayer();
 		LayerZero = Mesh->GetLayer(0);
@@ -169,7 +169,7 @@ FbxNode* FFbxExporter::CreateMesh(const USkeletalMesh* SkelMesh, const TCHAR* Me
 	for (int32 TexCoordSourceIndex = 0; TexCoordSourceIndex < TexCoordSourceCount; ++TexCoordSourceIndex)
 	{
 		FbxLayer* Layer = Mesh->GetLayer(TexCoordSourceIndex);
-		if (Layer == nullptr)
+		if (Layer == NULL)
 		{
 			Mesh->CreateLayer();
 			Layer = Mesh->GetLayer(TexCoordSourceIndex);
@@ -328,10 +328,10 @@ FbxNode* FFbxExporter::CreateMesh(const USkeletalMesh* SkelMesh, const TCHAR* Me
 		if(AnimSeq && BlendShapeCurvesMap.Num() > 0)
 		{
 			ExportCustomAnimCurvesToFbx(BlendShapeCurvesMap, AnimSeq,
-				FFrameTime(0),													// Start frame to export
-				FFrameTime(AnimSeq->GetDataModel()->GetNumberOfFrames()),		// Final frame to export
-				1.f,												// Frame rate scale
-				0.f,												// FBX StartTime
+				0.f,		// AnimStartOffset
+				0.f,		// AnimEndOffset
+				1.f,		// AnimPlayRate
+				0.f,		// StartTime
 				100.f);		// ValueScale, for some reason we need to scale BlendShape curves by a factor of 100.
 		}
 	}
@@ -339,6 +339,8 @@ FbxNode* FFbxExporter::CreateMesh(const USkeletalMesh* SkelMesh, const TCHAR* Me
 	FbxNode* MeshNode = FbxNode::Create(Scene, TCHAR_TO_UTF8(MeshName));
 	MeshNode->SetNodeAttribute(Mesh);
 
+
+
 	// Add the materials for the mesh
 	const TArray<FSkeletalMaterial>& SkelMeshMaterials = SkelMesh->GetMaterials();
 	int32 MaterialCount = SkelMeshMaterials.Num();
@@ -356,7 +358,7 @@ FbxNode* FFbxExporter::CreateMesh(const USkeletalMesh* SkelMesh, const TCHAR* Me
 			MatInterface = SkelMeshMaterials[MaterialIndex].MaterialInterface;
 		}
 
-		FbxSurfaceMaterial* FbxMaterial = nullptr;
+		FbxSurfaceMaterial* FbxMaterial = NULL;
 		if (LODIndex == 0)
 		{
 			if (MatInterface && !FbxMaterials.Find(MatInterface))
@@ -588,7 +590,7 @@ void FFbxExporter::ExportSkeletalMeshComponent(USkeletalMeshComponent* SkelMeshC
 	if (SkelMeshComp && SkelMeshComp->SkeletalMesh)
 	{
 		UAnimSequence* AnimSeq = (bSaveAnimSeq && SkelMeshComp->GetAnimationMode() == EAnimationMode::AnimationSingleNode) ? 
-			Cast<UAnimSequence>(SkelMeshComp->AnimationData.AnimToPlay) : nullptr;
+			Cast<UAnimSequence>(SkelMeshComp->AnimationData.AnimToPlay) : NULL;
 		FbxNode* SkeletonRootNode = ExportSkeletalMeshToFbx(SkelMeshComp->SkeletalMesh, AnimSeq, MeshName, ActorRootNode, &ToRawPtrTArrayUnsafe(SkelMeshComp->OverrideMaterials));
 		if(SkeletonRootNode)
 		{
@@ -764,7 +766,7 @@ FbxNode* FFbxExporter::ExportSkeletalMeshToFbx(const USkeletalMesh* SkeletalMesh
 		return SkeletonRootNode;
 	}
 
-	return nullptr;
+	return NULL;
 }
 
 } // namespace UnFbx
diff --git a/Engine/Source/Editor/UnrealEd/Private/FbxExporter.h b/Engine/Source/Editor/UnrealEd/Private/FbxExporter.h
--- a/Engine/Source/Editor/UnrealEd/Private/FbxExporter.h
+++ b/Engine/Source/Editor/UnrealEd/Private/FbxExporter.h
@@ -217,7 +217,7 @@ public:
 	 * @param StaticMesh	The static mesh to export
 	 * @param MaterialOrder	Optional ordering of materials to set up correct material ID's across multiple meshes being export such as BSP surfaces which share common materials. Should be used sparingly
 	 */
-	virtual void ExportStaticMesh( UStaticMesh* StaticMesh, const TArray<FStaticMaterial>* MaterialOrder = nullptr );
+	virtual void ExportStaticMesh( UStaticMesh* StaticMesh, const TArray<FStaticMaterial>* MaterialOrder = NULL );
 
 	/**
 	 * Exports BSP
@@ -249,7 +249,7 @@ public:
 	/**
 	 * Exports a single UAnimSequence, and optionally a skeletal mesh
 	 */
-	FbxNode* ExportAnimSequence( const UAnimSequence* AnimSeq, const USkeletalMesh* SkelMesh, bool bExportSkelMesh, const TCHAR* MeshNames=nullptr, FbxNode* ActorRootNode=nullptr, const TArray<UMaterialInterface*>* OverrideMaterials = nullptr);
+	FbxNode* ExportAnimSequence( const UAnimSequence* AnimSeq, const USkeletalMesh* SkelMesh, bool bExportSkelMesh, const TCHAR* MeshNames=NULL, FbxNode* ActorRootNode=NULL, const TArray<UMaterialInterface*>* OverrideMaterials = nullptr);
 
 	/**
 	 * Exports the list of UAnimSequences as a single animation based on the settings in the TrackKeys
@@ -339,7 +339,7 @@ private:
 	 * @param MaterialOrderOverride	Optional ordering of materials to set up correct material ID's across multiple meshes being export such as BSP surfaces which share common materials. Should be used sparingly
 	 * @param OverrideMaterials	Optional array of materials to be used instead of the static mesh materials. Used for material overrides in static mesh components.
 	 */
-	FbxNode* ExportStaticMeshToFbx(const UStaticMesh* StaticMesh, int32 ExportLOD, const TCHAR* MeshName, FbxNode* FbxActor, int32 LightmapUVChannel = -1, const FColorVertexBuffer* ColorBuffer = nullptr, const TArray<FStaticMaterial>* MaterialOrderOverride = nullptr, const TArray<UMaterialInterface*>* OverrideMaterials = nullptr);
+	FbxNode* ExportStaticMeshToFbx(const UStaticMesh* StaticMesh, int32 ExportLOD, const TCHAR* MeshName, FbxNode* FbxActor, int32 LightmapUVChannel = -1, const FColorVertexBuffer* ColorBuffer = NULL, const TArray<FStaticMaterial>* MaterialOrderOverride = NULL, const TArray<UMaterialInterface*>* OverrideMaterials = NULL);
 
 	/**
 	 * Exports a spline mesh
@@ -421,28 +421,18 @@ private:
 	/**
 	 * Add the given animation sequence as rotation and translation tracks to the given list of bone nodes
 	 */
-	void ExportAnimSequenceToFbx(const UAnimSequence* AnimSeq, const USkeletalMesh* SkelMesh, TArray<FbxNode*>& BoneNodes, FbxAnimLayer* AnimLayer, FFrameTime StartFrameTime, FFrameTime EndFrameTime, float FrameRateScale, float StartTime);
-
-	UE_DEPRECATED(5.1, "ExportAnimSequenceToFbx is deprecated, use different signature")
 	void ExportAnimSequenceToFbx(const UAnimSequence* AnimSeq, const USkeletalMesh* SkelMesh, TArray<FbxNode*>& BoneNodes, FbxAnimLayer* AnimLayer,
 		float AnimStartOffset, float AnimEndOffset, float AnimPlayRate, float StartTime);
 
 	/**
 	 * Add the custom Curve data to the FbxAnimCurves passed in parameter by matching their name to the skeletal mesh custom curves.
 	 */
-	void ExportCustomAnimCurvesToFbx(const TMap<FName, FbxAnimCurve*>& CustomCurves, const UAnimSequence* AnimSeq,
-		FFrameTime AnimStartOffset, FFrameTime AnimEndOffset, float FrameRateScale, float StartTime, float ValueScale = 1.f);
-	
-	UE_DEPRECATED(5.1, "ExportCustomAnimCurvesToFbx is deprecated, use different signature")
 	void ExportCustomAnimCurvesToFbx(const TMap<FName, FbxAnimCurve*>& CustomCurves, const UAnimSequence* AnimSeq,
 		float AnimStartOffset, float AnimEndOffset, float AnimPlayRate, float StartTime, float ValueScale = 1.f);
 
 	/**
 	 * Used internally to reuse the AnimSequence iteration code when exporting various kind of curves.
 	 */
-	void IterateInsideAnimSequence(const UAnimSequence* AnimSeq, FFrameTime StartFrameTime, FFrameTime EndFrameTime, float FrameRateScale, float StartTime, TFunctionRef<void(double, FbxTime, bool)> IterationLambda);
-
-	UE_DEPRECATED(5.1, "IterateInsideAnimSequence is deprecated, use different signature")
 	void IterateInsideAnimSequence(const UAnimSequence* AnimSeq, float AnimStartOffset, float AnimEndOffset, float AnimPlayRate, float StartTime, TFunctionRef<void(float, FbxTime, bool)> IterationLambda);
 
 	/** 
diff --git a/Engine/Source/Editor/UnrealEd/Private/SkeletalMeshEdit.cpp b/Engine/Source/Editor/UnrealEd/Private/SkeletalMeshEdit.cpp
--- a/Engine/Source/Editor/UnrealEd/Private/SkeletalMeshEdit.cpp
+++ b/Engine/Source/Editor/UnrealEd/Private/SkeletalMeshEdit.cpp
@@ -49,7 +49,7 @@ UAnimSequence * UEditorEngine::ImportFbxAnimation( USkeleton* Skeleton, UObject*
 {
 	check(Skeleton);
 
-	UAnimSequence * NewAnimation=nullptr;
+	UAnimSequence * NewAnimation=NULL;
 
 	UnFbx::FFbxImporter* FFbxImporter = UnFbx::FFbxImporter::GetInstance();
 
@@ -76,7 +76,7 @@ UAnimSequence * UEditorEngine::ImportFbxAnimation( USkeleton* Skeleton, UObject*
 			FFbxImporter->AddTokenizedErrorMessage(FTokenizedMessage::Create(EMessageSeverity::Error, FText::Format(LOCTEXT("Error_CouldNotFindFbxTrack", "Mesh contains {0} bone as root but animation doesn't contain the root track.\nImport failed."), FText::FromName(Skeleton->GetReferenceSkeleton().GetBoneName(0)))), FFbxErrors::Animation_CouldNotFindRootTrack);
 
 			FFbxImporter->ReleaseScene();
-			return nullptr;
+			return NULL;
 		}
 
 		// Check for blend shape curves that are not skinned.  Unskinned geometry can still contain morph curves
@@ -568,19 +568,19 @@ void UnFbx::FFbxImporter::GetAnimationIntervalMultiLayer(FbxNode* RootNode, FbxA
 UAnimSequence * UnFbx::FFbxImporter::ImportAnimations(USkeleton* Skeleton, UObject* Outer, TArray<FbxNode*>& SortedLinks, const FString& Name, UFbxAnimSequenceImportData* TemplateImportData, TArray<FbxNode*>& NodeArray)
 {
 	// we need skeleton to create animsequence
-	if (Skeleton == nullptr)
+	if (Skeleton == NULL)
 	{
-		return nullptr;
+		return NULL;
 	}
 
 	int32 ValidTakeCount = 0;
 	if (IsValidAnimationData(SortedLinks, NodeArray, ValidTakeCount) == false)
 	{
 		AddTokenizedErrorMessage(FTokenizedMessage::Create(EMessageSeverity::Warning, LOCTEXT("FBXImport_InvalidAnimationData", "This does not contain any valid animation takes.")), FFbxErrors::Animation_InvalidData);
-		return nullptr;
+		return NULL;
 	}
 
-	UAnimSequence* LastCreatedAnim = nullptr;
+	UAnimSequence* LastCreatedAnim = NULL;
 
 	int32 ResampleRate = DEFAULT_SAMPLERATE;
 	if ( ImportOptions->bResample )
@@ -635,7 +635,7 @@ UAnimSequence * UnFbx::FFbxImporter::ImportAnimations(USkeleton* Skeleton, UObje
 
 		FString 	ParentPath = FString::Printf(TEXT("%s/%s"), *FPackageName::GetLongPackagePath(*Outer->GetName()), *SequenceName);
 		UObject* 	ParentPackage = CreatePackage( *ParentPath);
-		UObject* Object = LoadObject<UObject>(ParentPackage, *SequenceName, nullptr, (LOAD_Quiet | LOAD_NoWarn), nullptr);
+		UObject* Object = LoadObject<UObject>(ParentPackage, *SequenceName, NULL, (LOAD_Quiet | LOAD_NoWarn), NULL);
 		UAnimSequence * DestSeq = Cast<UAnimSequence>(Object);
 		// if object with same name exists, warn user
 		if (Object && !DestSeq)
@@ -645,12 +645,12 @@ UAnimSequence * UnFbx::FFbxImporter::ImportAnimations(USkeleton* Skeleton, UObje
 		}
 
 		// If not, create new one now.
-
-		const bool bCreateAsset = !DestSeq;
-		if(bCreateAsset)
+		if(!DestSeq)
 		{
 			DestSeq = NewObject<UAnimSequence>(ParentPackage, *SequenceName, RF_Public | RF_Standalone);
 			CreatedObjects.Add(DestSeq);
+			// Notify the asset registry
+			FAssetRegistryModule::AssetCreated(DestSeq);
 		}
 
 		DestSeq->SetSkeleton(Skeleton);
@@ -662,16 +662,8 @@ UAnimSequence * UnFbx::FFbxImporter::ImportAnimations(USkeleton* Skeleton, UObje
 		DestSeq->ImportFileFramerate = GetOriginalFbxFramerate();
 		DestSeq->ImportResampleFramerate = ResampleRate;
 
-		DestSeq->GetController().InitializeModel();
-
 		ImportAnimation(Skeleton, DestSeq, Name, SortedLinks, NodeArray, CurAnimStack, ResampleRate, AnimTimeSpan);
 
-		if (bCreateAsset)
-		{			
-			// Notify the asset registry
-			FAssetRegistryModule::AssetCreated(DestSeq);
-		}
-
 		LastCreatedAnim = DestSeq;
 	}
 
@@ -1345,7 +1337,7 @@ bool UnFbx::FFbxImporter::ImportCurveToAnimSequence(class UAnimSequence * Target
 		FAnimationCurveIdentifier FloatCurveId(NewName, ERawCurveTrackTypes::RCT_Float);
 
 
-		IAnimationDataModel* DataModel = TargetSequence->GetDataModel();
+		UAnimDataModel* DataModel = TargetSequence->GetDataModel();
 		IAnimationDataController& Controller = TargetSequence->GetController();
 
 		const FFloatCurve* TargetCurve = DataModel->FindFloatCurve(FloatCurveId);
@@ -1565,14 +1557,8 @@ bool UnFbx::FFbxImporter::ImportAnimation(USkeleton* Skeleton, UAnimSequence * D
 	//This destroy all previously imported animation raw data
 	Controller.RemoveAllBoneTracks();
 
-	// First set frame rate
-	const FFrameRate ResampleFrameRate(ResampleRate, 1);
-	Controller.SetFrameRate(ResampleFrameRate);
-
 	// if you have one pose(thus 0.f duration), it still contains animation, so we'll need to consider that as MINIMUM_ANIMATION_LENGTH time length
-	const FFrameNumber NumberOfFrames = ResampleFrameRate.AsFrameNumber(SequenceLength.GetSecondDouble());
-	Controller.SetNumberOfFrames(FGenericPlatformMath::Max<int32>(NumberOfFrames.Value, 1));
-
+	Controller.SetPlayLength(FGenericPlatformMath::Max<float>(SequenceLength.GetSecondDouble(), MINIMUM_ANIMATION_LENGTH));
 	if(PreviousSequenceLength > MINIMUM_ANIMATION_LENGTH && DestSeq->GetDataModel()->GetNumberOfFloatCurves() > 0)
 	{
 		// The sequence already existed when we began the import. We need to scale the key times for all curves to match the new 
@@ -1802,7 +1788,7 @@ void UnFbx::FFbxImporter::ImportAnimationCustomAttribute(FAnimCurveImportSetting
 	UAnimSequence* DestSeq = AnimImportSettings.DestSeq;
 	USkeleton* MySkeleton = DestSeq->GetSkeleton();
 
-	const IAnimationDataModel* DataModel = DestSeq->GetDataModel();
+	const UAnimDataModel* DataModel = DestSeq->GetDataModel();
 	const int32 NumFloatCurves = DataModel->GetNumberOfFloatCurves();
 	const FAnimationCurveData& CurveData = DataModel->GetCurveData();
 
@@ -1918,7 +1904,7 @@ void UnFbx::FFbxImporter::ImportBoneTracks(USkeleton* Skeleton, FAnimCurveImport
 	const FbxTime TimeComparisonThreshold = (KINDA_SMALL_NUMBER * static_cast<float>(FBXSDK_TC_SECOND));
 	const FReferenceSkeleton& RefSkeleton = Skeleton->GetReferenceSkeleton();
 
-	IAnimationDataModel* DataModel = DestSeq->GetDataModel();
+	UAnimDataModel* DataModel = DestSeq->GetDataModel();
 	IAnimationDataController& Controller = DestSeq->GetController();	
 	Controller.SetFrameRate(FFrameRate(ResampleRate, 1));
 
@@ -2079,7 +2065,7 @@ void UnFbx::FFbxImporter::ImportBoneTracks(USkeleton* Skeleton, FAnimCurveImport
 								if (AttributeIdentifier.IsValid())
 								{
 									// remove any existing attribute with the same identifier
-									if (const IAnimationDataModel* Model = Controller.GetModel())
+									if (const UAnimDataModel* Model = Controller.GetModel())
 									{
 										if (Model->FindAttribute(AttributeIdentifier))
 										{
diff --git a/Engine/Source/Editor/UnrealEd/Public/ChangeTransactor.h b/Engine/Source/Editor/UnrealEd/Public/ChangeTransactor.h
--- a/Engine/Source/Editor/UnrealEd/Public/ChangeTransactor.h
+++ b/Engine/Source/Editor/UnrealEd/Public/ChangeTransactor.h
@@ -42,7 +42,7 @@ struct FChangeTransactor
 		CheckTransactionObject();
 		checkf(!IsTransactionPending(), TEXT("A transaction is already pending"));
 		
-		if (CanTransactChanges())
+		if (GEngine && GEngine->CanTransact() && ensure(!GIsTransacting))
 		{
 			CompoundChangeData = MakeUnique<FCompoundChangeInput>();
 			PendingTransactionIndex = GEngine->BeginTransaction(TEXT("FChangeTransactor"), TransactionDescription, TransactionObject.Get());
@@ -100,7 +100,7 @@ struct FChangeTransactor
 
 		CheckTransactionObject();
 
-		if (CanTransactChanges())
+		if (GEngine && GEngine->CanTransact() && ensure(!GIsTransacting))
 		{
 			checkf(IsTransactionPending(), TEXT("No transaction was previously opened"));
 
@@ -118,11 +118,6 @@ struct FChangeTransactor
 			}
 		}
 	}
-	
-	static bool CanTransactChanges()
-	{
-		return GEngine && GEngine->CanTransact() && !GIsTransacting;
-	};
 private:	
 	void CheckTransactionObject()
 	{
@@ -141,28 +136,5 @@ private:
 	int32 TransactedChanges;
 };
 
-struct FScopedCompoundTransaction
-{
-	FScopedCompoundTransaction(FChangeTransactor& InTransactor, const FText& InDescription) : Transactor(InTransactor), bCreated(false)
-	{
-		if (FChangeTransactor::CanTransactChanges() && !Transactor.IsTransactionPending())
-		{
-			Transactor.OpenTransaction(InDescription);
-			bCreated = true;
-		}
-	}
-
-	~FScopedCompoundTransaction()
-	{
-		if (bCreated)
-		{
-			Transactor.CloseTransaction();
-		}
-	}
-
-	FChangeTransactor& Transactor;
-	bool bCreated;
-};
-
 }
 #endif // WITH_EDITOR
diff --git a/Engine/Source/Runtime/AnimGraphRuntime/Private/AnimNodes/AnimNode_RandomPlayer.cpp b/Engine/Source/Runtime/AnimGraphRuntime/Private/AnimNodes/AnimNode_RandomPlayer.cpp
--- a/Engine/Source/Runtime/AnimGraphRuntime/Private/AnimNodes/AnimNode_RandomPlayer.cpp
+++ b/Engine/Source/Runtime/AnimGraphRuntime/Private/AnimNodes/AnimNode_RandomPlayer.cpp
@@ -296,8 +296,8 @@ void FAnimNode_RandomPlayer::Evaluate_AnyThread(FPoseContext& Output)
 		FAnimationPoseData CurrentPoseData(Poses[0], Curves[0], Attributes[0]);
 		FAnimationPoseData NextPoseData(Poses[1], Curves[1], Attributes[1]);
 
-		CurrentSequence->GetAnimationPose(CurrentPoseData, FAnimExtractContext(static_cast<double>(CurrentData.CurrentPlayTime), AnimProxy->ShouldExtractRootMotion(), CurrentData.DeltaTimeRecord, CurrentData.RemainingLoops > 0));
-		NextSequence->GetAnimationPose(NextPoseData, FAnimExtractContext(static_cast<double>(NextData.CurrentPlayTime), AnimProxy->ShouldExtractRootMotion(), NextData.DeltaTimeRecord, NextData.RemainingLoops > 0));
+		CurrentSequence->GetAnimationPose(CurrentPoseData, FAnimExtractContext(CurrentData.CurrentPlayTime, AnimProxy->ShouldExtractRootMotion(), CurrentData.DeltaTimeRecord, CurrentData.RemainingLoops > 0));
+		NextSequence->GetAnimationPose(NextPoseData, FAnimExtractContext(NextData.CurrentPlayTime, AnimProxy->ShouldExtractRootMotion(), NextData.DeltaTimeRecord, NextData.RemainingLoops > 0));
 
 		FAnimationPoseData AnimationPoseData(Output);
 		FAnimationRuntime::BlendPosesTogether(Poses, Curves, Attributes, Weights, AnimationPoseData);
@@ -306,7 +306,7 @@ void FAnimNode_RandomPlayer::Evaluate_AnyThread(FPoseContext& Output)
 	{
 		// Single animation, no blending needed.
 		FAnimationPoseData AnimationPoseData(Output);
-		CurrentSequence->GetAnimationPose(AnimationPoseData, FAnimExtractContext(static_cast<double>(CurrentData.CurrentPlayTime), Output.AnimInstanceProxy->ShouldExtractRootMotion(), CurrentData.DeltaTimeRecord, CurrentData.RemainingLoops > 0));
+		CurrentSequence->GetAnimationPose(AnimationPoseData, FAnimExtractContext(CurrentData.CurrentPlayTime, Output.AnimInstanceProxy->ShouldExtractRootMotion(), CurrentData.DeltaTimeRecord, CurrentData.RemainingLoops > 0));
 	}
 }
 
diff --git a/Engine/Source/Runtime/Core/Public/UObject/UE5MainStreamObjectVersion.h b/Engine/Source/Runtime/Core/Public/UObject/UE5MainStreamObjectVersion.h
--- a/Engine/Source/Runtime/Core/Public/UObject/UE5MainStreamObjectVersion.h
+++ b/Engine/Source/Runtime/Core/Public/UObject/UE5MainStreamObjectVersion.h
@@ -239,12 +239,17 @@ struct CORE_API FUE5MainStreamObjectVersion
 		// Remove serialization of bounds relevant from  WorldPartitionActorDesc
 		WorldPartitionActorDescRemoveBoundsRelevantSerialization,
 
-		// Add IAnimationDataModel interface and replace UObject based representation for Animation Assets
-		AnimationDataModelInterface,
+		// Added IAnimationDataModel interface and replace UObject based representation for Animation Assets
+		// This version had to be undone. Animation assets saved between this and the subsequent backout version
+		// will be unable to be loaded
+		AnimationDataModelInterface_BackedOut,
 
 		// Deprecate LandscapeSplineActorDesc
 		LandscapeSplineActorDescDeprecation,
 
+		// Revert the IAnimationDataModel changes. Animation assets 
+		BackoutAnimationDataModelInterface,
+
 		// -----<new versions can be added above this line>-------------------------------------------------
 		VersionPlusOne,
 		LatestVersion = VersionPlusOne - 1
@@ -256,4 +261,4 @@ struct CORE_API FUE5MainStreamObjectVersion
 	static TMap<FGuid, FGuid> GetSystemGuids();
 
 	FUE5MainStreamObjectVersion() = delete;
-};
+};
\ No newline at end of file
diff --git a/Engine/Source/Runtime/Engine/Classes/Animation/AnimCompositeBase.h b/Engine/Source/Runtime/Engine/Classes/Animation/AnimCompositeBase.h
--- a/Engine/Source/Runtime/Engine/Classes/Animation/AnimCompositeBase.h
+++ b/Engine/Source/Runtime/Engine/Classes/Animation/AnimCompositeBase.h
@@ -39,7 +39,7 @@ struct FRootMotionExtractionStep
 	float EndPosition;
 
 	FRootMotionExtractionStep() 
-		: AnimSequence(nullptr)
+		: AnimSequence(NULL)
 		, StartPosition(0.f)
 		, EndPosition(0.f)
 		{
@@ -348,10 +348,6 @@ class UAnimCompositeBase : public UAnimSequenceBase
 	//~ Begin UObject Interface
 	virtual void PostLoad() override;
 	//~ End UObject Interface
-	
-	//~ Begin UAnimSequenceBase Interface
-	virtual FFrameRate GetSamplingFrameRate() const override;
-	//~ End UAnimSequenceBase Interface
 
 	// Extracts root motion from the supplied FAnimTrack between the Start End range specified
 	ENGINE_API void ExtractRootMotionFromTrack(const FAnimTrack &SlotAnimTrack, float StartTrackPosition, float EndTrackPosition, FRootMotionMovementParams &RootMotion) const;
@@ -366,9 +362,5 @@ class UAnimCompositeBase : public UAnimSequenceBase
 	virtual bool ContainRecursive(TArray<UAnimCompositeBase*>& CurrentAccumulatedList) PURE_VIRTUAL(UAnimCompositeBase::ContainRecursive, return false; );
 
 	virtual void SetCompositeLength(float InLength) PURE_VIRTUAL(UAnimCompositeBase::SetCompositeLength, );
-
-#if WITH_EDITOR
-	virtual void PopulateWithExistingModel(TScriptInterface<IAnimationDataModel> ExistingDataModel) override;
-#endif // WITH_EDITOR
 };
 
diff --git a/Engine/Source/Runtime/Engine/Classes/Animation/AnimData/AnimDataModel.h b/Engine/Source/Runtime/Engine/Classes/Animation/AnimData/AnimDataModel.h
--- a/Engine/Source/Runtime/Engine/Classes/Animation/AnimData/AnimDataModel.h
+++ b/Engine/Source/Runtime/Engine/Classes/Animation/AnimData/AnimDataModel.h
@@ -4,75 +4,279 @@
 
 #include "CoreMinimal.h"
 #include "UObject/ObjectMacros.h"
-
-#include "IAnimationDataModel.h"
+#include "AnimDataNotifications.h"
+#include "Animation/AttributeCurve.h"
 
 #include "AnimDataModel.generated.h"
 
+/**
+ * Structure encapsulating a single bone animation track.
+ */
+USTRUCT(BlueprintType)
+struct ENGINE_API FBoneAnimationTrack
+{
+	GENERATED_BODY()
+
+	/** Internally stored data representing the animation bone data */
+	UPROPERTY(VisibleAnywhere, BlueprintReadOnly, Category = "Model")
+	FRawAnimSequenceTrack InternalTrackData;
+
+	/** Index corresponding to the bone this track corresponds to within the target USkeleton */
+	UPROPERTY(VisibleAnywhere, BlueprintReadOnly, Category = "Model")
+	int32 BoneTreeIndex = INDEX_NONE;
+
+	/** Name of the bone this track corresponds to */
+	UPROPERTY(VisibleAnywhere, BlueprintReadOnly, Category = "Model")
+	FName Name;
+};
+
+/**
+ * Structure encapsulating animated curve data. Currently only contains Float and Transform curves.
+ */
+USTRUCT(BlueprintType)
+struct ENGINE_API FAnimationCurveData
+{
+	GENERATED_BODY()
+
+	/** Float-based animation curves */
+	UPROPERTY(VisibleAnywhere, BlueprintReadOnly, Category = "Model")
+	TArray<FFloatCurve>	FloatCurves;
+
+	/** FTransform-based animation curves, used for animation layer editing */
+	UPROPERTY(VisibleAnywhere, BlueprintReadOnly, Category = "Model")
+	TArray<FTransformCurve>	TransformCurves;
+};
+
+/**
+ * Structure encapsulating animated (bone) attribute data.
+ */
+USTRUCT(BlueprintType)
+struct ENGINE_API FAnimatedBoneAttribute
+{
+	GENERATED_BODY()
+
+	/** Identifier to reference this attribute by */
+	UPROPERTY(VisibleAnywhere, BlueprintReadOnly, Category = "Model")
+	FAnimationAttributeIdentifier Identifier;	
+
+	/** Curve containing the (animated) attribute data */
+	UPROPERTY(VisibleAnywhere, BlueprintReadOnly, Category = "Model")
+	FAttributeCurve Curve;
+};
+
+struct FTrackToSkeletonMap;
+struct FAnimationCurveIdentifier;
+
 /**
  * The Model represents the source data for animations. It contains both bone animation data as well as animated curves.
  * They are currently only a sub-object of a AnimSequenceBase instance. The instance derives all runtime data from the source data. 
  */
 UCLASS(BlueprintType)
-class ENGINE_API UAnimDataModel : public UObject, public IAnimationDataModel
+class ENGINE_API UAnimDataModel : public UObject
 {
 	GENERATED_BODY()
 public:
-	/** Begin UObject overrides */
+
+	/** Begin UAnimDataModel overrides */
 	virtual void PostLoad() override;
 	virtual void PostDuplicate(bool bDuplicateForPIE) override;
 	virtual bool IsEditorOnly() const override { return true; }
-	/** End UObject overrides */
+	/** End UAnimDataModel overrides */
+
+	/**
+	* @return	Total length of play-able animation data 
+	*/
+	UFUNCTION(BlueprintPure, Category=AnimationDataModel)
+	float GetPlayLength() const;
 	
-	/** Begin IAnimationDataModel overrides */
-	virtual double GetPlayLength() const override;
-	virtual int32 GetNumberOfFrames() const override;
-	virtual int32 GetNumberOfKeys() const override;
-	virtual FFrameRate GetFrameRate() const override;
-	virtual const TArray<FBoneAnimationTrack>& GetBoneAnimationTracks() const override;
-	virtual const FBoneAnimationTrack& GetBoneTrackByIndex(int32 TrackIndex) const override;
-	virtual const FBoneAnimationTrack& GetBoneTrackByName(FName TrackName) const override;
-	virtual const FBoneAnimationTrack* FindBoneTrackByName(FName Name) const override;
-	virtual const FBoneAnimationTrack* FindBoneTrackByIndex(int32 BoneIndex) const override;
-	virtual int32 GetBoneTrackIndex(const FBoneAnimationTrack& Track) const override;
-	virtual int32 GetBoneTrackIndexByName(FName TrackName) const override;
-	virtual bool IsValidBoneTrackIndex(int32 TrackIndex) const override;
-	virtual int32 GetNumBoneTracks() const override;
-	virtual void GetBoneTrackNames(TArray<FName>& OutNames) const override;
-
-	virtual const FAnimCurveBase* FindCurve(const FAnimationCurveIdentifier& CurveIdentifier) const override;
-	virtual const FFloatCurve* FindFloatCurve(const FAnimationCurveIdentifier& CurveIdentifier) const override;
-	virtual const FTransformCurve* FindTransformCurve(const FAnimationCurveIdentifier& CurveIdentifier) const override;
-	virtual const FRichCurve* FindRichCurve(const FAnimationCurveIdentifier& CurveIdentifier) const override;
-	virtual const FAnimationCurveData& GetCurveData() const override;
-	virtual int32 GetNumberOfTransformCurves() const override;
-	virtual int32 GetNumberOfFloatCurves() const override;
-	virtual const TArray<struct FFloatCurve>& GetFloatCurves() const override;
-	virtual const TArray<struct FTransformCurve>& GetTransformCurves() const override;	
-	virtual const FAnimCurveBase& GetCurve(const FAnimationCurveIdentifier& CurveIdentifier) const override;
-	virtual const FFloatCurve& GetFloatCurve(const FAnimationCurveIdentifier& CurveIdentifier) const override;
-	virtual const FTransformCurve& GetTransformCurve(const FAnimationCurveIdentifier& CurveIdentifier) const override;
-	virtual const FRichCurve& GetRichCurve(const FAnimationCurveIdentifier& CurveIdentifier) const override;
-	virtual TArrayView<const FAnimatedBoneAttribute> GetAttributes() const override;
-	virtual int32 GetNumberOfAttributes() const override;
-	virtual int32 GetNumberOfAttributesForBoneIndex(const int32 BoneIndex) const override;
-	virtual void GetAttributesForBone(const FName& BoneName, TArray<const FAnimatedBoneAttribute*>& OutBoneAttributes) const override;
-	virtual const FAnimatedBoneAttribute& GetAttribute(const FAnimationAttributeIdentifier& AttributeIdentifier) const override;
-	virtual const FAnimatedBoneAttribute* FindAttribute(const FAnimationAttributeIdentifier& AttributeIdentifier) const override;
-	virtual UAnimSequence* GetAnimationSequence() const override;
-	virtual FAnimDataModelModifiedEvent& GetModifiedEvent() override { return ModifiedEvent; }
-	virtual FGuid GenerateGuid() const override;
-#if WITH_EDITOR
-	virtual void Evaluate(FAnimationPoseData& InOutPoseData, const UE::Anim::DataModel::FEvaluationContext& EvaluationContext) const override;
-#endif
-	virtual TScriptInterface<IAnimationDataController> GetController() override;
-	virtual bool HasBeenPopulated() const override { return bPopulated; }
-protected:
-	virtual IAnimationDataModel::FModelNotifier& GetNotifier() override;
-	virtual FAnimDataModelModifiedDynamicEvent& GetModifiedDynamicEvent() override { return ModifiedEventDynamic; }
-	virtual void OnNotify(const EAnimDataModelNotifyType& NotifyType, const FAnimDataModelNotifPayload& Payload) override {}
-	/** End IAnimationDataModel overrides */
+	/**
+	* @return	Total number of frames of animation data stored 
+	*/
+	UFUNCTION(BlueprintPure, Category = AnimationDataModel)
+	int32 GetNumberOfFrames() const;
+
+	/**
+	* @return	Total number of animation data keys stored 
+	*/
+	UFUNCTION(BlueprintPure, Category = AnimationDataModel)
+	int32 GetNumberOfKeys() const;
+
+	/**
+	* @return	Frame rate at which the animation data is key-ed 
+	*/
+	UFUNCTION(BlueprintPure, Category = AnimationDataModel)
+	const FFrameRate& GetFrameRate() const;
+	
+	/**
+	* @return	Array containg all bone animation tracks 
+	*/
+	UFUNCTION(BlueprintPure, Category = AnimationDataModel)
+	const TArray<FBoneAnimationTrack>& GetBoneAnimationTracks() const;
+	
+	/**
+	* @return	Bone animation track for the provided index
+	*/
+	UFUNCTION(BlueprintPure, Category = AnimationDataModel)
+	const FBoneAnimationTrack& GetBoneTrackByIndex(int32 TrackIndex) const;
+
+	/**
+	* @return	Bone animation track for the provided (bone) name
+	*/
+	UFUNCTION(BlueprintPure, Category = AnimationDataModel)
+	const FBoneAnimationTrack& GetBoneTrackByName(FName TrackName) const;
+	
+	/**
+	* @return	Bone animation track for the provided (bone) name if found, otherwise returns a nullptr 
+	*/
+	const FBoneAnimationTrack* FindBoneTrackByName(FName Name) const;
+
+	/**
+	* @return	Bone animation track for the provided index if valid, otherwise returns a nullptr 
+	*/
+	const FBoneAnimationTrack* FindBoneTrackByIndex(int32 BoneIndex) const;
+
+	/**
+	* @return	Internal track index for the provided bone animation track if found, otherwise returns INDEX_NONE 
+	*/
+	UFUNCTION(BlueprintPure, Category = AnimationDataModel)
+	int32 GetBoneTrackIndex(const FBoneAnimationTrack& Track) const;
+
+	/**
+	* @return	Internal track index for the provided (bone) name if found, otherwise returns INDEX_NONE 
+	*/
+	UFUNCTION(BlueprintPure, Category = AnimationDataModel)
+	int32 GetBoneTrackIndexByName(FName TrackName) const;
 
+	/**
+	* @return	Whether or not the provided track index is valid 
+	*/
+	UFUNCTION(BlueprintPure, Category = AnimationDataModel)
+	bool IsValidBoneTrackIndex(int32 TrackIndex) const;
+
+	/**
+	* @return	Total number of bone animation tracks
+	*/
+	UFUNCTION(BlueprintPure, Category = AnimationDataModel)
+	const int32 GetNumBoneTracks() const;
+
+	/**
+	* Populates the provided array with all contained (bone) track names
+	*
+	* @param	OutNames	[out] Array containing all bone track names
+	*/
+	UFUNCTION(BlueprintPure, Category = AnimationDataModel)
+	void GetBoneTrackNames(TArray<FName>& OutNames) const;
+
+	/** Returns all contained curve animation data */
+	const FAnimationCurveData& GetCurveData() const;
+
+	/**
+	* @return	Total number of stored FTransform curves
+	*/
+	UFUNCTION(BlueprintPure, Category = AnimationDataModel)
+	int32 GetNumberOfTransformCurves() const;
+
+	/**
+	* @return	Total number of stored float curves
+	*/
+	UFUNCTION(BlueprintPure, Category = AnimationDataModel)
+	int32 GetNumberOfFloatCurves() const;
+
+	/**
+	* @return	Array containing all stored float curves 
+	*/
+	const TArray<struct FFloatCurve>& GetFloatCurves() const;
+
+	/**
+	* @return	Array containing all stored FTransform curves 
+	*/
+	const TArray<struct FTransformCurve>& GetTransformCurves() const;
+
+	/**
+	* @return	Curve ptr for the provided identifier if valid, otherwise returns a nullptr 
+	*/
+	const FAnimCurveBase* FindCurve(const FAnimationCurveIdentifier& CurveIdentifier) const;
+
+	/**
+	* @return	Float Curve ptr for the provided identifier if valid, otherwise returns a nullptr
+	*/
+	const FFloatCurve* FindFloatCurve(const FAnimationCurveIdentifier& CurveIdentifier) const;
+
+	/**
+	* @return	Transform Curve ptr for the provided identifier if valid, otherwise returns a nullptr
+	*/
+	const FTransformCurve* FindTransformCurve(const FAnimationCurveIdentifier& CurveIdentifier) const;
+
+	/**
+	* @return	Rich curve ptr for the provided identifier if valid, otherwise returns a nullptr
+	*/
+	const FRichCurve* FindRichCurve(const FAnimationCurveIdentifier& CurveIdentifier) const;
+
+	/**
+	* @return	Curve object for the provided identifier if valid
+	*/
+	const FAnimCurveBase& GetCurve(const FAnimationCurveIdentifier& CurveIdentifier) const;
+
+	/**
+	* @return	Float Curve object for the provided identifier if valid
+	*/
+	const FFloatCurve& GetFloatCurve(const FAnimationCurveIdentifier& CurveIdentifier) const;
+
+	/**
+	* @return	Transform Curve object for the provided identifier if valid
+	*/
+	const FTransformCurve& GetTransformCurve(const FAnimationCurveIdentifier& CurveIdentifier) const;
+
+	/**
+	* @return	Rich Curve object for the provided identifier if valid
+	*/
+	const FRichCurve& GetRichCurve(const FAnimationCurveIdentifier& CurveIdentifier) const;
+		
+	/**
+	* @return	Animated (bone) attributes stored
+	*/
+	TArrayView<const FAnimatedBoneAttribute> GetAttributes() const;
+
+	/**
+	* @return	Number of animated (bone) attributes stored
+	*/
+	int32 GetNumberOfAttributes() const;
+
+	/**
+	* @return	Number of animated (bone) attributes stored for the specified bone index
+	*/
+	int32 GetNumberOfAttributesForBoneIndex(const int32 BoneIndex) const;
+
+	/**
+	* @return	All animated (bone) attributes stored for the specified bone name
+	*/
+	void GetAttributesForBone(const FName& BoneName, TArray<const FAnimatedBoneAttribute*>& OutBoneAttributes) const;
+
+	/**
+	* @return	Animated (bone) attribute object for the provided identifier if valid
+	*/
+	const FAnimatedBoneAttribute& GetAttribute(const FAnimationAttributeIdentifier& AttributeIdentifier) const;
+
+	/**
+	* @return	Animated (bone) attribute ptr for the provided identifier if valid, otherwise returns a nullptr 
+	*/
+	const FAnimatedBoneAttribute* FindAttribute(const FAnimationAttributeIdentifier& AttributeIdentifier) const;
+		
+	/**
+	* @return	The outer UAnimSequence object if found, otherwise returns a nullptr 
+	*/
+	UFUNCTION(BlueprintPure, Category = AnimationDataModel)
+	UAnimSequence* GetAnimationSequence() const;
+
+	/**
+	* @return	Multicast delegate which is broadcasted to propagated changes to any internal data, see FAnimDataModelModifiedEvent and EAnimDataModelNotifyType
+	*/
+	FAnimDataModelModifiedEvent& GetModifiedEvent() { return ModifiedEvent; }
+
+	/**
+	* @return	GUID representing the contained data and state 
+	*/
+	FGuid GenerateGuid() const;
 private:
 	/** Helper functionality used by UAnimDataController to retrieve mutable data */ 
 	FBoneAnimationTrack* FindMutableBoneTrackByName(FName Name);
@@ -82,10 +286,77 @@ private:
 	FAnimCurveBase* FindMutableCurveById(const FAnimationCurveIdentifier& CurveIdentifier);	   
 	FRichCurve* GetMutableRichCurve(const FAnimationCurveIdentifier& CurveIdentifier);
 
+	/**
+	* Broadcasts a new EAnimDataModelNotifyType with the provided payload data alongside it.
+	*
+	* @param	NotifyType			Type of notify to broadcast
+	* @param	PayloadData			Typed payload data
+	*/
+	template<typename T>
+	void Notify(EAnimDataModelNotifyType NotifyType, const T& PayloadData)
+	{
+		UScriptStruct* TypeScriptStruct = T::StaticStruct();
+
+		const FAnimDataModelNotifPayload Payload((int8*)&PayloadData, TypeScriptStruct);
+		ModifiedEvent.Broadcast(NotifyType, this, Payload);
+
+		if (ModifiedEventDynamic.IsBound())
+		{
+			ModifiedEventDynamic.Broadcast(NotifyType, this, Payload);
+		}
+
+		// Only regenerate transient data when not in a bracket, or at the end of one
+		{
+			if (NotifyType == EAnimDataModelNotifyType::BracketOpened)
+			{
+				++BracketCounter;
+			}
+			if (NotifyType == EAnimDataModelNotifyType::BracketClosed)
+			{
+				--BracketCounter;
+			}
+
+			check(BracketCounter >= 0);
+		}
+	}
+
+	/**
+	* Broadcasts a new EAnimDataModelNotifyType alongside of an empty payload.
+	*
+	* @param	NotifyType			Type of notify to broadcast
+	*/
+	void Notify(EAnimDataModelNotifyType NotifyType)
+	{
+		FEmptyPayload EmptyPayload;
+		const FAnimDataModelNotifPayload Payload((int8*)&EmptyPayload, FEmptyPayload::StaticStruct());
+
+		ModifiedEvent.Broadcast(NotifyType, this, Payload);
+
+		if (ModifiedEventDynamic.IsBound())
+		{
+			ModifiedEventDynamic.Broadcast(NotifyType, this, Payload);
+		}
+
+		// Only regenerate transient data when not in a bracket, or at the end of one
+		{
+			if (NotifyType == EAnimDataModelNotifyType::BracketOpened)
+			{
+				++BracketCounter;
+			}
+			if (NotifyType == EAnimDataModelNotifyType::BracketClosed)
+			{
+				--BracketCounter;
+			}
+
+			check(BracketCounter >= 0);
+		}
+	}
+
 private:
+
 	UPROPERTY(Transient)
 	int32 BracketCounter = 0;
-
+private:
 	/** Dynamic delegate event allows scripting to register to any broadcasted notify. */
 	UPROPERTY(BlueprintAssignable, Transient, Category = AnimationDataModel, meta = (ScriptName = "ModifiedEvent", AllowPrivateAccess = "true"))
 	FAnimDataModelModifiedDynamicEvent ModifiedEventDynamic;
@@ -98,7 +369,6 @@ private:
 	TArray<FBoneAnimationTrack> BoneAnimationTracks;
 
 	/** Total playable length of the contained animation data */
-	UE_DEPRECATED(5.1, "PlayLength is deprecated use GetPlayLength instead, as it is now calculated with Number of Frames * FrameRate instead of stored as a value")
 	UPROPERTY(VisibleAnywhere, BlueprintReadOnly, Category = "Animation Data Model", meta = (AllowPrivateAccess = "true"))
 	float PlayLength;
 	
@@ -120,12 +390,7 @@ private:
 	
 	/** Container with all animated (bone) attribute data */
 	UPROPERTY(VisibleAnywhere, BlueprintReadOnly, Category = "Animation Data Model", meta = (AllowPrivateAccess = "true"))
-	TArray<FAnimatedBoneAttribute> AnimatedBoneAttributes;
-
-	UPROPERTY()
-	bool bPopulated = false;
-
-	TUniquePtr<IAnimationDataModel::FModelNotifier> Notifier;
+	TArray<FAnimatedBoneAttribute> AnimatedBoneAttributes;	
 
 	friend class UAnimDataController;
 	friend class FAnimDataControllerTestBase;
diff --git a/Engine/Source/Runtime/Engine/Classes/Animation/AnimData/AnimDataNotifications.h b/Engine/Source/Runtime/Engine/Classes/Animation/AnimData/AnimDataNotifications.h
--- a/Engine/Source/Runtime/Engine/Classes/Animation/AnimData/AnimDataNotifications.h
+++ b/Engine/Source/Runtime/Engine/Classes/Animation/AnimData/AnimDataNotifications.h
@@ -13,7 +13,7 @@
 
 #include "AnimDataNotifications.generated.h"
 
-class IAnimationDataModel;
+class UAnimDataModel;
 
 UENUM(BlueprintType)
 enum class EAnimDataModelNotifyType : uint8
@@ -122,40 +122,17 @@ struct FSequenceLengthChangedPayload : public FEmptyPayload
 {
 	GENERATED_BODY()
 
-	PRAGMA_DISABLE_DEPRECATION_WARNINGS
-	FSequenceLengthChangedPayload() = default;
-	FSequenceLengthChangedPayload(const FSequenceLengthChangedPayload&) = default;
-	FSequenceLengthChangedPayload(FSequenceLengthChangedPayload&&) = default;
-	FSequenceLengthChangedPayload& operator=(const FSequenceLengthChangedPayload&) = default;
-	FSequenceLengthChangedPayload& operator=(FSequenceLengthChangedPayload&&) = default;
-	PRAGMA_ENABLE_DEPRECATION_WARNINGS
-
 	/** Previous playable length for the Model */
-	UE_DEPRECATED(5.0, "Time in seconds is deprecated use FFrameNumber members (PreviousNumberOfFrames) instead")
 	UPROPERTY(BlueprintReadOnly, Category = Payload)
 	float PreviousLength = 0.f;
 
 	/** Time at which the change in length has been made */
-	UE_DEPRECATED(5.0, "Time in seconds is deprecated use FFrameNumber members (Frame0) instead")
 	UPROPERTY(BlueprintReadOnly, Category = Payload)
 	float T0 = 0.f;
 
 	/** Length of time which is inserted or removed starting at T0 */
-	UE_DEPRECATED(5.0, "Time in seconds is deprecated use FFrameNumber members (Frame1) instead")
 	UPROPERTY(BlueprintReadOnly, Category = Payload)
 	float T1 = 0.f;
-
-	/** Previous playable number of frames for the Model */
-    UPROPERTY(BlueprintReadOnly, Category = Payload)
-    FFrameNumber PreviousNumberOfFrames;
-
-    /** Frame number at which the change in frames has been made */
-    UPROPERTY(BlueprintReadOnly, Category = Payload)
-    FFrameNumber Frame0;
-
-    /** Amount of frames which is inserted or removed starting at Frame0 */
-    UPROPERTY(BlueprintReadOnly, Category = Payload)
-	FFrameNumber Frame1;
 };
 
 USTRUCT(BlueprintType)
@@ -280,5 +257,5 @@ public:
 #endif // WITH_EDITOR
 };
 
-DECLARE_MULTICAST_DELEGATE_ThreeParams(FAnimDataModelModifiedEvent, const EAnimDataModelNotifyType& /* type */, IAnimationDataModel* /* Model */, const FAnimDataModelNotifPayload& /* payload */);
-DECLARE_DYNAMIC_MULTICAST_DELEGATE_ThreeParams(FAnimDataModelModifiedDynamicEvent, EAnimDataModelNotifyType, NotifType, TScriptInterface<IAnimationDataModel>, Model, const FAnimDataModelNotifPayload&, Payload);
\ No newline at end of file
+DECLARE_MULTICAST_DELEGATE_ThreeParams(FAnimDataModelModifiedEvent, const EAnimDataModelNotifyType& /* type */, UAnimDataModel* /* Model */, const FAnimDataModelNotifPayload& /* payload */);
+DECLARE_DYNAMIC_MULTICAST_DELEGATE_ThreeParams(FAnimDataModelModifiedDynamicEvent, EAnimDataModelNotifyType, NotifType, UAnimDataModel*, Model, const FAnimDataModelNotifPayload&, Payload);
\ No newline at end of file
diff --git a/Engine/Source/Runtime/Engine/Classes/Animation/AnimData/AnimDataNotifyCollector.h b/Engine/Source/Runtime/Engine/Classes/Animation/AnimData/AnimDataNotifyCollector.h
new file mode 100644
--- /dev/null
+++ b/Engine/Source/Runtime/Engine/Classes/Animation/AnimData/AnimDataNotifyCollector.h
@@ -0,0 +1,81 @@
+// Copyright Epic Games, Inc. All Rights Reserved.
+
+#pragma once
+
+#include "Animation/AnimData/AnimDataNotifications.h"
+#include "Containers/Set.h"
+#include "Containers/Array.h"
+
+enum class EAnimDataModelNotifyType : uint8;
+
+namespace UE {
+namespace Anim {
+
+#if WITH_EDITOR
+
+/** Helper structure for keeping track of which notifies of type EAnimDataModelNotifType are broadcasted
+between top-level EAnimDataModelNotifType::BracketOpened and EAnimDataModelNotifType::BracketClosed notifies */
+struct ENGINE_API FAnimationDataNotifyCollector
+{
+	FAnimationDataNotifyCollector() : BracketDepth(0) {}
+
+	/** Handle a broadcasted notify, reset if we are opening a new top-level bracket*/
+	void Handle(EAnimDataModelNotifType NotifyType)
+	{
+		NotifyTypes.Add(NotifyType);
+
+		if (NotifyType == EAnimDataModelNotifType::BracketOpened)
+		{
+			if (BracketDepth == 0)
+			{
+				Reset();
+			}
+
+			++BracketDepth;
+		}
+		else if (NotifyType == EAnimDataModelNotifType::BracketClosed)
+		{
+			--BracketDepth;
+		}
+	}
+	
+	/** Returns whether or not the notify of the provided types was broadcasted */
+	bool Contains(EAnimDataModelNotifType NotifyType) const
+	{
+		return NotifyTypes.Find(NotifyType) != nullptr;
+	}
+
+	/** Returns whether or not any of the provided notify types were broadcasted */
+	bool Contains(const TArray<EAnimDataModelNotifType>& TestNotifyTypes) const
+	{
+		for (EAnimDataModelNotifType Notify : TestNotifyTypes)
+		{
+			if (NotifyTypes.Find(Notify) != nullptr)
+			{
+				return true;
+			}
+		}
+
+		return false;
+	}
+
+	/** Returns whether or not a bracket is still open */
+	bool IsWithinBracket() const { return BracketDepth > 0; }
+
+	/** Returns whether or not all brackets have been closed */
+	bool IsNotWithinBracket() const { return BracketDepth == 0; }
+protected:
+	void Reset()
+	{
+		NotifyTypes.Empty();
+	}
+protected:
+	TSet<EAnimDataModelNotifType> NotifyTypes;
+	int32 BracketDepth;
+};
+
+#endif // WITH_EDITOR
+
+} // namespace Anim	
+
+} // namespace UE
\ No newline at end of file
diff --git a/Engine/Source/Runtime/Engine/Classes/Animation/AnimData/CurveIdentifier.h b/Engine/Source/Runtime/Engine/Classes/Animation/AnimData/CurveIdentifier.h
--- a/Engine/Source/Runtime/Engine/Classes/Animation/AnimData/CurveIdentifier.h
+++ b/Engine/Source/Runtime/Engine/Classes/Animation/AnimData/CurveIdentifier.h
@@ -69,22 +69,10 @@ struct ENGINE_API FAnimationCurveIdentifier
 
 		return *this;
 	}
-	
-	friend uint32 GetTypeHash(const FAnimationCurveIdentifier& CurveId)
-	{
-		return HashCombine(HashCombine(HashCombine((uint32)CurveId.Channel, (uint32)CurveId.Axis), (uint32)CurveId.CurveType), GetTypeHash(CurveId.InternalName.DisplayName));
-	}
 
-	UPROPERTY()
 	FSmartName InternalName;
-	
-	UPROPERTY()
 	ERawCurveTrackTypes CurveType;
-
-	UPROPERTY()
 	ETransformCurveChannel Channel;
-	
-	UPROPERTY()
 	EVectorCurveChannel Axis;
 };
 
diff --git a/Engine/Source/Runtime/Engine/Classes/Animation/AnimData/IAnimationDataController.h b/Engine/Source/Runtime/Engine/Classes/Animation/AnimData/IAnimationDataController.h
--- a/Engine/Source/Runtime/Engine/Classes/Animation/AnimData/IAnimationDataController.h
+++ b/Engine/Source/Runtime/Engine/Classes/Animation/AnimData/IAnimationDataController.h
@@ -7,25 +7,12 @@
 #include "CurveIdentifier.h"
 #include "Misc/FrameRate.h"
 #include "AttributeIdentifier.h"
-#include "IAnimationDataModel.h"
-#include "Logging/LogVerbosity.h"
 #include "Algo/Transform.h"
-#include "Animation/AnimCurveTypes.h"
-
-#if WITH_EDITOR
-#include "ChangeTransactor.h"
-#endif // WITH_EDITOR
 
 #include "IAnimationDataController.generated.h"
 
-#define LOCTEXT_NAMESPACE "IAnimationDataController"
-
-class FText;
 class UAssetUserData;
-class IAnimationDataModel;
-
-struct FAnimationCurveIdentifier;
-struct FAnimationAttributeIdentifier;
+class UAnimDataModel;
 
 namespace UE {
 namespace Anim {
@@ -35,7 +22,7 @@ namespace Anim {
 
 /**
  * The Controller is the sole authority to perform changes on the Animation Data Model. Any mutation to the model made will
- * cause a subsequent notify (EAnimDataModelNotifyType) to be broadcast-ed from the Model's ModifiedEvent. Alongside of it is a 
+ * cause a subsequent notify (EAnimDataModelNotifyType) to be broadcasted from the Model's ModifiedEvent. Alongside of it is a 
  * payload containing information relevant to the mutation. These notifies should be relied upon to update any dependent views 
  * or generated (derived) data.
  */
@@ -66,7 +53,7 @@ public:
 			Controller.OpenBracket(InDescription, bShouldTransact);
 		}
 		
-		FScopedBracket(const TScriptInterface<IAnimationDataController>& InController, const FText& InDescription, bool bInShouldTransact=true)
+		FScopedBracket(TScriptInterface<IAnimationDataController>& InController, const FText& InDescription, bool bInShouldTransact=true)
             : Controller(*InController), bShouldTransact(bInShouldTransact)
 		{
 			Controller.OpenBracket(InDescription, bShouldTransact);
@@ -85,25 +72,25 @@ public:
 	/**
 	* Sets the AnimDataModel instance this controller is supposed to be targeting
 	*
-	* @param	InModel		IAnimationDataModel instance to target
+	* @param	InModel		UAnimDataModel instance to target
 	*/
 	UFUNCTION(BlueprintCallable, Category = AnimationData)
-	virtual void SetModel(TScriptInterface<IAnimationDataModel> InModel) = 0;
+	virtual void SetModel(UAnimDataModel* InModel) = 0;
 
 	/**
-	* @return		The IAnimationDataModel instance this controller is currently targeting
+	* @return		The AnimDataModel instance this controller is currently targeting
 	*/
 	UFUNCTION(BlueprintCallable, Category = AnimationData)
-	virtual TScriptInterface<IAnimationDataModel> GetModelInterface() const = 0;
+	virtual UAnimDataModel* GetModel() = 0;
 
 	/**
-	* @return		The IAnimationDataModel instance this controller is currently targeting
+	* @return		The AnimDataModel instance this controller is currently targeting
 	*/
-	virtual const IAnimationDataModel* const GetModel() const = 0;
+	virtual const UAnimDataModel* const GetModel() const = 0;
 
 	/**
 	* Opens an interaction bracket, used for combining a set of controller actions. Broadcasts a EAnimDataModelNotifyType::BracketOpened notify,
-	* this can be used by any Views or dependent systems to halt any unnecessary or invalid operations until the (last) bracket is closed.
+	* this can be used by any Views or dependendent systems to halt any unnecessary or invalid operations until the (last) bracket is closed.
 	*
 	* @param	InTitle				Description of the bracket, e.g. "Generating Curve Data"
 	* @param	bShouldTransact		Whether or not any undo-redo changes should be generated
@@ -114,6 +101,7 @@ public:
 	/**
 	* Closes a previously opened interaction bracket, used for combining a set of controller actions. Broadcasts a EAnimDataModelNotifyType::BracketClosed notify.
 	*
+	* @param	InTitle				Description of the bracket, e.g. "Generating Curve Data"
 	* @param	bShouldTransact		Whether or not any undo-redo changes should be generated
 	*/
 	UFUNCTION(BlueprintCallable, Category = AnimationData)
@@ -123,30 +111,22 @@ public:
 	* Sets the total play-able length in seconds. Broadcasts a EAnimDataModelNotifyType::SequenceLengthChanged notify if successful.
 	* The number of frames and keys for the provided length is recalculated according to the current value of UAnimDataModel::FrameRate.
 	*
-	* @param	NewLengthInFrames	Total new play-able number of frames value (according to frame rate), has to be positive and non-zero
+	* @param	Length				New play-able length value, has to be positive and non-zero
 	* @param	bShouldTransact		Whether or not any undo-redo changes should be generated
 	*/
 	UFUNCTION(BlueprintCallable, Category = AnimationData)
-	virtual void SetNumberOfFrames(FFrameNumber NewLengthInFrames, bool bShouldTransact = true) = 0;
-
-	UE_DEPRECATED(5.1, "SetPlayLength with length in seconds in deprecated use SetNumberOfFrames using FFrameNumber instead")
-	UFUNCTION(BlueprintCallable, Category = AnimationData, meta=(DeprecatedFunction, DeprecationMessage="SetPlayLength is deprecated use SetNumberOfFrames instead."))
 	virtual void SetPlayLength(float Length, bool bShouldTransact = true) = 0;
 
 	/*** Sets the total play-able length in seconds. Broadcasts a EAnimDataModelNotifyType::SequenceLengthChanged notify if successful.
 	* T0 and T1 are expected to represent the window of time that was either added or removed. E.g. for insertion T0 indicates the time
 	* at which additional time starts and T1 were it ends. For removal T0 indicates the time at which time should be started to remove, and T1 indicates the end. Giving a total of T1 - T0 added or removed length.
 	* The number of frames and keys for the provided length is recalculated according to the current value of UAnimDataModel::FrameRate.
-	* @param	NewLengthInFrames	Total new play-able number of frames value (according to frame rate), has to be positive and non-zero
-	* @param	T0					Point between 0 and NewLengthInFrames at which the change in length starts
-	* @param	T1					Point between 0 and NewLengthInFrames at which the change in length ends
+	* @param	Length				Total new play-able length value, has to be positive and non-zero
+	* @param	T0					Point between 0 and Length at which the change in time starts
+	* @param	T1					Point between 0 and Length at which the change in time ends
 	* @param	bShouldTransact		Whether or not any undo-redo changes should be generated
 	*/
 	UFUNCTION(BlueprintCallable, Category = AnimationData)
-    virtual void ResizeNumberOfFrames(FFrameNumber NewLengthInFrames, FFrameNumber T0, FFrameNumber T1, bool bShouldTransact = true) = 0;
-
-	UE_DEPRECATED(5.1, "ResizePlayLength with length in seconds in deprecated use ResizeNumberOfFrames using FFrameNumber instead")
-	UFUNCTION(BlueprintCallable, Category = AnimationData, meta=(DeprecatedFunction, DeprecationMessage="ResizePlayLength is deprecated use ResizeNumberOfFrames instead."))
 	virtual void ResizePlayLength(float NewLength, float T0, float T1, bool bShouldTransact = true) = 0;
 
 	/**
@@ -156,18 +136,14 @@ public:
 	* at which additional time starts and T1 were it ends. For removal T0 indicates the time at which time should be started to remove, and T1 indicates the end. Giving a total of T1 - T0 added or removed length.
 	* The number of frames and keys for the provided length is recalculated according to the current value of UAnimDataModel::FrameRate.
 	*
-	* @param	NewLengthInFrames	Total new play-able number of frames value (according to frame rate), has to be positive and non-zero
-	* @param	T0					Point between 0 and NewLengthInFrames at which the change in length starts
-	* @param	T1					Point between 0 and NewLengthInFrames at which the change in length ends
+	* @param	Length				Total new play-able length value, has to be positive and non-zero
+	* @param	T0					Point between 0 and Length at which the change in time starts
+	* @param	T1					Point between 0 and Length at which the change in time ends
 	* @param	bShouldTransact		Whether or not any undo-redo changes should be generated
 	*/
 	UFUNCTION(BlueprintCallable, Category = AnimationData)
-	virtual void ResizeInFrames(FFrameNumber NewLengthInFrames, FFrameNumber T0, FFrameNumber T1, bool bShouldTransact = true) = 0;
-
-	UE_DEPRECATED(5.1, "Resize with length in seconds in deprecated use ResizeInFrames using FFrameNumber instead")
-	UFUNCTION(BlueprintCallable, Category = AnimationData, meta=(DeprecatedFunction, DeprecationMessage="Resize is deprecated use ResizeInFrames instead."))
 	virtual void Resize(float Length, float T0, float T1, bool bShouldTransact = true) = 0;
-	
+
 	/**
 	* Sets the frame rate according to which the bone animation is expected to be sampled. Broadcasts a EAnimDataModelNotifyType::FrameRateChanged notify if successful.
 	* The number of frames and keys for the provided frame rate is recalculated according to the current value of UAnimDataModel::PlayLength.
@@ -199,8 +175,7 @@ public:
 	*
 	* @return	The index at which the bone track was inserted, INDEX_NONE if the insertion failed
 	*/
-	UE_DEPRECATED(5.1, "InsertBoneTrack has been deprecated, use AddBoneTrack instead")
-	UFUNCTION(BlueprintCallable, Category = AnimationData, meta=(DeprecatedFunction, DeprecationMessage="InsertBoneTrack is deprecated use AddBoneTrack instead."))
+	UFUNCTION(BlueprintCallable, Category = AnimationData)
 	virtual int32 InsertBoneTrack(FName BoneName, int32 DesiredIndex, bool bShouldTransact = true) = 0;
 
 	/**
@@ -256,7 +231,7 @@ public:
 	* The provided number of keys provided is expected to match for each component, be between FrameLowerBound and FrameUpperBound (inclusive), and be non-zero.
 	*
 	* @param	BoneName			Bone name of the track for which the keys should be set
-	* @param	KeyRangeToSet			Range of frames to set keys for
+	* @param	FrameRange			Range of frames to set keys for
 	* @param	PositionalKeys		Array of keys for the translation component
 	* @param	RotationalKeys		Array of keys for the rotation component
 	* @param	ScalingKeys			Array of keys for the scale component
@@ -271,7 +246,7 @@ public:
 	* The provided number of keys provided is expected to match for each component, be between FrameLowerBound and FrameUpperBound (inclusive), and be non-zero.
 	*
 	* @param	BoneName			Bone name of the track for which the keys should be set
-	* @param	KeyRangeToSet			Range of frames to set keys for
+	* @param	FrameRange			Range of frames to set keys for
 	* @param	PositionalKeys		Array of keys for the translation component
 	* @param	RotationalKeys		Array of keys for the rotation component
 	* @param	ScalingKeys			Array of keys for the scale component
@@ -309,7 +284,7 @@ public:
 	/**
 	* Remove the curve with provided identifier. Broadcasts a EAnimDataModelNotifyType::CurveRemoved notify if successful.
 	*
-	* @param	CurveId				Identifier for the to-be-removed curve
+	* @param	CopyCurveId			Identifier for the to-be-removed curve
 	* @param	bShouldTransact		Whether or not any undo-redo changes should be generated
 	*
 	* @return	Whether or not the curve was successfully removed
@@ -343,7 +318,7 @@ public:
 	* Replace the flags for the curve with provided identifier. Broadcasts a EAnimDataModelNotifyType::CurveFlagsChanged notify if successful.
 	*
 	* @param	CurveId			    Identifier for the curve for which the flag state is to be set
-	* @param	Flags				Flag mask with which the existing flags are to be replaced
+	* @param	Flags				Flag mask with which the existings flags are to be replaced
 	* @param	bShouldTransact		Whether or not any undo-redo changes should be generated
 	*
 	* @return	Whether or not the flag mask was successfully set
@@ -476,7 +451,7 @@ public:
 	virtual void UpdateCurveNamesFromSkeleton(const USkeleton* Skeleton, ERawCurveTrackTypes SupportedCurveType, bool bShouldTransact = true) = 0;
 
 	/**
-	* Updates the curve names with the provided skeleton, if a display name is not found it will be added thus modifying the skeleton. Broadcasts a EAnimDataModelNotifyType::CurveRenamed for each curve name for which the UID was different or if it was added as a new smart-name, wrapped within BracketOpened/BracketClosed notifies.
+	* Updates the curve names with the provided skeleton, if a display name is not found it will be added thus modifying the skeleton. Broadcasts a EAnimDataModelNotifyType::CurveRenamed for each curve name for which the UID was different or if it was added as a new smartname, wrapped within BracketOpened/BracketClosed notifies.
 	*
 	* @param	Skeleton			Skeleton to retrieve the display name values from
 	* @param	SupportedCurveType	Curve type for which the names should be updated
@@ -573,9 +548,9 @@ public:
 	* @return	Whether or not the key was successfully set
 	*/
 	template<typename AttributeType>
-	bool SetTypedAttributeKey(const FAnimationAttributeIdentifier& AttributeIdentifier, float Time, const AttributeType& KeyValue, bool bShouldTransact = true)
+	bool SetTypedAttributeKey(const FAnimationAttributeIdentifier& AttributeIdentifier, float Time, AttributeType& KeyValue, bool bShouldTransact = true)
 	{
-		return SetAttributeKey(AttributeIdentifier, Time, static_cast<const void*>(&KeyValue), AttributeType::StaticStruct(), bShouldTransact);
+		return SetAttributeKey(AttributeIdentifier, Time, (const void*)&KeyValue, AttributeType::StaticStruct(), bShouldTransact);
 	}
 
 	/**
@@ -585,7 +560,6 @@ public:
 	* @param	AttributeIdentifier		Identifier for the attribute for which the key is to be set
 	* @param	Time					Time of the to-be-set key
 	* @param	KeyValue				Value of the to-be-set key
-	* @param	TypeStruct				UScriptStruct describing the type of KeyValue
 	* @param	bShouldTransact			Whether or not any undo-redo changes should be generated
 	*
 	* @return	Whether or not the key was successfully set
@@ -598,7 +572,6 @@ public:
 	* @param	AttributeIdentifier		Identifier for the attribute for which the keys are to be replaced
 	* @param	Times					Times with which the existing key timings are to be replaced
 	* @param	KeyValues				Values with which the existing key values are to be replaced
-	* @param	TypeStruct				UScriptStruct describing the type of KeyValues
 	* @param	bShouldTransact			Whether or not any undo-redo changes should be generated
 	*
 	* @return	Whether or not replacing the attribute keys was successful
@@ -616,12 +589,12 @@ public:
 	* @return	Whether or not replacing the attribute keys was successful
 	*/
 	template<typename AttributeType>
-	bool SetTypedAttributeKeys(const FAnimationAttributeIdentifier& AttributeIdentifier, TArrayView<const float> Times, TArrayView<const AttributeType> KeyValues, bool bShouldTransact = true)
+	bool SetTypedAttributeKeys(const FAnimationAttributeIdentifier& AttributeIdentifier, TArrayView<const float> Times, TArrayView<AttributeType> KeyValues, bool bShouldTransact = true)
 	{
 		TArray<const void*> KeyValuePtrs;
 		Algo::Transform(KeyValues, KeyValuePtrs, [](const AttributeType& Value)
 		{
-			return static_cast<const void*>(&Value);
+			return (const void*)&Value;
 		});
 
 		return SetAttributeKeys(AttributeIdentifier, Times, MakeArrayView(KeyValuePtrs), AttributeType::StaticStruct(), bShouldTransact);
@@ -650,175 +623,13 @@ public:
 	*/
 	UFUNCTION(BlueprintCallable, Category = AttributeData)
 	virtual bool DuplicateAttribute(const FAnimationAttributeIdentifier& AttributeIdentifier, const FAnimationAttributeIdentifier& NewAttributeIdentifier, bool bShouldTransact = true) = 0;
-
-	/** Updates/removes/remaps contained animation data according to the newly assigned skeleton
-	 *
-	* @param	TargetSkeleton		Skeleton to which the contained animation data is to be mapped
-	* @param	bShouldTransact		Whether or not any undo-redo changes should be generated
-	* 
-	 */
-	virtual void UpdateWithSkeleton(USkeleton* TargetSkeleton, bool bShouldTransact = true) = 0;
-
-	/** Copies any animation relevant data from an already existing IAnimationDataModel object */
-	virtual void PopulateWithExistingModel(TScriptInterface<IAnimationDataModel> InModel) = 0;
-
-	/** Initializes model data structures */
-	virtual void InitializeModel() = 0;
-
-	/** Returns the final frame number calculating according to the Model its frame-rate, additionally outputs log information for invalid/loss of precision */
-	FFrameNumber ConvertSecondsToFrameNumber(double Seconds) const
-	{
-		ValidateModel();
-		
-		const FFrameRate& ModelFrameRate = GetModel()->GetFrameRate();
-		
-		static const FNumberFormattingOptions DurationFormatOptions = FNumberFormattingOptions()
-			.SetMinimumFractionalDigits(8)
-			.SetMaximumFractionalDigits(8);
-	
-		const FFrameTime FrameTime = ModelFrameRate.AsFrameTime(Seconds);
-		// Check for either small sub-frame value or near zero seconds representation (of sub-frame only)
-		if (!FMath::IsNearlyZero(FrameTime.GetSubFrame(), UE_KINDA_SMALL_NUMBER) &&
-			!FMath::IsNearlyZero(ModelFrameRate.AsSeconds(FFrameTime(0, FrameTime.GetSubFrame())), UE_DOUBLE_KINDA_SMALL_NUMBER))
-		{
-			ReportWarningf(LOCTEXT("SecondsToFrameNumberPrecisionWarning", "Insufficient precision while converting seconds to frames: {0} seconds {1} frames using {2} (sub-frame in seconds {3})"), FText::AsNumber(Seconds), FText::AsNumber(FrameTime.AsDecimal(), &DurationFormatOptions), ModelFrameRate.ToPrettyText(), FText::AsNumber(ModelFrameRate.AsSeconds(FFrameTime(0, FrameTime.GetSubFrame())), &DurationFormatOptions));
-		}
-
-		return FrameTime.GetFrame();		
-	}	
 protected:
 	/** Functionality used by FOpenBracketAction and FCloseBracketAction to broadcast their equivalent notifies without actually opening a bracket. */
 	virtual void NotifyBracketOpen() = 0;
 	virtual void NotifyBracketClosed() = 0;
 
-	template <typename FmtType, typename... Types>
-    void ReportWarningf(const FmtType& Fmt, Types... Args) const
-	{
-		ReportWarning(FText::Format(Fmt, Args...));
-	}
-
-	template <typename FmtType, typename... Types>
-    void ReportErrorf(const FmtType& Fmt, Types... Args) const
-	{
-		ReportError(FText::Format(Fmt, Args...));
-	}
-
-	void ReportWarning(const FText& Message) const
-	{
-		ReportMessage(GetModelInterface().GetObject(), Message, ELogVerbosity::Warning);
-	}
-	
-	void ReportError(const FText& Message) const
-    {
-    	ReportMessage(GetModelInterface().GetObject(), Message, ELogVerbosity::Error);
-    }
-
-	/** Returns whether or not the supplied curve type is supported by the controller functionality */
-	static bool IsSupportedCurveType(ERawCurveTrackTypes CurveType)
-	{
-		const TArray<ERawCurveTrackTypes> SupportedTypes = { ERawCurveTrackTypes::RCT_Float, ERawCurveTrackTypes::RCT_Transform };
-		return SupportedTypes.Contains(CurveType);
-	}
-
-	/** Ensures that a valid model is currently targeted */
-	void ValidateModel() const
-	{
-		checkf(GetModel() != nullptr, TEXT("Invalid Model"));
-	}
-		
-	/** Verifies whether or not the Model's outer object is (or is derived from) the specified UClass */
-	bool CheckOuterClass(UClass* InClass) const
-	{
-		ValidateModel();
-	
-		if (const UObject* ModelOuter = GetModelInterface().GetObject()->GetOuter())
-		{
-			if (const UClass* OuterClass = ModelOuter->GetClass())
-			{
-				if (OuterClass == InClass || OuterClass->IsChildOf(InClass))
-				{
-					return true;
-				}
-				else
-				{
-					ReportErrorf(NSLOCTEXT("IAnimationDataController", "NoValidOuterClassError", "Incorrect outer object class found for Animation Data Model {0}, expected {1} actual {2}"), FText::FromString(GetModelInterface().GetObject()->GetName()), FText::FromString(InClass->GetName()), FText::FromString(OuterClass->GetName()));
-				}
-			}
-		}
-		else
-		{
-			ReportErrorf(NSLOCTEXT("IAnimationDataController", "NoValidOuterObjectFoundError", "No valid outer object found for Animation Data Model {0}"), FText::FromString(GetModelInterface().GetObject()->GetName()));
-		}
-
-		return false;
-	}
-	
-	/** Returns the string representation of the provided curve enum type value */
-	static FString GetCurveTypeValueName(ERawCurveTrackTypes InType)
-	{
-		FString ValueString;
-		if (const UEnum* Enum = FindObject<UEnum>(nullptr, TEXT("/Script/Engine.ERawCurveTrackTypes")))
-		{
-			ValueString = Enum->GetNameStringByValue(static_cast<int64>(InType));
-		}
-
-		return ValueString;
-	}
-	
-	static constexpr int32 DefaultCurveFlags = EAnimAssetCurveFlags::AACF_Editable;
-
-#if WITH_EDITOR
-	typedef TUniquePtr<UE::FScopedCompoundTransaction> FTransaction;
-	FTransaction ConditionalTransaction(const FText& Description, bool bCondition)
-	{
-		if (UE::FChangeTransactor::CanTransactChanges() && bCondition)
-		{ 
-			return MakeUnique<UE::FScopedCompoundTransaction>(ChangeTransactor, Description);
-		}
-
-		return nullptr;
-	}
-
-	typedef TUniquePtr<IAnimationDataController::FScopedBracket> FBracket; 
-	FBracket ConditionalBracket(const FText& Description, bool bCondition)
-	{	
-		return MakeUnique<IAnimationDataController::FScopedBracket>(this, Description, UE::FChangeTransactor::CanTransactChanges() && bCondition);	
-	}
-
-	template<typename ActionClass, class... ActionArguments>
-	void ConditionalAction(bool bCondition, ActionArguments... Arguments)
-	{
-		if (UE::FChangeTransactor::CanTransactChanges() && bCondition) \
-		{ 
-			ChangeTransactor.AddTransactionChange<ActionClass>(Arguments...); \
-		}
-	}
-protected:
-	UE::FChangeTransactor ChangeTransactor;
-#endif // WITH_EDITOR
-
-public:	
-	template <typename FmtType, typename... Types>
-	static void ReportObjectWarningf(const UObject* ErrorObject, const FmtType& Fmt, Types... Args)
-	{
-		ReportMessage(ErrorObject, FText::Format(Fmt, Args...), ELogVerbosity::Warning);
-	}
-
-	template <typename FmtType, typename... Types>
-	static void ReportObjectErrorf(const UObject* ErrorObject, const FmtType& Fmt, Types... Args)
-	{
-		ReportMessage(ErrorObject, FText::Format(Fmt, Args...), ELogVerbosity::Error);
-	}
-
-	static void ReportMessage(const UObject* ErrorObject, const FText& InMessage, ELogVerbosity::Type LogVerbosity);	
 private:
 	friend class FAnimDataControllerTestBase;
 	friend UE::Anim::FOpenBracketAction;
 	friend UE::Anim::FCloseBracketAction;
 };
-
-#define CONDITIONAL_TRANSACTION(Text) FTransaction Transaction = ConditionalTransaction(Test, bShouldTransact);
-#define CONDITIONAL_BRACKET(Text) FBracket Bracket = ConditionalBracket(Text, bShouldTransact);
-#define CONDITIONAL_ACTION(ActionClass, ...) ConditionalAction<ActionClass>(bShouldTransact, __VA_ARGS__);
-
-#undef LOCTEXT_NAMESPACE // "IAnimationDataController"
\ No newline at end of file
diff --git a/Engine/Source/Runtime/Engine/Classes/Animation/AnimMontage.h b/Engine/Source/Runtime/Engine/Classes/Animation/AnimMontage.h
--- a/Engine/Source/Runtime/Engine/Classes/Animation/AnimMontage.h
+++ b/Engine/Source/Runtime/Engine/Classes/Animation/AnimMontage.h
@@ -496,7 +496,7 @@ public:
 	bool SetNextSectionName(FName const & SectionName, FName const & NewNextSectionName);
 	bool SetNextSectionID(int32 const & SectionID, int32 const & NewNextSectionID);
 
-	bool IsValid() const { return (Montage!=nullptr); }
+	bool IsValid() const { return (Montage!=NULL); }
 	bool IsPlaying() const { return IsValid() && bPlaying; }
 	void SetPlaying(bool bInPlaying) { bPlaying = bInPlaying; }
 	bool IsStopped() const { return Blend.GetDesiredValue() == 0.f; }
@@ -924,9 +924,6 @@ public:
 	virtual void InvalidateRecursiveAsset() override;
 	virtual bool ContainRecursive(TArray<UAnimCompositeBase*>& CurrentAccumulatedList) override;
 	virtual void SetCompositeLength(float InLength) override;
-#if WITH_EDITOR
-	virtual void PopulateWithExistingModel(TScriptInterface<IAnimationDataModel> ExistingDataModel) override;
-#endif // WITH_EDITOR
 	//~End UAnimCompositeBase Interface
 
 	/** Utility function to create dynamic montage from AnimSequence */
diff --git a/Engine/Source/Runtime/Engine/Classes/Animation/AnimSequence.h b/Engine/Source/Runtime/Engine/Classes/Animation/AnimSequence.h
--- a/Engine/Source/Runtime/Engine/Classes/Animation/AnimSequence.h
+++ b/Engine/Source/Runtime/Engine/Classes/Animation/AnimSequence.h
@@ -23,7 +23,7 @@
 #include "Animation/AttributeCurve.h"
 
 #if WITH_EDITOR
-#include "AnimData/IAnimationDataModel.h"
+#include "AnimData/AnimDataModel.h"
 #endif // WITH_EDITOR
 
 #include "AnimSequence.generated.h"
@@ -233,11 +233,21 @@ protected:
 	UPROPERTY()
 	FFrameRate SamplingFrameRate;
 
+	/**
+	 * In the future, maybe keeping RawAnimSequenceTrack + TrackMap as one would be good idea to avoid inconsistent array size
+	 * TrackToSkeletonMapTable(i) should contains  track mapping data for RawAnimationData(i). 
+	 */
+	UE_DEPRECATED(5.0, "TrackToSkeletonMapTable has been deprecated see FBoneAnimationTrack::BoneTreeIndex")
+	UPROPERTY()
+	TArray<struct FTrackToSkeletonMap> TrackToSkeletonMapTable;
+
+	/**
+	 * Raw uncompressed keyframe data. 
+	 */
 	UE_DEPRECATED(5.0, "RawAnimationData has been deprecated see FBoneAnimationTrack::InternalTrackData")
 	TArray<struct FRawAnimSequenceTrack> RawAnimationData;
 
-	// Update this if the contents of RawAnimationData changes
-	UE_DEPRECATED(5.1, "RawDataGuid has been deprecated see GenerateGuidFromModel instead")
+	// Update this if the contents of RawAnimationData changes;
 	UPROPERTY()
 	FGuid RawDataGuid;
 
@@ -252,6 +262,7 @@ protected:
 	 * Source RawAnimationData. Only can be overridden by when transform curves are added first time OR imported
 	 */
 	TArray<struct FRawAnimSequenceTrack> SourceRawAnimationData_DEPRECATED;
+
 public:
 
 	/**
@@ -359,6 +370,10 @@ public:
 	UPROPERTY()
 	FString SourceFileTimestamp_DEPRECATED;
 
+	UE_DEPRECATED(5.0, "bNeedsRebake has been deprecated, transform curves are now baked during compression")
+	UPROPERTY(transient)
+	bool bNeedsRebake;
+
 	// Track whether we have updated markers so cached data can be updated
 	int32 MarkerDataUpdateCounter;
 #endif // WITH_EDITORONLY_DATA
@@ -399,10 +414,24 @@ public:
 	virtual void RefreshCacheData() override;
 	virtual EAdditiveAnimationType GetAdditiveAnimType() const override { return AdditiveAnimType; }
 	virtual int32 GetNumberOfSampledKeys() const override;
-	virtual FFrameRate GetSamplingFrameRate() const override { return TargetFrameRate; }
+	virtual const FFrameRate& GetSamplingFrameRate() const override { return TargetFrameRate; }
 	virtual void EvaluateCurveData(FBlendedCurve& OutCurve, float CurrentTime, bool bForceUseRawData = false) const override;
 	virtual float EvaluateCurveData(SmartName::UID_Type CurveUID, float CurrentTime, bool bForceUseRawData = false) const override;
 	virtual bool HasCurveData(SmartName::UID_Type CurveUID, bool bForceUseRawData) const override;
+
+#if WITH_EDITOR
+	UE_DEPRECATED(5.0, "MarkRawDataAsModified has been deprecated, any (Raw Data) modification should be applied using the UAnimDataController API instead. This will handle updating the GUID instead.")
+	PRAGMA_DISABLE_DEPRECATION_WARNINGS
+	virtual void MarkRawDataAsModified(bool bForceNewRawDatGuid = true) override
+	{
+		Super::MarkRawDataAsModified();
+		bUseRawDataOnly = true;
+		RawDataGuid = bForceNewRawDatGuid ? FGuid::NewGuid() : GenerateGuidFromRawData();
+		FlagDependentAnimationsAsRawDataOnly();
+		UpdateDependentStreamingAnimations();
+	}
+	PRAGMA_ENABLE_DEPRECATION_WARNINGS
+#endif
 	//~ End UAnimSequenceBase Interface
 
 	// Extract Root Motion transform from the animation
@@ -443,7 +472,9 @@ public:
 #if WITH_EDITORONLY_DATA
 protected:
 	void UpdateCompressedCurveName(SmartName::UID_Type CurveUID, const struct FSmartName& NewCurveName);
+#endif
 private:
+#if WITH_EDITORONLY_DATA
 	void UpdateRetargetSourceAsset();
 
 	/** Updates the stored sampling frame-rate using the sequence length and number of sampling keys */
@@ -495,9 +526,7 @@ public:
 	 * @param	Time			Time on track to interpolate to.
 	 * @param	bUseRawData		If true, use raw animation data instead of compressed data.
 	 */
-	UE_DEPRECATED(5.1, "Use other GetBoneTransform signature")
 	void GetBoneTransform(FTransform& OutAtom, int32 TrackIndex, float Time, bool bUseRawData) const;
-	void GetBoneTransform(FTransform& OutAtom, int32 TrackIndex, double Time, bool bUseRawData) const;
 
 	/**
 	 * Get Bone Transform of the Time given, relative to Parent for the Track Given
@@ -558,14 +587,10 @@ public:
 
 protected:
 	void ApplyCompressedData(const TArray<uint8>& Data);
-	void ClearCompressedBoneData();
-	void ClearCompressedCurveData();
-	// Write the compressed data to the supplied FArchive
-	void SerializeCompressedData(FArchive& Ar, bool bDDCData);
-
-#if WITH_EDITOR
-	virtual void OnAnimModelLoaded() override;
-#endif
+    void ClearCompressedBoneData();
+    void ClearCompressedCurveData();
+    // Write the compressed data to the supplied FArchive
+    void SerializeCompressedData(FArchive& Ar, bool bDDCData);
 public:
 	bool IsCompressedDataValid() const;
 	bool IsCurveCompressedDataValid() const;
@@ -675,7 +700,6 @@ public:
 	void UpdateDependentStreamingAnimations() const;
 
 	// Generate a GUID from a hash of our own raw data
-	UE_DEPRECATED(5.1, "GenerateGuidFromRawData has been deprecated use IAnimationDataModel::GenerateGuid instead")
 	FGuid GenerateGuidFromRawData() const;
 
 	// Should we be always using our raw data (i.e is our compressed data stale)
@@ -683,12 +707,9 @@ public:
 	void SetUseRawDataOnly(bool bInUseRawDataOnly) { bUseRawDataOnly = bInUseRawDataOnly; }
 
 	// Return this animations guid for the raw data
-	UE_DEPRECATED(5.1, "GetRawDataGuid has been deprecated use GenerateGuidFromModel instead")
 	FGuid GetRawDataGuid() const
 	{ 
-		PRAGMA_DISABLE_DEPRECATION_WARNINGS
 		return RawDataGuid;
-		PRAGMA_ENABLE_DEPRECATION_WARNINGS
 	}
 
 	/** Resets Bone Animation, Curve data and Notify tracks **/
@@ -718,13 +739,18 @@ private:
 	 */ 
 	virtual void RemapTracksToNewSkeleton( USkeleton* NewSkeleton, bool bConvertSpaces ) override;
 
+	/**
+	 * Remap NaN tracks from the RawAnimation data and recompress
+	 */	
+	void RemoveNaNTracks();
+
 	/** Retargeting functions */
 	bool ConvertAnimationDataToRiggingData(FAnimSequenceTrackContainer & RiggingAnimationData);
 	bool ConvertRiggingDataToAnimationData(FAnimSequenceTrackContainer & RiggingAnimationData);
 	int32 GetSpaceBasedAnimationData(TArray< TArray<FTransform> > & AnimationDataInComponentSpace, FAnimSequenceTrackContainer * RiggingAnimationData) const;
-#endif
 
 public:
+#endif
 	/** Refresh sync marker data*/
 	void RefreshSyncMarkerDataFromAuthored();
 
@@ -785,7 +811,7 @@ protected:
 protected:
 #if WITH_EDITOR
 	// Begin UAnimSequenceBase virtual overrides
-	virtual void OnModelModified(const EAnimDataModelNotifyType& NotifyType, IAnimationDataModel* Model, const FAnimDataModelNotifPayload& Payload) override;
+	virtual void OnModelModified(const EAnimDataModelNotifyType& NotifyType, UAnimDataModel* Model, const FAnimDataModelNotifPayload& Payload) override;
 	virtual void PopulateModel() override;
 	// End UAnimSequenceBase virtual overrides
 
@@ -836,5 +862,4 @@ public:
 	friend class FCustomAttributeCustomization;
 	friend class FAnimSequenceTestBase;
 	friend struct UE::Anim::Compression::FScopedCompressionGuard;
-	friend class FAnimDataControllerTestBase;
 };
diff --git a/Engine/Source/Runtime/Engine/Classes/Animation/AnimSequenceBase.h b/Engine/Source/Runtime/Engine/Classes/Animation/AnimSequenceBase.h
--- a/Engine/Source/Runtime/Engine/Classes/Animation/AnimSequenceBase.h
+++ b/Engine/Source/Runtime/Engine/Classes/Animation/AnimSequenceBase.h
@@ -29,7 +29,6 @@ enum ETypeAdvanceAnim
 struct FAnimationPoseData;
 struct FAnimDataModelNotifPayload;
 class UAnimDataModel;
-class IAnimationDataModel;
 enum class EAnimDataModelNotifyType : uint8;
 
 UCLASS(abstract, BlueprintType)
@@ -42,20 +41,14 @@ public:
 	UPROPERTY()
 	TArray<struct FAnimNotifyEvent> Notifies;
 
-protected:
 	/** Length (in seconds) of this AnimSequence if played back with a speed of 1.0. */
 	UE_DEPRECATED(5.0, "Public access to SequenceLength is deprecated, use GetPlayLength or UAnimDataController::SetPlayLength instead")
 	UPROPERTY(Category=Length, AssetRegistrySearchable, VisibleAnywhere, BlueprintReadOnly)
 	float SequenceLength;
 
-	/**
-	 * Raw uncompressed float curve data
-	 */
-	UE_DEPRECATED(5.0, "Public access to RawCurveData is deprecated, see UAnimDataModel for source data or use GetCurveData for runtime instead")
-	UPROPERTY()
-	struct FRawCurveTracks RawCurveData;
-	
-public:
+	UE_DEPRECATED(5.0, "SetSequenceLength is deprecated use UAnimDataController::SetPlayLength instead")
+	virtual void SetSequenceLength(float NewLength);
+
 	/** Number for tweaking playback rate of this animation globally. */
 	UPROPERTY(EditAnywhere, Category=Animation)
 	float RateScale;
@@ -66,6 +59,14 @@ public:
 	 */
 	UPROPERTY(EditAnywhere, Category=Animation)
 	bool bLoop;
+
+	/**
+	 * Raw uncompressed float curve data 
+	 */
+	UE_DEPRECATED(5.0, "Public access to RawCurveData is deprecated, see UAnimDataModel for source data or use GetCurveData for runtime instead")
+	UPROPERTY()
+	struct FRawCurveTracks RawCurveData;
+
 #if WITH_EDITORONLY_DATA
 	// if you change Notifies array, this will need to be rebuilt
 	UPROPERTY()
@@ -144,7 +145,7 @@ public:
 	virtual int32 GetNumberOfSampledKeys() const;
 
 	/** Return rate at which the animation is sampled **/
-	virtual FFrameRate GetSamplingFrameRate() const;
+	virtual const FFrameRate& GetSamplingFrameRate() const;
 
 #if WITH_EDITOR
 	/** Get the frame number for the provided time */
@@ -237,6 +238,10 @@ public:
 	virtual void EnableRootMotionSettingFromMontage(bool bInEnableRootMotion, const ERootMotionRootLock::Type InRootMotionRootLock) {};
 
 #if WITH_EDITOR
+	// Store that our raw data has changed so that we can get correct compressed data later on
+	UE_DEPRECATED(5.0, "MarkRawDataAsModified has been deprecated, any (Raw Data) modification should be applied using the UAnimDataController API instead. This will handle updating the GUID instead.")
+	virtual void MarkRawDataAsModified(bool bForceNewRawDatGuid = true) {}
+
 private:
 	DECLARE_MULTICAST_DELEGATE( FOnNotifyChangedMulticaster );
 	FOnNotifyChangedMulticaster OnNotifyChanged;
@@ -264,24 +269,18 @@ protected:
 		}
 	}
 #if WITH_EDITOR
-	void OnEndLoadPackage(const FEndLoadPackageContext& Context);
-	virtual void OnAnimModelLoaded();
 public:
-	/** Returns the IAnimationDataModel object embedded in this UAnimSequenceBase */
-	IAnimationDataModel* GetDataModel() const;
-
-	/** Returns the IAnimationDataModel as a script-interface, provides access to UObject and Interface */
-	TScriptInterface<IAnimationDataModel> GetDataModelInterface() const;
+	/** Returns the UAnimDataModel object embedded in this UAnimSequenceBase */
+	UAnimDataModel* GetDataModel() const;
 
 	/** Returns the transient UAnimDataController set to operate on DataModel */
 	IAnimationDataController& GetController();
 protected:
 	/** Populates the UAnimDataModel object according to any pre-existing data. (overrides expect to populate the model according to their data) */
 	virtual void PopulateModel();
-	virtual void PopulateWithExistingModel(TScriptInterface<IAnimationDataModel> ExistingDataModel);
 
 	/** Callback registered to UAnimDatModel::GetModifiedEvent for the embedded object */
-	virtual void OnModelModified(const EAnimDataModelNotifyType& NotifyType, IAnimationDataModel* Model, const FAnimDataModelNotifPayload& Payload);
+	virtual void OnModelModified(const EAnimDataModelNotifyType& NotifyType, UAnimDataModel* Model, const FAnimDataModelNotifPayload& Payload);
 
 	/** Validates that DataModel contains a valid UAnimDataModel object */
 	void ValidateModel() const;
@@ -290,7 +289,7 @@ protected:
 	void BindToModelModificationEvent();
 
 	/** Replaces the current DataModel, if any, with the provided one */
-	void CopyDataModel(const TScriptInterface<IAnimationDataModel>& ModelToDuplicate);
+	void CopyDataModel(const UAnimDataModel* ModelToDuplicate);
 private:
 	/** Creates a new UAnimDataModel instance and sets DataModel accordingly */
 	void CreateModel();
@@ -302,7 +301,7 @@ public:
 		if(ShouldDataModelBeValid())
 		{
 			ValidateModel();
-			return DataModelInterface != nullptr;
+			return DataModel != nullptr;
 		}
 
 		return false;
@@ -311,14 +310,10 @@ public:
 
 #if WITH_EDITORONLY_DATA
 protected:
-	UE_DEPRECATED(5.1, "DataModel has been converted to DataModelInterface")
+	/** UAnimDataModel instance containing source animation data */
 	UPROPERTY(VisibleAnywhere, BlueprintReadOnly, Category = "Animation Model")
 	TObjectPtr<UAnimDataModel> DataModel;
 
-	/** IAnimationDataModel instance containing (source) animation data */
-	UPROPERTY(VisibleAnywhere, BlueprintReadOnly, Category = "Animation Model")
-	TScriptInterface<IAnimationDataModel> DataModelInterface;
-
 	/** Flag set whenever the data-model is initially populated (during upgrade path) */
 	bool bPopulatingDataModel;
 
diff --git a/Engine/Source/Runtime/Engine/Classes/Animation/AnimSequenceHelpers.h b/Engine/Source/Runtime/Engine/Classes/Animation/AnimSequenceHelpers.h
--- a/Engine/Source/Runtime/Engine/Classes/Animation/AnimSequenceHelpers.h
+++ b/Engine/Source/Runtime/Engine/Classes/Animation/AnimSequenceHelpers.h
@@ -8,7 +8,7 @@
 #include "Delegates/IDelegateInstance.h"
 
 class UAnimSequence;
-class IAnimationDataModel;
+class UAnimDataModel;
 class UAnimDataController;
 class IAnimationDataController;
 class FName;
@@ -17,7 +17,6 @@ struct FRawCurveTracks;
 struct FRawAnimSequenceTrack;
 struct FBoneAnimationTrack;
 struct FRawAnimSequenceTrack;
-struct FAnimationPoseData;
 class IAnimationDataController;
 
 enum class EAnimInterpolationType : uint8;
@@ -28,7 +27,7 @@ namespace Anim {
 
 #if WITH_EDITOR
 	/**
-	* Populates a FCompactPose according to the data stored within a IAnimationDataModel provided a time to sample.
+	* Populates a FCompactPose according to the data stored within a UAnimDataModel provided a time to sample.
 	*
 	* @param	Model				Model to sample from
 	* @param	OutPose				[out] Pose to be populated using the model
@@ -37,42 +36,29 @@ namespace Anim {
 	* @param	RetargetSource		Name of the retarget source to retrieve from the skeleton
 	* @param	RetargetTransforms	Base transforms to use when retargetting the pose
 	*/
-	UE_DEPRECATED(5.1, "Use BuildPoseFromModel with different signature instead")
-	ENGINE_API void BuildPoseFromModel(const IAnimationDataModel* Model, FCompactPose& OutPose, const float Time, const EAnimInterpolationType& InterpolationType, const FName& RetargetSource, const TArray<FTransform>& RetargetTransforms);
+	ENGINE_API void BuildPoseFromModel(const UAnimDataModel* Model, FCompactPose& OutPose, const float Time, const EAnimInterpolationType& InterpolationType, const FName& RetargetSource, const TArray<FTransform>& RetargetTransforms);
 
 	/**
-	* Populates a FAnimationPoseData according to the data stored within a IAnimationDataModel provided a time to sample.
-	*
-	* @param	Model				Model to sample from
-	* @param	OutPoseData			[out] Pose data to be populated using the model
-	* @param	Time				Time value to sample at
-	* @param	InterpolationType	Method to be used to interpolate between different sample points
-	* @param	RetargetSource		Name of the retarget source to retrieve from the skeleton
-	* @param	RetargetTransforms	Base transforms to use when retargetting the pose
-	*/
-	ENGINE_API void BuildPoseFromModel(const IAnimationDataModel* Model, FAnimationPoseData& OutPoseData, const double Time, const EAnimInterpolationType& InterpolationType, const FName& RetargetSource, const TArray<FTransform>& RetargetTransforms);
-
-	/**
-	* Populates float curves according to the data stored within a IAnimationDataModel provided a time to sample.
+	* Populates float curves according to the data stored within a UAnimDataModel provided a time to sample.
 	*
 	* @param	Model				Model to sample from
 	* @param	OutCurves			[out] Curves to be populated using the model
 	* @param	Time				Time value to sample at
 	*/
-	ENGINE_API void EvaluateFloatCurvesFromModel(const IAnimationDataModel* Model, FBlendedCurve& OutCurves, double Time);
+	ENGINE_API void EvaluateFloatCurvesFromModel(const UAnimDataModel* Model, FBlendedCurve& OutCurves, float Time);
 
 	/**
-	* Populates transform curves according to the data stored within a IAnimationDataModel provided a time to sample.
+	* Populates transform curves according to the data stored within a UAnimDataModel provided a time to sample.
 	*
 	* @param	Model				Model to sample from
 	* @param	OutCurves			[out] Curves which were sampled from the model, key-ed with their equivalent name
 	* @param	Time				Time value to sample at
 	* @param	BlendWeight			Blending weight to apply when sampling
 	*/
-	ENGINE_API void EvaluateTransformCurvesFromModel(const IAnimationDataModel* Model, TMap<FName, FTransform>& OutCurves, double Time, float BlendWeight);
+	ENGINE_API void EvaluateTransformCurvesFromModel(const UAnimDataModel* Model, TMap<FName, FTransform>& OutCurves, float Time, float BlendWeight);
 
 	/**
-	* Retrieves a single bone transform according to the data stored within a IAnimationDataModel provided a time and track to sample.
+	* Retrieves a single bone transform according to the data stored within a UAnimDataModel provided a time and track to sample.
 	*
 	* @param	Model				Model to sample from
 	* @param	OutTransform		[out] Transform which was sampled from the model
@@ -80,20 +66,20 @@ namespace Anim {
 	* @param	Time				Time value to sample at
 	* @param	InterpolationType	Method to be used to interpolate between different sample points
 	*/
-	ENGINE_API void GetBoneTransformFromModel(const IAnimationDataModel* Model, FTransform& OutTransform, int32 TrackIndex, double Time, const EAnimInterpolationType& Interpolation);
+	ENGINE_API void GetBoneTransformFromModel(const UAnimDataModel* Model, FTransform& OutTransform, int32 TrackIndex, float Time, const EAnimInterpolationType& Interpolation);
 	
 	/**
-	* Retrieves a single bone transform according to the data stored within a IAnimationDataModel provided a key and track to sample.
+	* Retrieves a single bone transform according to the data stored within a UAnimDataModel provided a key and track to sample.
 	*
 	* @param	Model				Model to sample from
 	* @param	OutTransform		[out] Transform which was sampled from the model
 	* @param	TrackIndex			Track / bone index which should be sampld
 	* @param	KeyIndex			Specific key index to sample
 	*/
-	ENGINE_API void GetBoneTransformFromModel(const IAnimationDataModel* Model, FTransform& OutTransform, int32 TrackIndex, int32 KeyIndex);
+	ENGINE_API void GetBoneTransformFromModel(const UAnimDataModel* Model, FTransform& OutTransform, int32 TrackIndex, int32 KeyIndex);
 
 	/**
-	* Copies over any individual curve from FRawCurveTracks to a IAnimationDataModel instance targeted by the provided controller.
+	* Copies over any individual curve from FRawCurveTracks to a UAnimDataModel instance targeted by the provided controller.
 	*
 	* @param	CurveData			Container of the curves to be copied over
 	* @param	Skeleton			Skeleton to use for verifying and or add the curve names
@@ -102,7 +88,7 @@ namespace Anim {
 	ENGINE_API void CopyCurveDataToModel(const FRawCurveTracks& CurveData, const USkeleton* Skeleton, IAnimationDataController& Controller);
 
 	/**
-	* Copies over any individual curve from FRawCurveTracks to a IAnimationDataModel instance targeted by the provided controller.
+	* Copies over any individual curve from FRawCurveTracks to a UAnimDataModel instance targeted by the provided controller.
 	*
 	* @param	SourceAnimSeq				Animation Sequence to copy the notifies from
 	* @param	DestAnimSeq					Animation Sequence to copy the notifies to
@@ -226,30 +212,8 @@ namespace Anim {
 #endif // WITH_EDITOR
 	}
 
-	namespace DataModel { struct FEvaluationContext; }
 	namespace Retargeting
 	{
-		struct ENGINE_API FRetargetTracking
-		{
-			const FCompactPoseBoneIndex PoseBoneIndex;
-			const int32 SkeletonBoneIndex;
-
-			FRetargetTracking(const FCompactPoseBoneIndex InPoseBoneIndex, const int32 InSkeletonBoneIndex) : PoseBoneIndex(InPoseBoneIndex), SkeletonBoneIndex(InSkeletonBoneIndex) {}
-		};
-
-		struct ENGINE_API FRetargetingScope
-		{
-			FRetargetingScope(FCompactPose& ToRetargetPose, const DataModel::FEvaluationContext& InEvaluationContext);
-			~FRetargetingScope();
-			
-			void AddTrackedBone(FCompactPoseBoneIndex CompactBoneIndex, int32 SkeletonBoneIndex) const;
-		private:
-			FCompactPose& RetargetPose;
-			const DataModel::FEvaluationContext& EvaluationContext;
-			TArray<FRetargetTracking>& RetargetTracking;
-			bool bShouldRetarget;
-		};
-		
 		/**
 		* Retargeting the provided pose using the retarget source or transforms
 		*
@@ -257,16 +221,9 @@ namespace Anim {
 		* @param	RetargetSource		Name of the retarget source to retrieve data for
 		* @param	RetargetTransforms	Set of transforms to use as the ref skeleton bone transforms
 		*/
-		UE_DEPRECATED(5.1, "UE::Anim::Retargeting::RetargetPose has been deprecated as retargeting while evaluating pose data from a DataModel is now handled/performed implicitly")
 		ENGINE_API void RetargetPose(FCompactPose& InOutPose, const FName& RetargetSource, const TArray<FTransform>& RetargetTransforms);
 	}
 
-	struct ENGINE_API FBuildRawPoseScratchArea : public TThreadSingleton<FBuildRawPoseScratchArea>
-	{
-		TArray<Retargeting::FRetargetTracking> RetargetTracking;
-		TArray<FVirtualBoneCompactPoseData> VirtualBoneCompactPoseData;
-	};
-
 } // namespace Anim	
 
 } // namespace UE
\ No newline at end of file
diff --git a/Engine/Source/Runtime/Engine/Classes/Animation/AnimStreamable.h b/Engine/Source/Runtime/Engine/Classes/Animation/AnimStreamable.h
--- a/Engine/Source/Runtime/Engine/Classes/Animation/AnimStreamable.h
+++ b/Engine/Source/Runtime/Engine/Classes/Animation/AnimStreamable.h
@@ -38,7 +38,6 @@ public:
 	float StartTime;
 
 	float SequenceLength;
-	int32 NumFrames;
 
 	// Compressed Data for this chunk (if nullptr then data needs to be loaded via BulkData)
 	FCompressedAnimSequence* CompressedAnimSequence;
@@ -104,9 +103,6 @@ public:
 	UPROPERTY(EditAnywhere, AssetRegistrySearchable, Category = Animation)
 	FName RetargetSource;
 
-	UPROPERTY()
-	FFrameRate SamplingFrameRate;
-
 #if WITH_EDITORONLY_DATA
 
 	// Sequence the streamable was created from (used for reflecting changes to the source in editor)
diff --git a/Engine/Source/Runtime/Engine/Classes/Animation/AnimationAsset.h b/Engine/Source/Runtime/Engine/Classes/Animation/AnimationAsset.h
--- a/Engine/Source/Runtime/Engine/Classes/Animation/AnimationAsset.h
+++ b/Engine/Source/Runtime/Engine/Classes/Animation/AnimationAsset.h
@@ -245,7 +245,7 @@ struct FPoseCurve
 struct FAnimExtractContext
 {
 	/** Position in animation to extract pose from */
-	double CurrentTime;
+	float CurrentTime;
 
 	/** Is root motion being extracted? */
 	bool bExtractRootMotion;
@@ -268,19 +268,8 @@ struct FAnimExtractContext
 	 * by several animation nodes to optimize evaluation time.
 	 */
 	TArray<bool> BonesRequired;
-	
-	UE_DEPRECATED(5.1, "FAnimExtractContext construct with float-based time value is deprecated, use other signature")
-	FAnimExtractContext(float InCurrentTime, bool InbExtractRootMotion = false, FDeltaTimeRecord InDeltaTimeRecord = {}, bool InbLooping = false)
-		: CurrentTime((double)InCurrentTime)
-		, bExtractRootMotion(InbExtractRootMotion)
-		, DeltaTimeRecord(InDeltaTimeRecord)
-		, bLooping(InbLooping)
-		, PoseCurves()
-		, BonesRequired()
-	{
-	}
 
-	FAnimExtractContext(double InCurrentTime = 0.0, bool InbExtractRootMotion = false, FDeltaTimeRecord InDeltaTimeRecord = {}, bool InbLooping = false)
+	FAnimExtractContext(float InCurrentTime = 0.f, bool InbExtractRootMotion = false, FDeltaTimeRecord InDeltaTimeRecord = {}, bool InbLooping = false)
 		: CurrentTime(InCurrentTime)
 		, bExtractRootMotion(InbExtractRootMotion)
 		, DeltaTimeRecord(InDeltaTimeRecord)
@@ -1102,7 +1091,7 @@ protected:
 
 public:
 	/** Return a list of unique marker names for blending compatibility */
-	virtual TArray<FName>* GetUniqueMarkerNames() { return nullptr; }
+	virtual TArray<FName>* GetUniqueMarkerNames() { return NULL; }
 
 	//~ Begin IInterface_AssetUserData Interface
 	virtual void AddAssetUserData(UAssetUserData* InUserData) override;
diff --git a/Engine/Source/Runtime/Engine/Classes/Animation/AnimationSettings.h b/Engine/Source/Runtime/Engine/Classes/Animation/AnimationSettings.h
--- a/Engine/Source/Runtime/Engine/Classes/Animation/AnimationSettings.h
+++ b/Engine/Source/Runtime/Engine/Classes/Animation/AnimationSettings.h
@@ -95,15 +95,9 @@ class ENGINE_API UAnimationSettings : public UDeveloperSettings
 	UPROPERTY(config, EditAnywhere, Category = Mirroring)
 	TArray<FMirrorFindReplaceExpression> MirrorFindReplaceExpressions;
 
-	/** Project specific default frame-rate used when (re)initializing any animation based data */
-	UPROPERTY(config, EditAnywhere, Category = AnimationData)
-	FFrameRate DefaultFrameRate;
 public:
 	static UAnimationSettings * Get() { return CastChecked<UAnimationSettings>(UAnimationSettings::StaticClass()->GetDefaultObject()); }
 
-	/** Returns the project specific default frame-rate */
-	const FFrameRate& GetDefaultFrameRate() const;
-
 #if WITH_EDITOR
 	virtual void PostEditChangeProperty(struct FPropertyChangedEvent& PropertyChangedEvent) override;
 #endif // WITH_EDITOR
diff --git a/Engine/Source/Runtime/Engine/Classes/Animation/AttributeCurve.h b/Engine/Source/Runtime/Engine/Classes/Animation/AttributeCurve.h
--- a/Engine/Source/Runtime/Engine/Classes/Animation/AttributeCurve.h
+++ b/Engine/Source/Runtime/Engine/Classes/Animation/AttributeCurve.h
@@ -189,6 +189,7 @@ protected:
 	const UE::Anim::IAttributeBlendOperator* Operator;
 
 	friend class UAnimSequence;
+	friend class UAnimDataController;
 	friend struct UE::Anim::Attributes;
 };
 
diff --git a/Engine/Source/Runtime/Engine/Private/Animation/AnimComposite.cpp b/Engine/Source/Runtime/Engine/Private/Animation/AnimComposite.cpp
--- a/Engine/Source/Runtime/Engine/Private/Animation/AnimComposite.cpp
+++ b/Engine/Source/Runtime/Engine/Private/Animation/AnimComposite.cpp
@@ -144,10 +144,10 @@ bool UAnimComposite::ContainRecursive(TArray<UAnimCompositeBase*>& CurrentAccumu
 void UAnimComposite::SetCompositeLength(float InLength)
 {
 #if WITH_EDITOR		
-	Controller->SetNumberOfFrames(DataModelInterface->GetFrameRate().AsFrameNumber(InLength));
+	Controller->SetPlayLength(InLength);
 #else
 	PRAGMA_DISABLE_DEPRECATION_WARNINGS
-	SequenceLength = InLength;
+	SetSequenceLength(InLength);
 	PRAGMA_ENABLE_DEPRECATION_WARNINGS
 #endif	
 }
diff --git a/Engine/Source/Runtime/Engine/Private/Animation/AnimCompositeBase.cpp b/Engine/Source/Runtime/Engine/Private/Animation/AnimCompositeBase.cpp
--- a/Engine/Source/Runtime/Engine/Private/Animation/AnimCompositeBase.cpp
+++ b/Engine/Source/Runtime/Engine/Private/Animation/AnimCompositeBase.cpp
@@ -16,7 +16,7 @@
 // FAnimSegment
 ///////////////////////////////////////////////////////
 
-UAnimSequenceBase* FAnimSegment::GetAnimationData(float PositionInTrack, float& PositionInAnim) const
+UAnimSequenceBase * FAnimSegment::GetAnimationData(float PositionInTrack, float& PositionInAnim) const
 {
 	if( bValid && IsInRange(PositionInTrack) )
 	{
@@ -569,11 +569,9 @@ void FAnimTrack::GetAnimationPose(FAnimationPoseData& OutAnimationPoseData, cons
 		if (AnimSegment->bValid)
 		{
 			// Copy passed in Extraction Context, but override position and root motion parameters.
-			float PositionInAnim = 0.f;
-			if (const UAnimSequenceBase* const AnimRef = AnimSegment->GetAnimationData(ClampedTime, PositionInAnim))
+			FAnimExtractContext SequenceExtractionContext(ExtractionContext);
+			if (const UAnimSequenceBase* const AnimRef = AnimSegment->GetAnimationData(ClampedTime, SequenceExtractionContext.CurrentTime))
 			{
-				FAnimExtractContext SequenceExtractionContext(ExtractionContext);
-				SequenceExtractionContext.CurrentTime = static_cast<double>(PositionInAnim);
 				SequenceExtractionContext.DeltaTimeRecord.SetPrevious(
 					SequenceExtractionContext.CurrentTime - SequenceExtractionContext.DeltaTimeRecord.Delta);
 				SequenceExtractionContext.bExtractRootMotion &= AnimRef->HasRootMotion();
@@ -768,13 +766,6 @@ void UAnimCompositeBase::ExtractRootMotionFromTrack(const FAnimTrack &SlotAnimTr
 	}
 }
 
-FFrameRate UAnimCompositeBase::GetSamplingFrameRate() const
-{
-	// Allowing for 0.00001s precision in composite/montage length
-	static const FFrameRate CompositeFrameRate(100000, 1);
-	return CompositeFrameRate;
-}
-
 void UAnimCompositeBase::PostLoad()
 {
 	Super::PostLoad();
@@ -826,17 +817,11 @@ void FAnimSegment::UpdateCachedPlayLength()
 {
 	CachedPlayLength = 0.f;
 	PRAGMA_DISABLE_DEPRECATION_WARNINGS
-	const IAnimationDataModel* DataModel = AnimReference ? AnimReference->GetDataModel() : nullptr;	
+	const UAnimDataModel* DataModel = AnimReference ? AnimReference->GetDataModel() : nullptr;	
 	PRAGMA_ENABLE_DEPRECATION_WARNINGS
 	if(DataModel)
 	{
 		CachedPlayLength = DataModel->GetPlayLength();
 	}
 }
-
-void UAnimCompositeBase::PopulateWithExistingModel(TScriptInterface<IAnimationDataModel> ExistingDataModel)
-{
-	Super::PopulateWithExistingModel(ExistingDataModel);
-	Controller->SetFrameRate(GetSamplingFrameRate());
-}
 #endif // WITH_EDITOR
\ No newline at end of file
