[Backout] - CL20593017, 20594056, 20625634, 20645670, 20649179, 20649223, 20649380, 20658535

diff --git a/Engine/Source/Runtime/Engine/Private/Animation/AnimCompress_RemoveLinearKeys.cpp b/Engine/Source/Runtime/Engine/Private/Animation/AnimCompress_RemoveLinearKeys.cpp
--- a/Engine/Source/Runtime/Engine/Private/Animation/AnimCompress_RemoveLinearKeys.cpp
+++ b/Engine/Source/Runtime/Engine/Private/Animation/AnimCompress_RemoveLinearKeys.cpp
@@ -304,19 +304,20 @@ void UAnimCompress_RemoveLinearKeys::UpdateWorldBoneTransformTable(
 {
 	const FBoneData& Bone		= CompressibleAnimData.BoneData[BoneIndex];
 	const int32 NumKeys		= CompressibleAnimData.NumberOfKeys;
+	const float SequenceLength	= CompressibleAnimData.SequenceLength;
 	const int32 FrameStart		= (BoneIndex*NumKeys);
 	const int32 TrackIndex = FAnimationUtils::GetAnimTrackIndexForSkeletonBone(BoneIndex, CompressibleAnimData.TrackToSkeletonMapTable);
 	
 	check(OutputWorldBones.Num() >= (FrameStart+NumKeys));
 
-	const FFrameRate& SamplingRate = CompressibleAnimData.SampledFrameRate;
+	const float TimePerFrame = SequenceLength / (float)(NumKeys-1);
 
 	if( TrackIndex != INDEX_NONE )
 	{
 		// get the local-space bone transforms using the animation solver
 		for ( int32 KeyIndex = 0; KeyIndex < NumKeys; ++KeyIndex )
 		{
-			const double Time = SamplingRate.AsSeconds(KeyIndex);
+			float Time = (float)KeyIndex * TimePerFrame;
 			FTransform LocalAtom;
 
 			FAnimationUtils::ExtractTransformFromCompressionData(CompressibleAnimData, OutCompressedData, Time, TrackIndex, UseRaw, LocalAtom);
@@ -416,10 +417,9 @@ void UAnimCompress_RemoveLinearKeys::UpdateBoneAtomList(
 	TArray<FTransform>& BoneAtoms)
 {
 	BoneAtoms.Reset(NumFrames);
-	const FFrameRate& SamplingRate = CompressibleAnimData.SampledFrameRate;
 	for ( int32 FrameIndex = 0; FrameIndex < NumFrames; ++FrameIndex )
 	{
-		const double Time = SamplingRate.AsSeconds(FrameIndex);
+		float Time = (float)FrameIndex * TimePerFrame;
 		FTransform LocalAtom;
 		FAnimationUtils::ExtractTransformFromCompressionData(CompressibleAnimData, OutCompressedData, Time, TrackIndex, false, LocalAtom);
 
@@ -465,10 +465,11 @@ void UAnimCompress_RemoveLinearKeys::ConvertFromRelativeSpace(FCompressibleAnimD
 		// make sure scale key exists
 		if (RawTrack.ScaleKeys.Num() > 0)
 		{
-			const FVector3f& RefBoneScale = (BasePoseTrack.ScaleKeys.Num() > 0)? BasePoseTrack.ScaleKeys[0] : FVector3f::OneVector;
+			const FVector3f DefaultScale(1.f);
+			const FVector3f& RefBoneScale = (BasePoseTrack.ScaleKeys.Num() > 0)? BasePoseTrack.ScaleKeys[0] : DefaultScale;
 			for (int32 ScaleIndex = 0; ScaleIndex < RawTrack.ScaleKeys.Num(); ++ScaleIndex)
 			{
-				RawTrack.ScaleKeys[ScaleIndex] = RefBoneScale * (FVector3f::OneVector + RawTrack.ScaleKeys[ScaleIndex]);
+				RawTrack.ScaleKeys[ScaleIndex] = RefBoneScale * (DefaultScale + RawTrack.ScaleKeys[ScaleIndex]);
 			}
 		}
 	}
@@ -514,8 +515,7 @@ void UAnimCompress_RemoveLinearKeys::ConvertToRelativeSpace(FCompressibleAnimDat
 		// scale key
 		if (RawTrack.ScaleKeys.Num() > 0)
 		{
-        	const FVector3f& RefBoneScale = (BasePoseTrack.ScaleKeys.Num() > 0)? BasePoseTrack.ScaleKeys[0] : FVector3f::OneVector;
-			const FVector3f InvRefBoneScale = (FVector3f)FTransform::GetSafeScaleReciprocal((FVector)RefBoneScale);
+			const FVector3f InvRefBoneScale = (FVector3f)FTransform::GetSafeScaleReciprocal((FVector)BasePoseTrack.ScaleKeys[0]);
 
 			// transform scale keys.
 			for (int32 ScaleIndex = 0; ScaleIndex < RawTrack.ScaleKeys.Num(); ++ScaleIndex)
@@ -563,8 +563,7 @@ void UAnimCompress_RemoveLinearKeys::ConvertToRelativeSpace(
 		// scale key
 		if (ScaleData.Num() > 0)
 		{
-        	const FVector3f& RefBoneScale = (BasePoseTrack.ScaleKeys.Num() > 0)? BasePoseTrack.ScaleKeys[0] : FVector3f::OneVector;
-			const FVector3f InvRefBoneScale = (FVector3f)FTransform::GetSafeScaleReciprocal((FVector)RefBoneScale);
+			const FVector3f InvRefBoneScale = (FVector3f)FTransform::GetSafeScaleReciprocal((FVector)BasePoseTrack.ScaleKeys[0]);
 
 			// convert the new scale tracks to additive space
 			FScaleTrack& ScaleTrack = ScaleData[TrackIndex];
@@ -880,7 +879,7 @@ void UAnimCompress_RemoveLinearKeys::ProcessAnimationTracks(
 			}
 
 			// Determine if a guidance track should be used to aid in choosing keys to retain
-			TArray<float>* GuidanceTrack = nullptr;
+			TArray<float>* GuidanceTrack = NULL;
 			float GuidanceScale = 1.0f;
 			if (GuideTrackIndex != INDEX_NONE)
 			{
diff --git a/Engine/Source/Runtime/Engine/Private/Animation/AnimCompressionDerivedData.cpp b/Engine/Source/Runtime/Engine/Private/Animation/AnimCompressionDerivedData.cpp
--- a/Engine/Source/Runtime/Engine/Private/Animation/AnimCompressionDerivedData.cpp
+++ b/Engine/Source/Runtime/Engine/Private/Animation/AnimCompressionDerivedData.cpp
@@ -35,7 +35,7 @@ const TCHAR* FDerivedDataAnimationCompression::GetVersionString() const
 	// This is a version string that mimics the old versioning scheme. If you
 	// want to bump this version, generate a new guid using VS->Tools->Create GUID and
 	// return it here. Ex.
-	return TEXT("6CAE774196A34FEABF46EF3B30EF6DDC");
+	return TEXT("3A3F9FBB7A4047278CABE35820CC44D7");
 }
 
 bool FDerivedDataAnimationCompression::Build( TArray<uint8>& OutDataArray )
diff --git a/Engine/Source/Runtime/Engine/Private/Animation/AnimCurveCompressionCodec_CompressedRichCurve.cpp b/Engine/Source/Runtime/Engine/Private/Animation/AnimCurveCompressionCodec_CompressedRichCurve.cpp
--- a/Engine/Source/Runtime/Engine/Private/Animation/AnimCurveCompressionCodec_CompressedRichCurve.cpp
+++ b/Engine/Source/Runtime/Engine/Private/Animation/AnimCurveCompressionCodec_CompressedRichCurve.cpp
@@ -45,8 +45,8 @@ bool UAnimCurveCompressionCodec_CompressedRichCurve::Compress(const FCompressibl
 	int32 KeyDataOffset = 0;
 	KeyDataOffset += sizeof(FCurveDesc) * NumCurves;
 
-	const FFrameRate& SamplingFrameRate = AnimSeq.SampledFrameRate;
-	const double SampleRate = AnimSeq.SampledFrameRate.AsDecimal();
+	const FAnimKeyHelper Helper(AnimSeq.SequenceLength, AnimSeq.NumberOfKeys);
+	const float SampleRate = UseAnimSequenceSampleRate ? Helper.KeysPerSecond() : ErrorSampleRate;
 
 	TArray<uint8> KeyData;
 
@@ -78,15 +78,15 @@ bool UAnimCurveCompressionCodec_CompressedRichCurve::Compress(const FCompressibl
 			for (int32 KeyIndex = 0; KeyIndex < SamplesToTake; ++KeyIndex)
 			{
 				// Evaluated RawCurve and Compressed Curve
-				const double EvalTime = SamplingFrameRate.AsSeconds(KeyIndex);
-				const float RawValue = Curve.FloatCurve.Eval(EvalTime);
+				const float EvalTime = KeyIndex * Helper.TimePerKey();
+				const float RawValue = RawCurve.Eval(EvalTime);
 				const float CompressedValue = CompressedCurve.Eval(EvalTime);
 
 				const float AbsDelta = FMath::Abs(CompressedValue - RawValue);
-				if (!FMath::IsNearlyZero(AbsDelta, MaxCurveError + UE_KINDA_SMALL_NUMBER))
+				if (AbsDelta > MaxCurveError)
 				{
 					// Delta larger than tolerated error value
-					UE_LOG(LogAnimationCompression, Warning, TEXT("Curve %s: delta too large %f, between %f and %f, at %f"), *Curve.Name.DisplayName.ToString(), AbsDelta, RawValue, CompressedValue, EvalTime);
+					UE_LOG(LogAnimationCompression, Warning, TEXT("Curve %s: delta too large %f at %f"), *Curve.Name.DisplayName.ToString(), AbsDelta, EvalTime);
 				}
 			}
 		}
diff --git a/Engine/Source/Runtime/Engine/Private/Animation/AnimCurveCompressionCodec_UniformIndexable.cpp b/Engine/Source/Runtime/Engine/Private/Animation/AnimCurveCompressionCodec_UniformIndexable.cpp
--- a/Engine/Source/Runtime/Engine/Private/Animation/AnimCurveCompressionCodec_UniformIndexable.cpp
+++ b/Engine/Source/Runtime/Engine/Private/Animation/AnimCurveCompressionCodec_UniformIndexable.cpp
@@ -89,7 +89,11 @@ UAnimCurveCompressionCodec_UniformIndexable::UAnimCurveCompressionCodec_UniformI
 bool UAnimCurveCompressionCodec_UniformIndexable::Compress(const FCompressibleAnimData& AnimSeq, FAnimCurveCompressionResult& OutResult)
 {
 	const int32 NumCurves = AnimSeq.RawFloatCurves.Num();
-	const int32 NumSamples = AnimSeq.NumberOfKeys;
+	const float Duration = AnimSeq.SequenceLength;
+
+	const FAnimKeyHelper Helper(AnimSeq.SequenceLength, AnimSeq.NumberOfKeys);
+	const float SampleRate = Helper.KeysPerSecond();
+	const int32 NumSamples = FMath::RoundToInt(Duration * SampleRate) + 1;
 
 	int32 BufferSize = 0;
 	BufferSize += sizeof(int32);	// NumSamples
@@ -109,7 +113,7 @@ bool UAnimCurveCompressionCodec_UniformIndexable::Compress(const FCompressibleAn
 	BufferOffset += sizeof(float);
 
 	*NumSamplesPtr = NumSamples;
-	*SampleRatePtr = AnimSeq.SampledFrameRate.AsDecimal();
+	*SampleRatePtr = SampleRate;
 
 	if (NumCurves > 0 && NumSamples > 0)
 	{
@@ -119,12 +123,12 @@ bool UAnimCurveCompressionCodec_UniformIndexable::Compress(const FCompressibleAn
 		// Write out samples by curve so that we can more easily index curve keys
 
 		// Curve 0 Key 0, Curve 0 Key 1, Curve 0 Key N, Curve 1 Key 0, Curve 1 Key 1, Curve 1 Key N, Curve M Key 0, ...
+		const float InvSampleRate = 1.0f / SampleRate;
 		for (const FFloatCurve& Curve : AnimSeq.RawFloatCurves)
 		{
 			for (int32 SampleIndex = 0; SampleIndex < NumSamples; ++SampleIndex)
 			{
-				const double FrameSampleTime = AnimSeq.SampledFrameRate.AsSeconds(SampleIndex);
-				const float SampleTime = FMath::Clamp(FrameSampleTime, 0.0f, AnimSeq.SequenceLength);
+				const float SampleTime = FMath::Clamp(SampleIndex * InvSampleRate, 0.0f, Duration);
 				const float SampleValue = Curve.FloatCurve.Eval(SampleTime);
 
 				(*AnimatedSamplesPtr) = SampleValue;
diff --git a/Engine/Source/Runtime/Engine/Private/Animation/AnimCurveCompressionCodec_UniformlySampled.cpp b/Engine/Source/Runtime/Engine/Private/Animation/AnimCurveCompressionCodec_UniformlySampled.cpp
--- a/Engine/Source/Runtime/Engine/Private/Animation/AnimCurveCompressionCodec_UniformlySampled.cpp
+++ b/Engine/Source/Runtime/Engine/Private/Animation/AnimCurveCompressionCodec_UniformlySampled.cpp
@@ -23,8 +23,9 @@ bool UAnimCurveCompressionCodec_UniformlySampled::Compress(const FCompressibleAn
 	float SampleRate_;
 	if (UseAnimSequenceSampleRate)
 	{
-		NumSamples = AnimSeq.NumberOfKeys;
-		SampleRate_ = AnimSeq.SampledFrameRate.AsDecimal();
+		const FAnimKeyHelper Helper(AnimSeq.SequenceLength, AnimSeq.NumberOfKeys);
+		SampleRate_ = Helper.KeysPerSecond();
+		NumSamples = FMath::RoundToInt(Duration * SampleRate_) + 1;
 	}
 	else
 	{
diff --git a/Engine/Source/Runtime/Engine/Private/Animation/AnimDataModel.cpp b/Engine/Source/Runtime/Engine/Private/Animation/AnimDataModel.cpp
--- a/Engine/Source/Runtime/Engine/Private/Animation/AnimDataModel.cpp
+++ b/Engine/Source/Runtime/Engine/Private/Animation/AnimDataModel.cpp
@@ -1,50 +1,44 @@
 // Copyright Epic Games, Inc. All Rights Reserved.
 
 #include "Animation/AnimData/AnimDataModel.h"
-
-#include "AnimationRuntime.h"
-#include "AnimEncoding.h"
 #include "UObject/NameTypes.h"
 
 #include "Animation/AnimSequence.h"
 
 #include "Algo/Transform.h"
 #include "Algo/Accumulate.h"
-#include "Animation/AnimSequenceHelpers.h"
 #include "Animation/SmartName.h"
-#include "Animation/AttributesRuntime.h"
 #include "UObject/UE5MainStreamObjectVersion.h"
 #include "UObject/FortniteMainBranchObjectVersion.h"
 
-#if WITH_EDITOR
-#include "IAnimationDataControllerModule.h"
-#include "Modules/ModuleManager.h"
-#endif // WITH_EDITOR
-
-#define LOCTEXT_NAMESPACE "AnimDataModel"
-
 void UAnimDataModel::PostLoad()
 {
 	UObject::PostLoad();
 
+	if (GetLinkerCustomVersion(FUE5MainStreamObjectVersion::GUID) >= FUE5MainStreamObjectVersion::AnimationDataModelInterface_BackedOut &&
+		GetLinkerCustomVersion(FUE5MainStreamObjectVersion::GUID) < FUE5MainStreamObjectVersion::BackoutAnimationDataModelInterface)
+	{
+		UE_LOG(LogAnimation, Fatal, TEXT("This package was saved with a version that had to be backed out and is no longer able to be loaded."));
+	}
+
+
 	if (GetLinkerCustomVersion(FUE5MainStreamObjectVersion::GUID) < FUE5MainStreamObjectVersion::ForceUpdateAnimationAssetCurveTangents)
 	{
 		// Forcefully AutoSetTangents to fix-up any imported sequences pre the fix for flattening first/last key leave/arrive tangents
-		GetNotifier().Notify(EAnimDataModelNotifyType::BracketOpened);
+		Notify(EAnimDataModelNotifyType::BracketOpened);
 		for (FFloatCurve& FloatCurve : CurveData.FloatCurves)
 		{
 			FloatCurve.FloatCurve.AutoSetTangents();
 			FCurvePayload Payload;
 			Payload.Identifier = FAnimationCurveIdentifier(FloatCurve.Name.UID, ERawCurveTrackTypes::RCT_Float);
-			GetNotifier().Notify(EAnimDataModelNotifyType::CurveChanged, Payload);
+			Notify(EAnimDataModelNotifyType::CurveChanged, Payload);
 		}
-		GetNotifier().Notify(EAnimDataModelNotifyType::BracketClosed);
+		Notify(EAnimDataModelNotifyType::BracketClosed);
 	}
 
-	const bool bHasBoneTracks = BoneAnimationTracks.Num() > 0;
-
 	if (GetLinkerCustomVersion(FFortniteMainBranchObjectVersion::GUID) < FFortniteMainBranchObjectVersion::SingleFrameAndKeyAnimModel)
 	{
+		const bool bHasBoneTracks = BoneAnimationTracks.Num() > 0;
 		if (bHasBoneTracks)
 		{
 			// Number of keys was used directly rather than Max(Value,2), as a single _frame_ animation should always have two _keys_ 
@@ -65,47 +59,17 @@ void UAnimDataModel::PostLoad()
 					AddKey(BoneTrack.InternalTrackData.ScaleKeys);
 				}
 
-				GetNotifier().Notify(EAnimDataModelNotifyType::TrackChanged);
-			}
-		}
-	}
-	
-	if (GetLinkerCustomVersion(FUE5MainStreamObjectVersion::GUID) < FUE5MainStreamObjectVersion::AnimationDataModelInterface)
-	{		
-		const double NumFramesAsSeconds = FrameRate.AsSeconds(NumberOfFrames);
-		PRAGMA_DISABLE_DEPRECATION_WARNINGS
-		const float CachedPlayLength = PlayLength;
-		if (static_cast<float>(NumFramesAsSeconds) != PlayLength)
-		{
-			if (bHasBoneTracks)
-			{
-				// Number of keys was rounded up from sequence length, so set play length from stored number of bone keys 
-				const int32 ActualNumKeys = bHasBoneTracks ? FMath::Min(BoneAnimationTracks[0].InternalTrackData.PosKeys.Num(), NumberOfKeys) : NumberOfKeys;
-				if (ActualNumKeys != NumberOfKeys)
-				{
-					PlayLength = FrameRate.AsSeconds(ActualNumKeys - 1);
-					NumberOfFrames = ActualNumKeys - 1;
-					NumberOfKeys = ActualNumKeys;
-
-					IAnimationDataController::ReportObjectWarningf(this, LOCTEXT("AnimDataModelPlayLengthAlteredToBones", "Playable length {0} has been updated according to number of frames {1} and frame-rate {2}, new length: {3}"), FText::AsNumber(CachedPlayLength), FText::AsNumber(GetNumberOfFrames()), FrameRate.ToPrettyText(), FText::AsNumber(PlayLength));
-				}
+				Notify(EAnimDataModelNotifyType::TrackChanged);
 			}
-			else
-			{
-				// No way to determine 'correct' length so set current PlayLength value to # of frames * frame-rate
-				PlayLength = FrameRate.AsSeconds(NumberOfFrames);
-			}
-			GetNotifier().Notify(EAnimDataModelNotifyType::Populated);
 		}
-		PRAGMA_ENABLE_DEPRECATION_WARNINGS
-	}
+	}	
 }
 
 void UAnimDataModel::PostDuplicate(bool bDuplicateForPIE)
 {
 	UObject::PostDuplicate(bDuplicateForPIE);
 
-	GetNotifier().Notify(EAnimDataModelNotifyType::Populated);
+	Notify(EAnimDataModelNotifyType::Populated);
 }
 
 const TArray<FBoneAnimationTrack>& UAnimDataModel::GetBoneAnimationTracks() const
@@ -159,9 +123,9 @@ int32 UAnimDataModel::GetBoneTrackIndexByName(FName TrackName) const
 	return INDEX_NONE;
 }
 
-double UAnimDataModel::GetPlayLength() const
+float UAnimDataModel::GetPlayLength() const
 {
-	return FrameRate.AsSeconds(NumberOfFrames);
+	return PlayLength;
 }
 
 int32 UAnimDataModel::GetNumberOfFrames() const
@@ -174,7 +138,7 @@ int32 UAnimDataModel::GetNumberOfKeys() const
 	return NumberOfKeys;
 }
 
-FFrameRate UAnimDataModel::GetFrameRate() const
+const FFrameRate& UAnimDataModel::GetFrameRate() const
 {
 	return FrameRate;
 }
@@ -184,15 +148,14 @@ bool UAnimDataModel::IsValidBoneTrackIndex(int32 TrackIndex) const
 	return BoneAnimationTracks.IsValidIndex(TrackIndex);
 }
 
-int32 UAnimDataModel::GetNumBoneTracks() const
+const int32 UAnimDataModel::GetNumBoneTracks() const
 {
 	return BoneAnimationTracks.Num();
 }
 
 UAnimSequence* UAnimDataModel::GetAnimationSequence() const
 {
-	UObject* AnimObject = static_cast<UObject*>(FindObjectWithOuter(GetOuter(), UAnimSequence::StaticClass()));
-	return CastChecked<UAnimSequence>(AnimObject ? AnimObject : GetOuter());
+	return Cast<UAnimSequence>(GetOuter());
 }
 
 void UAnimDataModel::GetBoneTrackNames(TArray<FName>& OutNames) const
@@ -208,16 +171,6 @@ const FAnimationCurveData& UAnimDataModel::GetCurveData() const
 	return CurveData;
 }
 
-IAnimationDataModel::FModelNotifier& UAnimDataModel::GetNotifier()
-{
-	if (!Notifier)
-	{
-		Notifier.Reset(new IAnimationDataModel::FModelNotifier(this));
-	}
-	
-	return *Notifier.Get();
-}
-
 int32 UAnimDataModel::GetNumberOfTransformCurves() const
 {
 	return CurveData.TransformCurves.Num();
@@ -353,6 +306,7 @@ int32 UAnimDataModel::GetNumberOfAttributesForBoneIndex(const int32 BoneIndex) c
 
 void UAnimDataModel::GetAttributesForBone(const FName& BoneName, TArray<const FAnimatedBoneAttribute*>& OutBoneAttributes) const
 {
+	// Sum up total number of attributes with provided bone name
 	Algo::TransformIf(AnimatedBoneAttributes, OutBoneAttributes, [BoneName](const FAnimatedBoneAttribute& Attribute) -> bool
 	{
 		return Attribute.Identifier.GetBoneName() == BoneName;
@@ -431,7 +385,7 @@ FGuid UAnimDataModel::GenerateGuid() const
     {
     	UpdateWithData(Curve.Name.UID);
 
-    	auto UpdateWithComponent = [&UpdateWithFloatCurve](const FVectorCurve& VectorCurve)
+    	auto UpdateWithComponent = [&Sha, &UpdateWithFloatCurve](const FVectorCurve& VectorCurve)
     	{
     		for (int32 ChannelIndex = 0; ChannelIndex < 3; ++ChannelIndex)
     		{
@@ -459,252 +413,6 @@ FGuid UAnimDataModel::GenerateGuid() const
 	return Guid;
 }
 
-TScriptInterface<IAnimationDataController> UAnimDataModel::GetController()
-{
-	TScriptInterface<IAnimationDataController> Controller = nullptr;
-#if WITH_EDITOR
-	IAnimationDataControllerModule& ControllerModule = FModuleManager::Get().GetModuleChecked<IAnimationDataControllerModule>("AnimationDataController");
-	Controller = ControllerModule.GetController();
-	Controller->SetModel(this);
-#endif // WITH_EDITOR
-
-	return Controller;
-}
-
-#if WITH_EDITOR
-FTransform ExtractTransformForKeyIndex(int32 Key, const FRawAnimSequenceTrack& TrackToExtract)
-{
-	static const FVector DefaultScale3D = FVector(1.f);
-	const bool bHasScaleKey = TrackToExtract.ScaleKeys.Num() > 0;
-
-	const int32 PosKeyIndex = FMath::Min(Key, TrackToExtract.PosKeys.Num() - 1);
-	const int32 RotKeyIndex = FMath::Min(Key, TrackToExtract.RotKeys.Num() - 1);
-	if (bHasScaleKey)
-	{
-		const int32 ScaleKeyIndex = FMath::Min(Key, TrackToExtract.ScaleKeys.Num() - 1);
-		return FTransform(FQuat(TrackToExtract.RotKeys[RotKeyIndex]), FVector(TrackToExtract.PosKeys[PosKeyIndex]), FVector(TrackToExtract.ScaleKeys[ScaleKeyIndex]));
-	}
-	else
-	{
-		return FTransform(FQuat(TrackToExtract.RotKeys[RotKeyIndex]), FVector(TrackToExtract.PosKeys[PosKeyIndex]), DefaultScale3D);
-	}
-}
-
-template<bool bInterpolateT>
-void ExtractPose(const TArray<FBoneAnimationTrack>& BoneAnimationTracks, const TMap<FName, const FTransformCurve*>& OverrideBoneTransforms, FCompactPose& InOutPose, const int32 KeyIndex1, const int32 KeyIndex2, float Alpha, float TimePerKey, UE::Anim::Retargeting::FRetargetingScope& RetargetingScope)
-{
-	const int32 NumAnimationTracks = BoneAnimationTracks.Num();
-	const FBoneContainer& RequiredBones = InOutPose.GetBoneContainer();
-
-	TArray<FVirtualBoneCompactPoseData>& VBCompactPoseData = UE::Anim::FBuildRawPoseScratchArea::Get().VirtualBoneCompactPoseData;
-	VBCompactPoseData = RequiredBones.GetVirtualBoneCompactPoseData();
-
-	FCompactPose Key2Pose;
-	if (bInterpolateT)
-	{
-		Key2Pose.CopyBonesFrom(InOutPose);
-	}
-
-	for (const FBoneAnimationTrack& AnimationTrack : BoneAnimationTracks)
-	{
-		const int32 SkeletonBoneIndex = AnimationTrack.BoneTreeIndex;
-		// not sure it's safe to assume that SkeletonBoneIndex can never be INDEX_NONE
-		if ((SkeletonBoneIndex != INDEX_NONE) && (SkeletonBoneIndex < MAX_BONES))
-		{
-			const FCompactPoseBoneIndex PoseBoneIndex = RequiredBones.GetCompactPoseIndexFromSkeletonIndex(SkeletonBoneIndex);
-			if (PoseBoneIndex != INDEX_NONE)
-			{
-				for (int32 Idx = 0; Idx < VBCompactPoseData.Num(); ++Idx)
-				{
-					FVirtualBoneCompactPoseData& VB = VBCompactPoseData[Idx];
-					if (PoseBoneIndex == VB.VBIndex)
-					{
-						// Remove this bone as we have written data for it (false so we dont resize allocation)
-						VBCompactPoseData.RemoveAtSwap(Idx, 1, false);
-						break; //Modified TArray so must break here
-					}
-				}
-				
-				// extract animation
-				const FRawAnimSequenceTrack& TrackToExtract = AnimationTrack.InternalTrackData;
-				auto OverrideTransform = OverrideBoneTransforms.Find(AnimationTrack.Name);
-				{
-					// Bail out (with rather wacky data) if data is empty for some reason.
-					if (TrackToExtract.PosKeys.Num() == 0 || TrackToExtract.RotKeys.Num() == 0)
-					{
-						InOutPose[PoseBoneIndex].SetIdentity();
-
-						if (bInterpolateT)
-						{
-							Key2Pose[PoseBoneIndex].SetIdentity();
-						}
-					}
-					else
-					{
-						InOutPose[PoseBoneIndex] = ExtractTransformForKeyIndex(KeyIndex1, TrackToExtract);
-
-						if (bInterpolateT)
-						{
-							Key2Pose[PoseBoneIndex] = ExtractTransformForKeyIndex(KeyIndex2, TrackToExtract);
-						}
-					}
-				}
-
-				if (OverrideTransform)
-				{
-					{
-						const FTransform PoseOneAdditive = (*OverrideTransform)->Evaluate(KeyIndex1 * TimePerKey, 1.f);
-						const FTransform PoseOneLocalTransform = InOutPose[PoseBoneIndex];
-						InOutPose[PoseBoneIndex].SetRotation(PoseOneLocalTransform.GetRotation() * PoseOneAdditive.GetRotation());
-						InOutPose[PoseBoneIndex].SetTranslation(PoseOneLocalTransform.TransformPosition(PoseOneAdditive.GetTranslation()));
-						InOutPose[PoseBoneIndex].SetScale3D(PoseOneLocalTransform.GetScale3D() * PoseOneAdditive.GetScale3D());
-					}
-					
-
-					if (bInterpolateT)
-					{
-						const FTransform PoseTwoAdditive = (*OverrideTransform)->Evaluate(KeyIndex2 * TimePerKey, 1.f);
-						const FTransform PoseTwoLocalTransform = Key2Pose[PoseBoneIndex];						
-						Key2Pose[PoseBoneIndex].SetRotation(PoseTwoLocalTransform.GetRotation() * PoseTwoAdditive.GetRotation());
-						Key2Pose[PoseBoneIndex].SetTranslation(PoseTwoLocalTransform.TransformPosition(PoseTwoAdditive.GetTranslation()));
-						Key2Pose[PoseBoneIndex].SetScale3D(PoseTwoLocalTransform.GetScale3D() * PoseTwoAdditive.GetScale3D());
-					}
-				}
-				RetargetingScope.AddTrackedBone(PoseBoneIndex, SkeletonBoneIndex);
-			}
-		}
-	}
-
-	//Build Virtual Bones
-	if (VBCompactPoseData.Num() > 0)
-	{
-		FCSPose<FCompactPose> CSPose1;
-		CSPose1.InitPose(InOutPose);
-
-		FCSPose<FCompactPose> CSPose2;
-		if (bInterpolateT)
-		{
-			CSPose2.InitPose(Key2Pose);
-		}
-
-		for (FVirtualBoneCompactPoseData& VB : VBCompactPoseData)
-		{
-			FTransform Source = CSPose1.GetComponentSpaceTransform(VB.SourceIndex);
-			FTransform Target = CSPose1.GetComponentSpaceTransform(VB.TargetIndex);
-			InOutPose[VB.VBIndex] = Target.GetRelativeTransform(Source);
-
-			if (bInterpolateT)
-			{
-				FTransform Source2 = CSPose2.GetComponentSpaceTransform(VB.SourceIndex);
-				FTransform Target2 = CSPose2.GetComponentSpaceTransform(VB.TargetIndex);
-				Key2Pose[VB.VBIndex] = Target2.GetRelativeTransform(Source2);
-			}
-		}
-	}
-
-	if (bInterpolateT)
-	{
-		for (FCompactPoseBoneIndex BoneIndex : InOutPose.ForEachBoneIndex())
-		{
-			InOutPose[BoneIndex].Blend(InOutPose[BoneIndex], Key2Pose[BoneIndex], Alpha);
-		}
-	}
-}
-
-void UAnimDataModel::Evaluate(FAnimationPoseData& InOutPoseData, const UE::Anim::DataModel::FEvaluationContext& EvaluationContext) const
-{
-	const float SequenceLength = GetPlayLength();
-
-	const double Time = EvaluationContext.SampleFrameRate.AsSeconds(EvaluationContext.SampleTime);
-	const EAnimInterpolationType& InterpolationType = EvaluationContext.InterpolationType;
-	FCompactPose& OutPose = InOutPoseData.GetPose();
-	
-	// Generate keys to interpolate between
-	int32 KeyIndex1, KeyIndex2;
-	float Alpha;
-	FAnimationRuntime::GetKeyIndicesFromTime(KeyIndex1, KeyIndex2, Alpha, Time, FrameRate, NumberOfKeys);
-
-	if (InterpolationType == EAnimInterpolationType::Step)
-	{
-		// Force stepping between keys
-		Alpha = 0.f;
-	}
-
-	bool bShouldInterpolate = true;
-
-	if (Alpha < UE_KINDA_SMALL_NUMBER)
-	{
-		Alpha = 0.f;
-		bShouldInterpolate = false;
-	}
-	else if (Alpha > 1.f - UE_KINDA_SMALL_NUMBER)
-	{
-		bShouldInterpolate = false;
-		KeyIndex1 = KeyIndex2;
-	}
-
-	// Evaluate animation float curve data
-	FBlendedCurve& Curves = InOutPoseData.GetCurve();
-	if (Curves.NumValidCurveCount > 0)
-	{
-		// evaluate the curve data at the CurrentTime and add to Instance
-		for (const FFloatCurve& Curve : CurveData.FloatCurves)
-		{
-			if (Curves.IsEnabled(Curve.Name.UID))
-			{
-				float Value = Curve.Evaluate(Time);
-				Curves.Set(Curve.Name.UID, Value);
-			}
-		}
-	}
-
-	const double TimePerFrame = FrameRate.AsInterval();
-	TMap<FName, const FTransformCurve*> ActiveCurves;
-
-	if (!OutPose.GetBoneContainer().ShouldUseSourceData())
-	{
-		for (const FTransformCurve& Curve : GetTransformCurves())
-		{
-			// if disabled, do not handle
-			if (Curve.GetCurveTypeFlag(AACF_Disabled))
-			{
-				continue;
-			}
-
-			// Add or retrieve curve
-			const FName& CurveName = Curve.Name.DisplayName;
-
-			// note we're not checking Curve.GetCurveTypeFlags() yet
-			ActiveCurves.FindOrAdd(CurveName, &Curve);
-		}
-	}
-
-	{		
-		UE::Anim::Retargeting::FRetargetingScope RetargetingScope(OutPose, EvaluationContext);
-		if (bShouldInterpolate)
-		{		
-			ExtractPose<true>(GetBoneAnimationTracks(), ActiveCurves, OutPose, KeyIndex1, KeyIndex2, Alpha, TimePerFrame, RetargetingScope);
-		}
-		else
-		{
-			ExtractPose<false>(GetBoneAnimationTracks(), ActiveCurves, OutPose, KeyIndex1, KeyIndex2, Alpha, TimePerFrame, RetargetingScope);
-		}
-	}
-	
-
-	const FBoneContainer& RequiredBones = InOutPoseData.GetPose().GetBoneContainer();
-	for (const FAnimatedBoneAttribute& Attribute : AnimatedBoneAttributes)
-	{
-		const FCompactPoseBoneIndex PoseBoneIndex = RequiredBones.GetCompactPoseIndexFromSkeletonIndex(Attribute.Identifier.GetBoneIndex());
-		// Only add attribute if the bone its tied to exists in the currently evaluated set of bones
-		if(PoseBoneIndex.IsValid())
-		{
-			UE::Anim::Attributes::GetAttributeValue(InOutPoseData.GetAttributes(), PoseBoneIndex, Attribute, EvaluationContext.SampleFrameRate.AsSeconds(EvaluationContext.SampleTime));
-		}
-	}	
-}
-#endif // WITH_EDITOR
-
 FRichCurve* UAnimDataModel::GetMutableRichCurve(const FAnimationCurveIdentifier& CurveIdentifier)
 {
 	FRichCurve* RichCurve = nullptr;
@@ -808,4 +516,3 @@ FAnimCurveBase* UAnimDataModel::FindMutableCurveById(const FAnimationCurveIdenti
 	return nullptr;
 }
 
-#undef LOCTEXT_NAMESPACE // "AnimDataModel"
diff --git a/Engine/Source/Runtime/Engine/Private/Animation/AnimInstanceProxy.cpp b/Engine/Source/Runtime/Engine/Private/Animation/AnimInstanceProxy.cpp
--- a/Engine/Source/Runtime/Engine/Private/Animation/AnimInstanceProxy.cpp
+++ b/Engine/Source/Runtime/Engine/Private/Animation/AnimInstanceProxy.cpp
@@ -1346,7 +1346,7 @@ void FAnimInstanceProxy::EvaluateAnimationNode(FPoseContext& Output)
 
 void FAnimInstanceProxy::EvaluateAnimationNode_WithRoot(FPoseContext& Output, FAnimNode_Base* InRootNode)
 {
-	if (InRootNode != nullptr)
+	if (InRootNode != NULL)
 	{
 		ANIM_MT_SCOPE_CYCLE_COUNTER(EvaluateAnimGraph, !IsInGameThread());
 		if(InRootNode == RootNode)
@@ -1467,7 +1467,7 @@ void FAnimInstanceProxy::SlotEvaluatePoseWithBlendProfiles(const FName& SlotNode
 			NewPose.Curve.InitFrom(RequiredBones);
 
 			// Extract pose from Track.
-			FAnimExtractContext ExtractionContext(static_cast<double>(EvalState.MontagePosition), Montage->HasRootMotion() && RootMotionMode != ERootMotionMode::NoRootMotionExtraction, EvalState.DeltaTimeRecord);
+			FAnimExtractContext ExtractionContext(EvalState.MontagePosition, Montage->HasRootMotion() && RootMotionMode != ERootMotionMode::NoRootMotionExtraction, EvalState.DeltaTimeRecord);
 			FAnimationPoseData NewAnimationPoseData(NewPose);
 			AnimTrack->GetAnimationPose(NewAnimationPoseData, ExtractionContext);
 
@@ -1806,7 +1806,7 @@ void FAnimInstanceProxy::SlotEvaluatePose(const FName& SlotNodeName, const FAnim
 			NewPose.Curve.InitFrom(RequiredBones);
 
 			// Extract pose from Track
-			FAnimExtractContext ExtractionContext(static_cast<double>(EvalState.MontagePosition), Montage->HasRootMotion() && RootMotionMode != ERootMotionMode::NoRootMotionExtraction, EvalState.DeltaTimeRecord);
+			FAnimExtractContext ExtractionContext(EvalState.MontagePosition, Montage->HasRootMotion() && RootMotionMode != ERootMotionMode::NoRootMotionExtraction, EvalState.DeltaTimeRecord);
 
 			FAnimationPoseData NewAnimationPoseData(NewPose);
 			AnimTrack->GetAnimationPose(NewAnimationPoseData, ExtractionContext);
diff --git a/Engine/Source/Runtime/Engine/Private/Animation/AnimMontage.cpp b/Engine/Source/Runtime/Engine/Private/Animation/AnimMontage.cpp
--- a/Engine/Source/Runtime/Engine/Private/Animation/AnimMontage.cpp
+++ b/Engine/Source/Runtime/Engine/Private/Animation/AnimMontage.cpp
@@ -92,7 +92,7 @@ const FAnimTrack* UAnimMontage::GetAnimationData(FName InSlotName) const
 		}
 	}
 
-	return nullptr;
+	return NULL;
 }
 
 bool UAnimMontage::IsWithinPos(int32 FirstIndex, int32 SecondIndex, float CurrentTime) const
@@ -369,7 +369,7 @@ void UAnimMontage::PostLoad()
 		Track.ValidateSegmentTimes();
 
 		const float CurrentCalculatedLength = CalculateSequenceLength();
-		if(!FMath::IsNearlyEqual(CurrentCalculatedLength, GetPlayLength(), UE_KINDA_SMALL_NUMBER))		
+		if(CurrentCalculatedLength != GetPlayLength())		
 		{
 			UE_LOG(LogAnimMontage, Display, TEXT("UAnimMontage::PostLoad: The actual sequence length for %s does not match the length stored in the asset, please resave the asset."), *GetFullName());
 			SetCompositeLength(CurrentCalculatedLength);
@@ -418,7 +418,7 @@ void UAnimMontage::PostLoad()
 	}
 	// find preview base pose if it can
 #if WITH_EDITORONLY_DATA
-	if ( IsValidAdditive() && PreviewBasePose == nullptr )
+	if ( IsValidAdditive() && PreviewBasePose == NULL )
 	{
 		for (int32 I=0; I<SlotAnimTracks.Num(); ++I)
 		{
@@ -555,8 +555,8 @@ void UAnimMontage::ConvertBranchingPointsToAnimNotifies()
 #if WITH_EDITORONLY_DATA
 			NewEvent.TrackIndex = TrackIndex;
 #endif
-			NewEvent.Notify = nullptr;
-			NewEvent.NotifyStateClass = nullptr;
+			NewEvent.Notify = NULL;
+			NewEvent.NotifyStateClass = NULL;
 			NewEvent.bConvertedFromBranchingPoint = true;
 			NewEvent.MontageTickType = EMontageNotifyTickType::BranchingPoint;
 
@@ -1310,14 +1310,11 @@ bool UAnimMontage::ContainRecursive(TArray<UAnimCompositeBase*>& CurrentAccumula
 
 void UAnimMontage::SetCompositeLength(float InLength)
 {
-#if WITH_EDITOR
-	const FFrameTime LengthInFrameTime = DataModelInterface->GetFrameRate().AsFrameTime(InLength);
-	PRAGMA_DISABLE_DEPRECATION_WARNINGS
-	Controller->SetNumberOfFrames(LengthInFrameTime.RoundToFrame());
-	PRAGMA_ENABLE_DEPRECATION_WARNINGS
+#if WITH_EDITOR	
+	Controller->SetPlayLength(InLength);
 #else
 	PRAGMA_DISABLE_DEPRECATION_WARNINGS
-	SequenceLength = InLength;
+	SetSequenceLength(InLength);
 	PRAGMA_ENABLE_DEPRECATION_WARNINGS
 #endif	
 }
@@ -1327,11 +1324,11 @@ void UAnimMontage::SetCompositeLength(float InLength)
 /////////////////////////////////////////////////////////////////////////////////////////////
 
 FAnimMontageInstance::FAnimMontageInstance()
-	: Montage(nullptr)
+	: Montage(NULL)
 	, bPlaying(false)
 	, DefaultBlendTimeMultiplier(1.0f)
 	, bDidUseMarkerSyncThisTick(false)
-	, AnimInstance(nullptr)
+	, AnimInstance(NULL)
 	, InstanceID(INDEX_NONE)
 	, Position(0.f)
 	, PlayRate(1.f)
@@ -1343,13 +1340,13 @@ FAnimMontageInstance::FAnimMontageInstance()
 	, ActiveBlendProfile(nullptr)
 	, ActiveBlendProfileMode(EBlendProfileMode::TimeFactor)
 	, DisableRootMotionCount(0)
-	, MontageSyncLeader(nullptr)
+	, MontageSyncLeader(NULL)
 	, MontageSyncUpdateFrameCounter(INDEX_NONE)
 {
 }
 
 FAnimMontageInstance::FAnimMontageInstance(UAnimInstance * InAnimInstance)
-	: Montage(nullptr)
+	: Montage(NULL)
 	, bPlaying(false)
 	, DefaultBlendTimeMultiplier(1.0f)
 	, bDidUseMarkerSyncThisTick(false)
@@ -1366,7 +1363,7 @@ FAnimMontageInstance::FAnimMontageInstance(UAnimInstance * InAnimInstance)
 	, ActiveBlendProfile(nullptr)
 	, ActiveBlendProfileMode(EBlendProfileMode::TimeFactor)
 	, DisableRootMotionCount(0)
-	, MontageSyncLeader(nullptr)
+	, MontageSyncLeader(NULL)
 	, MontageSyncUpdateFrameCounter(INDEX_NONE)
 {
 }
@@ -1592,7 +1589,7 @@ void FAnimMontageInstance::Terminate()
 {
 	SCOPE_CYCLE_COUNTER(STAT_AnimMontageInstance_Terminate);
 
-	if (Montage == nullptr)
+	if (Montage == NULL)
 	{
 		return;
 	}
@@ -1631,7 +1628,7 @@ void FAnimMontageInstance::Terminate()
 	}
 
 	// clear Blend curve
-	Blend.SetCustomCurve(nullptr);
+	Blend.SetCustomCurve(NULL);
 	Blend.SetBlendOption(EAlphaBlendOption::Linear);
 
 	ActiveBlendProfile = nullptr;
@@ -1777,7 +1774,7 @@ void FAnimMontageInstance::MontageSync_StopLeading()
 		if (MontageSyncFollower)
 		{
 			ensure(MontageSyncFollower->MontageSyncLeader == this);
-			MontageSyncFollower->MontageSyncLeader = nullptr;
+			MontageSyncFollower->MontageSyncLeader = NULL;
 		}
 	}
 	MontageSyncFollowers.Empty();
@@ -1788,7 +1785,7 @@ void FAnimMontageInstance::MontageSync_StopFollowing()
 	if (MontageSyncLeader)
 	{
 		MontageSyncLeader->MontageSyncFollowers.RemoveSingleSwap(this);
-		MontageSyncLeader = nullptr;
+		MontageSyncLeader = NULL;
 	}
 }
 
@@ -2661,7 +2658,7 @@ void FAnimMontageInstance::BranchingPointEventHandler(const FBranchingPointMarke
 		// Must grab a reference on the stack in case "this" is deleted during iteration
 		TWeakObjectPtr<UAnimInstance> AnimInstanceLocal = AnimInstance;
 
-		FAnimNotifyEvent* NotifyEvent = (BranchingPointMarker->NotifyIndex < Montage->Notifies.Num()) ? &Montage->Notifies[BranchingPointMarker->NotifyIndex] : nullptr;
+		FAnimNotifyEvent* NotifyEvent = (BranchingPointMarker->NotifyIndex < Montage->Notifies.Num()) ? &Montage->Notifies[BranchingPointMarker->NotifyIndex] : NULL;
 		if (NotifyEvent)
 		{
 			// Handle backwards compatibility with older BranchingPoints.
@@ -2673,7 +2670,7 @@ void FAnimMontageInstance::BranchingPointEventHandler(const FBranchingPointMarke
 				UFunction* Function = AnimInstance.Get()->FindFunction(FuncFName);
 				if (Function)
 				{
-					AnimInstance.Get()->ProcessEvent(Function, nullptr);
+					AnimInstance.Get()->ProcessEvent(Function, NULL);
 				}
 				// In case older BranchingPoint has been re-implemented as a new Custom Notify, this is if BranchingPoint function hasn't been found.
 				else
@@ -2681,7 +2678,7 @@ void FAnimMontageInstance::BranchingPointEventHandler(const FBranchingPointMarke
 					AnimInstance.Get()->TriggerSingleAnimNotify(NotifyEvent);
 				}
 			}
-			else if (NotifyEvent->NotifyStateClass != nullptr)
+			else if (NotifyEvent->NotifyStateClass != NULL)
 			{
 				if (BranchingPointMarker->NotifyEventType == EAnimNotifyEventType::Begin)
 				{
@@ -2810,7 +2807,7 @@ UAnimMontage* FAnimMontageInstance::PreviewSequencerMontagePosition(FName SlotNa
 			// we have to get it again in case if this is new
 			MontageInstanceToUpdate = AnimInst->GetMontageInstanceForID(InOutInstanceId);
 			// since we don't advance montage in the tick, we manually have to handle notifies
-			MontageInstanceToUpdate->HandleEvents(InFromPosition, InToPosition, nullptr);
+			MontageInstanceToUpdate->HandleEvents(InFromPosition, InToPosition, NULL);
 			if (!bFireNotifies)
 			{
 				AnimInst->NotifyQueue.Reset(SkeletalMeshComponent);
@@ -2929,15 +2926,6 @@ void UAnimMontage::BakeTimeStretchCurve()
 	TimeStretchCurve.BakeFromFloatCurve(*TimeStretchFloatCurve, SequenceLength);
 	PRAGMA_ENABLE_DEPRECATION_WARNINGS
 }
-
-void UAnimMontage::PopulateWithExistingModel(TScriptInterface<IAnimationDataModel> ExistingDataModel)
-{
-	Super::PopulateWithExistingModel(ExistingDataModel);
-	
-	// Set composite length while model is being populated
-	const float CurrentCalculatedLength = CalculateSequenceLength();
-	SetCompositeLength(CurrentCalculatedLength);
-}
 #endif // WITH_EDITOR
 
 FMontageBlendSettings::FMontageBlendSettings()
diff --git a/Engine/Source/Runtime/Engine/Private/Animation/AnimNode_SequencePlayer.cpp b/Engine/Source/Runtime/Engine/Private/Animation/AnimNode_SequencePlayer.cpp
--- a/Engine/Source/Runtime/Engine/Private/Animation/AnimNode_SequencePlayer.cpp
+++ b/Engine/Source/Runtime/Engine/Private/Animation/AnimNode_SequencePlayer.cpp
@@ -128,7 +128,7 @@ void FAnimNode_SequencePlayerBase::Evaluate_AnyThread(FPoseContext& Output)
 		}
 
 		FAnimationPoseData AnimationPoseData(Output);
-		CurrentSequence->GetAnimationPose(AnimationPoseData, FAnimExtractContext(static_cast<double>(InternalTimeAccumulator), Output.AnimInstanceProxy->ShouldExtractRootMotion(), DeltaTimeRecord, GetLoopAnimation()));
+		CurrentSequence->GetAnimationPose(AnimationPoseData, FAnimExtractContext(InternalTimeAccumulator, Output.AnimInstanceProxy->ShouldExtractRootMotion(), DeltaTimeRecord, GetLoopAnimation()));
 	}
 	else
 	{
@@ -141,7 +141,7 @@ void FAnimNode_SequencePlayerBase::GatherDebugData(FNodeDebugData& DebugData)
 	FString DebugLine = DebugData.GetNodeName(this);
 
 	UAnimSequenceBase* CurrentSequence = GetSequence();
-	DebugLine += FString::Printf(TEXT("('%s' Play Time: %.3f)"), CurrentSequence ? *CurrentSequence->GetName() : TEXT("nullptr"), InternalTimeAccumulator);
+	DebugLine += FString::Printf(TEXT("('%s' Play Time: %.3f)"), CurrentSequence ? *CurrentSequence->GetName() : TEXT("NULL"), InternalTimeAccumulator);
 	DebugData.AddDebugItem(DebugLine, true);
 }
 
