Revert ApplyMode and bIsLowAreaModifier back to it's proper values if it had been replaced before.

diff --git a/Engine/Source/Runtime/Engine/Private/AI/Navigation/NavigationModifier.cpp b/Engine/Source/Runtime/Engine/Private/AI/Navigation/NavigationModifier.cpp
--- a/Engine/Source/Runtime/Engine/Private/AI/Navigation/NavigationModifier.cpp
+++ b/Engine/Source/Runtime/Engine/Private/AI/Navigation/NavigationModifier.cpp
@@ -450,6 +450,12 @@ void FAreaNavModifier::SetAreaClassToReplace(const TSubclassOf<UNavAreaBase> InA
 		bIsLowAreaModifier = AreaClass2->GetDefaultObject<UNavAreaBase>()->IsLowArea();
 		ApplyMode = bIsLowAreaModifier ? ENavigationAreaMode::ReplaceInLowPass : ENavigationAreaMode::Replace;
 	}
+	else if (ApplyMode == ENavigationAreaMode::ReplaceInLowPass || ApplyMode == ENavigationAreaMode::Replace)
+	{
+		// since we no longer have ReplaceAreaClass the new value of ApplyMode and bIsLowAreaModifier should depend on previous value of ApplyMode
+		bIsLowAreaModifier = ApplyMode == ENavigationAreaMode::ReplaceInLowPass;
+		ApplyMode = ApplyMode == ENavigationAreaMode::ReplaceInLowPass ? ENavigationAreaMode::ApplyInLowPass : ENavigationAreaMode::Apply;		
+	}
 }
 
 void FAreaNavModifier::SetApplyMode(ENavigationAreaMode::Type InApplyMode)
