Fix Queries transforming inputs by GetLocalTransform and inconsistent hit location.
Fixes HitResult.Location/End/Start sometimes being relative to actor, and sometimes shape, (when it is stored as TransformedImplicit).  Start/End now match query input, and Location is always based off actor location at hit.

Fixes incorrect overlap results for TransformedImplicits.

diff --git a/Engine/Source/Runtime/Engine/Private/Collision/WorldCollision.cpp b/Engine/Source/Runtime/Engine/Private/Collision/WorldCollision.cpp
--- a/Engine/Source/Runtime/Engine/Private/Collision/WorldCollision.cpp
+++ b/Engine/Source/Runtime/Engine/Private/Collision/WorldCollision.cpp
@@ -387,10 +387,6 @@ bool UWorld::ComponentSweepMulti(TArray<struct FHitResult>& OutHits, class UPrim
 		FInlineShapeArray PShapes;
 		const int32 NumShapes = FillInlineShapeArray_AssumesLocked(PShapes, Actor);
 
-		// calculate the test global pose of the actor
-		const FTransform GlobalStartTransform(Quat, Start);
-		const FTransform GlobalEndTransform(Quat, End);
-
 		for(FPhysicsShapeHandle& Shape : PShapes)
 		{
 			check(Shape.IsValid());
@@ -402,17 +398,9 @@ bool UWorld::ComponentSweepMulti(TArray<struct FHitResult>& OutHits, class UPrim
 				continue;
 			}
 
-			// Calc shape global pose
-			const FTransform ShapeLocalTransform = FPhysicsInterface::GetLocalTransform(Shape);
-			const FTransform GlobalStartTransform_Shape = ShapeLocalTransform * GlobalStartTransform;
-			FTransform GlobalEndTransform_Shape = ShapeLocalTransform * GlobalEndTransform;
-
-			// consider localshape rotation for shape rotation
-			const FQuat ShapeQuat = Quat * ShapeLocalTransform.GetRotation();
-
 			FPhysicsGeometryCollection GeomCollection = FPhysicsInterface::GetGeometryCollection(Shape);
 			TArray<FHitResult> TmpHits;
-			if(FPhysicsInterface::GeomSweepMulti(this, GeomCollection, ShapeQuat, TmpHits, GlobalStartTransform_Shape.GetTranslation(), GlobalEndTransform_Shape.GetTranslation(), TraceChannel, Params, FCollisionResponseParams(PrimComp->GetCollisionResponseToChannels())))
+			if(FPhysicsInterface::GeomSweepMulti(this, GeomCollection, Quat, TmpHits, Start, End, TraceChannel, Params, FCollisionResponseParams(PrimComp->GetCollisionResponseToChannels())))
 			{
 				bHaveBlockingHit = true;
 			}
diff --git a/Engine/Source/Runtime/Engine/Private/PhysicsEngine/BodyInstance.cpp b/Engine/Source/Runtime/Engine/Private/PhysicsEngine/BodyInstance.cpp
--- a/Engine/Source/Runtime/Engine/Private/PhysicsEngine/BodyInstance.cpp
+++ b/Engine/Source/Runtime/Engine/Private/PhysicsEngine/BodyInstance.cpp
@@ -3570,11 +3570,9 @@ bool FBodyInstance::OverlapTestForBodiesImpl(const FVector& Pos, const FQuat& Ro
 				continue;	//we skip complex shapes - should this respect ComplexAsSimple?
 			}
 
-			// Calc shape global pose
-			FTransform PShapeGlobalPose = FPhysicsInterface::GetLocalTransform(Shape) * PTestGlobalPose;
 			for (const FBodyInstance* BodyInstance : Bodies)
 			{
-				bHaveOverlap = FPhysicsInterface::Overlap_Geom(BodyInstance, FPhysicsInterface::GetGeometryCollection(Shape), PShapeGlobalPose);
+				bHaveOverlap = FPhysicsInterface::Overlap_Geom(BodyInstance, FPhysicsInterface::GetGeometryCollection(Shape), PTestGlobalPose);
 
 				if (bHaveOverlap)
 				{
@@ -3710,12 +3708,8 @@ bool FBodyInstance::OverlapMulti(TArray<struct FOverlapResult>& InOutOverlaps, c
 					continue;	//we skip complex shapes - should this respect ComplexAsSimple?
 				}
 
-				// Calc shape global pose
-				const FTransform LocalTransform = FPhysicsInterface::GetLocalTransform(ShapeRef);
-				const FTransform GlobalTransform = LocalTransform * BodyInstanceSpaceToTestSpace;
-			
 				TempOverlaps.Reset();
-				if(FPhysicsInterface::GeomOverlapMulti(World, GeomCollection, GlobalTransform.GetTranslation(), GlobalTransform.GetRotation(), TempOverlaps, TestChannel, Params, ResponseParams, ObjectQueryParams))
+				if(FPhysicsInterface::GeomOverlapMulti(World, GeomCollection, BodyInstanceSpaceToTestSpace.GetTranslation(), BodyInstanceSpaceToTestSpace.GetRotation(), TempOverlaps, TestChannel, Params, ResponseParams, ObjectQueryParams))
 				{
 					bHaveBlockingHit = true;
 				}
