#UE Add package name to fatal error about backed out serialization version

diff --git a/Engine/Source/Runtime/Engine/Private/Animation/AnimSequenceBase.cpp b/Engine/Source/Runtime/Engine/Private/Animation/AnimSequenceBase.cpp
--- a/Engine/Source/Runtime/Engine/Private/Animation/AnimSequenceBase.cpp
+++ b/Engine/Source/Runtime/Engine/Private/Animation/AnimSequenceBase.cpp
@@ -975,7 +975,7 @@ void UAnimSequenceBase::Serialize(FArchive& Ar)
 	if (GetLinkerCustomVersion(FUE5MainStreamObjectVersion::GUID) >= FUE5MainStreamObjectVersion::AnimationDataModelInterface_BackedOut &&
 		GetLinkerCustomVersion(FUE5MainStreamObjectVersion::GUID) < FUE5MainStreamObjectVersion::BackoutAnimationDataModelInterface)
 	{
-		UE_LOG(LogAnimation, Fatal, TEXT("This package was saved with a version that had to be backed out and is no longer able to be loaded."));
+		UE_LOG(LogAnimation, Fatal, TEXT("This package (%s) was saved with a version that had to be backed out and is no longer able to be loaded."), *GetPathNameSafe(this));
 	}
 
 	// fix up version issue and so on
