Strata - Clear coat
- secondary normal is now achieved with two slabs, no fast path for that one which before needed preoject settings anyway.
- coat is only enabled in shader when coverage is >0
- eye is removed from the dynamic case

diff --git a/Engine/Shaders/Private/Strata/StrataLegacyConversion.ush b/Engine/Shaders/Private/Strata/StrataLegacyConversion.ush
--- a/Engine/Shaders/Private/Strata/StrataLegacyConversion.ush
+++ b/Engine/Shaders/Private/Strata/StrataLegacyConversion.ush
@@ -163,9 +163,8 @@ FStrataData StrataConvertLegacyMaterialDynamic(
 
 		FStrataLegacyParameters Bot = InitStrataLegacyParameters(PixelFootprint, SharedLocalBasisIndex, 1.0f);
 
-		FStrataData WTop_Slab = StrataWeight(Top_SLW , 1.0, StrataTree,					/*OperatorIndex*/1, /*MaxDistanceFromLeaves*/1);
+		return StrataWeight(Top_SLW , 1.0, StrataTree,					/*OperatorIndex*/1, /*MaxDistanceFromLeaves*/1);
 		// Out maps to LegacyUpdateBSDFsOperators
-		return Out;
 	}
 	#endif
 	#if MATERIAL_SHADINGMODEL_HAIR
@@ -183,26 +182,6 @@ FStrataData StrataConvertLegacyMaterialDynamic(
 		// No need to take into account opacity weighting for translucent material, as the coverage is assumed to be 1
 	}
 	#endif
-	#if MATERIAL_SHADINGMODEL_EYE
-	if (ShadingModel == SHADINGMODELID_EYE)
-	{
-		const float IrisMask = ClearCoat;
-		const float IrisDistance = ClearCoatRoughness;
-		const float3 IrisNormal = RawClearCoatNormal;
-
-		return GetStrataEyeBSDF(
-			BaseColor,	        // DiffuseAlbedo
-			Roughness,	        // Roughness
-			ClearCoat,	        // IrisMask
-			ClearCoatRoughness, // IrisDistance
-			IrisNormal,         // IrisNormal TODO
-			Emissive,	        // EmissiveColor
-			SharedLocalBasisIndex);
-
-		// No need to take into account opacity weighting for translucent material, as the coverage is assumed to be 1
-	}
-	#endif
-
 
 	FStrataLegacyParameters LegacySlab = InitStrataLegacyParameters(PixelFootprint, SharedLocalBasisIndex, 1.0f);
 
@@ -324,7 +303,7 @@ FStrataData StrataConvertLegacyMaterialDynamic(
 		LegacySlab.SSSMFPScale = 0.f;
 		LegacySlab.SecondRoughnessWeight = ClearCoat;
 		LegacySlab.SecondRoughness = ClearCoatRoughness;
-		LegacySlab.SecondRoughnessAsSimpleClearCoat = true; // Haziness as simple clear coat layer.
+		LegacySlab.SecondRoughnessAsSimpleClearCoat = ClearCoat > 0.0f; // Haziness as simple clear coat layer.
 		LegacySlab.Emissive = Zeros;
 		LegacySlab.SharedLocalBasisIndex = ClearCoatBottomNormal_SharedLocalBasisIndex;
 		LegacySlab.Weight = 1.f;
@@ -421,6 +400,21 @@ FStrataData StrataConvertLegacyMaterialDynamic(
 		// Thin translucency model always have a coverage of 1. It also canont be used in a multiple shading model context.
 		FinalWeight = 1.0f;
 	}
+	else if (ShadingModel == SHADINGMODELID_EYE) 
+	{
+		// Fallback to bright purple in case eye shading is selected when using the material attribute workflow.
+		LegacySlab.BaseColor = float3(0, 1, 1);
+		LegacySlab.Specular = 0;
+		LegacySlab.Metallic = 0;
+		LegacySlab.Roughness = 0;
+		LegacySlab.Anisotropy = 0;
+		LegacySlab.Emissive = 0;
+
+		FinalWeight = 1.0f;
+#if STRATA_TRANSLUCENT_MATERIAL
+		FinalWeight = 1.0f;
+#endif
+	}
 
 	// Fixed layer structure for helping compiler to unroll and optimize shader
 	// This must match what is done in PackStrataOut
diff --git a/Engine/Source/Runtime/Engine/Classes/Materials/MaterialExpressionStrata.h b/Engine/Source/Runtime/Engine/Classes/Materials/MaterialExpressionStrata.h
--- a/Engine/Source/Runtime/Engine/Classes/Materials/MaterialExpressionStrata.h
+++ b/Engine/Source/Runtime/Engine/Classes/Materials/MaterialExpressionStrata.h
@@ -99,13 +99,13 @@ class UMaterialExpressionStrataLegacyConversion : public UMaterialExpressionStra
 	FExpressionInput SubSurfaceColor;
 
 	/**
-	 * Emissive color on top of the surface (type = float3, unit = luminance, default = 0)
+	 * Coverage of the clear coat layer. (type = float, unit = unitless, defaults to 0.0)
 	 */
 	UPROPERTY()
 	FExpressionInput ClearCoat;
 
 	/**
-	 * The amount of fuzz on top of the surface used to simulate cloth-like appearance.
+	 * Roughness of the top clear coat layer. (type = float, unit = unitless, defaults to 0.0)
 	 */
 	UPROPERTY()
 	FExpressionInput ClearCoatRoughness;
diff --git a/Engine/Source/Runtime/Engine/Private/Materials/Material.cpp b/Engine/Source/Runtime/Engine/Private/Materials/Material.cpp
--- a/Engine/Source/Runtime/Engine/Private/Materials/Material.cpp
+++ b/Engine/Source/Runtime/Engine/Private/Materials/Material.cpp
@@ -2994,7 +2994,7 @@ void UMaterial::ConvertMaterialToStrataMaterial()
 	bool bCustomNodesGathered = false;
 	UMaterialExpressionThinTranslucentMaterialOutput* ThinTranslucentOutput = nullptr;
 	UMaterialExpressionSingleLayerWaterMaterialOutput* SingleLayerWaterOutput = nullptr;
-	UMaterialExpressionClearCoatNormalCustomOutput* ClearCoatOutput = nullptr;
+	UMaterialExpressionClearCoatNormalCustomOutput* ClearCoatBottomNormalOutput = nullptr;
 	UMaterialExpressionTangentOutput* TangentOutput = nullptr;
 	auto GatherCustomNodes = [&]()
 	{
@@ -3018,9 +3018,9 @@ void UMaterial::ConvertMaterialToStrataMaterial()
 				}
 
 				// Gather custom output for clear coat
-				if (ClearCoatOutput == nullptr && Cast<UMaterialExpressionClearCoatNormalCustomOutput>(Expression))
+				if (ClearCoatBottomNormalOutput == nullptr && Cast<UMaterialExpressionClearCoatNormalCustomOutput>(Expression))
 				{
-					ClearCoatOutput = Cast<UMaterialExpressionClearCoatNormalCustomOutput>(Expression);
+					ClearCoatBottomNormalOutput = Cast<UMaterialExpressionClearCoatNormalCustomOutput>(Expression);
 				}
 
 				// Gather custom output for tangent (unused atm)
@@ -3029,7 +3029,7 @@ void UMaterial::ConvertMaterialToStrataMaterial()
 					TangentOutput = Cast<UMaterialExpressionTangentOutput>(Expression);
 				}
 
-				if (ThinTranslucentOutput && SingleLayerWaterOutput && ClearCoatOutput && TangentOutput)
+				if (ThinTranslucentOutput && SingleLayerWaterOutput && ClearCoatBottomNormalOutput && TangentOutput)
 				{
 					break;
 				}
@@ -3106,68 +3106,128 @@ void UMaterial::ConvertMaterialToStrataMaterial()
 	}
 	else if (!bUseMaterialAttributes && !EditorOnly->FrontMaterial.IsConnected())
 	{
-		// STRATA_TODO for material conversion
-		//  - WorldPositionOffset can remain on the end point node
-		//  - Refraction
-		//  - PixelDepthOffset
-
 		if (MaterialDomain == MD_Surface)
 		{
-			ConvertNode = NewObject<UMaterialExpressionStrataLegacyConversion>(this);
-			ConvertNode->SubsurfaceProfile = bRequireNoSubsurfaceProfile ? nullptr : SubsurfaceProfile;
-			MoveConnectionTo(EditorOnly->BaseColor, ConvertNode, 0);
-			MoveConnectionTo(EditorOnly->Metallic, ConvertNode, 1);
-			MoveConnectionTo(EditorOnly->Specular, ConvertNode, 2);
-			MoveConnectionTo(EditorOnly->Roughness, ConvertNode, 3);
-			MoveConnectionTo(EditorOnly->Anisotropy, ConvertNode, 4);
-			MoveConnectionTo(EditorOnly->EmissiveColor, ConvertNode, 5);
-			CopyConnectionTo(EditorOnly->Normal, ConvertNode, 6);
-			MoveConnectionTo(EditorOnly->Tangent, ConvertNode, 7);
-			MoveConnectionTo(EditorOnly->SubsurfaceColor, ConvertNode, 8);
-			MoveConnectionTo(EditorOnly->ClearCoat, ConvertNode, 9);
-			MoveConnectionTo(EditorOnly->ClearCoatRoughness, ConvertNode, 10);
-			MoveConnectionTo(EditorOnly->Opacity, ConvertNode, 11);
-			bRelinkCustomOutputNodes = true;
-			
-			// Shading Model
-			// * either use the shader graph expression 
-			// * or add a constant shading model
-			// 
-			// Note: store this conversion type(s) into ConvertedStrataMaterialInfo for having more context when 
-			// rebuilding the final Shading model (see RebuildShadingModelField())
-			if (ShadingModel == MSM_FromMaterialExpression)
+			bool bClearCoatConversionDone = false;
+			if (ShadingModel == MSM_ClearCoat)
 			{
-				check(EditorOnly->ShadingModelFromMaterialExpression.IsConnected());
-
-				// Reconnect the shading model expression
-				MoveConnectionTo(EditorOnly->ShadingModelFromMaterialExpression, ConvertNode, 18);
-
-				// Store strata shading model of the converted material. 
 				GatherCustomNodes();
-				if (SingleLayerWaterOutput)
+				if (ClearCoatBottomNormalOutput)
 				{
-					ConvertNode->ConvertedStrataMaterialInfo.AddShadingModel(SSM_SingleLayerWater);
+					// For this special case, using two slabs to create a clear coat material with separated top and bottom normal. 
+
+					// Top slab BSDF as a simple Disney material
+					UMaterialExpressionStrataSlabBSDF* BottomSlabBSDF = NewObject<UMaterialExpressionStrataSlabBSDF>(this);
+					MoveConnectionTo(EditorOnly->BaseColor, BottomSlabBSDF, 0);					// BaseColor
+					MoveConnectionTo(EditorOnly->Metallic, BottomSlabBSDF, 2);					// Metallic
+					MoveConnectionTo(EditorOnly->Specular, BottomSlabBSDF, 3);					// Specular
+					MoveConnectionTo(EditorOnly->Roughness, BottomSlabBSDF, 7);					// Roughness
+					CopyConnectionTo(EditorOnly->Anisotropy, BottomSlabBSDF, 8);				// Anisotropy
+					MoveConnectionTo(EditorOnly->Tangent, BottomSlabBSDF, 10);					// Tangent
+
+					check(ClearCoatBottomNormalOutput);
+					MoveConnectionTo(*ClearCoatBottomNormalOutput->GetInput(0), BottomSlabBSDF, 9);// ClearColorBottomNormal -> BottomSlabBSDF.Normal
+
+					// Now weight the top base material by opacity.
+					UMaterialExpressionStrataSlabBSDF* TopSlabBSDF = NewObject<UMaterialExpressionStrataSlabBSDF>(this);
+					MoveConnectionTo(EditorOnly->EmissiveColor, TopSlabBSDF, 14);				// Emissive
+					MoveConnectionTo(EditorOnly->ClearCoatRoughness, TopSlabBSDF, 7);			// ClearCoatRoughness => Roughness
+					MoveConnectionTo(EditorOnly->Normal, TopSlabBSDF, 9);						// Normal
+
+					//  The top layer has a hard coded specular value of 0.5 (F0 = 0.04)
+					UMaterialExpressionConstant* ConstantHalf = NewObject<UMaterialExpressionConstant>(this);
+					ConstantHalf->R = 0.5f;
+					TopSlabBSDF->GetInput(3)->Connect(0, ConstantHalf);
+
+					// The original clear coat is a complex assemblage of arbitrary functions that do not always make sense.
+					// To simplify things, we set the top slab BSDF as having a constant Grey scale transmittance.
+					// As for the original, this is achieved with coverage so both transmittance and specular contribution vanishes
+					UMaterialExpressionConstant* ConstantZero = NewObject<UMaterialExpressionConstant>(this);
+					ConstantZero->R = 0.0f;
+					TopSlabBSDF->GetInput(0)->Connect(0, ConstantZero);							// BaseColor = 0 to only feature absorption, no scattering
+
+					// Now setup the mean free path with a hard coded transmittance of 0.75 when viewing the surface perpendicularly
+					UMaterialExpressionConstant* Constant075 = NewObject<UMaterialExpressionConstant>(this);
+					Constant075->R = 0.75f;
+					UMaterialExpressionStrataTransmittanceToMFP* TransToMDFP = NewObject<UMaterialExpressionStrataTransmittanceToMFP>(this);
+					TransToMDFP->GetInput(0)->Connect(0, Constant075);
+					TopSlabBSDF->GetInput(11)->Connect(0, TransToMDFP);							// MFP -> MFP
+					TopSlabBSDF->GetInput(17)->Connect(1, TransToMDFP);							// Thickness -> Thickness
+
+					// Now weight the top base material by ClearCoat
+					UMaterialExpressionStrataWeight* TopSlabBSDFWithCoverage = NewObject<UMaterialExpressionStrataWeight>(this);
+					TopSlabBSDFWithCoverage->GetInput(0)->Connect(0, TopSlabBSDF);				// TopSlabBSDF -> A
+					MoveConnectionTo(EditorOnly->ClearCoat, TopSlabBSDFWithCoverage, 1);		// ClearCoat -> Weight
+
+					UMaterialExpressionStrataVerticalLayering* VerticalLayering = NewObject<UMaterialExpressionStrataVerticalLayering>(this);
+					VerticalLayering->GetInput(0)->Connect(0, TopSlabBSDFWithCoverage);			// Top -> Top
+					VerticalLayering->GetInput(1)->Connect(0, BottomSlabBSDF);					// Bottom -> Base
+
+					EditorOnly->FrontMaterial.Connect(0, VerticalLayering);
+					bClearCoatConversionDone = true;
+					bRelinkCustomOutputNodes = false;	// We do not want that to happen in this case
 				}
-
-				ConvertNode->ConvertedStrataMaterialInfo.SetShadingModelFromExpression(true);
-				AddStrataShadingModelFromMaterialShadingModel(ConvertNode->ConvertedStrataMaterialInfo, ShadingModels);
-				check(ConvertNode->ConvertedStrataMaterialInfo.CountShadingModels() >= 1);
 			}
-			else
-			{
-				check(!bHasShadingModelMixture);
+			
+			if (!bClearCoatConversionDone)
+			{
+				ConvertNode = NewObject<UMaterialExpressionStrataLegacyConversion>(this);
+				ConvertNode->SubsurfaceProfile = bRequireNoSubsurfaceProfile ? nullptr : SubsurfaceProfile;
+				MoveConnectionTo(EditorOnly->BaseColor, ConvertNode, 0);
+				MoveConnectionTo(EditorOnly->Metallic, ConvertNode, 1);
+				MoveConnectionTo(EditorOnly->Specular, ConvertNode, 2);
+				MoveConnectionTo(EditorOnly->Roughness, ConvertNode, 3);
+				MoveConnectionTo(EditorOnly->Anisotropy, ConvertNode, 4);
+				MoveConnectionTo(EditorOnly->EmissiveColor, ConvertNode, 5);
+				CopyConnectionTo(EditorOnly->Normal, ConvertNode, 6);
+				MoveConnectionTo(EditorOnly->Tangent, ConvertNode, 7);
+				MoveConnectionTo(EditorOnly->SubsurfaceColor, ConvertNode, 8);
+				MoveConnectionTo(EditorOnly->ClearCoat, ConvertNode, 9);
+				MoveConnectionTo(EditorOnly->ClearCoatRoughness, ConvertNode, 10);
+				MoveConnectionTo(EditorOnly->Opacity, ConvertNode, 11);
+				bRelinkCustomOutputNodes = true;
+			
+				// Shading Model
+				// * either use the shader graph expression 
+				// * or add a constant shading model
+				// 
+				// Note: store this conversion type(s) into ConvertedStrataMaterialInfo for having more context when 
+				// rebuilding the final Shading model (see RebuildShadingModelField())
+				if (ShadingModel == MSM_FromMaterialExpression)
+				{
+					check(EditorOnly->ShadingModelFromMaterialExpression.IsConnected());
+
+					// Reconnect the shading model expression
+					MoveConnectionTo(EditorOnly->ShadingModelFromMaterialExpression, ConvertNode, 18);
 
-				// Add constant for the shading model
-				UMaterialExpressionConstant* ShadingModelNode = NewObject<UMaterialExpressionConstant>(this);
-				ShadingModelNode->SetParameterName(FName(TEXT("ConstantShadingModel")));
-				ShadingModelNode->R = ShadingModel;
-				ConvertNode->ShadingModel.Connect(0, ShadingModelNode);
+					// Store strata shading model of the converted material. 
+					GatherCustomNodes();
+					if (SingleLayerWaterOutput)
+					{
+						ConvertNode->ConvertedStrataMaterialInfo.AddShadingModel(SSM_SingleLayerWater);
+					}
 
-				AddStrataShadingModelFromMaterialShadingModel(ConvertNode->ConvertedStrataMaterialInfo, ShadingModels);
-				check(ConvertNode->ConvertedStrataMaterialInfo.CountShadingModels() == 1);
+					ConvertNode->ConvertedStrataMaterialInfo.SetShadingModelFromExpression(true);
+					AddStrataShadingModelFromMaterialShadingModel(ConvertNode->ConvertedStrataMaterialInfo, ShadingModels);
+					check(ConvertNode->ConvertedStrataMaterialInfo.CountShadingModels() >= 1);
+				}
+				else
+				{
+					check(!bHasShadingModelMixture);
+				
+					// Add constant for the shading model
+					UMaterialExpressionConstant* ShadingModelNode = NewObject<UMaterialExpressionConstant>(this);
+					ShadingModelNode->SetParameterName(FName(TEXT("ConstantShadingModel")));
+					ShadingModelNode->R = ShadingModel;
+					ConvertNode->ShadingModel.Connect(0, ShadingModelNode);
+
+					AddStrataShadingModelFromMaterialShadingModel(ConvertNode->ConvertedStrataMaterialInfo, ShadingModels);
+					check(ConvertNode->ConvertedStrataMaterialInfo.CountShadingModels() == 1);
+				}
+
+				EditorOnly->FrontMaterial.Connect(0, ConvertNode);
 			}
 
-			EditorOnly->FrontMaterial.Connect(0, ConvertNode);
 			bInvalidateShader = true;
 		}
 		else if (MaterialDomain == MD_Volume)
@@ -3293,10 +3353,6 @@ void UMaterial::ConvertMaterialToStrataMaterial()
 			MoveConnectionTo(*SingleLayerWaterOutput->GetInput(2), ConvertNode, 15); // WaterPhaseG
 			MoveConnectionTo(*SingleLayerWaterOutput->GetInput(3), ConvertNode, 16); // ColorScaleBehindWater
 		}
-		if (ClearCoatOutput)
-		{
-			CopyConnectionTo(*ClearCoatOutput->GetInput(0), ConvertNode, 17);		 // ClearCoatNormal
-		}
 	}
 
 	if (bInvalidateShader)
