Fixes for new anim BP re-instancing path

Marked a number of node-internal properties as 'compiler generated defaults'.
Made sure that CPFUO skips 'compiler generated defaults' when re-instancing objects (previously it only did this in one of the reinstancing calls, not the main one).
Tweaked/added custom reinstancing logic for linked layers & graphs to ensure that linked instances get correctly hooked back up. This now only gets called on the already-reinstanced UAnimInstance rather than all of them

diff --git a/Engine/Source/Editor/UnrealEd/Private/Kismet2/KismetReinstanceUtilities.cpp b/Engine/Source/Editor/UnrealEd/Private/Kismet2/KismetReinstanceUtilities.cpp
--- a/Engine/Source/Editor/UnrealEd/Private/Kismet2/KismetReinstanceUtilities.cpp
+++ b/Engine/Source/Editor/UnrealEd/Private/Kismet2/KismetReinstanceUtilities.cpp
@@ -1918,6 +1918,7 @@ static void ReplaceObjectHelper(UObject*& OldObject, UClass* OldClass, UObject*&
 	// Copy property values
 	Options.bNotifyObjectReplacement = true;
 	Options.bDontClearReferenceIfNewerClassExists = true;
+	Options.bSkipCompilerGeneratedDefaults = true;
 	UEditorEngine::CopyPropertiesForUnrelatedObjects(OldObject, NewUObject, Options);
 	// Generate new subobjects
 	InstancedPropertyUtils::FArchiveInsertInstancedSubObjects InstancedSubObjSpawner(NewUObject, InstancedPropertyMap);
diff --git a/Engine/Source/Runtime/Engine/Classes/Animation/AnimInstance.h b/Engine/Source/Runtime/Engine/Classes/Animation/AnimInstance.h
--- a/Engine/Source/Runtime/Engine/Classes/Animation/AnimInstance.h
+++ b/Engine/Source/Runtime/Engine/Classes/Animation/AnimInstance.h
@@ -1641,6 +1641,6 @@ public:
 
 #if WITH_EDITOR
 	// Helper function to handle reinstancing in editor
-	virtual void HandleAnimInstanceReplaced();
+	virtual void HandleAnimInstanceReplaced(const TMap<UObject*, UObject*>& OldToNewInstanceMap);
 #endif
 };
diff --git a/Engine/Source/Runtime/Engine/Classes/Animation/AnimNodeBase.h b/Engine/Source/Runtime/Engine/Classes/Animation/AnimNodeBase.h
--- a/Engine/Source/Runtime/Engine/Classes/Animation/AnimNodeBase.h
+++ b/Engine/Source/Runtime/Engine/Classes/Animation/AnimNodeBase.h
@@ -700,12 +700,12 @@ protected:
 
 public:
 	/** Serialized link ID, used to build the non-serialized pointer map. */
-	UPROPERTY()
+	UPROPERTY(meta=(BlueprintCompilerGeneratedDefaults))
 	int32 LinkID;
 
 #if WITH_EDITORONLY_DATA
 	/** The source link ID, used for debug visualization. */
-	UPROPERTY()
+	UPROPERTY(meta=(BlueprintCompilerGeneratedDefaults))
 	int32 SourceLinkID;
 #endif
 
diff --git a/Engine/Source/Runtime/Engine/Classes/Animation/AnimNode_CustomProperty.h b/Engine/Source/Runtime/Engine/Classes/Animation/AnimNode_CustomProperty.h
--- a/Engine/Source/Runtime/Engine/Classes/Animation/AnimNode_CustomProperty.h
+++ b/Engine/Source/Runtime/Engine/Classes/Animation/AnimNode_CustomProperty.h
@@ -48,11 +48,11 @@ public:
 
 protected:
 	/** List of source properties to use, 1-1 with Dest names below, built by the compiler */
-	UPROPERTY()
+	UPROPERTY(meta=(BlueprintCompilerGeneratedDefaults))
 	TArray<FName> SourcePropertyNames;
 
 	/** List of destination properties to use, 1-1 with Source names above, built by the compiler */
-	UPROPERTY()
+	UPROPERTY(meta=(BlueprintCompilerGeneratedDefaults))
 	TArray<FName> DestPropertyNames;
 
 	/** This is the actual instance allocated at runtime that will run. Set by child class. */
diff --git a/Engine/Source/Runtime/Engine/Classes/Animation/AnimNode_LinkedAnimGraph.h b/Engine/Source/Runtime/Engine/Classes/Animation/AnimNode_LinkedAnimGraph.h
--- a/Engine/Source/Runtime/Engine/Classes/Animation/AnimNode_LinkedAnimGraph.h
+++ b/Engine/Source/Runtime/Engine/Classes/Animation/AnimNode_LinkedAnimGraph.h
@@ -120,7 +120,7 @@ protected:
 	}
 
 #if WITH_EDITOR
-	virtual void HandleObjectsReplaced(const TMap<UObject*, UObject*>& OldToNewInstanceMap) override;
+	virtual void HandleAnimInstanceReplaced(UAnimInstance* InNewInstance, const TMap<UObject*, UObject*>& OldToNewInstanceMap);
 #endif	// #if WITH_EDITOR
 
 	/** Link up pose links dynamically with linked instance */
diff --git a/Engine/Source/Runtime/Engine/Classes/Animation/AnimNode_LinkedAnimLayer.h b/Engine/Source/Runtime/Engine/Classes/Animation/AnimNode_LinkedAnimLayer.h
--- a/Engine/Source/Runtime/Engine/Classes/Animation/AnimNode_LinkedAnimLayer.h
+++ b/Engine/Source/Runtime/Engine/Classes/Animation/AnimNode_LinkedAnimLayer.h
@@ -55,6 +55,7 @@ public:
 #if WITH_EDITOR
 	// Event fired when the instance we are running has changed
 	FSimpleMulticastDelegate& OnInstanceChanged() { return OnInstanceChangedEvent; }
+	virtual void HandleAnimInstanceReplaced(UAnimInstance* InNewInstance, const TMap<UObject*, UObject*>& OldToNewInstanceMap) override;
 #endif
 
 protected:
diff --git a/Engine/Source/Runtime/Engine/Classes/Animation/AnimNode_Root.h b/Engine/Source/Runtime/Engine/Classes/Animation/AnimNode_Root.h
--- a/Engine/Source/Runtime/Engine/Classes/Animation/AnimNode_Root.h
+++ b/Engine/Source/Runtime/Engine/Classes/Animation/AnimNode_Root.h
@@ -19,7 +19,7 @@ struct ENGINE_API FAnimNode_Root : public FAnimNode_Base
 #if WITH_EDITORONLY_DATA
 protected:
 	/** The name of this root node, used to identify the output of this graph. Filled in by the compiler, propagated from the parent graph. */
-	UPROPERTY(meta=(FoldProperty))
+	UPROPERTY(meta=(FoldProperty, BlueprintCompilerGeneratedDefaults))
 	FName Name;
 
 	/** The group of this root node, used to group this output with others when used in a layer. */
diff --git a/Engine/Source/Runtime/Engine/Classes/Animation/AnimNode_UseCachedPose.h b/Engine/Source/Runtime/Engine/Classes/Animation/AnimNode_UseCachedPose.h
--- a/Engine/Source/Runtime/Engine/Classes/Animation/AnimNode_UseCachedPose.h
+++ b/Engine/Source/Runtime/Engine/Classes/Animation/AnimNode_UseCachedPose.h
@@ -17,7 +17,7 @@ struct ENGINE_API FAnimNode_UseCachedPose : public FAnimNode_Base
 	FPoseLink LinkToCachingNode;
 
 	/** Intentionally not exposed, set by AnimBlueprintCompiler */
-	UPROPERTY()
+	UPROPERTY(meta=(BlueprintCompilerGeneratedDefaults))
 	FName CachePoseName;
 
 public:	
diff --git a/Engine/Source/Runtime/Engine/Private/Animation/AnimInstance.cpp b/Engine/Source/Runtime/Engine/Private/Animation/AnimInstance.cpp
--- a/Engine/Source/Runtime/Engine/Private/Animation/AnimInstance.cpp
+++ b/Engine/Source/Runtime/Engine/Private/Animation/AnimInstance.cpp
@@ -3800,8 +3800,31 @@ void UAnimInstance::QueueRootMotionBlend(const FTransform& RootTransform, const
 }
 
 #if WITH_EDITOR
-void UAnimInstance::HandleAnimInstanceReplaced()
+void UAnimInstance::HandleAnimInstanceReplaced(const TMap<UObject*, UObject*>& OldToNewInstanceMap)
 {
+	if (IAnimClassInterface* AnimBlueprintClass = IAnimClassInterface::GetFromClass(GetClass()))
+	{
+		TArray<UObject*> Values;
+		OldToNewInstanceMap.GenerateValueArray(Values);
+		if(Values.Contains(this))
+		{
+			// Re-link graphs & layers if they are targets of a re-instance.
+			// As this gets called post-CPFUO and references are already replaced we need to only do this for 'new'
+			// instances that have just been replaced
+			for(const FStructProperty* LinkedGraphProperty : AnimBlueprintClass->GetLinkedAnimGraphNodeProperties())
+			{
+				FAnimNode_LinkedAnimGraph* LinkedAnimGraph = LinkedGraphProperty->ContainerPtrToValuePtr<FAnimNode_LinkedAnimGraph>(this);
+				LinkedAnimGraph->HandleAnimInstanceReplaced(this, OldToNewInstanceMap);
+			}
+
+			for(const FStructProperty* LinkedLayerProperty : AnimBlueprintClass->GetLinkedAnimLayerNodeProperties())
+			{
+				FAnimNode_LinkedAnimLayer* LinkedAnimLayer = LinkedLayerProperty->ContainerPtrToValuePtr<FAnimNode_LinkedAnimLayer>(this);
+				LinkedAnimLayer->HandleAnimInstanceReplaced(this, OldToNewInstanceMap);
+			}
+		}
+	}
+	
 	RecalcRequiredBones();
 	
 	// Reinit proxy
diff --git a/Engine/Source/Runtime/Engine/Private/Animation/AnimNode_LinkedAnimGraph.cpp b/Engine/Source/Runtime/Engine/Private/Animation/AnimNode_LinkedAnimGraph.cpp
--- a/Engine/Source/Runtime/Engine/Private/Animation/AnimNode_LinkedAnimGraph.cpp
+++ b/Engine/Source/Runtime/Engine/Private/Animation/AnimNode_LinkedAnimGraph.cpp
@@ -472,22 +472,13 @@ void FAnimNode_LinkedAnimGraph::RequestBlend(const IAnimClassInterface* PriorAni
 }
 
 #if WITH_EDITOR
-void FAnimNode_LinkedAnimGraph::HandleObjectsReplaced(const TMap<UObject*, UObject*>& OldToNewInstanceMap)
+void FAnimNode_LinkedAnimGraph::HandleAnimInstanceReplaced(UAnimInstance* InNewInstance, const TMap<UObject*, UObject*>& OldToNewInstanceMap)
 {
-	FAnimNode_CustomProperty::HandleObjectsReplaced(OldToNewInstanceMap);
-	
 	if (UAnimInstance* ThisTargetInstance = GetTargetInstance<UAnimInstance>())
 	{
-		UObject* const* ReinstancedTarget = OldToNewInstanceMap.Find(ThisTargetInstance);
-		if (ReinstancedTarget)
-		{
-			if(UAnimInstance* ReinstancedTargetInstance = Cast<UAnimInstance>(*ReinstancedTarget))
-			{
-				DynamicUnlink(ThisTargetInstance);
-				SetTargetInstance(ReinstancedTargetInstance);
-				DynamicLink(ReinstancedTargetInstance);
-			}
-		}
+		// If we have a target, unlink & re-link
+		DynamicUnlink(InNewInstance);
+		DynamicLink(InNewInstance);
 	}
 }
 #endif
diff --git a/Engine/Source/Runtime/Engine/Private/Animation/AnimNode_LinkedAnimLayer.cpp b/Engine/Source/Runtime/Engine/Private/Animation/AnimNode_LinkedAnimLayer.cpp
--- a/Engine/Source/Runtime/Engine/Private/Animation/AnimNode_LinkedAnimLayer.cpp
+++ b/Engine/Source/Runtime/Engine/Private/Animation/AnimNode_LinkedAnimLayer.cpp
@@ -138,4 +138,20 @@ void FAnimNode_LinkedAnimLayer::InitializeProperties(const UObject* InSourceInst
 			}
 		}
 	}
-}
\ No newline at end of file
+}
+
+#if WITH_EDITOR
+void FAnimNode_LinkedAnimLayer::HandleAnimInstanceReplaced(UAnimInstance* InNewInstance, const TMap<UObject*, UObject*>& OldToNewInstanceMap)
+{
+	UAnimInstance* CurrentTargetInstance = GetTargetInstance<UAnimInstance>();
+	if(CurrentTargetInstance && (Interface.Get() == nullptr || InstanceClass.Get() == nullptr))
+	{
+		InitializeSelfLayer(CurrentTargetInstance);
+	}
+	else
+	{
+		Super::HandleAnimInstanceReplaced(InNewInstance, OldToNewInstanceMap);
+	}
+}
+
+#endif
\ No newline at end of file
diff --git a/Engine/Source/Runtime/Engine/Private/Components/SkeletalMeshComponent.cpp b/Engine/Source/Runtime/Engine/Private/Components/SkeletalMeshComponent.cpp
--- a/Engine/Source/Runtime/Engine/Private/Components/SkeletalMeshComponent.cpp
+++ b/Engine/Source/Runtime/Engine/Private/Components/SkeletalMeshComponent.cpp
@@ -4621,7 +4621,7 @@ void USkeletalMeshComponent::HandleObjectsReplaced(const TMap<UObject*, UObject*
 			{
 				if (UAnimInstance* ReinstancedTargetInstance = Cast<UAnimInstance>(*ReinstancedTarget))
 				{
-					ReinstancedTargetInstance->HandleAnimInstanceReplaced();
+					ReinstancedTargetInstance->HandleAnimInstanceReplaced(OldToNewInstanceMap);
 					bReinstancedAnimInstance = true;
 				}
 			}
