Fixed editable propoerties when Strata is enabled (strata should not rely on ShadingModels)

diff --git a/Engine/Source/Runtime/Engine/Private/Materials/Material.cpp b/Engine/Source/Runtime/Engine/Private/Materials/Material.cpp
--- a/Engine/Source/Runtime/Engine/Private/Materials/Material.cpp
+++ b/Engine/Source/Runtime/Engine/Private/Materials/Material.cpp
@@ -3958,6 +3958,7 @@ bool UMaterial::CanEditChange(const FProperty* InProperty) const
 	{
 		const UMaterialEditorOnlyData* EditorOnly = GetEditorOnlyData();
 		FString PropertyName = InProperty->GetName();
+		const bool bStrataEnabled = Engine_IsStrataEnabled();
 
 		if (PropertyName == GET_MEMBER_NAME_STRING_CHECKED(UMaterial, PhysMaterial) || PropertyName == GET_MEMBER_NAME_STRING_CHECKED(UMaterial, PhysMaterialMask))
 		{
@@ -4026,17 +4027,17 @@ bool UMaterial::CanEditChange(const FProperty* InProperty) const
 
 		if (PropertyName == GET_MEMBER_NAME_STRING_CHECKED(UMaterial, BlendMode))
 		{
-			return !Engine_IsStrataEnabled() && (MaterialDomain == MD_DeferredDecal || MaterialDomain == MD_Surface || MaterialDomain == MD_Volume || MaterialDomain == MD_UI || (MaterialDomain == MD_PostProcess && BlendableOutputAlpha));
+			return !bStrataEnabled && (MaterialDomain == MD_DeferredDecal || MaterialDomain == MD_Surface || MaterialDomain == MD_Volume || MaterialDomain == MD_UI || (MaterialDomain == MD_PostProcess && BlendableOutputAlpha));
 		}
 
 		if (PropertyName == GET_MEMBER_NAME_STRING_CHECKED(UMaterial, StrataBlendMode))
 		{
-			return Engine_IsStrataEnabled() && ((MaterialDomain != MD_PostProcess && MaterialDomain != MD_LightFunction && MaterialDomain != MD_Volume) || (MaterialDomain == MD_PostProcess && BlendableOutputAlpha));
+			return bStrataEnabled && ((MaterialDomain != MD_PostProcess && MaterialDomain != MD_LightFunction && MaterialDomain != MD_Volume) || (MaterialDomain == MD_PostProcess && BlendableOutputAlpha));
 		}
 	
 		if (PropertyName == GET_MEMBER_NAME_STRING_CHECKED(UMaterial, ShadingModel))
 		{
-			return !Engine_IsStrataEnabled() && MaterialDomain == MD_Surface;
+			return !bStrataEnabled && MaterialDomain == MD_Surface;
 		}
 
 		if (FCString::Strncmp(*PropertyName, TEXT("bUsedWith"), 9) == 0)
@@ -4072,7 +4073,8 @@ bool UMaterial::CanEditChange(const FProperty* InProperty) const
 
 		if (PropertyName == GET_MEMBER_NAME_STRING_CHECKED(UMaterial, bIsSky))
 		{
-			return MaterialDomain != MD_DeferredDecal && GetShadingModels().IsUnlit() && (BlendMode == BLEND_Opaque || BlendMode == BLEND_Masked);
+			return (!bStrataEnabled && (MaterialDomain != MD_DeferredDecal && GetShadingModels().IsUnlit() && !IsTranslucentBlendMode(BlendMode)))
+				|| (bStrataEnabled && (MaterialDomain != MD_DeferredDecal && !IsTranslucentBlendMode(StrataBlendMode)));
 		}
 
 		if (PropertyName == GET_MEMBER_NAME_STRING_CHECKED(UMaterial, TranslucencyLightingMode)
@@ -4085,7 +4087,8 @@ bool UMaterial::CanEditChange(const FProperty* InProperty) const
 			|| PropertyName == GET_MEMBER_NAME_STRING_CHECKED(UMaterial, TranslucentMultipleScatteringExtinction)
 			|| PropertyName == GET_MEMBER_NAME_STRING_CHECKED(UMaterial, TranslucentShadowStartOffset))
 		{
-			return MaterialDomain != MD_DeferredDecal && IsTranslucentBlendMode(BlendMode) && GetShadingModels().IsLit();
+			return (!bStrataEnabled && (MaterialDomain != MD_DeferredDecal && IsTranslucentBlendMode(BlendMode) && GetShadingModels().IsLit()))
+				|| (bStrataEnabled && (MaterialDomain != MD_DeferredDecal && IsTranslucentBlendMode(StrataBlendMode)));
 		}
 
 		if (PropertyName == GET_MEMBER_NAME_STRING_CHECKED(UMaterial, SubsurfaceProfile))
diff --git a/Engine/Source/Runtime/Engine/Public/MaterialShared.h b/Engine/Source/Runtime/Engine/Public/MaterialShared.h
--- a/Engine/Source/Runtime/Engine/Public/MaterialShared.h
+++ b/Engine/Source/Runtime/Engine/Public/MaterialShared.h
@@ -2724,6 +2724,14 @@ inline bool IsTranslucentBlendMode(enum EBlendMode BlendMode)
 	return BlendMode != BLEND_Opaque && BlendMode != BLEND_Masked;
 }
 
+/**
+ * @return True if StrataBlendMode is translucent (should be part of the translucent rendering).
+ */
+inline bool IsTranslucentBlendMode(enum EStrataBlendMode StrataBlendMode)
+{
+	return StrataBlendMode != SBM_Opaque && StrataBlendMode != SBM_Masked;
+}
+
 /**
  * Implementation of the FMaterial interface for a UMaterial or UMaterialInstance.
  */
