Add USkinnedAsset class to sit over USkeletalMesh, capturing common skinned mesh functionalities

diff --git a/Engine/Source/Runtime/Engine/Classes/Components/SkinnedMeshComponent.h b/Engine/Source/Runtime/Engine/Classes/Components/SkinnedMeshComponent.h
--- a/Engine/Source/Runtime/Engine/Classes/Components/SkinnedMeshComponent.h
+++ b/Engine/Source/Runtime/Engine/Classes/Components/SkinnedMeshComponent.h
@@ -28,6 +28,7 @@ struct FSkelMeshRenderSection;
 class FPositionVertexBuffer;
 class UMeshDeformer;
 class UMeshDeformerInstance;
+class USkinnedAsset;
 
 DECLARE_DELEGATE_OneParam(FOnAnimUpdateRateParamsCreated, FAnimUpdateRateParameters*)
 DECLARE_MULTICAST_DELEGATE_ThreeParams(FOnTickPose, USkinnedMeshComponent* /*SkinnedMeshComponent*/, float /*DeltaTime*/, bool /*bNeedsValidRootMotion*/)
@@ -359,6 +360,7 @@ protected:
 	TSharedPtr<FSkelMeshRefPoseOverride> RefPoseOverride;
 
 public:
+	USkinnedAsset* GetSkinnedAsset() const;
 
 	const TArray<int32>& GetMasterBoneMap() const { return MasterBoneMap; }
 
diff --git a/Engine/Source/Runtime/Engine/Classes/Engine/SkeletalMesh.h b/Engine/Source/Runtime/Engine/Classes/Engine/SkeletalMesh.h
--- a/Engine/Source/Runtime/Engine/Classes/Engine/SkeletalMesh.h
+++ b/Engine/Source/Runtime/Engine/Classes/Engine/SkeletalMesh.h
@@ -32,7 +32,7 @@
 #include "Animation/MorphTarget.h"
 #include "Engine/StreamableRenderAsset.h"
 #include "PerQualityLevelProperties.h"
-
+#include "SkinnedAsset.h"
 #include "SkeletalMesh.generated.h"
 
 class UAnimInstance;
@@ -763,7 +763,7 @@ namespace NSSkeletalMeshSourceFileLabels
  * @see https://docs.unrealengine.com/latest/INT/Engine/Content/Types/SkeletalMeshes/
  */
 UCLASS(hidecategories=Object, BlueprintType)
-class ENGINE_API USkeletalMesh : public UStreamableRenderAsset, public IInterface_CollisionDataProvider, public IInterface_AssetUserData, public INodeMappingProviderInterface, public IInterface_AsyncCompilation
+class ENGINE_API USkeletalMesh : public USkinnedAsset, public IInterface_CollisionDataProvider, public IInterface_AssetUserData, public INodeMappingProviderInterface, public IInterface_AsyncCompilation
 {
 	GENERATED_UCLASS_BODY()
 
@@ -928,8 +928,8 @@ public:
 	virtual void ClearAllCachedCookedPlatformData() override;
 #endif
 
-	/** Get the data to use for rendering. */
-	FORCEINLINE FSkeletalMeshRenderData* GetResourceForRendering() const 
+	/** Get the data to use for rendering. USkinnedAsset interface. */
+	virtual FSkeletalMeshRenderData* GetResourceForRendering() const override
 	{
 		WaitUntilAsyncPropertyReleased(ESkeletalMeshAsyncProperties::SkeletalMeshRenderData);
 		PRAGMA_DISABLE_DEPRECATION_WARNINGS
@@ -949,7 +949,8 @@ public:
 		PRAGMA_ENABLE_DEPRECATION_WARNINGS
 	}
 
-	USkeleton* GetSkeleton()
+	/** USkinnedAsset interface. */
+	virtual USkeleton* GetSkeleton() override
 	{
 		WaitUntilAsyncPropertyReleased(ESkeletalMeshAsyncProperties::Skeleton);
 		PRAGMA_DISABLE_DEPRECATION_WARNINGS
@@ -957,8 +958,9 @@ public:
 		PRAGMA_ENABLE_DEPRECATION_WARNINGS
 	}
 
+	/** USkinnedAsset interface. */
 	UFUNCTION(BlueprintGetter)
-	const USkeleton* GetSkeleton() const
+	virtual const USkeleton* GetSkeleton() const override
 	{
 		WaitUntilAsyncPropertyReleased(ESkeletalMeshAsyncProperties::Skeleton, EAsyncPropertyLockType::ReadOnly);
 		PRAGMA_DISABLE_DEPRECATION_WARNINGS
@@ -1036,9 +1038,9 @@ protected:
 	}
 
 public:
-	/** Get the extended bounds of this mesh (imported bounds plus bounds extension) */
+	/** Get the extended bounds of this mesh (imported bounds plus bounds extension). USkinnedAsset interface. */
 	UFUNCTION(BlueprintCallable, Category = Mesh)
-	FBoxSphereBounds GetBounds() const;
+	virtual FBoxSphereBounds GetBounds() const override;
 
 	/** Get the original imported bounds of the skel mesh */
 	UFUNCTION(BlueprintCallable, Category = Mesh)
@@ -1113,7 +1115,8 @@ public:
 		PRAGMA_ENABLE_DEPRECATION_WARNINGS
 	}
 
-	TArray<FSkeletalMaterial>& GetMaterials()
+	/** USkinnedAsset interface. */
+	virtual TArray<FSkeletalMaterial>& GetMaterials() override
 	{
 		WaitUntilAsyncPropertyReleased(ESkeletalMeshAsyncProperties::Materials);
 		PRAGMA_DISABLE_DEPRECATION_WARNINGS
@@ -1121,8 +1124,9 @@ public:
 		PRAGMA_ENABLE_DEPRECATION_WARNINGS
 	}
 
+	/** USkinnedAsset interface. */
 	UFUNCTION(BlueprintGetter)
-	const TArray<FSkeletalMaterial>& GetMaterials() const
+	virtual const TArray<FSkeletalMaterial>& GetMaterials() const override
 	{
 		WaitUntilAsyncPropertyReleased(ESkeletalMeshAsyncProperties::Materials, EAsyncPropertyLockType::ReadOnly);
 		PRAGMA_DISABLE_DEPRECATION_WARNINGS
@@ -1229,11 +1233,15 @@ public:
 	/** Check the QualitLevel property is enabled for MinLod. */
 	bool IsMinLodQualityLevelEnable() const;
 
+	/** USkinnedAsset interface */
+	virtual int32 GetPlatformMinLODIdx(const ITargetPlatform* TargetPlatform) const override;
+	virtual void SetSkinWeightProfilesData(int32 LODIndex, FSkinWeightProfilesData& SkinWeightProfilesData) override;
+
 	static void OnLodStrippingQualityLevelChanged(IConsoleVariable* Variable);
 
 	/*Choose either PerPlatform or PerQuality override. Note: Enable PerQuality override in the Project Settings/ General Settings/ UseSkeletalMeshMinLODPerQualityLevels*/
-	int32 GetMinLodIdx(bool bForceLowestLODIdx = false) const;
-	int32 GetDefaultMinLod() const;
+	virtual int32 GetMinLodIdx(bool bForceLowestLODIdx = false) const override;
+	virtual int32 GetDefaultMinLod() const override;
 	void SetMinLodIdx(int32 InMinLOD);
 
 	/** Minimum LOD to render. Can be overridden per component as well as set here for all mesh instances here */
@@ -1276,7 +1284,8 @@ public:
 		PRAGMA_ENABLE_DEPRECATION_WARNINGS
 	}
 
-	const FPerPlatformBool& GetDisableBelowMinLodStripping() const
+	/** USkinnedAsset interface. */
+	virtual const FPerPlatformBool& GetDisableBelowMinLodStripping() const override
 	{
 		WaitUntilAsyncPropertyReleased(ESkeletalMeshAsyncProperties::DisableBelowMinLodStripping, EAsyncPropertyLockType::ReadOnly);
 		PRAGMA_DISABLE_DEPRECATION_WARNINGS
@@ -1458,14 +1467,14 @@ public:
 	void SetLODSettings(USkeletalMeshLODSettings* InLODSettings);
 
 #if WITH_EDITOR
-	/** Get whether this mesh should use LOD streaming for the given platform. Do not use bSupportLODStreaming directly. Call this method instead. */
-	bool GetEnableLODStreaming(const class ITargetPlatform* TargetPlatform) const;
+	/** Get whether this mesh should use LOD streaming for the given platform. Do not use bSupportLODStreaming directly. Call this method instead. USkinnedAsset Interface. */
+	virtual bool GetEnableLODStreaming(const class ITargetPlatform* TargetPlatform) const override;
 
-	/** Get the maximum number of LODs that can be streamed. Do not use MaxNumStreamedLODs directly. Call this method instead. */
-	int32 GetMaxNumStreamedLODs(const class ITargetPlatform* TargetPlatform) const;
+	/** Get the maximum number of LODs that can be streamed. Do not use MaxNumStreamedLODs directly. Call this method instead. USkinnedAsset Interface. */
+	virtual int32 GetMaxNumStreamedLODs(const class ITargetPlatform* TargetPlatform) const override;
 
-	/** Get the maximum number of optional LODs. Do not use MaxNumOptionalLODs directly. Call this method instead. */
-	int32 GetMaxNumOptionalLODs(const class ITargetPlatform* TargetPlatform) const;
+	/** Get the maximum number of optional LODs. Do not use MaxNumOptionalLODs directly. Call this method instead. USkinnedAsset Interface. */
+	virtual int32 GetMaxNumOptionalLODs(const class ITargetPlatform* TargetPlatform) const override;
 #endif
 
 	UFUNCTION(BlueprintSetter)
@@ -1577,7 +1586,8 @@ public:
 		PRAGMA_ENABLE_DEPRECATION_WARNINGS
 	}
 
-	bool GetHasVertexColors() const
+	/** Return whether or not the mesh has vertex colors. USkinnedAsset interface. */
+	virtual bool GetHasVertexColors() const override
 	{
 		WaitUntilAsyncPropertyReleased(ESkeletalMeshAsyncProperties::HasVertexColors, EAsyncPropertyLockType::ReadOnly);
 		PRAGMA_DISABLE_DEPRECATION_WARNINGS
@@ -1723,8 +1733,9 @@ public:
 		PRAGMA_ENABLE_DEPRECATION_WARNINGS
 	}
 
+	/** USkinnedAsset interface. */
 	UFUNCTION(BlueprintGetter)
-	class UPhysicsAsset* GetPhysicsAsset() const
+	virtual class UPhysicsAsset* GetPhysicsAsset() const override
 	{
 		WaitUntilAsyncPropertyReleased(ESkeletalMeshAsyncProperties::PhysicsAsset);
 		PRAGMA_DISABLE_DEPRECATION_WARNINGS
@@ -1755,8 +1766,9 @@ public:
 		PRAGMA_ENABLE_DEPRECATION_WARNINGS
 	}
 
+	/** USkinnedAsset interface. */
 	UFUNCTION(BlueprintGetter)
-	class UPhysicsAsset* GetShadowPhysicsAsset() const
+	virtual class UPhysicsAsset* GetShadowPhysicsAsset() const override
 	{
 		WaitUntilAsyncPropertyReleased(ESkeletalMeshAsyncProperties::ShadowPhysicsAsset);
 		PRAGMA_DISABLE_DEPRECATION_WARNINGS
@@ -2104,7 +2116,8 @@ public:
 	UPROPERTY(EditAnywhere, Category = RayTracing)
 	uint8 bSupportRayTracing : 1;
 
-	bool GetSupportRayTracing() const
+	/** USkinnedAsset interface. */
+	virtual bool GetSupportRayTracing() const override
 	{
 		WaitUntilAsyncPropertyReleased(ESkeletalMeshAsyncProperties::bSupportRayTracing, EAsyncPropertyLockType::ReadOnly);
 		PRAGMA_DISABLE_DEPRECATION_WARNINGS
@@ -2127,7 +2140,8 @@ public:
 	UPROPERTY(EditAnywhere, Category = RayTracing)
 	int32 RayTracingMinLOD;
 
-	int32 GetRayTracingMinLOD() const
+	/** USkinnedAsset interface. */
+	virtual int32 GetRayTracingMinLOD() const override
 	{
 		WaitUntilAsyncPropertyReleased(ESkeletalMeshAsyncProperties::RayTracingMinLOD, EAsyncPropertyLockType::ReadOnly);
 		PRAGMA_DISABLE_DEPRECATION_WARNINGS
@@ -2178,7 +2192,8 @@ public:
 		PRAGMA_ENABLE_DEPRECATION_WARNINGS
 	}
 
-	TArray<TObjectPtr<UMorphTarget>>& GetMorphTargets()
+	/** USkinnedAsset interface. */
+	virtual TArray<TObjectPtr<UMorphTarget>>& GetMorphTargets() override
 	{
 		WaitUntilAsyncPropertyReleased(ESkeletalMeshAsyncProperties::MorphTargets);
 		PRAGMA_DISABLE_DEPRECATION_WARNINGS
@@ -2186,8 +2201,9 @@ public:
 		PRAGMA_ENABLE_DEPRECATION_WARNINGS
 	}
 
+	/** USkinnedAsset interface. */
 	UFUNCTION(BlueprintGetter)
-	const TArray<UMorphTarget*>& GetMorphTargets() const
+	virtual const TArray<UMorphTarget*>& GetMorphTargets() const override
 	{
 		WaitUntilAsyncPropertyReleased(ESkeletalMeshAsyncProperties::MorphTargets, EAsyncPropertyLockType::ReadOnly);
 		PRAGMA_DISABLE_DEPRECATION_WARNINGS
@@ -2225,7 +2241,8 @@ public:
 		PRAGMA_ENABLE_DEPRECATION_WARNINGS
 	}
 
-	FReferenceSkeleton& GetRefSkeleton()
+	/** USkinnedAsset interface. */
+	virtual FReferenceSkeleton& GetRefSkeleton() override
 	{
 		WaitUntilAsyncPropertyReleased(ESkeletalMeshAsyncProperties::RefSkeleton);
 		PRAGMA_DISABLE_DEPRECATION_WARNINGS
@@ -2233,7 +2250,8 @@ public:
 		PRAGMA_ENABLE_DEPRECATION_WARNINGS
 	}
 
-	const FReferenceSkeleton& GetRefSkeleton() const
+	/** USkinnedAsset interface. */
+	virtual const FReferenceSkeleton& GetRefSkeleton() const override
 	{
 		WaitUntilAsyncPropertyReleased(ESkeletalMeshAsyncProperties::RefSkeleton, EAsyncPropertyLockType::ReadOnly);
 		PRAGMA_DISABLE_DEPRECATION_WARNINGS
@@ -2295,7 +2313,8 @@ public:
 		PRAGMA_ENABLE_DEPRECATION_WARNINGS
 	}
 
-	TArray<FMatrix44f>& GetRefBasesInvMatrix()
+	/** USkinnedAsset interface. */
+	virtual TArray<FMatrix44f>& GetRefBasesInvMatrix() override
 	{
 		WaitUntilAsyncPropertyReleased(ESkeletalMeshAsyncProperties::RefBasesInvMatrix);
 		PRAGMA_DISABLE_DEPRECATION_WARNINGS
@@ -2303,7 +2322,8 @@ public:
 		PRAGMA_ENABLE_DEPRECATION_WARNINGS
 	}
 
-	const TArray<FMatrix44f>& GetRefBasesInvMatrix() const
+	/** USkinnedAsset interface. */
+	virtual const TArray<FMatrix44f>& GetRefBasesInvMatrix() const override
 	{
 		WaitUntilAsyncPropertyReleased(ESkeletalMeshAsyncProperties::RefBasesInvMatrix, EAsyncPropertyLockType::ReadOnly);
 		PRAGMA_DISABLE_DEPRECATION_WARNINGS
@@ -2564,9 +2584,9 @@ public:
 #endif
 
 	/** 
-	True if this mesh LOD needs to keep it's data on CPU. 
+	True if this mesh LOD needs to keep it's data on CPU. USkinnedAsset interface.
 	*/
-	bool NeedCPUData(int32 LODIndex)const;
+	virtual bool NeedCPUData(int32 LODIndex)const override;
 
 protected:
 
@@ -2645,12 +2665,12 @@ public:
 
 	/**
 	 * Returns the UV channel data for a given material index. Used by the texture streamer.
-	 * This data applies to all lod-section using the same material.
+	 * This data applies to all lod-section using the same material. USkinnedAsset interface.
 	 *
 	 * @param MaterialIndex		the material index for which to get the data for.
 	 * @return the data, or null if none exists.
 	 */
-	const FMeshUVChannelInfo* GetUVChannelData(int32 MaterialIndex) const;
+	virtual const FMeshUVChannelInfo* GetUVChannelData(int32 MaterialIndex) const override;
 
 	/**
 	 * Computes flags for building vertex buffers.
@@ -2740,7 +2760,6 @@ public:
 	virtual bool IsPostLoadThreadSafe() const override;
 	virtual void GetAssetRegistryTags(TArray<FAssetRegistryTag>& OutTags) const override;
 	virtual FString GetDesc() override;
-	virtual FString GetDetailedInfoInternal() const override;
 	virtual void GetResourceSizeEx(FResourceSizeEx& CumulativeResourceSize) override;
 	virtual void GetPreloadDependencies(TArray<UObject*>& OutDeps) override;
 	//~ End UObject Interface.
@@ -2825,7 +2844,9 @@ public:
 	 *	Get the component orientation of a bone or socket. Transforms by parent bones.
 	 */
 	FMatrix GetComposedRefPoseMatrix( FName InBoneName ) const;
-	FMatrix GetComposedRefPoseMatrix( int32 InBoneIndex ) const;
+	/** USkinnedAsset interface. */
+	virtual FMatrix GetComposedRefPoseMatrix( int32 InBoneIndex ) const override;
+
 	UE_DEPRECATED(5.0, "Please use UMirrorDataTable for mirroring support.")
 	void InitBoneMirrorInfo();
 
@@ -2860,9 +2881,9 @@ public:
 
 	/**
 	* Returns the "active" socket list - all sockets from this mesh plus all non-duplicates from the skeleton
-	* Const ref return value as this cannot be modified externally
+	* Const ref return value as this cannot be modified externally. USkinnedAsset interface.
 	*/
-	TArray<USkeletalMeshSocket*> GetActiveSocketList() const;
+	virtual TArray<USkeletalMeshSocket*> GetActiveSocketList() const override;
 
 #if WITH_EDITOR
 	/**
@@ -2889,13 +2910,13 @@ public:
 
 	/**
 	 * Find a named MorphTarget from the MorphSets array in the SkinnedMeshComponent.
-	 * This searches the array in the same way as FindAnimSequence
+	 * This searches the array in the same way as FindAnimSequence. USkinnedAsset interface.
 	 *
 	 * @param MorphTargetName Name of MorphTarget to look for.
 	 *
 	 * @return Pointer to found MorphTarget. Returns NULL if could not find target with that name.
 	 */
-	UMorphTarget* FindMorphTarget(FName MorphTargetName) const;
+	virtual UMorphTarget* FindMorphTarget(FName MorphTargetName) const override;
 	UMorphTarget* FindMorphTargetAndIndex(FName MorphTargetName, int32& OutIndex) const;
 
 	/* Initialize morph targets and rebuild the render data */
@@ -3126,9 +3147,9 @@ public:
 	}
 
 	/*
-	 * Returns whole array of LODInfo non-const
+	 * Returns whole array of LODInfo non-const. USkinnedAsset interface.
 	 */
-	TArray<FSkeletalMeshLODInfo>& GetLODInfoArray()
+	virtual TArray<FSkeletalMeshLODInfo>& GetLODInfoArray() override
 	{ 
 		WaitUntilAsyncPropertyReleased(ESkeletalMeshAsyncProperties::LODInfo);
 		PRAGMA_DISABLE_DEPRECATION_WARNINGS
@@ -3137,9 +3158,9 @@ public:
 	}
 
 	/*
-	 * Returns whole array of LODInfo const
+	 * Returns whole array of LODInfo const. USkinnedAsset interface.
 	 */
-	const TArray<FSkeletalMeshLODInfo>& GetLODInfoArray() const 
+	virtual const TArray<FSkeletalMeshLODInfo>& GetLODInfoArray() const override
 	{
 		WaitUntilAsyncPropertyReleased(ESkeletalMeshAsyncProperties::LODInfo, EAsyncPropertyLockType::ReadOnly);
 		PRAGMA_DISABLE_DEPRECATION_WARNINGS
@@ -3148,9 +3169,9 @@ public:
 	}
 	
 	/* 
-	 * Get LODInfo of the given index non-const
+	 * Get LODInfo of the given index non-const. USkinnedAsset interface.
 	 */
-	FSkeletalMeshLODInfo* GetLODInfo(int32 Index) 
+	virtual FSkeletalMeshLODInfo* GetLODInfo(int32 Index) override
 	{
 		WaitUntilAsyncPropertyReleased(ESkeletalMeshAsyncProperties::LODInfo);
 		PRAGMA_DISABLE_DEPRECATION_WARNINGS
@@ -3159,9 +3180,9 @@ public:
 	}
 	
 	/* 
-	 * Get LODInfo of the given index const
+	 * Get LODInfo of the given index const. USkinnedAsset interface.
 	 */	
-	const FSkeletalMeshLODInfo* GetLODInfo(int32 Index) const 
+	virtual const FSkeletalMeshLODInfo* GetLODInfo(int32 Index) const override
 	{
 		WaitUntilAsyncPropertyReleased(ESkeletalMeshAsyncProperties::LODInfo, EAsyncPropertyLockType::ReadOnly);
 		PRAGMA_DISABLE_DEPRECATION_WARNINGS
@@ -3180,9 +3201,9 @@ public:
 	const USkeletalMeshLODSettings* GetDefaultLODSetting() const; 
 
 	/* 
-	 * Return true if given index's LOD is valid
+	 * Return true if given index's LOD is valid. USkinnedAsset interface.
 	 */
-	bool IsValidLODIndex(int32 Index) const 
+	virtual bool IsValidLODIndex(int32 Index) const override
 	{
 		WaitUntilAsyncPropertyReleased(ESkeletalMeshAsyncProperties::LODInfo, EAsyncPropertyLockType::ReadOnly);
 		PRAGMA_DISABLE_DEPRECATION_WARNINGS
@@ -3190,9 +3211,9 @@ public:
 		PRAGMA_ENABLE_DEPRECATION_WARNINGS
 	}
 	/* 
-	 * Returns total number of LOD
+	 * Returns total number of LOD. USkinnedAsset interface.
 	 */
-	int32 GetLODNum() const 
+	virtual int32 GetLODNum() const override
 	{
 		WaitUntilAsyncPropertyReleased(ESkeletalMeshAsyncProperties::LODInfo, EAsyncPropertyLockType::ReadOnly);
 		PRAGMA_DISABLE_DEPRECATION_WARNINGS
@@ -3200,7 +3221,8 @@ public:
 		PRAGMA_ENABLE_DEPRECATION_WARNINGS
 	}
 
-	bool IsMaterialUsed(int32 MaterialIndex) const;
+	/** USkinnedAsset interface. */
+	virtual bool IsMaterialUsed(int32 MaterialIndex) const override;
 
 public:
 	const TArray<FSkinWeightProfileInfo>& GetSkinWeightProfiles() const 
diff --git a/Engine/Source/Runtime/Engine/Classes/Engine/SkinnedAsset.h b/Engine/Source/Runtime/Engine/Classes/Engine/SkinnedAsset.h
new file mode 100644
--- /dev/null
+++ b/Engine/Source/Runtime/Engine/Classes/Engine/SkinnedAsset.h
@@ -0,0 +1,187 @@
+// Copyright Epic Games, Inc. All Rights Reserved.
+
+/**
+ * Base class for common skinned mesh assets.
+ */
+
+#pragma once
+
+#include "CoreMinimal.h"
+#include "UObject/ObjectMacros.h"
+#include "Engine/StreamableRenderAsset.h"
+#include "ReferenceSkeleton.h"
+#include "PerPlatformProperties.h"
+#include "SkinnedAsset.generated.h"
+
+class USkeleton;
+class UMorphTarget;
+class UPhysicsAsset;
+class ITargetPlatform;
+struct FSkeletalMaterial;
+struct FSkeletalMeshLODInfo;
+
+UCLASS(hidecategories = Object, config = Engine, editinlinenew, abstract)
+class ENGINE_API USkinnedAsset : public UStreamableRenderAsset
+{
+	GENERATED_BODY()
+
+public:
+	USkinnedAsset(const FObjectInitializer& ObjectInitializer)
+		: Super(ObjectInitializer)
+	{}
+
+	virtual ~USkinnedAsset() {}
+
+	/** Return the reference skeleton. */
+	virtual struct FReferenceSkeleton& GetRefSkeleton()
+	PURE_VIRTUAL(USkinnedAsset::GetRefSkeleton, static FReferenceSkeleton Dummy; return Dummy;);
+	virtual const struct FReferenceSkeleton& GetRefSkeleton() const
+	PURE_VIRTUAL(USkinnedAsset::GetRefSkeleton, static const FReferenceSkeleton Dummy; return Dummy;);
+
+	/** Return the LOD information for the specified LOD index. */
+	virtual FSkeletalMeshLODInfo* GetLODInfo(int32 Index)
+	PURE_VIRTUAL(USkinnedAsset::GetLODInfo, return nullptr;);
+	virtual const FSkeletalMeshLODInfo* GetLODInfo(int32 Index) const
+	PURE_VIRTUAL(USkinnedAsset::GetLODInfo, return nullptr;);
+
+	/** Return if the material index is valid. */
+	virtual bool IsValidMaterialIndex(int32 Index) const
+	{ return GetMaterials().IsValidIndex(Index); }
+
+	/** Return the number of materials of this mesh. */
+	virtual int32 GetNumMaterials() const
+	{ return GetMaterials().Num(); }
+
+	/** Return the physics asset whose shapes will be used for shadowing. */
+	virtual UPhysicsAsset* GetShadowPhysicsAsset() const
+	PURE_VIRTUAL(USkinnedAsset::GetShadowPhysicsAsset, return nullptr;);
+
+	/** Return the component orientation of a bone or socket. */
+	virtual FMatrix GetComposedRefPoseMatrix(int32 InBoneIndex) const
+	PURE_VIRTUAL(USkinnedAsset::GetShadowPhysicsAsset, return FMatrix::Identity;);
+
+	/**
+	 * Returns the UV channel data for a given material index. Used by the texture streamer.
+	 * This data applies to all lod-section using the same material.
+	 *
+	 * @param MaterialIndex		the material index for which to get the data for.
+	 * @return the data, or null if none exists.
+	 */
+	virtual const FMeshUVChannelInfo* GetUVChannelData(int32 MaterialIndex) const
+	PURE_VIRTUAL(USkinnedAsset::GetUVChannelData, return nullptr;);
+
+	/** Return whether ray tracing is supported on this mesh. */
+	virtual bool GetSupportRayTracing() const
+	PURE_VIRTUAL(USkinnedAsset::GetSupportRayTracing, return false;);
+
+	/** Return the minimum ray tracing LOD of this mesh. */
+	virtual int32 GetRayTracingMinLOD() const
+	PURE_VIRTUAL(USkinnedAsset::GetRayTracingMinLOD, return 0;);
+
+	/** Return the reference skeleton precomputed bases. */
+	virtual TArray<FMatrix44f>& GetRefBasesInvMatrix()
+	PURE_VIRTUAL(USkinnedAsset::GetRefBasesInvMatrix, static TArray<FMatrix44f> Dummy; return Dummy;);
+	virtual const TArray<FMatrix44f>& GetRefBasesInvMatrix() const
+	PURE_VIRTUAL(USkinnedAsset::GetRefBasesInvMatrix, static const TArray<FMatrix44f> Dummy; return Dummy;);
+
+	/** Return the whole array of LOD info. */
+	virtual TArray<FSkeletalMeshLODInfo>& GetLODInfoArray()
+	PURE_VIRTUAL(USkinnedAsset::GetLODInfoArray, static TArray<FSkeletalMeshLODInfo> Dummy; return Dummy;);
+	virtual const TArray<FSkeletalMeshLODInfo>& GetLODInfoArray() const
+	PURE_VIRTUAL(USkinnedAsset::GetLODInfoArray, static const TArray<FSkeletalMeshLODInfo> Dummy; return Dummy;);
+
+	/** Get the data to use for rendering. */
+	virtual FORCEINLINE class FSkeletalMeshRenderData* GetResourceForRendering() const
+	PURE_VIRTUAL(USkinnedAsset::GetResourceForRendering, return nullptr;);
+
+	virtual int32 GetDefaultMinLod() const
+	PURE_VIRTUAL(USkinnedAsset::GetDefaultMinLod, return 0;);
+
+	virtual UPhysicsAsset* GetPhysicsAsset() const
+	PURE_VIRTUAL(USkinnedAsset::GetPhysicsAsset, return nullptr;);
+
+	virtual TArray<FSkeletalMaterial>& GetMaterials()
+	PURE_VIRTUAL(USkinnedAsset::GetMaterials, static TArray<FSkeletalMaterial> Dummy; return Dummy;);
+	virtual const TArray<FSkeletalMaterial>& GetMaterials() const
+	PURE_VIRTUAL(USkinnedAsset::GetMaterials, static const TArray<FSkeletalMaterial> Dummy; return Dummy;);
+
+	virtual int32 GetLODNum() const
+	PURE_VIRTUAL(USkinnedAsset::GetLODNum, return 0;);
+
+	virtual bool IsMaterialUsed(int32 MaterialIndex) const
+	PURE_VIRTUAL(USkinnedAsset::IsMaterialUsed, return false;);
+
+	virtual FBoxSphereBounds GetBounds() const
+	PURE_VIRTUAL(USkinnedAsset::GetBounds, return FBoxSphereBounds(););
+
+	/**
+	 * Returns the "active" socket list - all sockets from this mesh plus all non-duplicates from the skeleton
+	 * Const ref return value as this cannot be modified externally
+	 */
+	virtual TArray<class USkeletalMeshSocket*> GetActiveSocketList() const
+	PURE_VIRTUAL(USkinnedAsset::GetActiveSocketList, static TArray<class USkeletalMeshSocket*> Dummy; return Dummy;);
+
+	virtual USkeleton* GetSkeleton()
+	PURE_VIRTUAL(USkinnedAsset::GetSkeleton, return nullptr;);
+	virtual const USkeleton* GetSkeleton() const
+	PURE_VIRTUAL(USkinnedAsset::GetSkeleton, return nullptr;);
+	virtual void SetSkeleton(USkeleton* InSkeleton)
+	PURE_VIRTUAL(USkinnedAsset::SetSkeleton,);
+
+	/** Return true if given index's LOD is valid */
+	virtual bool IsValidLODIndex(int32 Index) const
+	PURE_VIRTUAL(USkinnedAsset::IsValidLODIndex, return false;);
+
+	virtual int32 GetMinLodIdx(bool bForceLowestLODIdx = false) const
+	PURE_VIRTUAL(USkinnedAsset::GetMinLodIdx, return 0;);
+
+	/** Return the morph targets. */
+	virtual TArray<TObjectPtr<UMorphTarget>>& GetMorphTargets()
+	{ static TArray<TObjectPtr<UMorphTarget>> Dummy; return Dummy; }
+	virtual const TArray<UMorphTarget*>& GetMorphTargets() const
+	{ static const TArray<TObjectPtr<UMorphTarget>> Dummy; return Dummy; }
+
+	/**Find a named MorphTarget from the morph targets */
+	virtual UMorphTarget* FindMorphTarget(FName MorphTargetName) const
+	{ return nullptr; }
+
+	/** True if this mesh LOD needs to keep it's data on CPU. USkinnedAsset interface. */
+	virtual bool NeedCPUData(int32 LODIndex) const
+	PURE_VIRTUAL(USkinnedAsset::NeedCPUData, return false;);
+
+	/** Return whether or not the mesh has vertex colors. */
+	virtual bool GetHasVertexColors() const
+	PURE_VIRTUAL(USkinnedAsset::GetHasVertexColors, return false;);
+
+	virtual int32 GetPlatformMinLODIdx(const ITargetPlatform* TargetPlatform) const
+	PURE_VIRTUAL(USkinnedAsset::GetPlatformMinLODIdx, return 0;);
+
+	virtual const FPerPlatformBool& GetDisableBelowMinLodStripping() const
+	PURE_VIRTUAL(USkinnedAsset::GetDisableBelowMinLodStripping, static const FPerPlatformBool Dummy; return Dummy;);
+
+	virtual void SetSkinWeightProfilesData(int32 LODIndex, struct FSkinWeightProfilesData& SkinWeightProfilesData) {}
+
+	virtual int32 GetMaxNumOptionalLODs(const ITargetPlatform* TargetPlatform) const
+	PURE_VIRTUAL(USkinnedAsset::GetMaxNumOptionalLODs, return 0;);
+
+	/**
+	* UObject Interface
+	* This will return detail info about this specific object. (e.g. AudioComponent will return the name of the cue,
+	* ParticleSystemComponent will return the name of the ParticleSystem)  The idea here is that in many places
+	* you have a component of interest but what you really want is some characteristic that you can use to track
+	* down where it came from.
+	*/
+	virtual FString GetDetailedInfoInternal() const override
+	{ return GetPathName(nullptr); }
+
+#if WITH_EDITOR
+	/** Get whether this mesh should use LOD streaming for the given platform. */
+	virtual bool GetEnableLODStreaming(const ITargetPlatform* TargetPlatform) const
+	PURE_VIRTUAL(USkinnedAsset::GetEnableLODStreaming, return false;);
+
+	/** Get the maximum number of LODs that can be streamed. */
+	virtual int32 GetMaxNumStreamedLODs(const ITargetPlatform* TargetPlatform) const
+	PURE_VIRTUAL(USkinnedAsset::GetMaxNumStreamedLODs, return 0;);
+#endif
+};
+
diff --git a/Engine/Source/Runtime/Engine/Classes/PhysicsEngine/PhysicsAsset.h b/Engine/Source/Runtime/Engine/Classes/PhysicsEngine/PhysicsAsset.h
--- a/Engine/Source/Runtime/Engine/Classes/PhysicsEngine/PhysicsAsset.h
+++ b/Engine/Source/Runtime/Engine/Classes/PhysicsEngine/PhysicsAsset.h
@@ -289,10 +289,10 @@ public:
 #endif
 
 	// @todo document
-	void GetCollisionMesh(int32 ViewIndex, FMeshElementCollector& Collector, const USkeletalMesh* SkelMesh, const TArray<FTransform>& SpaceBases, const FTransform& LocalToWorld, const FVector& Scale3D);
+	void GetCollisionMesh(int32 ViewIndex, FMeshElementCollector& Collector, const FReferenceSkeleton& RefSkeleton, const TArray<FTransform>& SpaceBases, const FTransform& LocalToWorld, const FVector& Scale3D);
 
 	// @todo document
-	void DrawConstraints(int32 ViewIndex, FMeshElementCollector& Collector, const USkeletalMesh* SkelMesh, const TArray<FTransform>& SpaceBases, const FTransform& LocalToWorld, float Scale);
+	void DrawConstraints(int32 ViewIndex, FMeshElementCollector& Collector, const FReferenceSkeleton& RefSkeleton, const TArray<FTransform>& SpaceBases, const FTransform& LocalToWorld, float Scale);
 
 	void GetUsedMaterials(TArray<UMaterialInterface*>& Materials);
 
diff --git a/Engine/Source/Runtime/Engine/Private/Components/SkinnedMeshComponent.cpp b/Engine/Source/Runtime/Engine/Private/Components/SkinnedMeshComponent.cpp
--- a/Engine/Source/Runtime/Engine/Private/Components/SkinnedMeshComponent.cpp
+++ b/Engine/Source/Runtime/Engine/Private/Components/SkinnedMeshComponent.cpp
@@ -4056,6 +4056,11 @@ void USkinnedMeshComponent::SetRenderStatic(bool bNewValue)
 	}
 }
 
+USkinnedAsset* USkinnedMeshComponent::GetSkinnedAsset() const
+{
+	return SkeletalMesh;
+}
+
 #if WITH_EDITOR
 
 bool USkinnedMeshComponent::IsCompiling() const
diff --git a/Engine/Source/Runtime/Engine/Private/PhysicsEngine/PhysDrawing.cpp b/Engine/Source/Runtime/Engine/Private/PhysicsEngine/PhysDrawing.cpp
--- a/Engine/Source/Runtime/Engine/Private/PhysicsEngine/PhysDrawing.cpp
+++ b/Engine/Source/Runtime/Engine/Private/PhysicsEngine/PhysDrawing.cpp
@@ -1071,11 +1071,11 @@ FTransform GetSkelBoneTransform(int32 BoneIndex, const TArray<FTransform>& Space
 	}
 }
 
-void UPhysicsAsset::GetCollisionMesh(int32 ViewIndex, FMeshElementCollector& Collector, const USkeletalMesh* SkelMesh, const TArray<FTransform>& SpaceBases, const FTransform& LocalToWorld, const FVector& Scale3D)
+void UPhysicsAsset::GetCollisionMesh(int32 ViewIndex, FMeshElementCollector& Collector, const FReferenceSkeleton& RefSkeleton, const TArray<FTransform>& SpaceBases, const FTransform& LocalToWorld, const FVector& Scale3D)
 {
 	for( int32 i=0; i<SkeletalBodySetups.Num(); i++)
 	{
-		int32 BoneIndex = SkelMesh->GetRefSkeleton().FindBoneIndex( SkeletalBodySetups[i]->BoneName );
+		int32 BoneIndex = RefSkeleton.FindBoneIndex( SkeletalBodySetups[i]->BoneName );
 		
 		FColor* BoneColor = (FColor*)( &SkeletalBodySetups[i] );
 
@@ -1098,7 +1098,7 @@ void UPhysicsAsset::GetUsedMaterials(TArray<UMaterialInterface*>& Materials)
 	}
 }
 
-void UPhysicsAsset::DrawConstraints(int32 ViewIndex, FMeshElementCollector& Collector, const USkeletalMesh* SkelMesh, const TArray<FTransform>& SpaceBases, const FTransform& LocalToWorld, float Scale)
+void UPhysicsAsset::DrawConstraints(int32 ViewIndex, FMeshElementCollector& Collector, const FReferenceSkeleton& RefSkeleton, const TArray<FTransform>& SpaceBases, const FTransform& LocalToWorld, float Scale)
 {
 	for (int32 i = 0; i < ConstraintSetup.Num(); i++)
 	{
@@ -1106,7 +1106,7 @@ void UPhysicsAsset::DrawConstraints(int32 ViewIndex, FMeshElementCollector& Coll
 
 		// Get each constraint frame in world space.
 		FTransform Con1Frame = FTransform::Identity;
-		int32 Bone1Index = SkelMesh->GetRefSkeleton().FindBoneIndex(Instance.ConstraintBone1);
+		int32 Bone1Index = RefSkeleton.FindBoneIndex(Instance.ConstraintBone1);
 		if (Bone1Index != INDEX_NONE)
 		{
 			FTransform Body1TM = GetSkelBoneTransform(Bone1Index, SpaceBases, LocalToWorld);
@@ -1115,7 +1115,7 @@ void UPhysicsAsset::DrawConstraints(int32 ViewIndex, FMeshElementCollector& Coll
 		}
 
 		FTransform Con2Frame = FTransform::Identity;
-		int32 Bone2Index = SkelMesh->GetRefSkeleton().FindBoneIndex(Instance.ConstraintBone2);
+		int32 Bone2Index = RefSkeleton.FindBoneIndex(Instance.ConstraintBone2);
 		if (Bone2Index != INDEX_NONE)
 		{
 			FTransform Body2TM = GetSkelBoneTransform(Bone2Index, SpaceBases, LocalToWorld);
diff --git a/Engine/Source/Runtime/Engine/Private/SkeletalMesh.cpp b/Engine/Source/Runtime/Engine/Private/SkeletalMesh.cpp
--- a/Engine/Source/Runtime/Engine/Private/SkeletalMesh.cpp
+++ b/Engine/Source/Runtime/Engine/Private/SkeletalMesh.cpp
@@ -3988,19 +3988,6 @@ void USkeletalMesh::RebuildSocketMap()
 #endif // !WITH_EDITOR
 }
 
-
-/**
- * This will return detail info about this specific object. (e.g. AudioComponent will return the name of the cue,
- * ParticleSystemComponent will return the name of the ParticleSystem)  The idea here is that in many places
- * you have a component of interest but what you really want is some characteristic that you can use to track
- * down where it came from.  
- */
-FString USkeletalMesh::GetDetailedInfoInternal() const
-{
-	return GetPathName(nullptr);
-}
-
-
 FMatrix USkeletalMesh::GetRefPoseMatrix( int32 BoneIndex ) const
 {
  	check( BoneIndex >= 0 && BoneIndex < GetRefSkeleton().GetRawBoneNum() );
@@ -5568,6 +5555,41 @@ bool USkeletalMesh::IsMinLodQualityLevelEnable() const
 	return (GEngine && GEngine->UseSkeletalMeshMinLODPerQualityLevels);
 }
 
+int32 USkeletalMesh::GetPlatformMinLODIdx(const ITargetPlatform* TargetPlatform) const
+{
+#if WITH_EDITOR
+	check(TargetPlatform);
+
+	if (IsMinLodQualityLevelEnable())
+	{
+		// get all supported quality level from scalability + engine ini files
+		return GetQualityLevelMinLod().GetValueForPlatform(TargetPlatform);
+	}
+	else
+	{
+		return GetMinLod().GetValueForPlatform(*TargetPlatform->IniPlatformName());
+	}
+#else
+	return 0;
+#endif
+}
+
+void USkeletalMesh::SetSkinWeightProfilesData(int32 LODIndex, FSkinWeightProfilesData& SkinWeightProfilesData)
+{
+#if !WITH_EDITOR
+	if (GSkinWeightProfilesLoadByDefaultMode == 1)
+	{
+		// Only allow overriding the base buffer in non-editor builds as it could otherwise be serialized into the asset
+		SkinWeightProfilesData.OverrideBaseBufferSkinWeightData(this, LODIndex);
+	}
+	else
+#endif
+	if (GSkinWeightProfilesLoadByDefaultMode == 3)
+	{
+		SkinWeightProfilesData.SetDynamicDefaultSkinWeightProfile(this, LODIndex, true);
+	}
+}
+
 void USkeletalMesh::OnLodStrippingQualityLevelChanged(IConsoleVariable* Variable) {
 #if WITH_EDITOR || PLATFORM_DESKTOP
 	if (GEngine && GEngine->UseStaticMeshMinLODPerQualityLevels)
@@ -5782,11 +5804,11 @@ const FQuat SphylBasis(FVector(1.0f / FMath::Sqrt(2.0f), 0.0f, 1.0f / FMath::Sqr
  * @param	Component - skeletal mesh primitive being added
  */
 FSkeletalMeshSceneProxy::FSkeletalMeshSceneProxy(const USkinnedMeshComponent* Component, FSkeletalMeshRenderData* InSkelMeshRenderData)
-		:	FPrimitiveSceneProxy(Component, Component->SkeletalMesh->GetFName())
+		:	FPrimitiveSceneProxy(Component, Component->GetSkinnedAsset()->GetFName())
 		,	Owner(Component->GetOwner())
 		,	MeshObject(Component->MeshObject)
 		,	SkeletalMeshRenderData(InSkelMeshRenderData)
-		,	SkeletalMeshForDebug(Component->SkeletalMesh)
+		,	SkeletalMeshForDebug(Component->GetSkinnedAsset())
 		,	PhysicsAssetForDebug(Component->GetPhysicsAsset())
 		,	OverlayMaterial(Component->OverlayMaterial)
 		,	OverlayMaterialMaxDrawDistance(Component->OverlayMaterialMaxDrawDistance)
@@ -5844,7 +5866,7 @@ FSkeletalMeshSceneProxy::FSkeletalMeshSceneProxy(const USkinnedMeshComponent* Co
 	for(int32 LODIdx=0; LODIdx < SkeletalMeshRenderData->LODRenderData.Num(); LODIdx++)
 	{
 		const FSkeletalMeshLODRenderData& LODData = SkeletalMeshRenderData->LODRenderData[LODIdx];
-		const FSkeletalMeshLODInfo& Info = *(Component->SkeletalMesh->GetLODInfo(LODIdx));
+		const FSkeletalMeshLODInfo& Info = *(Component->GetSkinnedAsset()->GetLODInfo(LODIdx));
 
 		FLODSectionElements& LODSection = LODSections[LODIdx];
 
@@ -5857,10 +5879,10 @@ FSkeletalMeshSceneProxy::FSkeletalMeshSceneProxy(const USkinnedMeshComponent* Co
 			// If we are at a dropped LOD, route material index through the LODMaterialMap in the LODInfo struct.
 			int32 UseMaterialIndex = Section.MaterialIndex;			
 			{
-				if(SectionIndex < Info.LODMaterialMap.Num() && Component->SkeletalMesh->GetMaterials().IsValidIndex(Info.LODMaterialMap[SectionIndex]))
+				if(SectionIndex < Info.LODMaterialMap.Num() && Component->GetSkinnedAsset()->IsValidMaterialIndex(Info.LODMaterialMap[SectionIndex]))
 				{
 					UseMaterialIndex = Info.LODMaterialMap[SectionIndex];
-					UseMaterialIndex = FMath::Clamp( UseMaterialIndex, 0, Component->SkeletalMesh->GetMaterials().Num() );
+					UseMaterialIndex = FMath::Clamp( UseMaterialIndex, 0, Component->GetSkinnedAsset()->GetNumMaterials());
 				}
 			}
 
@@ -5888,14 +5910,14 @@ FSkeletalMeshSceneProxy::FSkeletalMeshSceneProxy(const USkinnedMeshComponent* Co
 			{
 				UE_CLOG(Material && !bValidUsage, LogSkeletalMesh, Error,
 					TEXT("Material with missing usage flag was applied to skeletal mesh %s"),
-					*Component->SkeletalMesh->GetPathName());
+					*Component->GetSkinnedAsset()->GetPathName());
 
 				Material = UMaterial::GetDefaultMaterial(MD_Surface);
 				MaterialRelevance |= Material->GetRelevance(FeatureLevel);
 			}
 
 			bool bSectionCastsShadow = !bSectionHidden && bCastShadow &&
-				(Component->SkeletalMesh->GetMaterials().IsValidIndex(UseMaterialIndex) == false || Section.bCastShadow);
+				(Component->GetSkinnedAsset()->IsValidMaterialIndex(UseMaterialIndex) == false || Section.bCastShadow);
 
 			bAnySectionCastsShadow |= bSectionCastsShadow;
 
@@ -5918,7 +5940,7 @@ FSkeletalMeshSceneProxy::FSkeletalMeshSceneProxy(const USkinnedMeshComponent* Co
 		if (!OverlayMaterial->CheckMaterialUsage_Concurrent(MATUSAGE_SkeletalMesh))
 		{
 			OverlayMaterial = UMaterial::GetDefaultMaterial(MD_Surface);
-			UE_LOG(LogSkeletalMesh, Error, TEXT("Overlay material with missing usage flag was applied to skeletal mesh %s"),	*Component->SkeletalMesh->GetPathName());
+			UE_LOG(LogSkeletalMesh, Error, TEXT("Overlay material with missing usage flag was applied to skeletal mesh %s"),	*Component->GetSkinnedAsset()->GetPathName());
 		}
 	}
 
@@ -5945,7 +5967,7 @@ FSkeletalMeshSceneProxy::FSkeletalMeshSceneProxy(const USkinnedMeshComponent* Co
 	// Copy out shadow physics asset data
 	if(SkinnedMeshComponent)
 	{
-		UPhysicsAsset* ShadowPhysicsAsset = SkinnedMeshComponent->SkeletalMesh->GetShadowPhysicsAsset();
+		UPhysicsAsset* ShadowPhysicsAsset = SkinnedMeshComponent->GetSkinnedAsset()->GetShadowPhysicsAsset();
 
 		if (ShadowPhysicsAsset
 			&& SkinnedMeshComponent->CastShadow
@@ -5958,7 +5980,7 @@ FSkeletalMeshSceneProxy::FSkeletalMeshSceneProxy(const USkinnedMeshComponent* Co
 
 				if (BoneIndex != INDEX_NONE)
 				{
-					const FMatrix& RefBoneMatrix = SkinnedMeshComponent->SkeletalMesh->GetComposedRefPoseMatrix(BoneIndex);
+					const FMatrix& RefBoneMatrix = SkinnedMeshComponent->GetSkinnedAsset()->GetComposedRefPoseMatrix(BoneIndex);
 
 					const int32 NumSpheres = BodySetup->AggGeom.SphereElems.Num();
 					for (int32 SphereIndex = 0; SphereIndex < NumSpheres; SphereIndex++)
@@ -6826,11 +6848,11 @@ void FSkeletalMeshSceneProxy::DebugDrawPhysicsAsset(int32 ViewIndex, FMeshElemen
 			check(PhysicsAssetForDebug);
 			if (EngineShowFlags.Collision && IsCollisionEnabled())
 			{
-				PhysicsAssetForDebug->GetCollisionMesh(ViewIndex, Collector, SkeletalMeshForDebug, *BoneSpaceBases, LocalToWorldTransform, TotalScale);
+				PhysicsAssetForDebug->GetCollisionMesh(ViewIndex, Collector, SkeletalMeshForDebug->GetRefSkeleton(), *BoneSpaceBases, LocalToWorldTransform, TotalScale);
 			}
 			if (EngineShowFlags.Constraints)
 			{
-				PhysicsAssetForDebug->DrawConstraints(ViewIndex, Collector, SkeletalMeshForDebug, *BoneSpaceBases, LocalToWorldTransform, TotalScale.X);
+				PhysicsAssetForDebug->DrawConstraints(ViewIndex, Collector, SkeletalMeshForDebug->GetRefSkeleton(), *BoneSpaceBases, LocalToWorldTransform, TotalScale.X);
 			}
 		}
 	}
diff --git a/Engine/Source/Runtime/Engine/Private/SkeletalMeshLODRenderData.cpp b/Engine/Source/Runtime/Engine/Private/SkeletalMeshLODRenderData.cpp
--- a/Engine/Source/Runtime/Engine/Private/SkeletalMeshLODRenderData.cpp
+++ b/Engine/Source/Runtime/Engine/Private/SkeletalMeshLODRenderData.cpp
@@ -2,7 +2,8 @@
 
 #include "Rendering/SkeletalMeshLODRenderData.h"
 #include "Rendering/SkeletalMeshRenderData.h"
-#include "Engine/SkeletalMesh.h"
+#include "Engine/SkinnedAsset.h"
+#include "Engine/SkeletalMesh.h"	// for FSkeletalMeshLODInfo
 #include "Animation/MorphTarget.h"
 #include "Misc/ConfigCacheIni.h"
 #include "EngineLogs.h"
@@ -158,7 +159,7 @@ FArchive& operator<<(FArchive& Ar, FSkelMeshRenderSection& S)
 	return Ar;
 }
 
-void FSkeletalMeshLODRenderData::InitResources(bool bNeedsVertexColors, int32 LODIndex, TArray<UMorphTarget*>& InMorphTargets, USkeletalMesh* Owner)
+void FSkeletalMeshLODRenderData::InitResources(bool bNeedsVertexColors, int32 LODIndex, TArray<UMorphTarget*>& InMorphTargets, USkinnedAsset* Owner)
 {
 	IncrementMemoryStats(bNeedsVertexColors);
 
@@ -205,7 +206,7 @@ void FSkeletalMeshLODRenderData::InitResources(bool bNeedsVertexColors, int32 LO
 	}
 
 	// Always make sure the morph target resources are in an initialized state. This could should not get hit since the data should come with the load.
-	bool bNeedsMorphTargetRenderData = Owner->GetMorphTargets().Num() > 0 && !MorphTargetVertexInfoBuffers.IsMorphResourcesInitialized();
+	bool bNeedsMorphTargetRenderData = Owner && Owner->GetMorphTargets().Num() > 0 && !MorphTargetVertexInfoBuffers.IsMorphResourcesInitialized();
 	if (bNeedsMorphTargetRenderData)
 	{
 		const FSkeletalMeshLODInfo* SkeletalMeshLODInfo = Owner->GetLODInfo(LODIndex);
@@ -455,26 +456,7 @@ SIZE_T FSkeletalMeshLODRenderData::GetCPUAccessMemoryOverhead() const
 	return Result;
 }
 
-int32 FSkeletalMeshLODRenderData::GetPlatformMinLODIdx(const ITargetPlatform* TargetPlatform, const USkeletalMesh* SkeletalMesh)
-{
-#if WITH_EDITOR
-	check(TargetPlatform && SkeletalMesh);
-
-	if(SkeletalMesh->IsMinLodQualityLevelEnable())
-	{
-		// get all supported quality level from scalability + engine ini files
-		return SkeletalMesh->GetQualityLevelMinLod().GetValueForPlatform(TargetPlatform);
-	}
-	else
-	{
-		return SkeletalMesh->GetMinLod().GetValueForPlatform(*TargetPlatform->IniPlatformName());
-	}
-#else
-	return 0;
-#endif
-}
-
-uint8 FSkeletalMeshLODRenderData::GenerateClassStripFlags(FArchive& Ar, const USkeletalMesh* OwnerMesh, int32 LODIdx)
+uint8 FSkeletalMeshLODRenderData::GenerateClassStripFlags(FArchive& Ar, const USkinnedAsset* OwnerMesh, int32 LODIdx)
 {
 #if WITH_EDITOR
 	const bool bIsCook = Ar.IsCooking();
@@ -484,7 +466,7 @@ uint8 FSkeletalMeshLODRenderData::GenerateClassStripFlags(FArchive& Ar, const US
 	bool bMeshDisablesMinLodStrip = false;
 	if (bIsCook)
 	{
-		MinMeshLod = OwnerMesh ? GetPlatformMinLODIdx(Ar.CookingTarget(), OwnerMesh) : 0;
+		MinMeshLod = OwnerMesh ? OwnerMesh->GetPlatformMinLODIdx(Ar.CookingTarget()) : 0;
 		bMeshDisablesMinLodStrip = OwnerMesh ? OwnerMesh->GetDisableBelowMinLodStripping().GetValueForPlatform(CookTarget->GetPlatformInfo().IniPlatformName) : false;
 	}
 	const bool bWantToStripBelowMinLod = bIsCook && GStripSkeletalMeshLodsDuringCooking != 0 && MinMeshLod > LODIdx && !bMeshDisablesMinLodStrip;
@@ -497,9 +479,9 @@ uint8 FSkeletalMeshLODRenderData::GenerateClassStripFlags(FArchive& Ar, const US
 #endif
 }
 
-bool FSkeletalMeshLODRenderData::IsLODCookedOut(const ITargetPlatform* TargetPlatform, const USkeletalMesh* SkeletalMesh, bool bIsBelowMinLOD)
+bool FSkeletalMeshLODRenderData::IsLODCookedOut(const ITargetPlatform* TargetPlatform, const USkinnedAsset* SkinnedAsset, bool bIsBelowMinLOD)
 {
-	check(SkeletalMesh);
+	check(SkinnedAsset);
 #if WITH_EDITOR
 	if (!bIsBelowMinLOD)
 	{
@@ -512,15 +494,15 @@ bool FSkeletalMeshLODRenderData::IsLODCookedOut(const ITargetPlatform* TargetPla
 	}
 	check(TargetPlatform);
 
-	return !SkeletalMesh->GetEnableLODStreaming(TargetPlatform);
+	return !SkinnedAsset->GetEnableLODStreaming(TargetPlatform);
 #else
 	return false;
 #endif
 }
 
-bool FSkeletalMeshLODRenderData::IsLODInlined(const ITargetPlatform* TargetPlatform, const USkeletalMesh* SkeletalMesh, int32 LODIdx, bool bIsBelowMinLOD)
+bool FSkeletalMeshLODRenderData::IsLODInlined(const ITargetPlatform* TargetPlatform, const USkinnedAsset* SkinnedAsset, int32 LODIdx, bool bIsBelowMinLOD)
 {
-	check(SkeletalMesh);
+	check(SkinnedAsset);
 #if WITH_EDITOR
 	if (!TargetPlatform)
 	{
@@ -528,7 +510,7 @@ bool FSkeletalMeshLODRenderData::IsLODInlined(const ITargetPlatform* TargetPlatf
 	}
 	check(TargetPlatform);
 	
-	if (!SkeletalMesh->GetEnableLODStreaming(TargetPlatform))
+	if (!SkinnedAsset->GetEnableLODStreaming(TargetPlatform))
 	{
 		return true;
 	}
@@ -538,8 +520,8 @@ bool FSkeletalMeshLODRenderData::IsLODInlined(const ITargetPlatform* TargetPlatf
 		return false;
 	}
 
-	const int32 MaxNumStreamedLODs = SkeletalMesh->GetMaxNumStreamedLODs(TargetPlatform);
-	const int32 NumLODs = SkeletalMesh->GetLODNum();
+	const int32 MaxNumStreamedLODs = SkinnedAsset->GetMaxNumStreamedLODs(TargetPlatform);
+	const int32 NumLODs = SkinnedAsset->GetLODNum();
 	const int32 NumStreamedLODs = FMath::Min(MaxNumStreamedLODs, NumLODs - 1);
 	const int32 InlinedLODStartIdx = NumStreamedLODs;
 	return LODIdx >= InlinedLODStartIdx;
@@ -548,11 +530,11 @@ bool FSkeletalMeshLODRenderData::IsLODInlined(const ITargetPlatform* TargetPlatf
 #endif
 }
 
-int32 FSkeletalMeshLODRenderData::GetNumOptionalLODsAllowed(const ITargetPlatform* TargetPlatform, const USkeletalMesh* SkeletalMesh)
+int32 FSkeletalMeshLODRenderData::GetNumOptionalLODsAllowed(const ITargetPlatform* TargetPlatform, const USkinnedAsset* SkinnedAsset)
 {
 #if WITH_EDITOR
-	check(TargetPlatform && SkeletalMesh);
-	return SkeletalMesh->GetMaxNumOptionalLODs(TargetPlatform);
+	check(TargetPlatform && SkinnedAsset);
+	return SkinnedAsset->GetMaxNumOptionalLODs(TargetPlatform);
 #else
 	return 0;
 #endif
@@ -570,11 +552,11 @@ bool FSkeletalMeshLODRenderData::ShouldForceKeepCPUResources()
 	return true;
 }
 
-bool FSkeletalMeshLODRenderData::ShouldKeepCPUResources(const USkeletalMesh* SkeletalMesh, int32 LODIdx, bool bForceKeep)
+bool FSkeletalMeshLODRenderData::ShouldKeepCPUResources(const USkinnedAsset* SkinnedAsset, int32 LODIdx, bool bForceKeep)
 {
 	return bForceKeep
-		|| SkeletalMesh->GetResourceForRendering()->RequiresCPUSkinning(GMaxRHIFeatureLevel)
-		|| SkeletalMesh->NeedCPUData(LODIdx);
+		|| SkinnedAsset->GetResourceForRendering()->RequiresCPUSkinning(GMaxRHIFeatureLevel)
+		|| SkinnedAsset->NeedCPUData(LODIdx);
 }
 
 class FSkeletalMeshLODSizeCounter : public FArchive
@@ -602,7 +584,7 @@ private:
 	int64 Size;
 };
 
-void FSkeletalMeshLODRenderData::SerializeStreamedData(FArchive& Ar, USkeletalMesh* Owner, int32 LODIdx, uint8 ClassDataStripFlags, bool bNeedsCPUAccess, bool bForceKeepCPUResources)
+void FSkeletalMeshLODRenderData::SerializeStreamedData(FArchive& Ar, USkinnedAsset* Owner, int32 LODIdx, uint8 ClassDataStripFlags, bool bNeedsCPUAccess, bool bForceKeepCPUResources)
 {
 	Ar.UsingCustomVersion(FUE5PrivateFrostyStreamObjectVersion::GUID);
 	Ar.UsingCustomVersion(FUE5ReleaseStreamObjectVersion::GUID);
@@ -644,18 +626,7 @@ void FSkeletalMeshLODRenderData::SerializeStreamedData(FArchive& Ar, USkeletalMe
 
 	if (Ar.IsLoading())
 	{
-#if !WITH_EDITOR
-		if (GSkinWeightProfilesLoadByDefaultMode == 1)
-		{
-			// Only allow overriding the base buffer in non-editor builds as it could otherwise be serialized into the asset
-			SkinWeightProfilesData.OverrideBaseBufferSkinWeightData(Owner, LODIdx);
-		} else
-	
-#endif
-		if (GSkinWeightProfilesLoadByDefaultMode == 3)
-		{
-			SkinWeightProfilesData.SetDynamicDefaultSkinWeightProfile(Owner, LODIdx, true);
-		}
+		Owner->SetSkinWeightProfilesData(LODIdx, SkinWeightProfilesData);
 	}
 	Ar << SourceRayTracingGeometry.RawData;
 
@@ -737,7 +708,7 @@ void FSkeletalMeshLODRenderData::SerializeStreamedData(FArchive& Ar, USkeletalMe
 	}
 }
 
-void FSkeletalMeshLODRenderData::SerializeAvailabilityInfo(FArchive& Ar, USkeletalMesh* Owner, int32 LODIdx, bool bAdjacencyDataStripped, bool bNeedsCPUAccess)
+void FSkeletalMeshLODRenderData::SerializeAvailabilityInfo(FArchive& Ar, int32 LODIdx, bool bAdjacencyDataStripped, bool bNeedsCPUAccess)
 {
 	Ar.UsingCustomVersion(FUE5ReleaseStreamObjectVersion::GUID);
 	MultiSizeIndexContainer.SerializeMetaData(Ar, bNeedsCPUAccess);
@@ -766,7 +737,7 @@ void FSkeletalMeshLODRenderData::Serialize(FArchive& Ar, UObject* Owner, int32 I
 {
 	DECLARE_SCOPE_CYCLE_COUNTER(TEXT("FSkeletalMeshLODRenderData::Serialize"), STAT_SkeletalMeshLODRenderData_Serialize, STATGROUP_LoadTime);
 
-	USkeletalMesh* OwnerMesh = CastChecked<USkeletalMesh>(Owner);
+	USkinnedAsset* OwnerMesh = CastChecked<USkinnedAsset>(Owner);
 	
 	// Shouldn't needed but to make some static analyzers happy
 	if (!OwnerMesh)
@@ -780,7 +751,7 @@ void FSkeletalMeshLODRenderData::Serialize(FArchive& Ar, UObject* Owner, int32 I
 
 #if WITH_EDITOR
 	const bool bIsBelowMinLOD = StripFlags.IsClassDataStripped(CDSF_MinLodData)
-		|| (Ar.IsCooking() && OwnerMesh && Idx < GetPlatformMinLODIdx(Ar.CookingTarget(), OwnerMesh));
+		|| (Ar.IsCooking() && OwnerMesh && Idx < OwnerMesh->GetPlatformMinLODIdx(Ar.CookingTarget()));
 #else
 	const bool bIsBelowMinLOD = false;
 #endif
@@ -854,7 +825,7 @@ void FSkeletalMeshLODRenderData::Serialize(FArchive& Ar, UObject* Owner, int32 I
 			if (Ar.IsSaving())
 			{
 				const int32 MaxNumOptionalLODs = GetNumOptionalLODsAllowed(Ar.CookingTarget(), OwnerMesh);
-				const int32 OptionalLODIdx = GetPlatformMinLODIdx(Ar.CookingTarget(), OwnerMesh) - Idx;
+				const int32 OptionalLODIdx = OwnerMesh->GetPlatformMinLODIdx(Ar.CookingTarget()) - Idx;
 				bDiscardBulkData = OptionalLODIdx > MaxNumOptionalLODs;
 
 				TArray<uint8> TmpBuff;
@@ -898,7 +869,7 @@ void FSkeletalMeshLODRenderData::Serialize(FArchive& Ar, UObject* Owner, int32 I
 
 			if (!bDiscardBulkData)
 			{
-				SerializeAvailabilityInfo(Ar, OwnerMesh, Idx, StripFlags.IsClassDataStripped(CDSF_AdjacencyData_DEPRECATED), bNeedsCPUAccess);
+				SerializeAvailabilityInfo(Ar, Idx, StripFlags.IsClassDataStripped(CDSF_AdjacencyData_DEPRECATED), bNeedsCPUAccess);
 			}
 		}
 	}
diff --git a/Engine/Source/Runtime/Engine/Private/SkeletalMeshRenderData.cpp b/Engine/Source/Runtime/Engine/Private/SkeletalMeshRenderData.cpp
--- a/Engine/Source/Runtime/Engine/Private/SkeletalMeshRenderData.cpp
+++ b/Engine/Source/Runtime/Engine/Private/SkeletalMeshRenderData.cpp
@@ -732,7 +732,7 @@ PRAGMA_DISABLE_DEPRECATION_WARNINGS
 PRAGMA_ENABLE_DEPRECATION_WARNINGS
 }
 
-void FSkeletalMeshRenderData::InitResources(bool bNeedsVertexColors, TArray<UMorphTarget*>& InMorphTargets, USkeletalMesh* Owner)
+void FSkeletalMeshRenderData::InitResources(bool bNeedsVertexColors, TArray<UMorphTarget*>& InMorphTargets, USkinnedAsset* Owner)
 {
 	if (!bInitialized)
 	{
diff --git a/Engine/Source/Runtime/Engine/Public/Rendering/SkeletalMeshLODRenderData.h b/Engine/Source/Runtime/Engine/Public/Rendering/SkeletalMeshLODRenderData.h
--- a/Engine/Source/Runtime/Engine/Public/Rendering/SkeletalMeshLODRenderData.h
+++ b/Engine/Source/Runtime/Engine/Public/Rendering/SkeletalMeshLODRenderData.h
@@ -16,6 +16,7 @@
 #include "GPUSkinVertexFactory.h"
 #include "Animation/SkinWeightProfile.h"
 
+class USkinnedAsset;
 #if WITH_EDITOR
 class FSkeletalMeshLODModel;
 #endif // WITH_EDITOR
@@ -174,7 +175,7 @@ public:
 	*
 	* @param Parent Parent mesh
 	*/
-	void InitResources(bool bNeedsVertexColors, int32 LODIndex, TArray<class UMorphTarget*>& InMorphTargets, USkeletalMesh* Owner);
+	void InitResources(bool bNeedsVertexColors, int32 LODIndex, TArray<class UMorphTarget*>& InMorphTargets, USkinnedAsset* Owner);
 
 	/**
 	* Releases the LOD's render resources.
@@ -219,9 +220,9 @@ public:
 	* @param bForceKeepCPUResources - Whether we want to force keeping CPU resources
 	* @return Size of streamed LOD data in bytes
 	*/
-	void SerializeStreamedData(FArchive& Ar, USkeletalMesh* Owner, int32 LODIdx, uint8 ClassStripFlags, bool bNeedsCPUAccess, bool bForceKeepCPUResources);
+	void SerializeStreamedData(FArchive& Ar, USkinnedAsset* Owner, int32 LODIdx, uint8 ClassStripFlags, bool bNeedsCPUAccess, bool bForceKeepCPUResources);
 
-	void SerializeAvailabilityInfo(FArchive& Ar, USkeletalMesh* Owner, int32 LODIdx, bool bHasAdjacencyData, bool bNeedsCPUAccess);
+	void SerializeAvailabilityInfo(FArchive& Ar, int32 LODIdx, bool bHasAdjacencyData, bool bNeedsCPUAccess);
 
 #if WITH_EDITOR
 	/**
@@ -294,7 +295,7 @@ public:
 	void DecrementMemoryStats();
 
 	static bool ShouldForceKeepCPUResources();
-	static bool ShouldKeepCPUResources(const USkeletalMesh* SkeletalMesh, int32 LODIdx, bool bForceKeep);
+	static bool ShouldKeepCPUResources(const USkinnedAsset* SkinnedAsset, int32 LODIdx, bool bForceKeep);
 
 private:
 	enum EClassDataStripFlag : uint8
@@ -303,15 +304,13 @@ private:
 		CDSF_MinLodData = 2
 	};
 
-	static int32 GetPlatformMinLODIdx(const ITargetPlatform* TargetPlatform, const USkeletalMesh* SkeletalMesh);
+	static uint8 GenerateClassStripFlags(FArchive& Ar, const USkinnedAsset* OwnerMesh, int32 LODIdx);
 
-	static uint8 GenerateClassStripFlags(FArchive& Ar, const USkeletalMesh* OwnerMesh, int32 LODIdx);
+	static bool IsLODCookedOut(const ITargetPlatform* TargetPlatform, const USkinnedAsset* SkinnedAsset, bool bIsBelowMinLOD);
 
-	static bool IsLODCookedOut(const ITargetPlatform* TargetPlatform, const USkeletalMesh* SkeletalMesh, bool bIsBelowMinLOD);
+	static bool IsLODInlined(const ITargetPlatform* TargetPlatform, const USkinnedAsset* SkinnedAsset, int32 LODIdx, bool bIsBelowMinLOD);
 
-	static bool IsLODInlined(const ITargetPlatform* TargetPlatform, const USkeletalMesh* SkeletalMesh, int32 LODIdx, bool bIsBelowMinLOD);
-
-	static int32 GetNumOptionalLODsAllowed(const ITargetPlatform* TargetPlatform, const USkeletalMesh* SkeletalMesh);
+	static int32 GetNumOptionalLODsAllowed(const ITargetPlatform* TargetPlatform, const USkinnedAsset* SkinnedAsset);
 
 	friend class FSkeletalMeshRenderData;
 };
\ No newline at end of file
diff --git a/Engine/Source/Runtime/Engine/Public/Rendering/SkeletalMeshRenderData.h b/Engine/Source/Runtime/Engine/Public/Rendering/SkeletalMeshRenderData.h
--- a/Engine/Source/Runtime/Engine/Public/Rendering/SkeletalMeshRenderData.h
+++ b/Engine/Source/Runtime/Engine/Public/Rendering/SkeletalMeshRenderData.h
@@ -64,7 +64,7 @@ public:
 	void Serialize(FArchive& Ar, USkeletalMesh* Owner);
 
 	/** Initializes rendering resources. */
-	ENGINE_API void InitResources(bool bNeedsVertexColors, TArray<UMorphTarget*>& InMorphTargets, USkeletalMesh* Owner);
+	ENGINE_API void InitResources(bool bNeedsVertexColors, TArray<UMorphTarget*>& InMorphTargets, USkinnedAsset* Owner);
 
 	/** Releases rendering resources. */
 	ENGINE_API void ReleaseResources();
diff --git a/Engine/Source/Runtime/Engine/Public/SkeletalMeshTypes.h b/Engine/Source/Runtime/Engine/Public/SkeletalMeshTypes.h
--- a/Engine/Source/Runtime/Engine/Public/SkeletalMeshTypes.h
+++ b/Engine/Source/Runtime/Engine/Public/SkeletalMeshTypes.h
@@ -240,7 +240,7 @@ protected:
 	FSkeletalMeshRenderData* SkeletalMeshRenderData;
 
 	/** The points to the skeletal mesh and physics assets are purely for debug purposes. Access is NOT thread safe! */
-	const USkeletalMesh* SkeletalMeshForDebug;
+	const class USkinnedAsset* SkeletalMeshForDebug;
 	class UPhysicsAsset* PhysicsAssetForDebug;
 	
 	UMaterialInterface* OverlayMaterial;
