Initalize the Input Device to the old controller id in the InputKey event args. This fixes a backwards compatbility issue where the input device could be invalid

#rnx

diff --git a/Engine/Source/Runtime/Engine/Public/UnrealClient.h b/Engine/Source/Runtime/Engine/Public/UnrealClient.h
--- a/Engine/Source/Runtime/Engine/Public/UnrealClient.h
+++ b/Engine/Source/Runtime/Engine/Public/UnrealClient.h
@@ -794,6 +794,7 @@ public:
 	FInputKeyEventArgs(FViewport* InViewport, int32 InControllerId, FKey InKey, EInputEvent InEvent)
 		: Viewport(InViewport)
 		, ControllerId(InControllerId)
+		, InputDevice(FInputDeviceId::CreateFromInternalId(InControllerId))
 		, Key(InKey)
 		, Event(InEvent)
 		, AmountDepressed(1.0f)
@@ -815,6 +816,7 @@ public:
 	FInputKeyEventArgs(FViewport* InViewport, int32 InControllerId, FKey InKey, EInputEvent InEvent, float InAmountDepressed, bool bInIsTouchEvent)
 		: Viewport(InViewport)
 		, ControllerId(InControllerId)
+		, InputDevice(FInputDeviceId::CreateFromInternalId(InControllerId))
 		, Key(InKey)
 		, Event(InEvent)
 		, AmountDepressed(InAmountDepressed)
