Remove the pasing of input to other viewports in the case of multiple clients. This is no longer necessary due to each player having a unique PlatformUserId, so as long as the viewport has focus it will get input from the right device.

#rnx

diff --git a/Engine/Source/Runtime/Engine/Private/GameViewportClient.cpp b/Engine/Source/Runtime/Engine/Private/GameViewportClient.cpp
--- a/Engine/Source/Runtime/Engine/Private/GameViewportClient.cpp
+++ b/Engine/Source/Runtime/Engine/Private/GameViewportClient.cpp
@@ -661,7 +661,7 @@ bool UGameViewportClient::InputKey(const FInputKeyEventArgs& InEventArgs)
 
 	if (!bResult)
 	{
-		ULocalPlayer* const TargetPlayer = GEngine->GetLocalPlayerFromControllerId(this, EventArgs.ControllerId);
+		ULocalPlayer* const TargetPlayer = GEngine->GetLocalPlayerFromInputDevice(this, EventArgs.InputDevice);
 		if (TargetPlayer && TargetPlayer->PlayerController)
 		{
 			bResult = TargetPlayer->PlayerController->InputKey(FInputKeyParams(EventArgs.Key, EventArgs.Event, static_cast<double>(EventArgs.AmountDepressed), EventArgs.IsGamepad(), EventArgs.InputDevice));
@@ -674,21 +674,6 @@ bool UGameViewportClient::InputKey(const FInputKeyEventArgs& InEventArgs)
 		}
 	}
 
-#if WITH_EDITOR
-	// For PIE, let the next PIE window handle the input if none of our players did
-	// (this allows people to use multiple controllers to control each window)
-	const int32 NumLocalPlayers = World ? World->GetGameInstance()->GetNumLocalPlayers() : 0;
-	if (!bResult && EventArgs.ControllerId > NumLocalPlayers - 1 && EventArgs.Viewport->IsPlayInEditorViewport())
-	{
-		UGameViewportClient* NextViewport = GEngine->GetNextPIEViewport(this);
-		if (NextViewport)
-		{
-			FInputKeyEventArgs NextViewportEventArgs = EventArgs;
-			NextViewportEventArgs.ControllerId = EventArgs.ControllerId - NumLocalPlayers;
-			bResult = NextViewport->InputKey(NextViewportEventArgs);
-		}
-	}
-#endif
 
 	return bResult;
 }
@@ -746,18 +731,6 @@ bool UGameViewportClient::InputAxis(FViewport* InViewport, FInputDeviceId InputD
 			}
 		}
 
-#if WITH_EDITOR
-		// If there are multiple local players and the input has not been handled yet, we should pass it along
-		// to the next viewport so that it can check if any of its Player Controllers need it
-		if (!bResult && InViewport->IsPlayInEditorViewport())
-		{
-			if (UGameViewportClient* NextViewport = GEngine->GetNextPIEViewport(this))
-			{
-				bResult = NextViewport->InputAxis(InViewport, EventArgs.InputDevice, EventArgs.Key, Delta, DeltaTime, NumSamples, EventArgs.IsGamepad());
-			}
-		}
-#endif
-
 		if( InViewport->IsSlateViewport() && InViewport->IsPlayInEditorViewport() )
 		{
 			// Absorb all keys so game input events are not routed to the Slate editor frame
