Replace direct access to SkeletalMesh object from USkinnedMeshComponent with GetSkeletalMesh function

diff --git a/Engine/Source/Runtime/Engine/Private/Components/SkeletalMeshComponent.cpp b/Engine/Source/Runtime/Engine/Private/Components/SkeletalMeshComponent.cpp
--- a/Engine/Source/Runtime/Engine/Private/Components/SkeletalMeshComponent.cpp
+++ b/Engine/Source/Runtime/Engine/Private/Components/SkeletalMeshComponent.cpp
@@ -185,7 +185,7 @@ public:
 		if (USkeletalMeshComponent* Comp = SkeletalMeshComponent.Get())
 		{
 			FScopeCycleCounterUObject ComponentScope(Comp);
-			FScopeCycleCounterUObject MeshScope(Comp->SkeletalMesh);
+			FScopeCycleCounterUObject MeshScope(Comp->GetSkeletalMesh());
 
 			if(IsValidRef(Comp->ParallelAnimationEvaluationTask))
 			{
@@ -465,12 +465,12 @@ extern TAutoConsoleVariable<int32> CVarEnableClothPhysics;
 
 bool USkeletalMeshComponent::CanSimulateClothing() const
 {
-	if(!SkeletalMesh || !bAllowClothActors || !CVarEnableClothPhysics.GetValueOnAnyThread())
+	if(!GetSkeletalMesh() || !bAllowClothActors || !CVarEnableClothPhysics.GetValueOnAnyThread())
 	{
 		return false;
 	}
 
-	return SkeletalMesh->HasActiveClothingAssets() && !IsNetMode(NM_DedicatedServer);
+	return GetSkeletalMesh()->HasActiveClothingAssets() && !IsNetMode(NM_DedicatedServer);
 }
 
 void USkeletalMeshComponent::UpdateClothTickRegisteredState()
@@ -524,8 +524,8 @@ bool USkeletalMeshComponent::NeedToSpawnAnimScriptInstance() const
 {
 	IAnimClassInterface* AnimClassInterface = IAnimClassInterface::GetFromClass(AnimClass);
 	const USkeleton* AnimSkeleton = (AnimClassInterface) ? AnimClassInterface->GetTargetSkeleton() : nullptr;
-	const bool bSkeletonCompatible = (SkeletalMesh && AnimSkeleton && SkeletalMesh->GetSkeleton()) ? SkeletalMesh->GetSkeleton()->IsCompatible(AnimSkeleton) : false;
-	const bool bSkelMeshCompatible = (SkeletalMesh && AnimSkeleton) ? AnimSkeleton->IsCompatibleMesh(SkeletalMesh, false) : false;
+	const bool bSkeletonCompatible = (GetSkeletalMesh() && AnimSkeleton && GetSkeletalMesh()->GetSkeleton()) ? GetSkeletalMesh()->GetSkeleton()->IsCompatible(AnimSkeleton) : false;
+	const bool bSkelMeshCompatible = (GetSkeletalMesh() && AnimSkeleton) ? AnimSkeleton->IsCompatibleMesh(GetSkeletalMesh(), false) : false;
 	const bool bAnimSkelValid = !AnimClassInterface || (bSkeletonCompatible && bSkelMeshCompatible);
 
 	if (AnimationMode == EAnimationMode::AnimationBlueprint && AnimClass && bAnimSkelValid)
@@ -545,10 +545,10 @@ bool USkeletalMeshComponent::NeedToSpawnAnimScriptInstance() const
 
 bool USkeletalMeshComponent::NeedToSpawnPostPhysicsInstance(bool bForceReinit) const
 {
-	if(SkeletalMesh)
+	if(GetSkeletalMesh())
 	{
 		const UClass* MainInstanceClass = *AnimClass;
-		const UClass* ClassToUse = *SkeletalMesh->GetPostProcessAnimBlueprint();
+		const UClass* ClassToUse = *GetSkeletalMesh()->GetPostProcessAnimBlueprint();
 		const UClass* CurrentClass = PostProcessAnimInstance ? PostProcessAnimInstance->GetClass() : nullptr;
 
 		// We need to have an instance, and we have the wrong class (different or null)
@@ -594,14 +594,15 @@ void USkeletalMeshComponent::OnRegister()
 		ClothingSimulationFactory = UClothingSimulationFactory::GetDefaultClothingSimulationFactoryClass();
 	}
 
+	USkeletalMesh* SkelMesh = GetSkeletalMesh();
 	// Look up for the best simulation factory to support each asset
-	if (ClothingSimulationFactory && SkeletalMesh)
+	if (ClothingSimulationFactory && SkelMesh)
 	{
 		// Check whether all clothing assets are supported by the current factory
 		bool bSupportsAllAssets = true;
 
 		UClothingSimulationFactory* const DefaultObject = ClothingSimulationFactory->GetDefaultObject<UClothingSimulationFactory>();
-		for (UClothingAssetBase* const ClothingAsset : SkeletalMesh->GetMeshClothingAssets())
+		for (UClothingAssetBase* const ClothingAsset : SkelMesh->GetMeshClothingAssets())
 		{
 			if (ClothingAsset && !DefaultObject->SupportsAsset(ClothingAsset))
 			{
@@ -609,11 +610,11 @@ void USkeletalMeshComponent::OnRegister()
 
 				UE_LOG(LogSkeletalMesh, Display,
 					TEXT("OnRegister[%s]: [%s] is currently unable to provide a fully functional simulation for each of this SkeletalMesh's clothing assets."),
-					*GetPathNameSafe(SkeletalMesh),
+					*GetPathNameSafe(SkelMesh),
 					*ClothingSimulationFactory->GetName());
 				UE_LOG(LogSkeletalMesh, Display,
 					TEXT("OnRegister[%s]: The ClothingSimulationFactory property will now be automatically updated to use the most functional simulation that can be found."),
-					*GetPathNameSafe(SkeletalMesh));
+					*GetPathNameSafe(SkelMesh));
 
 				break;
 			}
@@ -634,7 +635,7 @@ void USkeletalMeshComponent::OnRegister()
 						int NumAssets = 0;
 						int SupportedNumAssets = 0;
 						UClothingSimulationFactory* const NewDefaultObject = NewClothingSimulationFactory->GetDefaultObject<UClothingSimulationFactory>();
-						for (UClothingAssetBase* const ClothingAsset : SkeletalMesh->GetMeshClothingAssets())
+						for (UClothingAssetBase* const ClothingAsset : SkelMesh->GetMeshClothingAssets())
 						{
 							if (ClothingAsset)
 							{
@@ -662,11 +663,11 @@ void USkeletalMeshComponent::OnRegister()
 
 			UE_CLOG(!MostSupportedNumAssets, LogSkeletalMesh, Warning,
 				TEXT("OnRegister[%s]: There is no clothing simulation factory available that supports any of this SkeletalMesh's clothing assets."),
-				*GetPathNameSafe(SkeletalMesh));
+				*GetPathNameSafe(SkelMesh));
 
 			UE_CLOG(MostSupportedNumAssets && !bSupportsAllAssets, LogSkeletalMesh, Warning,
 				TEXT("OnRegister[%s]: The most suitable clothing simulation factory available only partially supports this SkeletalMesh's clothing assets."),
-				*GetPathNameSafe(SkeletalMesh));
+				*GetPathNameSafe(SkelMesh));
 		}
 	}
 
@@ -739,13 +740,13 @@ void USkeletalMeshComponent::InitAnim(bool bForceReinit)
 
 	// a lot of places just call InitAnim without checking Mesh, so 
 	// I'm moving the check here
-	if ( SkeletalMesh != nullptr && IsRegistered() )
+	if ( GetSkeletalMesh() != nullptr && IsRegistered() )
 	{
 		//clear cache UID since we don't know if skeleton changed
 		CachedAnimCurveUidVersion = 0;
 
 		// we still need this in case users doesn't call tick, but sent to renderer
-		MorphTargetWeights.SetNumZeroed(SkeletalMesh->GetMorphTargets().Num());
+		MorphTargetWeights.SetNumZeroed(GetSkeletalMesh()->GetMorphTargets().Num());
 
 		// We may be doing parallel evaluation on the current anim instance
 		// Calling this here with true will block this init till that thread completes
@@ -760,8 +761,8 @@ void USkeletalMeshComponent::InitAnim(bool bForceReinit)
 		const USkeleton* AnimSkeleton = (AnimScriptInstance)? AnimScriptInstance->CurrentSkeleton : nullptr;
 
 		const bool bClearAnimInstance = AnimScriptInstance && !AnimSkeleton;
-		const bool bSkeletonMismatch = AnimSkeleton && (AnimScriptInstance->CurrentSkeleton!=SkeletalMesh->GetSkeleton());
-		const bool bSkeletonCompatible = AnimSkeleton && SkeletalMesh->GetSkeleton() && !bSkeletonMismatch && SkeletalMesh->GetSkeleton()->IsCompatible(AnimSkeleton);
+		const bool bSkeletonMismatch = AnimSkeleton && (AnimScriptInstance->CurrentSkeleton!=GetSkeletalMesh()->GetSkeleton());
+		const bool bSkeletonCompatible = AnimSkeleton && GetSkeletalMesh()->GetSkeleton() && !bSkeletonMismatch && GetSkeletalMesh()->GetSkeleton()->IsCompatible(AnimSkeleton);
 
 		LastPoseTickFrame = 0;
 
@@ -796,9 +797,9 @@ void USkeletalMeshComponent::InitAnim(bool bForceReinit)
 				else
 				{
 					PRAGMA_DISABLE_DEPRECATION_WARNINGS
-					BoneSpaceTransforms = SkeletalMesh->GetRefSkeleton().GetRefBonePose();
+					BoneSpaceTransforms = GetSkeletalMesh()->GetRefSkeleton().GetRefBonePose();
 					//Mini RefreshBoneTransforms (the bit we actually care about)
-					FillComponentSpaceTransforms(SkeletalMesh, BoneSpaceTransforms, GetEditableComponentSpaceTransforms());
+					FillComponentSpaceTransforms(GetSkeletalMesh(), BoneSpaceTransforms, GetEditableComponentSpaceTransforms());
 					PRAGMA_ENABLE_DEPRECATION_WARNINGS
 					bNeedToFlipSpaceBaseBuffers = true; // Have updated space bases so need to flip
 					FlipEditableSpaceBases();
@@ -848,7 +849,8 @@ bool USkeletalMeshComponent::InitializeAnimScriptInstance(bool bForceReinit, boo
 
 	if (IsRegistered())
 	{
-		check(SkeletalMesh);
+		USkeletalMesh* SkelMesh = GetSkeletalMesh();
+		check(SkelMesh);
 
 		if (NeedToSpawnAnimScriptInstance())
 		{
@@ -866,7 +868,7 @@ bool USkeletalMeshComponent::InitializeAnimScriptInstance(bool bForceReinit, boo
 		}
 		else 
 		{
-			bool bShouldSpawnSingleNodeInstance = SkeletalMesh && SkeletalMesh->GetSkeleton() && AnimationMode == EAnimationMode::AnimationSingleNode;
+			bool bShouldSpawnSingleNodeInstance = SkelMesh && SkelMesh->GetSkeleton() && AnimationMode == EAnimationMode::AnimationSingleNode;
 			if (bShouldSpawnSingleNodeInstance)
 			{
 				SCOPE_CYCLE_COUNTER(STAT_AnimSpawnTime);
@@ -906,7 +908,7 @@ bool USkeletalMeshComponent::InitializeAnimScriptInstance(bool bForceReinit, boo
 		}
 
 		// May need to clear out the post physics instance
-		UClass* NewMeshInstanceClass = *SkeletalMesh->GetPostProcessAnimBlueprint();
+		UClass* NewMeshInstanceClass = *SkelMesh->GetPostProcessAnimBlueprint();
 		if(!NewMeshInstanceClass || NewMeshInstanceClass == *AnimClass)
 		{
 			PostProcessAnimInstance = nullptr;
@@ -914,7 +916,7 @@ bool USkeletalMeshComponent::InitializeAnimScriptInstance(bool bForceReinit, boo
 
 		if(NeedToSpawnPostPhysicsInstance(bForceReinit))
 		{
-			PostProcessAnimInstance = NewObject<UAnimInstance>(this, *SkeletalMesh->GetPostProcessAnimBlueprint());
+			PostProcessAnimInstance = NewObject<UAnimInstance>(this, *SkelMesh->GetPostProcessAnimBlueprint());
 
 			if(PostProcessAnimInstance)
 			{
@@ -928,7 +930,7 @@ bool USkeletalMeshComponent::InitializeAnimScriptInstance(bool bForceReinit, boo
 				bInitializedPostInstance = true;
 			}
 		}
-		else if (!SkeletalMesh->GetPostProcessAnimBlueprint().Get())
+		else if (!SkelMesh->GetPostProcessAnimBlueprint().Get())
 		{
 			PostProcessAnimInstance = nullptr;
 		}
@@ -1039,7 +1041,8 @@ void USkeletalMeshComponent::PostEditChangeProperty(FPropertyChangedEvent& Prope
 			InitAnim(false);
 		}
 
-		if(PropertyThatChanged->GetFName() == GET_MEMBER_NAME_CHECKED( USkeletalMeshComponent, SkeletalMesh))
+		USkeletalMesh* SkelMesh = GetSkeletalMesh();
+		if(PropertyThatChanged->GetFName() == GET_MEMBER_NAME_CHECKED( USkeletalMeshComponent, GetSkeletalMesh()))
 		{
 			ValidateAnimation();
 
@@ -1047,12 +1050,12 @@ void USkeletalMeshComponent::PostEditChangeProperty(FPropertyChangedEvent& Prope
 			if(PostProcessAnimInstance)
 			{
 				UClass* CurrentClass = PostProcessAnimInstance->GetClass();
-				UClass* MeshClass = SkeletalMesh ? *SkeletalMesh->GetPostProcessAnimBlueprint() : nullptr;
+				UClass* MeshClass = SkelMesh ? *SkelMesh->GetPostProcessAnimBlueprint() : nullptr;
 				if(CurrentClass != MeshClass)
 				{
 					if(MeshClass)
 					{
-						PostProcessAnimInstance = NewObject<UAnimInstance>(this, *SkeletalMesh->GetPostProcessAnimBlueprint());
+						PostProcessAnimInstance = NewObject<UAnimInstance>(this, *SkelMesh->GetPostProcessAnimBlueprint());
 						PostProcessAnimInstance->InitializeAnimation();
 					}
 					else
@@ -1086,7 +1089,7 @@ void USkeletalMeshComponent::PostEditChangeProperty(FPropertyChangedEvent& Prope
 		if ( PropertyThatChanged->GetFName() == GET_MEMBER_NAME_CHECKED( FSingleAnimationPlayData, AnimToPlay ))
 		{
 			// make sure the animation skeleton matches the current skeletalmesh
-			if (AnimationData.AnimToPlay && SkeletalMesh && SkeletalMesh->GetSkeleton() && SkeletalMesh->GetSkeleton()->IsCompatible(AnimationData.AnimToPlay->GetSkeleton()))
+			if (AnimationData.AnimToPlay && SkelMesh && SkelMesh->GetSkeleton() && SkelMesh->GetSkeleton()->IsCompatible(AnimationData.AnimToPlay->GetSkeleton()))
 			{
 				PlayAnimation(AnimationData.AnimToPlay, false);
 			}
@@ -1151,9 +1154,9 @@ TSoftObjectPtr<UObject> USkeletalMeshComponent::GetDefaultAnimatingRig() const
 	{
 		return DefaultAnimatingRigOverride;
 	}
-	if (SkeletalMesh)
+	if (GetSkeletalMesh())
 	{
-		return SkeletalMesh->GetDefaultAnimatingRig();
+		return GetSkeletalMesh()->GetDefaultAnimatingRig();
 	}
 	return nullptr;
 }
@@ -1195,7 +1198,7 @@ void USkeletalMeshComponent::TickAnimation(float DeltaTime, bool bNeedsValidRoot
 		RecalcRequiredCurves();
 	}
 
-	if (SkeletalMesh != nullptr)
+	if (GetSkeletalMesh() != nullptr)
 	{
 		// We're about to UpdateAnimation, this will potentially queue events that we'll need to dispatch.
 		bNeedsQueuedAnimEventsDispatched = true;
@@ -1270,9 +1273,9 @@ void USkeletalMeshComponent::UpdateVisualizeLODString(FString& DebugString)
 	Super::UpdateVisualizeLODString(DebugString);
 
 	uint32 NumVertices = 0;
-	if (SkeletalMesh)
+	if (GetSkeletalMesh())
 	{
-		if (FSkeletalMeshRenderData* RenderData = SkeletalMesh->GetResourceForRendering())
+		if (FSkeletalMeshRenderData* RenderData = GetSkeletalMesh()->GetResourceForRendering())
 		{
 			if (RenderData->LODRenderData.IsValidIndex(GetPredictedLODLevel()))
 			{
@@ -1421,9 +1424,9 @@ void USkeletalMeshComponent::ResetMorphTargetCurves()
 {
 	ActiveMorphTargets.Reset();
 
-	if (SkeletalMesh)
+	if (GetSkeletalMesh())
 	{
-		MorphTargetWeights.SetNum(SkeletalMesh->GetMorphTargets().Num());
+		MorphTargetWeights.SetNum(GetSkeletalMesh()->GetMorphTargets().Num());
 
 		// we need this code to ensure the buffer gets cleared whether or not you have morphtarget curve set
 		// the case, where you had morphtargets weight on, and when you clear the weight, you want to make sure 
@@ -1441,11 +1444,11 @@ void USkeletalMeshComponent::ResetMorphTargetCurves()
 
 void USkeletalMeshComponent::UpdateMorphTargetOverrideCurves()
 {
-	if (SkeletalMesh)
+	if (GetSkeletalMesh())
 	{
 		if (MorphTargetCurves.Num() > 0)
 		{
-			FAnimationRuntime::AppendActiveMorphTargets(SkeletalMesh, MorphTargetCurves, ActiveMorphTargets, MorphTargetWeights);
+			FAnimationRuntime::AppendActiveMorphTargets(GetSkeletalMesh(), MorphTargetCurves, ActiveMorphTargets, MorphTargetWeights);
 		}
 	}
 }
@@ -1720,14 +1723,14 @@ static void MergeInBoneIndexArrays(TArray<FBoneIndexType>& BaseArray, const TArr
 // if you call RecalcRequiredBones, curve should be refreshed
 void USkeletalMeshComponent::RecalcRequiredCurves()
 {
-	if (!SkeletalMesh)
+	if (!GetSkeletalMesh())
 	{
 		return;
 	}
 
-	if (SkeletalMesh->GetSkeleton())
+	if (GetSkeletalMesh()->GetSkeleton())
 	{
-		CachedCurveUIDList = SkeletalMesh->GetSkeleton()->GetDefaultCurveUIDList();
+		CachedCurveUIDList = GetSkeletalMesh()->GetSkeleton()->GetDefaultCurveUIDList();
 	}
 
 	const FCurveEvaluationOption CurveEvalOption(bAllowAnimCurveEvaluation, &DisallowedAnimCurves, GetPredictedLODLevel());
@@ -1756,7 +1759,8 @@ void USkeletalMeshComponent::ComputeRequiredBones(TArray<FBoneIndexType>& OutReq
 	OutRequiredBones.Reset();
 	OutFillComponentSpaceTransformsRequiredBones.Reset();
 
-	if (!SkeletalMesh)
+	USkeletalMesh* SkelMesh = GetSkeletalMesh();
+	if (!SkelMesh)
 	{
 		return;
 	}
@@ -1766,7 +1770,7 @@ void USkeletalMeshComponent::ComputeRequiredBones(TArray<FBoneIndexType>& OutReq
 	{
 		//No Render Data?
 		// Jira UE-64409
-		UE_LOG(LogAnimation, Warning, TEXT("Skeletal Mesh asset '%s' has no render data"), *SkeletalMesh->GetName());
+		UE_LOG(LogAnimation, Warning, TEXT("Skeletal Mesh asset '%s' has no render data"), *SkelMesh->GetName());
 		return;
 	}
 
@@ -1775,7 +1779,7 @@ void USkeletalMeshComponent::ComputeRequiredBones(TArray<FBoneIndexType>& OutReq
 	if (SkelMeshRenderData->LODRenderData.Num() == 0)
 	{
 		//No LODS?
-		UE_LOG(LogAnimation, Warning, TEXT("Skeletal Mesh asset '%s' has no LODs"), *SkeletalMesh->GetName());
+		UE_LOG(LogAnimation, Warning, TEXT("Skeletal Mesh asset '%s' has no LODs"), *SkelMesh->GetName());
 		return;
 	}
 
@@ -1786,7 +1790,7 @@ void USkeletalMeshComponent::ComputeRequiredBones(TArray<FBoneIndexType>& OutReq
 	OutRequiredBones = LODData.RequiredBones;
 
 	// Add virtual bones
-	MergeInBoneIndexArrays(OutRequiredBones, SkeletalMesh->GetRefSkeleton().GetRequiredVirtualBones());
+	MergeInBoneIndexArrays(OutRequiredBones, SkelMesh->GetRefSkeleton().GetRequiredVirtualBones());
 
 	const UPhysicsAsset* const PhysicsAsset = GetPhysicsAsset();
 	// If we have a PhysicsAsset, we also need to make sure that all the bones used by it are always updated, as its used
@@ -1801,7 +1805,7 @@ void USkeletalMeshComponent::ComputeRequiredBones(TArray<FBoneIndexType>& OutReq
 			{
 				continue;
 			}
-			int32 PhysBoneIndex = SkeletalMesh->GetRefSkeleton().FindBoneIndex(PhysicsAsset->SkeletalBodySetups[i]->BoneName);
+			int32 PhysBoneIndex = SkelMesh->GetRefSkeleton().FindBoneIndex(PhysicsAsset->SkeletalBodySetups[i]->BoneName);
 			if (PhysBoneIndex != INDEX_NONE)
 			{
 				PhysAssetBones.Add(PhysBoneIndex);
@@ -1827,7 +1831,7 @@ void USkeletalMeshComponent::ComputeRequiredBones(TArray<FBoneIndexType>& OutReq
 		
 		if (ensureMsgf(EditableBoneVisibilityStates.Num() >= OutRequiredBones.Num(), 
 			TEXT("Skeletal Mesh asset '%s' has incorrect BoneVisibilityStates. # of BoneVisibilityStatese (%d), # of OutRequiredBones (%d)"), 
-			*SkeletalMesh->GetName(), EditableBoneVisibilityStates.Num(), OutRequiredBones.Num()))
+			*SkelMesh->GetName(), EditableBoneVisibilityStates.Num(), OutRequiredBones.Num()))
 		{
 			int32 VisibleBoneWriteIndex = 0;
 			for (int32 i = 0; i < OutRequiredBones.Num(); ++i)
@@ -1853,13 +1857,13 @@ void USkeletalMeshComponent::ComputeRequiredBones(TArray<FBoneIndexType>& OutReq
 	// Add in any bones that may be required when mirroring.
 	// JTODO: This is only required if there are mirroring nodes in the tree, but hard to know...
 	PRAGMA_DISABLE_DEPRECATION_WARNINGS
-	if (SkeletalMesh->SkelMirrorTable.Num() > 0 &&
-		SkeletalMesh->SkelMirrorTable.Num() == BoneSpaceTransforms.Num())
+	if (SkelMesh->SkelMirrorTable.Num() > 0 &&
+		SkelMesh->SkelMirrorTable.Num() == BoneSpaceTransforms.Num())
 	{
 		TArray<FBoneIndexType> MirroredDesiredBones;
 		MirroredDesiredBones.AddUninitialized(RequiredBones.Num());
 
-		const TArray<struct FBoneMirrorInfo>& SkelMirrorTable = SkeletalMesh->GetSkelMirrorTable();
+		const TArray<struct FBoneMirrorInfo>& SkelMirrorTable = SkelMesh->GetSkelMirrorTable();
 		// Look up each bone in the mirroring table.
 		for (int32 i = 0; i<OutRequiredBones.Num(); i++)
 		{
@@ -1876,11 +1880,11 @@ void USkeletalMeshComponent::ComputeRequiredBones(TArray<FBoneIndexType>& OutReq
 	TArray<FBoneIndexType> NeededBonesForFillComponentSpaceTransforms;
 	{
 		TArray<FBoneIndexType> ForceAnimatedSocketBones;
-		TArray<USkeletalMeshSocket*> ActiveSocketList = SkeletalMesh->GetActiveSocketList();
+		TArray<USkeletalMeshSocket*> ActiveSocketList = SkelMesh->GetActiveSocketList();
 		ForceAnimatedSocketBones.Reserve(ActiveSocketList.Num());
 		for (const USkeletalMeshSocket* Socket : ActiveSocketList)
 		{
-			int32 BoneIndex = SkeletalMesh->GetRefSkeleton().FindBoneIndex(Socket->BoneName);
+			int32 BoneIndex = SkelMesh->GetRefSkeleton().FindBoneIndex(Socket->BoneName);
 			if (BoneIndex != INDEX_NONE)
 			{
 				if (Socket->bForceAlwaysAnimated)
@@ -1914,26 +1918,26 @@ void USkeletalMeshComponent::ComputeRequiredBones(TArray<FBoneIndexType>& OutReq
 	}
 
 	// Ensure that we have a complete hierarchy down to those bones.
-	FAnimationRuntime::EnsureParentsPresent(OutRequiredBones, SkeletalMesh->GetRefSkeleton());
+	FAnimationRuntime::EnsureParentsPresent(OutRequiredBones, SkelMesh->GetRefSkeleton());
 
 	OutFillComponentSpaceTransformsRequiredBones.Reset(OutRequiredBones.Num() + NeededBonesForFillComponentSpaceTransforms.Num());
 	OutFillComponentSpaceTransformsRequiredBones = OutRequiredBones;
 
 	NeededBonesForFillComponentSpaceTransforms.Sort();
 	MergeInBoneIndexArrays(OutFillComponentSpaceTransformsRequiredBones, NeededBonesForFillComponentSpaceTransforms);
-	FAnimationRuntime::EnsureParentsPresent(OutFillComponentSpaceTransformsRequiredBones, SkeletalMesh->GetRefSkeleton());
+	FAnimationRuntime::EnsureParentsPresent(OutFillComponentSpaceTransformsRequiredBones, SkelMesh->GetRefSkeleton());
 }
 
 void USkeletalMeshComponent::RecalcRequiredBones(int32 LODIndex)
 {
-	if (!SkeletalMesh)
+	if (!GetSkeletalMesh())
 	{
 		return;
 	}
 
 	ComputeRequiredBones(RequiredBones, FillComponentSpaceTransformsRequiredBones, LODIndex, /*bIgnorePhysicsAsset=*/ false);
 	PRAGMA_DISABLE_DEPRECATION_WARNINGS
-	BoneSpaceTransforms = SkeletalMesh->GetRefSkeleton().GetRefBonePose();
+	BoneSpaceTransforms = GetSkeletalMesh()->GetRefSkeleton().GetRefBonePose();
 	PRAGMA_ENABLE_DEPRECATION_WARNINGS
 	// make sure animation requiredBone to mark as dirty
 	if (AnimScriptInstance)
@@ -1962,15 +1966,15 @@ void USkeletalMeshComponent::RecalcRequiredBones(int32 LODIndex)
 
 void USkeletalMeshComponent::MarkRequiredCurveUpToDate()
 {
-	if (SkeletalMesh && SkeletalMesh->GetSkeleton())
+	if (GetSkeletalMesh() && GetSkeletalMesh()->GetSkeleton())
 	{
-		CachedAnimCurveUidVersion = SkeletalMesh->GetSkeleton()->GetAnimCurveUidVersion();
+		CachedAnimCurveUidVersion = GetSkeletalMesh()->GetSkeleton()->GetAnimCurveUidVersion();
 	}
 }
 
 bool USkeletalMeshComponent::AreRequiredCurvesUpToDate() const
 {
-	return (!SkeletalMesh || !SkeletalMesh->GetSkeleton() || CachedAnimCurveUidVersion == SkeletalMesh->GetSkeleton()->GetAnimCurveUidVersion());
+	return (!GetSkeletalMesh() || !GetSkeletalMesh()->GetSkeleton() || CachedAnimCurveUidVersion == GetSkeletalMesh()->GetSkeleton()->GetAnimCurveUidVersion());
 }
 
 void USkeletalMeshComponent::EvaluateAnimation(const USkeletalMesh* InSkeletalMesh, UAnimInstance* InAnimInstance, FVector& OutRootBoneTranslation, FBlendedHeapCurve& OutCurve, FCompactPose& OutPose, UE::Anim::FHeapAttributeContainer& OutAttributes) const
@@ -2013,18 +2017,18 @@ void USkeletalMeshComponent::UpdateSlaveComponent()
 		// we changed order of morphtarget to be overriden by SetMorphTarget from BP
 		// so this has to go first
 		// now propagate BP-driven curves from the master SMC...
-		if (SkeletalMesh)
+		if (GetSkeletalMesh())
 		{
-			check(MorphTargetWeights.Num() == SkeletalMesh->GetMorphTargets().Num());
+			check(MorphTargetWeights.Num() == GetSkeletalMesh()->GetMorphTargets().Num());
 			if (MasterSMC->MorphTargetCurves.Num() > 0)
 			{
-				FAnimationRuntime::AppendActiveMorphTargets(SkeletalMesh, MasterSMC->MorphTargetCurves, ActiveMorphTargets, MorphTargetWeights);
+				FAnimationRuntime::AppendActiveMorphTargets(GetSkeletalMesh(), MasterSMC->MorphTargetCurves, ActiveMorphTargets, MorphTargetWeights);
 			}
 
 			// if slave also has it, add it here. 
 			if (MorphTargetCurves.Num() > 0)
 			{
-				FAnimationRuntime::AppendActiveMorphTargets(SkeletalMesh, MorphTargetCurves, ActiveMorphTargets, MorphTargetWeights);
+				FAnimationRuntime::AppendActiveMorphTargets(GetSkeletalMesh(), MorphTargetCurves, ActiveMorphTargets, MorphTargetWeights);
 			}
 		}
 
@@ -2302,7 +2306,7 @@ void USkeletalMeshComponent::RefreshBoneTransforms(FActorComponentTickFunction*
 
 	check(IsInGameThread()); //Only want to call this from the game thread as we set up tasks etc
 	
-	if (!SkeletalMesh || GetNumComponentSpaceTransforms() == 0)
+	if (!GetSkeletalMesh() || GetNumComponentSpaceTransforms() == 0)
 	{
 		return;
 	}
@@ -2327,7 +2331,7 @@ void USkeletalMeshComponent::RefreshBoneTransforms(FActorComponentTickFunction*
 	//Dont mark cache as invalid if we aren't performing optimization anyway
 	PRAGMA_DISABLE_DEPRECATION_WARNINGS
 	const bool bInvalidCachedBones = bDoEvaluationRateOptimization &&
-									 ((BoneSpaceTransforms.Num() != SkeletalMesh->GetRefSkeleton().GetNum())
+									 ((BoneSpaceTransforms.Num() != GetSkeletalMesh()->GetRefSkeleton().GetNum())
 									 || (BoneSpaceTransforms.Num() != CachedBoneSpaceTransforms.Num())
 									 || (GetNumComponentSpaceTransforms() != CachedComponentSpaceTransforms.Num()));
 
@@ -2369,7 +2373,7 @@ void USkeletalMeshComponent::RefreshBoneTransforms(FActorComponentTickFunction*
 		return;
 	}
 
-	AnimEvaluationContext.SkeletalMesh = SkeletalMesh;
+	AnimEvaluationContext.SkeletalMesh = GetSkeletalMesh();
 	AnimEvaluationContext.AnimInstance = AnimScriptInstance;
 	AnimEvaluationContext.PostProcessAnimInstance = (ShouldEvaluatePostProcessInstance())? ToRawPtr(PostProcessAnimInstance): nullptr;
 
@@ -2515,9 +2519,9 @@ void USkeletalMeshComponent::DispatchParallelEvaluationTasks(FActorComponentTick
 
 #if WITH_EDITOR
 	// We can only finish compilation on the game-thread, so wait here before spawning eval tasks.
-	if (SkeletalMesh && SkeletalMesh->IsCompiling())
+	if (GetSkeletalMesh() && GetSkeletalMesh()->IsCompiling())
 	{
-		FSkeletalMeshCompilingManager::Get().FinishCompilation({ SkeletalMesh });
+		FSkeletalMeshCompilingManager::Get().FinishCompilation({ GetSkeletalMesh() });
 	}
 #endif
 
@@ -2549,7 +2553,7 @@ void USkeletalMeshComponent::DispatchParallelTickPose(FActorComponentTickFunctio
 {
 	check(TickFunction);
 	
-	if(SkeletalMesh != nullptr)
+	if(GetSkeletalMesh() != nullptr)
 	{
 		if ((AnimScriptInstance && AnimScriptInstance->NeedsUpdate()) ||
 			(PostProcessAnimInstance && PostProcessAnimInstance->NeedsUpdate()))
@@ -2573,7 +2577,7 @@ void USkeletalMeshComponent::DispatchParallelTickPose(FActorComponentTickFunctio
 				}
 
 				// Do a mini-setup of the eval context
-				AnimEvaluationContext.SkeletalMesh = SkeletalMesh;
+				AnimEvaluationContext.SkeletalMesh = GetSkeletalMesh();
 				AnimEvaluationContext.AnimInstance = AnimScriptInstance;
 
 				// We dont set up the Curve here, as we dont use it in Update()
@@ -2633,7 +2637,7 @@ void USkeletalMeshComponent::PostAnimEvaluation(FAnimationEvaluationContext& Eva
 		if (EvaluationContext.bDuplicateToCacheCurve)
 		{
 			ensureAlwaysMsgf(AnimCurves.IsValid(), TEXT("Animation Curve is invalid (%s). TotalCount(%d) "),
-				*GetPathNameSafe(SkeletalMesh), AnimCurves.NumValidCurveCount);
+				*GetPathNameSafe(GetSkeletalMesh()), AnimCurves.NumValidCurveCount);
 			CachedCurve.CopyFrom(AnimCurves);
 		}
 
@@ -2686,7 +2690,7 @@ void USkeletalMeshComponent::PostAnimEvaluation(FAnimationEvaluationContext& Eva
 
 			PRAGMA_DISABLE_DEPRECATION_WARNINGS
 			FAnimationRuntime::LerpBoneTransforms(BoneSpaceTransforms, CachedBoneSpaceTransforms, Alpha, RequiredBones);
-			FillComponentSpaceTransforms(SkeletalMesh, BoneSpaceTransforms, GetEditableComponentSpaceTransforms());
+			FillComponentSpaceTransforms(GetSkeletalMesh(), BoneSpaceTransforms, GetEditableComponentSpaceTransforms());
 			PRAGMA_ENABLE_DEPRECATION_WARNINGS
 			// interpolate curve
 			AnimCurves.LerpTo(CachedCurve, Alpha);
@@ -2815,10 +2819,10 @@ void USkeletalMeshComponent::ApplyAnimationCurvesToComponent(const TMap<FName, f
 	}
 
 	const bool bContainsMorphCurves = InAnimationMorphCurves && InAnimationMorphCurves->Num() > 0;
-	if (SkeletalMesh && bContainsMorphCurves)
+	if (GetSkeletalMesh() && bContainsMorphCurves)
 	{
 		// we want to append to existing curves - i.e. BP driven curves 
-		FAnimationRuntime::AppendActiveMorphTargets(SkeletalMesh, *InAnimationMorphCurves, ActiveMorphTargets, MorphTargetWeights);
+		FAnimationRuntime::AppendActiveMorphTargets(GetSkeletalMesh(), *InAnimationMorphCurves, ActiveMorphTargets, MorphTargetWeights);
 	}
 
 	/** Push through curves to slave components */
@@ -2874,7 +2878,7 @@ FBoxSphereBounds USkeletalMeshComponent::CalcBounds(const FTransform& LocalToWor
 		{
 			const USkinnedMeshComponent* const MasterPoseComponentInst = MasterPoseComponent.Get();
 			check(MasterPoseComponentInst);
-			if (MasterPoseComponentInst->SkeletalMesh &&
+			if (MasterPoseComponentInst->GetSkeletalMesh() &&
 				MasterPoseComponentInst->bComponentUseFixedSkelBounds &&
 				MasterPoseComponentInst->IsA((USkeletalMeshComponent::StaticClass())))
 			{
@@ -2920,7 +2924,7 @@ void USkeletalMeshComponent::SetSkeletalMesh(USkeletalMesh* InSkelMesh, bool bRe
 	QUICK_SCOPE_CYCLE_COUNTER(STAT_SetSkeletalMesh);
 	SCOPE_CYCLE_UOBJECT(NewSkelMesh, InSkelMesh);
 
-	if (InSkelMesh == SkeletalMesh)
+	if (InSkelMesh == GetSkeletalMesh())
 	{
 		// do nothing if the input mesh is the same mesh we're already using.
 		return;
@@ -2987,9 +2991,9 @@ bool USkeletalMeshComponent::AllocateTransformData()
 	if ( Super::AllocateTransformData() )
 	{
 		PRAGMA_DISABLE_DEPRECATION_WARNINGS
-		if(BoneSpaceTransforms.Num() != SkeletalMesh->GetRefSkeleton().GetNum() )
+		if(BoneSpaceTransforms.Num() != GetSkeletalMesh()->GetRefSkeleton().GetNum() )
 		{
-			BoneSpaceTransforms = SkeletalMesh->GetRefSkeleton().GetRefBonePose();
+			BoneSpaceTransforms = GetSkeletalMesh()->GetRefSkeleton().GetRefBonePose();
 		}
 		PRAGMA_ENABLE_DEPRECATION_WARNINGS
 
@@ -3220,7 +3224,7 @@ void USkeletalMeshComponent::HideBone( int32 BoneIndex, EPhysBodyOp PhysBodyOpti
 {
 	Super::HideBone(BoneIndex, PhysBodyOption);
 
-	if (!SkeletalMesh)
+	if (!GetSkeletalMesh())
 	{
 		return;
 	}
@@ -3237,7 +3241,7 @@ void USkeletalMeshComponent::HideBone( int32 BoneIndex, EPhysBodyOp PhysBodyOpti
 
 		if (PhysBodyOption != PBO_None)
 		{
-			FName HideBoneName = SkeletalMesh->GetRefSkeleton().GetBoneName(BoneIndex);
+			FName HideBoneName = GetSkeletalMesh()->GetRefSkeleton().GetBoneName(BoneIndex);
 			if (PhysBodyOption == PBO_Term)
 			{
 				TermBodiesBelow(HideBoneName);
@@ -3246,7 +3250,7 @@ void USkeletalMeshComponent::HideBone( int32 BoneIndex, EPhysBodyOp PhysBodyOpti
 	}
 	else
 	{
-		UE_LOG(LogSkeletalMesh, Warning, TEXT("HideBone[%s]: Invalid Body Index (%d) has entered. This component doesn't contain buffer for the given body."), *GetPathNameSafe(SkeletalMesh), BoneIndex);
+		UE_LOG(LogSkeletalMesh, Warning, TEXT("HideBone[%s]: Invalid Body Index (%d) has entered. This component doesn't contain buffer for the given body."), *GetPathNameSafe(GetSkeletalMesh()), BoneIndex);
 	}
 }
 
@@ -3254,7 +3258,7 @@ void USkeletalMeshComponent::UnHideBone( int32 BoneIndex )
 {
 	Super::UnHideBone(BoneIndex);
 
-	if (!SkeletalMesh)
+	if (!GetSkeletalMesh())
 	{
 		return;
 	}
@@ -3268,7 +3272,7 @@ void USkeletalMeshComponent::UnHideBone( int32 BoneIndex )
 	{
 		bRequiredBonesUpToDate = false;
 
-		//FName HideBoneName = SkeletalMesh->GetRefSkeleton().GetBoneName(BoneIndex);
+		//FName HideBoneName = GetSkeletalMesh()->GetRefSkeleton().GetBoneName(BoneIndex);
 		// It's okay to turn this on for terminated bodies
 		// It won't do any if BodyData isn't found
 		// @JTODO
@@ -3276,7 +3280,7 @@ void USkeletalMeshComponent::UnHideBone( int32 BoneIndex )
 	}
 	else
 	{
-		UE_LOG(LogSkeletalMesh, Warning, TEXT("UnHideBone[%s]: Invalid Body Index (%d) has entered. This component doesn't contain buffer for the given body."), *GetPathNameSafe(SkeletalMesh), BoneIndex);
+		UE_LOG(LogSkeletalMesh, Warning, TEXT("UnHideBone[%s]: Invalid Body Index (%d) has entered. This component doesn't contain buffer for the given body."), *GetPathNameSafe(GetSkeletalMesh()), BoneIndex);
 	}
 }
 
@@ -3391,7 +3395,7 @@ void USkeletalMeshComponent::SetAnimationMode(EAnimationMode::Type InAnimationMo
 	// when mode is swapped, make sure to reinitialize
 	// even if it was same mode, this was due to users who wants to use BP construction script to do this
 	// if you use it in the construction script, it gets serialized, but it never instantiate. 
-	if(SkeletalMesh != nullptr && (bNeedChange || AnimationMode == EAnimationMode::AnimationBlueprint))
+	if(GetSkeletalMesh() != nullptr && (bNeedChange || AnimationMode == EAnimationMode::AnimationBlueprint))
 	{
 		if (InitializeAnimScriptInstance(true))
 		{
@@ -3771,9 +3775,10 @@ void USkeletalMeshComponent::UnregisterOnSkeletalMeshPropertyChanged( FDelegateH
 
 void USkeletalMeshComponent::ValidateAnimation()
 {
-	if (SkeletalMesh && SkeletalMesh->GetSkeleton() == nullptr)
+	USkeletalMesh* SkelMesh = GetSkeletalMesh();
+	if (SkelMesh && SkelMesh->GetSkeleton() == nullptr)
 	{
-		UE_LOG(LogAnimation, Warning, TEXT("SkeletalMesh %s has no skeleton. This needs to fixed before an animation can be set"), *SkeletalMesh->GetFullName());
+		UE_LOG(LogAnimation, Warning, TEXT("SkeletalMesh %s has no skeleton. This needs to fixed before an animation can be set"), *SkelMesh->GetFullName());
 		if (AnimationMode == EAnimationMode::AnimationSingleNode)
 		{
 			AnimationData.AnimToPlay = nullptr;
@@ -3792,14 +3797,14 @@ void USkeletalMeshComponent::ValidateAnimation()
 
 	if(AnimationMode == EAnimationMode::AnimationSingleNode)
 	{
-		if (AnimationData.AnimToPlay && SkeletalMesh)
+		if (AnimationData.AnimToPlay && SkelMesh)
 		{
 			if (AnimationData.AnimToPlay->GetSkeleton() == nullptr)
 			{
 				UE_LOG(LogAnimation, Warning, TEXT("Animation %s is incompatible because it has no skeleton, removing animation from actor."), *AnimationData.AnimToPlay->GetName());
 				AnimationData.AnimToPlay = nullptr;
 			}
-			else if (!AnimationData.AnimToPlay->GetSkeleton()->IsCompatible(SkeletalMesh->GetSkeleton()))
+			else if (!AnimationData.AnimToPlay->GetSkeleton()->IsCompatible(SkelMesh->GetSkeleton()))
 			{
 				UE_LOG(LogAnimation, Warning, TEXT("Animation %s is incompatible with the skeletal mesh's skeleton, removing animation from actor."), *AnimationData.AnimToPlay->GetName());
 				AnimationData.AnimToPlay = nullptr;
@@ -3809,21 +3814,21 @@ void USkeletalMeshComponent::ValidateAnimation()
 	else if (AnimationMode == EAnimationMode::AnimationBlueprint)
 	{
 		IAnimClassInterface* AnimClassInterface = IAnimClassInterface::GetFromClass(AnimClass);
-		if (AnimClassInterface && SkeletalMesh)
+		if (AnimClassInterface && SkelMesh)
 		{
 			if (AnimClassInterface->GetTargetSkeleton() == nullptr)
 			{
 				UE_LOG(LogAnimation, Warning, TEXT("AnimBP %s is incompatible because it has no target skeleton, removing AnimBP from actor."), *AnimClass->GetName());
 				AnimClass = nullptr;
 			}
-			else if (SkeletalMesh->GetSkeleton() == nullptr)
+			else if (SkelMesh->GetSkeleton() == nullptr)
 			{
-				UE_LOG(LogAnimation, Warning, TEXT("AnimBP %s is incompatible because mesh %s has no skeleton, removing AnimBP from actor."), *AnimClass->GetName(), *SkeletalMesh->GetName());
+				UE_LOG(LogAnimation, Warning, TEXT("AnimBP %s is incompatible because mesh %s has no skeleton, removing AnimBP from actor."), *AnimClass->GetName(), *SkelMesh->GetName());
 				AnimClass = nullptr;
 			}
-			else if (!AnimClassInterface->GetTargetSkeleton()->IsCompatible(SkeletalMesh->GetSkeleton()))
+			else if (!AnimClassInterface->GetTargetSkeleton()->IsCompatible(SkelMesh->GetSkeleton()))
 			{
-				UE_LOG(LogAnimation, Warning, TEXT("AnimBP %s is incompatible with skeleton %s, removing AnimBP from actor."), *AnimClass->GetName(), *SkeletalMesh->GetSkeleton()->GetName());
+				UE_LOG(LogAnimation, Warning, TEXT("AnimBP %s is incompatible with skeleton %s, removing AnimBP from actor."), *AnimClass->GetName(), *SkelMesh->GetSkeleton()->GetName());
 				AnimClass = nullptr;
 			}
 		}
@@ -3893,7 +3898,7 @@ void USkeletalMeshComponent::RefreshMorphTargets()
 {
 	ResetMorphTargetCurves();
 
-	if (SkeletalMesh && AnimScriptInstance)
+	if (GetSkeletalMesh() && AnimScriptInstance)
 	{
 		// as this can be called from any worker thread (i.e. from CreateRenderState_Concurrent) we cant currently be doing parallel evaluation
 		check(!IsRunningParallelEvaluation());
@@ -3953,7 +3958,7 @@ void USkeletalMeshComponent::ParallelDuplicateAndInterpolate(FAnimationEvaluatio
 		if (InAnimEvaluationContext.bDuplicateToCacheCurve)
 		{
 			ensureAlwaysMsgf(InAnimEvaluationContext.Curve.IsValid(), TEXT("Animation Curve is invalid (%s). TotalCount(%d) "),
-				*GetPathNameSafe(SkeletalMesh), InAnimEvaluationContext.Curve.NumValidCurveCount);
+				*GetPathNameSafe(GetSkeletalMesh()), InAnimEvaluationContext.Curve.NumValidCurveCount);
 			InAnimEvaluationContext.CachedCurve.CopyFrom(InAnimEvaluationContext.Curve);
 		}
 
@@ -4020,7 +4025,7 @@ void USkeletalMeshComponent::CompleteParallelAnimationEvaluation(bool bDoPostAni
 	SCOPED_NAMED_EVENT(USkeletalMeshComponent_CompleteParallelAnimationEvaluation, FColor::Yellow);
 	ParallelAnimationEvaluationTask.SafeRelease(); //We are done with this task now, clean up!
 
-	if (bDoPostAnimEvaluation && (AnimEvaluationContext.AnimInstance == AnimScriptInstance) && (AnimEvaluationContext.SkeletalMesh == SkeletalMesh) && (AnimEvaluationContext.ComponentSpaceTransforms.Num() == GetNumComponentSpaceTransforms()))
+	if (bDoPostAnimEvaluation && (AnimEvaluationContext.AnimInstance == AnimScriptInstance) && (AnimEvaluationContext.SkeletalMesh == GetSkeletalMesh()) && (AnimEvaluationContext.ComponentSpaceTransforms.Num() == GetNumComponentSpaceTransforms()))
 	{
 		SwapEvaluationContextBuffers();
 
@@ -4071,7 +4076,7 @@ void USkeletalMeshComponent::BindClothToMasterPoseComponent()
 {
 	if(USkeletalMeshComponent* MasterComp = Cast<USkeletalMeshComponent>(MasterPoseComponent.Get()))
 	{
-		if(SkeletalMesh != MasterComp->SkeletalMesh)
+		if(GetSkeletalMesh() != MasterComp->GetSkeletalMesh())
 		{
 			// Not the same mesh, can't bind
 			return;
@@ -4179,18 +4184,19 @@ void USkeletalMeshComponent::FinalizeBoneTransform()
 
 void USkeletalMeshComponent::GetCurrentRefToLocalMatrices(TArray<FMatrix44f>& OutRefToLocals, int32 InLodIdx) const
 {
-	if(SkeletalMesh)
+	USkeletalMesh* SkelMesh = GetSkeletalMesh();
+	if(SkelMesh)
 	{
-		FSkeletalMeshRenderData* RenderData = SkeletalMesh->GetResourceForRendering();
+		FSkeletalMeshRenderData* RenderData = SkelMesh->GetResourceForRendering();
 		if (ensureMsgf(RenderData->LODRenderData.IsValidIndex(InLodIdx),
 			TEXT("GetCurrentRefToLocalMatrices (SkelMesh :%s) input LODIndex (%d) doesn't match with render data size (%d)."),
-			*SkeletalMesh->GetPathName(), InLodIdx, RenderData->LODRenderData.Num()))
+			*SkelMesh->GetPathName(), InLodIdx, RenderData->LODRenderData.Num()))
 		{
 			UpdateRefToLocalMatrices(OutRefToLocals, this, RenderData, InLodIdx, nullptr);
 		}
 		else
 		{
-			const FReferenceSkeleton& RefSkeleton = SkeletalMesh->GetRefSkeleton();
+			const FReferenceSkeleton& RefSkeleton = SkelMesh->GetRefSkeleton();
 			OutRefToLocals.AddUninitialized(RefSkeleton.GetNum());
 			for (int32 Index = 0; Index < OutRefToLocals.Num(); ++Index)
 			{
@@ -4313,13 +4319,14 @@ void USkeletalMeshComponent::RemoveSlavePoseComponent(USkinnedMeshComponent* Ski
 
 void USkeletalMeshComponent::SnapshotPose(FPoseSnapshot& Snapshot)
 {
-	if (ensureAsRuntimeWarning(SkeletalMesh != nullptr))
+	USkeletalMesh* SkelMesh = GetSkeletalMesh();
+	if (ensureAsRuntimeWarning(SkelMesh != nullptr))
 	{
 		const TArray<FTransform>& ComponentSpaceTMs = GetComponentSpaceTransforms();
-		const FReferenceSkeleton& RefSkeleton = SkeletalMesh->GetRefSkeleton();
+		const FReferenceSkeleton& RefSkeleton = SkelMesh->GetRefSkeleton();
 		const TArray<FTransform>& RefPoseSpaceBaseTMs = RefSkeleton.GetRefBonePose();
 
-		Snapshot.SkeletalMeshName = SkeletalMesh->GetFName();
+		Snapshot.SkeletalMeshName = SkelMesh->GetFName();
 
 		const int32 NumSpaceBases = ComponentSpaceTMs.Num();
 		Snapshot.LocalTransforms.Reset(NumSpaceBases);
@@ -4467,9 +4474,9 @@ void USkeletalMeshComponent::SetAllowedAnimCurvesEvaluation(const TArray<FName>&
 			const TArray<FName>& AllowedList;
 		};
 
-		if (SkeletalMesh)
+		if (GetSkeletalMesh())
 		{
-			USkeleton* Skeleton = SkeletalMesh->GetSkeleton();
+			USkeleton* Skeleton = GetSkeletalMesh()->GetSkeleton();
 			if (Skeleton)
 			{
 				const FSmartNameMapping* Mapping = Skeleton->GetSmartNameContainer(USkeleton::AnimCurveMappingName);
@@ -4544,10 +4551,11 @@ bool USkeletalMeshComponent::FindAttributeChecked(const FName& BoneName, const F
 	OutValue = DefaultValue;
 	bool bFound = false;
 
-	if (SkeletalMesh)
+	USkeletalMesh* SkelMesh = GetSkeletalMesh();
+	if (SkelMesh)
 	{
 		const UE::Anim::FMeshAttributeContainer& Attributes = GetCustomAttributes();
-		const int32 BoneIndex = SkeletalMesh->GetRefSkeleton().FindBoneIndex(BoneName);
+		const int32 BoneIndex = SkelMesh->GetRefSkeleton().FindBoneIndex(BoneName);
 
 		const CustomAttributeType* AttributePtr = Attributes.Find<CustomAttributeType>(UE::Anim::FAttributeId(AttributeName, FCompactPoseBoneIndex(BoneIndex)));
 
@@ -4555,7 +4563,7 @@ bool USkeletalMeshComponent::FindAttributeChecked(const FName& BoneName, const F
 		{
 			if (LookupType == ECustomBoneAttributeLookup::ImmediateParent)
 			{
-				const int32 ParentIndex = SkeletalMesh->GetRefSkeleton().GetParentIndex(BoneIndex);
+				const int32 ParentIndex = SkelMesh->GetRefSkeleton().GetParentIndex(BoneIndex);
 				if (ParentIndex != INDEX_NONE)
 				{
 					AttributePtr = Attributes.Find<CustomAttributeType>(UE::Anim::FAttributeId(AttributeName, FCompactPoseBoneIndex(ParentIndex)));
@@ -4564,7 +4572,7 @@ bool USkeletalMeshComponent::FindAttributeChecked(const FName& BoneName, const F
 			else if (LookupType == ECustomBoneAttributeLookup::ParentHierarchy)
 			{
 				int32 SearchBoneIndex = BoneIndex;
-				int32 ParentIndex = SkeletalMesh->GetRefSkeleton().GetParentIndex(SearchBoneIndex);
+				int32 ParentIndex = SkelMesh->GetRefSkeleton().GetParentIndex(SearchBoneIndex);
 
 				while (ParentIndex != INDEX_NONE)
 				{
@@ -4575,7 +4583,7 @@ bool USkeletalMeshComponent::FindAttributeChecked(const FName& BoneName, const F
 					}
 
 					SearchBoneIndex = ParentIndex;
-					ParentIndex = SkeletalMesh->GetRefSkeleton().GetParentIndex(SearchBoneIndex);
+					ParentIndex = SkelMesh->GetRefSkeleton().GetParentIndex(SearchBoneIndex);
 				}
 			}
 		}
diff --git a/Engine/Source/Runtime/Engine/Private/Components/SkinnedMeshComponent.cpp b/Engine/Source/Runtime/Engine/Private/Components/SkinnedMeshComponent.cpp
--- a/Engine/Source/Runtime/Engine/Private/Components/SkinnedMeshComponent.cpp
+++ b/Engine/Source/Runtime/Engine/Private/Components/SkinnedMeshComponent.cpp
@@ -402,10 +402,10 @@ void USkinnedMeshComponent::UpdateMorphMaterialUsageOnProxy()
 	// update morph material usage
 	if (SceneProxy)
 	{
-		if (ActiveMorphTargets.Num() > 0 && SkeletalMesh->GetMorphTargets().Num() > 0)
+		if (ActiveMorphTargets.Num() > 0 && GetSkeletalMesh()->GetMorphTargets().Num() > 0)
 		{
 			TArray<UMaterialInterface*> MaterialUsingMorphTarget;
-			for (UMorphTarget* MorphTarget : SkeletalMesh->GetMorphTargets())
+			for (UMorphTarget* MorphTarget : GetSkeletalMesh()->GetMorphTargets())
 			{
 				if (!MorphTarget)
 				{
@@ -415,9 +415,9 @@ void USkinnedMeshComponent::UpdateMorphMaterialUsageOnProxy()
 				{
 					for (int32 SectionIndex : MorphTargetLODModel.SectionIndices)
 					{
-						for (int32 LodIdx = 0; LodIdx < SkeletalMesh->GetResourceForRendering()->LODRenderData.Num(); LodIdx++)
+						for (int32 LodIdx = 0; LodIdx < GetSkeletalMesh()->GetResourceForRendering()->LODRenderData.Num(); LodIdx++)
 						{
-							const FSkeletalMeshLODRenderData& LODModel = SkeletalMesh->GetResourceForRendering()->LODRenderData[LodIdx];
+							const FSkeletalMeshLODRenderData& LODModel = GetSkeletalMesh()->GetResourceForRendering()->LODRenderData[LodIdx];
 							if (LODModel.RenderSections.IsValidIndex(SectionIndex))
 							{
 								MaterialUsingMorphTarget.AddUnique(GetMaterial(LODModel.RenderSections[SectionIndex].MaterialIndex));
@@ -544,10 +544,10 @@ void USkinnedMeshComponent::CreateRenderState_Concurrent(FRegisterComponentConte
 {
 	LLM_SCOPE(ELLMTag::SkeletalMesh);
 
-	if( SkeletalMesh )
+	if( GetSkeletalMesh() )
 	{
 		// Attempting to track down UE-45505, where it looks as if somehow a skeletal mesh component's mesh has only been partially loaded, causing a mismatch in the LOD arrays
-		checkf(!SkeletalMesh->HasAnyFlags(RF_NeedLoad | RF_NeedPostLoad | RF_NeedPostLoadSubobjects | RF_WillBeLoaded), TEXT("Attempting to create render state for a skeletal mesh that is is not fully loaded. Mesh: %s"), *SkeletalMesh->GetName());
+		checkf(!GetSkeletalMesh()->HasAnyFlags(RF_NeedLoad | RF_NeedPostLoad | RF_NeedPostLoadSubobjects | RF_WillBeLoaded), TEXT("Attempting to create render state for a skeletal mesh that is is not fully loaded. Mesh: %s"), *GetSkeletalMesh()->GetName());
 
 		// Initialize the alternate weight tracks if present BEFORE creating the new mesh object
 		InitLODInfos();
@@ -556,7 +556,7 @@ void USkinnedMeshComponent::CreateRenderState_Concurrent(FRegisterComponentConte
 		if ( FApp::CanEverRender() && ShouldComponentAddToScene() )
 		{
 			ERHIFeatureLevel::Type SceneFeatureLevel = GetWorld()->FeatureLevel;
-			FSkeletalMeshRenderData* SkelMeshRenderData = SkeletalMesh->GetResourceForRendering();
+			FSkeletalMeshRenderData* SkelMeshRenderData = GetSkeletalMesh()->GetResourceForRendering();
 			int32 MinLODIndex = ComputeMinLOD();
 			
 #if DO_CHECK
@@ -566,7 +566,7 @@ void USkinnedMeshComponent::CreateRenderState_Concurrent(FRegisterComponentConte
 				const FPositionVertexBuffer* PositionVertexBufferPtr = &LODData.StaticVertexBuffers.PositionVertexBuffer;
 				if (!PositionVertexBufferPtr || (PositionVertexBufferPtr->GetNumVertices() <= 0))
 				{
-					UE_LOG(LogSkinnedMeshComp, Warning, TEXT("Invalid Lod %i for Rendering Asset: %s"), LODIndex, *SkeletalMesh->GetFullName());
+					UE_LOG(LogSkinnedMeshComp, Warning, TEXT("Invalid Lod %i for Rendering Asset: %s"), LODIndex, *GetSkeletalMesh()->GetFullName());
 				}
 			}
 #endif
@@ -601,7 +601,7 @@ void USkinnedMeshComponent::CreateRenderState_Concurrent(FRegisterComponentConte
 					FString FeatureLevelName; GetFeatureLevelName(SceneFeatureLevel, FeatureLevelName);
 
 					UE_LOG(LogSkinnedMeshComp, Warning, TEXT("SkeletalMesh %s, is not supported for current feature level (%s) and will not be rendered. MinLOD %d, NumBones %d (supported %d), NumBoneInfluences: %d"),
-						*GetNameSafe(SkeletalMesh), *FeatureLevelName, MinLODIndex, MaxBonesPerChunk, MaxSupportedGPUSkinBones, NumBoneInfluences);
+						*GetNameSafe(GetSkeletalMesh()), *FeatureLevelName, MinLODIndex, MaxBonesPerChunk, MaxSupportedGPUSkinBones, NumBoneInfluences);
 				}
 			}
 
@@ -620,7 +620,7 @@ void USkinnedMeshComponent::CreateRenderState_Concurrent(FRegisterComponentConte
 
 	Super::CreateRenderState_Concurrent(Context);
 
-	if (SkeletalMesh)
+	if (GetSkeletalMesh())
 	{
 		// Update dynamic data
 
@@ -640,14 +640,14 @@ void USkinnedMeshComponent::CreateRenderState_Concurrent(FRegisterComponentConte
 			}
 
 			// Clamp to loaded streaming data if available
-			if (SkeletalMesh->IsStreamable() && MeshObject)
+			if (GetSkeletalMesh()->IsStreamable() && MeshObject)
 			{
 				ModifiedLODLevel = FMath::Max<int32>(ModifiedLODLevel, MeshObject->GetSkeletalMeshRenderData().PendingFirstLODIdx);
 			}
 
 			// If we have a valid LOD, set up required data, during reimport we may try to create data before we have all the LODs
 			// imported, in that case we skip until we have all the LODs
-			if(SkeletalMesh->IsValidLODIndex(ModifiedLODLevel))
+			if(GetSkeletalMesh()->IsValidLODIndex(ModifiedLODLevel))
 			{
 				const bool bMorphTargetsAllowed = CVarEnableMorphTargets.GetValueOnAnyThread(true) != 0;
 
@@ -706,9 +706,9 @@ FString USkinnedMeshComponent::GetDetailedInfoInternal() const
 {
 	FString Result;  
 
-	if( SkeletalMesh != NULL )
+	if(GetSkeletalMesh() != NULL )
 	{
-		Result = SkeletalMesh->GetDetailedInfoInternal();
+		Result = GetSkeletalMesh()->GetDetailedInfoInternal();
 	}
 	else
 	{
@@ -726,13 +726,13 @@ void USkinnedMeshComponent::SendRenderDynamicData_Concurrent()
 
 	// if we have not updated the transforms then no need to send them to the rendering thread
 	// @todo GIsEditor used to be bUpdateSkelWhenNotRendered. Look into it further to find out why it doesn't update animations in the AnimSetViewer, when a level is loaded in UED (like POC_Cover.gear).
-	if( MeshObject && SkeletalMesh && (bForceMeshObjectUpdate || (bRecentlyRendered || VisibilityBasedAnimTickOption == EVisibilityBasedAnimTickOption::AlwaysTickPoseAndRefreshBones || GIsEditor || MeshObject->bHasBeenUpdatedAtLeastOnce == false)) )
+	if( MeshObject && GetSkeletalMesh() && (bForceMeshObjectUpdate || (bRecentlyRendered || VisibilityBasedAnimTickOption == EVisibilityBasedAnimTickOption::AlwaysTickPoseAndRefreshBones || GIsEditor || MeshObject->bHasBeenUpdatedAtLeastOnce == false)) )
 	{
 		SCOPE_CYCLE_COUNTER(STAT_MeshObjectUpdate);
 
 		int32 UseLOD = GetPredictedLODLevel();
 		// Clamp to loaded streaming data if available
-		if (SkeletalMesh->IsStreamable() && MeshObject)
+		if (GetSkeletalMesh()->IsStreamable() && MeshObject)
 		{
 			UseLOD = FMath::Max<int32>(UseLOD, MeshObject->GetSkeletalMeshRenderData().PendingFirstLODIdx);
 		}
@@ -761,7 +761,7 @@ void USkinnedMeshComponent::SendRenderDynamicData_Concurrent()
 		{
 			if (MeshDeformerInstance->IsActive())
 			{
-				MeshDeformerInstance->EnqueueWork(GetScene(), UMeshDeformerInstance::WorkLoad_Update, SkeletalMesh->GetFName());
+				MeshDeformerInstance->EnqueueWork(GetScene(), UMeshDeformerInstance::WorkLoad_Update, GetSkeletalMesh()->GetFName());
 			}
 			else
 			{
@@ -779,7 +779,7 @@ void USkinnedMeshComponent::ClearMotionVector()
 	{
 		int32 UseLOD = GetPredictedLODLevel();
 		// Clamp to loaded streaming data if available
-		if (SkeletalMesh->IsStreamable() && MeshObject)
+		if (GetSkeletalMesh()->IsStreamable() && MeshObject)
 		{
 			UseLOD = FMath::Max<int32>(UseLOD, MeshObject->GetSkeletalMeshRenderData().PendingFirstLODIdx);
 		}
@@ -805,7 +805,7 @@ void USkinnedMeshComponent::ForceMotionVector()
 	{
 		int32 UseLOD = GetPredictedLODLevel();
 		// Clamp to loaded streaming data if available
-		if (SkeletalMesh->IsStreamable() && MeshObject)
+		if (GetSkeletalMesh()->IsStreamable() && MeshObject)
 		{
 			UseLOD = FMath::Max<int32>(UseLOD, MeshObject->GetSkeletalMeshRenderData().PendingFirstLODIdx);
 		}
@@ -844,12 +844,12 @@ bool USkinnedMeshComponent::CanEditChange(const FProperty* InProperty) const
 
 void USkinnedMeshComponent::InitLODInfos()
 {
-	if (SkeletalMesh != NULL)
+	if (GetSkeletalMesh() != NULL)
 	{
-		if (SkeletalMesh->GetLODNum() != LODInfo.Num())
+		if (GetSkeletalMesh()->GetLODNum() != LODInfo.Num())
 		{
-			LODInfo.Empty(SkeletalMesh->GetLODNum());
-			for (int32 Idx=0; Idx < SkeletalMesh->GetLODNum(); Idx++)
+			LODInfo.Empty(GetSkeletalMesh()->GetLODNum());
+			for (int32 Idx=0; Idx < GetSkeletalMesh()->GetLODNum(); Idx++)
 			{
 				new(LODInfo) FSkelMeshComponentLODInfo();
 			}
@@ -896,7 +896,7 @@ void USkinnedMeshComponent::TickUpdateRate(float DeltaTime, bool bNeedsValidRoot
 				DrawDebugBox(GetWorld(), Bounds.Origin, Bounds.BoxExtent, FQuat::Identity, DrawColor, false);
 
 				FString DebugString = FString::Printf(TEXT("%s UpdateRate(%d) EvaluationRate(%d) ShouldInterpolateSkippedFrames(%d) ShouldSkipUpdate(%d) Interp Alpha (%f) AdditionalTime(%f)"),
-					*GetNameSafe(SkeletalMesh), AnimUpdateRateParams->UpdateRate, AnimUpdateRateParams->EvaluationRate, 
+					*GetNameSafe(GetSkeletalMesh()), AnimUpdateRateParams->UpdateRate, AnimUpdateRateParams->EvaluationRate,
 					AnimUpdateRateParams->ShouldInterpolateSkippedFrames(), AnimUpdateRateParams->ShouldSkipUpdate(), AnimUpdateRateParams->GetInterpolationAlpha(), AnimUpdateRateParams->AdditionalTime);
 
 				GEngine->AddOnScreenDebugMessage(INDEX_NONE, 0.f, FColor::Red, DebugString, false);
@@ -966,7 +966,7 @@ void USkinnedMeshComponent::TickComponent(float DeltaTime, enum ELevelTick TickT
 
 UObject const* USkinnedMeshComponent::AdditionalStatObject() const
 {
-	return SkeletalMesh;
+	return GetSkeletalMesh();
 }
 
 
@@ -978,9 +978,9 @@ void USkinnedMeshComponent::UpdateSlaveComponent()
 // this has to be skeletalmesh material. You can't have more than what SkeletalMesh materials have
 int32 USkinnedMeshComponent::GetNumMaterials() const
 {
-	if (SkeletalMesh)
+	if (GetSkeletalMesh())
 	{
-		return SkeletalMesh->GetMaterials().Num();
+		return GetSkeletalMesh()->GetMaterials().Num();
 	}
 
 	return 0;
@@ -992,9 +992,9 @@ UMaterialInterface* USkinnedMeshComponent::GetMaterial(int32 MaterialIndex) cons
 	{
 		return OverrideMaterials[MaterialIndex];
 	}
-	else if (SkeletalMesh && SkeletalMesh->GetMaterials().IsValidIndex(MaterialIndex) && SkeletalMesh->GetMaterials()[MaterialIndex].MaterialInterface)
+	else if (GetSkeletalMesh() && GetSkeletalMesh()->GetMaterials().IsValidIndex(MaterialIndex) && GetSkeletalMesh()->GetMaterials()[MaterialIndex].MaterialInterface)
 	{
-		return SkeletalMesh->GetMaterials()[MaterialIndex].MaterialInterface;
+		return GetSkeletalMesh()->GetMaterials()[MaterialIndex].MaterialInterface;
 	}
 
 	return nullptr;
@@ -1002,9 +1002,9 @@ UMaterialInterface* USkinnedMeshComponent::GetMaterial(int32 MaterialIndex) cons
 
 int32 USkinnedMeshComponent::GetMaterialIndex(FName MaterialSlotName) const
 {
-	if (SkeletalMesh != nullptr)
+	if (GetSkeletalMesh() != nullptr)
 	{
-		const TArray<FSkeletalMaterial>& SkeletalMeshMaterials = SkeletalMesh->GetMaterials();
+		const TArray<FSkeletalMaterial>& SkeletalMeshMaterials = GetSkeletalMesh()->GetMaterials();
 		for (int32 MaterialIndex = 0; MaterialIndex < SkeletalMeshMaterials.Num(); ++MaterialIndex)
 		{
 			const FSkeletalMaterial &SkeletalMaterial = SkeletalMeshMaterials[MaterialIndex];
@@ -1020,9 +1020,9 @@ int32 USkinnedMeshComponent::GetMaterialIndex(FName MaterialSlotName) const
 TArray<FName> USkinnedMeshComponent::GetMaterialSlotNames() const
 {
 	TArray<FName> MaterialNames;
-	if (SkeletalMesh != nullptr)
+	if (GetSkeletalMesh() != nullptr)
 	{
-		const TArray<FSkeletalMaterial>& SkeletalMeshMaterials = SkeletalMesh->GetMaterials();
+		const TArray<FSkeletalMaterial>& SkeletalMeshMaterials = GetSkeletalMesh()->GetMaterials();
 		for (int32 MaterialIndex = 0; MaterialIndex < SkeletalMeshMaterials.Num(); ++MaterialIndex)
 		{
 			const FSkeletalMaterial &SkeletalMaterial = SkeletalMeshMaterials[MaterialIndex];
@@ -1052,7 +1052,7 @@ PRAGMA_ENABLE_DEPRECATION_WARNINGS
 void USkinnedMeshComponent::SetCPUSkinningEnabled(bool bEnable, bool bRecreateRenderStateImmediately)
 {
 	checkSlow(IsInGameThread());
-	check(SkeletalMesh && SkeletalMesh->GetResourceForRendering());
+	check(GetSkeletalMesh() && GetSkeletalMesh()->GetResourceForRendering());
 
 	if (GetCPUSkinningEnabled() == bEnable)
 	{
@@ -1066,16 +1066,16 @@ void USkinnedMeshComponent::SetCPUSkinningEnabled(bool bEnable, bool bRecreateRe
 		IRenderAssetStreamingManager& Manager = IStreamingManager::Get().GetRenderAssetStreamingManager();
 		Manager.BlockTillAllRequestsFinished();
 
-		const bool bOriginalForcedFullyLoad = SkeletalMesh->bForceMiplevelsToBeResident;
-		SkeletalMesh->bForceMiplevelsToBeResident = true;
-		Manager.UpdateIndividualRenderAsset(SkeletalMesh);
+		const bool bOriginalForcedFullyLoad = GetSkeletalMesh()->bForceMiplevelsToBeResident;
+		GetSkeletalMesh()->bForceMiplevelsToBeResident = true;
+		Manager.UpdateIndividualRenderAsset(GetSkeletalMesh());
 
-		SkeletalMesh->WaitForPendingInitOrStreaming();
+		GetSkeletalMesh()->WaitForPendingInitOrStreaming();
 
-		check(SkeletalMesh->GetResourceForRendering()->CurrentFirstLODIdx <= SkeletalMesh->GetDefaultMinLod());
+		check(GetSkeletalMesh()->GetResourceForRendering()->CurrentFirstLODIdx <= GetSkeletalMesh()->GetDefaultMinLod());
 
-		SkeletalMesh->UnlinkStreaming();
-		SkeletalMesh->bForceMiplevelsToBeResident = bOriginalForcedFullyLoad;
+		GetSkeletalMesh()->UnlinkStreaming();
+		GetSkeletalMesh()->bForceMiplevelsToBeResident = bOriginalForcedFullyLoad;
 	}
 
 PRAGMA_DISABLE_DEPRECATION_WARNINGS
@@ -1098,10 +1098,10 @@ PRAGMA_ENABLE_DEPRECATION_WARNINGS
 
 bool USkinnedMeshComponent::GetMaterialStreamingData(int32 MaterialIndex, FPrimitiveMaterialInfo& MaterialData) const
 {
-	if (SkeletalMesh && SkeletalMesh->IsMaterialUsed(MaterialIndex))
+	if (GetSkeletalMesh() && GetSkeletalMesh()->IsMaterialUsed(MaterialIndex))
 	{
 		MaterialData.Material = GetMaterial(MaterialIndex);
-		MaterialData.UVChannelData = SkeletalMesh->GetUVChannelData(MaterialIndex);
+		MaterialData.UVChannelData = GetSkeletalMesh()->GetUVChannelData(MaterialIndex);
 		MaterialData.PackedRelativeBox = PackedRelativeBox_Identity;
 	}
 	return MaterialData.IsValid();
@@ -1111,11 +1111,11 @@ void USkinnedMeshComponent::GetStreamingRenderAssetInfo(FStreamingTextureLevelCo
 {
 	GetStreamingTextureInfoInner(LevelContext, nullptr, StreamingDistanceMultiplier, OutStreamingRenderAssets);
 
-	if (SkeletalMesh && SkeletalMesh->IsStreamable())
+	if (GetSkeletalMesh() && GetSkeletalMesh()->IsStreamable())
 	{
 		const int32 LocalForcedLodModel = GetForcedLOD();
-		const float TexelFactor = LocalForcedLodModel > 0 ? -(SkeletalMesh->GetLODNum() - LocalForcedLodModel + 1) : Bounds.SphereRadius * 2.f;
-		new (OutStreamingRenderAssets) FStreamingRenderAssetPrimitiveInfo(SkeletalMesh, Bounds, TexelFactor, PackedRelativeBox_Identity, false, false);
+		const float TexelFactor = LocalForcedLodModel > 0 ? -(GetSkeletalMesh()->GetLODNum() - LocalForcedLodModel + 1) : Bounds.SphereRadius * 2.f;
+		new (OutStreamingRenderAssets) FStreamingRenderAssetPrimitiveInfo(GetSkeletalMesh(), Bounds, TexelFactor, PackedRelativeBox_Identity, false, false);
 	}
 }
 
@@ -1139,9 +1139,9 @@ void USkinnedMeshComponent::RebuildVisibilityArray()
 
 		// The following code relies on a complete hierarchy sorted from parent to children
 		TArray<uint8>& EditableBoneVisibilityStates = GetEditableBoneVisibilityStates();
-		if (EditableBoneVisibilityStates.Num() != SkeletalMesh->GetRefSkeleton().GetNum())
+		if (EditableBoneVisibilityStates.Num() != GetSkeletalMesh()->GetRefSkeleton().GetNum())
 		{
-			UE_LOG(LogSkinnedMeshComp, Warning, TEXT("RebuildVisibilityArray() failed because EditableBoneVisibilityStates size: %d not equal to RefSkeleton bone count: %d."), EditableBoneVisibilityStates.Num(), SkeletalMesh->GetRefSkeleton().GetNum());
+			UE_LOG(LogSkinnedMeshComp, Warning, TEXT("RebuildVisibilityArray() failed because EditableBoneVisibilityStates size: %d not equal to RefSkeleton bone count: %d."), EditableBoneVisibilityStates.Num(), GetSkeletalMesh()->GetRefSkeleton().GetNum());
 			return;
 		}
 
@@ -1153,7 +1153,7 @@ void USkinnedMeshComponent::RebuildVisibilityArray()
 			if (VisState != BVS_ExplicitlyHidden)
 			{
 				// Check direct parent (only need to do one deep, since we have already processed the parent and written to BoneVisibilityStates previously)
-				const int32 ParentIndex = SkeletalMesh->GetRefSkeleton().GetParentIndex(BoneId);
+				const int32 ParentIndex = GetSkeletalMesh()->GetRefSkeleton().GetParentIndex(BoneId);
 				if ((ParentIndex == -1) || (EditableBoneVisibilityStates[ParentIndex] == BVS_Visible))
 				{
 					EditableBoneVisibilityStates[BoneId] = BVS_Visible;
@@ -1202,9 +1202,9 @@ class UPhysicsAsset* USkinnedMeshComponent::GetPhysicsAsset() const
 		return PhysicsAssetOverride;
 	}
 
-	if (SkeletalMesh && SkeletalMesh->GetPhysicsAsset())
+	if (GetSkeletalMesh() && GetSkeletalMesh()->GetPhysicsAsset())
 	{
-		return SkeletalMesh->GetPhysicsAsset();
+		return GetSkeletalMesh()->GetPhysicsAsset();
 	}
 
 	return nullptr;
@@ -1223,8 +1223,8 @@ FBoxSphereBounds USkinnedMeshComponent::CalcMeshBound(const FVector3f& RootOffse
 	UPhysicsAsset * const MasterPhysicsAsset = (MasterPoseComponentInst != nullptr)? MasterPoseComponentInst->GetPhysicsAsset() : nullptr;
 
 	// Can only use the PhysicsAsset to calculate the bounding box if we are not non-uniformly scaling the mesh.
-	const USkeletalMesh* SkeletalMeshConst = SkeletalMesh;
-	const bool bCanUsePhysicsAsset = DrawScale.IsUniform() && (SkeletalMesh != NULL)
+	const USkeletalMesh* SkeletalMeshConst = GetSkeletalMesh();
+	const bool bCanUsePhysicsAsset = DrawScale.IsUniform() && (GetSkeletalMesh() != NULL)
 		// either space base exists or child component
 		&& ( (GetNumComponentSpaceTransforms() == SkeletalMeshConst->GetRefSkeleton().GetNum()) || (MasterPhysicsAsset) );
 
@@ -1235,15 +1235,15 @@ FBoxSphereBounds USkinnedMeshComponent::CalcMeshBound(const FVector3f& RootOffse
 	const bool bMasterHasPhysBodies = MasterPhysicsAsset && MasterPhysicsAsset->SkeletalBodySetups.Num();
 
 	// if not visible, or we were told to use fixed bounds, use skelmesh bounds
-	if ( (!bIsVisible || bComponentUseFixedSkelBounds) && SkeletalMesh ) 
+	if ( (!bIsVisible || bComponentUseFixedSkelBounds) && GetSkeletalMesh())
 	{
 		FBoxSphereBounds RootAdjustedBounds = SkeletalMeshConst->GetBounds();
 		RootAdjustedBounds.Origin += FVector(RootOffset); // Adjust bounds by root bone translation
 		NewBounds = RootAdjustedBounds.TransformBy(LocalToWorld);
 	}
-	else if(MasterPoseComponentInst && MasterPoseComponentInst->SkeletalMesh && MasterPoseComponentInst->bComponentUseFixedSkelBounds)
+	else if(MasterPoseComponentInst && MasterPoseComponentInst->GetSkeletalMesh() && MasterPoseComponentInst->bComponentUseFixedSkelBounds)
 	{
-		FBoxSphereBounds RootAdjustedBounds = MasterPoseComponentInst->SkeletalMesh->GetBounds();
+		FBoxSphereBounds RootAdjustedBounds = MasterPoseComponentInst->GetSkeletalMesh()->GetBounds();
 		RootAdjustedBounds.Origin += FVector(RootOffset); // Adjust bounds by root bone translation
 		NewBounds = RootAdjustedBounds.TransformBy(LocalToWorld);
 	}
@@ -1254,7 +1254,7 @@ FBoxSphereBounds USkinnedMeshComponent::CalcMeshBound(const FVector3f& RootOffse
 	}
 #if WITH_EDITOR
 	// For AnimSet Viewer, use 'bounds preview' physics asset if present.
-	else if(SkeletalMesh && bHasPhysBodies && bCanUsePhysicsAsset && PhysicsAsset->CanCalculateValidAABB(this, LocalToWorld))
+	else if(GetSkeletalMesh() && bHasPhysBodies && bCanUsePhysicsAsset && PhysicsAsset->CanCalculateValidAABB(this, LocalToWorld))
 	{
 		NewBounds = FBoxSphereBounds(PhysicsAsset->CalcAABB(this, LocalToWorld));
 	}
@@ -1270,9 +1270,9 @@ FBoxSphereBounds USkinnedMeshComponent::CalcMeshBound(const FVector3f& RootOffse
 		NewBounds = FBoxSphereBounds(MasterPhysicsAsset->CalcAABB(this, LocalToWorld));
 	}
 	// Fallback is to use the one from the skeletal mesh. Usually pretty bad in terms of Accuracy of where the SkelMesh Bounds are located (i.e. usually bigger than it needs to be)
-	else if( SkeletalMesh )
+	else if(GetSkeletalMesh())
 	{
-		FBoxSphereBounds RootAdjustedBounds = SkeletalMesh->GetBounds();
+		FBoxSphereBounds RootAdjustedBounds = GetSkeletalMesh()->GetBounds();
 
 		// Adjust bounds by root bone translation
 		RootAdjustedBounds.Origin += FVector(RootOffset);
@@ -1295,15 +1295,15 @@ void USkinnedMeshComponent::GetPreSkinnedLocalBounds(FBoxSphereBounds& OutBounds
 {
 	const USkinnedMeshComponent* const MasterPoseComponentInst = MasterPoseComponent.Get();
 
-	if (SkeletalMesh)
+	if (GetSkeletalMesh())
 	{
 		// Get the Pre-skinned bounds from the skeletal mesh. Note that these bounds are the "ExtendedBounds", so they can be tweaked on the SkeletalMesh   
-		OutBounds = SkeletalMesh->GetBounds();
+		OutBounds = GetSkeletalMesh()->GetBounds();
 	}
-	else if(MasterPoseComponentInst && MasterPoseComponentInst->SkeletalMesh)
+	else if(MasterPoseComponentInst && MasterPoseComponentInst->GetSkeletalMesh())
 	{
 		// Get the bounds from the master pose if ther is no skeletal mesh
-		OutBounds = MasterPoseComponentInst->SkeletalMesh->GetBounds();
+		OutBounds = MasterPoseComponentInst->GetSkeletalMesh()->GetBounds();
 	}
 	else
 	{
@@ -1337,13 +1337,13 @@ FMatrix USkinnedMeshComponent::GetBoneMatrix(int32 BoneIdx) const
 			}
 			else
 			{
-				UE_LOG(LogSkinnedMeshComp, Verbose, TEXT("GetBoneMatrix : ParentBoneIndex(%d:%s) out of range of MasterPoseComponent->SpaceBases for %s(%s)"), BoneIdx, *GetNameSafe(MasterPoseComponentInst->SkeletalMesh), *GetNameSafe(SkeletalMesh), *GetPathName());
+				UE_LOG(LogSkinnedMeshComp, Verbose, TEXT("GetBoneMatrix : ParentBoneIndex(%d:%s) out of range of MasterPoseComponent->SpaceBases for %s(%s)"), BoneIdx, *GetNameSafe(MasterPoseComponentInst->GetSkeletalMesh()), *GetNameSafe(GetSkeletalMesh()), *GetPathName());
 				return FMatrix::Identity;
 			}
 		}
 		else
 		{
-			UE_LOG(LogSkinnedMeshComp, Warning, TEXT("GetBoneMatrix : BoneIndex(%d) out of range of MasterBoneMap for %s (%s)"), BoneIdx, *this->GetFName().ToString(), SkeletalMesh ? *SkeletalMesh->GetFName().ToString() : TEXT("NULL") );
+			UE_LOG(LogSkinnedMeshComp, Warning, TEXT("GetBoneMatrix : BoneIndex(%d) out of range of MasterBoneMap for %s (%s)"), BoneIdx, *this->GetFName().ToString(), GetSkeletalMesh() ? *GetSkeletalMesh()->GetFName().ToString() : TEXT("NULL") );
 			return FMatrix::Identity;
 		}
 	}
@@ -1356,7 +1356,7 @@ FMatrix USkinnedMeshComponent::GetBoneMatrix(int32 BoneIdx) const
 		}
 		else
 		{
-			UE_LOG(LogSkinnedMeshComp, Warning, TEXT("GetBoneMatrix : BoneIndex(%d) out of range of SpaceBases for %s (%s)"), BoneIdx, *GetPathName(), SkeletalMesh ? *SkeletalMesh->GetFullName() : TEXT("NULL") );
+			UE_LOG(LogSkinnedMeshComp, Warning, TEXT("GetBoneMatrix : BoneIndex(%d) out of range of SpaceBases for %s (%s)"), BoneIdx, *GetPathName(), GetSkeletalMesh() ? *GetSkeletalMesh()->GetFullName() : TEXT("NULL") );
 			return FMatrix::Identity;
 		}
 	}
@@ -1438,7 +1438,7 @@ FTransform USkinnedMeshComponent::GetBoneTransform(int32 BoneIdx, const FTransfo
 		}
 		else
 		{
-			UE_LOG(LogSkinnedMeshComp, Verbose, TEXT("GetBoneTransform : BoneIndex(%d) out of range of SpaceBases for %s (%s)"), BoneIdx, *GetPathName(), SkeletalMesh ? *SkeletalMesh->GetFullName() : TEXT("NULL") );
+			UE_LOG(LogSkinnedMeshComp, Verbose, TEXT("GetBoneTransform : BoneIndex(%d) out of range of SpaceBases for %s (%s)"), BoneIdx, *GetPathName(), GetSkeletalMesh() ? *GetSkeletalMesh()->GetFullName() : TEXT("NULL") );
 			return FTransform::Identity;
 		}
 	}
@@ -1446,15 +1446,15 @@ FTransform USkinnedMeshComponent::GetBoneTransform(int32 BoneIdx, const FTransfo
 
 int32 USkinnedMeshComponent::GetNumBones()const
 {
-	return SkeletalMesh ? SkeletalMesh->GetRefSkeleton().GetNum() : 0;
+	return GetSkeletalMesh() ? GetSkeletalMesh()->GetRefSkeleton().GetNum() : 0;
 }
 
 int32 USkinnedMeshComponent::GetBoneIndex( FName BoneName) const
 {
 	int32 BoneIndex = INDEX_NONE;
-	if ( BoneName != NAME_None && SkeletalMesh )
+	if ( BoneName != NAME_None && GetSkeletalMesh())
 	{
-		BoneIndex = SkeletalMesh->GetRefSkeleton().FindBoneIndex( BoneName );
+		BoneIndex = GetSkeletalMesh()->GetRefSkeleton().FindBoneIndex( BoneName );
 	}
 
 	return BoneIndex;
@@ -1463,7 +1463,7 @@ int32 USkinnedMeshComponent::GetBoneIndex( FName BoneName) const
 
 FName USkinnedMeshComponent::GetBoneName(int32 BoneIndex) const
 {
-	return (SkeletalMesh != NULL && SkeletalMesh->GetRefSkeleton().IsValidIndex(BoneIndex)) ? SkeletalMesh->GetRefSkeleton().GetBoneName(BoneIndex) : NAME_None;
+	return (GetSkeletalMesh() != NULL && GetSkeletalMesh()->GetRefSkeleton().IsValidIndex(BoneIndex)) ? GetSkeletalMesh()->GetRefSkeleton().GetBoneName(BoneIndex) : NAME_None;
 }
 
 
@@ -1474,16 +1474,16 @@ FName USkinnedMeshComponent::GetParentBone( FName BoneName ) const
 	int32 BoneIndex = GetBoneIndex(BoneName);
 	if ((BoneIndex != INDEX_NONE) && (BoneIndex > 0)) // This checks that this bone is not the root (ie no parent), and that BoneIndex != INDEX_NONE (ie bone name was found)
 	{
-		Result = SkeletalMesh->GetRefSkeleton().GetBoneName(SkeletalMesh->GetRefSkeleton().GetParentIndex(BoneIndex));
+		Result = GetSkeletalMesh()->GetRefSkeleton().GetBoneName(GetSkeletalMesh()->GetRefSkeleton().GetParentIndex(BoneIndex));
 	}
 	return Result;
 }
 
 FTransform USkinnedMeshComponent::GetDeltaTransformFromRefPose(FName BoneName, FName BaseName/* = NAME_None*/) const
 {
-	if (SkeletalMesh)
+	if (GetSkeletalMesh())
 	{
-		const FReferenceSkeleton& RefSkeleton = SkeletalMesh->GetRefSkeleton();
+		const FReferenceSkeleton& RefSkeleton = GetSkeletalMesh()->GetRefSkeleton();
 		const int32 BoneIndex = GetBoneIndex(BoneName);
 		if (BoneIndex != INDEX_NONE)
 		{
@@ -1511,7 +1511,7 @@ FTransform USkinnedMeshComponent::GetDeltaTransformFromRefPose(FName BoneName, F
 
 bool USkinnedMeshComponent::GetTwistAndSwingAngleOfDeltaRotationFromRefPose(FName BoneName, float& OutTwistAngle, float& OutSwingAngle) const
 {
-	const FReferenceSkeleton& RefSkeleton = SkeletalMesh->GetRefSkeleton();
+	const FReferenceSkeleton& RefSkeleton = GetSkeletalMesh()->GetRefSkeleton();
 	const int32 BoneIndex = GetBoneIndex(BoneName);
 	const TArray<FTransform>& Transforms = GetComponentSpaceTransforms();
 
@@ -1564,12 +1564,12 @@ bool USkinnedMeshComponent::IsSkinCacheAllowed(int32 LodIdx) const
 		return false;
 	}
 
-	if (!SkeletalMesh)
+	if (!GetSkeletalMesh())
 	{
 		return GEnableGPUSkinCache && bGlobalDefault;
 	}
 
-	FSkeletalMeshLODInfo* LodInfo = SkeletalMesh->GetLODInfo(LodIdx);
+	FSkeletalMeshLODInfo* LodInfo = GetSkeletalMesh()->GetLODInfo(LodIdx);
 	if (!LodInfo)
 	{
 		return GEnableGPUSkinCache && bGlobalDefault;
@@ -1593,7 +1593,7 @@ bool USkinnedMeshComponent::IsSkinCacheAllowed(int32 LodIdx) const
 
 void USkinnedMeshComponent::GetBoneNames(TArray<FName>& BoneNames)
 {
-	if (SkeletalMesh == NULL)
+	if (GetSkeletalMesh() == NULL)
 	{
 		// no mesh, so no bones
 		BoneNames.Empty();
@@ -1601,11 +1601,11 @@ void USkinnedMeshComponent::GetBoneNames(TArray<FName>& BoneNames)
 	else
 	{
 		// pre-size the array to avoid unnecessary reallocation
-		BoneNames.Empty(SkeletalMesh->GetRefSkeleton().GetNum());
-		BoneNames.AddUninitialized(SkeletalMesh->GetRefSkeleton().GetNum());
-		for (int32 i = 0; i < SkeletalMesh->GetRefSkeleton().GetNum(); i++)
+		BoneNames.Empty(GetSkeletalMesh()->GetRefSkeleton().GetNum());
+		BoneNames.AddUninitialized(GetSkeletalMesh()->GetRefSkeleton().GetNum());
+		for (int32 i = 0; i < GetSkeletalMesh()->GetRefSkeleton().GetNum(); i++)
 		{
-			BoneNames[i] = SkeletalMesh->GetRefSkeleton().GetBoneName(i);
+			BoneNames[i] = GetSkeletalMesh()->GetRefSkeleton().GetBoneName(i);
 		}
 	}
 }
@@ -1615,23 +1615,23 @@ bool USkinnedMeshComponent::BoneIsChildOf(FName BoneName, FName ParentBoneName)
 {
 	bool bResult = false;
 
-	if( SkeletalMesh )
+	if(GetSkeletalMesh())
 	{
-		const int32 BoneIndex = SkeletalMesh->GetRefSkeleton().FindBoneIndex(BoneName);
+		const int32 BoneIndex = GetSkeletalMesh()->GetRefSkeleton().FindBoneIndex(BoneName);
 		if(BoneIndex == INDEX_NONE)
 		{
-			UE_LOG(LogSkinnedMeshComp, Log, TEXT("execBoneIsChildOf: BoneName '%s' not found in SkeletalMesh '%s'"), *BoneName.ToString(), *SkeletalMesh->GetName());
+			UE_LOG(LogSkinnedMeshComp, Log, TEXT("execBoneIsChildOf: BoneName '%s' not found in SkeletalMesh '%s'"), *BoneName.ToString(), *GetSkeletalMesh()->GetName());
 			return bResult;
 		}
 
-		const int32 ParentBoneIndex = SkeletalMesh->GetRefSkeleton().FindBoneIndex(ParentBoneName);
+		const int32 ParentBoneIndex = GetSkeletalMesh()->GetRefSkeleton().FindBoneIndex(ParentBoneName);
 		if(ParentBoneIndex == INDEX_NONE)
 		{
-			UE_LOG(LogSkinnedMeshComp, Log, TEXT("execBoneIsChildOf: ParentBoneName '%s' not found in SkeletalMesh '%s'"), *ParentBoneName.ToString(), *SkeletalMesh->GetName());
+			UE_LOG(LogSkinnedMeshComp, Log, TEXT("execBoneIsChildOf: ParentBoneName '%s' not found in SkeletalMesh '%s'"), *ParentBoneName.ToString(), *GetSkeletalMesh()->GetName());
 			return bResult;
 		}
 
-		bResult = SkeletalMesh->GetRefSkeleton().BoneIsChildOf(BoneIndex, ParentBoneIndex);
+		bResult = GetSkeletalMesh()->GetRefSkeleton().BoneIsChildOf(BoneIndex, ParentBoneIndex);
 	}
 
 	return bResult;
@@ -1640,9 +1640,9 @@ bool USkinnedMeshComponent::BoneIsChildOf(FName BoneName, FName ParentBoneName)
 
 FVector USkinnedMeshComponent::GetRefPosePosition(int32 BoneIndex) const
 {
-	if(SkeletalMesh && (BoneIndex >= 0) && (BoneIndex < SkeletalMesh->GetRefSkeleton().GetNum()))
+	if(GetSkeletalMesh() && (BoneIndex >= 0) && (BoneIndex < GetSkeletalMesh()->GetRefSkeleton().GetNum()))
 	{
-		return SkeletalMesh->GetRefSkeleton().GetRefBonePose()[BoneIndex].GetTranslation();
+		return GetSkeletalMesh()->GetRefSkeleton().GetRefBonePose()[BoneIndex].GetTranslation();
 	}
 	else
 	{
@@ -1652,9 +1652,9 @@ FVector USkinnedMeshComponent::GetRefPosePosition(int32 BoneIndex) const
 
 FTransform USkinnedMeshComponent::GetRefPoseTransform(int32 BoneIndex) const
 {
-	if(SkeletalMesh && (BoneIndex >= 0) && (BoneIndex < SkeletalMesh->GetRefSkeleton().GetNum()))
+	if(GetSkeletalMesh() && (BoneIndex >= 0) && (BoneIndex < GetSkeletalMesh()->GetRefSkeleton().GetNum()))
 	{
-		return SkeletalMesh->GetRefSkeleton().GetRefBonePose()[BoneIndex];
+		return GetSkeletalMesh()->GetRefSkeleton().GetRefBonePose()[BoneIndex];
 	}
 	else
 	{
@@ -1667,7 +1667,7 @@ void USkinnedMeshComponent::SetSkeletalMesh(USkeletalMesh* InSkelMesh, bool bRei
 	// NOTE: InSkelMesh may be NULL (useful in the editor for removing the skeletal mesh associated with
 	//   this component on-the-fly)
 
-	if (InSkelMesh == SkeletalMesh)
+	if (InSkelMesh == GetSkeletalMesh())
 	{
 		// do nothing if the input mesh is the same mesh we're already using.
 		return;
@@ -1677,7 +1677,7 @@ void USkinnedMeshComponent::SetSkeletalMesh(USkeletalMesh* InSkelMesh, bool bRei
 		//Handle destroying and recreating the renderstate
 		FRenderStateRecreator RenderStateRecreator(this);
 
-		SkeletalMesh = InSkelMesh;
+		GetSkeletalMesh() = InSkelMesh;
 		SetPredictedLODLevel(0);
 
 		//SlavePoseComponents is an array of weak obj ptrs, so it can contain null elements
@@ -1728,9 +1728,9 @@ FSkeletalMeshRenderData* USkinnedMeshComponent::GetSkeletalMeshRenderData() cons
 	{
 		return &MeshObject->GetSkeletalMeshRenderData();
 	}
-	else if (SkeletalMesh)
+	else if (GetSkeletalMesh())
 	{
-		return SkeletalMesh->GetResourceForRendering();
+		return GetSkeletalMesh()->GetResourceForRendering();
 	}
 	else
 	{
@@ -1743,25 +1743,25 @@ bool USkinnedMeshComponent::AllocateTransformData()
 	LLM_SCOPE_BYNAME(TEXT("SkeletalMesh/TransformData"));
 
 	// Allocate transforms if not present.
-	if ( SkeletalMesh != NULL && MasterPoseComponent == NULL )
+	if (GetSkeletalMesh() != NULL && MasterPoseComponent == NULL )
 	{
-		if(GetNumComponentSpaceTransforms() != SkeletalMesh->GetRefSkeleton().GetNum() )
+		if(GetNumComponentSpaceTransforms() != GetSkeletalMesh()->GetRefSkeleton().GetNum() )
 		{
 			for (int32 BaseIndex = 0; BaseIndex < 2; ++BaseIndex)
 			{
-				ComponentSpaceTransformsArray[BaseIndex].Empty(SkeletalMesh->GetRefSkeleton().GetNum());
-				ComponentSpaceTransformsArray[BaseIndex].AddUninitialized(SkeletalMesh->GetRefSkeleton().GetNum());
+				ComponentSpaceTransformsArray[BaseIndex].Empty(GetSkeletalMesh()->GetRefSkeleton().GetNum());
+				ComponentSpaceTransformsArray[BaseIndex].AddUninitialized(GetSkeletalMesh()->GetRefSkeleton().GetNum());
 
-				for (int32 I = 0; I < SkeletalMesh->GetRefSkeleton().GetNum(); ++I)
+				for (int32 I = 0; I < GetSkeletalMesh()->GetRefSkeleton().GetNum(); ++I)
 				{
 					ComponentSpaceTransformsArray[BaseIndex][I].SetIdentity();
 				}
 
-				BoneVisibilityStates[BaseIndex].Empty( SkeletalMesh->GetRefSkeleton().GetNum() );
-				if( SkeletalMesh->GetRefSkeleton().GetNum() )
+				BoneVisibilityStates[BaseIndex].Empty( GetSkeletalMesh()->GetRefSkeleton().GetNum() );
+				if( GetSkeletalMesh()->GetRefSkeleton().GetNum() )
 				{
-					BoneVisibilityStates[BaseIndex].AddUninitialized( SkeletalMesh->GetRefSkeleton().GetNum() );
-					for (int32 BoneIndex = 0; BoneIndex < SkeletalMesh->GetRefSkeleton().GetNum(); BoneIndex++)
+					BoneVisibilityStates[BaseIndex].AddUninitialized( GetSkeletalMesh()->GetRefSkeleton().GetNum() );
+					for (int32 BoneIndex = 0; BoneIndex < GetSkeletalMesh()->GetRefSkeleton().GetNum(); BoneIndex++)
 					{
 						BoneVisibilityStates[BaseIndex][BoneIndex] = BVS_Visible;
 					}
@@ -2033,9 +2033,9 @@ void USkinnedMeshComponent::SetSelectedEditorMaterial(int32 InSelectedEditorMate
 
 UMorphTarget* USkinnedMeshComponent::FindMorphTarget( FName MorphTargetName ) const
 {
-	if( SkeletalMesh != NULL )
+	if(GetSkeletalMesh() != NULL )
 	{
-		return SkeletalMesh->FindMorphTarget(MorphTargetName);
+		return GetSkeletalMesh()->FindMorphTarget(MorphTargetName);
 	}
 
 	return NULL;
@@ -2043,7 +2043,7 @@ UMorphTarget* USkinnedMeshComponent::FindMorphTarget( FName MorphTargetName ) co
 
 bool USkinnedMeshComponent::GetMissingMasterBoneRelativeTransform(int32 InBoneIndex, FMissingMasterBoneCacheEntry& OutInfo) const
 {
-	const FReferenceSkeleton& SlaveRefSkeleton = SkeletalMesh->GetRefSkeleton();
+	const FReferenceSkeleton& SlaveRefSkeleton = GetSkeletalMesh()->GetRefSkeleton();
 	check(SlaveRefSkeleton.IsValidIndex(InBoneIndex));
 	const TArray<FTransform>& BoneSpaceRefPoseTransforms = SlaveRefSkeleton.GetRefBonePose();
 
@@ -2079,17 +2079,17 @@ void USkinnedMeshComponent::UpdateMasterBoneMap()
 	MasterBoneMap.Reset();
 	MissingMasterBoneMap.Reset();
 
-	if (SkeletalMesh)
+	if (GetSkeletalMesh())
 	{
 		if (USkinnedMeshComponent* MasterPoseComponentPtr = MasterPoseComponent.Get())
 		{
-			if (USkeletalMesh* MasterMesh = MasterPoseComponentPtr->SkeletalMesh)
+			if (USkeletalMesh* MasterMesh = MasterPoseComponentPtr->GetSkeletalMesh())
 			{
-				const FReferenceSkeleton& SlaveRefSkeleton = SkeletalMesh->GetRefSkeleton();
+				const FReferenceSkeleton& SlaveRefSkeleton = GetSkeletalMesh()->GetRefSkeleton();
 				const FReferenceSkeleton& MasterRefSkeleton = MasterMesh->GetRefSkeleton();
 
 				MasterBoneMap.AddUninitialized(SlaveRefSkeleton.GetNum());
-				if (SkeletalMesh == MasterMesh)
+				if (GetSkeletalMesh() == MasterMesh)
 				{
 					// if the meshes are the same, the indices must match exactly so we don't need to look them up
 					for (int32 BoneIndex = 0; BoneIndex < MasterBoneMap.Num(); BoneIndex++)
@@ -2108,7 +2108,7 @@ void USkinnedMeshComponent::UpdateMasterBoneMap()
 					// Cache bones for any SOCKET bones that are missing in the master.
 					// We assume that sockets will be potentially called more often, so we
 					// leave out missing BONE transforms here to try to balance memory & performance.
-					for(USkeletalMeshSocket* Socket : SkeletalMesh->GetActiveSocketList())
+					for(USkeletalMeshSocket* Socket : GetSkeletalMesh()->GetActiveSocketList())
 					{
 						int32 BoneIndex = SlaveRefSkeleton.FindBoneIndex(Socket->BoneName);
 						int32 MasterBoneIndex = MasterRefSkeleton.FindBoneIndex(Socket->BoneName);
@@ -2202,10 +2202,10 @@ class USkeletalMeshSocket const* USkinnedMeshComponent::GetSocketInfoByName(FNam
 
 	USkeletalMeshSocket const* Socket = nullptr;
 
-	if( SkeletalMesh )
+	if(GetSkeletalMesh())
 	{
 		int32 SocketIndex;
-		Socket = SkeletalMesh->FindSocketInfo(OverrideSocketName, OutTransform, OutBoneIndex, SocketIndex);
+		Socket = GetSkeletalMesh()->FindSocketInfo(OverrideSocketName, OutTransform, OutBoneIndex, SocketIndex);
 	}
 	else
 	{
@@ -2231,9 +2231,9 @@ class USkeletalMeshSocket const* USkinnedMeshComponent::GetSocketByName(FName In
 
 	USkeletalMeshSocket const* Socket = nullptr;
 
-	if( SkeletalMesh )
+	if(GetSkeletalMesh())
 	{
-		Socket = SkeletalMesh->FindSocket(OverrideSocketName);
+		Socket = GetSkeletalMesh()->FindSocket(OverrideSocketName);
 	}
 	else
 	{
@@ -2289,7 +2289,7 @@ bool USkinnedMeshComponent::DoesSocketExist(FName InSocketName) const
 
 FName USkinnedMeshComponent::GetSocketBoneName(FName InSocketName) const
 {
-	if(!SkeletalMesh)
+	if(!GetSkeletalMesh())
 	{
 		return NAME_None;
 	}
@@ -2298,7 +2298,7 @@ FName USkinnedMeshComponent::GetSocketBoneName(FName InSocketName) const
 	const FName OverrideSocketName = OverrideSocket ? *OverrideSocket : InSocketName;
 
 	// First check for a socket
-	USkeletalMeshSocket const* TmpSocket = SkeletalMesh->FindSocket(OverrideSocketName);
+	USkeletalMeshSocket const* TmpSocket = GetSkeletalMesh()->FindSocket(OverrideSocketName);
 	if( TmpSocket )
 	{
 		return TmpSocket->BoneName;
@@ -2433,19 +2433,19 @@ FVector USkinnedMeshComponent::GetBoneAxis( FName BoneName, EAxis::Type Axis ) c
 
 bool USkinnedMeshComponent::HasAnySockets() const
 {
-	return (SkeletalMesh != NULL) && (
+	return (GetSkeletalMesh() != NULL) && (
 #if WITH_EDITOR
-		(SkeletalMesh->GetActiveSocketList().Num() > 0) ||
+		(GetSkeletalMesh()->GetActiveSocketList().Num() > 0) ||
 #endif
-		(SkeletalMesh->GetRefSkeleton().GetNum() > 0));
+		(GetSkeletalMesh()->GetRefSkeleton().GetNum() > 0));
 }
 
 void USkinnedMeshComponent::QuerySupportedSockets(TArray<FComponentSocketDescription>& OutSockets) const
 {
-	if (SkeletalMesh != NULL)
+	if (GetSkeletalMesh() != NULL)
 	{
 		// Grab all the mesh and skeleton sockets
-		const TArray<USkeletalMeshSocket*> AllSockets = SkeletalMesh->GetActiveSocketList();
+		const TArray<USkeletalMeshSocket*> AllSockets = GetSkeletalMesh()->GetActiveSocketList();
 
 		for (int32 SocketIdx = 0; SocketIdx < AllSockets.Num(); ++SocketIdx)
 		{
@@ -2456,9 +2456,9 @@ void USkinnedMeshComponent::QuerySupportedSockets(TArray<FComponentSocketDescrip
 		}
 
 		// Now grab the bones, which can behave exactly like sockets
-		for (int32 BoneIdx = 0; BoneIdx < SkeletalMesh->GetRefSkeleton().GetNum(); ++BoneIdx)
+		for (int32 BoneIdx = 0; BoneIdx < GetSkeletalMesh()->GetRefSkeleton().GetNum(); ++BoneIdx)
 		{
-			const FName BoneName = SkeletalMesh->GetRefSkeleton().GetBoneName(BoneIdx);
+			const FName BoneName = GetSkeletalMesh()->GetRefSkeleton().GetBoneName(BoneIdx);
 			new (OutSockets) FComponentSocketDescription(BoneName, EComponentSocketType::Bone);
 		}
 	}
@@ -2484,9 +2484,9 @@ bool USkinnedMeshComponent::GetMaterialPropertyPath(int32 ElementIndex, UObject*
 		}
 		return true;
 	}
-	if (SkeletalMesh && SkeletalMesh->GetMaterials().IsValidIndex(ElementIndex))
+	if (GetSkeletalMesh() && GetSkeletalMesh()->GetMaterials().IsValidIndex(ElementIndex))
 	{
-		OutOwner = SkeletalMesh;
+		OutOwner = GetSkeletalMesh();
 		OutPropertyPath = FString::Printf(TEXT("%s[%d].%s"), *USkeletalMesh::GetMaterialsMemberName().ToString(), ElementIndex, GET_MEMBER_NAME_STRING_CHECKED(FSkeletalMaterial, MaterialInterface));
 		OutProperty = FSkeletalMaterial::StaticStruct()->FindPropertyByName(GET_MEMBER_NAME_CHECKED(FSkeletalMaterial, MaterialInterface));
 		return true;
@@ -2530,7 +2530,7 @@ void USkinnedMeshComponent::TransformFromBoneSpace(FName BoneName, FVector InPos
 
 FName USkinnedMeshComponent::FindClosestBone(FVector TestLocation, FVector* BoneLocation, float IgnoreScale, bool bRequirePhysicsAsset) const
 {
-	if (SkeletalMesh == NULL)
+	if (GetSkeletalMesh() == NULL)
 	{
 		if (BoneLocation != NULL)
 		{
@@ -2567,7 +2567,7 @@ FName USkinnedMeshComponent::FindClosestBone(FVector TestLocation, FVector* Bone
 			bool bPassPACheck = !bRequirePhysicsAsset;
 			if (bRequirePhysicsAsset)
 			{
-				FName BoneName = SkeletalMesh->GetRefSkeleton().GetBoneName(i);
+				FName BoneName = GetSkeletalMesh()->GetRefSkeleton().GetBoneName(i);
 				bPassPACheck = (PhysAsset->BodySetupIndexMap.Find(BoneName) != nullptr);
 			}
 
@@ -2597,7 +2597,7 @@ FName USkinnedMeshComponent::FindClosestBone(FVector TestLocation, FVector* Bone
 			{
 				*BoneLocation = (CompSpaceTransforms[BestIndex] * GetComponentTransform()).GetLocation();
 			}
-			return SkeletalMesh->GetRefSkeleton().GetBoneName(BestIndex);
+			return GetSkeletalMesh()->GetRefSkeleton().GetBoneName(BestIndex);
 		}
 	}
 }
@@ -2610,7 +2610,7 @@ FName USkinnedMeshComponent::FindClosestBone_K2(FVector TestLocation, FVector& B
 
 void USkinnedMeshComponent::ShowMaterialSection(int32 MaterialID, int32 SectionIndex, bool bShow, int32 LODIndex)
 {
-	if (!SkeletalMesh)
+	if (!GetSkeletalMesh())
 	{
 		// no skeletalmesh, then nothing to do. 
 		return;
@@ -2619,16 +2619,16 @@ void USkinnedMeshComponent::ShowMaterialSection(int32 MaterialID, int32 SectionI
 	InitLODInfos();
 	if (LODInfo.IsValidIndex(LODIndex))
 	{
-		const FSkeletalMeshLODInfo& SkelLODInfo = *SkeletalMesh->GetLODInfo(LODIndex);
+		const FSkeletalMeshLODInfo& SkelLODInfo = *GetSkeletalMesh()->GetLODInfo(LODIndex);
 		FSkelMeshComponentLODInfo& SkelCompLODInfo = LODInfo[LODIndex];
 		TArray<bool>& HiddenMaterials = SkelCompLODInfo.HiddenMaterials;
 	
 		// allocate if not allocated yet
-		if ( HiddenMaterials.Num() != SkeletalMesh->GetMaterials().Num() )
+		if ( HiddenMaterials.Num() != GetSkeletalMesh()->GetMaterials().Num() )
 		{
-			// Using skeletalmesh component because Materials.Num() should be <= SkeletalMesh->GetMaterials().Num()		
-			HiddenMaterials.Empty(SkeletalMesh->GetMaterials().Num());
-			HiddenMaterials.AddZeroed(SkeletalMesh->GetMaterials().Num());
+			// Using skeletalmesh component because Materials.Num() should be <= GetSkeletalMesh()->GetMaterials().Num()		
+			HiddenMaterials.Empty(GetSkeletalMesh()->GetMaterials().Num());
+			HiddenMaterials.AddZeroed(GetSkeletalMesh()->GetMaterials().Num());
 		}
 		// If we have a valid LODInfo LODMaterialMap, route material index through it.
 		int32 UseMaterialIndex = MaterialID;			
@@ -2704,10 +2704,10 @@ bool USkinnedMeshComponent::IsMaterialSectionShown(int32 MaterialID, int32 LODIn
 
 void USkinnedMeshComponent::GetUsedMaterials( TArray<UMaterialInterface*>& OutMaterials, bool bGetDebugMaterials ) const
 {
-	if( SkeletalMesh )
+	if(GetSkeletalMesh())
 	{
 		// The max number of materials used is the max of the materials on the skeletal mesh and the materials on the mesh component
-		const int32 NumMaterials = FMath::Max(SkeletalMesh->GetMaterials().Num(), OverrideMaterials.Num() );
+		const int32 NumMaterials = FMath::Max(GetSkeletalMesh()->GetMaterials().Num(), OverrideMaterials.Num() );
 		for( int32 MatIdx = 0; MatIdx < NumMaterials; ++MatIdx )
 		{
 			// GetMaterial will determine the correct material to use for this index.  
@@ -2736,11 +2736,11 @@ FSkinWeightVertexBuffer* USkinnedMeshComponent::GetSkinWeightBuffer(int32 LODInd
 {
 	FSkinWeightVertexBuffer* WeightBuffer = nullptr;
 
-	if (SkeletalMesh && 
-		SkeletalMesh->GetResourceForRendering() && 
-		SkeletalMesh->GetResourceForRendering()->LODRenderData.IsValidIndex(LODIndex) )
+	if (GetSkeletalMesh() &&
+		GetSkeletalMesh()->GetResourceForRendering() && 
+		GetSkeletalMesh()->GetResourceForRendering()->LODRenderData.IsValidIndex(LODIndex) )
 	{
-		FSkeletalMeshLODRenderData& LODData = SkeletalMesh->GetResourceForRendering()->LODRenderData[LODIndex];
+		FSkeletalMeshLODRenderData& LODData = GetSkeletalMesh()->GetResourceForRendering()->LODRenderData[LODIndex];
 
 		// Grab weight buffer (check for override)
 		if (LODInfo.IsValidIndex(LODIndex) &&
@@ -2794,17 +2794,17 @@ FVector3f USkinnedMeshComponent::GetSkinnedVertexPosition(USkinnedMeshComponent*
 
 void USkinnedMeshComponent::SetRefPoseOverride(const TArray<FTransform>& NewRefPoseTransforms)
 {
-	if (!SkeletalMesh)
+	if (!GetSkeletalMesh())
 	{
 		UE_LOG(LogSkeletalMesh, Warning, TEXT("SetRefPoseOverride (%s) : Not valid without SkeletalMesh assigned."), *GetName());
 		return;
 	}
 
-	const int32 NumRealBones = SkeletalMesh->GetRefSkeleton().GetRawBoneNum();
+	const int32 NumRealBones = GetSkeletalMesh()->GetRefSkeleton().GetRawBoneNum();
 
 	if (NumRealBones != NewRefPoseTransforms.Num())
 	{
-		UE_LOG(LogSkeletalMesh, Warning, TEXT("SetRefPoseOverride (%s) : Expected %d transforms, got %d."), *SkeletalMesh->GetName(), NumRealBones, NewRefPoseTransforms.Num());
+		UE_LOG(LogSkeletalMesh, Warning, TEXT("SetRefPoseOverride (%s) : Expected %d transforms, got %d."), *GetSkeletalMesh()->GetName(), NumRealBones, NewRefPoseTransforms.Num());
 		return;
 	}
 
@@ -2835,7 +2835,7 @@ void USkinnedMeshComponent::SetRefPoseOverride(const TArray<FTransform>& NewRefP
 		// Construct mesh-space skeletal hierarchy.
 		if (BoneIndex > 0)
 		{
-			int32 ParentIndex = SkeletalMesh->GetRefSkeleton().GetRawParentIndex(BoneIndex);
+			int32 ParentIndex = GetSkeletalMesh()->GetRefSkeleton().GetRawParentIndex(BoneIndex);
 			CachedComposedRefPoseMatrices[BoneIndex] = CachedComposedRefPoseMatrices[BoneIndex] * CachedComposedRefPoseMatrices[ParentIndex];
 		}
 
@@ -2847,7 +2847,7 @@ void USkinnedMeshComponent::SetRefPoseOverride(const TArray<FTransform>& NewRefP
 			ZAxis.IsNearlyZero(UE_SMALL_NUMBER))
 		{
 			// this is not allowed, warn them 
-			UE_LOG(LogSkeletalMesh, Warning, TEXT("Reference Pose for asset %s for joint (%s) includes NIL matrix. Zero scale isn't allowed on ref pose. "), *SkeletalMesh->GetPathName(), *SkeletalMesh->GetRefSkeleton().GetBoneName(BoneIndex).ToString());
+			UE_LOG(LogSkeletalMesh, Warning, TEXT("Reference Pose for asset %s for joint (%s) includes NIL matrix. Zero scale isn't allowed on ref pose. "), *GetSkeletalMesh()->GetPathName(), *GetSkeletalMesh()->GetRefSkeleton().GetBoneName(BoneIndex).ToString());
 		}
 
 		// Precompute inverse so we can use from-refpose-skin vertices.
@@ -2867,7 +2867,7 @@ void USkinnedMeshComponent::ClearRefPoseOverride()
 void USkinnedMeshComponent::CacheRefToLocalMatrices(TArray<FMatrix44f>& OutRefToLocal)const
 {
 	const USkinnedMeshComponent* BaseComponent = GetBaseComponent();
-	OutRefToLocal.SetNumUninitialized(SkeletalMesh->GetRefBasesInvMatrix().Num());
+	OutRefToLocal.SetNumUninitialized(GetSkeletalMesh()->GetRefBasesInvMatrix().Num());
 	const TArray<FTransform>& CompSpaceTransforms = BaseComponent->GetComponentSpaceTransforms();
 	if(CompSpaceTransforms.Num())
 	{
@@ -2875,7 +2875,7 @@ void USkinnedMeshComponent::CacheRefToLocalMatrices(TArray<FMatrix44f>& OutRefTo
 
 		for (int32 MatrixIdx = 0; MatrixIdx < OutRefToLocal.Num(); ++MatrixIdx)
 		{
-			OutRefToLocal[MatrixIdx] = SkeletalMesh->GetRefBasesInvMatrix()[MatrixIdx] * FMatrix44f(CompSpaceTransforms[MatrixIdx].ToMatrixWithScale());
+			OutRefToLocal[MatrixIdx] = GetSkeletalMesh()->GetRefBasesInvMatrix()[MatrixIdx] * FMatrix44f(CompSpaceTransforms[MatrixIdx].ToMatrixWithScale());
 		}
 	}
 	else
@@ -2883,7 +2883,7 @@ void USkinnedMeshComponent::CacheRefToLocalMatrices(TArray<FMatrix44f>& OutRefTo
 		//Possible in some cases to request this before the component space transforms are prepared (undo/redo)
 		for (int32 MatrixIdx = 0; MatrixIdx < OutRefToLocal.Num(); ++MatrixIdx)
 		{
-			OutRefToLocal[MatrixIdx] = SkeletalMesh->GetRefBasesInvMatrix()[MatrixIdx];
+			OutRefToLocal[MatrixIdx] = GetSkeletalMesh()->GetRefBasesInvMatrix()[MatrixIdx];
 		}
 	}
 }
@@ -2893,7 +2893,7 @@ void USkinnedMeshComponent::ComputeSkinnedPositions(USkinnedMeshComponent* Compo
 	OutPositions.Empty();
 
 	// Fail if no mesh
-	if (!Component || !Component->SkeletalMesh)
+	if (!Component || !Component->GetSkeletalMesh())
 	{
 		return;
 	}
@@ -2920,7 +2920,7 @@ FColor USkinnedMeshComponent::GetVertexColor(int32 VertexIndex) const
 {
 	// Fail if no mesh or no color vertex buffer.
 	FColor FallbackColor = FColor(255, 255, 255, 255);
-	if (!SkeletalMesh || !MeshObject)
+	if (!GetSkeletalMesh() || !MeshObject)
 	{
 		return FallbackColor;
 	}
@@ -2958,7 +2958,7 @@ FVector2D USkinnedMeshComponent::GetVertexUV(int32 VertexIndex, uint32 UVChannel
 {
 	// Fail if no mesh or no vertex buffer.
 	FVector2D FallbackUV = FVector2D::ZeroVector;
-	if (!SkeletalMesh || !MeshObject)
+	if (!GetSkeletalMesh() || !MeshObject)
 	{
 		return FallbackUV;
 	}
@@ -3101,7 +3101,7 @@ void USkinnedMeshComponent::SetMinLOD(int32 InNewMinLOD)
 
 int32 USkinnedMeshComponent::ComputeMinLOD() const
 {
-	int32 MinLodIndex = bOverrideMinLod ? MinLodModel : SkeletalMesh->GetMinLodIdx();
+	int32 MinLodIndex = bOverrideMinLod ? MinLodModel : GetSkeletalMesh()->GetMinLodIdx();
 	int32 NumLODs = GetNumLODs();
 	// want to make sure MinLOD stays within the valid range
 	MinLodIndex = FMath::Min(MinLodIndex, NumLODs - 1);
@@ -3146,7 +3146,7 @@ void USkinnedMeshComponent::SetCapsuleIndirectShadowMinVisibility(float NewValue
 // @todo: think about consolidating this with UpdateLODStatus_Internal
 int32 USkinnedMeshComponent::GetDesiredSyncLOD() const
 {
-	if (SkeletalMesh && MeshObject)
+	if (GetSkeletalMesh() && MeshObject)
 	{
 #if WITH_EDITOR
 		const int32 LODBias = GetLODBias();
@@ -3195,7 +3195,7 @@ bool USkinnedMeshComponent::UpdateLODStatus_Internal(int32 InMasterPoseComponent
 	const int32 OldPredictedLODLevel = GetPredictedLODLevel();
 	int32 NewPredictedLODLevel = OldPredictedLODLevel;
 
-	if (SkeletalMesh != nullptr)
+	if (GetSkeletalMesh() != nullptr)
 	{
 #if WITH_EDITOR
 		const int32 LODBias = GetLODBias();
@@ -3252,7 +3252,7 @@ bool USkinnedMeshComponent::UpdateLODStatus_Internal(int32 InMasterPoseComponent
 			}
 		}
 
-		if (SkeletalMesh->IsStreamable() && MeshObject)
+		if (GetSkeletalMesh()->IsStreamable() && MeshObject)
 		{
 			NewPredictedLODLevel = FMath::Max<int32>(NewPredictedLODLevel, MeshObject->GetSkeletalMeshRenderData().PendingFirstLODIdx);
 		}
@@ -3261,7 +3261,7 @@ bool USkinnedMeshComponent::UpdateLODStatus_Internal(int32 InMasterPoseComponent
 		if (CVarAnimVisualizeLODs.GetValueOnAnyThread() != 0)
 		{
 			// Reduce to visible animated, non SyncAttachParentLOD to reduce clutter.
-			if (SkeletalMesh && MeshObject && bRecentlyRendered)
+			if (GetSkeletalMesh() && MeshObject && bRecentlyRendered)
 			{
 				const bool bHasValidSyncAttachParent = bSyncAttachParentLOD && GetAttachParent() && GetAttachParent()->IsA(USkinnedMeshComponent::StaticClass());
 				if (!bHasValidSyncAttachParent)
@@ -3438,18 +3438,18 @@ void USkinnedMeshComponent::ReleaseResources()
 
 void USkinnedMeshComponent::RegisterLODStreamingCallback(FLODStreamingCallback&& Callback, int32 LODIdx, float TimeoutSecs, bool bOnStreamIn)
 {
-	if (SkeletalMesh)
+	if (GetSkeletalMesh())
 	{
-		SkeletalMesh->RegisterMipLevelChangeCallback(this, LODIdx, TimeoutSecs, bOnStreamIn, MoveTemp(Callback));
+		GetSkeletalMesh()->RegisterMipLevelChangeCallback(this, LODIdx, TimeoutSecs, bOnStreamIn, MoveTemp(Callback));
 		bMipLevelCallbackRegistered = true;
 	}
 }
 
 void USkinnedMeshComponent::BeginDestroy()
 {
-	if (SkeletalMesh && bMipLevelCallbackRegistered)
+	if (GetSkeletalMesh() && bMipLevelCallbackRegistered)
 	{
-		SkeletalMesh->RemoveMipLevelChangeCallback(this);
+		GetSkeletalMesh()->RemoveMipLevelChangeCallback(this);
 		bMipLevelCallbackRegistered = false;
 	}
 
@@ -3816,7 +3816,7 @@ void USkinnedMeshComponent::SetSkinWeightOverride(int32 LODIndex, const TArray<F
 			Info.OverrideSkinWeights->SetMaxBoneInfluences(NumBoneInfluences);
 			Info.OverrideSkinWeights->SetUse16BitBoneIndex(bUse16BitBoneIndex);
 
-			const FReferenceSkeleton& RefSkel = SkeletalMesh->GetRefSkeleton();
+			const FReferenceSkeleton& RefSkel = GetSkeletalMesh()->GetRefSkeleton();
 			TArray<FSkinWeightInfo> GPUWeights;
 			CreateSkinWeightsArray(SkinWeights, LODData, GPUWeights, RefSkel);
 			*(Info.OverrideSkinWeights) = GPUWeights;
@@ -3925,7 +3925,7 @@ bool USkinnedMeshComponent::SetSkinWeightProfile(FName InProfileName)
 				// Put in a skin weight profile request
 				if (FSkinWeightProfileManager* Manager = FSkinWeightProfileManager::Get(GetWorld()))
 				{
-					Manager->RequestSkinWeightProfile(InProfileName, SkeletalMesh, this, Callback);
+					Manager->RequestSkinWeightProfile(InProfileName, GetSkeletalMesh(), this, Callback);
 					bSkinWeightProfilePending = true;
 				}
 			}
@@ -4065,7 +4065,7 @@ USkinnedAsset* USkinnedMeshComponent::GetSkinnedAsset() const
 
 bool USkinnedMeshComponent::IsCompiling() const
 {
-	return SkeletalMesh && SkeletalMesh->IsCompiling();
+	return GetSkeletalMesh() && GetSkeletalMesh()->IsCompiling();
 }
 
 void USkinnedMeshComponent::BindWorldDelegates()
@@ -4280,7 +4280,7 @@ FVector3f GetTypedSkinnedVertexPosition(
 #endif
 	{
 		const int32 BoneMapIndex = SkinWeightVertexBuffer.GetBoneIndex(BufferVertIndex, InfluenceIndex);
-		if (!ensureMsgf(Section.BoneMap.IsValidIndex(BoneMapIndex), TEXT("%s has attempted to access a BoneMap of size %i with an invalid index of %i in GetTypedSkinnedVertexPosition()"), *SkinnedComp->SkeletalMesh->GetFullName(), Section.BoneMap.Num(), BoneMapIndex))
+		if (!ensureMsgf(Section.BoneMap.IsValidIndex(BoneMapIndex), TEXT("%s has attempted to access a BoneMap of size %i with an invalid index of %i in GetTypedSkinnedVertexPosition()"), *SkinnedComp->GetSkeletalMesh()->GetFullName(), Section.BoneMap.Num(), BoneMapIndex))
 		{
 			continue;
 		}
@@ -4291,7 +4291,7 @@ FVector3f GetTypedSkinnedVertexPosition(
 		if (MasterPoseComponentInst)
 		{
 			const TArray<int32>& MasterBoneMap = SkinnedComp->GetMasterBoneMap();
-			check(MasterBoneMap.Num() == SkinnedComp->SkeletalMesh->GetRefSkeleton().GetNum());
+			check(MasterBoneMap.Num() == SkinnedComp->GetSkeletalMesh()->GetRefSkeleton().GetNum());
 			TransformBoneIndex = MasterBoneMap[MeshBoneIndex];
 		}
 
@@ -4299,7 +4299,7 @@ FVector3f GetTypedSkinnedVertexPosition(
 		{
 			if (bCachedMatrices)
 			{
-				if (ensureMsgf(RefToLocals.IsValidIndex(MeshBoneIndex), TEXT("%s has attempted to access a RefToLocals of size %i with an invalid index of %i in GetTypedSkinnedVertexPosition()"), *SkinnedComp->SkeletalMesh->GetFullName(), RefToLocals.Num(), MeshBoneIndex))
+				if (ensureMsgf(RefToLocals.IsValidIndex(MeshBoneIndex), TEXT("%s has attempted to access a RefToLocals of size %i with an invalid index of %i in GetTypedSkinnedVertexPosition()"), *SkinnedComp->GetSkeletalMesh()->GetFullName(), RefToLocals.Num(), MeshBoneIndex))
 				{
 					const FMatrix44f& RefToLocal = RefToLocals[MeshBoneIndex];
 					SkinnedPos += RefToLocal.TransformPosition(PositionVertexBuffer.VertexPosition(BufferVertIndex)) * Weight;
@@ -4308,7 +4308,7 @@ FVector3f GetTypedSkinnedVertexPosition(
 			else
 			{
 				const FMatrix44f BoneTransformMatrix = (TransformBoneIndex != INDEX_NONE) ? (FMatrix44f)BaseComponent->GetComponentSpaceTransforms()[TransformBoneIndex].ToMatrixWithScale() : FMatrix44f::Identity;
-				const FMatrix44f RefToLocal = SkinnedComp->SkeletalMesh->GetRefBasesInvMatrix()[MeshBoneIndex] * BoneTransformMatrix;
+				const FMatrix44f RefToLocal = SkinnedComp->GetSkeletalMesh()->GetRefBasesInvMatrix()[MeshBoneIndex] * BoneTransformMatrix;
 				SkinnedPos += RefToLocal.TransformPosition(PositionVertexBuffer.VertexPosition(BufferVertIndex)) * Weight;
 			}
 		}
diff --git a/Engine/Source/Runtime/Engine/Private/LocalPlayer.cpp b/Engine/Source/Runtime/Engine/Private/LocalPlayer.cpp
--- a/Engine/Source/Runtime/Engine/Private/LocalPlayer.cpp
+++ b/Engine/Source/Runtime/Engine/Private/LocalPlayer.cpp
@@ -1259,7 +1259,7 @@ bool ULocalPlayer::HandleListSkelMeshesCommand( const TCHAR* Cmd, FOutputDevice&
 	for( TObjectIterator<USkeletalMeshComponent> It; It; ++It )
 	{
 		USkeletalMeshComponent* SkeletalMeshComponent = *It;
-		USkeletalMesh* SkeletalMesh = SkeletalMeshComponent->SkeletalMesh;
+		USkeletalMesh* SkeletalMesh = SkeletalMeshComponent->GetSkeletalMesh();
 
 		if( !SkeletalMeshComponent->IsTemplate() )
 		{
diff --git a/Engine/Source/Runtime/Engine/Private/ObjectCacheContext.cpp b/Engine/Source/Runtime/Engine/Private/ObjectCacheContext.cpp
--- a/Engine/Source/Runtime/Engine/Private/ObjectCacheContext.cpp
+++ b/Engine/Source/Runtime/Engine/Private/ObjectCacheContext.cpp
@@ -574,7 +574,7 @@ TObjectCacheIterator<USkinnedMeshComponent> FObjectCacheContext::GetSkinnedMeshC
 		TempMap.Reserve(8192);
 		for (USkinnedMeshComponent* Component : GetSkinnedMeshComponents())
 		{
-			TempMap.FindOrAdd(ToRawPtr(Component->SkeletalMesh)).Add(Component);
+			TempMap.FindOrAdd(ToRawPtr(Component->GetSkeletalMesh())).Add(Component);
 		}
 		SkeletalMeshToComponents = MoveTemp(TempMap);
 	}
