Replace direct access to SkeletalMesh object from USkinnedMeshComponent with GetSkeletalMesh function

diff --git a/Engine/Source/Runtime/Engine/Private/Particles/ParticleModules_Location.cpp b/Engine/Source/Runtime/Engine/Private/Particles/ParticleModules_Location.cpp
--- a/Engine/Source/Runtime/Engine/Private/Particles/ParticleModules_Location.cpp
+++ b/Engine/Source/Runtime/Engine/Private/Particles/ParticleModules_Location.cpp
@@ -1467,7 +1467,7 @@ void UParticleModuleLocationBoneSocket::RegeneratePreSelectedIndices(FModuleLoca
 {
 	if (SourceIndexMode == EBoneSocketSourceIndexMode::PreSelectedIndices)
 	{
-		int32 MaxIndex = SourceType == BONESOCKETSOURCE_Sockets ? SourceComponent->SkeletalMesh->NumSockets() : SourceComponent->GetNumBones();
+		int32 MaxIndex = SourceType == BONESOCKETSOURCE_Sockets ? SourceComponent->GetSkeletalMesh()->NumSockets() : SourceComponent->GetNumBones();
 		for (int32 i = 0; i < NumPreSelectedIndices; ++i)
 		{
 			//Should we provide sequential selection here? Does that make sense for the pre selected list?
@@ -1708,7 +1708,7 @@ void UParticleModuleLocationBoneSocket::FinalUpdate(FParticleEmitterInstance* Ow
 			BEGIN_UPDATE_LOOP;
 			{
 				FModuleLocationBoneSocketParticlePayload* ParticlePayload = (FModuleLocationBoneSocketParticlePayload*)((uint8*)&Particle + Offset);
-				if (SourceComponent && SourceComponent->SkeletalMesh)
+				if (SourceComponent && SourceComponent->GetSkeletalMesh())
 				{
 					USkeletalMeshSocket* Socket;
 					FVector SocketOffset;
@@ -1916,9 +1916,9 @@ int32 UParticleModuleLocationBoneSocket::GetMaxSourceIndex(FModuleLocationBoneSo
 		{
 			if (SourceType == BONESOCKETSOURCE_Sockets)
 			{
-				if (ensure(SourceComponent->SkeletalMesh))
+				if (ensure(SourceComponent->GetSkeletalMesh()))
 				{
-					return SourceComponent->SkeletalMesh->NumSockets();
+					return SourceComponent->GetSkeletalMesh()->NumSockets();
 				}
 				return 0;
 			}
@@ -1997,7 +1997,7 @@ bool UParticleModuleLocationBoneSocket::GetSocketInfoForSourceIndex(FModuleLocat
 {
 	if (!ensureMsgf(SourceType == BONESOCKETSOURCE_Sockets, TEXT("Invalid source type %d for %s"), SourceType, *GetPathName()) ||
 		!ensureMsgf(SourceComponent, TEXT("Null SkeletalMeshComponent for %s"), *GetPathName()) ||
-		!ensureMsgf(SourceComponent->SkeletalMesh, TEXT("Null SkeletalMesh on Component %s for %s"), *SourceComponent->GetPathName(), *GetPathName()))
+		!ensureMsgf(SourceComponent->GetSkeletalMesh(), TEXT("Null SkeletalMesh on Component %s for %s"), *SourceComponent->GetPathName(), *GetPathName()))
 	{
 		return false;
 	}
@@ -2008,7 +2008,7 @@ bool UParticleModuleLocationBoneSocket::GetSocketInfoForSourceIndex(FModuleLocat
 		{
 			if (ensureMsgf(SourceLocations.IsValidIndex(SourceIndex), TEXT("Invalid index of %d for %s"), SourceIndex, *GetPathName()))
 			{
-				OutSocket = SourceComponent->SkeletalMesh->FindSocket(SourceLocations[SourceIndex].BoneSocketName);
+				OutSocket = SourceComponent->GetSkeletalMesh()->FindSocket(SourceLocations[SourceIndex].BoneSocketName);
 				OutOffset = SourceLocations[SourceIndex].Offset + UniversalOffset;
 			}
 			else
@@ -2022,7 +2022,7 @@ bool UParticleModuleLocationBoneSocket::GetSocketInfoForSourceIndex(FModuleLocat
 			if (ensureMsgf(InstancePayload, TEXT("Invalid instance payload parameter on GetSocketInfoForSourceIndex %d for %s"), SourceIndex, *GetPathName()) &&
 				ensureMsgf(SourceIndex >= 0 && SourceIndex < InstancePayload->PreSelectedBoneSocketIndices.Num(), TEXT("Invalid index of %d for %s"), SourceIndex, *GetPathName()))
 			{
-				OutSocket = SourceComponent->SkeletalMesh->GetSocketByIndex(InstancePayload->PreSelectedBoneSocketIndices[SourceIndex]);
+				OutSocket = SourceComponent->GetSkeletalMesh()->GetSocketByIndex(InstancePayload->PreSelectedBoneSocketIndices[SourceIndex]);
 				OutOffset = UniversalOffset;
 			} 
 			else
@@ -2034,7 +2034,7 @@ bool UParticleModuleLocationBoneSocket::GetSocketInfoForSourceIndex(FModuleLocat
 		break;
 		case EBoneSocketSourceIndexMode::Direct:
 		{
-			OutSocket = SourceComponent->SkeletalMesh->GetSocketByIndex(SourceIndex);
+			OutSocket = SourceComponent->GetSkeletalMesh()->GetSocketByIndex(SourceIndex);
 			OutOffset = UniversalOffset;
 		}
 		break;
@@ -2098,7 +2098,7 @@ bool UParticleModuleLocationBoneSocket::GetParticleLocation(FModuleLocationBoneS
 
 	if (SourceType == BONESOCKETSOURCE_Sockets)
 	{
-		if (InSkelMeshComponent->SkeletalMesh)
+		if (InSkelMeshComponent->GetSkeletalMesh())
 		{
 			USkeletalMeshSocket* Socket;
 			FVector SocketOffset;
@@ -2583,7 +2583,7 @@ void UParticleModuleLocationSkelVertSurface::UpdateBoneIndicesList(FParticleEmit
 		int32 InsertionIndex = 0;
 		for (int32 FindBoneIdx = 0; FindBoneIdx < ValidAssociatedBones.Num(); FindBoneIdx++)
 		{
-			const int32 BoneIdx = SkelMeshComp->SkeletalMesh->GetRefSkeleton().FindBoneIndex(ValidAssociatedBones[FindBoneIdx]);
+			const int32 BoneIdx = SkelMeshComp->GetSkeletalMesh()->GetRefSkeleton().FindBoneIndex(ValidAssociatedBones[FindBoneIdx]);
 			if (BoneIdx != INDEX_NONE && ValidAssociatedBones.Num() > InsertionIndex)
 			{
 				InstancePayload->ValidAssociatedBoneIndices[InsertionIndex++] = BoneIdx;
@@ -2747,10 +2747,10 @@ void UParticleModuleLocationSkelVertSurface::GetSkeletalMeshComponentSource(FPar
 
 		bool bMeshIsValid = false;
 		int32 MinLOD = INDEX_NONE;
-		if (NewSkelMeshComp && NewSkelMeshComp->GetScene() && NewSkelMeshComp->SkeletalMesh)
+		if (NewSkelMeshComp && NewSkelMeshComp->GetScene() && NewSkelMeshComp->GetSkeletalMesh())
 		{
 			FSkeletalMeshRenderData* SkelMeshResource = NewSkelMeshComp->GetSkeletalMeshRenderData();
-			MinLOD = SkelMeshResource->GetFirstValidLODIdx(NewSkelMeshComp->SkeletalMesh->GetMinLodIdx());
+			MinLOD = SkelMeshResource->GetFirstValidLODIdx(NewSkelMeshComp->GetSkeletalMesh()->GetMinLodIdx());
 
 			if (MinLOD != INDEX_NONE)
 			{
@@ -2786,7 +2786,7 @@ void UParticleModuleLocationSkelVertSurface::GetSkeletalMeshComponentSource(FPar
 			UE_LOG(LogParticles, Warning, TEXT("Attempting to use Cascade SkelVertSurface module on mesh without valid data."));
 			UE_LOG(LogParticles, Warning, TEXT("Likely due to CPU side buffers being stripped with r.FreeSkeletalMeshBuffers=1."));
 			UE_LOG(LogParticles, Warning, TEXT("This emitter should probably be culled for this platform using detail mode."));
-			UE_LOG(LogParticles, Warning, TEXT("Mesh: %s"), NewSkelMeshComp ? *NewSkelMeshComp->SkeletalMesh->GetFullName() : TEXT("NULL"));
+			UE_LOG(LogParticles, Warning, TEXT("Mesh: %s"), NewSkelMeshComp ? *NewSkelMeshComp->GetSkeletalMesh()->GetFullName() : TEXT("NULL"));
 			UE_LOG(LogParticles, Warning, TEXT("Comp: %s"), *Owner->Component->GetFullName());
 			UE_LOG(LogParticles, Warning, TEXT("System: %s"), *Owner->Component->Template->GetFullName());
 			UE_LOG(LogParticles, Warning, TEXT("----------------------------------------------------------------------"));
@@ -2964,7 +2964,7 @@ bool UParticleModuleLocationSkelVertSurface::VertInfluencedByActiveBoneTyped(
 			int32 BoneIndex = Section.BoneMap[WeightBuffer->GetBoneIndex(Section.GetVertexBufferIndex() + VertIndex, InfluenceIndex)];
 			if (InSkelMeshComponent->MasterPoseComponent.IsValid())
 			{
-				check(MasterBoneMap.Num() == InSkelMeshComponent->SkeletalMesh->GetRefSkeleton().GetNum());
+				check(MasterBoneMap.Num() == InSkelMeshComponent->GetSkeletalMesh()->GetRefSkeleton().GetNum());
 				BoneIndex = MasterBoneMap[BoneIndex];
 			}
 
diff --git a/Engine/Source/Runtime/Engine/Private/PhysicsEngine/PhysAnim.cpp b/Engine/Source/Runtime/Engine/Private/PhysicsEngine/PhysAnim.cpp
--- a/Engine/Source/Runtime/Engine/Private/PhysicsEngine/PhysAnim.cpp
+++ b/Engine/Source/Runtime/Engine/Private/PhysicsEngine/PhysAnim.cpp
@@ -132,7 +132,7 @@ void UpdateWorldBoneTM(TAssetWorldBoneTMArray& WorldBoneTMs, const TArray<FTrans
 	else
 	{
 		// If not root, use our cached world-space bone transforms.
-		int32 ParentIndex = SkelComp->SkeletalMesh->GetRefSkeleton().GetParentIndex(BoneIndex);
+		int32 ParentIndex = SkelComp->GetSkeletalMesh()->GetRefSkeleton().GetParentIndex(BoneIndex);
 		UpdateWorldBoneTM(WorldBoneTMs, InBoneSpaceTransforms, ParentIndex, SkelComp, LocalToWorldTM, Scale3D);
 		ParentTM = WorldBoneTMs[ParentIndex].TM;
 	}
@@ -201,7 +201,7 @@ void USkeletalMeshComponent::PerformBlendPhysicsBones(const TArray<FBoneIndexTyp
 			int32 BoneIndex = InRequiredBones[i];
 
 			// See if this is a physics bone..
-			int32 BodyIndex = PhysicsAsset->FindBodyIndex(SkeletalMesh->GetRefSkeleton().GetBoneName(BoneIndex));
+			int32 BodyIndex = PhysicsAsset->FindBodyIndex(GetSkeletalMesh()->GetRefSkeleton().GetBoneName(BoneIndex));
 			// need to update back to physX so that physX knows where it was after blending
 			FBodyInstance* PhysicsAssetBodyInstance = nullptr;
 
@@ -213,7 +213,7 @@ void USkeletalMeshComponent::PerformBlendPhysicsBones(const TArray<FBoneIndexTyp
 				if ( !ensure(Bodies.IsValidIndex(BodyIndex)) )
 				{
 					UE_LOG(LogPhysics, Warning, TEXT("%s(Mesh %s, PhysicsAsset %s)"), 
-						*GetName(), *GetNameSafe(SkeletalMesh), *GetNameSafe(PhysicsAsset));
+						*GetName(), *GetNameSafe(GetSkeletalMesh()), *GetNameSafe(PhysicsAsset));
 					UE_LOG(LogPhysics, Warning, TEXT(" - # of BodySetup (%d), # of Bodies (%d), Invalid BodyIndex(%d)"), 
 						PhysicsAsset->SkeletalBodySetups.Num(), Bodies.Num(), BodyIndex);
 					continue;
@@ -256,7 +256,7 @@ void USkeletalMeshComponent::PerformBlendPhysicsBones(const TArray<FBoneIndexTyp
 						else
 						{
 							// If not root, get parent TM from cache (making sure its up-to-date).
-							int32 ParentIndex = SkeletalMesh->GetRefSkeleton().GetParentIndex(BoneIndex);
+							int32 ParentIndex = GetSkeletalMesh()->GetRefSkeleton().GetParentIndex(BoneIndex);
 							UpdateWorldBoneTM(WorldBoneTMs, InOutBoneSpaceTransforms, ParentIndex, this, LocalToWorldTM, TotalScale3D);
 							ParentWorldTM = WorldBoneTMs[ParentIndex].TM;
 						}
@@ -306,7 +306,7 @@ void USkeletalMeshComponent::PerformBlendPhysicsBones(const TArray<FBoneIndexTyp
 					{
 						continue;
 					}
-					const int32 ParentIndex = SkeletalMesh->GetRefSkeleton().GetParentIndex(BoneIndex);
+					const int32 ParentIndex = GetSkeletalMesh()->GetRefSkeleton().GetParentIndex(BoneIndex);
 					InOutComponentSpaceTransforms[BoneIndex] = InOutBoneSpaceTransforms[BoneIndex] * InOutComponentSpaceTransforms[ParentIndex];
 
 					/**
@@ -357,7 +357,7 @@ void USkeletalMeshComponent::BlendInPhysicsInternal(FTickFunction& ThisTickFunct
 	check(IsInGameThread());
 
 	// Can't do anything without a SkeletalMesh
-	if( !SkeletalMesh )
+	if( !GetSkeletalMesh())
 	{
 		return;
 	}
@@ -527,13 +527,13 @@ void USkeletalMeshComponent::UpdateKinematicBonesToAnim(const TArray<FTransform>
 
 #if !(UE_BUILD_SHIPPING || UE_BUILD_TEST)
 	// If desired, draw the skeleton at the point where we pass it to the physics.
-	if (bShowPrePhysBones && SkeletalMesh && InSpaceBases.Num() == SkeletalMesh->GetRefSkeleton().GetNum())
+	if (bShowPrePhysBones && GetSkeletalMesh() && InSpaceBases.Num() == GetSkeletalMesh()->GetRefSkeleton().GetNum())
 	{
 		for (int32 i = 1; i<InSpaceBases.Num(); i++)
 		{
 			FVector ThisPos = CurrentLocalToWorld.TransformPosition(InSpaceBases[i].GetLocation());
 
-			int32 ParentIndex = SkeletalMesh->GetRefSkeleton().GetParentIndex(i);
+			int32 ParentIndex = GetSkeletalMesh()->GetRefSkeleton().GetParentIndex(i);
 			FVector ParentPos = CurrentLocalToWorld.TransformPosition(InSpaceBases[ParentIndex].GetLocation());
 
 			World->LineBatcher->DrawLine(ThisPos, ParentPos, AnimSkelDrawColor, SDPG_Foreground);
@@ -546,7 +546,7 @@ void USkeletalMeshComponent::UpdateKinematicBonesToAnim(const TArray<FTransform>
 #if !(UE_BUILD_SHIPPING || UE_BUILD_TEST)
 	if( !MeshScale3D.IsUniform() )
 	{
-		UE_LOG(LogPhysics, Log, TEXT("USkeletalMeshComponent::UpdateKinematicBonesToAnim : Non-uniform scale factor (%s) can cause physics to mismatch for %s  SkelMesh: %s"), *MeshScale3D.ToString(), *GetFullName(), SkeletalMesh ? *SkeletalMesh->GetFullName() : TEXT("NULL"));
+		UE_LOG(LogPhysics, Log, TEXT("USkeletalMeshComponent::UpdateKinematicBonesToAnim : Non-uniform scale factor (%s) can cause physics to mismatch for %s  SkelMesh: %s"), *MeshScale3D.ToString(), *GetFullName(), GetSkeletalMesh() ? *GetSkeletalMesh()->GetFullName() : TEXT("NULL"));
 	}
 #endif
 
@@ -554,11 +554,11 @@ void USkeletalMeshComponent::UpdateKinematicBonesToAnim(const TArray<FTransform>
 	if (bEnablePerPolyCollision == false)
 	{
 		const UPhysicsAsset* const PhysicsAsset = GetPhysicsAsset();
-		if (PhysicsAsset && SkeletalMesh && Bodies.Num() > 0)
+		if (PhysicsAsset && GetSkeletalMesh() && Bodies.Num() > 0)
 		{
 #if !(UE_BUILD_SHIPPING || UE_BUILD_TEST)
 			if (!ensureMsgf(PhysicsAsset->SkeletalBodySetups.Num() == Bodies.Num(), TEXT("Mesh (%s) has PhysicsAsset(%s), and BodySetup(%d) and Bodies(%d) don't match"),
-						*SkeletalMesh->GetName(), *PhysicsAsset->GetName(), PhysicsAsset->SkeletalBodySetups.Num(), Bodies.Num()))
+						*GetSkeletalMesh()->GetName(), *PhysicsAsset->GetName(), PhysicsAsset->SkeletalBodySetups.Num(), Bodies.Num()))
 			{
 				return;
 			}
@@ -593,7 +593,7 @@ void USkeletalMeshComponent::UpdateKinematicBonesToAnim(const TArray<FTransform>
 							{
 								BodyName = PhysicsAsset->SkeletalBodySetups[i]->BoneName;
 							}
-							UE_LOG(LogPhysics, Log, TEXT("UpdateRBBones: WARNING: Failed to find bone '%s' need by PhysicsAsset '%s' in SkeletalMesh '%s'."), *BodyName.ToString(), *PhysicsAsset->GetName(), *SkeletalMesh->GetName());
+							UE_LOG(LogPhysics, Log, TEXT("UpdateRBBones: WARNING: Failed to find bone '%s' need by PhysicsAsset '%s' in SkeletalMesh '%s'."), *BodyName.ToString(), *PhysicsAsset->GetName(), *GetSkeletalMesh()->GetName());
 						}
 						else
 						{
@@ -604,7 +604,7 @@ void USkeletalMeshComponent::UpdateKinematicBonesToAnim(const TArray<FTransform>
 								{
 									BodyName = PhysicsAsset->SkeletalBodySetups[i]->BoneName;
 								}
-								UE_LOG(LogPhysics, Warning, TEXT("BoneIndex %d out of range of SpaceBases (Size %d) on PhysicsAsset '%s' in SkeletalMesh '%s' for bone '%s'"), BoneIndex, InSpaceBases.Num(), *PhysicsAsset->GetName(), *SkeletalMesh->GetName(), *BodyName.ToString());
+								UE_LOG(LogPhysics, Warning, TEXT("BoneIndex %d out of range of SpaceBases (Size %d) on PhysicsAsset '%s' in SkeletalMesh '%s' for bone '%s'"), BoneIndex, InSpaceBases.Num(), *PhysicsAsset->GetName(), *GetSkeletalMesh()->GetName(), *BodyName.ToString());
 								continue;
 							}
 
@@ -618,7 +618,7 @@ void USkeletalMeshComponent::UpdateKinematicBonesToAnim(const TArray<FTransform>
 									BodyName = PhysicsAsset->SkeletalBodySetups[i]->BoneName;
 								}
 
-								UE_LOG(LogPhysics, Warning, TEXT("UpdateKinematicBonesToAnim: Trying to set transform with bad data %s on PhysicsAsset '%s' in SkeletalMesh '%s' for bone '%s'"), *BoneTransform.ToHumanReadableString(), *PhysicsAsset->GetName(), *SkeletalMesh->GetName(), *BodyName.ToString());
+								UE_LOG(LogPhysics, Warning, TEXT("UpdateKinematicBonesToAnim: Trying to set transform with bad data %s on PhysicsAsset '%s' in SkeletalMesh '%s' for bone '%s'"), *BoneTransform.ToHumanReadableString(), *PhysicsAsset->GetName(), *GetSkeletalMesh()->GetName(), *BodyName.ToString());
 								BoneTransform.DiagnosticCheck_IsValid();	//In special nan mode we want to actually ensure
 
 								continue;
@@ -665,7 +665,7 @@ void USkeletalMeshComponent::UpdateKinematicBonesToAnim(const TArray<FTransform>
 							//It's not clear whether this should be a warning. There are certainly cases where you interpolate the blend weight towards 0. The blend feature needs some work which will probably change this in the future.
 							//Making it Verbose for now
 							UE_LOG(LogPhysics, Verbose, TEXT("%s(Mesh %s, PhysicsAsset %s, Bone %s) is simulating, but no blending. "),
-								*GetName(), *GetNameSafe(SkeletalMesh), *GetNameSafe(PhysicsAsset), *BodyInst->BodySetup.Get()->BoneName.ToString());
+								*GetName(), *GetNameSafe(GetSkeletalMesh()), *GetNameSafe(PhysicsAsset), *BodyInst->BodySetup.Get()->BoneName.ToString());
 						}
 					}
 				}
@@ -724,7 +724,7 @@ void USkeletalMeshComponent::UpdateRBJointMotors()
 	}
 
 	const UPhysicsAsset* const PhysicsAsset = GetPhysicsAsset();
-	if(PhysicsAsset && Constraints.Num() > 0 && SkeletalMesh)
+	if(PhysicsAsset && Constraints.Num() > 0 && GetSkeletalMesh())
 	{
 		check( PhysicsAsset->ConstraintSetup.Num() == Constraints.Num() );
 
@@ -736,7 +736,7 @@ void USkeletalMeshComponent::UpdateRBJointMotors()
 			FConstraintInstance* CI = Constraints[i];
 
 			FName JointChildBoneName = CS->DefaultInstance.GetChildBoneName();
-			int32 BoneIndex = SkeletalMesh->GetRefSkeleton().FindBoneIndex(JointChildBoneName);
+			int32 BoneIndex = GetSkeletalMesh()->GetRefSkeleton().FindBoneIndex(JointChildBoneName);
 
 			// If we found this bone, and a visible bone that is not the root, and its joint is motorised in some way..
 			if( (BoneIndex != INDEX_NONE) && (BoneIndex != 0) &&
@@ -755,8 +755,8 @@ void USkeletalMeshComponent::UpdateRBJointMotors()
 				// We need this to compensate for welding, where graphics and physics parents may not be the same.
 				FMatrix ControlBodyToParentBoneTM = FMatrix::Identity;
 
-				int32 TestBoneIndex = SkeletalMesh->GetRefSkeleton().GetParentIndex(BoneIndex); // This give the 'graphics' parent of this bone
-				bool bFoundControlBody = (SkeletalMesh->GetRefSkeleton().GetBoneName(TestBoneIndex) == CS->DefaultInstance.ConstraintBone2); // ConstraintBone2 is the 'physics' parent of this joint.
+				int32 TestBoneIndex = GetSkeletalMesh()->GetRefSkeleton().GetParentIndex(BoneIndex); // This give the 'graphics' parent of this bone
+				bool bFoundControlBody = (GetSkeletalMesh()->GetRefSkeleton().GetBoneName(TestBoneIndex) == CS->DefaultInstance.ConstraintBone2); // ConstraintBone2 is the 'physics' parent of this joint.
 
 				while(!bFoundControlBody)
 				{
@@ -778,7 +778,7 @@ void USkeletalMeshComponent::UpdateRBJointMotors()
 					ControlBodyToParentBoneTM = ControlBodyToParentBoneTM * RelTM;
 
 					// Move on to parent
-					TestBoneIndex = SkeletalMesh->GetRefSkeleton().GetParentIndex(TestBoneIndex);
+					TestBoneIndex = GetSkeletalMesh()->GetRefSkeleton().GetParentIndex(TestBoneIndex);
 
 					// If we are at the root - bail out.
 					if(TestBoneIndex == 0)
@@ -787,7 +787,7 @@ void USkeletalMeshComponent::UpdateRBJointMotors()
 					}
 
 					// See if this is the controlling body
-					bFoundControlBody = (SkeletalMesh->GetRefSkeleton().GetBoneName(TestBoneIndex) == CS->DefaultInstance.ConstraintBone2);
+					bFoundControlBody = (GetSkeletalMesh()->GetRefSkeleton().GetBoneName(TestBoneIndex) == CS->DefaultInstance.ConstraintBone2);
 				}
 
 				// If after that we didn't find a parent body, we can' do this, so skip.
diff --git a/Engine/Source/Runtime/Engine/Private/PhysicsEngine/PhysicalAnimationComponent.cpp b/Engine/Source/Runtime/Engine/Private/PhysicsEngine/PhysicalAnimationComponent.cpp
--- a/Engine/Source/Runtime/Engine/Private/PhysicsEngine/PhysicalAnimationComponent.cpp
+++ b/Engine/Source/Runtime/Engine/Private/PhysicsEngine/PhysicalAnimationComponent.cpp
@@ -214,7 +214,7 @@ FTransform ComputeWorldSpaceTargetTM(const USkeletalMeshComponent& SkeletalMeshC
 
 FTransform ComputeLocalSpaceTargetTM(const USkeletalMeshComponent& SkeletalMeshComponent, const UPhysicsAsset& PhysAsset, const TArray<FTransform>& LocalTransforms, int32 BoneIndex)
 {
-	const FReferenceSkeleton& RefSkeleton = SkeletalMeshComponent.SkeletalMesh->GetRefSkeleton();
+	const FReferenceSkeleton& RefSkeleton = SkeletalMeshComponent.GetSkeletalMesh()->GetRefSkeleton();
 	FTransform AccumulatedDelta = LocalTransforms[BoneIndex];
 	int32 CurBoneIdx = BoneIndex;
 	while ((CurBoneIdx = RefSkeleton.GetParentIndex(CurBoneIdx)) != INDEX_NONE)
@@ -256,9 +256,9 @@ FTransform ComputeTargetTM(const FPhysicalAnimationData& PhysAnimData, const USk
 void UPhysicalAnimationComponent::UpdateTargetActors(ETeleportType TeleportType)
 {
 	UPhysicsAsset* PhysAsset = SkeletalMeshComponent ? SkeletalMeshComponent->GetPhysicsAsset() : nullptr;
-	if (PhysAsset && SkeletalMeshComponent->SkeletalMesh)
+	if (PhysAsset && SkeletalMeshComponent->GetSkeletalMesh())
 	{
-		const FReferenceSkeleton& RefSkeleton = SkeletalMeshComponent->SkeletalMesh->GetRefSkeleton();
+		const FReferenceSkeleton& RefSkeleton = SkeletalMeshComponent->GetSkeletalMesh()->GetRefSkeleton();
 
 		// Note we use GetEditableComponentSpaceTransforms because we need to update target actors in the midst of the 
 		// various anim ticks, before buffers are flipped (which happens in the skel mesh component's post-physics tick)
@@ -357,7 +357,7 @@ void UPhysicalAnimationComponent::UpdatePhysicsEngineImp()
 {
 	bPhysicsEngineNeedsUpdating = false;
 	UPhysicsAsset* PhysAsset = SkeletalMeshComponent ? SkeletalMeshComponent->GetPhysicsAsset() : nullptr;
-	if(PhysAsset && SkeletalMeshComponent->SkeletalMesh)
+	if(PhysAsset && SkeletalMeshComponent->GetSkeletalMesh())
 	{
 		//TODO: This is hacky and assumes constraints can only be added and not removed. True for now, but bad in general!
 		const int32 NumData = DriveData.Num();
@@ -368,7 +368,7 @@ void UPhysicalAnimationComponent::UpdatePhysicsEngineImp()
 		// Note we use GetEditableComponentSpaceTransforms because we need to update target actors in the midst of the 
 		// various anim ticks, before buffers are flipped (which happens in the skel mesh component's post-physics tick)
 		const TArray<FTransform>& SpaceBases = SkeletalMeshComponent->GetEditableComponentSpaceTransforms();
-		const FReferenceSkeleton& RefSkeleton = SkeletalMeshComponent->SkeletalMesh->GetRefSkeleton();
+		const FReferenceSkeleton& RefSkeleton = SkeletalMeshComponent->GetSkeletalMesh()->GetRefSkeleton();
 
 		FPhysicsCommand::ExecuteWrite(SkeletalMeshComponent, [&]()
 		{
diff --git a/Engine/Source/Runtime/Engine/Private/PhysicsEngine/PhysicsAsset.cpp b/Engine/Source/Runtime/Engine/Private/PhysicsEngine/PhysicsAsset.cpp
--- a/Engine/Source/Runtime/Engine/Private/PhysicsEngine/PhysicsAsset.cpp
+++ b/Engine/Source/Runtime/Engine/Private/PhysicsEngine/PhysicsAsset.cpp
@@ -394,7 +394,7 @@ FBox UPhysicsAsset::CalcAABB(const USkinnedMeshComponent* MeshComp, const FTrans
 	}
 	else
 	{
-		UE_LOG(LogPhysics, Log,  TEXT("UPhysicsAsset::CalcAABB : Non-uniform scale factor. You will not be able to collide with it.  Turn off collision and wrap it with a blocking volume.  MeshComp: %s  SkelMesh: %s"), *MeshComp->GetFullName(), MeshComp->SkeletalMesh ? *MeshComp->SkeletalMesh->GetFullName() : TEXT("NULL") );
+		UE_LOG(LogPhysics, Log,  TEXT("UPhysicsAsset::CalcAABB : Non-uniform scale factor. You will not be able to collide with it.  Turn off collision and wrap it with a blocking volume.  MeshComp: %s  SkelMesh: %s"), *MeshComp->GetFullName(), MeshComp->GetSkeletalMesh() ? *MeshComp->GetSkeletalMesh()->GetFullName() : TEXT("NULL") );
 	}
 
 	if(!Box.IsValid)
diff --git a/Engine/Source/Runtime/Engine/Private/SkeletalMeshComponentPhysics.cpp b/Engine/Source/Runtime/Engine/Private/SkeletalMeshComponentPhysics.cpp
--- a/Engine/Source/Runtime/Engine/Private/SkeletalMeshComponentPhysics.cpp
+++ b/Engine/Source/Runtime/Engine/Private/SkeletalMeshComponentPhysics.cpp
@@ -160,10 +160,10 @@ void USkeletalMeshComponent::CreateBodySetup()
 		BodySetup = NewObject<UBodySetup>(this);
 	}
 
-	if (SkeletalMesh)
+	if (GetSkeletalMesh())
 	{
-		const USkeletalMesh* SkeletalMeshConst = SkeletalMesh;
-		SkeletalMesh->CreateBodySetup();
+		const USkeletalMesh* SkeletalMeshConst = GetSkeletalMesh();
+		GetSkeletalMesh()->CreateBodySetup();
 		UBodySetup* OriginalBodySetup = SkeletalMeshConst->GetBodySetup();
 		BodySetup->CopyBodyPropertiesFrom(OriginalBodySetup);
 		BodySetup->CookedFormatDataOverride = &OriginalBodySetup->CookedFormatData;
@@ -184,11 +184,11 @@ UBodySetup* USkeletalMeshComponent::GetBodySetup()
 	if (bEnablePerPolyCollision == false)
 	{
 		UPhysicsAsset * const PhysicsAsset = GetPhysicsAsset();
-		if (SkeletalMesh && PhysicsAsset)
+		if (GetSkeletalMesh() && PhysicsAsset)
 		{
-			for (int32 i = 0; i < SkeletalMesh->GetRefSkeleton().GetNum(); i++)
+			for (int32 i = 0; i < GetSkeletalMesh()->GetRefSkeleton().GetNum(); i++)
 			{
-				int32 BodyIndex = PhysicsAsset->FindBodyIndex(SkeletalMesh->GetRefSkeleton().GetBoneName(i));
+				int32 BodyIndex = PhysicsAsset->FindBodyIndex(GetSkeletalMesh()->GetRefSkeleton().GetBoneName(i));
 				if (BodyIndex != INDEX_NONE)
 				{
 					return PhysicsAsset->SkeletalBodySetups[BodyIndex];
@@ -651,13 +651,13 @@ int32 USkeletalMeshComponent::FindRootBodyIndex() const
 {
 	// Find root physics body
 	int32 RootBodyIndex = RootBodyData.BodyIndex;
-	if(RootBodyIndex == INDEX_NONE && SkeletalMesh)
+	if(RootBodyIndex == INDEX_NONE && GetSkeletalMesh())
 	{
 		if(const UPhysicsAsset* PhysicsAsset = GetPhysicsAsset())
 		{
-			for (int32 i = 0; i<SkeletalMesh->GetRefSkeleton().GetNum(); i++)
+			for (int32 i = 0; i< GetSkeletalMesh()->GetRefSkeleton().GetNum(); i++)
 			{
-				int32 BodyInstIndex = PhysicsAsset->FindBodyIndex(SkeletalMesh->GetRefSkeleton().GetBoneName(i));
+				int32 BodyInstIndex = PhysicsAsset->FindBodyIndex(GetSkeletalMesh()->GetRefSkeleton().GetBoneName(i));
 				if (BodyInstIndex != INDEX_NONE)
 				{
 					RootBodyIndex = BodyInstIndex;
@@ -684,7 +684,7 @@ void USkeletalMeshComponent::InitArticulated(FPhysScene* PhysScene)
 
 	UPhysicsAsset* const PhysicsAsset = GetPhysicsAsset();
 
-	if (PhysScene == nullptr || PhysicsAsset == nullptr || SkeletalMesh == nullptr || !ShouldCreatePhysicsState())
+	if (PhysScene == nullptr || PhysicsAsset == nullptr || GetSkeletalMesh() == nullptr || !ShouldCreatePhysicsState())
 	{
 		return;
 	}
@@ -774,9 +774,9 @@ void USkeletalMeshComponent::InstantiatePhysicsAssetBodies(const UPhysicsAsset&
 
 void USkeletalMeshComponent::InstantiatePhysicsAssetRefPose(const UPhysicsAsset& PhysAsset, const FVector& Scale3D, TArray<FBodyInstance*>& OutBodies, TArray<FConstraintInstance*>& OutConstraints, FPhysScene* PhysScene /*= nullptr*/, USkeletalMeshComponent* OwningComponent /*= nullptr*/, int32 UseRootBodyIndex /*= INDEX_NONE*/, const FPhysicsAggregateHandle& UseAggregate, bool bCreateBodiesInRefPose) const
 {
-	if(SkeletalMesh)
+	if(GetSkeletalMesh())
 	{
-		const FReferenceSkeleton& RefSkeleton = SkeletalMesh->GetRefSkeleton();
+		const FReferenceSkeleton& RefSkeleton = GetSkeletalMesh()->GetRefSkeleton();
 
 		if (bCreateBodiesInRefPose)
 		{
@@ -1028,7 +1028,7 @@ void USkeletalMeshComponent::TermArticulated()
 void USkeletalMeshComponent::TermBodiesBelow(FName ParentBoneName)
 {
 	UPhysicsAsset* const PhysicsAsset = GetPhysicsAsset();
-	if(PhysicsAsset && SkeletalMesh && Bodies.Num() > 0)
+	if(PhysicsAsset && GetSkeletalMesh() && Bodies.Num() > 0)
 	{
 		check(Bodies.Num() == PhysicsAsset->SkeletalBodySetups.Num());
 
@@ -1048,7 +1048,7 @@ void USkeletalMeshComponent::TermBodiesBelow(FName ParentBoneName)
 			int32 JointBoneIndex = GetBoneIndex(JointChildBoneName);
 
 			// If constraint has bone in mesh, and is either the parent or child of it, term it
-			if(	JointBoneIndex != INDEX_NONE && (JointChildBoneName == ParentBoneName ||	SkeletalMesh->GetRefSkeleton().BoneIsChildOf(JointBoneIndex, ParentBoneIndex)) )
+			if(	JointBoneIndex != INDEX_NONE && (JointChildBoneName == ParentBoneName || GetSkeletalMesh()->GetRefSkeleton().BoneIsChildOf(JointBoneIndex, ParentBoneIndex)) )
 			{
 				Constraints[i]->TermConstraint();
 			}
@@ -1064,7 +1064,7 @@ void USkeletalMeshComponent::TermBodiesBelow(FName ParentBoneName)
 				int32 BodyBoneIndex = GetBoneIndex(BodyName);
 
 				// If body has bone in mesh, and is either the parent or child of it, term it
-				if(	BodyBoneIndex != INDEX_NONE && (BodyName == ParentBoneName ||	SkeletalMesh->GetRefSkeleton().BoneIsChildOf(BodyBoneIndex, ParentBoneIndex)) )
+				if(	BodyBoneIndex != INDEX_NONE && (BodyName == ParentBoneName || GetSkeletalMesh()->GetRefSkeleton().BoneIsChildOf(BodyBoneIndex, ParentBoneIndex)) )
 				{
 					Bodies[i]->TermBody();
 				}
@@ -1655,7 +1655,7 @@ void USkeletalMeshComponent::UpdateMeshForBrokenConstraints()
 {
 	UPhysicsAsset * const PhysicsAsset = GetPhysicsAsset();
 	// Needs to have a SkeletalMesh, and PhysicsAsset.
-	if( !SkeletalMesh || !PhysicsAsset )
+	if( !GetSkeletalMesh() || !PhysicsAsset )
 	{
 		return;
 	}
@@ -1684,7 +1684,7 @@ void USkeletalMeshComponent::UpdateMeshForBrokenConstraints()
 				UBodySetup* PhysicsAssetBodySetup = PhysicsAsset->SkeletalBodySetups[BodySetupIndex];
 				int32 BoneIndex = GetBoneIndex(PhysicsAssetBodySetup->BoneName);
 				if( BoneIndex != INDEX_NONE && 
-					(BoneIndex == JointBoneIndex || SkeletalMesh->GetRefSkeleton().BoneIsChildOf(BoneIndex, JointBoneIndex)) )
+					(BoneIndex == JointBoneIndex || GetSkeletalMesh()->GetRefSkeleton().BoneIsChildOf(BoneIndex, JointBoneIndex)) )
 				{
 					DEBUGBROKENCONSTRAINTUPDATE(UE_LOG(LogSkeletalMesh, Log, TEXT("    Found Child Bone: (%d) %s"), BoneIndex, *PhysicsAssetBodySetup->BoneName.ToString());)
 
@@ -1825,7 +1825,7 @@ int32 USkeletalMeshComponent::ForEachBodyBelow(FName BoneName, bool bIncludeSelf
 	else
 	{
 		UPhysicsAsset* const PhysicsAsset = GetPhysicsAsset();
-		if (!PhysicsAsset || !SkeletalMesh)
+		if (!PhysicsAsset || !GetSkeletalMesh())
 		{
 			return 0;
 		}
@@ -1839,7 +1839,7 @@ int32 USkeletalMeshComponent::ForEachBodyBelow(FName BoneName, bool bIncludeSelf
 
 		TArray<int32> BodyIndices;
 		BodyIndices.Reserve(Bodies.Num());
-		PhysicsAsset->GetBodyIndicesBelow(BodyIndices, BoneName, SkeletalMesh, bIncludeSelf);
+		PhysicsAsset->GetBodyIndicesBelow(BodyIndices, BoneName, GetSkeletalMesh(), bIncludeSelf);
 
 		int32 NumBodiesFound = 0;
 		for (int32 BodyIdx : BodyIndices)
@@ -2062,12 +2062,12 @@ void USkeletalMeshComponent::SetPhysicsAsset(UPhysicsAsset* InPhysicsAsset, bool
 		UpdateHasValidBodies();
 
 		// Component should be re-attached here, so create physics.
-		if( SkeletalMesh )
+		if(GetSkeletalMesh())
 		{
 			// Because we don't know what bones the new PhysicsAsset might want, we have to force an update to _all_ bones in the skeleton.
-			RequiredBones.Reset(SkeletalMesh->GetRefSkeleton().GetNum());
-			RequiredBones.AddUninitialized( SkeletalMesh->GetRefSkeleton().GetNum() );
-			for(int32 i=0; i<SkeletalMesh->GetRefSkeleton().GetNum(); i++)
+			RequiredBones.Reset(GetSkeletalMesh()->GetRefSkeleton().GetNum());
+			RequiredBones.AddUninitialized(GetSkeletalMesh()->GetRefSkeleton().GetNum() );
+			for(int32 i=0; i< GetSkeletalMesh()->GetRefSkeleton().GetNum(); i++)
 			{
 				RequiredBones[i] = (FBoneIndexType)i;
 			}
@@ -2184,8 +2184,8 @@ FVector GetTypedSkinnedVertexPositionWithCloth(USkeletalMeshComponent* Component
 	const FSkelMeshRenderSection& Section = LODData.RenderSections[SectionIndex];
 
 	// only if this component has clothing and is showing simulated results	
-	if (Component->SkeletalMesh &&
-		Component->SkeletalMesh->GetMeshClothingAssets().Num() > 0 &&
+	if (Component->GetSkeletalMesh() &&
+		Component->GetSkeletalMesh()->GetMeshClothingAssets().Num() > 0 &&
 		!Component->bDisableClothSimulation &&
 		Component->ClothBlendWeight > 0.0f // if cloth blend weight is 0.0, only showing skinned vertices regardless of simulation positions
 		)
@@ -2240,12 +2240,12 @@ FVector3f USkeletalMeshComponent::GetSkinnedVertexPosition(USkeletalMeshComponen
 void USkeletalMeshComponent::ComputeSkinnedPositions(USkeletalMeshComponent* Component, TArray<FVector3f> & OutPositions, TArray<FMatrix44f>& CachedRefToLocals, const FSkeletalMeshLODRenderData& LODData, const FSkinWeightVertexBuffer& SkinWeightBuffer)
 {
 	// Fail if no mesh
-	if (!Component->SkeletalMesh)
+	if (!Component->GetSkeletalMesh())
 	{
 		return;
 	}
 
-	if (Component->SkeletalMesh->GetMeshClothingAssets().Num() > 0 &&
+	if (Component->GetSkeletalMesh()->GetMeshClothingAssets().Num() > 0 &&
 		!Component->bDisableClothSimulation &&
 		Component->ClothBlendWeight > 0.0f // if cloth blend weight is 0.0, only showing skinned vertices regardless of simulation positions
 		)
@@ -2280,7 +2280,7 @@ void USkeletalMeshComponent::ComputeSkinnedPositions(USkeletalMeshComponent* Com
 
 			if (bClothVertex)
 			{
-				int32 AssetIndex = Component->SkeletalMesh->GetClothingAssetIndex(ClothAssetGuid);
+				int32 AssetIndex = Component->GetSkeletalMesh()->GetClothingAssetIndex(ClothAssetGuid);
 				if (AssetIndex != INDEX_NONE)
 				{
 					const FClothSimulData* ActorData = ClothData.Find(AssetIndex);
@@ -2360,7 +2360,7 @@ void USkeletalMeshComponent::GetSkinnedTangentBasis(USkeletalMeshComponent* Comp
 void USkeletalMeshComponent::ComputeSkinnedTangentBasis(USkeletalMeshComponent* Component, TArray<FVector3f>& OutTangenXYZ, TArray<FMatrix44f>& CachedRefToLocals, const FSkeletalMeshLODRenderData& LODData, const FSkinWeightVertexBuffer& SkinWeightBuffer)
 {
 	// Fail if no mesh
-	if (!Component->SkeletalMesh)
+	if (!Component->GetSkeletalMesh())
 	{
 		return;
 	}
@@ -2454,7 +2454,7 @@ bool USkeletalMeshComponent::GetClosestPointOnPhysicsAsset(const FVector& WorldP
 
 	bool bSuccess = false;
 	const UPhysicsAsset* PhysicsAsset = GetPhysicsAsset();
-	const FReferenceSkeleton* RefSkeleton = SkeletalMesh ? &SkeletalMesh->GetRefSkeleton() : nullptr;
+	const FReferenceSkeleton* RefSkeleton = GetSkeletalMesh() ? &GetSkeletalMesh()->GetRefSkeleton() : nullptr;
 	if(PhysicsAsset && RefSkeleton)
 	{
 		const TArray<FTransform>& BoneTransforms = GetComponentSpaceTransforms();
@@ -2666,12 +2666,13 @@ void USkeletalMeshComponent::RecreateClothingActors()
 
 	ReleaseAllClothingResources();
 
-	if(!bAllowClothActors || SkeletalMesh == nullptr || !IsRegistered())
+	USkeletalMesh* SkelMesh = GetSkeletalMesh();
+	if(!bAllowClothActors || SkelMesh == nullptr || !IsRegistered())
 	{
 		return;
 	}
 
-	if(CVarEnableClothPhysics.GetValueOnGameThread() && (SkeletalMesh->GetMeshClothingAssets().Num() > 0))
+	if(CVarEnableClothPhysics.GetValueOnGameThread() && (SkelMesh->GetMeshClothingAssets().Num() > 0))
 	{
 		UClass* SimFactoryClass = *ClothingSimulationFactory;
 		if (SimFactoryClass)
@@ -2705,15 +2706,15 @@ void USkeletalMeshComponent::RecreateClothingActors()
 #endif
 			{
 				TArray<UClothingAssetBase*> AssetsInUse;
-				SkeletalMesh->GetClothingAssetsInUse(AssetsInUse);
+				SkelMesh->GetClothingAssetsInUse(AssetsInUse);
 
 				const bool bIsInitialization = true;
 				ClothingSimulation->FillContext(this, 0.f, ClothingSimulationContext, bIsInitialization);
 
-				const int32 NumMeshAssets = SkeletalMesh->GetMeshClothingAssets().Num();
+				const int32 NumMeshAssets = SkelMesh->GetMeshClothingAssets().Num();
 				for (int32 BaseAssetIndex = 0; BaseAssetIndex < NumMeshAssets; ++BaseAssetIndex)
 				{
-					UClothingAssetBase* const Asset = SkeletalMesh->GetMeshClothingAssets()[BaseAssetIndex];
+					UClothingAssetBase* const Asset = SkelMesh->GetMeshClothingAssets()[BaseAssetIndex];
 					if (Asset && AssetsInUse.Contains(Asset))
 					{
 						ClothingSimulation->CreateActor(this, Asset, BaseAssetIndex);
@@ -2825,7 +2826,7 @@ void USkeletalMeshComponent::ExtractCollisionsForCloth(
 {
 	// Extract collisions from this mesh 'raw', as this isn't a mesh that has cloth simulation
 	// (but we want it to affect other meshes with cloth simulation)
-	if(SourceComponent->SkeletalMesh && PhysicsAsset)
+	if(SourceComponent->GetSkeletalMesh() && PhysicsAsset)
 	{
 		FTransform ComponentToComponentTransform;
 		if(SourceComponent != DestClothComponent)
@@ -2836,12 +2837,12 @@ void USkeletalMeshComponent::ExtractCollisionsForCloth(
 		}
 
 		// Init cache on first copy
-		if(!ClothCollisionSource.bCached || ClothCollisionSource.CachedSkeletalMesh.Get() != SourceComponent->SkeletalMesh)
+		if(!ClothCollisionSource.bCached || ClothCollisionSource.CachedSkeletalMesh.Get() != SourceComponent->GetSkeletalMesh())
 		{
 			for(const USkeletalBodySetup* SkeletalBodySetup : PhysicsAsset->SkeletalBodySetups)
 			{
 				// Cache bones
-				int32 MeshBoneIndex = SourceComponent->SkeletalMesh->GetRefSkeleton().FindBoneIndex(SkeletalBodySetup->BoneName);
+				int32 MeshBoneIndex = SourceComponent->GetSkeletalMesh()->GetRefSkeleton().FindBoneIndex(SkeletalBodySetup->BoneName);
 				if(MeshBoneIndex != INDEX_NONE)
 				{
 					// Cache spheres & capsules form physics asset
@@ -2903,7 +2904,7 @@ void USkeletalMeshComponent::ExtractCollisionsForCloth(
 				}
 			}
 
-			ClothCollisionSource.CachedSkeletalMesh = SourceComponent->SkeletalMesh;
+			ClothCollisionSource.CachedSkeletalMesh = SourceComponent->GetSkeletalMesh();
 			ClothCollisionSource.bCached = true;
 		}
 
@@ -3336,7 +3337,7 @@ void USkeletalMeshComponent::ProcessClothCollisionWithEnvironment()
 				// Possibly a skeletal mesh, extract it's clothing collisions if necessary
 				USkeletalMeshComponent* SkelComp = Cast<USkeletalMeshComponent>(Component.Get());
 
-				if(SkelComp && SkelComp->SkeletalMesh)
+				if(SkelComp && SkelComp->GetSkeletalMesh())
 				{
 					if(SkelComp == this)
 					{
@@ -3601,7 +3602,7 @@ public:
 
 	void DoTask(ENamedThreads::Type CurrentThread, const FGraphEventRef& MyCompletionGraphEvent)
 	{
-		FScopeCycleCounterUObject ContextScope(SkeletalMeshComponent.SkeletalMesh);
+		FScopeCycleCounterUObject ContextScope(SkeletalMeshComponent.GetSkeletalMesh());
 		SCOPE_CYCLE_COUNTER(STAT_ClothTotalTime);
 		CSV_SCOPED_TIMING_STAT(Animation, Cloth);
 
@@ -3758,7 +3759,7 @@ DECLARE_CYCLE_STAT(TEXT("Cloth Sim"), STAT_ClothSimTime, STATGROUP_Physics);
 
 bool USkeletalMeshComponent::GetClothSimulatedPosition_GameThread(const FGuid& AssetGuid, int32 VertexIndex, FVector& OutSimulPos) const
 {
-	if(!SkeletalMesh)
+	if(!GetSkeletalMesh())
 	{
 		// Can't proceed without a mesh
 		return false;
@@ -3766,7 +3767,7 @@ bool USkeletalMeshComponent::GetClothSimulatedPosition_GameThread(const FGuid& A
 		
 	bool bSucceed = false;
 
-	int32 AssetIndex = SkeletalMesh->GetClothingAssetIndex(AssetGuid);
+	int32 AssetIndex = GetSkeletalMesh()->GetClothingAssetIndex(AssetGuid);
 
 	if(AssetIndex != INDEX_NONE)
 	{
@@ -3784,7 +3785,7 @@ bool USkeletalMeshComponent::GetClothSimulatedPosition_GameThread(const FGuid& A
 
 void USkeletalMeshComponent::TickClothing(float DeltaTime, FTickFunction& ThisTickFunction)
 {
-	if (SkeletalMesh == nullptr || !ClothingSimulation || CVarEnableClothPhysics.GetValueOnGameThread() == 0)
+	if (GetSkeletalMesh() == nullptr || !ClothingSimulation || CVarEnableClothPhysics.GetValueOnGameThread() == 0)
 	{
 		return;
 	}
diff --git a/Engine/Source/Runtime/Engine/Private/UnrealEngine.cpp b/Engine/Source/Runtime/Engine/Private/UnrealEngine.cpp
--- a/Engine/Source/Runtime/Engine/Private/UnrealEngine.cpp
+++ b/Engine/Source/Runtime/Engine/Private/UnrealEngine.cpp
@@ -7860,7 +7860,7 @@ bool UEngine::HandleMergeMeshCommand( const TCHAR* Cmd, FOutputDevice& Ar, UWorl
 		if (PlayerController && PlayerController->GetCharacter() != NULL && PlayerController->GetCharacter()->GetMesh())
 		{
 			PlayerPawn = PlayerController->GetCharacter();
-			PlayerMesh = PlayerController->GetCharacter()->GetMesh()->SkeletalMesh;
+			PlayerMesh = PlayerController->GetCharacter()->GetMesh()->GetSkeletalMesh();
 			break;
 		}
 	}
diff --git a/Engine/Source/Runtime/Experimental/GeometryCollectionEngine/Private/GeometryCollection/PhysicsAssetSimulation.cpp b/Engine/Source/Runtime/Experimental/GeometryCollectionEngine/Private/GeometryCollection/PhysicsAssetSimulation.cpp
--- a/Engine/Source/Runtime/Experimental/GeometryCollectionEngine/Private/GeometryCollection/PhysicsAssetSimulation.cpp
+++ b/Engine/Source/Runtime/Experimental/GeometryCollectionEngine/Private/GeometryCollection/PhysicsAssetSimulation.cpp
@@ -36,7 +36,7 @@ void FPhysicsAssetSimulationUtil::BuildParams(const UObject* Caller, const AActo
 	Chaos::FVec3 ActorOrigin(ActorOriginVec);
 	Chaos::FVec3 ActorBoxExtent(ActorBoxExtentVec);
 
-	const USkeletalMesh* SkeletalMesh = SkelMeshComponent->SkeletalMesh;
+	const USkeletalMesh* SkeletalMesh = SkelMeshComponent->GetSkeletalMesh();
 	if (!SkeletalMesh)
 	{
 		return;
@@ -345,7 +345,7 @@ void FPhysicsAssetSimulationUtil::BuildParams(const UObject* Caller, const AActo
 		int32 TmpBoneIndex = BoneIndex;
 		int32 SocketIndex = INDEX_NONE;
 		if (USkeletalMeshSocket const* Socket =
-			SkelMeshComponent->SkeletalMesh->FindSocketInfo(
+			SkelMeshComponent->GetSkeletalMesh()->FindSocketInfo(
 				Params.BoneHierarchy.GetAnalyticShapeGroup(BoneIndex)->GetBoneName(),
 				BoneXf,
 				TmpBoneIndex,
@@ -375,7 +375,7 @@ bool FPhysicsAssetSimulationUtil::UpdateAnimState(const UObject* Caller, const A
 {
 	int32 UpdatedBones = 0;
 
-	const USkeletalMesh* SkeletalMesh = SkelMeshComponent->SkeletalMesh;
+	const USkeletalMesh* SkeletalMesh = SkelMeshComponent->GetSkeletalMesh();
 	
 	if (!SkeletalMesh)
 	{
@@ -413,7 +413,7 @@ bool FPhysicsAssetSimulationUtil::UpdateAnimState(const UObject* Caller, const A
 			const int32 SocketIndex = Params.BoneHierarchy.GetSocketIndexForBone(BoneIndex);
 			if (SocketIndex != INDEX_NONE)
 			{
-				USkeletalMeshSocket *Socket = SkelMeshComponent->SkeletalMesh->GetSocketByIndex(SocketIndex);
+				USkeletalMeshSocket *Socket = SkelMeshComponent->GetSkeletalMesh()->GetSocketByIndex(SocketIndex);
 				Params.BoneHierarchy.SetAnimLocalSpaceTransform(
 					BoneIndex,
 					Socket->GetSocketLocalTransform());
diff --git a/Engine/Source/Runtime/MovieSceneTracks/Private/Evaluation/MovieSceneSkeletalAnimationTemplate.cpp b/Engine/Source/Runtime/MovieSceneTracks/Private/Evaluation/MovieSceneSkeletalAnimationTemplate.cpp
--- a/Engine/Source/Runtime/MovieSceneTracks/Private/Evaluation/MovieSceneSkeletalAnimationTemplate.cpp
+++ b/Engine/Source/Runtime/MovieSceneTracks/Private/Evaluation/MovieSceneSkeletalAnimationTemplate.cpp
@@ -31,8 +31,8 @@ bool ShouldUsePreviewPlayback(IMovieScenePlayer& Player, UObject& RuntimeObject)
 
 bool CanPlayAnimation(USkeletalMeshComponent* SkeletalMeshComponent, UAnimSequenceBase* AnimAssetBase)
 {
-	return (SkeletalMeshComponent->SkeletalMesh && SkeletalMeshComponent->SkeletalMesh->GetSkeleton() && 
-		(!AnimAssetBase || SkeletalMeshComponent->SkeletalMesh->GetSkeleton()->IsCompatible(AnimAssetBase->GetSkeleton())));
+	return (SkeletalMeshComponent->GetSkeletalMesh() && SkeletalMeshComponent->GetSkeletalMesh()->GetSkeleton() &&
+		(!AnimAssetBase || SkeletalMeshComponent->GetSkeletalMesh()->GetSkeleton()->IsCompatible(AnimAssetBase->GetSkeleton())));
 }
 
 void ResetAnimSequencerInstance(UObject& InObject, const UE::MovieScene::FRestoreStateParams& Params)
@@ -142,10 +142,10 @@ struct FPreAnimatedAnimationTokenProducer : IMovieScenePreAnimatedTokenProducer
 				{
 					Component->AnimScriptInstance = CachedAnimInstance.Get();
 					CachedAnimInstance.Reset();
-					if (Component->AnimScriptInstance && Component->SkeletalMesh && Component->AnimScriptInstance->CurrentSkeleton != Component->SkeletalMesh->GetSkeleton())
+					if (Component->AnimScriptInstance && Component->GetSkeletalMesh() && Component->AnimScriptInstance->CurrentSkeleton != Component->GetSkeletalMesh()->GetSkeleton())
 					{
 						//the skeleton may have changed so need to recalc required bones as needed.
-						Component->AnimScriptInstance->CurrentSkeleton = Component->SkeletalMesh->GetSkeleton();
+						Component->AnimScriptInstance->CurrentSkeleton = Component->GetSkeletalMesh()->GetSkeleton();
 						//Need at least RecalcRequiredbones and UpdateMorphTargetrs
 						Component->InitializeAnimScriptInstance(true);
 					}
@@ -281,7 +281,7 @@ namespace MovieScene
 			ensureMsgf(InObject, TEXT("Attempting to evaluate an Animation track with a null object."));
 
 			USkeletalMeshComponent* SkeletalMeshComponent = SkeletalMeshComponentFromObject(InObject);
-			if (!SkeletalMeshComponent || !SkeletalMeshComponent->SkeletalMesh)
+			if (!SkeletalMeshComponent || !SkeletalMeshComponent->GetSkeletalMesh())
 			{
 				return;
 			}
