World Partition:
- Added actor descriptors properties in ToString function.
- Added ForEachProperty and ToString to FPropertyPairsMap.
- Generalized breakable iterators with TOverrideVoidReturnInvoker.

#rnx

diff --git a/Engine/Source/Runtime/Engine/Private/Tests/PropertyPairsMapTests.cpp b/Engine/Source/Runtime/Engine/Private/Tests/PropertyPairsMapTests.cpp
--- a/Engine/Source/Runtime/Engine/Private/Tests/PropertyPairsMapTests.cpp
+++ b/Engine/Source/Runtime/Engine/Private/Tests/PropertyPairsMapTests.cpp
@@ -40,6 +40,19 @@ namespace PropertyPairsMapTests
 		TestTrue(TEXT("PropertyPairsMap.GetProperty(WithValue)"), PropertyPairsMap.GetProperty(TEXT("WithValue"), &WithValueValue));
 		TestTrue(TEXT("WithValueValue is Value2"), WithValueValue.ToString() == TEXT("Value2"));
 
+		PropertyPairsMap.ForEachProperty([this](FName Name, FName Value)
+		{
+			TestTrue(TEXT("Unbreakable iterator"), (Name == TEXT("WithoutValue")) || (Name == TEXT("WithValue")));
+			TestTrue(TEXT("Unbreakable iterator"), Value.IsNone() || (Value == TEXT("Value2")));
+		});
+
+		PropertyPairsMap.ForEachProperty([this](FName Name, FName Value)
+		{
+			TestTrue(TEXT("Breakable iterator"), (Name == TEXT("WithoutValue")) || (Name == TEXT("WithValue")));
+			TestTrue(TEXT("Breakable iterator"), Value.IsNone() || (Value == TEXT("Value2")));
+			return true;
+		});
+
 		return true;
 	}
 }
diff --git a/Engine/Source/Runtime/Engine/Private/WorldPartition/WorldPartition.cpp b/Engine/Source/Runtime/Engine/Private/WorldPartition/WorldPartition.cpp
--- a/Engine/Source/Runtime/Engine/Private/WorldPartition/WorldPartition.cpp
+++ b/Engine/Source/Runtime/Engine/Private/WorldPartition/WorldPartition.cpp
@@ -72,7 +72,7 @@ TMap<FName, FString> GetDataLayersDumpString(const UWorldPartition* WorldPartiti
 	const UDataLayerSubsystem* DataLayerSubsystem = UWorld::GetSubsystem<UDataLayerSubsystem>(WorldPartition->GetWorld());
 	DataLayerSubsystem->ForEachDataLayer([&DataLayersDumpString](const UDataLayerInstance* DataLayer)
 	{
-		DataLayersDumpString.FindOrAdd(DataLayer->GetDataLayerFName()) = FString::Format(TEXT("{0}({1})"), { DataLayer->GetDataLayerShortName(), DataLayer->GetDataLayerFName().ToString() });
+		DataLayersDumpString.FindOrAdd(DataLayer->GetDataLayerFName()) = FString::Format(TEXT("{0}{1})"), { DataLayer->GetDataLayerShortName(), DataLayer->GetDataLayerFName().ToString() });
 		return true;
 	});
 	
@@ -88,31 +88,20 @@ FString GetActorDescDumpString(const FWorldPartitionActorDesc* ActorDesc, const
 			return FString("None");
 		}
 
-		return FString::JoinBy(DataLayerNames, TEXT(", "), 
-			[&DataLayersDumpString](const FName& DataLayerName) 
-			{ 
-				if (const FString* DumpString = DataLayersDumpString.Find(DataLayerName))
-				{
-					return *DumpString;
-				}
-				return DataLayerName.ToString(); 
-			});
+		return FString::JoinBy(DataLayerNames, TEXT(", "), [&DataLayersDumpString](const FName& DataLayerName) 
+		{ 
+			if (const FString* DumpString = DataLayersDumpString.Find(DataLayerName))
+			{
+				return *DumpString;
+			}
+			return DataLayerName.ToString(); 
+		});
 	};
 
 	check(ActorDesc);
 	return FString::Printf(
-		TEXT("%s, %s, %s, %s, %s, %.2f, %.2f, %.2f, %.2f, %.2f, %.2f, %s") LINE_TERMINATOR,
-		*ActorDesc->GetGuid().ToString(),
-		*ActorDesc->GetNativeClass().ToString(),
-		*ActorDesc->GetActorName().ToString(),
-		*ActorDesc->GetActorPackage().ToString(),
-		*ActorDesc->GetActorLabel().ToString(),
-		ActorDesc->GetBounds().GetCenter().X,
-		ActorDesc->GetBounds().GetCenter().Y,
-		ActorDesc->GetBounds().GetCenter().Z,
-		ActorDesc->GetBounds().GetExtent().X,
-		ActorDesc->GetBounds().GetExtent().Y,
-		ActorDesc->GetBounds().GetExtent().Z,
+		TEXT("%s DataLayerNames:%s") LINE_TERMINATOR, 
+		*ActorDesc->ToString(FWorldPartitionActorDesc::EToStringMode::Full), 
 		*GetDataLayerString(ActorDesc->GetDataLayerInstanceNames())
 	);
 }
@@ -147,7 +136,6 @@ static FAutoConsoleCommand DumpActorDesc(
 					if (UWorldPartition* WorldPartition = World->GetWorldPartition())
 					{
 						TMap<FName, FString> DataLayersDumpString = GetDataLayersDumpString(WorldPartition);
-						UE_LOG(LogWorldPartition, Log, TEXT("Guid, Class, Name, Package, BVCenterX, BVCenterY, BVCenterZ, BVExtentX, BVExtentY, BVExtentZ, DataLayers"));
 						for (const FString& ActorPath : ActorPaths)
 						{
 							if (const FWorldPartitionActorDesc* ActorDesc = WorldPartition->GetActorDesc(ActorPath))
@@ -1345,9 +1333,6 @@ void UWorldPartition::DumpActorDescs(const FString& Path)
 {
 	if (FArchive* LogFile = IFileManager::Get().CreateFileWriter(*Path))
 	{
-		FString LineEntry = TEXT("Guid, Class, Name, Package, BVCenterX, BVCenterY, BVCenterZ, BVExtentX, BVExtentY, BVExtentZ, DataLayers") LINE_TERMINATOR;
-		LogFile->Serialize(TCHAR_TO_ANSI(*LineEntry), LineEntry.Len());
-
 		TArray<const FWorldPartitionActorDesc*> ActorDescs;
 		TMap<FName, FString> DataLayersDumpString = GetDataLayersDumpString(this);
 		for (UActorDescContainer::TConstIterator<> ActorDescIterator(this); ActorDescIterator; ++ActorDescIterator)
@@ -1360,7 +1345,7 @@ void UWorldPartition::DumpActorDescs(const FString& Path)
 		});
 		for (const FWorldPartitionActorDesc* ActorDescIterator : ActorDescs)
 		{
-			LineEntry = GetActorDescDumpString(ActorDescIterator, DataLayersDumpString);
+			FString LineEntry = GetActorDescDumpString(ActorDescIterator, DataLayersDumpString);
 			LogFile->Serialize(TCHAR_TO_ANSI(*LineEntry), LineEntry.Len());
 		}
 
diff --git a/Engine/Source/Runtime/Engine/Private/WorldPartition/WorldPartitionActorDesc.cpp b/Engine/Source/Runtime/Engine/Private/WorldPartition/WorldPartitionActorDesc.cpp
--- a/Engine/Source/Runtime/Engine/Private/WorldPartition/WorldPartitionActorDesc.cpp
+++ b/Engine/Source/Runtime/Engine/Private/WorldPartition/WorldPartitionActorDesc.cpp
@@ -233,29 +233,76 @@ void FWorldPartitionActorDesc::TransformInstance(const FString& From, const FStr
 	}
 }
 
-FString FWorldPartitionActorDesc::ToString() const
+FString FWorldPartitionActorDesc::ToString(EToStringMode Mode) const
 {
 	auto GetBoolStr = [](bool bValue) -> const TCHAR*
 	{
-		return bValue ? TEXT("true") : TEXT("false");
+		return bValue ? TEXT("1") : TEXT("0");
 	};
 
-	return FString::Printf(
-		TEXT("Guid:%s BaseClass:%s NativeClass:%s Name:%s Label:%s SpatiallyLoaded:%s Bounds:%s RuntimeGrid:%s EditorOnly:%s HLODRelevant:%s FolderPath:%s FolderGuid:%s Parent:%s"), 
-		*Guid.ToString(), 
-		*BaseClass.ToString(), 
-		*NativeClass.ToString(), 
-		*GetActorName().ToString(),
-		*GetActorLabel().ToString(),
-		GetBoolStr(bIsSpatiallyLoaded),
-		*GetBounds().ToString(),
-		*RuntimeGrid.ToString(),
-		GetBoolStr(bActorIsEditorOnly),
-		GetBoolStr(bActorIsHLODRelevant),
-		*FolderPath.ToString(),
-		*FolderGuid.ToString(),
-		*ParentActor.ToString()
-	);
+	TStringBuilder<1024> Result;
+	Result.Appendf(TEXT("Guid:%s"), *Guid.ToString());
+
+	if (Mode >= EToStringMode::Compact)
+	{
+		Result.Appendf(
+			TEXT(" BaseClass:%s NativeClass:%s Name:%s Label:%s SpatiallyLoaded:%s Bounds:%s RuntimeGrid:%s EditorOnly:%s HLODRelevant:%s"),
+			*BaseClass.ToString(), 
+			*NativeClass.ToString(), 
+			*GetActorName().ToString(),
+			*GetActorLabel().ToString(),
+			GetBoolStr(bIsSpatiallyLoaded),
+			*GetBounds().ToString(),
+			*RuntimeGrid.ToString(),
+			GetBoolStr(bActorIsEditorOnly),
+			GetBoolStr(bActorIsHLODRelevant)
+		);
+
+		if (ParentActor.IsValid())
+		{
+			Result.Appendf(TEXT(" Parent:%s"), *ParentActor.ToString());
+		}
+
+		if (!HLODLayer.IsNone())
+		{
+			Result.Appendf(TEXT(" HLODLayer:%s"), *HLODLayer.ToString());
+		}
+
+		if (!FolderPath.IsNone())
+		{
+			Result.Appendf(TEXT(" FolderPath:%s"), *FolderPath.ToString());
+		}
+
+		if (FolderGuid.IsValid())
+		{
+			Result.Appendf(TEXT(" FolderGuid:%s"), *FolderGuid.ToString());
+		}
+
+		if (Mode >= EToStringMode::Full)
+		{
+			if (References.Num())
+			{
+				Result.Appendf(TEXT(" References:%s"), *FString::JoinBy(References, TEXT(","), [&](const FGuid& ReferenceGuid) { return ReferenceGuid.ToString(); }));
+			}
+
+			if (Tags.Num())
+			{
+				Result.Appendf(TEXT(" Tags:%s"), *FString::JoinBy(Tags, TEXT(","), [&](const FName& TagName) { return TagName.ToString(); }));
+			}
+
+			if (Properties.Num())
+			{
+				Result.Appendf(TEXT(" Properties:%s"), *Properties.ToString());
+			}
+
+			if (DataLayers.Num())
+			{
+				Result.Appendf(TEXT(" DataLayers:%s"), *FString::JoinBy(DataLayers, TEXT(","), [&](const FName& DataLayerName) { return DataLayerName.ToString(); }));
+			}
+		}
+	}
+
+	return Result.ToString();
 }
 
 void FWorldPartitionActorDesc::Serialize(FArchive& Ar)
diff --git a/Engine/Source/Runtime/Engine/Public/OverrideVoidReturnInvoker.h b/Engine/Source/Runtime/Engine/Public/OverrideVoidReturnInvoker.h
new file mode 100644
--- /dev/null
+++ b/Engine/Source/Runtime/Engine/Public/OverrideVoidReturnInvoker.h
@@ -0,0 +1,38 @@
+// Copyright Epic Games, Inc. All Rights Reserved.
+
+#pragma once
+
+#include "CoreTypes.h"
+
+/**
+ * Wraps invocation of a function that can return a value or not. If it doesn't, the function call is wrapped into
+ * a function that will invoke the provided function and return the provided default value.
+ */
+template <typename ReturnType, typename FuncType>
+struct TOverrideVoidReturnInvoker
+{
+	TOverrideVoidReturnInvoker(ReturnType InDefaultValue, FuncType InFunction)
+		: DefaultValue(InDefaultValue)
+		, Function(InFunction)
+	{}
+
+	template <typename... ArgTypes>
+	bool operator()(ArgTypes&&... Args)
+	{
+		using FuncResultType = std::invoke_result_t<FuncType, ArgTypes...>;
+		static_assert(std::disjunction_v<std::is_same<FuncResultType, void>, std::is_same<FuncResultType, ReturnType>>, "Passed in function should either return void or ReturnType");
+
+		if constexpr (std::is_same_v<FuncResultType, void>)
+		{
+			Invoke(Function, Forward<ArgTypes>(Args)...);
+			return DefaultValue;
+		}
+		else
+		{
+			return Invoke(Function, Forward<ArgTypes>(Args)...);
+		}
+	}
+
+	ReturnType DefaultValue;
+	FuncType Function;
+};
\ No newline at end of file
diff --git a/Engine/Source/Runtime/Engine/Public/PropertyPairsMap.h b/Engine/Source/Runtime/Engine/Public/PropertyPairsMap.h
--- a/Engine/Source/Runtime/Engine/Public/PropertyPairsMap.h
+++ b/Engine/Source/Runtime/Engine/Public/PropertyPairsMap.h
@@ -3,6 +3,8 @@
 #pragma once
 
 #include "CoreTypes.h"
+#include "Misc/StringBuilder.h"
+#include "OverrideVoidReturnInvoker.h"
 
 /** A simple name/value pairs map */
 class FPropertyPairsMap
@@ -90,6 +92,25 @@ public:
 		return InAr << InPropertyPairsMap.Properties;
 	}
 
+	/**
+	 * Iterate through the property map and invoke a functor that can optionally break iteration by returning false.
+	 *
+	 * @param InFunc Function to invoke on the properties of the map.
+	 */
+	template <class Func>
+	FORCEINLINE void ForEachProperty(Func InFunc)
+	{
+		TOverrideVoidReturnInvoker Invoker(true, InFunc);
+
+		for (auto [Name, Value] : Properties)
+		{
+			if (!Invoker(Name, Value))
+			{
+				return;
+			}
+		}
+	}
+
 	/** @returns True if the property map is empty. */
 	FORCEINLINE bool IsEmpty() const
 	{
@@ -102,6 +123,41 @@ public:
 		return Properties.Num();
 	}
 
+	/**
+	 * Converts this property map to a string representation.
+	 *
+	 * @return The string representation.
+	 */
+	FORCEINLINE FString ToString() const
+	{
+		TStringBuilder<1024> StringBuilder;
+		
+		StringBuilder.AppendChar(TEXT('('));
+
+		if (Properties.Num())
+		{
+			for (auto [Name, Value] : Properties)
+			{
+				StringBuilder.AppendChar(TEXT('('));
+				StringBuilder.Append(Name.ToString());
+
+				if (!Value.IsNone())
+				{
+					StringBuilder.AppendChar(TEXT(','));
+					StringBuilder.Append(Value.ToString());
+				}
+
+				StringBuilder.AppendChar(TEXT(')'));
+				StringBuilder.AppendChar(TEXT(','));
+			}
+			StringBuilder.RemoveSuffix(1);
+		}
+
+		StringBuilder.AppendChar(TEXT(')'));
+
+		return StringBuilder.ToString();
+	}
+
 protected:
 	TMap<FName, FName> Properties;
 };
\ No newline at end of file
diff --git a/Engine/Source/Runtime/Engine/Public/WorldPartition/WorldPartitionActorDesc.h b/Engine/Source/Runtime/Engine/Public/WorldPartition/WorldPartitionActorDesc.h
--- a/Engine/Source/Runtime/Engine/Public/WorldPartition/WorldPartitionActorDesc.h
+++ b/Engine/Source/Runtime/Engine/Public/WorldPartition/WorldPartitionActorDesc.h
@@ -175,7 +175,14 @@ public:
 		Container = InContainer;
 	}
 
-	FString ToString() const;
+	enum class EToStringMode : uint8
+	{
+		Guid,
+		Compact,
+		Full
+	};
+
+	FString ToString(EToStringMode Mode = EToStringMode::Compact) const;
 
 	bool IsLoaded(bool bEvenIfPendingKill=false) const;
 	AActor* GetActor(bool bEvenIfPendingKill=true, bool bEvenIfUnreachable=false) const;
diff --git a/Engine/Source/Runtime/Engine/Public/WorldPartition/WorldPartitionStreamingGeneration.h b/Engine/Source/Runtime/Engine/Public/WorldPartition/WorldPartitionStreamingGeneration.h
--- a/Engine/Source/Runtime/Engine/Public/WorldPartition/WorldPartitionStreamingGeneration.h
+++ b/Engine/Source/Runtime/Engine/Public/WorldPartition/WorldPartitionStreamingGeneration.h
@@ -3,8 +3,7 @@
 
 #if WITH_EDITOR
 #include "CoreMinimal.h"
-#include "Templates/IntegralConstant.h"
-#include "Templates/UnrealTypeTraits.h"
+#include "OverrideVoidReturnInvoker.h"
 #include "WorldPartition/WorldPartitionActorDescView.h"
 
 class ENGINE_API FActorDescViewMap
@@ -12,15 +11,14 @@ class ENGINE_API FActorDescViewMap
 	friend class FWorldPartitionStreamingGenerator;
 
 private:
-    template<typename F, typename V>bool DoCall(F Func, V Val, std::true_type) { return Func(*Val); }
-	template<typename F, typename V> bool DoCall(F Func, V Val, std::false_type) { Func(*Val); return true; }
-
-	template <class T>
-	void ForEachActorDescView(T Func)
+	template <class Func>
+	void ForEachActorDescView(Func InFunc)
 	{
+		TOverrideVoidReturnInvoker Invoker(true, InFunc);
+
 		for (TUniquePtr<FWorldPartitionActorDescView>& ActorDescView : ActorDescViewList)
 		{
-			if (!DoCall(Func, ActorDescView.Get(), std::is_same<bool, decltype(Func(*ActorDescView))>()))
+			if (!Invoker(*ActorDescView))
 			{
 				return;
 			}
@@ -36,6 +34,11 @@ private:
 		return nullptr;
 	}
 
+	FWorldPartitionActorDescView& FindByGuidChecked(const FGuid& InGuid)
+	{
+		return *ActorDescViewsByGuid.FindChecked(InGuid);
+	}
+
 public:
 	FActorDescViewMap();
 
@@ -47,12 +50,19 @@ public:
 
 	FWorldPartitionActorDescView* Emplace(const FGuid& InActorGuid, const FWorldPartitionActorDescView& InActorDescView);
 
-	template <class T>
-	void ForEachActorDescView(T Func) const
+	FORCEINLINE int32 Num() const
 	{
+		return ActorDescViewList.Num();
+	}
+
+	template <class Func>
+	void ForEachActorDescView(Func InFunc) const
+	{
+		TOverrideVoidReturnInvoker Invoker(true, InFunc);
+
 		for (const TUniquePtr<FWorldPartitionActorDescView>& ActorDescView : ActorDescViewList)
 		{
-			if (!DoCall(Func, ActorDescView.Get(), std::is_same<bool, decltype(Func(*ActorDescView))>()))
+			if (!Invoker(*ActorDescView))
 			{
 				return;
 			}
@@ -68,10 +78,10 @@ public:
 		return nullptr;
 	}
 
-	template <class T>
+	template <class ClassType>
 	TArray<const FWorldPartitionActorDescView*> FindByExactNativeClass() const
 	{
-		return FindByExactNativeClass(T::StaticClass());
+		return FindByExactNativeClass(ClassType::StaticClass());
 	}
 
 	TArray<const FWorldPartitionActorDescView*> FindByExactNativeClass(UClass* InExactNativeClass) const;
