Fix default texture placeholder type and allow VT upload cache update on Mac

#rnx

diff --git a/Engine/Source/Runtime/Engine/Engine.Build.cs b/Engine/Source/Runtime/Engine/Engine.Build.cs
--- a/Engine/Source/Runtime/Engine/Engine.Build.cs
+++ b/Engine/Source/Runtime/Engine/Engine.Build.cs
@@ -389,6 +389,9 @@ public class Engine : ModuleRules
 
 		if (Target.Platform == UnrealTargetPlatform.Mac)
 		{
+			// Allow VirtualTextureUploadCache to use the UpdateTexture path.
+			PublicDefinitions.Add("ALLOW_UPDATE_TEXTURE=1");
+
 			PublicFrameworks.AddRange(new string[] { "AVFoundation", "CoreVideo", "CoreMedia" });
 		}
 
diff --git a/Engine/Source/Runtime/Engine/Private/SceneManagement.cpp b/Engine/Source/Runtime/Engine/Private/SceneManagement.cpp
--- a/Engine/Source/Runtime/Engine/Private/SceneManagement.cpp
+++ b/Engine/Source/Runtime/Engine/Private/SceneManagement.cpp
@@ -969,8 +969,8 @@ void GetLightmapClusterResourceParameters(
 	Parameters.VTSkyOcclusionTexture = GWhiteTextureWithSRV->ShaderResourceViewRHI;
 	Parameters.VTAOMaterialMaskTexture = GBlackTextureWithSRV->ShaderResourceViewRHI;
 	Parameters.VTStaticShadowTexture = GWhiteTextureWithSRV->ShaderResourceViewRHI;
-	Parameters.LightmapVirtualTexturePageTable0 = GBlackTexture->TextureRHI;
-	Parameters.LightmapVirtualTexturePageTable1 = GBlackTexture->TextureRHI;
+	Parameters.LightmapVirtualTexturePageTable0 = GBlackUintTexture->TextureRHI;
+	Parameters.LightmapVirtualTexturePageTable1 = GBlackUintTexture->TextureRHI;
 	Parameters.LightMapSampler = GBlackTexture->SamplerStateRHI;
 	Parameters.LightMapSampler_1 = GBlackTexture->SamplerStateRHI;
 	Parameters.SkyOcclusionSampler = GWhiteTexture->SamplerStateRHI;
@@ -1050,8 +1050,8 @@ void GetLightmapClusterResourceParameters(
 		Parameters.StaticShadowTexture = Input.ShadowMapTexture ? Input.ShadowMapTexture->TextureReference.TextureReferenceRHI.GetReference() : GWhiteTexture->TextureRHI;
 		Parameters.StaticShadowTextureSampler = GetTextureSamplerState(Input.ShadowMapTexture, GWhiteTexture->SamplerStateRHI);
 
-		Parameters.LightmapVirtualTexturePageTable0 = GBlackTexture->TextureRHI;
-		Parameters.LightmapVirtualTexturePageTable1 = GBlackTexture->TextureRHI;
+		Parameters.LightmapVirtualTexturePageTable0 = GBlackUintTexture->TextureRHI;
+		Parameters.LightmapVirtualTexturePageTable1 = GBlackUintTexture->TextureRHI;
 	}
 }
 
