Store lossy compression setting for RVT streaming mips in the component, not the RVT asset.
This makes it possible to set per volume.

diff --git a/Engine/Source/Runtime/Engine/Classes/Components/RuntimeVirtualTextureComponent.h b/Engine/Source/Runtime/Engine/Classes/Components/RuntimeVirtualTextureComponent.h
--- a/Engine/Source/Runtime/Engine/Classes/Components/RuntimeVirtualTextureComponent.h
+++ b/Engine/Source/Runtime/Engine/Classes/Components/RuntimeVirtualTextureComponent.h
@@ -3,8 +3,8 @@
 #pragma once
 
 #include "CoreMinimal.h"
-#include "SceneComponent.h"
 #include "Engine/TextureDefines.h"
+#include "SceneComponent.h"
 #include "RuntimeVirtualTextureComponent.generated.h"
 
 class URuntimeVirtualTexture;
@@ -58,6 +58,14 @@ protected:
 	UPROPERTY(VisibleAnywhere, Transient, Category = VirtualTextureBuild)
 	bool bBuildStreamingMipsButton;
 
+	/** 
+	 * How aggressively should any relevant lossy compression be applied. 
+	 * For compressors that support EncodeSpeed (i.e. Oodle), this is only applied if enabled (see Project Settings -> Texture Encoding). 
+	 * Note that this is in addition to any unavoidable loss due to the target format. Selecting "No Lossy Compression" will not result in zero distortion for BCn formats.
+	*/
+	UPROPERTY(EditAnywhere, BlueprintReadWrite, AdvancedDisplay, Category = VirtualTextureBuild)
+	TEnumAsByte<ETextureLossyCompressionAmount> LossyCompressionAmount = TLCA_Default;
+
 	/** Use streaming low mips when rendering in editor. Set true to view and debug the baked streaming low mips. */
 	UPROPERTY(EditAnywhere, AdvancedDisplay, Category = VirtualTextureBuild, meta = (DisplayName = "View in Editor"))
 	bool bUseStreamingLowMipsInEditor = false;
@@ -113,6 +121,9 @@ public:
 	/** Public getter for debug streaming mips flag. */
 	bool IsBuildDebugStreamingMips() { return bBuildDebugStreamingMips; }
 
+	/** Public getter for lossy compression setting. */
+	TEnumAsByte<ETextureLossyCompressionAmount> GetLossyCompressionAmount() const { return LossyCompressionAmount; }
+
 	/** Returns true if there are StreamingTexure contents but they are not valid for use. */
 	bool IsStreamingTextureInvalid() const;
 
diff --git a/Engine/Source/Runtime/Engine/Classes/VT/RuntimeVirtualTexture.h b/Engine/Source/Runtime/Engine/Classes/VT/RuntimeVirtualTexture.h
--- a/Engine/Source/Runtime/Engine/Classes/VT/RuntimeVirtualTexture.h
+++ b/Engine/Source/Runtime/Engine/Classes/VT/RuntimeVirtualTexture.h
@@ -4,7 +4,6 @@
 
 #include "CoreMinimal.h"
 #include "Engine/TextureDefines.h"
-#include "UObject/ObjectMacros.h"
 #include "VirtualTexturing.h"
 #include "VT/RuntimeVirtualTextureEnum.h"
 #include "RuntimeVirtualTexture.generated.h"
@@ -76,13 +75,6 @@ protected:
 	UPROPERTY(EditAnywhere, BlueprintReadWrite, Category = LevelOfDetail, meta = (DisplayName = "Texture Group"), AssetRegistrySearchable)
 	TEnumAsByte<enum TextureGroup> LODGroup = TEXTUREGROUP_World;
 	
-	/** How aggressively should any relevant lossy compression be applied. For compressors that support EncodeSpeed (i.e. Oodle), this is only
-	*	applied if enabled (see Project Settings -> Texture Encoding). Note that this is *in addition* to any
-	*	unavoidable loss due to the target format - selecting "No Lossy Compression" will not result in zero distortion for BCn formats.
-	*/
-	UPROPERTY(EditAnywhere, BlueprintReadWrite, Category = LevelOfDetail)
-	TEnumAsByte<ETextureLossyCompressionAmount> LossyCompressionAmount = TLCA_Default;
-
 	/** Deprecated size of virtual texture. */
 	UPROPERTY()
 	int32 Size_DEPRECATED = -1;
@@ -129,11 +121,10 @@ public:
 	bool GetContinuousUpdate() const { return bContinuousUpdate; }
 	/** Public getter for virtual texture removed low mips */
 	int32 GetRemoveLowMips() const { return FMath::Clamp(RemoveLowMips, 0, 5); }
-
+	/** Public getter for virtual texture using low quality compression flag. */
 	bool GetLQCompression() const { return bUseLowQualityCompression; }
 	/** Public getter for texture LOD Group */
 	TEnumAsByte<enum TextureGroup> GetLODGroup() const { return LODGroup; }
-	TEnumAsByte<ETextureLossyCompressionAmount> GetLossyCompressionAmount() const { return LossyCompressionAmount; }
 
 	/** Structure containing additional settings for initializing the producer. */
 	struct FInitSettings
diff --git a/Engine/Source/Runtime/Engine/Private/Components/RuntimeVirtualTextureComponent.cpp b/Engine/Source/Runtime/Engine/Private/Components/RuntimeVirtualTextureComponent.cpp
--- a/Engine/Source/Runtime/Engine/Private/Components/RuntimeVirtualTextureComponent.cpp
+++ b/Engine/Source/Runtime/Engine/Private/Components/RuntimeVirtualTextureComponent.cpp
@@ -166,7 +166,7 @@ uint64 URuntimeVirtualTextureComponent::CalculateStreamingTextureSettingsHash()
 	Settings.ContinuousUpdate = (uint32)VirtualTexture->GetContinuousUpdate();
 	Settings.SinglePhysicalSpace = (uint32)VirtualTexture->GetSinglePhysicalSpace();
 	Settings.bUseLowQualityCompression = (uint32)VirtualTexture->GetLQCompression();
-	Settings.LossyCompressionAmount = (uint32) VirtualTexture->GetLossyCompressionAmount();
+	Settings.LossyCompressionAmount = (uint32)GetLossyCompressionAmount();
 
 	return Settings.PackedValue;
 }
@@ -246,7 +246,7 @@ void URuntimeVirtualTextureComponent::InitializeStreamingTexture(uint32 InSizeX,
 		BuildDesc.TileSize = VirtualTexture->GetTileSize();
 		BuildDesc.TileBorderSize = VirtualTexture->GetTileBorderSize();
 		BuildDesc.LODGroup = VirtualTexture->GetLODGroup();
-		BuildDesc.LossyCompressionAmount = VirtualTexture->GetLossyCompressionAmount();
+		BuildDesc.LossyCompressionAmount = GetLossyCompressionAmount();
 
 		BuildDesc.LayerCount = VirtualTexture->GetLayerCount();
 		check(BuildDesc.LayerCount <= RuntimeVirtualTexture::MaxTextureLayers);
