Servers should ignore NMT_DestructionInfo

diff --git a/Engine/Source/Runtime/Engine/Private/DataChannel.cpp b/Engine/Source/Runtime/Engine/Private/DataChannel.cpp
--- a/Engine/Source/Runtime/Engine/Private/DataChannel.cpp
+++ b/Engine/Source/Runtime/Engine/Private/DataChannel.cpp
@@ -1713,7 +1713,7 @@ void UControlChannel::ReceivedBunch( FInBunch& Bunch )
 							// but the client is reporting an actor channel failure that isn't the PlayerController
 							else
 							{
-								//UE_LOG(LogNet, Warning, TEXT("UControlChannel::RecievedBunch: PlayerController doesn't exist for the client, but the client is reporting an actor channel failure that isn't the PlayerController."));
+								//UE_LOG(LogNet, Warning, TEXT("UControlChannel::ReceivedBunch: PlayerController doesn't exist for the client, but the client is reporting an actor channel failure that isn't the PlayerController."));
 							}
 						}
 					}
@@ -1722,7 +1722,7 @@ void UControlChannel::ReceivedBunch( FInBunch& Bunch )
 					// @PotentialDOSAttackDetection
 					else
 					{
-						UE_LOG(LogNet, Warning, TEXT("UControlChannel::RecievedBunch: The client is sending an actor channel failure message with an invalid actor channel index."));
+						UE_LOG(LogNet, Warning, TEXT("UControlChannel::ReceivedBunch: The client is sending an actor channel failure message with an invalid actor channel index."));
 					}
 				}
 			}
@@ -1759,7 +1759,15 @@ void UControlChannel::ReceivedBunch( FInBunch& Bunch )
 		}
 		else if (MessageType == NMT_DestructionInfo)
 		{
-			ReceiveDestructionInfo(Bunch);
+			if (!Connection->Driver->IsServer())
+			{
+				ReceiveDestructionInfo(Bunch);
+			}
+			else
+			{
+				UE_LOG(LogNet, Warning, TEXT("UControlChannel::ReceivedBunch: Server received a NMT_DestructionInfo which is not supported. Closing connection."));
+				Connection->Close(ENetCloseResult::ClientSentDestructionInfo);
+			}
 		}
 		else if (MessageType == NMT_CloseReason)
 		{
@@ -2059,6 +2067,7 @@ int64 UControlChannel::SendDestructionInfo(FActorDestructionInfo* DestructionInf
 void UControlChannel::ReceiveDestructionInfo(FInBunch& Bunch)
 {
 	checkf(Connection && Connection->PackageMap && Connection->Driver, TEXT("UControlChannel::ReceiveDestructionInfo requires a valid connection, package map, and driver: %s"), *Describe());
+	checkf(!Connection->Driver->IsServer(), TEXT("UControlChannel::ReceiveDestructionInfo should not be called on a server: %s"), *Describe());
 
 	EChannelCloseReason CloseReason = EChannelCloseReason::Destroyed;
 	Bunch << CloseReason;
@@ -2070,63 +2079,59 @@ void UControlChannel::ReceiveDestructionInfo(FInBunch& Bunch)
 	{
 		if (AActor* TheActor = Cast<AActor>(Object))
 		{
-			// If we're the client, destroy this actor.
-			if (!Connection->Driver->IsServer())
-			{
-				checkf(TheActor->IsValidLowLevel(), TEXT("ReceiveDestructionInfo serialized an invalid actor: %s"), *Describe());
-				checkSlow(Connection->IsValidLowLevel());
-				checkSlow(Connection->Driver->IsValidLowLevel());
+			checkf(TheActor->IsValidLowLevel(), TEXT("ReceiveDestructionInfo serialized an invalid actor: %s"), *Describe());
+			checkSlow(Connection->IsValidLowLevel());
+			checkSlow(Connection->Driver->IsValidLowLevel());
 
-				if (TheActor->GetTearOff() && !Connection->Driver->ShouldClientDestroyTearOffActors())
+			if (TheActor->GetTearOff() && !Connection->Driver->ShouldClientDestroyTearOffActors())
+			{
+				if (!bTornOff)
 				{
-					if (!bTornOff)
+					TheActor->SetRole(ROLE_Authority);
+					TheActor->SetReplicates(false);
+					bTornOff = true;
+					if (TheActor->GetWorld() != nullptr && !IsEngineExitRequested())
 					{
-						TheActor->SetRole(ROLE_Authority);
-						TheActor->SetReplicates(false);
-						bTornOff = true;
-						if (TheActor->GetWorld() != nullptr && !IsEngineExitRequested())
-						{
-							TheActor->TornOff();
-						}
-
-						Connection->Driver->NotifyActorTornOff(TheActor);
+						TheActor->TornOff();
 					}
+
+					Connection->Driver->NotifyActorTornOff(TheActor);
 				}
-				else if (Dormant && (CloseReason == EChannelCloseReason::Dormancy) && !TheActor->GetTearOff())
+			}
+			else if (Dormant && (CloseReason == EChannelCloseReason::Dormancy) && !TheActor->GetTearOff())
+			{
+				Connection->Driver->ClientSetActorDormant(TheActor);
+				Connection->Driver->NotifyActorFullyDormantForConnection(TheActor, Connection);
+			}
+			else if (!TheActor->bNetTemporary && TheActor->GetWorld() != nullptr && !IsEngineExitRequested() && Connection->Driver->ShouldClientDestroyActor(TheActor))
+			{
+				// Destroy the actor
+
+				// Unmap any components in this actor. This will make sure that once the Actor is remapped
+				// any references to components will be remapped as well.
+				for (UActorComponent* Component : TheActor->GetComponents())
 				{
-					Connection->Driver->ClientSetActorDormant(TheActor);
-					Connection->Driver->NotifyActorFullyDormantForConnection(TheActor, Connection);
+					Connection->Driver->MoveMappedObjectToUnmapped(Component);
 				}
-				else if (!TheActor->bNetTemporary && TheActor->GetWorld() != nullptr && !IsEngineExitRequested() && Connection->Driver->ShouldClientDestroyActor(TheActor))
-				{
-					// Destroy the actor
-
-					// Unmap any components in this actor. This will make sure that once the Actor is remapped
-					// any references to components will be remapped as well.
-					for (UActorComponent* Component : TheActor->GetComponents())
-					{
-						Connection->Driver->MoveMappedObjectToUnmapped(Component);
-					}
 
-					// Unmap this object so we can remap it if it becomes relevant again in the future
-					Connection->Driver->MoveMappedObjectToUnmapped(TheActor);
+				// Unmap this object so we can remap it if it becomes relevant again in the future
+				Connection->Driver->MoveMappedObjectToUnmapped(TheActor);
 
-					TheActor->PreDestroyFromReplication();
-					TheActor->Destroy(true);
+				TheActor->PreDestroyFromReplication();
+				TheActor->Destroy(true);
 
-					if (UE::Net::FilterGuidRemapping != 0)
+				if (UE::Net::FilterGuidRemapping != 0)
+				{
+					// Remove this actor's NetGUID from the list of unmapped values, it will be added back if it replicates again
+					if (NetGUID.IsValid() && Connection != nullptr && Connection->Driver != nullptr && Connection->Driver->GuidCache.IsValid())
 					{
-						// Remove this actor's NetGUID from the list of unmapped values, it will be added back if it replicates again
-						if (NetGUID.IsValid() && Connection != nullptr && Connection->Driver != nullptr && Connection->Driver->GuidCache.IsValid())
-						{
-							Connection->Driver->GuidCache->ImportedNetGuids.Remove(NetGUID);
-						}
+						Connection->Driver->GuidCache->ImportedNetGuids.Remove(NetGUID);
 					}
+				}
 
-					if (UPackageMapClient* PackageMapClient = Cast<UPackageMapClient>(Connection->PackageMap))
-					{
-						PackageMapClient->SetHasQueuedBunches(NetGUID, false);
-					}
+				if (UPackageMapClient* PackageMapClient = Cast<UPackageMapClient>(Connection->PackageMap))
+				{
+					PackageMapClient->SetHasQueuedBunches(NetGUID, false);
 				}
 			}
 		}
diff --git a/Engine/Source/Runtime/Net/Core/Public/Net/Core/Connection/NetCloseResult.h b/Engine/Source/Runtime/Net/Core/Public/Net/Core/Connection/NetCloseResult.h
--- a/Engine/Source/Runtime/Net/Core/Public/Net/Core/Connection/NetCloseResult.h
+++ b/Engine/Source/Runtime/Net/Core/Public/Net/Core/Connection/NetCloseResult.h
@@ -260,6 +260,9 @@ enum class ENetCloseResult : uint32
 	/** Client tried to set bHasMustBeMappedGUIDs */
 	ClientHasMustBeMappedGUIDs,
 
+	/** Client tried to send a destruction info */
+	ClientSentDestructionInfo,
+
 	/** Received MustBeMappedGUID that isn't registered */
 	UnregisteredMustBeMappedGUID,
 
