Additional safety check for max net exports per group

diff --git a/Engine/Source/Runtime/Engine/Classes/Engine/PackageMapClient.h b/Engine/Source/Runtime/Engine/Classes/Engine/PackageMapClient.h
--- a/Engine/Source/Runtime/Engine/Classes/Engine/PackageMapClient.h
+++ b/Engine/Source/Runtime/Engine/Classes/Engine/PackageMapClient.h
@@ -54,47 +54,7 @@ public:
 	{
 	}
 
-	friend FArchive& operator<<( FArchive& Ar, FNetFieldExport& C )
-	{
-		uint8 Flags = C.bExported ? 1 : 0;
-
-		Ar << Flags;
-
-		if ( Ar.IsLoading() )
-		{
-			C.bExported = (Flags == 1);
-		}
-
-		if ( C.bExported )
-		{
-			Ar.SerializeIntPacked( C.Handle );
-			Ar << C.CompatibleChecksum;
-
-			if (Ar.IsLoading() && Ar.EngineNetVer() < HISTORY_NETEXPORT_SERIALIZATION)
-			{
-				FName TempName;
-				FString TempType;
-
-				Ar << TempName;
-				Ar << TempType;
-
-				C.ExportName = TempName;
-			}
-			else
-			{
-				if (Ar.IsLoading() && Ar.EngineNetVer() < HISTORY_NETEXPORT_SERIALIZE_FIX)
-				{
-					Ar << C.ExportName;
-				}
-				else
-				{
-					UPackageMap::StaticSerializeName(Ar, C.ExportName);
-				}
-			}
-		}
-
-		return Ar;
-	}
+	friend FArchive& operator<<(FArchive& Ar, FNetFieldExport& C);
 
 	void CountBytes(FArchive& Ar) const;
 
@@ -118,39 +78,19 @@ public:
 	TArray< FNetFieldExport >	NetFieldExports;
 	bool						bDirtyForReplay;
 
-	friend FArchive& operator<<( FArchive& Ar, FNetFieldExportGroup& C )
-	{
-		Ar << C.PathName;
-
-		Ar.SerializeIntPacked( C.PathNameIndex );
-
-		uint32 NumNetFieldExports = C.NetFieldExports.Num();
-		Ar.SerializeIntPacked( NumNetFieldExports );
-
-		if ( Ar.IsLoading() )
-		{
-			C.NetFieldExports.AddDefaulted( ( int32 )NumNetFieldExports );
-		}
-
-		for ( int32 i = 0; i < C.NetFieldExports.Num(); i++ )
-		{
-			Ar << C.NetFieldExports[i];
-		}
-
-		return Ar;
-	}
+	friend FArchive& operator<<(FArchive& Ar, FNetFieldExportGroup& C);
 
-	int32 FindNetFieldExportHandleByChecksum( const uint32 Checksum )
+	int32 FindNetFieldExportHandleByChecksum(const uint32 Checksum) const
 	{
-		for ( int32 i = 0; i < NetFieldExports.Num(); i++ )
+		for (int32 i = 0; i < NetFieldExports.Num(); i++)
 		{
-			if ( NetFieldExports[i].CompatibleChecksum == Checksum )
+			if (NetFieldExports[i].CompatibleChecksum == Checksum)
 			{
 				return i;
 			}
 		}
 
-		return -1;
+		return INDEX_NONE;
 	}
 
 	void CountBytes(FArchive& Ar) const;
diff --git a/Engine/Source/Runtime/Engine/Private/PackageMapClient.cpp b/Engine/Source/Runtime/Engine/Private/PackageMapClient.cpp
--- a/Engine/Source/Runtime/Engine/Private/PackageMapClient.cpp
+++ b/Engine/Source/Runtime/Engine/Private/PackageMapClient.cpp
@@ -3845,6 +3845,48 @@ void FNetFieldExport::CountBytes(FArchive& Ar) const
 
 }
 
+FArchive& operator<<(FArchive& Ar, FNetFieldExport& C)
+{
+	uint8 Flags = C.bExported ? 1 : 0;
+
+	Ar << Flags;
+
+	if (Ar.IsLoading())
+	{
+		C.bExported = (Flags == 1);
+	}
+
+	if (C.bExported)
+	{
+		Ar.SerializeIntPacked(C.Handle);
+		Ar << C.CompatibleChecksum;
+
+		if (Ar.IsLoading() && Ar.EngineNetVer() < HISTORY_NETEXPORT_SERIALIZATION)
+		{
+			FName TempName;
+			FString TempType;
+
+			Ar << TempName;
+			Ar << TempType;
+
+			C.ExportName = TempName;
+		}
+		else
+		{
+			if (Ar.IsLoading() && Ar.EngineNetVer() < HISTORY_NETEXPORT_SERIALIZE_FIX)
+			{
+				Ar << C.ExportName;
+			}
+			else
+			{
+				UPackageMap::StaticSerializeName(Ar, C.ExportName);
+			}
+		}
+	}
+
+	return Ar;
+}
+
 void FNetFieldExportGroup::CountBytes(FArchive& Ar) const
 {
 	PathName.CountBytes(Ar);
@@ -3855,6 +3897,35 @@ void FNetFieldExportGroup::CountBytes(FArchive& Ar) const
 	}
 }
 
+FArchive& operator<<(FArchive& Ar, FNetFieldExportGroup& C)
+{
+	Ar << C.PathName;
+
+	Ar.SerializeIntPacked(C.PathNameIndex);
+
+	uint32 NumNetFieldExports = C.NetFieldExports.Num();
+	Ar.SerializeIntPacked(NumNetFieldExports);
+
+	if (Ar.IsLoading())
+	{
+		if (NumNetFieldExports > (uint32)UE::Net::MaxSerializedNetExportsPerGroup)
+		{
+			UE_LOG(LogNetPackageMap, Error, TEXT("FNetFieldExportGroup - NumNetFieldExports exceeds MaxSerializedNetExportsPerGroup (%d / %d)"), NumNetFieldExports, UE::Net::MaxSerializedNetExportsPerGroup);
+			Ar.SetError();
+			return Ar;
+		}
+
+		C.NetFieldExports.AddDefaulted((int32)NumNetFieldExports);
+	}
+
+	for (int32 i = 0; i < C.NetFieldExports.Num(); i++)
+	{
+		Ar << C.NetFieldExports[i];
+	}
+
+	return Ar;
+}
+
 void FPackageMapAckState::CountBytes(FArchive& Ar) const
 {
 	NetGUIDAckStatus.CountBytes(Ar);
