reduced mem usage for DispatchCanvasRenderCommands task by capturing only required data instead of entire canvas

diff --git a/Engine/Source/Editor/UnrealEd/Private/Lightmass/LightmassLandscapeRender.cpp b/Engine/Source/Editor/UnrealEd/Private/Lightmass/LightmassLandscapeRender.cpp
--- a/Engine/Source/Editor/UnrealEd/Private/Lightmass/LightmassLandscapeRender.cpp
+++ b/Engine/Source/Editor/UnrealEd/Private/Lightmass/LightmassLandscapeRender.cpp
@@ -135,7 +135,7 @@ void RenderLandscapeMaterialForLightmass(const FLandscapeStaticLightingMesh* Lan
 				DrawRenderState.SetDepthStencilState(TStaticDepthStencilState<false, CF_Always>::GetRHI());
 
 				//SCOPED_DRAW_EVENT(RHICmdList, RenderLandscapeMaterialToTexture);
-				FCanvasRenderContext RenderContext(GraphBuilder, RenderTarget->GetRenderTargetTexture(GraphBuilder), RTViewRect, FIntRect(0, 0, 0, 0));
+				FCanvasRenderContext RenderContext(GraphBuilder, RenderTarget, RTViewRect, FIntRect(0, 0, 0, 0), false);
 				GetRendererModule().DrawTileMesh(RenderContext, DrawRenderState, View, Mesh, false, FHitProxyId());
 
 				GraphBuilder.Execute();
diff --git a/Engine/Source/Runtime/Engine/Private/UserInterface/Canvas.cpp b/Engine/Source/Runtime/Engine/Private/UserInterface/Canvas.cpp
--- a/Engine/Source/Runtime/Engine/Private/UserInterface/Canvas.cpp
+++ b/Engine/Source/Runtime/Engine/Private/UserInterface/Canvas.cpp
@@ -545,25 +545,18 @@ bool FCanvasBatchedElementRenderItem::Render_GameThread(const FCanvas* Canvas, F
 	return bDirty;
 }
 
-FCanvasRenderContext::FCanvasRenderContext(FRDGBuilder& InGraphBuilder, const FCanvas& Canvas)
+FCanvasRenderContext::FCanvasRenderContext(FRDGBuilder& InGraphBuilder, const FRenderTarget* InRenderTarget, FIntRect InViewportRect, FIntRect InScissorRect, bool bScaledToRenderTarget)
 	: GraphBuilder(InGraphBuilder)
-	, RenderTarget(Canvas.RenderTarget->GetRenderTargetTexture(InGraphBuilder))
-	, ViewportRect(Canvas.ViewRect)
-	, ScissorRect(Canvas.ScissorRect)
+	, RenderTarget(InRenderTarget->GetRenderTargetTexture(InGraphBuilder))
+	, ViewportRect(InViewportRect)
+	, ScissorRect(InScissorRect)
 {
-	if (ViewportRect.Area() <= 0 || Canvas.bScaledToRenderTarget)
+	if (ViewportRect.Area() <= 0 || bScaledToRenderTarget)
 	{
 		ViewportRect = FIntRect(FIntPoint::ZeroValue, RenderTarget->Desc.Extent);
 	}
 }
 
-FCanvasRenderContext::FCanvasRenderContext(FRDGBuilder& InGraphBuilder, FRDGTextureRef InRenderTarget, FIntRect InViewportRect, FIntRect InScissorRect)
-	: GraphBuilder(InGraphBuilder)
-	, RenderTarget(InRenderTarget)
-	, ViewportRect(InViewportRect)
-	, ScissorRect(InScissorRect)
-{}
-
 FCanvasRenderThreadScope::FCanvasRenderThreadScope(const FCanvas& InCanvas)
 	: Canvas(InCanvas)
 {
@@ -575,12 +568,12 @@ FCanvasRenderThreadScope::~FCanvasRenderThreadScope()
 	RenderCommandFunctionArray* RenderCommandArray = RenderCommands;
 
 	ENQUEUE_RENDER_COMMAND(DispatchCanvasRenderCommands)(
-		[RenderCommandArray, LocalCanvas = Canvas](FRHICommandListImmediate& RHICmdList)
+		[RenderCommandArray, RenderTarget = Canvas.RenderTarget, ViewRect = Canvas.ViewRect, ScissorRect = Canvas.ScissorRect, bScaledToRenderTarget = Canvas.bScaledToRenderTarget](FRHICommandListImmediate& RHICmdList) mutable
 	{
 		GetRendererModule().InitializeSystemTextures(RHICmdList);
 		FMemMark MemMark(FMemStack::Get());
 		FRDGBuilder GraphBuilder(RHICmdList, RDG_EVENT_NAME("CanvasRenderThreadScope"));
-		FCanvasRenderContext RenderContext(GraphBuilder, LocalCanvas);
+		FCanvasRenderContext RenderContext(GraphBuilder, RenderTarget, ViewRect, ScissorRect, bScaledToRenderTarget);
 
 		for (uint32 Index = 0, Count = RenderCommandArray->Num(); Index < Count; ++Index)
 		{
@@ -783,7 +776,7 @@ void FCanvas::Flush_RenderThread(FRDGBuilder& GraphBuilder, bool bForce)
 
 	GetRendererModule().InitializeSystemTextures(GraphBuilder.RHICmdList);
 
-	FCanvasRenderContext RenderContext(GraphBuilder, *this);
+	FCanvasRenderContext RenderContext(GraphBuilder, RenderTarget, ViewRect, ScissorRect, bScaledToRenderTarget);
 
 	// Set the RHI render target.
 	if (IsUsingInternalTexture())
diff --git a/Engine/Source/Runtime/Engine/Public/CanvasRender.h b/Engine/Source/Runtime/Engine/Public/CanvasRender.h
--- a/Engine/Source/Runtime/Engine/Public/CanvasRender.h
+++ b/Engine/Source/Runtime/Engine/Public/CanvasRender.h
@@ -8,8 +8,7 @@
 class FCanvasRenderContext
 {
 public:
-	ENGINE_API FCanvasRenderContext(FRDGBuilder& GraphBuilder, const FCanvas& Canvas);
-	ENGINE_API FCanvasRenderContext(FRDGBuilder& GraphBuilder, FRDGTextureRef Texture, FIntRect ViewportRect, FIntRect ScissorRect);
+	ENGINE_API FCanvasRenderContext(FRDGBuilder& InGraphBuilder, const FRenderTarget* RenderTarget, FIntRect InViewportRect, FIntRect InScissorRect, bool bScaledToRenderTarget);
 
 	template <typename ExecuteLambdaType, typename ParameterStructType>
 	void AddPass(FRDGEventName&& PassName, const ParameterStructType* PassParameters, ExecuteLambdaType&& ExecuteLambda)
diff --git a/Engine/Source/Runtime/Engine/Public/CanvasTypes.h b/Engine/Source/Runtime/Engine/Public/CanvasTypes.h
--- a/Engine/Source/Runtime/Engine/Public/CanvasTypes.h
+++ b/Engine/Source/Runtime/Engine/Public/CanvasTypes.h
@@ -792,6 +792,7 @@ public:
 	ENGINE_API FCanvasSortElement& GetSortElement(int32 DepthSortKey);
 
 	friend class FCanvasRenderContext;
+	friend class FCanvasRenderThreadScope;
 };
 
 /**
