Data Layers: remove spamming log.

#rnx

diff --git a/Engine/Source/Runtime/Engine/Private/WorldPartition/DataLayer/DataLayerSubsystem.cpp b/Engine/Source/Runtime/Engine/Private/WorldPartition/DataLayer/DataLayerSubsystem.cpp
--- a/Engine/Source/Runtime/Engine/Private/WorldPartition/DataLayer/DataLayerSubsystem.cpp
+++ b/Engine/Source/Runtime/Engine/Private/WorldPartition/DataLayer/DataLayerSubsystem.cpp
@@ -642,8 +642,6 @@ void UDataLayerSubsystem::OnActorDescContainerInitialized(UActorDescContainer* I
 {
 	check(InActorDescContainer);
 
-	UE_SCOPED_TIMER(*FString::Printf(TEXT("Resolving Data Layer Instance Names for %s"), *InActorDescContainer->GetContainerPackage().ToString()), LogWorldPartition, Display);
-
 	auto GetWorldDataLayersActorDesc = [](const UActorDescContainer* InContainer) -> FWorldDataLayersActorDesc*
 	{
 		if (InContainer)
