Don't iterate garbage collected objects in FScopedRuntimeVirtualTextureRecreate.
Could potentially be the cause of a crash.
#rnx

diff --git a/Engine/Source/Runtime/Engine/Private/Components/RuntimeVirtualTextureComponent.cpp b/Engine/Source/Runtime/Engine/Private/Components/RuntimeVirtualTextureComponent.cpp
--- a/Engine/Source/Runtime/Engine/Private/Components/RuntimeVirtualTextureComponent.cpp
+++ b/Engine/Source/Runtime/Engine/Private/Components/RuntimeVirtualTextureComponent.cpp
@@ -196,7 +196,7 @@ class FScopedRuntimeVirtualTextureRecreate
 public:
 	FScopedRuntimeVirtualTextureRecreate(UVirtualTextureBuilder* VirtualTextureBuilder)
 	{
-		for (TObjectIterator<URuntimeVirtualTextureComponent> It; It; ++It)
+		for (TObjectIterator<URuntimeVirtualTextureComponent> It(RF_ClassDefaultObject, false, EInternalObjectFlags::Garbage); It; ++It)
 		{
 			if (It->GetStreamingTexture() == VirtualTextureBuilder)
 			{
