Add console command to reinitialize timecode and genlock

diff --git a/Engine/Source/Runtime/Engine/Private/UnrealEngine.cpp b/Engine/Source/Runtime/Engine/Private/UnrealEngine.cpp
--- a/Engine/Source/Runtime/Engine/Private/UnrealEngine.cpp
+++ b/Engine/Source/Runtime/Engine/Private/UnrealEngine.cpp
@@ -2892,6 +2892,17 @@ static void LoadCustomTimeStep(UEngine* Engine)
 	}
 }
 
+static void ResetCustomTimeStep()
+{
+	LoadCustomTimeStep(GEngine);
+}
+
+static FAutoConsoleCommand GResetCustomTimeStep(
+	TEXT("CustomTimeStep.reset"),
+	TEXT("Resets the current custom step."),
+	FConsoleCommandDelegate::CreateStatic(&ResetCustomTimeStep));
+
+
 static void LoadTimecodeProvider(UEngine* Engine)
 {
 	if (Engine->TimecodeProviderClassName.IsValid())
@@ -2931,6 +2942,16 @@ static void LoadTimecodeProvider(UEngine* Engine)
 	}
 }
 
+static void ResetTimecode()
+{
+	LoadTimecodeProvider(GEngine);
+}
+
+static FAutoConsoleCommand GResetTimecode(
+	TEXT("TimecodeProvider.reset"),
+	TEXT("Resets the current timecode provider."),
+	FConsoleCommandDelegate::CreateStatic(&ResetTimecode));
+
 /**
 * Loads all Engine object references from their corresponding config entries.
 */
