World Partition: don't allow changing properties that affects streaming when the world is partitioned but with streaming disabled.

#rnx

diff --git a/Engine/Source/Runtime/Engine/Private/ActorEditor.cpp b/Engine/Source/Runtime/Engine/Private/ActorEditor.cpp
--- a/Engine/Source/Runtime/Engine/Private/ActorEditor.cpp
+++ b/Engine/Source/Runtime/Engine/Private/ActorEditor.cpp
@@ -12,11 +12,11 @@
 #include "Components/PrimitiveComponent.h"
 #include "AI/NavigationSystemBase.h"
 #include "Engine/BlueprintGeneratedClass.h"
+#include "WorldPartition/WorldPartition.h"
 #include "WorldPartition/WorldPartitionSubsystem.h"
 #include "WorldPartition/DataLayer/DataLayer.h"
 #include "WorldPartition/DataLayer/DataLayerSubsystem.h"
 #include "WorldPartition/DataLayer/DataLayerInstance.h"
-#include "WorldPartition/WorldPartitionRuntimeCellInterface.h"
 #include "EditorSupportDelegates.h"
 #include "Logging/TokenizedMessage.h"
 #include "Logging/MessageLog.h"
@@ -83,10 +83,14 @@ bool AActor::CanEditChange(const FProperty* PropertyThatWillChange) const
 	{
 		if (!IsTemplate())
 		{
-			if (!UWorld::IsPartitionedWorld(GetTypedOuter<UWorld>()))
+			if (UWorld* World = GetTypedOuter<UWorld>())
 			{
-				return false;
+				if (UWorldPartition* WorldPartition = World->GetWorldPartition())
+				{
+					return bIsDataLayersProperty || WorldPartition->IsStreamingEnabled();
+				}
 			}
+			return false;
 		}
 	}
 
