Fixes PinFriendlyName in Niagara being initialized from FName as CultureInvariant instead

diff --git a/Engine/Plugins/FX/Niagara/Source/NiagaraEditor/Private/NiagaraGraph.cpp b/Engine/Plugins/FX/Niagara/Source/NiagaraEditor/Private/NiagaraGraph.cpp
--- a/Engine/Plugins/FX/Niagara/Source/NiagaraEditor/Private/NiagaraGraph.cpp
+++ b/Engine/Plugins/FX/Niagara/Source/NiagaraEditor/Private/NiagaraGraph.cpp
@@ -1066,7 +1066,7 @@ void UNiagaraGraph::StandardizeParameterNames()
 						bAnyNamesModified = true;
 					}
 					Pin->PinName = NewName;
-					Pin->PinFriendlyName = FText::FromName(NewName);
+					Pin->PinFriendlyName = FText::AsCultureInvariant(NewName.ToString());
 				}
 			}
 		}
diff --git a/Engine/Plugins/FX/Niagara/Source/NiagaraEditor/Private/NiagaraNodeParameterMapBase.cpp b/Engine/Plugins/FX/Niagara/Source/NiagaraEditor/Private/NiagaraNodeParameterMapBase.cpp
--- a/Engine/Plugins/FX/Niagara/Source/NiagaraEditor/Private/NiagaraNodeParameterMapBase.cpp
+++ b/Engine/Plugins/FX/Niagara/Source/NiagaraEditor/Private/NiagaraNodeParameterMapBase.cpp
@@ -245,7 +245,7 @@ bool UNiagaraNodeParameterMapBase::HandleDropOperation(TSharedPtr<FDragDropOpera
 
 void UNiagaraNodeParameterMapBase::OnPinRenamed(UEdGraphPin* RenamedPin, const FString& OldName)
 {
-	RenamedPin->PinFriendlyName = FText::FromName(RenamedPin->PinName);
+	RenamedPin->PinFriendlyName = FText::AsCultureInvariant(RenamedPin->PinName.ToString());
 
 	FPinCollectorArray InOrOutPins;
 	if (RenamedPin->Direction == EEdGraphPinDirection::EGPD_Input)
diff --git a/Engine/Plugins/FX/Niagara/Source/NiagaraEditor/Private/NiagaraNodeParameterMapGet.cpp b/Engine/Plugins/FX/Niagara/Source/NiagaraEditor/Private/NiagaraNodeParameterMapGet.cpp
--- a/Engine/Plugins/FX/Niagara/Source/NiagaraEditor/Private/NiagaraNodeParameterMapGet.cpp
+++ b/Engine/Plugins/FX/Niagara/Source/NiagaraEditor/Private/NiagaraNodeParameterMapGet.cpp
@@ -225,7 +225,7 @@ void UNiagaraNodeParameterMapGet::OnNewTypedPinAdded(UEdGraphPin*& NewPin)
 		const FName NewUniqueName = FNiagaraUtilities::GetUniqueName(NewPinName, Names);
 
 		NewPin->PinName = NewUniqueName;
-		NewPin->PinFriendlyName = FText::FromName(NewPin->PinName);
+		NewPin->PinFriendlyName = FText::AsCultureInvariant(NewPin->PinName.ToString());
 		NewPin->PinType.PinSubCategory = UNiagaraNodeParameterMapBase::ParameterPinSubCategory;
 
 		UEdGraphPin* MatchingDefault = GetDefaultPin(NewPin);
diff --git a/Engine/Plugins/FX/Niagara/Source/NiagaraEditor/Private/NiagaraNodeParameterMapSet.cpp b/Engine/Plugins/FX/Niagara/Source/NiagaraEditor/Private/NiagaraNodeParameterMapSet.cpp
--- a/Engine/Plugins/FX/Niagara/Source/NiagaraEditor/Private/NiagaraNodeParameterMapSet.cpp
+++ b/Engine/Plugins/FX/Niagara/Source/NiagaraEditor/Private/NiagaraNodeParameterMapSet.cpp
@@ -121,7 +121,7 @@ void UNiagaraNodeParameterMapSet::OnNewTypedPinAdded(UEdGraphPin*& NewPin)
 
 		//GetDefault<UEdGraphSchema_Niagara>()->PinToNiagaraVariable()
 		NewPin->PinName = NewUniqueName;
-		NewPin->PinFriendlyName = FText::FromName(NewPin->PinName);
+		NewPin->PinFriendlyName = FText::AsCultureInvariant(NewPin->PinName.ToString());
 		NewPin->PinType.PinSubCategory = UNiagaraNodeParameterMapBase::ParameterPinSubCategory;
 		
 		// If dragging from a function or other non-namespaced parent node, we should 
diff --git a/Engine/Plugins/FX/Niagara/Source/NiagaraEditor/Private/NiagaraNodeReadDataSet.cpp b/Engine/Plugins/FX/Niagara/Source/NiagaraEditor/Private/NiagaraNodeReadDataSet.cpp
--- a/Engine/Plugins/FX/Niagara/Source/NiagaraEditor/Private/NiagaraNodeReadDataSet.cpp
+++ b/Engine/Plugins/FX/Niagara/Source/NiagaraEditor/Private/NiagaraNodeReadDataSet.cpp
@@ -34,7 +34,7 @@ void UNiagaraNodeReadDataSet::AllocateDefaultPins()
 		UEdGraphPin* NewPin = CreatePin(EGPD_Output, Schema->TypeDefinitionToPinType(Var.GetType()), Var.GetName());
 		if (useFriendlyNames && VariableFriendlyNames[i].IsEmpty() == false)
 		{
-			NewPin->PinFriendlyName = FText::FromString(VariableFriendlyNames[i]);
+			NewPin->PinFriendlyName = FText::AsCultureInvariant(VariableFriendlyNames[i]);
 		}
 	}
 }
diff --git a/Engine/Plugins/FX/Niagara/Source/NiagaraEditor/Private/NiagaraNodeWriteDataSet.cpp b/Engine/Plugins/FX/Niagara/Source/NiagaraEditor/Private/NiagaraNodeWriteDataSet.cpp
--- a/Engine/Plugins/FX/Niagara/Source/NiagaraEditor/Private/NiagaraNodeWriteDataSet.cpp
+++ b/Engine/Plugins/FX/Niagara/Source/NiagaraEditor/Private/NiagaraNodeWriteDataSet.cpp
@@ -47,7 +47,7 @@ void UNiagaraNodeWriteDataSet::AllocateDefaultPins()
 		UEdGraphPin* NewPin = CreatePin(EGPD_Input, Schema->TypeDefinitionToPinType(Var.GetType()), Var.GetName());
 		if (useFriendlyNames && VariableFriendlyNames[i].IsEmpty() == false)
 		{
-			NewPin->PinFriendlyName = FText::FromString(VariableFriendlyNames[i]);
+			NewPin->PinFriendlyName = FText::AsCultureInvariant(VariableFriendlyNames[i]);
 		}
 	}
 }
diff --git a/Engine/Source/Runtime/Engine/Private/EdGraph/EdGraphPin.cpp b/Engine/Source/Runtime/Engine/Private/EdGraph/EdGraphPin.cpp
--- a/Engine/Source/Runtime/Engine/Private/EdGraph/EdGraphPin.cpp
+++ b/Engine/Source/Runtime/Engine/Private/EdGraph/EdGraphPin.cpp
@@ -801,13 +801,13 @@ FText UEdGraphPin::GetDisplayName() const
 	}
 	else
 	{
-		DisplayName = (!PinFriendlyName.IsEmpty()) ? PinFriendlyName : FText::FromName(PinName);
+		DisplayName = (!PinFriendlyName.IsEmpty()) ? PinFriendlyName : FText::AsCultureInvariant(PinName.ToString());
 
 		bool bShouldUseLocalizedNodeAndPinNames = false;
 		GConfig->GetBool( TEXT("Internationalization"), TEXT("ShouldUseLocalizedNodeAndPinNames"), bShouldUseLocalizedNodeAndPinNames, GEditorSettingsIni );
 		if (!bShouldUseLocalizedNodeAndPinNames)
 		{
-			return FText::FromString(DisplayName.BuildSourceString());
+			return FText::AsCultureInvariant(DisplayName.BuildSourceString());
 		}
 	}
 	return DisplayName;
