Spline mesh component : make sure we have a valid value to divide with when mesh are extremely thin on the forward axis

diff --git a/Engine/Source/Runtime/Engine/Private/Components/SplineMeshComponent.cpp b/Engine/Source/Runtime/Engine/Private/Components/SplineMeshComponent.cpp
--- a/Engine/Source/Runtime/Engine/Private/Components/SplineMeshComponent.cpp
+++ b/Engine/Source/Runtime/Engine/Private/Components/SplineMeshComponent.cpp
@@ -777,7 +777,10 @@ FTransform USplineMeshComponent::CalcSliceTransform(const float DistanceAlong) c
 		const FBoxSphereBounds StaticMeshBounds = GetStaticMesh()->GetBounds();
 		const float MeshMinZ = GetAxisValue(StaticMeshBounds.Origin, ForwardAxis) - GetAxisValue(StaticMeshBounds.BoxExtent, ForwardAxis);
 		const float MeshRangeZ = 2.0f * GetAxisValue(StaticMeshBounds.BoxExtent, ForwardAxis);
-		Alpha = (DistanceAlong - MeshMinZ) / MeshRangeZ;
+		if (MeshRangeZ > UE_SMALL_NUMBER)
+		{
+			Alpha = (DistanceAlong - MeshMinZ) / MeshRangeZ;
+		}
 	}
 
 	return CalcSliceTransformAtSplineOffset(Alpha);
