Rollback CL 20560777 as it seems the USubsystem* in SubsystemMap are already stale before we start iterating on it
As reported by Asan when creating the TWeakObjectPtr<USubsystem>

diff --git a/Engine/Source/Runtime/Engine/Private/Subsystems/SubsystemCollection.cpp b/Engine/Source/Runtime/Engine/Private/Subsystems/SubsystemCollection.cpp
--- a/Engine/Source/Runtime/Engine/Private/Subsystems/SubsystemCollection.cpp
+++ b/Engine/Source/Runtime/Engine/Private/Subsystems/SubsystemCollection.cpp
@@ -7,7 +7,6 @@
 #include "UObject/Package.h"
 #include "Modules/ModuleManager.h"
 #include "Misc/PackageName.h"
-#include "Algo/Transform.h"
 
 DEFINE_LOG_CATEGORY(LogSubsystemCollection);
 
@@ -186,26 +185,17 @@ void FSubsystemCollectionBase::Deinitialize()
 
 	// Deinit and clean up existing systems
 	SubsystemArrayMap.Empty();
-	if (!SubsystemMap.IsEmpty())
+	for (auto Iter = SubsystemMap.CreateIterator(); Iter; ++Iter)
 	{
-		TMap<UClass*, TWeakObjectPtr<USubsystem>> SubsystemMapCopy;
-		SubsystemMapCopy.Reserve(SubsystemMap.Num());
-		Algo::Transform(SubsystemMap, SubsystemMapCopy, [](const TPair<UClass*, USubsystem*> Iter){ return TPair<UClass*, TWeakObjectPtr<USubsystem>>(Iter.Key, Iter.Value); });
-
-		for (auto Iter = SubsystemMapCopy.CreateIterator(); Iter; ++Iter)
+		UClass* KeyClass = Iter.Key();
+		USubsystem* Subsystem = Iter.Value();
+		if ( Subsystem != nullptr && Subsystem->GetClass() == KeyClass)
 		{
-			UClass* KeyClass = Iter.Key();
-			USubsystem* Subsystem = Iter.Value().GetEvenIfUnreachable();
-			if (Subsystem && Subsystem->InternalOwningSubsystem && Subsystem->GetClass() == KeyClass)
-			{
-				Subsystem->Deinitialize();
-				Subsystem->InternalOwningSubsystem = nullptr;
-			}
+			Subsystem->Deinitialize();
+			Subsystem->InternalOwningSubsystem = nullptr;
 		}
-
-		SubsystemMap.Empty();
-		SubsystemArrayMap.Empty(); // Clear again in case it gets repopulated during Deinitialize
 	}
+	SubsystemMap.Empty();
 	Outer = nullptr;
 }
 
