Better error message when finding duplicate actors.

#rnx

diff --git a/Engine/Source/Runtime/Engine/Private/World.cpp b/Engine/Source/Runtime/Engine/Private/World.cpp
--- a/Engine/Source/Runtime/Engine/Private/World.cpp
+++ b/Engine/Source/Runtime/Engine/Private/World.cpp
@@ -1556,15 +1556,15 @@ void UWorld::RepairSingletonActorOfClass(TSubclassOf<AActor> ActorClass)
 			{
 				if (FoundActor)
 				{
-					UE_LOG(LogWorld, Log, TEXT("Extra '%s' actor found. Resave level %s or actors to clean up."), *CurrentActor->GetPathName(), *PersistentLevel->GetPathName());
-
 					if (FoundActor == CurrentActor)
 					{
+						UE_LOG(LogWorld, Log, TEXT("Extra '%s' actor found. Resave level %s to clean up."), *CurrentActor->GetPathName(), *PersistentLevel->GetPathName());
 						PersistentLevel->Actors[i] = nullptr;
 					}
 					else
 					{
-						CurrentActor->Destroy();
+						UE_LOG(LogWorld, Log, TEXT("Extra '%s' actor found. Resave level %s and actor to cleanup."), *CurrentActor->GetPathName(), *PersistentLevel->GetPathName());
+						CurrentActor->Destroy();						
 					}
 				}
 				else
@@ -1642,7 +1642,7 @@ void UWorld::RepairWorldSettings()
 			AActor* Actor = PersistentLevel->Actors[Index];
 			if (Actor != nullptr && Actor->IsA<AWorldSettings>())
 			{
-				UE_LOG(LogWorld, Warning, TEXT("Extra World Settings '%s' actor found. Resave level %s or actors to clean up"), *Actor->GetPathName(), *PersistentLevel->GetPathName());
+				UE_LOG(LogWorld, Warning, TEXT("Extra World Settings '%s' actor found. Resave level %s to clean up."), *Actor->GetPathName(), *PersistentLevel->GetPathName());
 				Actor->Destroy();
 			}
 		}
