[KillPendingKill] Mark packages as garbage in PrepareStreamedOutLevelsForGC so that weak pointers to them in the net driver's guid->object cache become invalidated.

diff --git a/Engine/Source/Runtime/Engine/Private/World.cpp b/Engine/Source/Runtime/Engine/Private/World.cpp
--- a/Engine/Source/Runtime/Engine/Private/World.cpp
+++ b/Engine/Source/Runtime/Engine/Private/World.cpp
@@ -3346,6 +3346,7 @@ void FLevelStreamingGCHelper::PrepareStreamedOutLevelsForGC()
 				if (!bIsAlreadyInSet)
 				{
 					ForEachObjectWithPackage(Package, [](UObject* PackageObject) { PackageObject->MarkAsGarbage(); return true; }, true, RF_NoFlags, EInternalObjectFlags::Garbage);
+					Package->MarkAsGarbage();
 				}
 			}, true, RF_NoFlags, EInternalObjectFlags::Garbage);
 
