Replace 5.00 deprecation with 5.0 to avoid confusion

#rnx

diff --git a/Engine/Source/Runtime/Engine/Classes/Engine/SkeletalMesh.h b/Engine/Source/Runtime/Engine/Classes/Engine/SkeletalMesh.h
--- a/Engine/Source/Runtime/Engine/Classes/Engine/SkeletalMesh.h
+++ b/Engine/Source/Runtime/Engine/Classes/Engine/SkeletalMesh.h
@@ -778,12 +778,12 @@ class ENGINE_API USkeletalMesh : public USkinnedAsset, public IInterface_Collisi
 private:
 #if WITH_EDITORONLY_DATA
 	/** Imported skeletal mesh geometry information (not used at runtime). */
-	UE_DEPRECATED(5.00, "This must be protected for async build, always use the accessors even internally.")
+	UE_DEPRECATED(5.0, "This must be protected for async build, always use the accessors even internally.")
 	TSharedPtr<FSkeletalMeshModel> ImportedModel;
 #endif
 
 	/** Rendering resources used at runtime */
-	UE_DEPRECATED(5.00, "This must be protected for async build, always use the accessors even internally.")
+	UE_DEPRECATED(5.0, "This must be protected for async build, always use the accessors even internally.")
 	TUniquePtr<FSkeletalMeshRenderData> SkeletalMeshRenderData;
 
 	FSkeletalMeshRenderData* GetSkeletalMeshRenderData() const
@@ -979,12 +979,12 @@ public:
 
 private:
 	/** Original imported mesh bounds */
-	UE_DEPRECATED(5.00, "This must be protected for async build, always use the accessors even internally.")
+	UE_DEPRECATED(5.0, "This must be protected for async build, always use the accessors even internally.")
 	UPROPERTY(transient, duplicatetransient)
 	FBoxSphereBounds ImportedBounds;
 
 	/** Bounds extended by user values below */
-	UE_DEPRECATED(5.00, "This must be protected for async build, always use the accessors even internally.")
+	UE_DEPRECATED(5.0, "This must be protected for async build, always use the accessors even internally.")
 	UPROPERTY(transient, duplicatetransient)
 	FBoxSphereBounds ExtendedBounds;
 
@@ -1012,7 +1012,7 @@ protected:
 	/** Bound extension values in addition to imported bound in the positive direction of XYZ, 
 	 *	positive value increases bound size and negative value decreases bound size. 
 	 *	The final bound would be from [Imported Bound - Negative Bound] to [Imported Bound + Positive Bound]. */
-	UE_DEPRECATED(5.00, "This must be protected for async build, always use the accessors even internally.")
+	UE_DEPRECATED(5.0, "This must be protected for async build, always use the accessors even internally.")
 	UPROPERTY(EditAnywhere, BlueprintReadOnly, Category = Mesh)
 	FVector PositiveBoundsExtension;
 
@@ -1026,7 +1026,7 @@ protected:
 	/** Bound extension values in addition to imported bound in the negative direction of XYZ, 
 	 *	positive value increases bound size and negative value decreases bound size. 
 	 *	The final bound would be from [Imported Bound - Negative Bound] to [Imported Bound + Positive Bound]. */
-	UE_DEPRECATED(5.00, "This must be protected for async build, always use the accessors even internally.")
+	UE_DEPRECATED(5.0, "This must be protected for async build, always use the accessors even internally.")
 	UPROPERTY(EditAnywhere, BlueprintReadOnly, Category = Mesh)
 	FVector NegativeBoundsExtension;
 
@@ -1185,7 +1185,7 @@ public:
 
 private:
 	/** Struct containing information for each LOD level, such as materials to use, and when use the LOD. */
-	UE_DEPRECATED(5.00, "This must be protected for async build, always use the accessors even internally.")
+	UE_DEPRECATED(5.0, "This must be protected for async build, always use the accessors even internally.")
 	UPROPERTY(EditAnywhere, EditFixedSize, Category=LevelOfDetail)
 	TArray<struct FSkeletalMeshLODInfo> LODInfo;
 
@@ -2112,7 +2112,7 @@ public:
 	/**
 	 * If true, a ray tracing acceleration structure will be built for this mesh and it may be used in ray tracing effects
 	 */
-	UE_DEPRECATED(5.00, "This must be protected for async build, always use the accessors even internally.")
+	UE_DEPRECATED(5.0, "This must be protected for async build, always use the accessors even internally.")
 	UPROPERTY(EditAnywhere, Category = RayTracing)
 	uint8 bSupportRayTracing : 1;
 
@@ -2136,7 +2136,7 @@ public:
 	/**
 	 * LOD bias for ray tracing. When non-zero, a different LOD level other than the predicted LOD level will be used for ray tracing. Advanced features like morph targets and cloth simulation may not work properly.
 	 */
-	UE_DEPRECATED(5.00, "This must be protected for async build, always use the accessors even internally.")
+	UE_DEPRECATED(5.0, "This must be protected for async build, always use the accessors even internally.")
 	UPROPERTY(EditAnywhere, Category = RayTracing)
 	int32 RayTracingMinLOD;
 
@@ -2161,7 +2161,7 @@ public:
 	 * Set the strategy used for storing the additional cloth deformer mappings depending on the desired use of Raytracing LOD bias.
 	 * This parameter is only used in relation to raytracing of the cloth sections.
 	 */
-	UE_DEPRECATED(5.00, "This must be protected for async build, always use the accessors even internally.")
+	UE_DEPRECATED(5.0, "This must be protected for async build, always use the accessors even internally.")
 	UPROPERTY(EditAnywhere, Category = RayTracing)
 	EClothLODBiasMode ClothLODBiasMode;
 
@@ -2591,7 +2591,7 @@ public:
 protected:
 
 	/** Defines if and how to generate a set of precomputed data allowing targeted and fast sampling of this mesh on the CPU. */
-	UE_DEPRECATED(5.00, "This must be protected for async build, always use the accessors even internally.")
+	UE_DEPRECATED(5.0, "This must be protected for async build, always use the accessors even internally.")
 	UPROPERTY(EditAnywhere, Category = "Sampling", meta=(ShowOnlyInnerProperties))
 	FSkeletalMeshSamplingInfo SamplingInfo;
 
@@ -2615,7 +2615,7 @@ private:
 	TArray<TObjectPtr<class USkeletalMeshSocket>> Sockets;
 
 	/** Cached matrices from GetComposedRefPoseMatrix */
-	UE_DEPRECATED(5.00, "This must be protected for async build, always use the accessors even internally.")
+	UE_DEPRECATED(5.0, "This must be protected for async build, always use the accessors even internally.")
 	TArray<FMatrix> CachedComposedRefPoseMatrices;
 
 	TArray<FMatrix>& GetCachedComposedRefPoseMatrices() 
@@ -3292,7 +3292,7 @@ public:
 
 protected:
 	/** Set of skin weight profiles associated with this mesh */
-	UE_DEPRECATED(5.00, "This must be protected for async build, always use the accessors even internally.")
+	UE_DEPRECATED(5.0, "This must be protected for async build, always use the accessors even internally.")
 	UPROPERTY(EditAnywhere, Category = SkinWeights, EditFixedSize, Meta=(NoResetToDefault))
 	TArray<FSkinWeightProfileInfo> SkinWeightProfiles;
 
diff --git a/Engine/Source/Runtime/Engine/Classes/Engine/StaticMesh.h b/Engine/Source/Runtime/Engine/Classes/Engine/StaticMesh.h
--- a/Engine/Source/Runtime/Engine/Classes/Engine/StaticMesh.h
+++ b/Engine/Source/Runtime/Engine/Classes/Engine/StaticMesh.h
@@ -609,7 +609,7 @@ private:
 #endif
 
 	/** Pointer to the data used to render this static mesh. */
-	UE_DEPRECATED(5.00, "This must be protected for async build, always use the accessors even internally.")
+	UE_DEPRECATED(5.0, "This must be protected for async build, always use the accessors even internally.")
 	TUniquePtr<class FStaticMeshRenderData> RenderData;
 
 public:
@@ -630,12 +630,12 @@ public:
 
 private:
 	/** Imported raw mesh bulk data. */
-	UE_DEPRECATED(5.00, "This must be protected for async build, always use the accessors even internally.")
+	UE_DEPRECATED(5.0, "This must be protected for async build, always use the accessors even internally.")
 	UPROPERTY(Setter = None, Getter = None)
 	TArray<FStaticMeshSourceModel> SourceModels;
 
 	/** Optional hi-res source data */
-	UE_DEPRECATED(5.00, "This must be protected for async build, always use the accessors even internally.")
+	UE_DEPRECATED(5.0, "This must be protected for async build, always use the accessors even internally.")
 	UPROPERTY(Setter = None, Getter = None)
 	FStaticMeshSourceModel HiResSourceModel;
 
@@ -648,7 +648,7 @@ private:
 	}
 
 	/** Map of LOD+Section index to per-section info. */
-	UE_DEPRECATED(5.00, "This must be protected for async build, always use the accessors even internally.")
+	UE_DEPRECATED(5.0, "This must be protected for async build, always use the accessors even internally.")
 	UPROPERTY()
 	FMeshSectionInfoMap SectionInfoMap;
 
@@ -660,7 +660,7 @@ private:
 	 *
 	 * We do not update it when the user shuffle section in the staticmesh editor because the OriginalSectionInfoMap must always be in sync with the saved rawMesh bulk data.
 	 */
-	UE_DEPRECATED(5.00, "This must be protected for async build, always use the accessors even internally.")
+	UE_DEPRECATED(5.0, "This must be protected for async build, always use the accessors even internally.")
 	UPROPERTY()
 	FMeshSectionInfoMap OriginalSectionInfoMap;
 
@@ -692,7 +692,7 @@ public:
 
 private:
 	/* The lightmap UV generation version used during the last derived data build */
-	UE_DEPRECATED(5.00, "This must be protected for async build, always use the accessors even internally.")
+	UE_DEPRECATED(5.0, "This must be protected for async build, always use the accessors even internally.")
 	UPROPERTY()
 	int32 LightmapUVVersion;
 public:
@@ -851,7 +851,7 @@ public:
 	}
 
 private:
-	UE_DEPRECATED(5.00, "This must be protected for async build, always use the accessors even internally.")
+	UE_DEPRECATED(5.0, "This must be protected for async build, always use the accessors even internally.")
 	UPROPERTY(BlueprintGetter = GetStaticMaterials, BlueprintSetter = SetStaticMaterials, Category = StaticMesh)
 	TArray<FStaticMaterial> StaticMaterials;
 
@@ -890,7 +890,7 @@ public:
 	}
 
 private:
-	UE_DEPRECATED(5.00, "This must be protected for async build, always use the accessors even internally.")
+	UE_DEPRECATED(5.0, "This must be protected for async build, always use the accessors even internally.")
 	UPROPERTY()
 	float LightmapUVDensity;
 public:
@@ -971,7 +971,7 @@ public:
 
 private:
 	// Physics data.
-	UE_DEPRECATED(5.00, "This must be protected for async build, always use the accessors even internally.")
+	UE_DEPRECATED(5.0, "This must be protected for async build, always use the accessors even internally.")
 	UPROPERTY(EditAnywhere, transient, duplicatetransient, Instanced, Category = StaticMesh)
 	TObjectPtr<class UBodySetup> BodySetup;
 public:
@@ -1055,13 +1055,13 @@ private:
 
 public:
 
-	UE_DEPRECATED(5.00, "IsBuiltAtRuntime() is no longer used.")
+	UE_DEPRECATED(5.0, "IsBuiltAtRuntime() is no longer used.")
 	bool IsBuiltAtRuntime() const
 	{
 		return false;
 	}
 
-	UE_DEPRECATED(5.00, "SetIsBuiltAtRuntime() is no longer used.")
+	UE_DEPRECATED(5.0, "SetIsBuiltAtRuntime() is no longer used.")
 	void SetIsBuiltAtRuntime(bool InIsBuiltAtRuntime)
 	{
 	}
@@ -1115,7 +1115,7 @@ public:
 
 private:
 	/** Unique ID for tracking/caching this mesh during distributed lighting */
-	UE_DEPRECATED(5.00, "This must be protected for async build, always use the accessors even internally.")
+	UE_DEPRECATED(5.0, "This must be protected for async build, always use the accessors even internally.")
 	FGuid LightingGuid;
 public:
 
@@ -1209,7 +1209,7 @@ public:
 	}
 
 	/** Original mesh bounds extended with Positive/NegativeBoundsExtension */
-	UE_DEPRECATED(5.00, "This must be protected for async build, always use the accessors even internally.")
+	UE_DEPRECATED(5.0, "This must be protected for async build, always use the accessors even internally.")
 	UPROPERTY()
 	FBoxSphereBounds ExtendedBounds;
 public:
@@ -1499,7 +1499,7 @@ public:
 
 	/** Pre-build navigation collision */
 private:
-	UE_DEPRECATED(5.00, "This must be protected for async build, always use the accessors even internally.")
+	UE_DEPRECATED(5.0, "This must be protected for async build, always use the accessors even internally.")
 	UPROPERTY(VisibleAnywhere, transient, duplicatetransient, Instanced, Category = Navigation, meta = (EditCondition = "bHasNavigationData"))
 	TObjectPtr<UNavCollisionBase> NavCollision;
 
diff --git a/Engine/Source/Runtime/Engine/Classes/Engine/Texture.h b/Engine/Source/Runtime/Engine/Classes/Engine/Texture.h
--- a/Engine/Source/Runtime/Engine/Classes/Engine/Texture.h
+++ b/Engine/Source/Runtime/Engine/Classes/Engine/Texture.h
@@ -402,7 +402,7 @@ struct FTextureSource
 	// Apply a visitor to the bulkdata :
 	ENGINE_API void OperateOnLoadedBulkData(TFunctionRef<void (const FSharedBuffer& BulkDataBuffer)> Operation);
 
-	UE_DEPRECATED(5.00, "There is no longer a need to call LoadBulkDataWithFileReader, FTextureSource::BulkData can now load the data on demand without it.")
+	UE_DEPRECATED(5.0, "There is no longer a need to call LoadBulkDataWithFileReader, FTextureSource::BulkData can now load the data on demand without it.")
 	FORCEINLINE bool LoadBulkDataWithFileReader() { return true; }
 
 	FORCEINLINE void RemoveBulkData() { BulkData.UnloadData(); }
@@ -981,9 +981,9 @@ public:
 		FStringView InDebugContext);
 	bool AreDerivedMipsAvailable(FStringView Context) const;
 	bool AreDerivedVTChunksAvailable(FStringView Context) const;
-	UE_DEPRECATED(5.00, "Use AreDerivedMipsAvailable with the context instead.")
+	UE_DEPRECATED(5.0, "Use AreDerivedMipsAvailable with the context instead.")
 	bool AreDerivedMipsAvailable() const;
-	UE_DEPRECATED(5.00, "Use AreDerivedVTChunksAvailable with the context instead.")
+	UE_DEPRECATED(5.0, "Use AreDerivedVTChunksAvailable with the context instead.")
 	bool AreDerivedVTChunksAvailable() const;
 #endif
 
@@ -1393,7 +1393,7 @@ private:
 
 public:
 #if WITH_TEXTURE_RESOURCE_DEPRECATIONS
-	UE_DEPRECATED(5.00, "Use GetResource() / SetResource() accessors instead.")
+	UE_DEPRECATED(5.0, "Use GetResource() / SetResource() accessors instead.")
 	TFieldPtrAccessor<FTextureResource> Resource;
 #endif
 
diff --git a/Engine/Source/Runtime/Engine/Classes/Engine/Texture2D.h b/Engine/Source/Runtime/Engine/Classes/Engine/Texture2D.h
--- a/Engine/Source/Runtime/Engine/Classes/Engine/Texture2D.h
+++ b/Engine/Source/Runtime/Engine/Classes/Engine/Texture2D.h
@@ -73,7 +73,7 @@ private:
 
 public:
 #if WITH_TEXTURE_PLATFORMDATA_DEPRECATIONS
-	UE_DEPRECATED(5.00, "Use GetPlatformData() / SetPlatformData() accessors instead.")
+	UE_DEPRECATED(5.0, "Use GetPlatformData() / SetPlatformData() accessors instead.")
 	TFieldPtrAccessor<FTexturePlatformData> PlatformData;
 #endif
 
diff --git a/Engine/Source/Runtime/Engine/Classes/Engine/TextureCube.h b/Engine/Source/Runtime/Engine/Classes/Engine/TextureCube.h
--- a/Engine/Source/Runtime/Engine/Classes/Engine/TextureCube.h
+++ b/Engine/Source/Runtime/Engine/Classes/Engine/TextureCube.h
@@ -21,7 +21,7 @@ class UTextureCube : public UTexture
 
 public:
 #if WITH_TEXTURE_PLATFORMDATA_DEPRECATIONS
-	UE_DEPRECATED(5.00, "Use GetPlatformData() / SetPlatformData() accessors instead.")
+	UE_DEPRECATED(5.0, "Use GetPlatformData() / SetPlatformData() accessors instead.")
 	TFieldPtrAccessor<FTexturePlatformData> PlatformData;
 #endif
 
diff --git a/Engine/Source/Runtime/Engine/Classes/Sound/SoundWave.h b/Engine/Source/Runtime/Engine/Classes/Sound/SoundWave.h
--- a/Engine/Source/Runtime/Engine/Classes/Sound/SoundWave.h
+++ b/Engine/Source/Runtime/Engine/Classes/Sound/SoundWave.h
@@ -154,7 +154,7 @@ struct ENGINE_API FStreamedAudioPlatformData
 	bool IsCompiling() const;
 	bool TryInlineChunkData();
 
-	UE_DEPRECATED(5.00, "Use AreDerivedChunksAvailable with the context instead.")
+	UE_DEPRECATED(5.0, "Use AreDerivedChunksAvailable with the context instead.")
 	bool AreDerivedChunksAvailable() const;
 	
 	bool AreDerivedChunksAvailable(FStringView Context) const;
diff --git a/Engine/Source/Runtime/Engine/Public/WorldPartition/DataLayer/WorldDataLayers.h b/Engine/Source/Runtime/Engine/Public/WorldPartition/DataLayer/WorldDataLayers.h
--- a/Engine/Source/Runtime/Engine/Public/WorldPartition/DataLayer/WorldDataLayers.h
+++ b/Engine/Source/Runtime/Engine/Public/WorldPartition/DataLayer/WorldDataLayers.h
@@ -236,7 +236,7 @@ private:
 public:
 	DECLARE_DELEGATE_RetVal_ThreeParams(bool, FDataLayersFilterDelegate, FName /*DataLayerName*/, EDataLayerRuntimeState /*CurrentState*/, EDataLayerRuntimeState /*TargetState*/);
 
-	UE_DEPRECATED(5.00, "do not use, will be replaced by another mechanism for initial release.")
+	UE_DEPRECATED(5.0, "do not use, will be replaced by another mechanism for initial release.")
 	FDataLayersFilterDelegate DataLayersFilterDelegate;
 };
 
