Fix missing include in PhysicsFieldComponent for CIS

diff --git a/Engine/Source/Runtime/Engine/Classes/PhysicsField/PhysicsFieldComponent.h b/Engine/Source/Runtime/Engine/Classes/PhysicsField/PhysicsFieldComponent.h
--- a/Engine/Source/Runtime/Engine/Classes/PhysicsField/PhysicsFieldComponent.h
+++ b/Engine/Source/Runtime/Engine/Classes/PhysicsField/PhysicsFieldComponent.h
@@ -8,6 +8,7 @@
 #include "RHI.h"
 #include "RenderResource.h"
 #include "PrimitiveSceneProxy.h"
+#include "Kismet/BlueprintFunctionLibrary.h"
 
 #include "PhysicsFieldComponent.generated.h"
 
