PR #9299: Add option to get short FKey display name from blueprints (Contributed by [REDACTED_PERSON])

diff --git a/Engine/Source/Runtime/Engine/Classes/Kismet/KismetInputLibrary.h b/Engine/Source/Runtime/Engine/Classes/Kismet/KismetInputLibrary.h
--- a/Engine/Source/Runtime/Engine/Classes/Kismet/KismetInputLibrary.h
+++ b/Engine/Source/Runtime/Engine/Classes/Kismet/KismetInputLibrary.h
@@ -155,7 +155,7 @@ class ENGINE_API UKismetInputLibrary : public UBlueprintFunctionLibrary
 	 * Returns the display name of the key.
 	 */
 	UFUNCTION(BlueprintPure, meta = (DisplayName = "Get Key Display Name"), Category="Input|Key")
-	static FText Key_GetDisplayName(const FKey& Key);
+	static FText Key_GetDisplayName(const FKey& Key, bool bLongDisplayName = true);
 
 	/**
 	 * Returns whether or not this character is an auto-repeated keystroke
diff --git a/Engine/Source/Runtime/Engine/Private/KismetInputLibrary.cpp b/Engine/Source/Runtime/Engine/Private/KismetInputLibrary.cpp
--- a/Engine/Source/Runtime/Engine/Private/KismetInputLibrary.cpp
+++ b/Engine/Source/Runtime/Engine/Private/KismetInputLibrary.cpp
@@ -146,9 +146,9 @@ EUINavigation UKismetInputLibrary::Key_GetNavigationDirectionFromAnalog(const FA
 }
 
 
-FText UKismetInputLibrary::Key_GetDisplayName(const FKey& Key)
+FText UKismetInputLibrary::Key_GetDisplayName(const FKey& Key, const bool bLongDisplayName)
 {
-	return Key.GetDisplayName();
+	return Key.GetDisplayName(bLongDisplayName);
 }
 
 bool UKismetInputLibrary::InputEvent_IsRepeat(const FInputEvent& Input)
