Output the WorldGridMaterial's human readable DDC key string to the log.

- This is done of diagnostic purposes on users machines, and for the build machines.

diff --git a/Engine/Source/Runtime/Engine/Private/Materials/MaterialShader.cpp b/Engine/Source/Runtime/Engine/Private/Materials/MaterialShader.cpp
--- a/Engine/Source/Runtime/Engine/Private/Materials/MaterialShader.cpp
+++ b/Engine/Source/Runtime/Engine/Private/Materials/MaterialShader.cpp
@@ -1575,6 +1575,13 @@ TSharedRef<FMaterialShaderMap::FAsyncLoadContext> FMaterialShaderMap::BeginLoadF
 				DumpShaderDDCKeyToFile(InPlatform, ShaderMapId.LayoutParams.WithEditorOnly(), FileName, Result->DataKey);
 			}
 
+			if (Material->IsDefaultMaterial() && Material->GetMaterialDomain() == EMaterialDomain::MD_Surface)
+			{
+				FString SpecialEngineDDCKey = Result->DataKey;
+				SpecialEngineDDCKey.RemoveFromEnd(TEXT("\n"));
+				UE_LOG(LogMaterial, Display, TEXT("%s-%s-%s: %s"), *Material->GetAssetName(), *LexToString(ShaderMapId.FeatureLevel), *LexToString(ShaderMapId.QualityLevel), *SpecialEngineDDCKey);
+			}
+
 			bool CheckCache = true;
 
 			// If NoShaderDDC then don't check for a material the first time we encounter it to simulate
