LevelInstance: Fix SpriteComponent
- Avoid showing red sprite when locked
- Update visibility properly

diff --git a/Engine/Source/Runtime/Engine/Classes/Components/BillboardComponent.h b/Engine/Source/Runtime/Engine/Classes/Components/BillboardComponent.h
--- a/Engine/Source/Runtime/Engine/Classes/Components/BillboardComponent.h
+++ b/Engine/Source/Runtime/Engine/Classes/Components/BillboardComponent.h
@@ -57,6 +57,9 @@ class ENGINE_API UBillboardComponent : public UPrimitiveComponent
 	/** Whether to use in-editor arrow scaling (i.e. to be affected by the global arrow scale) */
 	UPROPERTY(EditAnywhere, BlueprintReadWrite, Category=Sprite)
 	bool bUseInEditorScaling;
+
+	UPROPERTY()
+	bool bShowLockedLocation;
 #endif // WITH_EDITORONLY_DATA
 	/** Change the sprite texture used by this component */
 	UFUNCTION(BlueprintCallable, Category="Rendering|Components|Sprite")
diff --git a/Engine/Source/Runtime/Engine/Classes/Components/SceneComponent.h b/Engine/Source/Runtime/Engine/Classes/Components/SceneComponent.h
--- a/Engine/Source/Runtime/Engine/Classes/Components/SceneComponent.h
+++ b/Engine/Source/Runtime/Engine/Classes/Components/SceneComponent.h
@@ -297,6 +297,7 @@ protected:
 	class UBillboardComponent* SpriteComponent;
 	/** Creates the editor only component used to display the sprite */
 	void CreateSpriteComponent(class UTexture2D* SpriteTexture = nullptr);
+	void CreateSpriteComponent(class UTexture2D* SpriteTexture, bool bRegister);
 #endif
 
 public:
diff --git a/Engine/Source/Runtime/Engine/Private/Components/BillboardComponent.cpp b/Engine/Source/Runtime/Engine/Private/Components/BillboardComponent.cpp
--- a/Engine/Source/Runtime/Engine/Private/Components/BillboardComponent.cpp
+++ b/Engine/Source/Runtime/Engine/Private/Components/BillboardComponent.cpp
@@ -99,7 +99,7 @@ public:
 
 			//save off override states
 #if WITH_EDITORONLY_DATA
-			bIsActorLocked = Owner->IsLockLocation();
+			bIsActorLocked = InComponent->bShowLockedLocation && Owner->IsLockLocation();
 #else // WITH_EDITORONLY_DATA
 			bIsActorLocked = false;
 #endif // WITH_EDITORONLY_DATA
diff --git a/Engine/Source/Runtime/Engine/Private/Components/SceneComponent.cpp b/Engine/Source/Runtime/Engine/Private/Components/SceneComponent.cpp
--- a/Engine/Source/Runtime/Engine/Private/Components/SceneComponent.cpp
+++ b/Engine/Source/Runtime/Engine/Private/Components/SceneComponent.cpp
@@ -710,6 +710,11 @@ void USceneComponent::OnRegister()
 
 #if WITH_EDITORONLY_DATA
 void USceneComponent::CreateSpriteComponent(UTexture2D* SpriteTexture)
+{
+	CreateSpriteComponent(SpriteTexture, true);
+}
+
+void USceneComponent::CreateSpriteComponent(class UTexture2D* SpriteTexture, bool bRegister)
 {
 	if (bVisualizeComponent && SpriteComponent == nullptr && GetOwner() && !GetWorld()->IsGameWorld())
 	{
@@ -729,7 +734,11 @@ void USceneComponent::CreateSpriteComponent(UTexture2D* SpriteTexture)
 		SpriteComponent->OpacityMaskRefVal = .3f;
 
 		SpriteComponent->SetupAttachment(this);
-		SpriteComponent->RegisterComponent();
+
+		if (bRegister)
+		{
+			SpriteComponent->RegisterComponent();
+		}
 	}
 }
 #endif
diff --git a/Engine/Source/Runtime/Engine/Private/LevelInstance/LevelInstanceComponent.cpp b/Engine/Source/Runtime/Engine/Private/LevelInstance/LevelInstanceComponent.cpp
--- a/Engine/Source/Runtime/Engine/Private/LevelInstance/LevelInstanceComponent.cpp
+++ b/Engine/Source/Runtime/Engine/Private/LevelInstance/LevelInstanceComponent.cpp
@@ -3,6 +3,7 @@
 #include "LevelInstance/LevelInstanceInterface.h"
 #include "LevelInstance/LevelInstanceSubsystem.h"
 #include "LevelInstance/LevelInstanceEditorInstanceActor.h"
+#include "Components/BillboardComponent.h"
 #include "Engine/Texture2D.h"
 #include "Engine/World.h"
 
@@ -27,15 +28,26 @@ void ULevelInstanceComponent::OnRegister()
 #if WITH_EDITORONLY_DATA
 	AActor* Owner = GetOwner();
 	// Only show Sprite for non-instanced LevelInstances
-	if (Owner && Owner->GetLevel() && !GetOwner()->GetLevel()->IsInstancedLevel() && !GetWorld()->IsGameWorld())
+	if (GetWorld() && !GetWorld()->IsGameWorld())
 	{
 		// Re-enable before calling CreateSpriteComponent
 		bVisualizeComponent = true;
-		CreateSpriteComponent(LoadObject<UTexture2D>(nullptr, TEXT("/Engine/EditorResources/LevelInstance")));
+		CreateSpriteComponent(LoadObject<UTexture2D>(nullptr, TEXT("/Engine/EditorResources/LevelInstance")), false);
+		if (SpriteComponent)
+		{
+			SpriteComponent->bShowLockedLocation = false;
+			SpriteComponent->SetVisibility(ShouldShowSpriteComponent());
+			SpriteComponent->RegisterComponent();
+		}
 	}
 #endif //WITH_EDITORONLY_DATA
 }
 
+bool ULevelInstanceComponent::ShouldShowSpriteComponent() const
+{
+	return GetOwner() && GetOwner()->GetLevel() && (GetOwner()->GetLevel()->IsPersistentLevel() || !GetOwner()->GetLevel()->IsInstancedLevel());
+}
+
 void ULevelInstanceComponent::PostEditUndo()
 {
 	Super::PostEditUndo();
@@ -89,4 +101,20 @@ void ULevelInstanceComponent::UpdateEditorInstanceActor()
 	}
 }
 
+void ULevelInstanceComponent::OnEdit()
+{
+	if (SpriteComponent)
+	{
+		SpriteComponent->SetVisibility(false);
+	}
+}
+
+void ULevelInstanceComponent::OnCommit()
+{
+	if (SpriteComponent)
+	{
+		SpriteComponent->SetVisibility(ShouldShowSpriteComponent());
+	}
+}
+
 #endif
\ No newline at end of file
diff --git a/Engine/Source/Runtime/Engine/Private/LevelInstance/LevelInstanceInterface.cpp b/Engine/Source/Runtime/Engine/Private/LevelInstance/LevelInstanceInterface.cpp
--- a/Engine/Source/Runtime/Engine/Private/LevelInstance/LevelInstanceInterface.cpp
+++ b/Engine/Source/Runtime/Engine/Private/LevelInstance/LevelInstanceInterface.cpp
@@ -7,6 +7,7 @@
 
 #if WITH_EDITOR
 #include "LevelInstance/LevelInstanceEditorInstanceActor.h"
+#include "LevelInstance/LevelInstanceComponent.h"
 #endif
 
 ULevelInstanceInterface::ULevelInstanceInterface(const FObjectInitializer& ObjectInitializer)
@@ -162,6 +163,14 @@ bool ILevelInstanceInterface::EnterEdit(AActor* ContextActor)
 	return true;
 }
 
+void ILevelInstanceInterface::OnEdit()
+{
+	if (ULevelInstanceComponent* LevelInstanceComponent = GetLevelInstanceComponent())
+	{
+		LevelInstanceComponent->OnEdit();
+	}
+}
+
 bool ILevelInstanceInterface::CanExitEdit(bool bDiscardEdits, FText* OutReason) const
 {
 	if (HasValidLevelInstanceID())
@@ -182,6 +191,14 @@ bool ILevelInstanceInterface::ExitEdit(bool bDiscardEdits)
 	return LevelInstanceSubsystem->CommitLevelInstance(this, bDiscardEdits);
 }
 
+void ILevelInstanceInterface::OnCommit(bool bChanged)
+{
+	if (ULevelInstanceComponent* LevelInstanceComponent = GetLevelInstanceComponent())
+	{
+		LevelInstanceComponent->OnCommit();
+	}
+}
+
 bool ILevelInstanceInterface::SetCurrent()
 {
 	if (ULevelInstanceSubsystem* LevelInstanceSubsystem = GetLevelInstanceSubsystem())
diff --git a/Engine/Source/Runtime/Engine/Public/LevelInstance/LevelInstanceComponent.h b/Engine/Source/Runtime/Engine/Public/LevelInstance/LevelInstanceComponent.h
--- a/Engine/Source/Runtime/Engine/Public/LevelInstance/LevelInstanceComponent.h
+++ b/Engine/Source/Runtime/Engine/Public/LevelInstance/LevelInstanceComponent.h
@@ -31,7 +31,11 @@ public:
 	virtual void OnUpdateTransform(EUpdateTransformFlags UpdateTransformFlags, ETeleportType Teleport) override;
 
 	void UpdateEditorInstanceActor();
+	void OnEdit();
+	void OnCommit();
 private:
+	bool ShouldShowSpriteComponent() const;
+
 	TWeakObjectPtr<AActor> CachedEditorInstanceActorPtr;
 #endif
 };
\ No newline at end of file
diff --git a/Engine/Source/Runtime/Engine/Public/LevelInstance/LevelInstanceInterface.h b/Engine/Source/Runtime/Engine/Public/LevelInstance/LevelInstanceInterface.h
--- a/Engine/Source/Runtime/Engine/Public/LevelInstance/LevelInstanceInterface.h
+++ b/Engine/Source/Runtime/Engine/Public/LevelInstance/LevelInstanceInterface.h
@@ -70,11 +70,11 @@
 
 	virtual void UpdateLevelInstanceFromWorldAsset();
 
-	virtual void OnEdit() {}
+	virtual void OnEdit();
 
 	virtual void OnEditChild() {}
 
-	virtual void OnCommit(bool bChanged) {}
+	virtual void OnCommit(bool bChanged);
 
 	virtual void OnCommitChild(bool bChanged) {}
 	
