Make sure replicated subobjects with empty rep layouts make their way into replay checkpoints

diff --git a/Engine/Source/Runtime/Engine/Private/DataChannel.cpp b/Engine/Source/Runtime/Engine/Private/DataChannel.cpp
--- a/Engine/Source/Runtime/Engine/Private/DataChannel.cpp
+++ b/Engine/Source/Runtime/Engine/Private/DataChannel.cpp
@@ -4905,7 +4905,12 @@ bool UActorChannel::WriteSubObjectInBunch(UObject* Obj, FOutBunch& Bunch, FRepli
 	
 	FReplicationFlags ObjRepFlags = RepFlags;
 	TSharedRef<FObjectReplicator>& ObjectReplicator = !bFoundReplicator ? CreateReplicator(Obj, !bFoundInvalidReplicator) : *FoundReplicator;
-	if (!bFoundReplicator)
+
+	// Special case for replay checkpoints where the replicator could already exist but we still need to consider this a new object even if it has an empty layout
+	const bool bNewToReplay = (Connection->ResendAllDataState == EResendAllDataState::SinceOpen)
+		|| ((Connection->ResendAllDataState == EResendAllDataState::SinceCheckpoint) && ObjectReplicator->IsDirtyForReplay());
+
+	if (!bFoundReplicator || bNewToReplay)
 	{
 		// This is the first time replicating this subobject
 		// This bunch should be reliable and we should always return true
diff --git a/Engine/Source/Runtime/Engine/Private/DataReplication.cpp b/Engine/Source/Runtime/Engine/Private/DataReplication.cpp
--- a/Engine/Source/Runtime/Engine/Private/DataReplication.cpp
+++ b/Engine/Source/Runtime/Engine/Private/DataReplication.cpp
@@ -317,6 +317,8 @@ FObjectReplicator::FObjectReplicator()
 	, bForceUpdateUnmapped(false)
 	, bHasReplicatedProperties(false)
 	, bSupportsFastArrayDelta(false)
+	, bCanUseNonDirtyOptimization(false)
+	, bDirtyForReplay(true)
 	, ObjectClass(nullptr)
 	, ObjectPtr(nullptr)
 	, Connection(nullptr)
@@ -1807,6 +1809,8 @@ bool FObjectReplicator::ReplicateProperties_r( FOutBunch & Bunch, FReplicationFl
 
 	if ( Connection->ResendAllDataState != EResendAllDataState::None )
 	{
+		bDirtyForReplay = false;
+
 		// If we are resending data since open, we don't want to affect the current state of channel/replication, so just send the data, and return
 		const bool bWroteImportantData = Writer.GetNumBits() != 0;
 
diff --git a/Engine/Source/Runtime/Engine/Public/Net/DataReplication.h b/Engine/Source/Runtime/Engine/Public/Net/DataReplication.h
--- a/Engine/Source/Runtime/Engine/Public/Net/DataReplication.h
+++ b/Engine/Source/Runtime/Engine/Public/Net/DataReplication.h
@@ -257,6 +257,8 @@ public:
 	 */
 	bool CanSkipUpdate(FReplicationFlags Flags);
 
+	bool IsDirtyForReplay() const { return bDirtyForReplay; }
+
 public:
 
 	/** Net GUID for the object we're replicating. */
@@ -288,6 +290,9 @@ private:
 	 */
 	uint32 bCanUseNonDirtyOptimization : 1;
 
+	/** Used to track if we've replicated this object into a checkpoint */
+	uint32 bDirtyForReplay : 1;
+
 public:
 	
 	TSharedPtr<class FReplicationChangelistMgr> ChangelistMgr;
