Implement proper preview for TextureCubeArray in the texture editor.

diff --git a/Engine/Shaders/Private/SimpleElementPixelShader.usf b/Engine/Shaders/Private/SimpleElementPixelShader.usf
--- a/Engine/Shaders/Private/SimpleElementPixelShader.usf
+++ b/Engine/Shaders/Private/SimpleElementPixelShader.usf
@@ -14,6 +14,10 @@
 #define ENABLE_EDITOR_COMPOSITING (FEATURE_LEVEL >= FEATURE_LEVEL_SM4 || MOBILE_EMULATION)
 #define WRITE_TO_SHADING_MODEL (FEATURE_LEVEL >= FEATURE_LEVEL_SM4 && !FORWARD_SHADING)
 
+#ifndef TEXTURECUBE_ARRAY
+#define TEXTURECUBE_ARRAY 0
+#endif
+
 Texture2D InTexture;
 SamplerState InTextureSampler;
 
@@ -305,7 +309,12 @@ void GammaDistanceFieldMain(
 	OutColor = SimpleElementEncodeFor32BPPHDR(OutColor);
 }
 
+#if TEXTURECUBE_ARRAY
+TextureCubeArray CubeTexture;
+#else
 TextureCube CubeTexture;
+#endif
+
 SamplerState CubeTextureSampler;
 
 float3 GetCubemapVector(float2 ScaledUVs, uint CubeFace)
@@ -345,6 +354,10 @@ float3 GetCubemapVector(float2 ScaledUVs, uint CubeFace)
 float4 PackedProperties0;
 float4x4 ColorWeights;
 
+#if TEXTURECUBE_ARRAY
+float NumSlices;
+#endif
+
 // used in the Texture info dialog when looking at a cubemap texture
 void CubemapTextureProperties(
 	in float2 TextureCoordinate : TEXCOORD0,
@@ -359,18 +372,43 @@ void CubemapTextureProperties(
 	float SliceIndex = PackedProperties0.y;
 	float4 BaseColor;
 
+#if TEXTURECUBE_ARRAY
+
+	// for a TextureCube array SliceIndex represents an index of a cubemap in the array
+	// each slice of a TextureCube array is displayed as a longlat unwrap
+
+	float2 Angles = float2(2 * PI * (TextureCoordinate.x + 0.5f), PI * TextureCoordinate.y);
+	float s = sin(Angles.y);
+	float3 Direction = float3(s * sin(Angles.x), cos(Angles.y), -s * cos(Angles.x));
+
+	float CurrentSliceIndex = SliceIndex;
+	if (CurrentSliceIndex < 0.0)
+	{
+		// if the slice index is not specified, display a portion of each cubemap from the TextureCube array
+		CurrentSliceIndex = floor(TextureCoordinate.y * NumSlices);
+	}
+
+	BaseColor = TextureCubeArraySampleLevel(CubeTexture, CubeTextureSampler, Direction.xzy, CurrentSliceIndex, MipLevel);
+
+#else //TEXTURECUBE_ARRAY
+
+	// for a TextureCube SliceIndex represents an index of a cubemap face
+
 	if (SliceIndex >= 0.0)
 	{
 		BaseColor = TextureCubeSampleLevel(CubeTexture, CubeTextureSampler, GetCubemapVector(TextureCoordinate * 2 - 1, SliceIndex), MipLevel);
 	}
 	else
 	{
+		// if the slice index is not specified, display the cubemap as a longlat unwrap
 		float2 Angles = float2(2 * PI * (TextureCoordinate.x + 0.5f), PI * TextureCoordinate.y);
 		float s = sin(Angles.y);
 		float3 Direction = float3(s * sin(Angles.x), cos(Angles.y), -s * cos(Angles.x));
 		BaseColor = TextureCubeSampleLevel(CubeTexture, CubeTextureSampler, Direction.xzy, MipLevel);
 	}
 
+#endif //TEXTURECUBE_ARRAY
+
 	float4 FinalColor;
 
 	// Seperate the Color weights and use against the Base colour to detrmine the actual colour from our filter
diff --git a/Engine/Source/Editor/TextureEditor/Private/Models/TextureEditorViewportClient.cpp b/Engine/Source/Editor/TextureEditor/Private/Models/TextureEditorViewportClient.cpp
--- a/Engine/Source/Editor/TextureEditor/Private/Models/TextureEditorViewportClient.cpp
+++ b/Engine/Source/Editor/TextureEditor/Private/Models/TextureEditorViewportClient.cpp
@@ -112,7 +112,7 @@ void FTextureEditorViewportClient::Draw(FViewport* Viewport, FCanvas* Canvas)
 	{
 		if (TextureCube || TextureCubeArray || RTTextureCube)
 		{
-			BatchedElementParameters = new FMipLevelBatchedElementParameters(MipLevel, SliceIndex, false);
+			BatchedElementParameters = new FMipLevelBatchedElementParameters(MipLevel, SliceIndex, TextureCubeArray != nullptr, false);
 		}
 		else if (VolumeTexture)
 		{
diff --git a/Engine/Source/Editor/TextureEditor/Private/TextureEditorToolkit.cpp b/Engine/Source/Editor/TextureEditor/Private/TextureEditorToolkit.cpp
--- a/Engine/Source/Editor/TextureEditor/Private/TextureEditorToolkit.cpp
+++ b/Engine/Source/Editor/TextureEditor/Private/TextureEditorToolkit.cpp
@@ -715,7 +715,8 @@ void FTextureEditorToolkit::PopulateQuickInfo( )
 
 	// Cubes are previewed as unwrapped 2D textures in case when slice index is not selected.
 	// These have 2x the width of a cube face.
-	if (IsCubeTexture() && GetSlice() < 0)
+	// Cube arrays are always displayed as longlat unwraps.
+	if (IsCubeTexture() && (IsArrayTexture() || GetSlice() < 0))
 	{
 		PreviewEffectiveTextureWidth *= 2;
 	}
@@ -2151,7 +2152,12 @@ int32 FTextureEditorToolkit::GetNumSlices() const
 	{
 		return Cast<UTextureRenderTarget2DArray>(Texture)->GetSurfaceArraySize();
 	}
-	else if (IsCubeTexture())
+	else if (Texture->IsA(UTextureCubeArray::StaticClass()))
+	{
+		// for a TextureCube array SliceIndex represents an index of a cubemap in the array
+		return Cast<UTextureCubeArray>(Texture)->GetSurfaceArraySize() / 6;
+	}
+	else if (IsCubeTexture() && !IsArrayTexture())
 	{
 		return 6;
 	}
diff --git a/Engine/Source/Editor/UnrealEd/Private/ThumbnailRendering/TextureThumbnailRenderer.cpp b/Engine/Source/Editor/UnrealEd/Private/ThumbnailRendering/TextureThumbnailRenderer.cpp
--- a/Engine/Source/Editor/UnrealEd/Private/ThumbnailRendering/TextureThumbnailRenderer.cpp
+++ b/Engine/Source/Editor/UnrealEd/Private/ThumbnailRendering/TextureThumbnailRenderer.cpp
@@ -72,7 +72,7 @@ void UTextureThumbnailRenderer::Draw(UObject* Object, int32 X, int32 Y, uint32 W
 		if(TextureCube || TextureCubeArray || RTTextureCube)
 		{
 			// is released by the render thread when it was rendered
-			BatchedElementParameters = new FMipLevelBatchedElementParameters((float)0, (float)-1);
+			BatchedElementParameters = new FMipLevelBatchedElementParameters((float)0, (float)-1, TextureCubeArray != nullptr);
 			
 			// If the thumbnail is square then make it 2:1 for cubes.
 			if(Width == Height)
diff --git a/Engine/Source/Runtime/Engine/Private/CubemapUnwrapUtils.cpp b/Engine/Source/Runtime/Engine/Private/CubemapUnwrapUtils.cpp
--- a/Engine/Source/Runtime/Engine/Private/CubemapUnwrapUtils.cpp
+++ b/Engine/Source/Runtime/Engine/Private/CubemapUnwrapUtils.cpp
@@ -37,7 +37,7 @@ namespace CubemapHelpers
 	bool GenerateLongLatUnwrap(const FTextureResource* TextureResource, const uint32 AxisDimenion, const EPixelFormat SourcePixelFormat, TArray64<uint8>& BitsOUT, FIntPoint& SizeOUT, EPixelFormat& FormatOUT)
 	{
 		TRefCountPtr<FBatchedElementParameters> BatchedElementParameters;
-		BatchedElementParameters = new FMipLevelBatchedElementParameters((float)0, (float)-1, true);
+		BatchedElementParameters = new FMipLevelBatchedElementParameters((float)0, (float)-1, false, true);
 		const FIntPoint LongLatDimensions(AxisDimenion * 2, AxisDimenion);
 
 		// If the source format is 8 bit per channel or less then select a LDR target format.
@@ -121,14 +121,25 @@ void FCubemapTexturePropertiesVS::SetParameters( FRHICommandList& RHICmdList, co
 }
 
 template<bool bHDROutput>
-void FCubemapTexturePropertiesPS<bHDROutput>::SetParameters(FRHICommandList& RHICmdList, const FTexture* Texture, const FMatrix& ColorWeightsValue, float MipLevel, float SliceIndex, float GammaValue)
+void FCubemapTexturePropertiesPS<bHDROutput>::SetParameters(FRHICommandList& RHICmdList, const FTexture* Texture, const FMatrix& ColorWeightsValue, float MipLevel, float SliceIndex, float GammaValue, bool bIsTextureCubeArray)
 {
-	SetTextureParameter(RHICmdList, RHICmdList.GetBoundPixelShader(),CubeTexture,CubeTextureSampler,Texture);
+	FRHIPixelShader* ShaderRHI = RHICmdList.GetBoundPixelShader();
+
+	SetTextureParameter(RHICmdList, ShaderRHI, CubeTexture, CubeTextureSampler, Texture);
 
 	FVector4f PackedProperties0Value(MipLevel, SliceIndex, 0, 0);
-	SetShaderValue(RHICmdList, RHICmdList.GetBoundPixelShader(), PackedProperties0, PackedProperties0Value);
-	SetShaderValue(RHICmdList, RHICmdList.GetBoundPixelShader(), ColorWeights, (FMatrix44f)ColorWeightsValue);
-	SetShaderValue(RHICmdList, RHICmdList.GetBoundPixelShader(), Gamma, GammaValue);
+	SetShaderValue(RHICmdList, ShaderRHI, PackedProperties0, PackedProperties0Value);
+	SetShaderValue(RHICmdList, ShaderRHI, ColorWeights, (FMatrix44f)ColorWeightsValue);
+	SetShaderValue(RHICmdList, ShaderRHI, Gamma, GammaValue);
+
+	// store slice count for the texture cube array
+	if (bIsTextureCubeArray)
+	{
+		// GetSizeZ() returns the total number of slices stored in the platform data
+		// for a TextureCube array this value is equal to the size of the array multiplied by 6
+		const float NumSlicesData = (float)(Texture ? FMath::Max((int32)Texture->GetSizeZ() / 6, 1) : 1);
+		SetShaderValue(RHICmdList, ShaderRHI, NumSlices, NumSlicesData);
+	}
 }
 
 void FMipLevelBatchedElementParameters::BindShaders(FRHICommandList& RHICmdList, FGraphicsPipelineStateInitializer& GraphicsPSOInit, ERHIFeatureLevel::Type InFeatureLevel, const FMatrix& InTransform, const float InGamma, const FMatrix& ColorWeights, const FTexture* Texture)
@@ -149,7 +160,10 @@ void FMipLevelBatchedElementParameters::BindShaders(FRHICommandList& RHICmdList,
 	GraphicsPSOInit.BlendState = TStaticBlendState<>::GetRHI();
 
 	TShaderMapRef<FCubemapTexturePropertiesVS> VertexShader(GetGlobalShaderMap(InFeatureLevel));
-	TShaderMapRef<TPixelShader> PixelShader(GetGlobalShaderMap(InFeatureLevel));
+
+	typename TPixelShader::FPermutationDomain PermutationVector;
+	PermutationVector.Set<typename TPixelShader::FCubemapArrayTexturePropertiesPS>(bIsTextureCubeArray);
+	TShaderMapRef<TPixelShader> PixelShader(GetGlobalShaderMap(InFeatureLevel), PermutationVector);
 	
 	GraphicsPSOInit.BoundShaderState.VertexDeclarationRHI = GSimpleElementVertexDeclaration.VertexDeclarationRHI;
 	GraphicsPSOInit.BoundShaderState.VertexShaderRHI = VertexShader.GetVertexShader();
@@ -160,7 +174,7 @@ void FMipLevelBatchedElementParameters::BindShaders(FRHICommandList& RHICmdList,
 	SetGraphicsPipelineState(RHICmdList, GraphicsPSOInit, 0);
 
 	VertexShader->SetParameters(RHICmdList, InTransform);
-	PixelShader->SetParameters(RHICmdList, Texture, ColorWeights, MipLevel, SliceIndex, InGamma);
+	PixelShader->SetParameters(RHICmdList, Texture, ColorWeights, MipLevel, SliceIndex, InGamma, bIsTextureCubeArray);
 }
 
 void FIESLightProfilePS::SetParameters( FRHICommandList& RHICmdList, const FTexture* Texture, float InBrightnessInLumens )
diff --git a/Engine/Source/Runtime/Engine/Public/CubemapUnwrapUtils.h b/Engine/Source/Runtime/Engine/Public/CubemapUnwrapUtils.h
--- a/Engine/Source/Runtime/Engine/Public/CubemapUnwrapUtils.h
+++ b/Engine/Source/Runtime/Engine/Public/CubemapUnwrapUtils.h
@@ -70,6 +70,10 @@ template<bool bHDROutput>
 class FCubemapTexturePropertiesPS : public FGlobalShader
 {
 	DECLARE_SHADER_TYPE(FCubemapTexturePropertiesPS,Global);
+
+	class FCubemapArrayTexturePropertiesPS : SHADER_PERMUTATION_BOOL("TEXTURECUBE_ARRAY");
+	using FPermutationDomain = TShaderPermutationDomain<FCubemapArrayTexturePropertiesPS>;
+
 public:
 
 	static bool ShouldCompilePermutation(const FGlobalShaderPermutationParameters& Parameters) { return IsPCPlatform(Parameters.Platform);}
@@ -82,10 +86,11 @@ public:
 		ColorWeights.Bind(Initializer.ParameterMap,TEXT("ColorWeights"));
 		PackedProperties0.Bind(Initializer.ParameterMap,TEXT("PackedProperties0"));
 		Gamma.Bind(Initializer.ParameterMap,TEXT("Gamma"));
+		NumSlices.Bind(Initializer.ParameterMap, TEXT("NumSlices"));
 	}
 	FCubemapTexturePropertiesPS() {}
 
-	void SetParameters(FRHICommandList& RHICmdList, const FTexture* Texture, const FMatrix& ColorWeightsValue, float MipLevel, float SliceIndex, float GammaValue);
+	void SetParameters(FRHICommandList& RHICmdList, const FTexture* Texture, const FMatrix& ColorWeightsValue, float MipLevel, float SliceIndex, float GammaValue, bool bIsTextureCubeArray);
 
 	static void ModifyCompilationEnvironment(const FGlobalShaderPermutationParameters& Parameters, FShaderCompilerEnvironment& OutEnvironment)
 	{
@@ -98,16 +103,18 @@ private:
 	LAYOUT_FIELD(FShaderParameter, PackedProperties0);
 	LAYOUT_FIELD(FShaderParameter, ColorWeights);
 	LAYOUT_FIELD(FShaderParameter, Gamma);
+	LAYOUT_FIELD(FShaderParameter, NumSlices);
 };
 
 
 class ENGINE_API FMipLevelBatchedElementParameters : public FBatchedElementParameters
 {
 public:
-	FMipLevelBatchedElementParameters(float InMipLevel, float InSliceIndex, bool bInHDROutput = false)
+	FMipLevelBatchedElementParameters(float InMipLevel, float InSliceIndex, bool bInIsTextureCubeArray, bool bInHDROutput = false)
 		: bHDROutput(bInHDROutput)
 		, MipLevel(InMipLevel)
 		, SliceIndex(InSliceIndex)
+		, bIsTextureCubeArray(bInIsTextureCubeArray)
 	{
 	}
 
@@ -118,9 +125,13 @@ private:
 	template<typename TPixelShader> void BindShaders(FRHICommandList& RHICmdList, FGraphicsPipelineStateInitializer& GraphicsPSOInit, ERHIFeatureLevel::Type InFeatureLevel, const FMatrix& InTransform, const float InGamma, const FMatrix& ColorWeights, const FTexture* Texture);
 
 	bool bHDROutput;
+
 	/** Parameters that need to be passed to the shader */
 	float MipLevel;
 	float SliceIndex;
+
+	/** Parameters that are used to select a shader permutation */
+	bool bIsTextureCubeArray;
 };
 
 
