Fix template compilation error on Linux.
#rnx

diff --git a/Engine/Source/Runtime/Engine/Private/CubemapUnwrapUtils.cpp b/Engine/Source/Runtime/Engine/Private/CubemapUnwrapUtils.cpp
--- a/Engine/Source/Runtime/Engine/Private/CubemapUnwrapUtils.cpp
+++ b/Engine/Source/Runtime/Engine/Private/CubemapUnwrapUtils.cpp
@@ -162,7 +162,7 @@ void FMipLevelBatchedElementParameters::BindShaders(FRHICommandList& RHICmdList,
 	TShaderMapRef<FCubemapTexturePropertiesVS> VertexShader(GetGlobalShaderMap(InFeatureLevel));
 
 	typename TPixelShader::FPermutationDomain PermutationVector;
-	PermutationVector.Set<typename TPixelShader::FCubemapArrayTexturePropertiesPS>(bIsTextureCubeArray);
+	PermutationVector.template Set<typename TPixelShader::FCubemapArrayTexturePropertiesPS>(bIsTextureCubeArray);
 	TShaderMapRef<TPixelShader> PixelShader(GetGlobalShaderMap(InFeatureLevel), PermutationVector);
 	
 	GraphicsPSOInit.BoundShaderState.VertexDeclarationRHI = GSimpleElementVertexDeclaration.VertexDeclarationRHI;
