Fix crash in FSubsystemCollectionBase::Deinitialize when hitting stale entries in SubsystemMap

diff --git a/Engine/Source/Runtime/Engine/Private/Subsystems/SubsystemCollection.cpp b/Engine/Source/Runtime/Engine/Private/Subsystems/SubsystemCollection.cpp
--- a/Engine/Source/Runtime/Engine/Private/Subsystems/SubsystemCollection.cpp
+++ b/Engine/Source/Runtime/Engine/Private/Subsystems/SubsystemCollection.cpp
@@ -173,28 +173,43 @@ void FSubsystemCollectionBase::Deinitialize()
 	check(IsInGameThread());
 
 	// already Deinitialize'd :
-	if ( Outer == nullptr )
+	if (Outer == nullptr)
+	{
 		return;
+	}
 
 	// Remove static tracking 
 	GlobalSubsystemCollections.Remove(this);
-	if (GlobalSubsystemCollections.Num() == 0)
+	if (GlobalSubsystemCollections.IsEmpty())
 	{
 		FSubsystemModuleWatcher::DeinitializeModuleWatcher();
 	}
 
 	// Deinit and clean up existing systems
-	SubsystemArrayMap.Empty();
-	for (auto Iter = SubsystemMap.CreateIterator(); Iter; ++Iter)
 	{
-		UClass* KeyClass = Iter.Key();
-		USubsystem* Subsystem = Iter.Value();
-		if ( Subsystem != nullptr && Subsystem->GetClass() == KeyClass)
+		// @note Iterate over SubsystemMapWeak to test whether subsystems have already been 
+		// deleted which can happen when Deinitialize() is called from the destructor
+		TMap<UClass*, TWeakObjectPtr<USubsystem>> SubsystemMapWeakCopy;
+		Swap(SubsystemMapWeakCopy, SubsystemMapWeak);
+		for (const TPair<UClass*, TWeakObjectPtr<USubsystem>>& Iter : SubsystemMapWeakCopy)
 		{
-			Subsystem->Deinitialize();
-			Subsystem->InternalOwningSubsystem = nullptr;
+			UClass* KeyClass = Iter.Key;
+			USubsystem* Subsystem = Iter.Value.GetEvenIfUnreachable();
+			if (Subsystem)
+			{
+				if ((Subsystem->GetClass() == KeyClass) && (Subsystem->InternalOwningSubsystem != nullptr))
+				{
+					Subsystem->Deinitialize();
+					Subsystem->InternalOwningSubsystem = nullptr;
+				}
+			}
+			else
+			{
+				UE_LOG(LogSubsystemCollection, Warning, TEXT("Cannot deinitialize %s because it has been deleted"), *KeyClass->GetName());
+			}
 		}
 	}
+	SubsystemArrayMap.Empty();
 	SubsystemMap.Empty();
 	Outer = nullptr;
 }
@@ -256,6 +271,7 @@ USubsystem* FSubsystemCollectionBase::AddAndInitializeSubsystem(UClass* Subsyste
 			{
 				USubsystem* Subsystem = NewObject<USubsystem>(Outer, SubsystemClass);
 				SubsystemMap.Add(SubsystemClass,Subsystem);
+				SubsystemMapWeak.Add(SubsystemClass, Subsystem);
 				Subsystem->InternalOwningSubsystem = this;
 				Subsystem->Initialize(*this);
 				return Subsystem;
@@ -276,6 +292,7 @@ void FSubsystemCollectionBase::RemoveAndDeinitializeSubsystem(USubsystem* Subsys
 	USubsystem* SubsystemFound = SubsystemMap.FindAndRemoveChecked(Subsystem->GetClass());
 	check(Subsystem == SubsystemFound);
 
+	SubsystemMapWeak.Remove(Subsystem->GetClass());
 	SubsystemArrayMap.Remove(Subsystem->GetClass());
 
 	Subsystem->Deinitialize();
diff --git a/Engine/Source/Runtime/Engine/Public/Subsystems/SubsystemCollection.h b/Engine/Source/Runtime/Engine/Public/Subsystems/SubsystemCollection.h
--- a/Engine/Source/Runtime/Engine/Public/Subsystems/SubsystemCollection.h
+++ b/Engine/Source/Runtime/Engine/Public/Subsystems/SubsystemCollection.h
@@ -87,6 +87,8 @@ private:
 
 	TMap<UClass*, USubsystem*> SubsystemMap;
 
+	TMap<UClass*, TWeakObjectPtr<USubsystem>> SubsystemMapWeak;
+
 	mutable TMap<UClass*, TArray<USubsystem*>> SubsystemArrayMap;
 
 	UClass* BaseType;
