Fix static analysis warning

diff --git a/Engine/Source/Runtime/Engine/Private/SkeletalMeshLODRenderData.cpp b/Engine/Source/Runtime/Engine/Private/SkeletalMeshLODRenderData.cpp
--- a/Engine/Source/Runtime/Engine/Private/SkeletalMeshLODRenderData.cpp
+++ b/Engine/Source/Runtime/Engine/Private/SkeletalMeshLODRenderData.cpp
@@ -624,7 +624,7 @@ void FSkeletalMeshLODRenderData::SerializeStreamedData(FArchive& Ar, USkinnedAss
 	Ar << SkinWeightProfilesData;
 	SkinWeightProfilesData.Init(&SkinWeightVertexBuffer);
 
-	if (Ar.IsLoading())
+	if (Ar.IsLoading() && Owner)
 	{
 		Owner->SetSkinWeightProfilesData(LODIdx, SkinWeightProfilesData);
 	}
@@ -632,7 +632,7 @@ void FSkeletalMeshLODRenderData::SerializeStreamedData(FArchive& Ar, USkinnedAss
 
 	// Determine if morph target data should be cooked out for this platform
 	EShaderPlatform MorphTargetShaderPlatform = GMaxRHIShaderPlatform;
-	bool bSerializeCompressedMorphTargets = Ar.IsSaving() && FMorphTargetVertexInfoBuffers::IsPlatformShaderSupported(MorphTargetShaderPlatform) && Owner->GetMorphTargets().Num() > 0;
+	bool bSerializeCompressedMorphTargets = Ar.IsSaving() && FMorphTargetVertexInfoBuffers::IsPlatformShaderSupported(MorphTargetShaderPlatform) && Owner && Owner->GetMorphTargets().Num() > 0;
 	if (Ar.IsCooking())
 	{
 		const ITargetPlatform* Platform = Ar.CookingTarget();
@@ -640,7 +640,7 @@ void FSkeletalMeshLODRenderData::SerializeStreamedData(FArchive& Ar, USkinnedAss
 		bSerializeCompressedMorphTargets = false;
 
 		// Make sure to avoid cooking the morph target data when build a server only executable
-		if (!Platform->IsServerOnly() && Owner->GetMorphTargets().Num() > 0)
+		if (!Platform->IsServerOnly() && Owner && Owner->GetMorphTargets().Num() > 0)
 		{
 			// Test if any of the supported shader formats supports SM5
 			TArray<FName> ShaderFormats;
@@ -662,7 +662,7 @@ void FSkeletalMeshLODRenderData::SerializeStreamedData(FArchive& Ar, USkinnedAss
 	{
 		Ar << bSerializeCompressedMorphTargets;
 
-		if (bSerializeCompressedMorphTargets)
+		if (bSerializeCompressedMorphTargets && Owner)
 		{
 #if WITH_EDITOR
 			if (Ar.IsSaving())
