Name all RHI buffers from a FStaticMeshLODResources via RHIBindDebugLabelName

diff --git a/Engine/Source/Runtime/Engine/Private/StaticMesh.cpp b/Engine/Source/Runtime/Engine/Private/StaticMesh.cpp
--- a/Engine/Source/Runtime/Engine/Private/StaticMesh.cpp
+++ b/Engine/Source/Runtime/Engine/Private/StaticMesh.cpp
@@ -1361,6 +1361,39 @@ void FStaticMeshLODResources::InitResources(UStaticMesh* Parent)
 		UpdateVertexMemoryStats<true>();
 	});
 #endif
+
+#if !(UE_BUILD_SHIPPING || UE_BUILD_TEST)
+	ENQUEUE_RENDER_COMMAND(NameRHIResources)(
+		[this, DebugName = Parent->GetFName()](FRHICommandListImmediate&)
+	{
+		TStringBuilder<512> StringBuilder;
+		auto SetDebugName = [&StringBuilder, DebugName](FRHIBuffer* RHIBuffer, const TCHAR* Extension)
+		{
+			if (RHIBuffer)
+			{
+				StringBuilder.Reset();
+				DebugName.ToString(StringBuilder);
+				StringBuilder.Append(Extension);
+				RHIBindDebugLabelName(RHIBuffer, StringBuilder.ToString());
+			}
+		};
+		SetDebugName(IndexBuffer.IndexBufferRHI, TEXT("_IB"));
+		SetDebugName(VertexBuffers.PositionVertexBuffer.VertexBufferRHI, TEXT("_VB"));
+		SetDebugName(VertexBuffers.StaticMeshVertexBuffer.TexCoordVertexBuffer.VertexBufferRHI, TEXT("_TC"));
+		SetDebugName(VertexBuffers.StaticMeshVertexBuffer.TangentsVertexBuffer.VertexBufferRHI, TEXT("_TB"));
+		SetDebugName(VertexBuffers.ColorVertexBuffer.VertexBufferRHI, TEXT("_CB"));
+
+		SetDebugName(DepthOnlyIndexBuffer.IndexBufferRHI, TEXT("_IB_DepthOnly"));
+		if (AdditionalIndexBuffers)
+		{
+			SetDebugName(AdditionalIndexBuffers->WireframeIndexBuffer.IndexBufferRHI, TEXT("_IB_WireFrame"));
+			SetDebugName(AdditionalIndexBuffers->ReversedIndexBuffer.IndexBufferRHI, TEXT("_IB_Reversed"));
+			SetDebugName(AdditionalIndexBuffers->ReversedDepthOnlyIndexBuffer.IndexBufferRHI, TEXT("_IB_DepthOnly_Reversed"));
+		}
+		SetDebugName(AreaWeightedSectionSamplersBuffer.GetBufferRHI(), TEXT("_AreaWeightSectionSamplers"));
+	});
+#endif // !(UE_BUILD_SHIPPING || UE_BUILD_TEST)
+
 }
 
 #if RHI_RAYTRACING
diff --git a/Engine/Source/Runtime/Engine/Public/StaticMeshResources.h b/Engine/Source/Runtime/Engine/Public/StaticMeshResources.h
--- a/Engine/Source/Runtime/Engine/Public/StaticMeshResources.h
+++ b/Engine/Source/Runtime/Engine/Public/StaticMeshResources.h
@@ -276,6 +276,7 @@ public:
 	ENGINE_API virtual void ReleaseRHI() override;
 	virtual FString GetFriendlyName() const override { return TEXT("FStaticMeshSectionAreaWeightedTriangleSamplerBuffer"); }
 
+	ENGINE_API const FBufferRHIRef& GetBufferRHI() const { return BufferSectionTriangleRHI; }
 	ENGINE_API const FShaderResourceViewRHIRef& GetBufferSRV() const { return BufferSectionTriangleSRV; }
 
 private:
