fix for streamed seek failures in LPCM decoder #preflight 62b048a9de510ae06e403d38 #jira UE-157058

diff --git a/Engine/Source/Runtime/Engine/Private/ADPCMAudioInfo.cpp b/Engine/Source/Runtime/Engine/Private/ADPCMAudioInfo.cpp
--- a/Engine/Source/Runtime/Engine/Private/ADPCMAudioInfo.cpp
+++ b/Engine/Source/Runtime/Engine/Private/ADPCMAudioInfo.cpp
@@ -202,19 +202,28 @@ void FADPCMAudioInfo::SeekToTimeInternal(const float InSeekTime)
 		else if (Format == WAVE_FORMAT_LPCM)
 		{
 			const int32 ChannelBlockSize = sizeof(int16) * NumChannels;
-
-			uint32 TotalByteSizeToSeek = TotalSamplesStreamed * ChannelBlockSize;		
-			uint32 SizeOfCurrentChunk = StreamingSoundWave->GetSizeOfChunk(CurrentChunkIndex);
+			const uint32 TotalByteSizeToSeek = TotalSamplesStreamed * ChannelBlockSize;	
+			
+			uint32 TargetChunkIndex = FirstChunkSampleDataIndex;
+			uint32 SizeOfTargetChunk = StreamingSoundWave->GetSizeOfChunk(TargetChunkIndex);
+			
 			uint32 CurrentBytes = 0;
-			while (CurrentBytes + SizeOfCurrentChunk < TotalByteSizeToSeek)
+			while (CurrentBytes + SizeOfTargetChunk < TotalByteSizeToSeek)
 			{
-				CurrentChunkIndex++;
-				CurrentBytes += SizeOfCurrentChunk;
-				CurrentChunkBufferOffset -= SizeOfCurrentChunk;
-				SizeOfCurrentChunk = StreamingSoundWave->GetSizeOfChunk(CurrentChunkIndex);
+				if (!ensure(TargetChunkIndex >= TotalStreamingChunks))
+				{
+					break;
+				}
+
+				++TargetChunkIndex;
+				CurrentBytes += SizeOfTargetChunk;
+				CurrentChunkBufferOffset -= SizeOfTargetChunk;
+				SizeOfTargetChunk = StreamingSoundWave->GetSizeOfChunk(TargetChunkIndex);
 			}
 
 			check(CurrentBytes <= TotalByteSizeToSeek);
+
+			CurrentChunkIndex = TargetChunkIndex;
 			CurrentChunkBufferOffset = TotalByteSizeToSeek - CurrentBytes;
 			CurrentChunkBufferOffset -= CurrentChunkBufferOffset % ChannelBlockSize;
 		}
@@ -1025,4 +1034,4 @@ const uint8* FADPCMAudioInfo::GetLoadedChunk(const FSoundWaveProxyPtr& InSoundWa
 		
 		return CurCompressedChunkHandle.GetData();
 	}
-}
+}
\ No newline at end of file
