Revert CL 20735243 #trivial #preflight none

diff --git a/Engine/Source/Runtime/Engine/Private/ADPCMAudioInfo.cpp b/Engine/Source/Runtime/Engine/Private/ADPCMAudioInfo.cpp
--- a/Engine/Source/Runtime/Engine/Private/ADPCMAudioInfo.cpp
+++ b/Engine/Source/Runtime/Engine/Private/ADPCMAudioInfo.cpp
@@ -202,28 +202,19 @@ void FADPCMAudioInfo::SeekToTimeInternal(const float InSeekTime)
 		else if (Format == WAVE_FORMAT_LPCM)
 		{
 			const int32 ChannelBlockSize = sizeof(int16) * NumChannels;
-			const uint32 TotalByteSizeToSeek = TotalSamplesStreamed * ChannelBlockSize;	
-			
-			uint32 TargetChunkIndex = FirstChunkSampleDataIndex;
-			uint32 SizeOfTargetChunk = StreamingSoundWave->GetSizeOfChunk(TargetChunkIndex);
-			
+
+			uint32 TotalByteSizeToSeek = TotalSamplesStreamed * ChannelBlockSize;
+			uint32 SizeOfCurrentChunk = StreamingSoundWave->GetSizeOfChunk(CurrentChunkIndex);
 			uint32 CurrentBytes = 0;
-			while (CurrentBytes + SizeOfTargetChunk < TotalByteSizeToSeek)
+			while (CurrentBytes + SizeOfCurrentChunk < TotalByteSizeToSeek)
 			{
-				if (!ensure(TargetChunkIndex >= TotalStreamingChunks))
-				{
-					break;
-				}
-
-				++TargetChunkIndex;
-				CurrentBytes += SizeOfTargetChunk;
-				CurrentChunkBufferOffset -= SizeOfTargetChunk;
-				SizeOfTargetChunk = StreamingSoundWave->GetSizeOfChunk(TargetChunkIndex);
+				CurrentChunkIndex++;
+				CurrentBytes += SizeOfCurrentChunk;
+				CurrentChunkBufferOffset -= SizeOfCurrentChunk;
+				SizeOfCurrentChunk = StreamingSoundWave->GetSizeOfChunk(CurrentChunkIndex);
 			}
 
 			check(CurrentBytes <= TotalByteSizeToSeek);
-
-			CurrentChunkIndex = TargetChunkIndex;
 			CurrentChunkBufferOffset = TotalByteSizeToSeek - CurrentBytes;
 			CurrentChunkBufferOffset -= CurrentChunkBufferOffset % ChannelBlockSize;
 		}
