PR #9179: Add SetSourceCubemapAngle function to SkyLight Component (Contributed by [REDACTED_PERSON])

diff --git a/Engine/Source/Runtime/Engine/Classes/Components/SkyLightComponent.h b/Engine/Source/Runtime/Engine/Classes/Components/SkyLightComponent.h
--- a/Engine/Source/Runtime/Engine/Classes/Components/SkyLightComponent.h
+++ b/Engine/Source/Runtime/Engine/Classes/Components/SkyLightComponent.h
@@ -246,6 +246,7 @@ class ENGINE_API USkyLightComponent : public ULightComponentBase
 	UFUNCTION(BlueprintCallable, Category="SkyLight")
 	void SetCubemap(UTextureCube* NewCubemap);
 
+	/** Sets the angle of the cubemap used when SourceType is set to SpecifiedCubemap and it is non static. It will cause the skylight to update on the next tick. */
 	UFUNCTION(BlueprintCallable, Category = "Rendering|Components|SkyLight")
 	void SetSourceCubemapAngle(float NewValue);
 
diff --git a/Engine/Source/Runtime/Engine/Private/Components/SkyLightComponent.cpp b/Engine/Source/Runtime/Engine/Private/Components/SkyLightComponent.cpp
--- a/Engine/Source/Runtime/Engine/Private/Components/SkyLightComponent.cpp
+++ b/Engine/Source/Runtime/Engine/Private/Components/SkyLightComponent.cpp
@@ -986,15 +986,14 @@ void USkyLightComponent::SetCubemap(UTextureCube* NewCubemap)
 
 void USkyLightComponent::SetSourceCubemapAngle(float NewValue)
 {
+	// Can't set on a static light
 	if (AreDynamicDataChangesAllowed()
 		&& SourceCubemapAngle != NewValue)
 	{
 		SourceCubemapAngle = NewValue;
-
-		if (GetWorld())
-		{
-			GetWorld()->UpdateAllSkyCaptures();
-		}
+		MarkRenderStateDirty();
+		// Note: this will cause the cubemap to be reprocessed including readback from the GPU
+		SetCaptureIsDirty();
 	}
 }
 
