fix for streamed seek failures in LPCM decoder #jira UE-157058 #preflight 62b074a7de510ae06e4c4620

diff --git a/Engine/Source/Runtime/Engine/Private/ADPCMAudioInfo.cpp b/Engine/Source/Runtime/Engine/Private/ADPCMAudioInfo.cpp
--- a/Engine/Source/Runtime/Engine/Private/ADPCMAudioInfo.cpp
+++ b/Engine/Source/Runtime/Engine/Private/ADPCMAudioInfo.cpp
@@ -202,19 +202,31 @@ void FADPCMAudioInfo::SeekToTimeInternal(const float InSeekTime)
 		else if (Format == WAVE_FORMAT_LPCM)
 		{
 			const int32 ChannelBlockSize = sizeof(int16) * NumChannels;
-
-			uint32 TotalByteSizeToSeek = TotalSamplesStreamed * ChannelBlockSize;
-			uint32 SizeOfCurrentChunk = StreamingSoundWave->GetSizeOfChunk(CurrentChunkIndex);
+			const uint32 TotalByteSizeToSeek = TotalSamplesStreamed * ChannelBlockSize;	
+			
+			uint32 TargetChunkIndex = FirstChunkSampleDataIndex;
+			uint32 SizeOfTargetChunk = StreamingSoundWave->GetSizeOfChunk(TargetChunkIndex);
+			
 			uint32 CurrentBytes = 0;
-			while (CurrentBytes + SizeOfCurrentChunk < TotalByteSizeToSeek)
+			while (CurrentBytes + SizeOfTargetChunk < TotalByteSizeToSeek)
 			{
-				CurrentChunkIndex++;
-				CurrentBytes += SizeOfCurrentChunk;
-				CurrentChunkBufferOffset -= SizeOfCurrentChunk;
-				SizeOfCurrentChunk = StreamingSoundWave->GetSizeOfChunk(CurrentChunkIndex);
+				const bool TargetChunkInRange = TargetChunkIndex < TotalStreamingChunks;
+				ensure(TargetChunkInRange);
+
+				if (!TargetChunkInRange)
+				{
+					break;
+				}
+
+				++TargetChunkIndex;
+				CurrentBytes += SizeOfTargetChunk;
+				CurrentChunkBufferOffset -= SizeOfTargetChunk;
+				SizeOfTargetChunk = StreamingSoundWave->GetSizeOfChunk(TargetChunkIndex);
 			}
 
 			check(CurrentBytes <= TotalByteSizeToSeek);
+
+			CurrentChunkIndex = TargetChunkIndex;
 			CurrentChunkBufferOffset = TotalByteSizeToSeek - CurrentBytes;
 			CurrentChunkBufferOffset -= CurrentChunkBufferOffset % ChannelBlockSize;
 		}
