[KillPendingKill] Release references from LODActor when SubActors are destroyed.

diff --git a/Engine/Source/Runtime/Engine/Classes/Engine/LODActor.h b/Engine/Source/Runtime/Engine/Classes/Engine/LODActor.h
--- a/Engine/Source/Runtime/Engine/Classes/Engine/LODActor.h
+++ b/Engine/Source/Runtime/Engine/Classes/Engine/LODActor.h
@@ -350,6 +350,9 @@ private:
 	void RegisterMeshComponents();
 	void UnregisterMeshComponents();
 
+	UFUNCTION()
+	void OnSubActorEndPlay(AActor* Actor, EEndPlayReason::Type Reason);
+
 #if WITH_EDITOR
 	// Get/Create the LOD instanced static mesh component for a given imposter material.
 	UInstancedStaticMeshComponent* GetISMComponent(const FHLODInstancingKey& InstancingKey) const;
diff --git a/Engine/Source/Runtime/Engine/Private/LODActor.cpp b/Engine/Source/Runtime/Engine/Private/LODActor.cpp
--- a/Engine/Source/Runtime/Engine/Private/LODActor.cpp
+++ b/Engine/Source/Runtime/Engine/Private/LODActor.cpp
@@ -309,7 +309,7 @@ void ALODActor::PostLoad()
 			{
 					// Make the parent HLOD
 					ParentLODActor->SubActors.Remove(this);
-					ParentLODActor->SubActors.Append(SubActors);
+					ParentLODActor->SubActors.Append(SubActors); // Don't register callbacks here, PostLoad should happen before PostRegisterAllComponents
 					for (AActor* Actor : SubActors)
 					{
 						if (Actor)
@@ -528,6 +528,17 @@ void ALODActor::PostRegisterAllComponents()
 		UnregisterMeshComponents();
 	}
 
+	if( UWorld* World = GetWorld(); World && World->IsGameWorld())
+	{
+		for (TObjectPtr<AActor> ActorPtr : SubActors)
+		{
+			if (AActor* Actor = ActorPtr.Get())
+			{
+				Actor->OnEndPlay.AddDynamic(this, &ALODActor::OnSubActorEndPlay);
+			}		
+		}
+	}
+
 #if WITH_EDITOR
 	if(!GetWorld()->IsPlayInEditor())
 	{
@@ -539,6 +550,15 @@ void ALODActor::PostRegisterAllComponents()
 #endif
 }
 
+void ALODActor::OnSubActorEndPlay(AActor* Actor, EEndPlayReason::Type Reason)
+{
+	// Other end play reasons will also be removing this actor from play so we don't need to touch our array
+	if (Reason == EEndPlayReason::Destroyed)
+	{
+		SubActors.RemoveSwap(Actor);
+	}
+}
+
 void ALODActor::RegisterMeshComponents()
 {
 	ForEachComponent<UStaticMeshComponent>(false, [](UStaticMeshComponent* SMComponent)
@@ -797,6 +817,16 @@ void ALODActor::AddSubActor(AActor* InActor)
 
 void ALODActor::AddSubActors(const TArray<AActor*>& InActors)
 {
+	if (UWorld* World = GetWorld(); World && World->IsGameWorld())
+	{
+		for (AActor* Actor : InActors)
+		{
+			if (Actor)
+			{
+				Actor->OnEndPlay.AddDynamic(this, &ALODActor::OnSubActorEndPlay);
+			}
+		}
+	}
 	SubActors.Append(InActors);
 
 	float LODDrawDistanceWithOverride = GetLODDrawDistanceWithOverride();
@@ -839,6 +869,7 @@ const bool ALODActor::RemoveSubActor(AActor* InActor)
 {
 	if ((InActor != nullptr) && SubActors.Contains(InActor))
 	{
+		InActor->OnEndPlay.RemoveAll(this);
 		SubActors.Remove(InActor);
 		InActor->SetLODParent(nullptr, 0);
 
@@ -1072,7 +1103,7 @@ void ALODActor::CleanSubActorArray()
 	for (int32 SubActorIndex = 0; SubActorIndex < SubActors.Num(); ++SubActorIndex)
 	{
 		AActor* Actor = SubActors[SubActorIndex];
-		if (Actor == nullptr)
+		if (!IsValid(Actor))
 		{
 			SubActors.RemoveAtSwap(SubActorIndex);
 			SubActorIndex--;
