allow consequential seeks to the same timestamp on LPCM/ADPCM decoder #preflight 62b184f7827ccccb2cbf2786 #jira UE-157058

diff --git a/Engine/Source/Runtime/Engine/Private/ADPCMAudioInfo.cpp b/Engine/Source/Runtime/Engine/Private/ADPCMAudioInfo.cpp
--- a/Engine/Source/Runtime/Engine/Private/ADPCMAudioInfo.cpp
+++ b/Engine/Source/Runtime/Engine/Private/ADPCMAudioInfo.cpp
@@ -9,13 +9,6 @@
 #include "Audio.h"
 #include "ContentStreaming.h"
 
-static int32 bDisableADPCMSeekLockCVar = 0;
-FAutoConsoleVariableRef CVarDisableADPCMSeekLock(
-	TEXT("au.DisableADPCMSeekLock"),
-	bDisableADPCMSeekLockCVar,
-	TEXT("Disables ADPCM seek crit section fix for multiple seek requests per frame.\n"),
-	ECVF_Default);
-
 static int32 bDisableADPCMSeekingCVar = 0;
 FAutoConsoleVariableRef CVarDisableADPCMSeeking(
 	TEXT("au.adpcm.DisableSeeking"),
@@ -75,10 +68,10 @@ FADPCMAudioInfo::FADPCMAudioInfo(void)
 	, FirstChunkSampleDataOffset(0)
 	, FirstChunkSampleDataIndex(0)
 	, bDecompressorReleased(false)
-	, bSeekPending(false)
+	, bNewSeekRequest(false)
+	, bSeekPendingRead(false)
 	, bSeekedFowardToNextChunk(false)
 	, TargetSeekTime(0.0f)
-	, LastSeekTime(0.0f)
 {
 }
 
@@ -98,20 +91,12 @@ void FADPCMAudioInfo::SeekToTime(const float InSeekTime)
 		return;
 	}
 
-	if (bDisableADPCMSeekLockCVar)
-	{
-		SeekToTimeInternal(InSeekTime);
-	}
-	else
-	{
-		TargetSeekTime = InSeekTime;
-	}
+	TargetSeekTime = InSeekTime;
+	bNewSeekRequest = true;
 }
 
 void FADPCMAudioInfo::SeekToTimeInternal(const float InSeekTime)
 {
-	LastSeekTime = InSeekTime;
-
 	// Reset chunk handle in preperation for a new chunk.
 	CurCompressedChunkData = nullptr;
 
@@ -124,8 +109,9 @@ void FADPCMAudioInfo::SeekToTimeInternal(const float InSeekTime)
 		CurrentChunkIndex = FirstChunkSampleDataIndex;
 		CurrentChunkBufferOffset = FirstChunkSampleDataOffset;
 		TotalSamplesStreamed = 0;
-
-		bSeekPending = false;
+		
+		bNewSeekRequest = false;
+		bSeekPendingRead = false;
 		return;
 	}
 
@@ -237,7 +223,7 @@ void FADPCMAudioInfo::SeekToTimeInternal(const float InSeekTime)
 			return;
 		}
 	}
-	bSeekPending = true;
+	bSeekPendingRead = true;
 }
 
 bool FADPCMAudioInfo::ReadCompressedInfo(const uint8* InSrcBufferData, uint32 InSrcBufferDataSize, struct FSoundQualityInfo* QualityInfo)
@@ -459,20 +445,13 @@ int FADPCMAudioInfo::GetStreamBufferSize() const
 
 void FADPCMAudioInfo::ProcessSeekRequest()
 {
-	if (bDisableADPCMSeekLockCVar)
-	{
-		return;
-	}
-
 	float NewSeekTime = -1.0f;
+	if (bNewSeekRequest)
 	{
-		FScopeLock Lock(&StreamSeekCriticalSection);
-		if (!FMath::IsNearlyEqual(TargetSeekTime, LastSeekTime))
-		{
-			NewSeekTime = TargetSeekTime;
-		}
+		NewSeekTime = TargetSeekTime;
 	}
 
+
 	if (NewSeekTime >= 0.0f)
 	{
 		SeekToTimeInternal(NewSeekTime);
@@ -660,14 +639,14 @@ bool FADPCMAudioInfo::StreamCompressedData(uint8* Destination, bool bLooping, ui
 						if (!bSeekedFowardToNextChunk)
 						{
 							// Ensure that we're either seeking or moving sequentially through the file.
-							ensureAlways(bSeekPending || CurrentChunkIndex == PreviouslyRequestedChunkIndex);
+							check(bSeekPendingRead || CurrentChunkIndex == PreviouslyRequestedChunkIndex);
 							++CurrentChunkIndex;
 
 							// Set the current buffer offset accounting for the header in the first chunk
-							if (!bSeekPending)
+							if (!bSeekPendingRead)
 							{
 								// Ensure that, if we hit this code, we are not at the beginning of the file.
-								ensureAlways(CurrentChunkIndex != FirstChunkSampleDataIndex);
+								check(CurrentChunkIndex != FirstChunkSampleDataIndex);
 								CurrentChunkBufferOffset = 0;
 							}
 						}
@@ -699,7 +678,7 @@ bool FADPCMAudioInfo::StreamCompressedData(uint8* Destination, bool bLooping, ui
 					if(CurCompressedChunkData == nullptr)
 					{
 						// We only need to worry about missing the stream chunk if we were seeking. Seeking might cause a bit of latency with chunk loading. That is expected.
-						if (!bSeekPending)
+						if (!bSeekPendingRead)
 						{
 							// If we did not load chunk then bail. CurrentChunkIndex will not get incremented on the next callback so in effect another attempt will be made to fetch the chunk.
 							// Since audio streaming depends on the general data streaming mechanism used by other parts of the engine and new data is pre-fetched on the game thread, its
@@ -710,7 +689,7 @@ bool FADPCMAudioInfo::StreamCompressedData(uint8* Destination, bool bLooping, ui
 						// If we have a seek pending, or we're in the middle of playing back the file, seek forward in the chunk offset to where we would have been.
 						// If we've already seeked forward one chunk and haven't loaded the next chunk in time and are skipping to the next one, we stop seeking to prevent snowballing chunk load requests.
 						const bool bAlreadySeekedFowardOneChunk = ADPCMOnlySeekForwardOneChunkCVar && bSeekedFowardToNextChunk;
-						const bool bShouldSeekForwardOnMissedChunk = !ADPCMDisableSeekForwardOnChunkMissesCVar && !bAlreadySeekedFowardOneChunk && (bSeekPending || CurrentChunkIndex > FirstChunkSampleDataIndex);
+						const bool bShouldSeekForwardOnMissedChunk = !ADPCMDisableSeekForwardOnChunkMissesCVar && !bAlreadySeekedFowardOneChunk && (bSeekPendingRead || CurrentChunkIndex > FirstChunkSampleDataIndex);
 
 						if (bShouldSeekForwardOnMissedChunk)
 						{
@@ -749,7 +728,7 @@ bool FADPCMAudioInfo::StreamCompressedData(uint8* Destination, bool bLooping, ui
 								ensureAlways(!bSeekedFowardToNextChunk);
 
 								// Ensure that we are either seeking of moving sequentially through the file.
-								ensureAlwaysMsgf(bSeekPending || CurrentChunkIndex == PreviouslyRequestedChunkIndex, TEXT("Failed to load a chunk of ADPCM audio for the entire duration that chunk of audio was supposed to be played."));
+								ensureAlwaysMsgf(bSeekPendingRead || CurrentChunkIndex == PreviouslyRequestedChunkIndex, TEXT("Failed to load a chunk of ADPCM audio for the entire duration that chunk of audio was supposed to be played."));
 								++CurrentChunkIndex;
 								CurrentChunkBufferOffset -= SizeOfCurrentChunk;
 
@@ -789,13 +768,14 @@ bool FADPCMAudioInfo::StreamCompressedData(uint8* Destination, bool bLooping, ui
 					}
 
 					
-					if (CurrentChunkIndex == FirstChunkSampleDataIndex && !bSeekPending)
+					if (CurrentChunkIndex == FirstChunkSampleDataIndex && !bSeekPendingRead)
 					{
 						// If we're in the first chunk, set the current buffer offset accounting for the header.
 						CurrentChunkBufferOffset = FirstChunkSampleDataOffset;
 					}
 
-					bSeekPending = false;
+					bSeekPendingRead = false;
+					bNewSeekRequest = false;
 				}
 
 				// Decompress one block for each channel and store it in UncompressedBlockData
@@ -897,7 +877,7 @@ bool FADPCMAudioInfo::StreamCompressedData(uint8* Destination, bool bLooping, ui
 				{
 					// Only report missing the stream chunk if we were seeking. Seeking
 					// may cause a bit of latency with chunk loading, which is expected.
-					if (!bSeekPending)
+					if (!bSeekPendingRead)
 					{
 						// CurrentChunkIndex will not get incremented on the next callback, effectively causing another attempt to fetch the chunk.
 						// Since audio streaming depends on the general data streaming mechanism used by other parts of the engine and new data is
@@ -916,12 +896,13 @@ bool FADPCMAudioInfo::StreamCompressedData(uint8* Destination, bool bLooping, ui
 				NumConsecutiveReadFailiures = 0;
 
 				// Set the current buffer offset accounting for the header in the first chunk
-				if (!bSeekPending)
+				if (!bSeekPendingRead)
 				{
 					CurrentChunkBufferOffset = CurrentChunkIndex == FirstChunkSampleDataIndex ? FirstChunkSampleDataOffset : 0;
 				}
 
-				bSeekPending = false;
+				bSeekPendingRead = false;
+				bNewSeekRequest = false;
 			}
 			
 			uint32 DecompressedSamplesToCopy = FMath::Min<uint32>(
@@ -1027,7 +1008,7 @@ const uint8* FADPCMAudioInfo::GetLoadedChunk(const FSoundWaveProxyPtr& InSoundWa
 	}
 	else
 	{
-		const bool bIsSeekingOrLooping = bSeekPending || ChunkIndex == FirstChunkSampleDataIndex;
+		const bool bIsSeekingOrLooping = bSeekPendingRead || ChunkIndex == FirstChunkSampleDataIndex;
 		ensureAlwaysMsgf(bIsSeekingOrLooping || ChunkIndex == PreviouslyRequestedChunkIndex || ChunkIndex == PreviouslyRequestedChunkIndex + 1, TEXT("ADPCM playback error! We skipped from the end of chunk %d to the beginning of chunk %d."), PreviouslyRequestedChunkIndex, ChunkIndex);
 
 		CurCompressedChunkHandle = IStreamingManager::Get().GetAudioStreamingManager().GetLoadedChunk(InSoundWave, ChunkIndex, false, true);
diff --git a/Engine/Source/Runtime/Engine/Public/ADPCMAudioInfo.h b/Engine/Source/Runtime/Engine/Public/ADPCMAudioInfo.h
--- a/Engine/Source/Runtime/Engine/Public/ADPCMAudioInfo.h
+++ b/Engine/Source/Runtime/Engine/Public/ADPCMAudioInfo.h
@@ -173,15 +173,15 @@ private:
 
 	uint32			CurrentCompressedBlockIndex;		// For non disk streaming - the current compressed block in the compressed source data
 	uint32			TotalCompressedBlocksPerChannel;	// For non disk streaming - the total number of compressed blocks per channel
-	uint8			bSeekPending : 1;					// Whether or not seek has been requested and pending a read
+	uint8			bNewSeekRequest : 1;				// Whether or not there is a new seek request to be processed 
+	uint8			bSeekPendingRead : 1;				// Whether or not a requested seek is pending a read
 	uint8           bSeekedFowardToNextChunk : 1;       // If this is true, we have already seeked forward to the next chunk while waiting for the current chunk of audio to load.
 
 private:
 	void ProcessSeekRequest();
 	void SeekToTimeInternal(const float InSeekTime);
 
-	float TargetSeekTime;
-	float LastSeekTime;
+	std::atomic<float> TargetSeekTime;
 
 	FCriticalSection StreamSeekCriticalSection;
 };
