Fix Screenshot tracing naming issues

diff --git a/Engine/Source/Runtime/Engine/Private/ProfilingDebugging/TraceScreenshot.cpp b/Engine/Source/Runtime/Engine/Private/ProfilingDebugging/TraceScreenshot.cpp
--- a/Engine/Source/Runtime/Engine/Private/ProfilingDebugging/TraceScreenshot.cpp
+++ b/Engine/Source/Runtime/Engine/Private/ProfilingDebugging/TraceScreenshot.cpp
@@ -9,6 +9,7 @@
 #include "HAL/PlatformTLS.h"
 #include "HAL/TlsAutoCleanup.h"
 #include "Misc/DateTime.h"
+#include "Misc/Paths.h"
 #include "ProfilingDebugging/MiscTrace.h"
 #include "Trace/Trace.inl"
 
@@ -31,25 +32,27 @@ static FAutoConsoleCommand TraceScreenshotCmd(
 	FConsoleCommandWithArgsDelegate::CreateStatic(TraceScreenshotCommandCallback)
 );
 
-FString FTraceScreenshot::RequestedScreenshotName;
 bool FTraceScreenshot::bSuppressWritingToFile = false;
 
 void FTraceScreenshot::RequestScreenshot(FString Name)
 {
-	RequestedScreenshotName = Name;
-	bool bShowUI = false;
 	bSuppressWritingToFile = true;
-	FScreenshotRequest::RequestScreenshot(bShowUI);
+
+	if (Name.IsEmpty())
+	{
+		Name = FDateTime::Now().ToString(TEXT("Screenshot_%Y%m%d_%H%M%S"));
+	}
+
+	constexpr bool bShowUI = false;
+	constexpr bool bAddUniqueSuffix = false;
+	FScreenshotRequest::RequestScreenshot(Name, bShowUI, bAddUniqueSuffix);
 }
 
 void FTraceScreenshot::TraceScreenshot(int32 InSizeX, int32 InSizeY, const TArray<FColor>& InImageData, const FString& InScreenshotName, int32 DesiredX)
 {
-	FString ScreenshotName = InScreenshotName;
-	if (!RequestedScreenshotName.IsEmpty())
-	{
-		ScreenshotName = RequestedScreenshotName;
-	}
-	else if (ScreenshotName.IsEmpty())
+	FString ScreenshotName = FPaths::GetBaseFilename(InScreenshotName);
+
+	if (ScreenshotName.IsEmpty())
 	{
 		ScreenshotName = FDateTime::Now().ToString(TEXT("Screenshot_%Y%m%d_%H%M%S"));
 	}
@@ -81,6 +84,5 @@ void FTraceScreenshot::TraceScreenshot(int32 InSizeX, int32 InSizeY, const TArra
 
 void FTraceScreenshot::Reset()
 {
-	RequestedScreenshotName.Empty();
 	bSuppressWritingToFile = false;
 }
diff --git a/Engine/Source/Runtime/Engine/Public/ProfilingDebugging/TraceScreenshot.h b/Engine/Source/Runtime/Engine/Public/ProfilingDebugging/TraceScreenshot.h
--- a/Engine/Source/Runtime/Engine/Public/ProfilingDebugging/TraceScreenshot.h
+++ b/Engine/Source/Runtime/Engine/Public/ProfilingDebugging/TraceScreenshot.h
@@ -41,5 +41,4 @@ private:
 
 private:
 	static bool bSuppressWritingToFile;
-	static FString RequestedScreenshotName;
 };
\ No newline at end of file
