Fixed FWorldPartitionLoadingContext::FDeferred to handle when load of actor fails.

diff --git a/Engine/Source/Runtime/Engine/Private/WorldPartition/WorldPartitionHandle.cpp b/Engine/Source/Runtime/Engine/Private/WorldPartition/WorldPartitionHandle.cpp
--- a/Engine/Source/Runtime/Engine/Private/WorldPartition/WorldPartitionHandle.cpp
+++ b/Engine/Source/Runtime/Engine/Private/WorldPartition/WorldPartitionHandle.cpp
@@ -166,31 +166,33 @@ void FWorldPartitionLoadingContext::FDeferred::RegisterActor(FWorldPartitionActo
 	TGuardValue<bool> IsEditorLoadingPackageGuard(GIsEditorLoadingPackage, true);
 
 	check(ActorDesc);
+	if (ActorDesc->Load())
+	{
+		UActorDescContainer* Container = ActorDesc->GetContainer();
+		check(Container);
 
-	UActorDescContainer* Container = ActorDesc->GetContainer();
-	check(Container);
-
-	bool bIsAlreadyInSetPtr;
-	ContainerOps.FindOrAdd(Container).Registrations.Add(ActorDesc, &bIsAlreadyInSetPtr);
-	check(!bIsAlreadyInSetPtr);
-
-	check(!ContainerOps.FindChecked(Container).Unregistrations.Contains(ActorDesc));
+		bool bIsAlreadyInSetPtr;
+		ContainerOps.FindOrAdd(Container).Registrations.Add(ActorDesc, &bIsAlreadyInSetPtr);
+		check(!bIsAlreadyInSetPtr);
 
-	ActorDesc->Load();
+		check(!ContainerOps.FindChecked(Container).Unregistrations.Contains(ActorDesc));
+	}
 }
 
 void FWorldPartitionLoadingContext::FDeferred::UnregisterActor(FWorldPartitionActorDesc* ActorDesc)
 {
 	check(ActorDesc);
+	if (AActor* Actor = ActorDesc->GetActor(); IsValid(Actor))
+	{
+		UActorDescContainer* Container = ActorDesc->GetContainer();
+		check(Container);
 
-	UActorDescContainer* Container = ActorDesc->GetContainer();
-	check(Container);
-
-	bool bIsAlreadyInSetPtr;
-	ContainerOps.FindOrAdd(Container).Unregistrations.Add(ActorDesc, &bIsAlreadyInSetPtr);
-	check(!bIsAlreadyInSetPtr);
+		bool bIsAlreadyInSetPtr;
+		ContainerOps.FindOrAdd(Container).Unregistrations.Add(ActorDesc, &bIsAlreadyInSetPtr);
+		check(!bIsAlreadyInSetPtr);
 
-	check(!ContainerOps.FindChecked(Container).Registrations.Contains(ActorDesc));
+		check(!ContainerOps.FindChecked(Container).Registrations.Contains(ActorDesc));
+	}
 }
 
 /**
