Modernize FCubemapTexturePropertiesPS permutations.

diff --git a/Engine/Source/Runtime/Engine/Private/CubemapUnwrapUtils.cpp b/Engine/Source/Runtime/Engine/Private/CubemapUnwrapUtils.cpp
--- a/Engine/Source/Runtime/Engine/Private/CubemapUnwrapUtils.cpp
+++ b/Engine/Source/Runtime/Engine/Private/CubemapUnwrapUtils.cpp
@@ -16,8 +16,7 @@
 #include "PipelineStateCache.h"
 
 IMPLEMENT_SHADER_TYPE(,FCubemapTexturePropertiesVS,TEXT("/Engine/Private/SimpleElementVertexShader.usf"),TEXT("Main"),SF_Vertex);
-IMPLEMENT_SHADER_TYPE(template<>,FCubemapTexturePropertiesPS<false>,TEXT("/Engine/Private/SimpleElementPixelShader.usf"),TEXT("CubemapTextureProperties"),SF_Pixel);
-IMPLEMENT_SHADER_TYPE(template<>,FCubemapTexturePropertiesPS<true>,TEXT("/Engine/Private/SimpleElementPixelShader.usf"),TEXT("CubemapTextureProperties"),SF_Pixel);
+IMPLEMENT_SHADER_TYPE(,FCubemapTexturePropertiesPS,TEXT("/Engine/Private/SimpleElementPixelShader.usf"),TEXT("CubemapTextureProperties"),SF_Pixel);
 IMPLEMENT_SHADER_TYPE(,FIESLightProfilePS,TEXT("/Engine/Private/SimpleElementPixelShader.usf"),TEXT("IESLightProfileMain"),SF_Pixel);
 
 namespace CubemapHelpers
@@ -120,8 +119,33 @@ void FCubemapTexturePropertiesVS::SetParameters( FRHICommandList& RHICmdList, co
 	SetShaderValue(RHICmdList, RHICmdList.GetBoundVertexShader(), Transform, (FMatrix44f)TransformValue);
 }
 
-template<bool bHDROutput>
-void FCubemapTexturePropertiesPS<bHDROutput>::SetParameters(FRHICommandList& RHICmdList, const FTexture* Texture, const FMatrix& ColorWeightsValue, float MipLevel, float SliceIndex, float GammaValue, bool bIsTextureCubeArray)
+FCubemapTexturePropertiesPS::FCubemapTexturePropertiesPS(const ShaderMetaType::CompiledShaderInitializerType& Initializer)
+	: FGlobalShader(Initializer)
+{
+	CubeTexture.Bind(Initializer.ParameterMap, TEXT("CubeTexture"));
+	CubeTextureSampler.Bind(Initializer.ParameterMap, TEXT("CubeTextureSampler"));
+	ColorWeights.Bind(Initializer.ParameterMap, TEXT("ColorWeights"));
+	PackedProperties0.Bind(Initializer.ParameterMap, TEXT("PackedProperties0"));
+	Gamma.Bind(Initializer.ParameterMap, TEXT("Gamma"));
+	NumSlices.Bind(Initializer.ParameterMap, TEXT("NumSlices"));
+}
+
+bool FCubemapTexturePropertiesPS::ShouldCompilePermutation(const FGlobalShaderPermutationParameters& Parameters)
+{
+	if (!IsPCPlatform(Parameters.Platform))
+	{
+		return false;
+	}
+
+	FPermutationDomain PermutationVector(Parameters.PermutationId);
+	if (PermutationVector.Get<FCubeArray>() && !IsFeatureLevelSupported(Parameters.Platform, ERHIFeatureLevel::SM5))
+	{
+		return false;
+	}
+	return true;
+}
+
+void FCubemapTexturePropertiesPS::SetParameters(FRHICommandList& RHICmdList, const FTexture* Texture, const FMatrix& ColorWeightsValue, float MipLevel, float SliceIndex, float GammaValue, bool bIsTextureCubeArray)
 {
 	FRHIPixelShader* ShaderRHI = RHICmdList.GetBoundPixelShader();
 
@@ -142,29 +166,18 @@ void FCubemapTexturePropertiesPS<bHDROutput>::SetParameters(FRHICommandList& RHI
 	}
 }
 
-void FMipLevelBatchedElementParameters::BindShaders(FRHICommandList& RHICmdList, FGraphicsPipelineStateInitializer& GraphicsPSOInit, ERHIFeatureLevel::Type InFeatureLevel, const FMatrix& InTransform, const float InGamma, const FMatrix& ColorWeights, const FTexture* Texture)
-{
-	if(bHDROutput)
-	{
-		BindShaders<FCubemapTexturePropertiesPS<true> >(RHICmdList, GraphicsPSOInit, InFeatureLevel, InTransform, InGamma, ColorWeights, Texture);
-	}
-	else
-	{
-		BindShaders<FCubemapTexturePropertiesPS<false> >(RHICmdList, GraphicsPSOInit, InFeatureLevel, InTransform, InGamma, ColorWeights, Texture);
-	}
-}
-
-template<typename TPixelShader>
 void FMipLevelBatchedElementParameters::BindShaders(FRHICommandList& RHICmdList, FGraphicsPipelineStateInitializer& GraphicsPSOInit, ERHIFeatureLevel::Type InFeatureLevel, const FMatrix& InTransform, const float InGamma, const FMatrix& ColorWeights, const FTexture* Texture)
 {
 	GraphicsPSOInit.BlendState = TStaticBlendState<>::GetRHI();
 
 	TShaderMapRef<FCubemapTexturePropertiesVS> VertexShader(GetGlobalShaderMap(InFeatureLevel));
 
-	typename TPixelShader::FPermutationDomain PermutationVector;
-	PermutationVector.template Set<typename TPixelShader::FCubemapArrayTexturePropertiesPS>(bIsTextureCubeArray);
-	TShaderMapRef<TPixelShader> PixelShader(GetGlobalShaderMap(InFeatureLevel), PermutationVector);
-	
+	typename FCubemapTexturePropertiesPS::FPermutationDomain PermutationVector;
+	PermutationVector.Set<FCubemapTexturePropertiesPS::FHDROutput>(bHDROutput);
+	PermutationVector.Set<FCubemapTexturePropertiesPS::FCubeArray>(bIsTextureCubeArray);
+
+	TShaderMapRef<FCubemapTexturePropertiesPS> PixelShader(GetGlobalShaderMap(InFeatureLevel), PermutationVector);
+
 	GraphicsPSOInit.BoundShaderState.VertexDeclarationRHI = GSimpleElementVertexDeclaration.VertexDeclarationRHI;
 	GraphicsPSOInit.BoundShaderState.VertexShaderRHI = VertexShader.GetVertexShader();
 	GraphicsPSOInit.BoundShaderState.PixelShaderRHI = PixelShader.GetPixelShader();
diff --git a/Engine/Source/Runtime/Engine/Public/CubemapUnwrapUtils.h b/Engine/Source/Runtime/Engine/Public/CubemapUnwrapUtils.h
--- a/Engine/Source/Runtime/Engine/Public/CubemapUnwrapUtils.h
+++ b/Engine/Source/Runtime/Engine/Public/CubemapUnwrapUtils.h
@@ -66,37 +66,21 @@ private:
 /**
  * Simple pixel shader reads from a cube map texture and unwraps it in the LongitudeLatitude form.
  */
-template<bool bHDROutput>
 class FCubemapTexturePropertiesPS : public FGlobalShader
 {
 	DECLARE_SHADER_TYPE(FCubemapTexturePropertiesPS,Global);
 
-	class FCubemapArrayTexturePropertiesPS : SHADER_PERMUTATION_BOOL("TEXTURECUBE_ARRAY");
-	using FPermutationDomain = TShaderPermutationDomain<FCubemapArrayTexturePropertiesPS>;
+	class FHDROutput : SHADER_PERMUTATION_BOOL("HDR_OUTPUT");
+	class FCubeArray : SHADER_PERMUTATION_BOOL("TEXTURECUBE_ARRAY");
+	using FPermutationDomain = TShaderPermutationDomain<FHDROutput, FCubeArray>;
 
 public:
+	FCubemapTexturePropertiesPS() = default;
+	FCubemapTexturePropertiesPS(const ShaderMetaType::CompiledShaderInitializerType& Initializer);
 
-	static bool ShouldCompilePermutation(const FGlobalShaderPermutationParameters& Parameters) { return IsPCPlatform(Parameters.Platform);}
-
-	FCubemapTexturePropertiesPS(const ShaderMetaType::CompiledShaderInitializerType& Initializer)
-		: FGlobalShader(Initializer)
-	{
-		CubeTexture.Bind(Initializer.ParameterMap,TEXT("CubeTexture"));
-		CubeTextureSampler.Bind(Initializer.ParameterMap,TEXT("CubeTextureSampler"));
-		ColorWeights.Bind(Initializer.ParameterMap,TEXT("ColorWeights"));
-		PackedProperties0.Bind(Initializer.ParameterMap,TEXT("PackedProperties0"));
-		Gamma.Bind(Initializer.ParameterMap,TEXT("Gamma"));
-		NumSlices.Bind(Initializer.ParameterMap, TEXT("NumSlices"));
-	}
-	FCubemapTexturePropertiesPS() {}
-
+	static bool ShouldCompilePermutation(const FGlobalShaderPermutationParameters& Parameters);
 	void SetParameters(FRHICommandList& RHICmdList, const FTexture* Texture, const FMatrix& ColorWeightsValue, float MipLevel, float SliceIndex, float GammaValue, bool bIsTextureCubeArray);
 
-	static void ModifyCompilationEnvironment(const FGlobalShaderPermutationParameters& Parameters, FShaderCompilerEnvironment& OutEnvironment)
-	{
-		FGlobalShader::ModifyCompilationEnvironment(Parameters, OutEnvironment);
-		OutEnvironment.SetDefine(TEXT("HDR_OUTPUT"), bHDROutput ? TEXT("1") : TEXT("0"));
-	}
 private:
 	LAYOUT_FIELD(FShaderResourceParameter, CubeTexture);
 	LAYOUT_FIELD(FShaderResourceParameter, CubeTextureSampler);
@@ -122,8 +106,6 @@ public:
 	virtual void BindShaders(FRHICommandList& RHICmdList, FGraphicsPipelineStateInitializer& GraphicsPSOInit, ERHIFeatureLevel::Type InFeatureLevel, const FMatrix& InTransform, const float InGamma, const FMatrix& ColorWeights, const FTexture* Texture) override;
 
 private:
-	template<typename TPixelShader> void BindShaders(FRHICommandList& RHICmdList, FGraphicsPipelineStateInitializer& GraphicsPSOInit, ERHIFeatureLevel::Type InFeatureLevel, const FMatrix& InTransform, const float InGamma, const FMatrix& ColorWeights, const FTexture* Texture);
-
 	bool bHDROutput;
 
 	/** Parameters that need to be passed to the shader */
