Optimized PIE performance while compiling shaders:

* Fixed O(n^2) behavior in FShaderCompileJobCollection::SubmitJobs().

diff --git a/Engine/Source/Runtime/Engine/Private/ShaderCompiler/ShaderCompiler.cpp b/Engine/Source/Runtime/Engine/Private/ShaderCompiler/ShaderCompiler.cpp
--- a/Engine/Source/Runtime/Engine/Private/ShaderCompiler/ShaderCompiler.cpp
+++ b/Engine/Source/Runtime/Engine/Private/ShaderCompiler/ShaderCompiler.cpp
@@ -581,6 +581,9 @@ void FShaderCompileJobCollection::SubmitJobs(const TArray<FShaderCommonCompileJo
 
 			FWriteScopeLock Locker(Lock);
 
+			// Optimization: Only search the linked list once to prevent O(n^2) behavior
+			FShaderCommonCompileJob* ExistingJobs[NumShaderCompileJobPriorities] = { nullptr };
+
 			for (FShaderCommonCompileJob* Job : InJobs)
 			{
 				check(Job->JobIndex != INDEX_NONE);
@@ -640,13 +643,23 @@ void FShaderCompileJobCollection::SubmitJobs(const TArray<FShaderCommonCompileJo
 					// link the job at the end of pending list, so we're executing them in a FIFO and not LIFO order
 					if (PendingJobs[PriorityIndex])
 					{
-						for (FShaderCommonCompileJob::TIterator It(PendingJobs[PriorityIndex]); It; ++It)
+						// Optimization: Only search the linked list if we have to
+						if (ExistingJobs[PriorityIndex])
+						{
+							Job->LinkAfter(ExistingJobs[PriorityIndex]);
+							ExistingJobs[PriorityIndex] = Job;
+						}
+						else
 						{
-							FShaderCommonCompileJob& ExistingJob = *It;
-							if (ExistingJob.GetNextLink() == nullptr)
+							for (FShaderCommonCompileJob::TIterator It(PendingJobs[PriorityIndex]); It; ++It)
 							{
-								Job->LinkAfter(&ExistingJob);
-								break;
+								FShaderCommonCompileJob& ExistingJob = *It;
+								if (ExistingJob.GetNextLink() == nullptr)
+								{
+									Job->LinkAfter(&ExistingJob);
+									ExistingJobs[PriorityIndex] = Job;
+									break;
+								}
 							}
 						}
 					}
@@ -4614,7 +4627,7 @@ void FShaderCompilingManager::PropagateMaterialChangesToPrimitives(const TMap<TR
 				{
 					FMaterial* UpdatedMaterial = MaterialIt.Key();
 					UMaterialInterface* UpdatedMaterialInterface = UpdatedMaterial->GetMaterialInterface();
-						
+
 					if (UpdatedMaterialInterface)
 					{
 						for (UMaterialInterface* TestMaterial : UsedMaterials)
