Fix issues with COTF and materials that don't use a vertex factory or are not mesh materials.
#rnx

diff --git a/Engine/Source/Runtime/Engine/Private/Materials/Material.cpp b/Engine/Source/Runtime/Engine/Private/Materials/Material.cpp
--- a/Engine/Source/Runtime/Engine/Private/Materials/Material.cpp
+++ b/Engine/Source/Runtime/Engine/Private/Materials/Material.cpp
@@ -5052,33 +5052,35 @@ void UMaterial::CompileODSCMaterialsForRemoteRecompile(TArray<FODSCRequestPayloa
 			const FVertexFactoryType* VFType = FVertexFactoryType::GetVFByName(payload.VertexFactoryName);
 			const FShaderPipelineType* PipelineType = FShaderPipelineType::GetShaderPipelineTypeByName(payload.PipelineName);
 
+			Shaders.ShaderPlatform = payload.ShaderPlatform;
+			Shaders.FeatureLevel = payload.FeatureLevel;
+			Shaders.QualityLevel = payload.QualityLevel;
+
 			if (VFType)
 			{
 				UE_LOG(LogShaders, Display, TEXT("VF Type:     %s "), *payload.VertexFactoryName);
-				Shaders.ShaderPlatform = payload.ShaderPlatform;
-				Shaders.FeatureLevel = payload.FeatureLevel;
-				Shaders.QualityLevel = payload.QualityLevel;
-				Shaders.VFTypes.Add(VFType);
+			}
 
-				if (PipelineType)
-				{
-					UE_LOG(LogShaders, Display, TEXT("Pipeline Type: %s"), *payload.PipelineName);
+			if (PipelineType)
+			{
+				UE_LOG(LogShaders, Display, TEXT("Pipeline Type: %s"), *payload.PipelineName);
 
-					Shaders.PipelineTypes.Add(PipelineType);
-					Shaders.ShaderTypes.Add(nullptr);
-				}
-				else
+				Shaders.VFTypes.Add(VFType);
+				Shaders.PipelineTypes.Add(PipelineType);
+				Shaders.ShaderTypes.Add(nullptr);
+			}
+			else
+			{
+				for (const FString& ShaderTypeName : payload.ShaderTypeNames)
 				{
-					for (const FString& ShaderTypeName : payload.ShaderTypeNames)
+					const FShaderType* ShaderType = FShaderType::GetShaderTypeByName(*ShaderTypeName);
+					if (ShaderType)
 					{
-						const FShaderType* ShaderType = FShaderType::GetShaderTypeByName(*ShaderTypeName);
-						if (ShaderType)
-						{
-							UE_LOG(LogShaders, Display, TEXT("\tShader Type: %s"), *ShaderTypeName);
+						UE_LOG(LogShaders, Display, TEXT("\tShader Type: %s"), *ShaderTypeName);
 
-							Shaders.PipelineTypes.Add(nullptr);
-							Shaders.ShaderTypes.Add(ShaderType);
-						}
+						Shaders.VFTypes.Add(VFType);
+						Shaders.PipelineTypes.Add(nullptr);
+						Shaders.ShaderTypes.Add(ShaderType);
 					}
 				}
 			}
diff --git a/Engine/Source/Runtime/Engine/Private/Materials/MaterialShared.cpp b/Engine/Source/Runtime/Engine/Private/Materials/MaterialShared.cpp
--- a/Engine/Source/Runtime/Engine/Private/Materials/MaterialShared.cpp
+++ b/Engine/Source/Runtime/Engine/Private/Materials/MaterialShared.cpp
@@ -2637,8 +2637,6 @@ void FMaterial::CacheGivenTypes(EShaderPlatform Platform, const TArray<const FVe
 		for (int i = 0; i < VFTypes.Num(); ++i)
 		{
 			const FVertexFactoryType* VFType = VFTypes[i];
-			check(VFType);
-
 			const FShaderPipelineType* PipelineType = PipelineTypes[i];
 			const FShaderType* ShaderType = ShaderTypes[i];
 
@@ -2658,10 +2656,22 @@ void FMaterial::CacheGivenTypes(EShaderPlatform Platform, const TArray<const FVe
 					nullptr,
 					nullptr);
 			}
-			else if (ShaderType)
+			else if (ShaderType->GetTypeForDynamicCast() == FShaderType::EShaderTypeForDynamicCast::Material)
+			{
+				ShaderType->AsMaterialShaderType()->BeginCompileShader(
+					EShaderCompileJobPriority::ForceLocal,
+					GetGameThreadCompilingShaderMapId(),
+					0,
+					this,
+					GameThreadPendingCompilerEnvironment,
+					Platform,
+					GameThreadShaderMap->GetPermutationFlags(),
+					CompileJobs,
+					nullptr,
+					nullptr);
+			}
+			else if (ShaderType->GetTypeForDynamicCast() == FShaderType::EShaderTypeForDynamicCast::MeshMaterial)
 			{
-				check(ShaderType->AsMeshMaterialShaderType());
-
 				ShaderType->AsMeshMaterialShaderType()->BeginCompileShader(
 					EShaderCompileJobPriority::ForceLocal,
 					GetGameThreadCompilingShaderMapId(),
