Fixed materials not properly setting bHasPerInstanceCustomData when a material has only ever seen UMaterialExpressionPerInstanceCustomData3Vector, not UMaterialExpressionPerInstanceCustomData.

diff --git a/Engine/Source/Runtime/Engine/Private/Materials/MaterialCachedData.cpp b/Engine/Source/Runtime/Engine/Private/Materials/MaterialCachedData.cpp
--- a/Engine/Source/Runtime/Engine/Private/Materials/MaterialCachedData.cpp
+++ b/Engine/Source/Runtime/Engine/Private/Materials/MaterialCachedData.cpp
@@ -366,6 +366,10 @@ void FMaterialCachedExpressionData::UpdateForExpressions(const FMaterialCachedEx
 		{
 			bHasPerInstanceCustomData = true;
 		}
+		else if (Expression->IsA(UMaterialExpressionPerInstanceCustomData3Vector::StaticClass()))
+		{
+			bHasPerInstanceCustomData = true;
+		}
 		else if (Expression->IsA(UMaterialExpressionVertexInterpolator::StaticClass()))
 		{
 			bHasVertexInterpolator = true;
