WorldPartition: Ensure the fake engine tick used for commandlets does not have a delta time
* Fixes partial captures of the minimap having changes between captured tiles

diff --git a/Engine/Source/Runtime/Engine/Classes/Commandlets/Commandlet.h b/Engine/Source/Runtime/Engine/Classes/Commandlets/Commandlet.h
--- a/Engine/Source/Runtime/Engine/Classes/Commandlets/Commandlet.h
+++ b/Engine/Source/Runtime/Engine/Classes/Commandlets/Commandlet.h
@@ -168,5 +168,5 @@ namespace CommandletHelpers
 	 * Can be used by commandlets to tick various subsystems. If running with -AllowCommandletRendering, this will
 	 * also tick rendering for the provided world scene(s).
 	 */
-	ENGINE_API void TickEngine(class UWorld* InWorld = nullptr);
+	ENGINE_API void TickEngine(class UWorld* InWorld = nullptr, double InDeltaTime = 0.0);
 }
\ No newline at end of file
diff --git a/Engine/Source/Runtime/Engine/Private/Commandlets/Commandlet.cpp b/Engine/Source/Runtime/Engine/Private/Commandlets/Commandlet.cpp
--- a/Engine/Source/Runtime/Engine/Private/Commandlets/Commandlet.cpp
+++ b/Engine/Source/Runtime/Engine/Private/Commandlets/Commandlet.cpp
@@ -65,14 +65,13 @@ bool FCommandletCommandLineParsingTest::RunTest(const FString& Parameters)
 #endif // WITH_DEV_AUTOMATION_TESTS
 
 
-void CommandletHelpers::TickEngine(UWorld* InWorld)
+void CommandletHelpers::TickEngine(UWorld* InWorld, double InDeltaTime)
 {
 	// Simulate an engine frame tick
 	// Will make sure systems can perform their internal bookkeeping properly. For example, the VT system needs to 
 	// process deleted VTs.
 
-	// Update FApp::CurrentTime / FApp::DeltaTime while taking into account max tick rate.
-	GEngine->UpdateTimeAndHandleMaxTickRate();
+	FApp::SetDeltaTime(InDeltaTime);
 
 	// Tick the engine.
 	GEngine->Tick(FApp::GetDeltaTime(), false);
