Added "debug alloc" and "debug gpucalloc" commands.

diff --git a/Engine/Source/Runtime/Engine/Private/UnrealEngine.cpp b/Engine/Source/Runtime/Engine/Private/UnrealEngine.cpp
--- a/Engine/Source/Runtime/Engine/Private/UnrealEngine.cpp
+++ b/Engine/Source/Runtime/Engine/Private/UnrealEngine.cpp
@@ -10210,6 +10210,71 @@ bool UEngine::PerformError(const TCHAR* Cmd, FOutputDevice& Ar)
 
 		return true;
 	}
+	else if (FParse::Command(&Cmd, TEXT("ALLOC")))
+	{
+		int64 MBToAllocate = FCString::Atoi64(Cmd);
+		if (MBToAllocate < 0)
+		{
+			Ar.Log(TEXT("Cannot allocate a negative amount of memory. Nice try though"));
+			return false;
+		}
+		else
+		{
+			Ar.Logf(TEXT("Allocating %d MB..."), MBToAllocate);
+		}
+
+
+		constexpr static size_t BytesPerMegabyte = 1024 * 1024;
+		while (MBToAllocate > 0)
+		{
+			// allocate 1 MB at a time (trying to malloc a huge number of bytes in one call seems to silently fail)
+			void* Eat = FMemory::Malloc(BytesPerMegabyte);
+			// newly allocated memory must be touched
+			FMemory::Memset(Eat, 0, BytesPerMegabyte);
+			MBToAllocate--;
+		}
+
+		return true;
+	}
+	else if (FParse::Command(&Cmd, TEXT("GPUALLOC")))
+	{
+		int64 MBToAllocate = FCString::Atoi64(Cmd);
+		if (MBToAllocate < 0)
+		{
+			Ar.Log(TEXT("Cannot allocate a negative amount of memory. Nice try though"));
+			return false;
+		}
+		else
+		{
+			Ar.Logf(TEXT("Allocating %d MB of VRAM..."), MBToAllocate);
+		}
+
+		static TArray<FBufferRHIRef> BufferRefs;
+		ENQUEUE_RENDER_COMMAND(GPUOOMAllocate)([MBToAllocate](FRHICommandListImmediate& CmdList)
+		{
+			// variables captured by value are const
+			size_t RemainingMB = MBToAllocate;
+			while (RemainingMB > 0)
+			{
+				// allocate up to 64MB at a time
+				const size_t CurrentAllocMB = FMath::Min<size_t>(64, RemainingMB);
+				const size_t CurrentAllocBytes = CurrentAllocMB * 1024 * 1024;
+
+				// create and fill buffer
+				FRHIResourceCreateInfo Info(TEXT("Debug GPUAlloc"));
+				FBufferRHIRef Buf = RHICreateBuffer(CurrentAllocBytes, EBufferUsageFlags::VertexBuffer, 32, ERHIAccess::Unknown, Info);
+				// only need to touch some of the memory
+				void* Data = RHILockBuffer(Buf, 0, CurrentAllocBytes, EResourceLockMode::RLM_WriteOnly);
+				FMemory::Memset(Data, 'd', CurrentAllocBytes);
+				RHIUnlockBuffer(Buf);
+
+				// store temporarily to avoid RAII
+				BufferRefs.Add(Buf);
+				
+				RemainingMB -= CurrentAllocMB;
+			}
+		});
+	}
 	else if (FParse::Command(&Cmd, TEXT("STACKOVERFLOW")))
 	{
 		Ar.Log(TEXT("Infinite recursion to cause stack overflow"));
