Debug Hud - Stat rendering optimizations

This cuts the cost of DrawStatsHUD by about half.

+ Remove implicit conversions from large world coords to floats by explicitly using float vector types.
+ Draw text effects (shadow, outline) in same pass as main text draw to avoid calling draw several times.
+ Updated the CharacterList to cache the FCharacterEntry's directly to avoid creation/copying of them.
+ Removed direct index cache from CharacterList as it didn't actually make it faster due to the complex logic when getting a character.
+ Updated FBatchedElements::AddVertex to takes float vectors to avoid expensive LWC conversions.

diff --git a/Engine/Source/Runtime/Engine/Private/BatchedElements.cpp b/Engine/Source/Runtime/Engine/Private/BatchedElements.cpp
--- a/Engine/Source/Runtime/Engine/Private/BatchedElements.cpp
+++ b/Engine/Source/Runtime/Engine/Private/BatchedElements.cpp
@@ -108,7 +108,14 @@ void FBatchedElements::AddPoint(const FVector& Position,float Size,const FLinear
 	Point->HitProxyId = HitProxyId;
 }
 
-int32 FBatchedElements::AddVertex(const FVector4& InPosition,const FVector2D& InTextureCoordinate,const FLinearColor& InColor,FHitProxyId HitProxyId)
+int32 FBatchedElements::AddVertex(const FVector4& InPosition, const FVector2D& InTextureCoordinate, const FLinearColor& InColor, FHitProxyId HitProxyId)
+{
+	int32 VertexIndex = MeshVertices.Num();
+	new(MeshVertices) FSimpleElementVertex(InPosition, InTextureCoordinate, InColor, HitProxyId);
+	return VertexIndex;
+}
+
+int32 FBatchedElements::AddVertexf(const FVector4f& InPosition,const FVector2f& InTextureCoordinate,const FLinearColor& InColor,FHitProxyId HitProxyId)
 {
 	int32 VertexIndex = MeshVertices.Num();
 	new(MeshVertices) FSimpleElementVertex(InPosition,InTextureCoordinate,InColor,HitProxyId);
diff --git a/Engine/Source/Runtime/Engine/Private/UserInterface/CanvasItem.cpp b/Engine/Source/Runtime/Engine/Private/UserInterface/CanvasItem.cpp
--- a/Engine/Source/Runtime/Engine/Private/UserInterface/CanvasItem.cpp
+++ b/Engine/Source/Runtime/Engine/Private/UserInterface/CanvasItem.cpp
@@ -5,6 +5,7 @@
 =============================================================================*/
 
 #include "CanvasItem.h"
+#include "BatchedElements.h"
 #include "EngineStats.h"
 #include "EngineGlobals.h"
 #include "HitProxies.h"
@@ -18,7 +19,7 @@
 #include "Fonts/FontMeasure.h"
 #include "Framework/Application/SlateApplication.h"
 #include "EngineFontServices.h"
-
+#include "TextureResource.h"
 
 DECLARE_CYCLE_STAT(TEXT("CanvasTileTextureItem Time"),STAT_Canvas_TileTextureItemTime,STATGROUP_Canvas);
 DECLARE_CYCLE_STAT(TEXT("CanvasTileMaterialItem Time"),STAT_Canvas_TileMaterialItemTime,STATGROUP_Canvas);
@@ -447,24 +448,24 @@ void FCanvasTileItem::Draw( class FCanvas* InCanvas )
 		Right =		( Position.X + Size.X ) * Z;
 		Bottom =	( Position.Y + Size.Y ) * Z;		
 
-		int32 V00 = BatchedElements->AddVertex(
-			FVector4( Left, Top, 0.0f, Z ),
-			FVector2D( UV0.X, UV0.Y ),
+		int32 V00 = BatchedElements->AddVertexf(
+			FVector4f( Left, Top, 0.0f, Z ),
+			FVector2f( UV0.X, UV0.Y ),
 			ActualColor,
 			HitProxyId );
-		int32 V10 = BatchedElements->AddVertex(
-			FVector4( Right, Top, 0.0f, Z ),
-			FVector2D( UV1.X, UV0.Y ),
+		int32 V10 = BatchedElements->AddVertexf(
+			FVector4f( Right, Top, 0.0f, Z ),
+			FVector2f( UV1.X, UV0.Y ),
 			ActualColor,
 			HitProxyId );
-		int32 V01 = BatchedElements->AddVertex(
-			FVector4( Left, Bottom, 0.0f, Z ),
-			FVector2D( UV0.X, UV1.Y ),		
+		int32 V01 = BatchedElements->AddVertexf(
+			FVector4f( Left, Bottom, 0.0f, Z ),
+			FVector2f( UV0.X, UV1.Y ),		
 			ActualColor,
 			HitProxyId );
-		int32 V11 = BatchedElements->AddVertex(
-			FVector4(Right,	Bottom,	0.0f, Z ),
-			FVector2D( UV1.X, UV1.Y ),
+		int32 V11 = BatchedElements->AddVertexf(
+			FVector4f(Right,	Bottom,	0.0f, Z ),
+			FVector2f( UV1.X, UV1.Y ),
 			ActualColor,
 			HitProxyId);
 
@@ -573,24 +574,24 @@ void FCanvasBorderItem::Draw( class FCanvas* InCanvas )
 		BatchedElements->ReserveVertices(4 * NumElements); // 4 verts each
 
 		//Draw background
-		int32 V00 = BatchedElements->AddVertex(
-			FVector4( Left + BorderLeftDrawSizeX, Top + BorderTopDrawSizeY, 0.0f, Z ),
-			FVector2D( 0, 0 ),
+		int32 V00 = BatchedElements->AddVertexf(
+			FVector4f( Left + BorderLeftDrawSizeX, Top + BorderTopDrawSizeY, 0.0f, Z ),
+			FVector2f( 0, 0 ),
 			ActualColor,
 			HitProxyId );
-		int32 V10 = BatchedElements->AddVertex(
-			FVector4( Right - BorderRightDrawSizeX, Top + BorderTopDrawSizeY, 0.0f, Z ),
-			FVector2D( BackgroundTilingX, 0 ),
+		int32 V10 = BatchedElements->AddVertexf(
+			FVector4f( Right - BorderRightDrawSizeX, Top + BorderTopDrawSizeY, 0.0f, Z ),
+			FVector2f( BackgroundTilingX, 0 ),
 			ActualColor,
 			HitProxyId );
-		int32 V01 = BatchedElements->AddVertex(
-			FVector4( Left + BorderLeftDrawSizeX, Bottom - BorderBottomDrawSizeY, 0.0f, Z ),
-			FVector2D( 0, BackgroundTilingY ),		
+		int32 V01 = BatchedElements->AddVertexf(
+			FVector4f( Left + BorderLeftDrawSizeX, Bottom - BorderBottomDrawSizeY, 0.0f, Z ),
+			FVector2f( 0, BackgroundTilingY ),		
 			ActualColor,
 			HitProxyId );
-		int32 V11 = BatchedElements->AddVertex(
-			FVector4(Right - BorderRightDrawSizeX, Bottom - BorderBottomDrawSizeY, 0.0f, Z ),
-			FVector2D( BackgroundTilingX, BackgroundTilingY ),
+		int32 V11 = BatchedElements->AddVertexf(
+			FVector4f(Right - BorderRightDrawSizeX, Bottom - BorderBottomDrawSizeY, 0.0f, Z ),
+			FVector2f( BackgroundTilingX, BackgroundTilingY ),
 			ActualColor,
 			HitProxyId);
 
@@ -604,24 +605,24 @@ void FCanvasBorderItem::Draw( class FCanvas* InCanvas )
 		const float CornerDrawHeight = BorderTextureHeight * CornerSize.Y * BorderScale.Y;
 
 		//Top Left Corner
-		V00 = BatchedElements->AddVertex(
-			FVector4( Left, Top, 0.0f, Z ),
-			FVector2D( BorderUV0.X, BorderUV0.Y ),
+		V00 = BatchedElements->AddVertexf(
+			FVector4f( Left, Top, 0.0f, Z ),
+			FVector2f( BorderUV0.X, BorderUV0.Y ),
 			ActualColor,
 			HitProxyId );
-		V10 = BatchedElements->AddVertex(
-			FVector4( Left + CornerDrawWidth, Top, 0.0f, Z ),
-			FVector2D( BorderUV1.X*CornerSize.X, BorderUV0.Y ),
+		V10 = BatchedElements->AddVertexf(
+			FVector4f( Left + CornerDrawWidth, Top, 0.0f, Z ),
+			FVector2f( BorderUV1.X*CornerSize.X, BorderUV0.Y ),
 			ActualColor,
 			HitProxyId );
-		V01 = BatchedElements->AddVertex(
-			FVector4( Left, Top + CornerDrawHeight, 0.0f, Z ),
-			FVector2D( BorderUV0.X, BorderUV1.Y*CornerSize.Y ),		
+		V01 = BatchedElements->AddVertexf(
+			FVector4f( Left, Top + CornerDrawHeight, 0.0f, Z ),
+			FVector2f( BorderUV0.X, BorderUV1.Y*CornerSize.Y ),		
 			ActualColor,
 			HitProxyId );
-		V11 = BatchedElements->AddVertex(
-			FVector4( Left + CornerDrawWidth, Top + CornerDrawHeight, 0.0f, Z ),
-			FVector2D( BorderUV1.X*CornerSize.X, BorderUV1.Y*CornerSize.Y ),
+		V11 = BatchedElements->AddVertexf(
+			FVector4f( Left + CornerDrawWidth, Top + CornerDrawHeight, 0.0f, Z ),
+			FVector2f( BorderUV1.X*CornerSize.X, BorderUV1.Y*CornerSize.Y ),
 			ActualColor,
 			HitProxyId);
 
@@ -629,24 +630,24 @@ void FCanvasBorderItem::Draw( class FCanvas* InCanvas )
 		BatchedElements->AddTriangleExtensive( V00, V11, V01, BatchedElementParameters, CornersTexture, BlendMode );
 
 		// Top Right Corner
-		V00 = BatchedElements->AddVertex(
-			FVector4( Right - CornerDrawWidth, Top, 0.0f, Z ),
-			FVector2D( BorderUV1.X - (BorderUV1.X - BorderUV0.X)*CornerSize.X, BorderUV0.Y ),
+		V00 = BatchedElements->AddVertexf(
+			FVector4f( Right - CornerDrawWidth, Top, 0.0f, Z ),
+			FVector2f( BorderUV1.X - (BorderUV1.X - BorderUV0.X)*CornerSize.X, BorderUV0.Y ),
 			ActualColor,
 			HitProxyId );
-		V10 = BatchedElements->AddVertex(
-			FVector4( Right, Top, 0.0f, Z ),
-			FVector2D( BorderUV1.X, BorderUV0.Y ),
+		V10 = BatchedElements->AddVertexf(
+			FVector4f( Right, Top, 0.0f, Z ),
+			FVector2f( BorderUV1.X, BorderUV0.Y ),
 			ActualColor,
 			HitProxyId );
-		V01 = BatchedElements->AddVertex(
-			FVector4( Right - CornerDrawWidth, Top + CornerDrawHeight, 0.0f, Z ),
-			FVector2D( BorderUV1.X - (BorderUV1.X - BorderUV0.X)*CornerSize.X, BorderUV1.Y*CornerSize.Y ),		
+		V01 = BatchedElements->AddVertexf(
+			FVector4f( Right - CornerDrawWidth, Top + CornerDrawHeight, 0.0f, Z ),
+			FVector2f( BorderUV1.X - (BorderUV1.X - BorderUV0.X)*CornerSize.X, BorderUV1.Y*CornerSize.Y ),		
 			ActualColor,
 			HitProxyId );
-		V11 = BatchedElements->AddVertex(
-			FVector4( Right, Top + CornerDrawHeight, 0.0f, Z ),
-			FVector2D( BorderUV1.X, BorderUV1.Y*CornerSize.Y ),
+		V11 = BatchedElements->AddVertexf(
+			FVector4f( Right, Top + CornerDrawHeight, 0.0f, Z ),
+			FVector2f( BorderUV1.X, BorderUV1.Y*CornerSize.Y ),
 			ActualColor,
 			HitProxyId);
 
@@ -654,24 +655,24 @@ void FCanvasBorderItem::Draw( class FCanvas* InCanvas )
 		BatchedElements->AddTriangleExtensive( V00, V11, V01, BatchedElementParameters, CornersTexture, BlendMode );
 
 		//Left Bottom Corner
-		V00 = BatchedElements->AddVertex(
-			FVector4( Left, Bottom - CornerDrawHeight, 0.0f, Z ),
-			FVector2D( BorderUV0.X, BorderUV1.Y - (BorderUV1.Y - BorderUV0.Y)*CornerSize.Y),
+		V00 = BatchedElements->AddVertexf(
+			FVector4f( Left, Bottom - CornerDrawHeight, 0.0f, Z ),
+			FVector2f( BorderUV0.X, BorderUV1.Y - (BorderUV1.Y - BorderUV0.Y)*CornerSize.Y),
 			ActualColor,
 			HitProxyId );
-		V10 = BatchedElements->AddVertex(
-			FVector4( Left + CornerDrawWidth, Bottom - CornerDrawHeight, 0.0f, Z ),
-			FVector2D( BorderUV1.X*CornerSize.X, BorderUV1.Y - (BorderUV1.Y - BorderUV0.Y)*CornerSize.Y ),
+		V10 = BatchedElements->AddVertexf(
+			FVector4f( Left + CornerDrawWidth, Bottom - CornerDrawHeight, 0.0f, Z ),
+			FVector2f( BorderUV1.X*CornerSize.X, BorderUV1.Y - (BorderUV1.Y - BorderUV0.Y)*CornerSize.Y ),
 			ActualColor,
 			HitProxyId );
-		V01 = BatchedElements->AddVertex(
-			FVector4( Left, Bottom, 0.0f, Z ),
-			FVector2D( BorderUV0.X, BorderUV1.Y ),		
+		V01 = BatchedElements->AddVertexf(
+			FVector4f( Left, Bottom, 0.0f, Z ),
+			FVector2f( BorderUV0.X, BorderUV1.Y ),		
 			ActualColor,
 			HitProxyId );
-		V11 = BatchedElements->AddVertex(
-			FVector4( Left + CornerDrawWidth, Bottom, 0.0f, Z ),
-			FVector2D( BorderUV1.X*CornerSize.X, BorderUV1.Y),
+		V11 = BatchedElements->AddVertexf(
+			FVector4f( Left + CornerDrawWidth, Bottom, 0.0f, Z ),
+			FVector2f( BorderUV1.X*CornerSize.X, BorderUV1.Y),
 			ActualColor,
 			HitProxyId);
 
@@ -679,24 +680,24 @@ void FCanvasBorderItem::Draw( class FCanvas* InCanvas )
 		BatchedElements->AddTriangleExtensive( V00, V11, V01, BatchedElementParameters, CornersTexture, BlendMode );
 
 		// Right Bottom Corner
-		V00 = BatchedElements->AddVertex(
-			FVector4( Right - CornerDrawWidth, Bottom - CornerDrawHeight, 0.0f, Z ),
-			FVector2D( BorderUV1.X - (BorderUV1.X - BorderUV0.X)*CornerSize.X, BorderUV1.Y - (BorderUV1.Y - BorderUV0.Y)*CornerSize.Y ),
+		V00 = BatchedElements->AddVertexf(
+			FVector4f( Right - CornerDrawWidth, Bottom - CornerDrawHeight, 0.0f, Z ),
+			FVector2f( BorderUV1.X - (BorderUV1.X - BorderUV0.X)*CornerSize.X, BorderUV1.Y - (BorderUV1.Y - BorderUV0.Y)*CornerSize.Y ),
 			ActualColor,
 			HitProxyId );
-		V10 = BatchedElements->AddVertex(
-			FVector4( Right, Bottom - CornerDrawHeight, 0.0f, Z ),
-			FVector2D( BorderUV1.X, BorderUV1.Y - (BorderUV1.Y - BorderUV0.Y)*CornerSize.Y ),
+		V10 = BatchedElements->AddVertexf(
+			FVector4f( Right, Bottom - CornerDrawHeight, 0.0f, Z ),
+			FVector2f( BorderUV1.X, BorderUV1.Y - (BorderUV1.Y - BorderUV0.Y)*CornerSize.Y ),
 			ActualColor,
 			HitProxyId );
-		V01 = BatchedElements->AddVertex(
-			FVector4( Right - CornerDrawWidth, Bottom, 0.0f, Z ),
-			FVector2D( BorderUV1.X - (BorderUV1.X - BorderUV0.X)*CornerSize.X, BorderUV1.Y ),		
+		V01 = BatchedElements->AddVertexf(
+			FVector4f( Right - CornerDrawWidth, Bottom, 0.0f, Z ),
+			FVector2f( BorderUV1.X - (BorderUV1.X - BorderUV0.X)*CornerSize.X, BorderUV1.Y ),		
 			ActualColor,
 			HitProxyId );
-		V11 = BatchedElements->AddVertex(
-			FVector4( Right, Bottom, 0.0f, Z ),
-			FVector2D( BorderUV1.X, BorderUV1.Y ),
+		V11 = BatchedElements->AddVertexf(
+			FVector4f( Right, Bottom, 0.0f, Z ),
+			FVector2f( BorderUV1.X, BorderUV1.Y ),
 			ActualColor,
 			HitProxyId);
 			
@@ -711,24 +712,24 @@ void FCanvasBorderItem::Draw( class FCanvas* InCanvas )
 		//Top Frame Border
 		const float TopFrameTilingX = (BorderRight-BorderLeft)/BorderTopDrawSizeX;
 
-		V00 = BatchedElements->AddVertex(
-			FVector4( BorderLeft, Top, 0.0f, Z ),
-			FVector2D( 0, 0 ),
+		V00 = BatchedElements->AddVertexf(
+			FVector4f( BorderLeft, Top, 0.0f, Z ),
+			FVector2f( 0, 0 ),
 			ActualColor,
 			HitProxyId );
-		V10 = BatchedElements->AddVertex(
-			FVector4( BorderRight, Top, 0.0f, Z ),
-			FVector2D( TopFrameTilingX , 0 ),
+		V10 = BatchedElements->AddVertexf(
+			FVector4f( BorderRight, Top, 0.0f, Z ),
+			FVector2f( TopFrameTilingX , 0 ),
 			ActualColor,
 			HitProxyId );
-		V01 = BatchedElements->AddVertex(
-			FVector4( BorderLeft, Top + BorderTopDrawSizeY, 0.0f, Z ),
-			FVector2D( 0, 1.0f ),		
+		V01 = BatchedElements->AddVertexf(
+			FVector4f( BorderLeft, Top + BorderTopDrawSizeY, 0.0f, Z ),
+			FVector2f( 0, 1.0f ),		
 			ActualColor,
 			HitProxyId );
-		V11 = BatchedElements->AddVertex(
-			FVector4( BorderRight, Top + BorderTopDrawSizeY, 0.0f, Z ),
-			FVector2D( TopFrameTilingX, 1.0f ),
+		V11 = BatchedElements->AddVertexf(
+			FVector4f( BorderRight, Top + BorderTopDrawSizeY, 0.0f, Z ),
+			FVector2f( TopFrameTilingX, 1.0f ),
 			ActualColor,
 			HitProxyId);
 
@@ -738,24 +739,24 @@ void FCanvasBorderItem::Draw( class FCanvas* InCanvas )
 		//Bottom Frame Border
 		const float BottomFrameTilingX = (BorderRight-BorderLeft)/BorderBottomDrawSizeX;
 
-		V00 = BatchedElements->AddVertex(
-			FVector4( BorderLeft, Bottom - BorderBottomDrawSizeY, 0.0f, Z ),
-			FVector2D( 0, 0 ),
+		V00 = BatchedElements->AddVertexf(
+			FVector4f( BorderLeft, Bottom - BorderBottomDrawSizeY, 0.0f, Z ),
+			FVector2f( 0, 0 ),
 			ActualColor,
 			HitProxyId );
-		V10 = BatchedElements->AddVertex(
-			FVector4( BorderRight, Bottom - BorderBottomDrawSizeY, 0.0f, Z ),
-			FVector2D( BottomFrameTilingX, 0 ),
+		V10 = BatchedElements->AddVertexf(
+			FVector4f( BorderRight, Bottom - BorderBottomDrawSizeY, 0.0f, Z ),
+			FVector2f( BottomFrameTilingX, 0 ),
 			ActualColor,
 			HitProxyId );
-		V01 = BatchedElements->AddVertex(
-			FVector4( BorderLeft, Bottom, 0.0f, Z ),
-			FVector2D( 0, 1.0f ),		
+		V01 = BatchedElements->AddVertexf(
+			FVector4f( BorderLeft, Bottom, 0.0f, Z ),
+			FVector2f( 0, 1.0f ),		
 			ActualColor,
 			HitProxyId );
-		V11 = BatchedElements->AddVertex(
-			FVector4( BorderRight, Bottom, 0.0f, Z ),
-			FVector2D( BottomFrameTilingX, 1.0f ),
+		V11 = BatchedElements->AddVertexf(
+			FVector4f( BorderRight, Bottom, 0.0f, Z ),
+			FVector2f( BottomFrameTilingX, 1.0f ),
 			ActualColor,
 			HitProxyId);
 
@@ -766,24 +767,24 @@ void FCanvasBorderItem::Draw( class FCanvas* InCanvas )
 		//Left Frame Border
 		const float LeftFrameTilingY = (BorderBottom-BorderTop) / BorderLeftDrawSizeY;
 
-		V00 = BatchedElements->AddVertex(
-			FVector4( Left, BorderTop, 0.0f, Z ),
-			FVector2D( 0, 0 ),
+		V00 = BatchedElements->AddVertexf(
+			FVector4f( Left, BorderTop, 0.0f, Z ),
+			FVector2f( 0, 0 ),
 			ActualColor,
 			HitProxyId );
-		V10 = BatchedElements->AddVertex(
-			FVector4( Left +  BorderLeftDrawSizeX , BorderTop, 0.0f, Z ),
-			FVector2D( 1.0f, 0 ),
+		V10 = BatchedElements->AddVertexf(
+			FVector4f( Left +  BorderLeftDrawSizeX , BorderTop, 0.0f, Z ),
+			FVector2f( 1.0f, 0 ),
 			ActualColor,
 			HitProxyId );
-		V01 = BatchedElements->AddVertex(
-			FVector4( Left, BorderBottom, 0.0f, Z ),
-			FVector2D( 0, LeftFrameTilingY ),		
+		V01 = BatchedElements->AddVertexf(
+			FVector4f( Left, BorderBottom, 0.0f, Z ),
+			FVector2f( 0, LeftFrameTilingY ),		
 			ActualColor,
 			HitProxyId );
-		V11 = BatchedElements->AddVertex(
-			FVector4( Left + BorderLeftDrawSizeX, BorderBottom, 0.0f, Z ),
-			FVector2D( 1.0f, LeftFrameTilingY ),
+		V11 = BatchedElements->AddVertexf(
+			FVector4f( Left + BorderLeftDrawSizeX, BorderBottom, 0.0f, Z ),
+			FVector2f( 1.0f, LeftFrameTilingY ),
 			ActualColor,
 			HitProxyId);
 
@@ -794,24 +795,24 @@ void FCanvasBorderItem::Draw( class FCanvas* InCanvas )
 		//Right Frame Border
 		const float RightFrameTilingY = (BorderBottom-BorderTop)/BorderRightDrawSizeY;
 
-		V00 = BatchedElements->AddVertex(
-			FVector4( Right - BorderRightDrawSizeX, BorderTop, 0.0f, Z ),
-			FVector2D( 0, 0 ),
+		V00 = BatchedElements->AddVertexf(
+			FVector4f( Right - BorderRightDrawSizeX, BorderTop, 0.0f, Z ),
+			FVector2f( 0, 0 ),
 			ActualColor,
 			HitProxyId );
-		V10 = BatchedElements->AddVertex(
-			FVector4( Right, BorderTop, 0.0f, Z ),
-			FVector2D( 1.0f , 0 ),
+		V10 = BatchedElements->AddVertexf(
+			FVector4f( Right, BorderTop, 0.0f, Z ),
+			FVector2f( 1.0f , 0 ),
 			ActualColor,
 			HitProxyId );
-		V01 = BatchedElements->AddVertex(
-			FVector4( Right - BorderRightDrawSizeX, BorderBottom, 0.0f, Z ),
-			FVector2D( 0, RightFrameTilingY ),		
+		V01 = BatchedElements->AddVertexf(
+			FVector4f( Right - BorderRightDrawSizeX, BorderBottom, 0.0f, Z ),
+			FVector2f( 0, RightFrameTilingY ),		
 			ActualColor,
 			HitProxyId );
-		V11 = BatchedElements->AddVertex(
-			FVector4( Right, BorderBottom, 0.0f, Z ),
-			FVector2D( 1.0f, RightFrameTilingY ),
+		V11 = BatchedElements->AddVertexf(
+			FVector4f( Right, BorderBottom, 0.0f, Z ),
+			FVector2f( 1.0f, RightFrameTilingY ),
 			ActualColor,
 			HitProxyId);
 
@@ -887,32 +888,33 @@ void FCanvasTextItemBase::Draw( class FCanvas* InCanvas )
 		}		
 	}
 	
-	FLinearColor DrawColor;
-	BatchedElements = nullptr;
+	TArray<FTextEffect, TFixedAllocator<6>> TextEffects;
+
+	const auto ModulateColor = [InCanvas, this](FLinearColor InColor)
+	{
+		InColor.A = Color.A;
+		// Copy the Alpha from the shadow otherwise if we fade the text the shadow wont fade - which is almost certainly not what we will want.
+		InColor.A *= InCanvas->AlphaModulate;
+		return InColor;
+	};
 
 	// If we have a shadow - draw it now
 	if (bHasShadow)
 	{
-		DrawColor = ShadowColor;
-		// Copy the Alpha from the shadow otherwise if we fade the text the shadow wont fade - which is almost certainly not what we will want.
-		DrawColor.A = Color.A;
-		DrawColor.A *= InCanvas->AlphaModulate;
-		DrawStringInternal(InCanvas, DrawPos + ShadowOffset*DPIScale, DrawColor);
+		TextEffects.Add({ FVector2f{ShadowOffset * DPIScale}, ModulateColor(ShadowColor) });
 	}
 
 	if( bOutlined )
 	{
-		DrawColor = OutlineColor;
-		DrawColor.A *= InCanvas->AlphaModulate;
-		DrawStringInternal( InCanvas, DrawPos + FVector2D( -1.0f, -1.0f ), DrawColor );
-		DrawStringInternal( InCanvas, DrawPos + FVector2D( -1.0f, 1.0f ), DrawColor );
-		DrawStringInternal( InCanvas, DrawPos + FVector2D( 1.0f, 1.0f ), DrawColor );
-		DrawStringInternal( InCanvas, DrawPos + FVector2D( 1.0f, -1.0f ), DrawColor );
+		const FLinearColor DrawColor = ModulateColor(OutlineColor);
+		TextEffects.Add({ FVector2f(-1.0f, -1.0f), OutlineColor });
+		TextEffects.Add({ FVector2f(-1.0f, 1.0f), OutlineColor });
+		TextEffects.Add({ FVector2f(1.0f, 1.0f), OutlineColor });
+		TextEffects.Add({ FVector2f(1.0f, -1.0f), OutlineColor });
 	}
 
-	DrawColor = Color;
-	DrawColor.A *= InCanvas->AlphaModulate;	
-	DrawStringInternal( InCanvas, DrawPos, DrawColor );
+	const FLinearColor DrawColor = ModulateColor(Color);
+	DrawStringInternal( InCanvas, DrawPos, DrawColor, TextEffects);
 }
 
 
@@ -979,16 +981,16 @@ FVector2D FCanvasTextItem::GetTextSize(float DPIScale) const
 	return MeasuredTextSize;
 }
 
-void FCanvasTextItem::DrawStringInternal( FCanvas* InCanvas, const FVector2D& DrawPos, const FLinearColor& InColor )
+void FCanvasTextItem::DrawStringInternal(FCanvas* InCanvas, const FVector2D& DrawPos, const FLinearColor& InColor, TArrayView<FTextEffect> TextEffects)
 {
 	switch(GetFontCacheType())
 	{
 	case EFontCacheType::Offline:
-		DrawStringInternal_OfflineCache(InCanvas, DrawPos, InColor);
+		DrawStringInternal_OfflineCache(InCanvas, DrawPos, InColor, TextEffects);
 		break;
 
 	case EFontCacheType::Runtime:
-		DrawStringInternal_RuntimeCache(InCanvas, DrawPos, InColor);
+		DrawStringInternal_RuntimeCache(InCanvas, DrawPos, InColor, TextEffects);
 		break;
 
 	default:
@@ -996,8 +998,10 @@ void FCanvasTextItem::DrawStringInternal( FCanvas* InCanvas, const FVector2D& Dr
 	}
 }
 
-void FCanvasTextItem::DrawStringInternal_OfflineCache( FCanvas* InCanvas, const FVector2D& DrawPos, const FLinearColor& InColor )
+void FCanvasTextItem::DrawStringInternal_OfflineCache(FCanvas* InCanvas, const FVector2D& DrawPos, const FLinearColor& InColor, TArrayView<FTextEffect> TextEffects)
 {
+	QUICK_SCOPE_CYCLE_COUNTER(DrawStringInternal_OfflineCache);
+
 	DrawnSize = FVector2D::ZeroVector;
 
 	// Nothing to do if no text
@@ -1084,29 +1088,41 @@ void FCanvasTextItem::DrawStringInternal_OfflineCache( FCanvas* InCanvas, const
 			const float Right = (X + SizeX) * Depth;
 			const float Bottom = (Y + SizeY) * Depth;
 
-			int32 V00 = BatchedElements->AddVertex(
-				FVector4( Left, Top, 0.f, Depth ),
-				FVector2D( U, V ),
-				InColor,
-				HitProxyId );
-			int32 V10 = BatchedElements->AddVertex(
-				FVector4( Right, Top, 0.0f, Depth ),
-				FVector2D( U + SizeU, V ),			
-				InColor,
-				HitProxyId );
-			int32 V01 = BatchedElements->AddVertex(
-				FVector4( Left, Bottom, 0.0f, Depth ),
-				FVector2D( U, V + SizeV ),	
-				InColor,
-				HitProxyId);
-			int32 V11 = BatchedElements->AddVertex(
-				FVector4( Right, Bottom, 0.0f, Depth ),
-				FVector2D( U + SizeU, V + SizeV ),
-				InColor,
-				HitProxyId);
-
-			BatchedElements->AddTriangle(V00, V10, V11, Tex->GetResource(), BlendMode, FontRenderInfo.GlowInfo);
-			BatchedElements->AddTriangle(V00, V11, V01, Tex->GetResource(), BlendMode, FontRenderInfo.GlowInfo);
+			const auto AddTriangles = [&](const FVector2f& Offset, const FLinearColor& InColor)
+			{
+				int32 V00 = BatchedElements->AddVertexf(
+					FVector4f(Left + Offset.X, Top + Offset.Y, 0.f, Depth),
+					FVector2f(U, V),
+					InColor,
+					HitProxyId);
+				int32 V10 = BatchedElements->AddVertexf(
+					FVector4f(Right + Offset.X, Top + Offset.Y, 0.0f, Depth),
+					FVector2f(U + SizeU, V),
+					InColor,
+					HitProxyId);
+				int32 V01 = BatchedElements->AddVertexf(
+					FVector4f(Left + Offset.X, Bottom + Offset.Y, 0.0f, Depth),
+					FVector2f(U, V + SizeV),
+					InColor,
+					HitProxyId);
+				int32 V11 = BatchedElements->AddVertexf(
+					FVector4f(Right + Offset.X, Bottom + Offset.Y, 0.0f, Depth),
+					FVector2f(U + SizeU, V + SizeV),
+					InColor,
+					HitProxyId);
+
+				BatchedElements->AddTriangle(V00, V10, V11, Tex->GetResource(), BlendMode, FontRenderInfo.GlowInfo);
+				BatchedElements->AddTriangle(V00, V11, V01, Tex->GetResource(), BlendMode, FontRenderInfo.GlowInfo);
+			};
+
+			// Render all the effects first that appear under the text
+			for (const FTextEffect& TextEffect : TextEffects)
+			{
+				AddTriangles(TextEffect.Offset, TextEffect.Color);
+			}
+
+			// The draw the main text
+			AddTriangles(FVector2f::ZeroVector, InColor);
 
 			// if we have another non-whitespace character to render, add the font's kerning.
 			if ( Chars[i+1] && !FChar::IsWhitespace(Chars[i+1]) )
@@ -1126,8 +1142,10 @@ void FCanvasTextItem::DrawStringInternal_OfflineCache( FCanvas* InCanvas, const
 	}
 }
 
-void FCanvasTextItem::DrawStringInternal_RuntimeCache( FCanvas* InCanvas, const FVector2D& DrawPos, const FLinearColor& InColor )
+void FCanvasTextItem::DrawStringInternal_RuntimeCache(FCanvas* InCanvas, const FVector2D& DrawPos, const FLinearColor& InColor, TArrayView<FTextEffect> TextEffects)
 {
+	QUICK_SCOPE_CYCLE_COUNTER(DrawStringInternal_RuntimeCache);
+
 	const float DPIScale = InCanvas->GetDPIScale();
 
 	DrawnSize = FVector2D::ZeroVector;
@@ -1149,7 +1167,7 @@ void FCanvasTextItem::DrawStringInternal_RuntimeCache( FCanvas* InCanvas, const
 	const FSlateFontInfo LegacyFontInfo = (SlateFontInfo.IsSet()) ? SlateFontInfo.GetValue() : Font->GetLegacySlateFontInfo();
 	FCharacterList& CharacterList = FontCache->GetCharacterList( LegacyFontInfo, FontScale );
 
-	FHitProxyId HitProxyId = InCanvas->GetHitProxyId();
+	const FHitProxyId HitProxyId = InCanvas->GetHitProxyId();
 
 	uint32 FontTextureIndex = 0;
 	FTextureResource* FontTexture = nullptr;
@@ -1159,8 +1177,6 @@ void FCanvasTextItem::DrawStringInternal_RuntimeCache( FCanvas* InCanvas, const
 	float InvTextureSizeX = 0;
 	float InvTextureSizeY = 0;
 
-	FCharacterEntry PreviousCharEntry;
-
 	FVector2D TopLeft(0,0);
 
 	const float PosX = TopLeft.X * DPIScale;
@@ -1170,11 +1186,12 @@ void FCanvasTextItem::DrawStringInternal_RuntimeCache( FCanvas* InCanvas, const
 	const float ScaledMaxHeight = CharacterList.GetMaxHeight() * Scale.Y;
 
 	float LineX = PosX;
-	
+
+	TCHAR PreviousChar = 0;
 	const int32 TextLen = TextString.Len();
-	for( int32 CharIndex = 0; CharIndex < TextLen; ++CharIndex )
+	for (int32 CharIndex = 0; CharIndex < TextLen; ++CharIndex)
 	{
-		const TCHAR CurrentChar = TextString[ CharIndex ];
+		const TCHAR CurrentChar = TextString[CharIndex];
 
 		if (DrawnSize.Y == 0)
 		{
@@ -1195,9 +1212,9 @@ void FCanvasTextItem::DrawStringInternal_RuntimeCache( FCanvas* InCanvas, const
 		}
 		else
 		{
-			FCharacterEntry Entry{ CharacterList.GetCharacter(CurrentChar, LegacyFontInfo.FontFallback) };
+			const FCharacterEntry& Entry = CharacterList.GetCharacter(CurrentChar, LegacyFontInfo.FontFallback);
 
-			if( Entry.Valid && (FontTexture == nullptr || Entry.TextureIndex != FontTextureIndex) )
+			if (Entry.Valid && (FontTexture == nullptr || Entry.TextureIndex != FontTextureIndex))
 			{
 				// Font has a new texture for this glyph. Refresh the batch we use and the index we are currently using
 				FontTextureIndex = Entry.TextureIndex;
@@ -1220,24 +1237,33 @@ void FCanvasTextItem::DrawStringInternal_RuntimeCache( FCanvas* InCanvas, const
 				// for the triangles/vertices of the batched elements used to render the text tiles.
 				// Only reserve initial batch, allow growth afterwards in case there are multiple repeated calls.
 				// Reserving exactly the added amount each time would essentially force an alloc each time on subsequent calls.
-				BatchedElements->ReserveTriangles(TextLen*2, FontTexture, FontTextureBlendMode);
-				BatchedElements->ReserveVertices(TextLen*4);
+				{
+					const int32 NumTris = 2 * (TextEffects.Num() + 1);
+					const int32 NumVerts = 4 * (TextEffects.Num() + 1);
+					BatchedElements->ReserveTriangles(TextLen * NumTris, FontTexture, FontTextureBlendMode);
+					BatchedElements->ReserveVertices(TextLen * NumVerts);
+				}
 
-				InvTextureSizeX = 1.0f/FontTexture->GetSizeX();
-				InvTextureSizeY = 1.0f/FontTexture->GetSizeY();
+				InvTextureSizeX = 1.0f / FontTexture->GetSizeX();
+				InvTextureSizeY = 1.0f / FontTexture->GetSizeY();
 			}
 
 			const bool bIsWhitespace = !Entry.Valid || FChar::IsWhitespace(CurrentChar);
 
 			int32 Kerning = 0;
-			if( !bIsWhitespace && PreviousCharEntry.Valid )
+			if (!bIsWhitespace)
 			{
-				Kerning = CharacterList.GetKerning( PreviousCharEntry, Entry ) * Scale.X;
+				if (CharIndex > 0)
+				{
+					// We can't store a pointer to the previous entry as calls to CharacterList.GetCharacter can invalidate previous pointers.
+					const FCharacterEntry& PreviousEntry = CharacterList.GetCharacter(PreviousChar, LegacyFontInfo.FontFallback);
+					Kerning = CharacterList.GetKerning(PreviousEntry, Entry) * Scale.X;
+				}
 			}
 
 			LineX += Kerning;
 
-			if( !bIsWhitespace )
+			if (!bIsWhitespace)
 			{
 				const float InvBitmapRenderScale = 1.0f / Entry.BitmapRenderScale;
 
@@ -1257,29 +1283,41 @@ void FCanvasTextItem::DrawStringInternal_RuntimeCache( FCanvas* InCanvas, const
 				const float Right = (X + SizeX) * Depth;
 				const float Bottom = (Y + SizeY) * Depth;
 
-				int32 V00 = BatchedElements->AddVertex(
-					FVector4( Left, Top, 0.f, Depth ),
-					FVector2D( U, V ),
-					FontColor,
-					HitProxyId );
-				int32 V10 = BatchedElements->AddVertex(
-					FVector4( Right, Top, 0.0f, Depth ),
-					FVector2D( U + SizeU, V ),			
-					FontColor,
-					HitProxyId );
-				int32 V01 = BatchedElements->AddVertex(
-					FVector4( Left, Bottom, 0.0f, Depth ),
-					FVector2D( U, V + SizeV ),	
-					FontColor,
-					HitProxyId);
-				int32 V11 = BatchedElements->AddVertex(
-					FVector4( Right, Bottom, 0.0f, Depth ),
-					FVector2D( U + SizeU, V + SizeV ),
-					FontColor,
-					HitProxyId);
+				const auto AddTriangles = [&](const FVector2f& Offset, const FLinearColor& InColor)
+				{
+					int32 V00 = BatchedElements->AddVertexf(
+						FVector4f(Left + Offset.X, Top + Offset.Y, 0.f, Depth),
+						FVector2f(U, V),
+						InColor,
+						HitProxyId);
+					int32 V10 = BatchedElements->AddVertexf(
+						FVector4f(Right + Offset.X, Top + Offset.Y, 0.0f, Depth),
+						FVector2f(U + SizeU, V),
+						InColor,
+						HitProxyId);
+					int32 V01 = BatchedElements->AddVertexf(
+						FVector4f(Left + Offset.X, Bottom + Offset.Y, 0.0f, Depth),
+						FVector2f(U, V + SizeV),
+						InColor,
+						HitProxyId);
+					int32 V11 = BatchedElements->AddVertexf(
+						FVector4f(Right + Offset.X, Bottom + Offset.Y, 0.0f, Depth),
+						FVector2f(U + SizeU, V + SizeV),
+						InColor,
+						HitProxyId);
+
+					BatchedElements->AddTriangle(V00, V10, V11, FontTexture, FontTextureBlendMode, FontRenderInfo.GlowInfo);
+					BatchedElements->AddTriangle(V00, V11, V01, FontTexture, FontTextureBlendMode, FontRenderInfo.GlowInfo);
+				};
+
+				// Render all the effects first that appear under the text
+				for (const FTextEffect& TextEffect : TextEffects)
+				{
+					AddTriangles(TextEffect.Offset, TextEffect.Color);
+				}
 
-				BatchedElements->AddTriangle(V00, V10, V11, FontTexture, FontTextureBlendMode, FontRenderInfo.GlowInfo);
-				BatchedElements->AddTriangle(V00, V11, V01, FontTexture, FontTextureBlendMode, FontRenderInfo.GlowInfo);
+				// The draw the main text
+				AddTriangles(FVector2f::ZeroVector, FontColor);
 			}
 
 			LineX += Entry.XAdvance * Scale.X;
@@ -1291,13 +1329,11 @@ void FCanvasTextItem::DrawStringInternal_RuntimeCache( FCanvas* InCanvas, const
 				DrawnSize.X = LineX;
 			}
 
-			// MoveTemp to reduce the number of SharedPtr creation/destruction.
-			PreviousCharEntry = MoveTemp(Entry);
+			PreviousChar = CurrentChar;
 		}
 	}
 }
 
-
 bool FCanvasShapedTextItem::HasValidText() const
 {
 	return ShapedGlyphSequence.IsValid() && ShapedGlyphSequence->GetGlyphsToRender().Num();
@@ -1319,7 +1355,7 @@ FVector2D FCanvasShapedTextItem::GetTextSize(float DPIScale) const
 	return FVector2D(ShapedGlyphSequence->GetMeasuredWidth(), ShapedGlyphSequence->GetMaxTextHeight());
 }
 
-void FCanvasShapedTextItem::DrawStringInternal( FCanvas* InCanvas, const FVector2D& DrawPos, const FLinearColor& InColor )
+void FCanvasShapedTextItem::DrawStringInternal(FCanvas* InCanvas, const FVector2D& DrawPos, const FLinearColor& InColor, TArrayView<FTextEffect> TextEffects)
 {
 	DrawnSize = FVector2D::ZeroVector;
 
@@ -1410,29 +1446,41 @@ void FCanvasShapedTextItem::DrawStringInternal( FCanvas* InCanvas, const FVector
 				const float Right = (X + SizeX) * Depth;
 				const float Bottom = (Y + SizeY) * Depth;
 
-				int32 V00 = BatchedElements->AddVertex(
-					FVector4( Left, Top, 0.f, Depth ),
-					FVector2D( U, V ),
-					FontColor,
-					HitProxyId );
-				int32 V10 = BatchedElements->AddVertex(
-					FVector4( Right, Top, 0.0f, Depth ),
-					FVector2D( U + SizeU, V ),			
-					FontColor,
-					HitProxyId );
-				int32 V01 = BatchedElements->AddVertex(
-					FVector4( Left, Bottom, 0.0f, Depth ),
-					FVector2D( U, V + SizeV ),	
-					FontColor,
-					HitProxyId);
-				int32 V11 = BatchedElements->AddVertex(
-					FVector4( Right, Bottom, 0.0f, Depth ),
-					FVector2D( U + SizeU, V + SizeV ),
-					FontColor,
-					HitProxyId);
+				const auto AddTriangles = [&](const FVector2f& Offset, const FLinearColor& InColor)
+				{
+					int32 V00 = BatchedElements->AddVertexf(
+						FVector4f(Left + Offset.X, Top + Offset.Y, 0.f, Depth),
+						FVector2f(U, V),
+						InColor,
+						HitProxyId);
+					int32 V10 = BatchedElements->AddVertexf(
+						FVector4f(Right + Offset.X, Top + Offset.Y, 0.0f, Depth),
+						FVector2f(U + SizeU, V),
+						InColor,
+						HitProxyId);
+					int32 V01 = BatchedElements->AddVertexf(
+						FVector4f(Left + Offset.X, Bottom + Offset.Y, 0.0f, Depth),
+						FVector2f(U, V + SizeV),
+						InColor,
+						HitProxyId);
+					int32 V11 = BatchedElements->AddVertexf(
+						FVector4f(Right + Offset.X, Bottom + Offset.Y, 0.0f, Depth),
+						FVector2f(U + SizeU, V + SizeV),
+						InColor,
+						HitProxyId);
+
+					BatchedElements->AddTriangle(V00, V10, V11, FontTexture, FontTextureBlendMode, FontRenderInfo.GlowInfo);
+					BatchedElements->AddTriangle(V00, V11, V01, FontTexture, FontTextureBlendMode, FontRenderInfo.GlowInfo);
+				};
+
+				// Render all the effects first that appear under the text
+				for (const FTextEffect& TextEffect : TextEffects)
+				{
+					AddTriangles(TextEffect.Offset, TextEffect.Color);
+				}
 
-				BatchedElements->AddTriangle(V00, V10, V11, FontTexture, FontTextureBlendMode, FontRenderInfo.GlowInfo);
-				BatchedElements->AddTriangle(V00, V11, V01, FontTexture, FontTextureBlendMode, FontRenderInfo.GlowInfo);
+				// The draw the main text
+				AddTriangles(FVector2f::ZeroVector, FontColor);
 			}
 		}
 
@@ -1508,9 +1556,9 @@ void FCanvasTriangleItem::Draw( class FCanvas* InCanvas )
 		for (int32 i = 0; i < NumTriangles; i++)
 		{
 			const FCanvasUVTri& Tri = TriangleList[i];
-			int32 V0 = BatchedElements->AddVertex(FVector4(Tri.V0_Pos.X, Tri.V0_Pos.Y, 0, 1), Tri.V0_UV, Tri.V0_Color, HitProxyId);
-			int32 V1 = BatchedElements->AddVertex(FVector4(Tri.V1_Pos.X, Tri.V1_Pos.Y, 0, 1), Tri.V1_UV, Tri.V1_Color, HitProxyId);
-			int32 V2 = BatchedElements->AddVertex(FVector4(Tri.V2_Pos.X, Tri.V2_Pos.Y, 0, 1), Tri.V2_UV, Tri.V2_Color, HitProxyId);
+			int32 V0 = BatchedElements->AddVertexf(FVector4f(Tri.V0_Pos.X, Tri.V0_Pos.Y, 0, 1), FVector2f{Tri.V0_UV}, Tri.V0_Color, HitProxyId);
+			int32 V1 = BatchedElements->AddVertexf(FVector4f(Tri.V1_Pos.X, Tri.V1_Pos.Y, 0, 1), FVector2f{Tri.V1_UV}, Tri.V1_Color, HitProxyId);
+			int32 V2 = BatchedElements->AddVertexf(FVector4f(Tri.V2_Pos.X, Tri.V2_Pos.Y, 0, 1), FVector2f{Tri.V2_UV}, Tri.V2_Color, HitProxyId);
 
 			if (BatchedElementParameters)
 			{
diff --git a/Engine/Source/Runtime/Engine/Public/BatchedElements.h b/Engine/Source/Runtime/Engine/Public/BatchedElements.h
--- a/Engine/Source/Runtime/Engine/Public/BatchedElements.h
+++ b/Engine/Source/Runtime/Engine/Public/BatchedElements.h
@@ -46,6 +46,14 @@ struct FSimpleElementVertex
 
 	FSimpleElementVertex() {}
 
+	FSimpleElementVertex(const FVector4f& InPosition, const FVector2f& InTextureCoordinate, const FLinearColor& InColor, FHitProxyId InHitProxyId) :
+		RelativePosition(InPosition),
+		TilePosition(ForceInitToZero),
+		TextureCoordinate(InTextureCoordinate),
+		Color(InColor),
+		HitProxyIdColor(InHitProxyId.GetColor())
+	{}
+
 	FSimpleElementVertex(const FVector4f& InPosition, const FVector2D& InTextureCoordinate, const FLinearColor& InColor, FHitProxyId InHitProxyId) :
 		RelativePosition(InPosition),
 		TilePosition(ForceInitToZero),
@@ -158,8 +166,10 @@ public:
 	/** Adds a point to the batch. Note only SE_BLEND_Opaque will be used for batched point rendering. */
 	void AddPoint(const FVector& Position,float Size,const FLinearColor& Color,FHitProxyId HitProxyId);
 
+	/** This is for compatibility but should be avoided since it's slower due to conversions. */
+	int32 AddVertex(const FVector4& InPosition, const FVector2D& InTextureCoordinate, const FLinearColor& InColor, FHitProxyId HitProxyId);
 	/** Adds a mesh vertex to the batch. */
-	int32 AddVertex(const FVector4& InPosition,const FVector2D& InTextureCoordinate,const FLinearColor& InColor,FHitProxyId HitProxyId);
+	int32 AddVertexf(const FVector4f& InPosition, const FVector2f& InTextureCoordinate, const FLinearColor& InColor, FHitProxyId HitProxyId);
 
 	/** Adds a triangle to the batch. */
 	void AddTriangle(int32 V0,int32 V1,int32 V2,const FTexture* Texture,EBlendMode BlendMode);
diff --git a/Engine/Source/Runtime/Engine/Public/CanvasItem.h b/Engine/Source/Runtime/Engine/Public/CanvasItem.h
--- a/Engine/Source/Runtime/Engine/Public/CanvasItem.h
+++ b/Engine/Source/Runtime/Engine/Public/CanvasItem.h
@@ -413,12 +413,18 @@ protected:
 	 */
 	virtual FVector2D GetTextSize(float DPIScale) const = 0;
 
+	struct FTextEffect
+	{
+		FVector2f Offset;
+		FLinearColor Color;
+	};
+
 	/** 
 	 * Internal string draw
 	 *
 	 * In a method to make it simpler to do effects like shadow, outline
 	 */
-	virtual void DrawStringInternal( FCanvas* InCanvas, const FVector2D& DrawPos, const FLinearColor& DrawColor ) = 0;
+	virtual void DrawStringInternal(FCanvas* InCanvas, const FVector2D& DrawPos, const FLinearColor& DrawColor, TArrayView<FTextEffect> Offsets) = 0;
 
 	/** 
 	 * These are used by the DrawStringInternal function. 
@@ -480,15 +486,16 @@ protected:
 	virtual bool HasValidText() const override;
 	virtual ESimpleElementBlendMode GetTextBlendMode( const bool bHasShadow ) const override;
 	virtual FVector2D GetTextSize(float DPIScale) const override;
-	virtual void DrawStringInternal( FCanvas* InCanvas, const FVector2D& DrawPos, const FLinearColor& DrawColor ) override;
+	virtual void DrawStringInternal(FCanvas* InCanvas, const FVector2D& DrawPos, const FLinearColor& DrawColor, TArrayView<FTextEffect> TextEffects) override;
+
 
 	/** 
 	 * Internal string draw
 	 *
 	 * In a method to make it simpler to do effects like shadow, outline
 	 */
-	void DrawStringInternal_OfflineCache( FCanvas* InCanvas, const FVector2D& DrawPos, const FLinearColor& DrawColor );
-	void DrawStringInternal_RuntimeCache( FCanvas* InCanvas, const FVector2D& DrawPos, const FLinearColor& DrawColor );
+	void DrawStringInternal_OfflineCache(FCanvas* InCanvas, const FVector2D& DrawPos, const FLinearColor& DrawColor, TArrayView<FTextEffect> TextEffects);
+	void DrawStringInternal_RuntimeCache(FCanvas* InCanvas, const FVector2D& DrawPos, const FLinearColor& DrawColor, TArrayView<FTextEffect> TextEffects);
 };
 
 /* Text item which can handle complex shaped text. */
@@ -517,7 +524,7 @@ protected:
 	virtual bool HasValidText() const override;
 	virtual ESimpleElementBlendMode GetTextBlendMode( const bool bHasShadow ) const override;
 	virtual FVector2D GetTextSize(float DPIScale) const override;
-	virtual void DrawStringInternal( FCanvas* InCanvas, const FVector2D& DrawPos, const FLinearColor& DrawColor ) override;
+	virtual void DrawStringInternal(FCanvas* InCanvas, const FVector2D& DrawPos, const FLinearColor& DrawColor, TArrayView<FTextEffect> TextEffects) override;
 };
 
 /* Line item. Note blend mode will be disregarded for these - only SE_BLEND_Opaque is currently supported. */
diff --git a/Engine/Source/Runtime/Renderer/Private/SceneHitProxyRendering.cpp b/Engine/Source/Runtime/Renderer/Private/SceneHitProxyRendering.cpp
--- a/Engine/Source/Runtime/Renderer/Private/SceneHitProxyRendering.cpp
+++ b/Engine/Source/Runtime/Renderer/Private/SceneHitProxyRendering.cpp
@@ -519,10 +519,10 @@ static void DoRenderHitProxies(
 
 				// Note: High DPI .  We are drawing to the size of the unscaled view rect because that is the size of the views render target
 				// if we do not do this clicking would be off.
-				const int32 V00 = BatchedElements.AddVertex(FVector4(View.UnscaledViewRect.Min.X, View.UnscaledViewRect.Min.Y, 0, 1), FVector2D(U0, V0), FLinearColor::White, FHitProxyId());
-				const int32 V10 = BatchedElements.AddVertex(FVector4(View.UnscaledViewRect.Max.X, View.UnscaledViewRect.Min.Y, 0, 1), FVector2D(U1, V0), FLinearColor::White, FHitProxyId());
-				const int32 V01 = BatchedElements.AddVertex(FVector4(View.UnscaledViewRect.Min.X, View.UnscaledViewRect.Max.Y, 0, 1), FVector2D(U0, V1), FLinearColor::White, FHitProxyId());
-				const int32 V11 = BatchedElements.AddVertex(FVector4(View.UnscaledViewRect.Max.X, View.UnscaledViewRect.Max.Y, 0, 1), FVector2D(U1, V1), FLinearColor::White, FHitProxyId());
+				const int32 V00 = BatchedElements.AddVertexf(FVector4f(View.UnscaledViewRect.Min.X, View.UnscaledViewRect.Min.Y, 0, 1), FVector2f(U0, V0), FLinearColor::White, FHitProxyId());
+				const int32 V10 = BatchedElements.AddVertexf(FVector4f(View.UnscaledViewRect.Max.X, View.UnscaledViewRect.Min.Y, 0, 1), FVector2f(U1, V0), FLinearColor::White, FHitProxyId());
+				const int32 V01 = BatchedElements.AddVertexf(FVector4f(View.UnscaledViewRect.Min.X, View.UnscaledViewRect.Max.Y, 0, 1), FVector2f(U0, V1), FLinearColor::White, FHitProxyId());
+				const int32 V11 = BatchedElements.AddVertexf(FVector4f(View.UnscaledViewRect.Max.X, View.UnscaledViewRect.Max.Y, 0, 1), FVector2f(U1, V1), FLinearColor::White, FHitProxyId());
 
 				BatchedElements.AddTriangle(V00, V10, V11, &HitProxyRenderTargetTexture, BLEND_Opaque);
 				BatchedElements.AddTriangle(V00, V11, V01, &HitProxyRenderTargetTexture, BLEND_Opaque);
diff --git a/Engine/Source/Runtime/SlateCore/Private/Fonts/FontCache.cpp b/Engine/Source/Runtime/SlateCore/Private/Fonts/FontCache.cpp
--- a/Engine/Source/Runtime/SlateCore/Private/Fonts/FontCache.cpp
+++ b/Engine/Source/Runtime/SlateCore/Private/Fonts/FontCache.cpp
@@ -46,14 +46,6 @@ FAutoConsoleVariableRef CVarUnloadFreeTypeDataOnFlush(
 	UnloadFreeTypeDataOnFlush,
 	TEXT("Releases the free type data when the font cache is flushed"));
 
-
-namespace FontCacheConstants
-{
-	/** Number of characters that can be indexed directly in the cache */
-	const int32 DirectAccessSize = 256;
-}
-
-
 static TAutoConsoleVariable<int32> CVarDefaultTextShapingMethod(
 	TEXT("Slate.DefaultTextShapingMethod"),
 	static_cast<int32>(ETextShapingMethod::Auto),
@@ -592,7 +584,6 @@ FCharacterList::FCharacterList( const FSlateFontKey& InFontKey, FSlateFontCache&
 #if WITH_EDITORONLY_DATA
 	, CompositeFontHistoryRevision( INDEX_NONE )
 #endif	// WITH_EDITORONLY_DATA
-	, MaxDirectIndexedEntries( FontCacheConstants::DirectAccessSize )
 	, MaxHeight( 0 )
 	, Baseline( 0 )
 {
@@ -615,8 +606,8 @@ bool FCharacterList::IsStale() const
 
 int8 FCharacterList::GetKerning(TCHAR FirstChar, TCHAR SecondChar, const EFontFallback MaxFontFallback)
 {
-	const FCharacterEntry First = GetCharacter(FirstChar, MaxFontFallback);
-	const FCharacterEntry Second = GetCharacter(SecondChar, MaxFontFallback);
+	const FCharacterEntry& First = GetCharacter(FirstChar, MaxFontFallback);
+	const FCharacterEntry& Second = GetCharacter(SecondChar, MaxFontFallback);
 	return GetKerning(First, Second);
 }
 
@@ -684,93 +675,40 @@ bool FCharacterList::CanCacheCharacter(TCHAR Character, const EFontFallback MaxF
 	return bReturnVal;
 }
 
-FCharacterEntry FCharacterList::GetCharacter(TCHAR Character, const EFontFallback MaxFontFallback)
+const FCharacterEntry& FCharacterList::GetCharacter(TCHAR Character, const EFontFallback MaxFontFallback)
 {
 #if WITH_FREETYPE
-	const FCharacterListEntry* InternalEntry = nullptr;
-	const bool bDirectIndexChar = Character < MaxDirectIndexedEntries;
-
-	// First get a reference to the character, if it is already mapped (mapped does not mean cached though)
-	if (bDirectIndexChar)
-	{
-		if (DirectIndexEntries.IsValidIndex(Character))
-		{
-			InternalEntry = &DirectIndexEntries[Character];
-		}
-	}
-	else
+	if (const FCharacterEntry* const FoundEntry = MappedEntries.Find(Character))
 	{
-		const FCharacterListEntry* const FoundEntry = MappedEntries.Find(Character);
-		if (FoundEntry)
-		{
-			InternalEntry = FoundEntry;
-		}
-	}
-
-	// Determine whether the character needs caching, and map it if needed
-	bool bNeedCaching = false;
-
-	if (InternalEntry)
-	{
-		bNeedCaching = !InternalEntry->Valid;
-
-		// If the character needs caching, but can't be cached, reject the character
-		if (bNeedCaching && !CanCacheCharacter(Character, MaxFontFallback))
+		// For already-cached characters, reject characters that don't fall within maximum font fallback level requirements
+		if (Character == SlateFontRendererUtils::InvalidSubChar || MaxFontFallback >= FoundEntry->FallbackLevel)
 		{
-			bNeedCaching = false;
-			InternalEntry = nullptr;
+			return *FoundEntry;
 		}
 	}
-	// Only map the character if it can be cached
 	else if (CanCacheCharacter(Character, MaxFontFallback))
 	{
-		bNeedCaching = true;
-
-		if (bDirectIndexChar)
+		if (const FCharacterEntry* NewEntry = CacheCharacter(Character))
 		{
-			DirectIndexEntries.AddZeroed((Character - DirectIndexEntries.Num()) + 1);
-			InternalEntry = &DirectIndexEntries[Character];
-		}
-		else
-		{
-			InternalEntry = &MappedEntries.Add(Character);
+			return *NewEntry;
 		}
 	}
-
-	if (InternalEntry)
-	{
-		if (bNeedCaching)
-		{
-			InternalEntry = CacheCharacter(Character);
-		}
-		// For already-cached characters, reject characters that don't fall within maximum font fallback level requirements
-		else if (Character != SlateFontRendererUtils::InvalidSubChar && MaxFontFallback < InternalEntry->FallbackLevel)
-		{
-			InternalEntry = nullptr;
-		}
-	}
-
-	if (InternalEntry)
-	{
-		return MakeCharacterEntry(Character, *InternalEntry);
-	}
-
-	// We have exhausted our options, looks like even InvalidSubChar can't be cached
-	// so just return an invalid character instead of going into an infinite loop.
-	if (Character == SlateFontRendererUtils::InvalidSubChar)
+	
+	// If we weren't able to cache the character, try fetching an invalid character to display instead.
+	// If we're already invalid, just return an invalid entry so we don't loop.
+	if (Character != SlateFontRendererUtils::InvalidSubChar)
 	{
-		return FCharacterEntry{};
+		return GetCharacter(SlateFontRendererUtils::InvalidSubChar, MaxFontFallback);
 	}
+#endif
 
-	return GetCharacter(SlateFontRendererUtils::InvalidSubChar, MaxFontFallback);
-#else
 	// CacheCharacter always returns invalid characters when FreeType is not available
 	// so just shortcut to this instead.
-	return FCharacterEntry{};
-#endif
+	static const FCharacterEntry Invalid{};
+	return Invalid;
 }
 
-FCharacterList::FCharacterListEntry* FCharacterList::CacheCharacter(TCHAR Character)
+const FCharacterEntry* FCharacterList::CacheCharacter(TCHAR Character)
 {
 #if WITH_FREETYPE
 	// Fake shape the character
@@ -811,33 +749,46 @@ FCharacterList::FCharacterListEntry* FCharacterList::CacheCharacter(TCHAR Charac
 				}
 			}
 
-			FCharacterListEntry NewInternalEntry;
-			NewInternalEntry.ShapedGlyphEntry.FontFaceData = MakeShared<FShapedGlyphFaceData>(FaceGlyphData.FaceAndMemory, GlyphFlags, FontInfo.Size, FinalFontScale, FontInfo.GetClampSkew());
-			NewInternalEntry.ShapedGlyphEntry.GlyphIndex = GlyphIndex;
-			NewInternalEntry.ShapedGlyphEntry.XAdvance = XAdvance;
-			NewInternalEntry.ShapedGlyphEntry.bIsVisible = !bIsWhitespace;
+			FShapedGlyphEntry ShapedGlyphEntry;
+			ShapedGlyphEntry.FontFaceData = MakeShared<FShapedGlyphFaceData>(FaceGlyphData.FaceAndMemory, GlyphFlags, FontInfo.Size, FinalFontScale, FontInfo.GetClampSkew());
+			ShapedGlyphEntry.GlyphIndex = GlyphIndex;
+			ShapedGlyphEntry.XAdvance = XAdvance;
+			ShapedGlyphEntry.bIsVisible = !bIsWhitespace;
 
-			NewInternalEntry.KerningCache = FontCache.FontRenderer->GetKerningCache(*FontDataPtr, FontInfo.Size, FinalFontScale);
-
-			NewInternalEntry.FontData = FontDataPtr;
-			NewInternalEntry.FallbackLevel = FaceGlyphData.CharFallbackLevel;
-			NewInternalEntry.HasKerning = bHasKerning;
-			NewInternalEntry.Valid = Character == 0 || GlyphIndex != 0;
+			FCharacterEntry CharEntry;
+			CharEntry.Valid = Character == 0 || GlyphIndex != 0;
 
 			// Cache the shaped entry in the font cache
-			if (NewInternalEntry.Valid)
+			if (CharEntry.Valid)
 			{
-				FontCache.GetShapedGlyphFontAtlasData(NewInternalEntry.ShapedGlyphEntry, FontKey.GetFontOutlineSettings());
-
-				if (Character < MaxDirectIndexedEntries)
+				const FShapedGlyphFontAtlasData ShapedGlyphFontAtlasData = FontCache.GetShapedGlyphFontAtlasData(ShapedGlyphEntry, FontKey.GetFontOutlineSettings());
+				if (ShapedGlyphFontAtlasData.Valid)
 				{
-					DirectIndexEntries[Character] = MoveTemp(NewInternalEntry);
-					return &DirectIndexEntries[Character];
+					CharEntry.Character = Character;
+					CharEntry.GlyphIndex = GlyphIndex;
+					CharEntry.FontData = FontDataPtr;
+					CharEntry.KerningCache = FontCache.FontRenderer->GetKerningCache(*FontDataPtr, FontInfo.Size, FinalFontScale);
+					CharEntry.FontScale = ShapedGlyphEntry.FontFaceData->FontScale;
+					CharEntry.BitmapRenderScale = ShapedGlyphEntry.FontFaceData->BitmapRenderScale;
+					CharEntry.StartU = ShapedGlyphFontAtlasData.StartU;
+					CharEntry.StartV = ShapedGlyphFontAtlasData.StartV;
+					CharEntry.USize = ShapedGlyphFontAtlasData.USize;
+					CharEntry.VSize = ShapedGlyphFontAtlasData.VSize;
+					CharEntry.VerticalOffset = ShapedGlyphFontAtlasData.VerticalOffset;
+					CharEntry.HorizontalOffset = ShapedGlyphFontAtlasData.HorizontalOffset;
+					CharEntry.GlobalDescender = GetBaseline(); // All fonts within a composite font need to use the baseline of the default font
+					CharEntry.XAdvance = ShapedGlyphEntry.XAdvance;
+					CharEntry.TextureIndex = ShapedGlyphFontAtlasData.TextureIndex;
+					CharEntry.FallbackLevel = FaceGlyphData.CharFallbackLevel;
+					CharEntry.HasKerning = bHasKerning;
+					CharEntry.SupportsOutline = ShapedGlyphFontAtlasData.SupportsOutline;
 				}
 				else
 				{
-					return &MappedEntries.Add(Character, MoveTemp(NewInternalEntry));
+					CharEntry.Valid = false;
 				}
+
+				return &MappedEntries.Add(Character, MoveTemp(CharEntry));
 			}
 		}
 	}
@@ -847,42 +798,6 @@ FCharacterList::FCharacterListEntry* FCharacterList::CacheCharacter(TCHAR Charac
 	return nullptr;
 }
 
-FCharacterEntry FCharacterList::MakeCharacterEntry(TCHAR Character, const FCharacterListEntry& InternalEntry) const
-{
-	FCharacterEntry CharEntry;
-
-	CharEntry.Valid = InternalEntry.Valid;
-	if (CharEntry.Valid)
-	{
-		FShapedGlyphFontAtlasData ShapedGlyphFontAtlasData = FontCache.GetShapedGlyphFontAtlasData(InternalEntry.ShapedGlyphEntry, FontKey.GetFontOutlineSettings());
-		CharEntry.Valid = ShapedGlyphFontAtlasData.Valid;
-
-		if (CharEntry.Valid)
-		{
-			CharEntry.Character = Character;
-			CharEntry.GlyphIndex = InternalEntry.ShapedGlyphEntry.GlyphIndex;
-			CharEntry.FontData = InternalEntry.FontData;
-			CharEntry.KerningCache = InternalEntry.KerningCache;
-			CharEntry.FontScale = InternalEntry.ShapedGlyphEntry.FontFaceData->FontScale;
-			CharEntry.BitmapRenderScale = InternalEntry.ShapedGlyphEntry.FontFaceData->BitmapRenderScale;
-			CharEntry.StartU = ShapedGlyphFontAtlasData.StartU;
-			CharEntry.StartV = ShapedGlyphFontAtlasData.StartV;
-			CharEntry.USize = ShapedGlyphFontAtlasData.USize;
-			CharEntry.VSize = ShapedGlyphFontAtlasData.VSize;
-			CharEntry.VerticalOffset = ShapedGlyphFontAtlasData.VerticalOffset;
-			CharEntry.HorizontalOffset = ShapedGlyphFontAtlasData.HorizontalOffset;
-			CharEntry.GlobalDescender = GetBaseline(); // All fonts within a composite font need to use the baseline of the default font
-			CharEntry.XAdvance = InternalEntry.ShapedGlyphEntry.XAdvance;
-			CharEntry.TextureIndex = ShapedGlyphFontAtlasData.TextureIndex;
-			CharEntry.HasKerning = InternalEntry.HasKerning;
-			CharEntry.SupportsOutline = ShapedGlyphFontAtlasData.SupportsOutline;
-			CharEntry.FallbackLevel = InternalEntry.FallbackLevel;
-		}
-	}
-
-	return CharEntry;
-}
-
 FSlateFontCache::FSlateFontCache( TSharedRef<ISlateFontAtlasFactory> InFontAtlasFactory, ESlateTextureAtlasThreadId InOwningThread)
 	: FSlateFlushableAtlasCache(&FontAtlasFlushParams)
 	, FTLibrary( new FFreeTypeLibrary() )
diff --git a/Engine/Source/Runtime/SlateCore/Public/Fonts/FontCache.h b/Engine/Source/Runtime/SlateCore/Public/Fonts/FontCache.h
--- a/Engine/Source/Runtime/SlateCore/Public/Fonts/FontCache.h
+++ b/Engine/Source/Runtime/SlateCore/Public/Fonts/FontCache.h
@@ -511,21 +511,16 @@ class SLATECORE_API FCharacterList
 public:
 	FCharacterList( const FSlateFontKey& InFontKey, FSlateFontCache& InFontCache );
 
-	/* @return Is the character in this list */
-	bool IsValidIndex( TCHAR Character ) const
-	{
-		return DirectIndexEntries.IsValidIndex( Character ) || ( Character >= MaxDirectIndexedEntries && MappedEntries.Contains( Character ) );
-	}
-
 	/**
-	 * Gets data about how to render and measure a character 
-	 * Caching and atlasing it if needed
+	 * Gets data about how to render and measure a character.
+	 * Caching and atlasing it if needed.
+	 * Subsequent calls may invalidate previous pointers.
 	 *
 	 * @param Character			The character to get
 	 * @param MaxFontFallback	The maximum fallback level that can be used when resolving glyphs
 	 * @return				Data about the character
 	 */
-	FCharacterEntry GetCharacter(TCHAR Character, const EFontFallback MaxFontFallback);
+	const FCharacterEntry& GetCharacter(TCHAR Character, const EFontFallback MaxFontFallback);
 
 #if WITH_EDITORONLY_DATA
 	/** Check to see if our cached data is potentially stale for our font */
@@ -564,23 +559,6 @@ public:
 	int16 GetBaseline() const;
 
 private:
-	/** Maintains a fake shaped glyph for each character in the character list */
-	struct FCharacterListEntry
-	{
-		/** The shaped glyph data for this character */
-		FShapedGlyphEntry ShapedGlyphEntry;
-		/** Font data this character was rendered with */
-		const FFontData* FontData = nullptr;
-		/** Kerning cache that this character uses */
-		TSharedPtr<FFreeTypeKerningCache> KerningCache;
-		/** The fallback level this character represents */
-		EFontFallback FallbackLevel = EFontFallback::FF_Max;
-		/** Does this character have kerning? */
-		bool HasKerning = false;
-		/** Has this entry been initialized? */
-		bool Valid = false;
-	};
-
 	/**
 	 * Returns whether the specified character is valid for caching (i.e. whether it matches the FontFallback level)
 	 *
@@ -594,18 +572,12 @@ private:
 	 * 
 	 * @param Character	The character to cache
 	 */
-	FCharacterListEntry* CacheCharacter(TCHAR Character);
-
-	/**
-	 * Convert the cached internal entry to the external data for the old non-shaped API
-	 */
-	FCharacterEntry MakeCharacterEntry(TCHAR Character, const FCharacterListEntry& InternalEntry) const;
+	const FCharacterEntry* CacheCharacter(TCHAR Character);
 
 private:
 	/** Entries for larger character sets to conserve memory */
-	TMap<TCHAR, FCharacterListEntry> MappedEntries; 
-	/** Directly indexed entries for fast lookup */
-	TArray<FCharacterListEntry> DirectIndexEntries;
+	TMap<TCHAR, FCharacterEntry> MappedEntries;
+
 	/** Font for this character list */
 	FSlateFontKey FontKey;
 	/** Reference to the font cache for accessing new unseen characters */
@@ -614,8 +586,6 @@ private:
 	/** The history revision of the cached composite font */
 	int32 CompositeFontHistoryRevision;
 #endif	// WITH_EDITORONLY_DATA
-	/** Number of directly indexed entries */
-	int32 MaxDirectIndexedEntries;
 	/** The global max height for any character in this font */
 	mutable uint16 MaxHeight;
 	/** The offset from the bottom of the max character height to the baseline. */
