[KillPendingKill] Don't truncate reference trace on first garbage object.

diff --git a/Engine/Source/Runtime/Engine/Private/UnrealEngine.cpp b/Engine/Source/Runtime/Engine/Private/UnrealEngine.cpp
--- a/Engine/Source/Runtime/Engine/Private/UnrealEngine.cpp
+++ b/Engine/Source/Runtime/Engine/Private/UnrealEngine.cpp
@@ -15460,9 +15460,7 @@ static bool PrintStaleReferenceChainsAndFindReferencingObjects(UObject* ObjectTo
 					OutGarbageObject = Params.Object;
 					OutReferencingObject = Params.Referencer;
 				}
-				// We only really care about a reference to the first garbage / pending kill object in the chain
-				// as the reference that holds it is the one that prevents the world from being GC'd
-				return false;
+				return true;
 			}
 			else
 			{
