WorldPartitionCookPackageSplitter: Changed the interface with the cooker to allow returning altered packages, which will be kept referenced during GC, rather than needing to pre-identify all objects that will be moved.

diff --git a/Engine/Source/Editor/UnrealEd/Classes/CookOnTheSide/CookOnTheFlyServer.h b/Engine/Source/Editor/UnrealEd/Classes/CookOnTheSide/CookOnTheFlyServer.h
--- a/Engine/Source/Editor/UnrealEd/Classes/CookOnTheSide/CookOnTheFlyServer.h
+++ b/Engine/Source/Editor/UnrealEd/Classes/CookOnTheSide/CookOnTheFlyServer.h
@@ -1206,7 +1206,8 @@ private:
 		TArray<ICookPackageSplitter::FGeneratedPackageForPreSave>& GeneratedPackagesForPresave);
 	/** Call the Generator's PreSaveGeneratorPackage to create/move objects into its main UPackage. */
 	UE::Cook::EPollStatus PreSaveGeneratorPackage(UE::Cook::FPackageData& PackageData,
-		UE::Cook::FGeneratorPackage& Generator, TArray<ICookPackageSplitter::FGeneratedPackageForPreSave>& GeneratedPackagesForPresave);
+		UE::Cook::FGeneratorPackage& Generator, UE::Cook::FCookGenerationInfo& Info,
+		TArray<ICookPackageSplitter::FGeneratedPackageForPreSave>& GeneratedPackagesForPresave);
 	/** Construct the list of generated packages that is required for some of the CookPackageSplitter interface calls. */
 	void ConstructGeneratedPackagesForPresave(UE::Cook::FPackageData& PackageData, UE::Cook::FGeneratorPackage& Generator,
 		TArray<ICookPackageSplitter::FGeneratedPackageForPreSave>& GeneratedPackagesForPresave);
diff --git a/Engine/Source/Editor/UnrealEd/Private/CookOnTheFlyServer.cpp b/Engine/Source/Editor/UnrealEd/Private/CookOnTheFlyServer.cpp
--- a/Engine/Source/Editor/UnrealEd/Private/CookOnTheFlyServer.cpp
+++ b/Engine/Source/Editor/UnrealEd/Private/CookOnTheFlyServer.cpp
@@ -2897,7 +2897,7 @@ UE::Cook::EPollStatus UCookOnTheFlyServer::PrepareSaveGeneratedPackage(UE::Cook:
 		EPollStatus Result;
 		if (Info.IsGenerator())
 		{
-			Result = PreSaveGeneratorPackage(PackageData, Generator, GeneratedPackagesForPresave);
+			Result = PreSaveGeneratorPackage(PackageData, Generator, Info, GeneratedPackagesForPresave);
 		}
 		else
 		{
@@ -2970,15 +2970,18 @@ UE::Cook::EPollStatus UCookOnTheFlyServer::BeginCacheObjectsToMove(UE::Cook::FGe
 	{
 		bool bPopulateSucceeded = false;
 		TArray<UObject*> ObjectsToMove;
+		TArray<UPackage*> ModifiedPackages;
 		if (Info.IsGenerator())
 		{
 			ConstructGeneratedPackagesForPresave(PackageData, Generator, GeneratedPackagesForPresave);
-			bPopulateSucceeded = Splitter->PopulateGeneratorPackage(Package, SplitDataObject, GeneratedPackagesForPresave, ObjectsToMove);
+			bPopulateSucceeded = Splitter->PopulateGeneratorPackage(Package, SplitDataObject, GeneratedPackagesForPresave,
+				ObjectsToMove, ModifiedPackages);
 		}
 		else
 		{
 			ICookPackageSplitter::FGeneratedPackageForPopulate SplitterInfo{ Info.RelativePath, Package, Info.IsCreateAsMap() };
-			bPopulateSucceeded = Splitter->PopulateGeneratedPackage(Package, SplitDataObject, SplitterInfo, ObjectsToMove);
+			bPopulateSucceeded = Splitter->PopulateGeneratedPackage(Package, SplitDataObject, SplitterInfo,
+				ObjectsToMove, ModifiedPackages);
 		}
 
 		if (!bPopulateSucceeded)
@@ -2990,6 +2993,7 @@ UE::Cook::EPollStatus UCookOnTheFlyServer::BeginCacheObjectsToMove(UE::Cook::FGe
 			return EPollStatus::Error;
 		}
 
+		Info.AddModifiedPackages(ModifiedPackages);
 		Info.TakeOverCachedObjectsAndAddMoved(Generator, PackageData.GetCachedObjectsInOuter(), ObjectsToMove);
 		Info.SetSaveStateComplete(FCookGenerationInfo::ESaveState::CallObjectsToMove);
 	}
@@ -3004,7 +3008,8 @@ UE::Cook::EPollStatus UCookOnTheFlyServer::BeginCacheObjectsToMove(UE::Cook::FGe
 	return EPollStatus::Success;
 }
 
-UE::Cook::EPollStatus UCookOnTheFlyServer::PreSaveGeneratorPackage(UE::Cook::FPackageData& PackageData, UE::Cook::FGeneratorPackage& Generator,
+UE::Cook::EPollStatus UCookOnTheFlyServer::PreSaveGeneratorPackage(UE::Cook::FPackageData& PackageData,
+	UE::Cook::FGeneratorPackage& Generator, UE::Cook::FCookGenerationInfo& Info,
 	TArray<ICookPackageSplitter::FGeneratedPackageForPreSave>& GeneratedPackagesForPresave)
 {
 	using namespace UE::Cook;
@@ -3020,13 +3025,16 @@ UE::Cook::EPollStatus UCookOnTheFlyServer::PreSaveGeneratorPackage(UE::Cook::FPa
 		return EPollStatus::Error;
 	}
 
+	TArray<UPackage*> ModifiedPackages;
 	ConstructGeneratedPackagesForPresave(PackageData, Generator, GeneratedPackagesForPresave);
-	if (!Splitter->PreSaveGeneratorPackage(Package, SplitDataObject, GeneratedPackagesForPresave))
+	if (!Splitter->PreSaveGeneratorPackage(Package, SplitDataObject, GeneratedPackagesForPresave, ModifiedPackages))
 	{
 		UE_LOG(LogCook, Error, TEXT("PackageSplitter returned false from PreSaveGeneratorPackage. Splitter=%s"),
 			*Generator.GetSplitDataObjectName().ToString());
 		return EPollStatus::Error;
 	}
+	Info.AddModifiedPackages(ModifiedPackages);
+
 	return EPollStatus::Success;
 }
 
@@ -3095,7 +3103,9 @@ UE::Cook::EPollStatus UCookOnTheFlyServer::BeginCachePostMove(UE::Cook::FGenerat
 
 	EPollStatus Result = CallBeginCacheOnObjects(PackageData, Package, Info.BeginCacheObjects.ObjectsInRound,
 		Info.BeginCacheObjects.NextIndexInRound, Timer);
-	if (PackageData.GetNumPendingCookedPlatformData() > 0 && !Info.HasIssuedUndeclaredMovedObjectsWarning())
+	if (PackageData.GetNumPendingCookedPlatformData() > 0 &&
+		!Generator.GetCookPackageSplitterInstance()->UseInternalReferenceToAvoidGarbageCollect() &&
+		!Info.HasIssuedUndeclaredMovedObjectsWarning())
 	{
 		UObject* FirstPendingObject = nullptr;
 		FString FirstPendingObjectName;
@@ -3197,12 +3207,14 @@ UE::Cook::EPollStatus UCookOnTheFlyServer::TryPopulateGeneratedPackage(UE::Cook:
 	PopulateData.RelativePath = GeneratedInfo.RelativePath;
 	PopulateData.Package = GeneratedPackage;
 	PopulateData.bCreatedAsMap = GeneratedInfo.IsCreateAsMap();
-	if (!Splitter->PreSaveGeneratedPackage(OwnerPackage, OwnerObject, PopulateData))
+	TArray<UPackage*> ModifiedPackages;
+	if (!Splitter->PreSaveGeneratedPackage(OwnerPackage, OwnerObject, PopulateData, ModifiedPackages))
 	{
 		UE_LOG(LogCook, Error, TEXT("PackageSplitter returned false from PreSaveGeneratedPackage. Splitter=%s, Generated=%s."),
 			*Generator.GetSplitDataObjectName().ToString(), *GeneratedPackageName);
 		return EPollStatus::Error;
 	}
+	GeneratedInfo.AddModifiedPackages(ModifiedPackages);
 	bool bPackageIsMap = GeneratedPackage->ContainsMap();
 	if (bPackageIsMap != GeneratedInfo.IsCreateAsMap())
 	{
diff --git a/Engine/Source/Editor/UnrealEd/Private/Cooker/CookPackageData.cpp b/Engine/Source/Editor/UnrealEd/Private/Cooker/CookPackageData.cpp
--- a/Engine/Source/Editor/UnrealEd/Private/Cooker/CookPackageData.cpp
+++ b/Engine/Source/Editor/UnrealEd/Private/Cooker/CookPackageData.cpp
@@ -1402,6 +1402,7 @@ void FGeneratorPackage::PreGarbageCollect(FCookGenerationInfo& Info, TArray<UObj
 				GCKeepPackages.Add(Package);
 				GCKeepPackageDatas.Add(Info.PackageData);
 			}
+			GCKeepPackages.Append(Info.ModifiedPackages);
 			for (FBeginCacheObject& BeginCacheObject : Info.BeginCacheObjects.Objects)
 			{
 				UObject* Object = BeginCacheObject.Object.Get();
@@ -1571,6 +1572,7 @@ void FGeneratorPackage::ResetSaveState(FCookGenerationInfo& Info, UPackage* Pack
 	Info.BeginCacheObjects.Reset();
 	Info.SetHasTakenOverCachedCookedPlatformData(false);
 	Info.SetHasIssuedUndeclaredMovedObjectsWarning(false);
+	Info.ModifiedPackages.Reset();
 }
 
 void FGeneratorPackage::UpdateSaveAfterGarbageCollect(const FPackageData& PackageData, bool& bInOutDemote)
diff --git a/Engine/Source/Editor/UnrealEd/Private/Cooker/CookPackageData.h b/Engine/Source/Editor/UnrealEd/Private/Cooker/CookPackageData.h
--- a/Engine/Source/Editor/UnrealEd/Private/Cooker/CookPackageData.h
+++ b/Engine/Source/Editor/UnrealEd/Private/Cooker/CookPackageData.h
@@ -668,11 +668,14 @@ public:
 	EPollStatus RefreshPackageObjects(FGeneratorPackage& Generator, UPackage* Package, bool& bOutFoundNewObjects,
 		ESaveState DemotionState);
 
+	void AddModifiedPackages(TArray<UPackage*>& InModifiedPackages) { ModifiedPackages.Append(InModifiedPackages); }
+
 public:
 	FBeginCacheObjects BeginCacheObjects;
 	FString RelativePath;
 	TArray<FName> Dependencies;
 	FPackageData* PackageData = nullptr;
+	TArray<UPackage*> ModifiedPackages;
 private:
 	ESaveState GeneratorSaveState = ESaveState::StartGenerate;
 	bool bCreateAsMap : 1;
diff --git a/Engine/Source/Editor/UnrealEd/Public/CookPackageSplitter.h b/Engine/Source/Editor/UnrealEd/Public/CookPackageSplitter.h
--- a/Engine/Source/Editor/UnrealEd/Public/CookPackageSplitter.h
+++ b/Engine/Source/Editor/UnrealEd/Public/CookPackageSplitter.h
@@ -85,14 +85,15 @@ public:
 	 * @return							True if successfully populates,  false on error (this will cause a cook error).
 	 */
 	virtual bool PopulateGeneratedPackage(UPackage* OwnerPackage, UObject* OwnerObject,
-		const FGeneratedPackageForPopulate& GeneratedPackage, TArray<UObject*>& OutObjectsToMove) { return false; }
+		const FGeneratedPackageForPopulate& GeneratedPackage, TArray<UObject*>& OutObjectsToMove,
+		TArray<UPackage*>& OutModifiedPackages) { return true; }
 	
 	/**
 	 * Called before saving the parent generator package, which itself occurs before PreSaveGeneratedPackage is called on the generated packages.
 	 * Make any required adjustments to the parent package before it is saved into the target domain.
 	 */
 	virtual bool PreSaveGeneratedPackage(UPackage* OwnerPackage, UObject* OwnerObject,
-		const FGeneratedPackageForPopulate& GeneratedPackage) = 0;
+		const FGeneratedPackageForPopulate& GeneratedPackage, TArray<UPackage*>& OutModifiedPackages) { return true; };
 	/**
 	 * Called after saving the generated package. Undo any required adjustments to the parent package that
 	 * were made in PreSaveGeneratedPackage, so that the package is once again ready for use in the editor or in
@@ -119,13 +120,15 @@ public:
 	 * can call BeginCacheForCookedPlatformData on them before the move
 	 */
 	virtual bool PopulateGeneratorPackage(UPackage* OwnerPackage, UObject* OwnerObject,
-		const TArray<ICookPackageSplitter::FGeneratedPackageForPreSave>& GeneratedPackages, TArray<UObject*>& OutObjectsToMove) { return false; }
+		const TArray<ICookPackageSplitter::FGeneratedPackageForPreSave>& GeneratedPackages, TArray<UObject*>& OutObjectsToMove,
+		TArray<UPackage*>& OutModifiedPackages) { return true; }
 
 	/**
 	 * Called before saving the parent generator package, which itself occurs before PreSaveGeneratedPackage is called on the generated packages.
 	 * Make any required adjustments to the parent package before it is saved into the target domain.
 	 */
-	virtual bool PreSaveGeneratorPackage(UPackage* OwnerPackage, UObject* OwnerObject, const TArray<FGeneratedPackageForPreSave>& PlaceholderPackages) { return false; }
+	virtual bool PreSaveGeneratorPackage(UPackage* OwnerPackage, UObject* OwnerObject, const TArray<FGeneratedPackageForPreSave>& PlaceholderPackages,
+		TArray<UPackage*>& OutModifiedPackages) { return true; }
 
 	/**
 	 * Called after saving the parent generator package. Undo any required adjustments to the parent package that
diff --git a/Engine/Source/Runtime/Engine/Private/WorldPartition/WorldPartition.cpp b/Engine/Source/Runtime/Engine/Private/WorldPartition/WorldPartition.cpp
--- a/Engine/Source/Runtime/Engine/Private/WorldPartition/WorldPartition.cpp
+++ b/Engine/Source/Runtime/Engine/Private/WorldPartition/WorldPartition.cpp
@@ -1256,34 +1256,16 @@ void UWorldPartition::DrawRuntimeHashPreview()
 	RuntimeHash->DrawPreview();
 }
 
-bool UWorldPartition::PopulateGeneratorPackageForCook(TArray<UObject*>& OutLoadedObjects)
+bool UWorldPartition::PopulateGeneratorPackageForCook(const TArray<ICookPackageSplitter::FGeneratedPackageForPreSave>& InGeneratedPackages, TArray<UPackage*>& OutModifiedPackages)
 {
 	check(RuntimeHash);
-	return RuntimeHash->PopulateGeneratorPackageForCook(OutLoadedObjects);
+	return RuntimeHash->PopulateGeneratorPackageForCook(InGeneratedPackages, OutModifiedPackages);
 }
 
-bool UWorldPartition::PopulateGeneratedPackageForCook(UPackage* InPackage, const FString& InPackageRelativePath, TArray<UObject*>& OutLoadedObjects)
+bool UWorldPartition::PopulateGeneratedPackageForCook(UPackage* InPackage, const FString& InPackageRelativePath, TArray<UPackage*>& OutModifiedPackages)
 {
 	check(RuntimeHash);
-	return RuntimeHash->PopulateGeneratedPackageForCook(InPackage, InPackageRelativePath, OutLoadedObjects);
-}
-
-bool UWorldPartition::FinalizeGeneratedPackageForCook(const FString& InPackageRelativePath)
-{
-	check(RuntimeHash);
-	return RuntimeHash->FinalizeGeneratedPackageForCook(InPackageRelativePath);
-}
-
-bool UWorldPartition::FinalizeGeneratorPackageForCook(const TArray<ICookPackageSplitter::FGeneratedPackageForPreSave>& InGeneratedPackages)
-{
-	check(RuntimeHash);
-	if (RuntimeHash->FinalizeGeneratorPackageForCook(InGeneratedPackages))
-	{
-		// Apply remapping of Persistent Level's SoftObjectPaths
-		FWorldPartitionLevelHelper::RemapLevelSoftObjectPaths(World->PersistentLevel, this);
-		return true;
-	}
-	return false;
+	return RuntimeHash->PopulateGeneratedPackageForCook(InPackage, InPackageRelativePath, OutModifiedPackages);
 }
 
 TArray<FBox> UWorldPartition::GetUserLoadedEditorRegions() const
diff --git a/Engine/Source/Runtime/Engine/Private/WorldPartition/WorldPartitionCookPackageSplitter.cpp b/Engine/Source/Runtime/Engine/Private/WorldPartition/WorldPartitionCookPackageSplitter.cpp
--- a/Engine/Source/Runtime/Engine/Private/WorldPartition/WorldPartitionCookPackageSplitter.cpp
+++ b/Engine/Source/Runtime/Engine/Private/WorldPartition/WorldPartitionCookPackageSplitter.cpp
@@ -118,35 +118,21 @@ TArray<ICookPackageSplitter::FGeneratedPackage> FWorldPartitionCookPackageSplitt
 }
 
 bool FWorldPartitionCookPackageSplitter::PopulateGeneratedPackage(UPackage* OwnerPackage, UObject* OwnerObject,
-	const FGeneratedPackageForPopulate& GeneratedPackage, TArray<UObject*>& OutObjectsToMove)
+	const FGeneratedPackageForPopulate& GeneratedPackage, TArray<UObject*>& OutObjectsToMove,
+	TArray<UPackage*>& OutModifiedPackages)
 {
 	UWorld* PartitionedWorld = ValidateDataObject(OwnerObject);
 	UWorldPartition* WorldPartition = PartitionedWorld->PersistentLevel->GetWorldPartition();
-	return WorldPartition->PopulateGeneratedPackageForCook(GeneratedPackage.Package, GeneratedPackage.RelativePath, OutObjectsToMove);
-}
-
-bool FWorldPartitionCookPackageSplitter::PreSaveGeneratedPackage(UPackage* OwnerPackage, UObject* OwnerObject,
-	const FGeneratedPackageForPopulate& GeneratedPackage)
-{
-	UWorld* PartitionedWorld = ValidateDataObject(OwnerObject);
-	UWorldPartition* WorldPartition = PartitionedWorld->PersistentLevel->GetWorldPartition();
-	return WorldPartition->FinalizeGeneratedPackageForCook(GeneratedPackage.RelativePath);
+	return WorldPartition->PopulateGeneratedPackageForCook(GeneratedPackage.Package, GeneratedPackage.RelativePath, OutModifiedPackages);
 }
 
 bool FWorldPartitionCookPackageSplitter::PopulateGeneratorPackage(UPackage* OwnerPackage, UObject* OwnerObject,
-	const TArray<ICookPackageSplitter::FGeneratedPackageForPreSave>& GeneratedPackages, TArray<UObject*>& OutObjectsToMove)
-{
-	UWorld* PartitionedWorld = ValidateDataObject(OwnerObject);
-	UWorldPartition* WorldPartition = PartitionedWorld->PersistentLevel->GetWorldPartition();
-	return WorldPartition->PopulateGeneratorPackageForCook(OutObjectsToMove);
-}
-
-bool FWorldPartitionCookPackageSplitter::PreSaveGeneratorPackage(UPackage* OwnerPackage, UObject* OwnerObject,
-	const TArray<FGeneratedPackageForPreSave>& GeneratedPackages)
+	const TArray<ICookPackageSplitter::FGeneratedPackageForPreSave>& GeneratedPackages, TArray<UObject*>& OutObjectsToMove,
+	TArray<UPackage*>& OutModifiedPackages)
 {
 	UWorld* PartitionedWorld = ValidateDataObject(OwnerObject);
 	UWorldPartition* WorldPartition = PartitionedWorld->PersistentLevel->GetWorldPartition();
-	return WorldPartition->FinalizeGeneratorPackageForCook(GeneratedPackages);
+	return WorldPartition->PopulateGeneratorPackageForCook(GeneratedPackages, OutModifiedPackages);
 }
 
 void FWorldPartitionCookPackageSplitter::OnOwnerReloaded(UPackage* OwnerPackage, UObject* OwnerObject)
diff --git a/Engine/Source/Runtime/Engine/Private/WorldPartition/WorldPartitionCookPackageSplitter.h b/Engine/Source/Runtime/Engine/Private/WorldPartition/WorldPartitionCookPackageSplitter.h
--- a/Engine/Source/Runtime/Engine/Private/WorldPartition/WorldPartitionCookPackageSplitter.h
+++ b/Engine/Source/Runtime/Engine/Private/WorldPartition/WorldPartitionCookPackageSplitter.h
@@ -21,13 +21,10 @@ public:
 	virtual bool UseInternalReferenceToAvoidGarbageCollect() override { return true; }
 	virtual TArray<ICookPackageSplitter::FGeneratedPackage> GetGenerateList(const UPackage* OwnerPackage, const UObject* OwnerObject) override;
 	virtual bool PopulateGeneratedPackage(UPackage* OwnerPackage, UObject* OwnerObject,
-		const FGeneratedPackageForPopulate& GeneratedPackage, TArray<UObject*>& OutObjectsToMove) override;
-	virtual bool PreSaveGeneratedPackage(UPackage* OwnerPackage, UObject* OwnerObject,
-		const ICookPackageSplitter::FGeneratedPackageForPopulate& GeneratedPackage) override;
+		const FGeneratedPackageForPopulate& GeneratedPackage, TArray<UObject*>& OutObjectsToMove, TArray<UPackage*>& OutModifiedPackages) override;
 	virtual bool PopulateGeneratorPackage(UPackage* OwnerPackage, UObject* OwnerObject,
-		const TArray<ICookPackageSplitter::FGeneratedPackageForPreSave>& GeneratedPackages, TArray<UObject*>& OutObjectsToMove) override;
-	virtual bool PreSaveGeneratorPackage(UPackage* OwnerPackage, UObject* OwnerObject,
-		const TArray<ICookPackageSplitter::FGeneratedPackageForPreSave>& GeneratedPackages) override;
+		const TArray<ICookPackageSplitter::FGeneratedPackageForPreSave>& GeneratedPackages, TArray<UObject*>& OutObjectsToMove,
+		TArray<UPackage*>& OutModifiedPackages) override;
 	virtual void OnOwnerReloaded(UPackage* OwnerPackage, UObject* OwnerObject) override;
 	//~ End of ICookPackageSplitter
 
diff --git a/Engine/Source/Runtime/Engine/Private/WorldPartition/WorldPartitionLevelHelper.cpp b/Engine/Source/Runtime/Engine/Private/WorldPartition/WorldPartitionLevelHelper.cpp
--- a/Engine/Source/Runtime/Engine/Private/WorldPartition/WorldPartitionLevelHelper.cpp
+++ b/Engine/Source/Runtime/Engine/Private/WorldPartition/WorldPartitionLevelHelper.cpp
@@ -17,6 +17,7 @@
 #include "WorldPartition/WorldPartitionActorDesc.h"
 #include "WorldPartition/WorldPartition.h"
 #include "WorldPartition/WorldPartitionPackageHelper.h"
+#include "AssetCompilingManager.h"
 #include "LevelUtils.h"
 #include "Templates/SharedPointer.h"
 #include "ActorFolder.h"
@@ -99,10 +100,15 @@ UWorld::InitializationValues FWorldPartitionLevelHelper::GetWorldInitializationV
 /**
  * Moves external actors into the given level
  */
-void FWorldPartitionLevelHelper::MoveExternalActorsToLevel(const TArray<FWorldPartitionRuntimeCellObjectMapping>& InChildPackages, ULevel* InLevel)
+void FWorldPartitionLevelHelper::MoveExternalActorsToLevel(const TArray<FWorldPartitionRuntimeCellObjectMapping>& InChildPackages, ULevel* InLevel, TArray<UPackage*>& OutModifiedPackages)
 {
 	TRACE_CPUPROFILER_EVENT_SCOPE(FWorldPartitionLevelHelper::MoveExternalActorsToLevel);
 
+	// We can't have async compilation still going on while we move actors as this is going to ResetLoaders which will move bulkdata around that
+	// might still be used by async compilation. 
+	// #TODO_DC Revisit once virtualbulkdata are enabled
+	FAssetCompilingManager::Get().FinishAllCompilation();
+
 	check(InLevel);
 	UPackage* LevelPackage = InLevel->GetPackage();
 	
@@ -122,6 +128,7 @@ void FWorldPartitionLevelHelper::MoveExternalActorsToLevel(const TArray<FWorldPa
 			check(ActorExternalPackage);
 
 			const bool bSameOuter = (InLevel == Actor->GetOuter());
+			Actor->SetPackageExternal(false, false);
 						
 			// Avoid calling Rename on the actor if it's already outered to InLevel as this will cause it's name to be changed. 
 			// (UObject::Rename doesn't check if Rename is being called with existing outer and assigns new name)
@@ -130,10 +137,25 @@ void FWorldPartitionLevelHelper::MoveExternalActorsToLevel(const TArray<FWorldPa
 				Actor->Rename(nullptr, InLevel, REN_ForceNoResetLoaders);
 			}
 			
+			check(Actor->GetPackage() == LevelPackage);
 			if (bSameOuter && !InLevel->Actors.Contains(Actor))
 			{
 				InLevel->AddLoadedActor(Actor);
 			}
+
+			// Include objects found in the source actor package in the destination level package
+			TArray<UObject*> Objects;
+			const bool bIncludeNestedSubobjects = false;
+			GetObjectsWithOuter(ActorExternalPackage, Objects, bIncludeNestedSubobjects);
+			for (UObject* Object : Objects)
+			{
+				if (Object->GetFName() != NAME_PackageMetaData)
+				{
+					Object->Rename(nullptr, LevelPackage, REN_ForceNoResetLoaders);
+				}
+			}
+
+			OutModifiedPackages.Add(ActorExternalPackage);
 		}
 		else
 		{
diff --git a/Engine/Source/Runtime/Engine/Private/WorldPartition/WorldPartitionRuntimeCell.cpp b/Engine/Source/Runtime/Engine/Private/WorldPartition/WorldPartitionRuntimeCell.cpp
--- a/Engine/Source/Runtime/Engine/Private/WorldPartition/WorldPartitionRuntimeCell.cpp
+++ b/Engine/Source/Runtime/Engine/Private/WorldPartition/WorldPartitionRuntimeCell.cpp
@@ -26,7 +26,7 @@ UWorldPartitionRuntimeCell::UWorldPartitionRuntimeCell(const FObjectInitializer&
 #if WITH_EDITOR
 bool UWorldPartitionRuntimeCell::NeedsActorToCellRemapping() const
 {
-	// When cooking, always loaded cells content is moved to persistent level (see FinalizeGeneratorPackageForCook)
+	// When cooking, always loaded cells content is moved to persistent level (see PopulateGeneratorPackageForCook)
 	return !(IsAlwaysLoaded() && IsRunningCookCommandlet());
 }
 
diff --git a/Engine/Source/Runtime/Engine/Private/WorldPartition/WorldPartitionRuntimeLevelStreamingCell.cpp b/Engine/Source/Runtime/Engine/Private/WorldPartition/WorldPartitionRuntimeLevelStreamingCell.cpp
--- a/Engine/Source/Runtime/Engine/Private/WorldPartition/WorldPartitionRuntimeLevelStreamingCell.cpp
+++ b/Engine/Source/Runtime/Engine/Private/WorldPartition/WorldPartitionRuntimeLevelStreamingCell.cpp
@@ -137,56 +137,21 @@ UWorldPartitionLevelStreamingDynamic* UWorldPartitionRuntimeLevelStreamingCell::
 	return nullptr;
 }
 
-bool UWorldPartitionRuntimeLevelStreamingCell::LoadCellObjectsForCook(TArray<UObject*>& OutLoadedObjects)
+bool UWorldPartitionRuntimeLevelStreamingCell::PopulateGeneratorPackageForCook(TArray<UPackage*>& OutModifiedPackages)
 {
-	check(!bCookHasLoadedCellObjects);
+	check(IsAlwaysLoaded());
+
 	if (GetActorCount() > 0)
 	{
+		FWorldPartitionLevelHelper::FPackageReferencer PackageReferencer;
 		const bool bLoadAsync = false;
 		UWorld* OuterWorld = GetOuterUWorldPartition()->GetTypedOuter<UWorld>();
-		verify(FWorldPartitionLevelHelper::LoadActors(OuterWorld, nullptr, Packages, CookPackageReferencer, [](bool) {}, bLoadAsync, FLinkerInstancingContext()));
-
-		TArray<UObject*> ObjectsInPackage;
-		for (const FWorldPartitionRuntimeCellObjectMapping& Mapping : Packages)
-		{
-			if (AActor* Actor = FindObject<AActor>(nullptr, *Mapping.LoadedPath.ToString()))
-			{
-				OutLoadedObjects.Add(Actor);
-				CookLoadedObjects.Add(Actor);
-
-				// Include objects found in the source actor package
-				const bool bIncludeNestedSubobjects = false;
-				ObjectsInPackage.Reset();
-				GetObjectsWithOuter(Actor->GetExternalPackage(), ObjectsInPackage, bIncludeNestedSubobjects);
-				for (UObject* Object : ObjectsInPackage)
-				{
-					if (Object->GetFName() != NAME_PackageMetaData && Object != Actor)
-					{
-						OutLoadedObjects.Add(Object);
-						CookLoadedObjects.Add(Object);
-					}
-				}
-			}
-		}
-	}
-	bCookHasLoadedCellObjects = true;
-	return true;
-}
-
-bool UWorldPartitionRuntimeLevelStreamingCell::PopulateGeneratorPackageForCook(TArray<UObject*>& OutLoadedObjects)
-{
-	check(IsAlwaysLoaded());
+		verify(FWorldPartitionLevelHelper::LoadActors(OuterWorld, nullptr, Packages, PackageReferencer, [](bool) {}, bLoadAsync, FLinkerInstancingContext()));
 
-	if (!LoadCellObjectsForCook(OutLoadedObjects))
-	{
-		return false;
-	}
+		FWorldPartitionLevelHelper::MoveExternalActorsToLevel(Packages, OuterWorld->PersistentLevel, OutModifiedPackages);
 
-	check(bCookHasLoadedCellObjects);
-	if (GetActorCount() > 0)
-	{
-		UWorld* OuterWorld = GetOuterUWorldPartition()->GetTypedOuter<UWorld>();
-		FWorldPartitionLevelHelper::MoveExternalActorsToLevel(Packages, OuterWorld->PersistentLevel);
+		// Empty cell's package list (this ensures that no one can rely on cell's content).
+		Packages.Empty();
 	}
 
 	return true;
@@ -208,7 +173,7 @@ bool UWorldPartitionRuntimeLevelStreamingCell::PrepareCellForCook(UPackage* InPa
 	return true;
 }
 
-bool UWorldPartitionRuntimeLevelStreamingCell::PopulateGeneratedPackageForCook(UPackage* InPackage, TArray<UObject*>& OutLoadedObjects)
+bool UWorldPartitionRuntimeLevelStreamingCell::PopulateGeneratedPackageForCook(UPackage* InPackage, TArray<UPackage*>& OutModifiedPackage)
 {
 	check(!IsAlwaysLoaded());
 	if (!InPackage)
@@ -216,13 +181,6 @@ bool UWorldPartitionRuntimeLevelStreamingCell::PopulateGeneratedPackageForCook(U
 		return false;
 	}
 
-	if (!LoadCellObjectsForCook(OutLoadedObjects))
-	{
-		return false;
-	}
-
-	check(bCookHasLoadedCellObjects);
-
 	if (GetActorCount() > 0)
 	{
 		// When cook splitter doesn't use deferred populate, cell needs to be prepared here.
@@ -231,70 +189,22 @@ bool UWorldPartitionRuntimeLevelStreamingCell::PopulateGeneratedPackageForCook(U
 			return false;
 		}
 
-		// Create a level and move loaded actors in it
 		UWorldPartition* WorldPartition = GetOuterUWorldPartition();
 		UWorld* OuterWorld = WorldPartition->GetTypedOuter<UWorld>();
+
+		// Load cell Actors
+		FWorldPartitionLevelHelper::FPackageReferencer PackageReferencer;
+		const bool bLoadAsync = false;
+		verify(FWorldPartitionLevelHelper::LoadActors(OuterWorld, nullptr, Packages, PackageReferencer, [](bool) {}, bLoadAsync, FLinkerInstancingContext()));
+
+		// Create a level and move these actors in it
 		ULevel* NewLevel = FWorldPartitionLevelHelper::CreateEmptyLevelForRuntimeCell(this, OuterWorld, LevelStreaming->GetWorldAsset().ToString(), InPackage);
 		check(NewLevel->GetPackage() == InPackage);
-		FWorldPartitionLevelHelper::MoveExternalActorsToLevel(Packages, NewLevel);
+		FWorldPartitionLevelHelper::MoveExternalActorsToLevel(Packages, NewLevel, OutModifiedPackage);
 
 		// Remap Level's SoftObjectPaths
 		FWorldPartitionLevelHelper::RemapLevelSoftObjectPaths(NewLevel, WorldPartition);
 	}
-
-	return true;
-}
-
-bool UWorldPartitionRuntimeLevelStreamingCell::FinalizePackageForCook()
-{
-	check(bCookHasLoadedCellObjects);
-
-	// Taken from original code introduced in CL 15771404 :
-	//     We can't have async compilation still going on while we move actors as this is going to ResetLoaders which will move bulkdata around that
-	//     might still be used by async compilation. 
-	// This code was modified to use FAssetCompilingManager instead of FStaticMeshCompilingManager.
-	//
-	// TODO: Remove once the core team fixes a remaining bug in the COTFS that requires this code to wait for async compilation.
-	FAssetCompilingManager::Get().FinishAllCompilation();
-
-	UPackage* LevelPackage = nullptr;
-	for (const TWeakObjectPtr<UObject>& LoadedObject : CookLoadedObjects)
-	{
-		check(LoadedObject.Get());
-		AActor* Actor = Cast<AActor>(LoadedObject);
-		if (UPackage* ActorExternalPackage = Actor ? Actor->GetExternalPackage() : nullptr)
-		{
-			UPackage* ActorLevelPackage = Actor->GetLevel()->GetPackage();
-			check(!LevelPackage || (LevelPackage == ActorLevelPackage));
-			LevelPackage = ActorLevelPackage;
-
-			// Remove package external on actor
-			Actor->SetPackageExternal(false, false);
-
-			// Make sure to also include objects found in the source actor package into the level package
-			TArray<UObject*> Objects;
-			const bool bIncludeNestedSubobjects = false;
-			GetObjectsWithOuter(ActorExternalPackage, Objects, bIncludeNestedSubobjects);
-			for (UObject* Object : Objects)
-			{
-				if (Object->GetFName() != NAME_PackageMetaData)
-				{
-					Object->Rename(nullptr, LevelPackage, REN_ForceNoResetLoaders);
-				}
-			}
-		}
-	}
-
-	if (IsAlwaysLoaded())
-	{
-		// Empty cell's package list (this ensures that no one can rely on cell's content).
-		Packages.Empty();
-	}
-
-	CookLoadedObjects.Empty();
-	CookPackageReferencer.RemoveReferences();
-	bCookHasLoadedCellObjects = false;
-
 	return true;
 }
 
diff --git a/Engine/Source/Runtime/Engine/Private/WorldPartition/WorldPartitionRuntimeSpatialHash.cpp b/Engine/Source/Runtime/Engine/Private/WorldPartition/WorldPartitionRuntimeSpatialHash.cpp
--- a/Engine/Source/Runtime/Engine/Private/WorldPartition/WorldPartitionRuntimeSpatialHash.cpp
+++ b/Engine/Source/Runtime/Engine/Private/WorldPartition/WorldPartitionRuntimeSpatialHash.cpp
@@ -1136,7 +1136,7 @@ bool UWorldPartitionRuntimeSpatialHash::CreateStreamingGrid(const FSpatialHashRu
 
 					if (StreamingCell->GetActorCount())
 					{
-						// Always loaded cell actors are transfered to World's Persistent Level (see UWorldPartitionRuntimeSpatialHash::FinalizeGeneratorPackageForCook)
+						// Always loaded cell actors are transfered to World's Persistent Level (see UWorldPartitionRuntimeSpatialHash::PopulateGeneratorPackageForCook)
 						if (!StreamingCell->IsAlwaysLoaded())
 						{
 							if (!OutPackagesToGenerate)
@@ -1149,7 +1149,7 @@ bool UWorldPartitionRuntimeSpatialHash::CreateStreamingGrid(const FSpatialHashRu
 							check(!PackageRelativePath.IsEmpty());
 							OutPackagesToGenerate->Add(PackageRelativePath);
 
-							// Map relative package to StreamingCell for PopulateGeneratedPackageForCook/FinalizeGeneratorPackageForCook
+							// Map relative package to StreamingCell for PopulateGeneratedPackageForCook/PopulateGeneratorPackageForCook
 							PackagesToGenerateForCook.Add(PackageRelativePath, StreamingCell);
 						}
 					}
@@ -1161,28 +1161,15 @@ bool UWorldPartitionRuntimeSpatialHash::CreateStreamingGrid(const FSpatialHashRu
 	return true;
 }
 
-bool UWorldPartitionRuntimeSpatialHash::PopulateGeneratedPackageForCook(UPackage* InPackage, const FString& InPackageRelativePath, TArray<UObject*>& OutLoadedObjects)
+bool UWorldPartitionRuntimeSpatialHash::PopulateGeneratedPackageForCook(UPackage* InPackage, const FString& InPackageRelativePath, TArray<UPackage*>& OutModifiedPackages)
 {
-	OutLoadedObjects.Reset();
+	OutModifiedPackages.Reset();
 	if (UWorldPartitionRuntimeCell** MatchingCell = PackagesToGenerateForCook.Find(InPackageRelativePath))
 	{
 		UWorldPartitionRuntimeCell* Cell = *MatchingCell;
 		if (ensure(Cell))
 		{
-			return Cell->PopulateGeneratedPackageForCook(InPackage, OutLoadedObjects);
-		}
-	}
-	return false;
-}
-
-bool UWorldPartitionRuntimeSpatialHash::FinalizeGeneratedPackageForCook(const FString& InPackageRelativePath)
-{
-	if (UWorldPartitionRuntimeCell** MatchingCell = PackagesToGenerateForCook.Find(InPackageRelativePath))
-	{
-		UWorldPartitionRuntimeCell* Cell = *MatchingCell;
-		if (ensure(Cell))
-		{
-			return Cell->FinalizePackageForCook();
+			return Cell->PopulateGeneratedPackageForCook(InPackage, OutModifiedPackages);
 		}
 	}
 	return false;
@@ -1203,33 +1190,20 @@ TArray<UWorldPartitionRuntimeCell*> UWorldPartitionRuntimeSpatialHash::GetAlways
 	return Result;
 }
 
-bool UWorldPartitionRuntimeSpatialHash::PopulateGeneratorPackageForCook(TArray<UObject*>& OutLoadedObjects)
-{
-	OutLoadedObjects.Reset();
-	for (UWorldPartitionRuntimeCell* Cell : GetAlwaysLoadedCells())
-	{
-		check(Cell->IsAlwaysLoaded());
-		if (!Cell->PopulateGeneratorPackageForCook(OutLoadedObjects))
-		{
-			return false;
-		}
-	}
-	return true;
-}
-
-bool UWorldPartitionRuntimeSpatialHash::FinalizeGeneratorPackageForCook(const TArray<ICookPackageSplitter::FGeneratedPackageForPreSave>& InGeneratedPackages)
+bool UWorldPartitionRuntimeSpatialHash::PopulateGeneratorPackageForCook(const TArray<ICookPackageSplitter::FGeneratedPackageForPreSave>& InGeneratedPackages, TArray<UPackage*>& OutModifiedPackages)
 {
 	check(IsRunningCookCommandlet());
 
+	OutModifiedPackages.Reset();
 	for (UWorldPartitionRuntimeCell* Cell : GetAlwaysLoadedCells())
 	{
 		check(Cell->IsAlwaysLoaded());
-		if (!Cell->FinalizePackageForCook())
+		if (!Cell->PopulateGeneratorPackageForCook(OutModifiedPackages))
 		{
 			return false;
 		}
 	}
-	
+
 	for (const ICookPackageSplitter::FGeneratedPackageForPreSave& GeneratedPackage : InGeneratedPackages)
 	{
 		UWorldPartitionRuntimeCell** MatchingCell = PackagesToGenerateForCook.Find(GeneratedPackage.RelativePath);
diff --git a/Engine/Source/Runtime/Engine/Public/WorldPartition/WorldPartition.h b/Engine/Source/Runtime/Engine/Public/WorldPartition/WorldPartition.h
--- a/Engine/Source/Runtime/Engine/Public/WorldPartition/WorldPartition.h
+++ b/Engine/Source/Runtime/Engine/Public/WorldPartition/WorldPartition.h
@@ -151,10 +151,8 @@ public:
 	void RemapSoftObjectPath(FSoftObjectPath& ObjectPath);
 
 	// Cooking
-	bool PopulateGeneratorPackageForCook(TArray<UObject*>& OutLoadedObjects);
-	bool PopulateGeneratedPackageForCook(UPackage* InPackage, const FString& InPackageRelativePath, TArray<UObject*>& OutLoadedObjects);
-	bool FinalizeGeneratedPackageForCook(const FString& InPackageRelativePath);
-	bool FinalizeGeneratorPackageForCook(const TArray<ICookPackageSplitter::FGeneratedPackageForPreSave>& InGeneratedPackages);
+	bool PopulateGeneratorPackageForCook(const TArray<ICookPackageSplitter::FGeneratedPackageForPreSave>& InGeneratedPackages, TArray<UPackage*>& OutModifiedPackages);
+	bool PopulateGeneratedPackageForCook(UPackage* InPackage, const FString& InPackageRelativePath, TArray<UPackage*>& OutModifiedPackages);
 
 	UE_DEPRECATED(5.1, "GetWorldBounds is deprecated, use GetEditorWorldBounds or GetRuntimeWorldBounds instead.")
 	FBox GetWorldBounds() const { return GetRuntimeWorldBounds(); }
diff --git a/Engine/Source/Runtime/Engine/Public/WorldPartition/WorldPartitionLevelHelper.h b/Engine/Source/Runtime/Engine/Public/WorldPartition/WorldPartitionLevelHelper.h
--- a/Engine/Source/Runtime/Engine/Public/WorldPartition/WorldPartitionLevelHelper.h
+++ b/Engine/Source/Runtime/Engine/Public/WorldPartition/WorldPartitionLevelHelper.h
@@ -32,7 +32,7 @@ public:
 	};
 
 	static ULevel* CreateEmptyLevelForRuntimeCell(const UWorldPartitionRuntimeCell* Cell, const UWorld* InWorld, const FString& InWorldAssetName, UPackage* DestPackage = nullptr);
-	static void MoveExternalActorsToLevel(const TArray<FWorldPartitionRuntimeCellObjectMapping>& InChildPackages, ULevel* InLevel);
+	static void MoveExternalActorsToLevel(const TArray<FWorldPartitionRuntimeCellObjectMapping>& InChildPackages, ULevel* InLevel, TArray<UPackage*>& OutModifiedPackages);
 	static void RemapLevelSoftObjectPaths(ULevel* InLevel, UWorldPartition* InWorldPartition);
 	
 	static bool LoadActors(UWorld* InOwningWorld, ULevel* InDestLevel, TArrayView<FWorldPartitionRuntimeCellObjectMapping> InActorPackages, FPackageReferencer& InPackageReferencer, TFunction<void(bool)> InCompletionCallback, bool bInLoadAsync, FLinkerInstancingContext InOutInstancingContext);
diff --git a/Engine/Source/Runtime/Engine/Public/WorldPartition/WorldPartitionRuntimeCell.h b/Engine/Source/Runtime/Engine/Public/WorldPartition/WorldPartitionRuntimeCell.h
--- a/Engine/Source/Runtime/Engine/Public/WorldPartition/WorldPartitionRuntimeCell.h
+++ b/Engine/Source/Runtime/Engine/Public/WorldPartition/WorldPartitionRuntimeCell.h
@@ -202,9 +202,8 @@ class UWorldPartitionRuntimeCell : public UObject, public IWorldPartitionCell
 
 	// Cook methods
 	virtual bool PrepareCellForCook(UPackage* InPackage) { return false; }
-	virtual bool PopulateGeneratorPackageForCook(TArray<UObject*>& OutLoadedObjects) PURE_VIRTUAL(UWorldPartitionRuntimeCell::PopulateGeneratorPackageForCook, return false;);
-	virtual bool PopulateGeneratedPackageForCook(UPackage* InPackage, TArray<UObject*>& OutLoadedObjects) PURE_VIRTUAL(UWorldPartitionRuntimeCell::PopulateGeneratedPackageForCook, return false;);
-	virtual bool FinalizePackageForCook() PURE_VIRTUAL(UWorldPartitionRuntimeCell::FinalizePackageForCook, return false;);
+	virtual bool PopulateGeneratorPackageForCook(TArray<UPackage*>& OutModifiedPackages) PURE_VIRTUAL(UWorldPartitionRuntimeCell::PopulateGeneratorPackageForCook, return false;);
+	virtual bool PopulateGeneratedPackageForCook(UPackage* InPackage, TArray<UPackage*>& OutModifiedPackages) PURE_VIRTUAL(UWorldPartitionRuntimeCell::PopulateGeneratedPackageForCook, return false;);
 	virtual FString GetPackageNameToCreate() const PURE_VIRTUAL(UWorldPartitionRuntimeCell::GetPackageNameToCreate, return FString(""););
 
 	void SetIsHLOD(bool bInIsHLOD) { bIsHLOD = bInIsHLOD; }
diff --git a/Engine/Source/Runtime/Engine/Public/WorldPartition/WorldPartitionRuntimeHash.h b/Engine/Source/Runtime/Engine/Public/WorldPartition/WorldPartitionRuntimeHash.h
--- a/Engine/Source/Runtime/Engine/Public/WorldPartition/WorldPartitionRuntimeHash.h
+++ b/Engine/Source/Runtime/Engine/Public/WorldPartition/WorldPartitionRuntimeHash.h
@@ -34,10 +34,8 @@ class ENGINE_API UWorldPartitionRuntimeHash : public UObject
 
 #if WITH_EDITOR
 	virtual void SetDefaultValues() {}
-	virtual bool PopulateGeneratorPackageForCook(TArray<UObject*>& OutLoadedObjects) { return false; }
-	virtual bool FinalizeGeneratorPackageForCook(const TArray<ICookPackageSplitter::FGeneratedPackageForPreSave>& InGeneratedPackages) { return false; }
-	virtual bool PopulateGeneratedPackageForCook(UPackage* InPackage, const FString& InPackageRelativePath, TArray<UObject*>& OutLoadedObjects) { return false; }
-	virtual bool FinalizeGeneratedPackageForCook(const FString& InPackageRelativePath) { return false; }
+	virtual bool PopulateGeneratorPackageForCook(const TArray<ICookPackageSplitter::FGeneratedPackageForPreSave>& InGeneratedPackages, TArray<UPackage*>& OutModifiedPackages) { return false; }
+	virtual bool PopulateGeneratedPackageForCook(UPackage* InPackage, const FString& InPackageRelativePath, TArray<UPackage*>& OutModifiedPackages) { return false; }
 	virtual bool GenerateStreaming(class UWorldPartitionStreamingPolicy* StreamingPolicy, const FActorClusterContext& ActorClusterContext, TArray<FString>* OutPackagesToGenerate) { return false; }
 	virtual void FlushStreaming() {}
 	virtual bool GenerateHLOD(ISourceControlHelper* SourceControlHelper, FActorClusterContext& ActorClusterContext, bool bCreateActorsOnly) { return false; }
diff --git a/Engine/Source/Runtime/Engine/Public/WorldPartition/WorldPartitionRuntimeLevelStreamingCell.h b/Engine/Source/Runtime/Engine/Public/WorldPartition/WorldPartitionRuntimeLevelStreamingCell.h
--- a/Engine/Source/Runtime/Engine/Public/WorldPartition/WorldPartitionRuntimeLevelStreamingCell.h
+++ b/Engine/Source/Runtime/Engine/Public/WorldPartition/WorldPartitionRuntimeLevelStreamingCell.h
@@ -40,9 +40,8 @@ class UWorldPartitionRuntimeLevelStreamingCell : public UWorldPartitionRuntimeSp
 	virtual void DumpStateLog(FHierarchicalLogArchive& Ar) override;
 	// Cook methods
 	virtual bool PrepareCellForCook(UPackage* InPackage) override;
-	virtual bool PopulateGeneratorPackageForCook(TArray<UObject*>& OutLoadedObjects) override;
-	virtual bool PopulateGeneratedPackageForCook(UPackage* InPackage, TArray<UObject*>& OutLoadedObjects) override;
-	virtual bool FinalizePackageForCook() override;
+	virtual bool PopulateGeneratorPackageForCook(TArray<UPackage*>& OutModifiedPackages) override;
+	virtual bool PopulateGeneratedPackageForCook(UPackage* InPackage, TArray<UPackage*>& OutModifiedPackages) override;
 	virtual FString GetPackageNameToCreate() const override;
 	//~End UWorldPartitionRuntimeCell Interface
 
@@ -59,12 +58,7 @@ private:
 	class UWorldPartitionLevelStreamingDynamic* GetOrCreateLevelStreaming() const;
 
 #if WITH_EDITOR
-	bool LoadCellObjectsForCook(TArray<UObject*>& OutLoadedObjects);
 	class UWorldPartitionLevelStreamingDynamic* CreateLevelStreaming(const FString& InPackageName = FString()) const;
-
-	FWorldPartitionLevelHelper::FPackageReferencer CookPackageReferencer;
-	TArray<TWeakObjectPtr<UObject>> CookLoadedObjects;
-	bool bCookHasLoadedCellObjects = false;
 #endif
 
 #if WITH_EDITORONLY_DATA
diff --git a/Engine/Source/Runtime/Engine/Public/WorldPartition/WorldPartitionRuntimeSpatialHash.h b/Engine/Source/Runtime/Engine/Public/WorldPartition/WorldPartitionRuntimeSpatialHash.h
--- a/Engine/Source/Runtime/Engine/Public/WorldPartition/WorldPartitionRuntimeSpatialHash.h
+++ b/Engine/Source/Runtime/Engine/Public/WorldPartition/WorldPartitionRuntimeSpatialHash.h
@@ -232,10 +232,8 @@ public:
 
 #if WITH_EDITOR
 	virtual void SetDefaultValues() override;
-	virtual bool PopulateGeneratorPackageForCook(TArray<UObject*>& OutLoadedObjects) override;
-	virtual bool FinalizeGeneratorPackageForCook(const TArray<ICookPackageSplitter::FGeneratedPackageForPreSave>& InGeneratedPackages) override;
-	virtual bool PopulateGeneratedPackageForCook(UPackage* InPackage, const FString& InPackageRelativePath, TArray<UObject*>& OutLoadedObjects) override;
-	virtual bool FinalizeGeneratedPackageForCook(const FString& InPackageRelativePath) override;
+	virtual bool PopulateGeneratorPackageForCook(const TArray<ICookPackageSplitter::FGeneratedPackageForPreSave>& InGeneratedPackages, TArray<UPackage*>& OutModifiedPackages) override;
+	virtual bool PopulateGeneratedPackageForCook(UPackage* InPackage, const FString& InPackageRelativePath, TArray<UPackage*>& OutModifiedPackages) override;
 	virtual void FlushStreaming() override;
 	virtual bool GenerateHLOD(ISourceControlHelper* SourceControlHelper, FActorClusterContext& ActorClusterContext, bool bCreateActorsOnly) override;
 	virtual void DrawPreview() const override;
