Creating a Level Instance inside another Level Instance didn't put the world current level back to the parent level instance.

diff --git a/Engine/Source/Runtime/Engine/Private/LevelInstance/LevelInstanceSubsystem.cpp b/Engine/Source/Runtime/Engine/Private/LevelInstance/LevelInstanceSubsystem.cpp
--- a/Engine/Source/Runtime/Engine/Private/LevelInstance/LevelInstanceSubsystem.cpp
+++ b/Engine/Source/Runtime/Engine/Private/LevelInstance/LevelInstanceSubsystem.cpp
@@ -1139,6 +1139,12 @@ ILevelInstanceInterface* ULevelInstanceSubsystem::CreateLevelInstanceFrom(const
 		EditorModule.ActivateEditorMode();
 	}
 
+	// After commit, CurrentLevel goes back to world's PersistentLevel. Set it back to the current editing level instance (if any).
+	if (ILevelInstanceInterface* EditingLevelInstance = GetEditingLevelInstance())
+	{
+		SetCurrent(EditingLevelInstance);
+	}
+
 	return GetLevelInstance(NewLevelInstanceID);
 }
 
