Added FWorldDelegates::FWorldCurrentLevelChangedEvent. It is broadcasted when the world's current level changes.

diff --git a/Engine/Source/Runtime/Engine/Classes/Engine/World.h b/Engine/Source/Runtime/Engine/Classes/Engine/World.h
--- a/Engine/Source/Runtime/Engine/Classes/Engine/World.h
+++ b/Engine/Source/Runtime/Engine/Classes/Engine/World.h
@@ -4168,6 +4168,7 @@ public:
 #if WITH_EDITOR
 	DECLARE_MULTICAST_DELEGATE_FiveParams(FWorldPreRenameEvent, UWorld* /*World*/, const TCHAR* /*InName*/, UObject* /*NewOuter*/, ERenameFlags /*Flags*/, bool& /*bShouldFailRename*/);
 	DECLARE_MULTICAST_DELEGATE_OneParam(FWorldPostRenameEvent, UWorld*);
+	DECLARE_MULTICAST_DELEGATE_ThreeParams(FWorldCurrentLevelChangedEvent, ULevel* /*NewLevel*/, ULevel* /*OldLevel*/, UWorld* /*World*/);
 #endif // WITH_EDITOR
 
 	// Delegate type for level change events
@@ -4204,6 +4205,8 @@ public:
 
 	// Callback for world rename event (post)
 	static FWorldPostRenameEvent OnPostWorldRename;
+
+	static FWorldCurrentLevelChangedEvent OnCurrentLevelChanged;
 #endif // WITH_EDITOR
 
 	// Post duplication event.
diff --git a/Engine/Source/Runtime/Engine/Private/World.cpp b/Engine/Source/Runtime/Engine/Private/World.cpp
--- a/Engine/Source/Runtime/Engine/Private/World.cpp
+++ b/Engine/Source/Runtime/Engine/Private/World.cpp
@@ -351,6 +351,7 @@ FWorldDelegates::FWorldInitializationEvent FWorldDelegates::OnPostWorldInitializ
 #if WITH_EDITOR
 FWorldDelegates::FWorldPreRenameEvent FWorldDelegates::OnPreWorldRename;
 FWorldDelegates::FWorldPostRenameEvent FWorldDelegates::OnPostWorldRename;
+FWorldDelegates::FWorldCurrentLevelChangedEvent FWorldDelegates::OnCurrentLevelChanged;
 #endif // WITH_EDITOR
 FWorldDelegates::FWorldPostDuplicateEvent FWorldDelegates::OnPostDuplicate;
 FWorldDelegates::FWorldCleanupEvent FWorldDelegates::OnWorldCleanup;
@@ -7906,8 +7907,11 @@ bool UWorld::SetCurrentLevel( class ULevel* InLevel )
 	bool bChanged = false;
 	if( CurrentLevel != InLevel )
 	{
+		ULevel* OldLevel = CurrentLevel;
 		CurrentLevel = InLevel;
 		bChanged = true;
+
+		FWorldDelegates::OnCurrentLevelChanged.Broadcast(CurrentLevel, OldLevel, this);
 	}
 	return bChanged;
 }
