Fix vertex colors being stripped from static mesh components when cooking against sever data.
StaticMesh assets in server data contain no vertex buffers so the asset vertex count was always being calculated as 0.

diff --git a/Engine/Source/Runtime/Engine/Private/Components/StaticMeshComponent.cpp b/Engine/Source/Runtime/Engine/Private/Components/StaticMeshComponent.cpp
--- a/Engine/Source/Runtime/Engine/Private/Components/StaticMeshComponent.cpp
+++ b/Engine/Source/Runtime/Engine/Private/Components/StaticMeshComponent.cpp
@@ -3208,11 +3208,42 @@ FArchive& operator<<(FArchive& Ar,FStaticMeshComponentLODInfo& I)
 		{}
 		check( LODIndex < I.OwningComponent->LODData.Num() );
 
-		bStrippedOverrideColors = !I.OverrideVertexColors || 
-			( I.OwningComponent->GetStaticMesh() == nullptr || 
-			I.OwningComponent->GetStaticMesh()->GetRenderData() == nullptr ||
-			LODIndex >= I.OwningComponent->GetStaticMesh()->GetRenderData()->LODResources.Num() ||
-			I.OverrideVertexColors->GetNumVertices() != I.OwningComponent->GetStaticMesh()->GetRenderData()->LODResources[LODIndex].VertexBuffers.StaticMeshVertexBuffer.GetNumVertices() );
+		bStrippedOverrideColors = true;
+
+		if (I.OverrideVertexColors &&
+			I.OwningComponent->GetStaticMesh() &&
+			I.OwningComponent->GetStaticMesh()->GetRenderData() &&
+			I.OwningComponent->GetStaticMesh()->GetRenderData()->LODResources.IsValidIndex(LODIndex))
+		{
+			const FStaticMeshLODResources& StaticMeshLODResources = I.OwningComponent->GetStaticMesh()->GetRenderData()->LODResources[LODIndex];
+			const int32 StaticMeshVertexBufferCount = StaticMeshLODResources.VertexBuffers.StaticMeshVertexBuffer.GetNumVertices();
+			if (StaticMeshVertexBufferCount == I.OverrideVertexColors->GetNumVertices())
+			{
+				bStrippedOverrideColors = false;
+			}
+			else if (StaticMeshVertexBufferCount == 0)
+			{
+				// StaticMeshVertexBuffer is not available when StaticMesh loaded from a cooked build made with IsDataStrippedForServer()
+				// TODO: Could be using PKG_ServerSideOnly but PKG_ServerSideOnly flag is not currently being set
+				if (I.OwningComponent->GetStaticMesh()->GetPackage()->HasAllPackagesFlags(PKG_Cooked))
+				{
+					// Calculate VertexCount by iterating section data we do have access to
+					uint32 MaxMaxVertexIndex = 0;
+					const FStaticMeshSectionArray& SectionsList = StaticMeshLODResources.Sections;
+					for (const FStaticMeshSection& Section : SectionsList)
+					{
+						MaxMaxVertexIndex = FMath::Max<uint32>(Section.MaxVertexIndex, MaxMaxVertexIndex);
+					}
+
+					++MaxMaxVertexIndex;
+
+					if (I.OverrideVertexColors->GetNumVertices() == MaxMaxVertexIndex)
+					{
+						bStrippedOverrideColors = false;
+					}
+				}
+			}
+		}
 	}
 #endif // WITH_EDITORONLY_DATA
 	FStripDataFlags StripFlags( Ar, bStrippedOverrideColors ? OverrideColorsStripFlag : 0 );
