Remove some deprecated options from the input settings. These have been deprecated for a while and are not used in any configs in our streams. They also weren't exposed in the editor, so nobody could have turned them on.

bDefaultViewportMouseLock_DEPRECATED was removed in 2018, so that is fine to delete.
ConsoleKey_DEPRECATED was removed in 2014 (!) which is also ok to be removed.

#rnx

diff --git a/Engine/Source/Runtime/Engine/Classes/GameFramework/InputSettings.h b/Engine/Source/Runtime/Engine/Classes/GameFramework/InputSettings.h
--- a/Engine/Source/Runtime/Engine/Classes/GameFramework/InputSettings.h
+++ b/Engine/Source/Runtime/Engine/Classes/GameFramework/InputSettings.h
@@ -49,10 +49,6 @@ class ENGINE_API UInputSettings
 	UPROPERTY(config, EditAnywhere, Category = "ViewportProperties")
 	uint8 bCaptureMouseOnLaunch:1;
 
-	/** The default mouse lock state when the viewport acquires capture */
-	UPROPERTY(config)
-	uint8 bDefaultViewportMouseLock_DEPRECATED:1;
-
 	/** Enable the use of legacy input scales on the player controller (InputYawScale, InputPitchScale, and InputRollScale) */
 	UPROPERTY(config, EditAnywhere, Category = "Input")
 	uint8 bEnableLegacyInputScales:1;
@@ -151,10 +147,6 @@ public:
 	UPROPERTY(config, EditAnywhere, Category="Mobile", meta=(AllowedClasses="/Script/Engine.TouchInterface"))
 	FSoftObjectPath DefaultTouchInterface;
 
-	/** The key which opens the console. */
-	UPROPERTY(config)
-	FKey ConsoleKey_DEPRECATED;
-
 	/** The keys which open the console. */
 	UPROPERTY(config, EditAnywhere, Category="Console")
 	TArray<FKey> ConsoleKeys;
diff --git a/Engine/Source/Runtime/Engine/Private/UserInterface/InputSettings.cpp b/Engine/Source/Runtime/Engine/Private/UserInterface/InputSettings.cpp
--- a/Engine/Source/Runtime/Engine/Private/UserInterface/InputSettings.cpp
+++ b/Engine/Source/Runtime/Engine/Private/UserInterface/InputSettings.cpp
@@ -21,7 +21,6 @@
 UInputSettings::UInputSettings(const FObjectInitializer& ObjectInitializer)
 	: Super(ObjectInitializer)
 	, bCaptureMouseOnLaunch(true)
-	, bDefaultViewportMouseLock_DEPRECATED(false)
 	, bEnableLegacyInputScales(true)
 	, bEnableMotionControls(true)
 	, bFilterInputByPlatformUser(false)
@@ -73,12 +72,6 @@ void UInputSettings::PostInitProperties()
 {
 	Super::PostInitProperties();
 
-	if (ConsoleKey_DEPRECATED.IsValid())
-	{
-		ConsoleKeys.Empty(1);
-		ConsoleKeys.Add(ConsoleKey_DEPRECATED);
-	}
-
 	PopulateAxisConfigs();
 
 #if PLATFORM_WINDOWS
