Moving scene texture config and parameters into the engine library.

diff --git a/Engine/Source/Runtime/Engine/Private/SceneTexturesConfig.cpp b/Engine/Source/Runtime/Engine/Private/SceneTexturesConfig.cpp
new file mode 100644
--- /dev/null
+++ b/Engine/Source/Runtime/Engine/Private/SceneTexturesConfig.cpp
@@ -0,0 +1,24 @@
+// Copyright Epic Games, Inc. All Rights Reserved.
+
+#include "SceneTexturesConfig.h"
+#include "RenderGraphResources.h"
+
+FSceneTexturesConfig FSceneTexturesConfig::GlobalInstance;
+
+IMPLEMENT_STATIC_UNIFORM_BUFFER_SLOT(SceneTextures);
+IMPLEMENT_STATIC_UNIFORM_BUFFER_STRUCT(FSceneTextureUniformParameters, "SceneTexturesStruct", SceneTextures);
+IMPLEMENT_STATIC_UNIFORM_BUFFER_STRUCT(FMobileSceneTextureUniformParameters, "MobileSceneTextures", SceneTextures);
+
+FSceneTextureShaderParameters GetSceneTextureShaderParameters(TRDGUniformBufferRef<FSceneTextureUniformParameters> UniformBuffer)
+{
+	FSceneTextureShaderParameters Parameters;
+	Parameters.SceneTextures = UniformBuffer;
+	return Parameters;
+}
+
+FSceneTextureShaderParameters GetSceneTextureShaderParameters(TRDGUniformBufferRef<FMobileSceneTextureUniformParameters> UniformBuffer)
+{
+	FSceneTextureShaderParameters Parameters;
+	Parameters.MobileSceneTextures = UniformBuffer;
+	return Parameters;
+}
\ No newline at end of file
diff --git a/Engine/Source/Runtime/Engine/Public/SceneTexturesConfig.h b/Engine/Source/Runtime/Engine/Public/SceneTexturesConfig.h
new file mode 100644
--- /dev/null
+++ b/Engine/Source/Runtime/Engine/Public/SceneTexturesConfig.h
@@ -0,0 +1,172 @@
+// Copyright Epic Games, Inc. All Rights Reserved.
+
+#pragma once
+
+#include "RenderGraphDefinitions.h"
+#include "ShaderParameterMacros.h"
+#include "GBufferInfo.h"
+#include "SceneUtils.h"
+
+class FSceneViewFamily;
+
+/** A uniform buffer containing common scene textures used by materials or global shaders. */
+BEGIN_GLOBAL_SHADER_PARAMETER_STRUCT(FSceneTextureUniformParameters, ENGINE_API)
+	// Scene Color / Depth
+	SHADER_PARAMETER_RDG_TEXTURE(Texture2D, SceneColorTexture)
+	SHADER_PARAMETER_RDG_TEXTURE(Texture2D, SceneDepthTexture)
+
+	// GBuffer
+	SHADER_PARAMETER_RDG_TEXTURE(Texture2D, GBufferATexture)
+	SHADER_PARAMETER_RDG_TEXTURE(Texture2D, GBufferBTexture)
+	SHADER_PARAMETER_RDG_TEXTURE(Texture2D, GBufferCTexture)
+	SHADER_PARAMETER_RDG_TEXTURE(Texture2D, GBufferDTexture)
+	SHADER_PARAMETER_RDG_TEXTURE(Texture2D, GBufferETexture)
+	SHADER_PARAMETER_RDG_TEXTURE(Texture2D, GBufferFTexture)
+	SHADER_PARAMETER_RDG_TEXTURE(Texture2D, GBufferVelocityTexture)
+
+	// SSAO
+	SHADER_PARAMETER_RDG_TEXTURE(Texture2D, ScreenSpaceAOTexture)
+
+	// Custom Depth / Stencil
+	SHADER_PARAMETER_RDG_TEXTURE(Texture2D, CustomDepthTexture)
+	SHADER_PARAMETER_RDG_TEXTURE_SRV(Texture2D<uint2>, CustomStencilTexture)
+
+	// Misc
+	SHADER_PARAMETER_SAMPLER(SamplerState, PointClampSampler)
+END_GLOBAL_SHADER_PARAMETER_STRUCT()
+
+BEGIN_GLOBAL_SHADER_PARAMETER_STRUCT(FMobileSceneTextureUniformParameters, ENGINE_API)
+	SHADER_PARAMETER_RDG_TEXTURE(Texture2D, SceneColorTexture)
+	SHADER_PARAMETER_SAMPLER(SamplerState, SceneColorTextureSampler)
+	SHADER_PARAMETER_RDG_TEXTURE(Texture2D, SceneDepthTexture)
+	SHADER_PARAMETER_SAMPLER(SamplerState, SceneDepthTextureSampler)
+	SHADER_PARAMETER_RDG_TEXTURE(Texture2D, CustomDepthTexture)
+	SHADER_PARAMETER_SAMPLER(SamplerState, CustomDepthTextureSampler)
+	SHADER_PARAMETER_RDG_TEXTURE_SRV(Texture2D<uint2>, CustomStencilTexture)
+	SHADER_PARAMETER_RDG_TEXTURE(Texture2D, SceneVelocityTexture)
+	SHADER_PARAMETER_SAMPLER(SamplerState, SceneVelocityTextureSampler)
+	// GBuffer
+	SHADER_PARAMETER_RDG_TEXTURE(Texture2D, GBufferATexture)
+	SHADER_PARAMETER_RDG_TEXTURE(Texture2D, GBufferBTexture)
+	SHADER_PARAMETER_RDG_TEXTURE(Texture2D, GBufferCTexture)
+	SHADER_PARAMETER_RDG_TEXTURE(Texture2D, GBufferDTexture)
+	SHADER_PARAMETER_RDG_TEXTURE(Texture2D, SceneDepthAuxTexture)
+	SHADER_PARAMETER_SAMPLER(SamplerState, GBufferATextureSampler)
+	SHADER_PARAMETER_SAMPLER(SamplerState, GBufferBTextureSampler)
+	SHADER_PARAMETER_SAMPLER(SamplerState, GBufferCTextureSampler)
+	SHADER_PARAMETER_SAMPLER(SamplerState, GBufferDTextureSampler)
+	SHADER_PARAMETER_SAMPLER(SamplerState, SceneDepthAuxTextureSampler)
+END_GLOBAL_SHADER_PARAMETER_STRUCT()
+
+BEGIN_SHADER_PARAMETER_STRUCT(FSceneTextureShaderParameters, ENGINE_API)
+	SHADER_PARAMETER_RDG_UNIFORM_BUFFER(FSceneTextureUniformParameters, SceneTextures)
+	SHADER_PARAMETER_RDG_UNIFORM_BUFFER(FMobileSceneTextureUniformParameters, MobileSceneTextures)
+END_SHADER_PARAMETER_STRUCT()
+
+extern ENGINE_API FSceneTextureShaderParameters GetSceneTextureShaderParameters(TRDGUniformBufferRef<FSceneTextureUniformParameters> UniformBuffer);
+
+extern ENGINE_API FSceneTextureShaderParameters GetSceneTextureShaderParameters(TRDGUniformBufferRef<FMobileSceneTextureUniformParameters> UniformBuffer);
+
+enum class ESceneTextureExtracts : uint32
+{
+	/** No textures are extracted from the render graph after execution. */
+	None = 0,
+
+	/** Extracts scene depth after execution */
+	Depth = 1 << 0,
+
+	/** Extracts custom depth after execution. */
+	CustomDepth = 1 << 1,
+
+	/** Extracts all available textures after execution. */
+	All = Depth | CustomDepth
+};
+
+/** Struct containing the scene texture configuration used to create scene textures.  Use InitializeViewFamily to initialize the
+ *  SceneTexturesConfig structure in the FViewFamilyInfo.  A global singleton instance is maintained manually with static Set / Get
+ *  functions, but will soon be deprecated, in preference of using the structure from the FViewFamilyInfo.
+ */
+struct ENGINE_API FSceneTexturesConfig
+{
+	// Sets the persistent global config instance.
+	static void Set(const FSceneTexturesConfig& Config)
+	{
+		GlobalInstance = Config;
+	}
+
+	// Gets the persistent global config instance. If unset, will return a default constructed instance.
+	static const FSceneTexturesConfig& Get()
+	{
+		return GlobalInstance;
+	}
+
+	FSceneTexturesConfig()
+		: bRequireMultiView{}
+		, bIsUsingGBuffers{}
+		, bKeepDepthContent{ 1 }
+		, bPreciseDepthAux{}
+		, bSamplesCustomDepthAndStencil{}
+	{}
+
+	FORCEINLINE bool IsValid() const
+	{
+		return ShadingPath != EShadingPath::Num;
+	}
+
+	// Extractions to queue for after execution of the render graph.
+	ESceneTextureExtracts Extracts = ESceneTextureExtracts::All;
+
+	// Enums describing the shading / feature / platform configurations used to construct the config.
+	EShadingPath ShadingPath = EShadingPath::Num;
+	ERHIFeatureLevel::Type FeatureLevel = ERHIFeatureLevel::SM5;
+	EShaderPlatform ShaderPlatform = SP_PCD3D_SM5;
+
+	// Extent of all full-resolution textures.
+	FIntPoint Extent = FIntPoint::ZeroValue;
+
+	// Extend of the mobile Pixel Projected Reflection texture
+	FIntPoint MobilePixelProjectedReflectionExtent = FIntPoint::ZeroValue;
+
+	// Downsample factors to divide against the full resolution texture extent.
+	uint32 SmallDepthDownsampleFactor = 2;
+
+	// Number of MSAA samples used by color and depth targets.
+	uint32 NumSamples = 1;
+
+	// Number of MSAA sampled used by the editor primitive composition targets.
+	uint32 EditorPrimitiveNumSamples = 1;
+
+	// Pixel format to use when creating scene color.
+	EPixelFormat ColorFormat = PF_Unknown;
+
+	// Optimized clear values to use for color / depth textures.
+	FClearValueBinding ColorClearValue = FClearValueBinding::Black;
+	FClearValueBinding DepthClearValue = FClearValueBinding::DepthFar;
+
+	// (Deferred Shading) Dynamic GBuffer configuration used to control allocation and slotting of base pass textures.
+	FGBufferParams GBufferParams;
+	FGBufferBinding GBufferA;
+	FGBufferBinding GBufferB;
+	FGBufferBinding GBufferC;
+	FGBufferBinding GBufferD;
+	FGBufferBinding GBufferE;
+	FGBufferBinding GBufferVelocity;
+
+	// (VR) True if scene color and depth should be multi-view allocated.
+	uint32 bRequireMultiView : 1;
+
+	// True if platform is using GBuffers.
+	uint32 bIsUsingGBuffers : 1;
+
+	// (Mobile) True if the platform should write depth content back to memory.
+	uint32 bKeepDepthContent : 1;
+
+	// (Mobile) True if SceneDepthAux should use a precise pixel format
+	uint32 bPreciseDepthAux : 1;
+
+	// (Mobile) True if both CustomDepth and CustomStencil are sampled in a shader
+	uint32 bSamplesCustomDepthAndStencil : 1;
+
+private:
+	static FSceneTexturesConfig GlobalInstance;
+};
diff --git a/Engine/Source/Runtime/Renderer/Private/DeferredShadingRenderer.cpp b/Engine/Source/Runtime/Renderer/Private/DeferredShadingRenderer.cpp
--- a/Engine/Source/Runtime/Renderer/Private/DeferredShadingRenderer.cpp
+++ b/Engine/Source/Runtime/Renderer/Private/DeferredShadingRenderer.cpp
@@ -2081,7 +2081,7 @@ void FDeferredShadingSceneRenderer::Render(FRDGBuilder& GraphBuilder)
 		RectLightAtlas::UpdateRectLightAtlasTexture(GraphBuilder, FeatureLevel);
 	}
 
-	FSceneTexturesConfig::InitializeViewFamily(ViewFamily);
+	InitializeSceneTexturesConfig(ViewFamily.SceneTexturesConfig, ViewFamily);
 	FSceneTexturesConfig& SceneTexturesConfig = GetActiveSceneTexturesConfig();
 	FSceneTexturesConfig::Set(SceneTexturesConfig);
 
diff --git a/Engine/Source/Runtime/Renderer/Private/MobileShadingRenderer.cpp b/Engine/Source/Runtime/Renderer/Private/MobileShadingRenderer.cpp
--- a/Engine/Source/Runtime/Renderer/Private/MobileShadingRenderer.cpp
+++ b/Engine/Source/Runtime/Renderer/Private/MobileShadingRenderer.cpp
@@ -768,7 +768,7 @@ void FMobileSceneRenderer::Render(FRDGBuilder& GraphBuilder)
 
 	WaitOcclusionTests(GraphBuilder.RHICmdList);
 
-	FSceneTexturesConfig::InitializeViewFamily(ViewFamily);
+	InitializeSceneTexturesConfig(ViewFamily.SceneTexturesConfig, ViewFamily);
 	FSceneTexturesConfig& SceneTexturesConfig = GetActiveSceneTexturesConfig();
 	FSceneTexturesConfig::Set(SceneTexturesConfig);
 
diff --git a/Engine/Source/Runtime/Renderer/Private/Renderer.cpp b/Engine/Source/Runtime/Renderer/Private/Renderer.cpp
--- a/Engine/Source/Runtime/Renderer/Private/Renderer.cpp
+++ b/Engine/Source/Runtime/Renderer/Private/Renderer.cpp
@@ -136,7 +136,7 @@ void FRendererModule::DrawTileMesh(FCanvasRenderContext& RenderContext, FMeshPas
 
 		// Default init of SceneTexturesConfig will take extents from FSceneTextureExtentState.
 		// We want the view extents, so explicitly set that.
-		FSceneTexturesConfig::InitializeViewFamily(*ViewFamily);
+		InitializeSceneTexturesConfig(ViewFamily->SceneTexturesConfig, *ViewFamily);
 		ViewFamily->SceneTexturesConfig.Extent = View.ViewRect.Size();
 
 		const auto FeatureLevel = View.GetFeatureLevel();
diff --git a/Engine/Source/Runtime/Renderer/Private/SceneHitProxyRendering.cpp b/Engine/Source/Runtime/Renderer/Private/SceneHitProxyRendering.cpp
--- a/Engine/Source/Runtime/Renderer/Private/SceneHitProxyRendering.cpp
+++ b/Engine/Source/Runtime/Renderer/Private/SceneHitProxyRendering.cpp
@@ -573,7 +573,7 @@ void FMobileSceneRenderer::RenderHitProxies(FRDGBuilder& GraphBuilder)
 	PrepareViewRectsForRendering(GraphBuilder.RHICmdList);
 
 #if WITH_EDITOR
-	FSceneTexturesConfig::InitializeViewFamily(ViewFamily);
+	InitializeSceneTexturesConfig(ViewFamily.SceneTexturesConfig, ViewFamily);
 	FSceneTexturesConfig& SceneTexturesConfig = GetActiveSceneTexturesConfig();
 	FSceneTexturesConfig::Set(SceneTexturesConfig);
 
@@ -615,7 +615,7 @@ void FDeferredShadingSceneRenderer::RenderHitProxies(FRDGBuilder& GraphBuilder)
 	PrepareViewRectsForRendering(GraphBuilder.RHICmdList);
 
 #if WITH_EDITOR
-	FSceneTexturesConfig::InitializeViewFamily(ViewFamily);
+	InitializeSceneTexturesConfig(ViewFamily.SceneTexturesConfig, ViewFamily);
 	FSceneTexturesConfig& SceneTexturesConfig = GetActiveSceneTexturesConfig();
 	FSceneTexturesConfig::Set(SceneTexturesConfig);
 
diff --git a/Engine/Source/Runtime/Renderer/Private/SceneTextures.cpp b/Engine/Source/Runtime/Renderer/Private/SceneTextures.cpp
--- a/Engine/Source/Runtime/Renderer/Private/SceneTextures.cpp
+++ b/Engine/Source/Runtime/Renderer/Private/SceneTextures.cpp
@@ -75,10 +75,6 @@ static TAutoConsoleVariable<int32> CVarDefaultBackBufferPixelFormat(
 	TEXT(" 4: 10bit RGB, 2bit Alpha\n"),
 	ECVF_ReadOnly);
 
-IMPLEMENT_STATIC_UNIFORM_BUFFER_SLOT(SceneTextures);
-IMPLEMENT_STATIC_UNIFORM_BUFFER_STRUCT(FSceneTextureUniformParameters, "SceneTexturesStruct", SceneTextures);
-IMPLEMENT_STATIC_UNIFORM_BUFFER_STRUCT(FMobileSceneTextureUniformParameters, "MobileSceneTextures", SceneTextures);
-
 RDG_REGISTER_BLACKBOARD_STRUCT(FSceneTextures);
 
 static EPixelFormat GetGBufferFFormat()
@@ -531,11 +527,8 @@ void ResetSceneTextureExtentHistory()
 
 ENUM_CLASS_FLAGS(FSceneTextureExtentState::ERenderTargetHistory);
 
-FSceneTexturesConfig FSceneTexturesConfig::GlobalInstance;
-
-void FSceneTexturesConfig::InitializeViewFamily(FViewFamilyInfo& ViewFamily)
+void InitializeSceneTexturesConfig(FSceneTexturesConfig& Config, const FSceneViewFamily& ViewFamily)
 {
-	FSceneTexturesConfig& Config = ViewFamily.SceneTexturesConfig;
 	Config.FeatureLevel = ViewFamily.GetFeatureLevel();
 	Config.ShadingPath = FSceneInterface::GetShadingPath(Config.FeatureLevel);
 	Config.ShaderPlatform = GetFeatureLevelShaderPlatform(Config.FeatureLevel);
@@ -610,16 +603,6 @@ void FSceneTexturesConfig::InitializeViewFamily(FViewFamilyInfo& ViewFamily)
 	}
 }
 
-void FSceneTexturesConfig::Set(const FSceneTexturesConfig& Config)
-{
-	GlobalInstance = Config;
-}
-
-const FSceneTexturesConfig& FSceneTexturesConfig::Get()
-{
-	return GlobalInstance;
-}
-
 void FMinimalSceneTextures::InitializeViewFamily(FRDGBuilder& GraphBuilder, FViewFamilyInfo& ViewFamily)
 {
 	const FSceneTexturesConfig& Config = ViewFamily.SceneTexturesConfig;
diff --git a/Engine/Source/Runtime/Renderer/Private/SceneTextures.h b/Engine/Source/Runtime/Renderer/Private/SceneTextures.h
--- a/Engine/Source/Runtime/Renderer/Private/SceneTextures.h
+++ b/Engine/Source/Runtime/Renderer/Private/SceneTextures.h
@@ -12,106 +12,8 @@ struct FSceneTextures;
 class FViewInfo;
 class FViewFamilyInfo;
 
-enum class ESceneTextureExtracts : uint32
-{
-	/** No textures are extracted from the render graph after execution. */
-	None = 0,
-
-	/** Extracts scene depth after execution */
-	Depth = 1 << 0,
-
-	/** Extracts custom depth after execution. */
-	CustomDepth = 1 << 1,
-
-	/** Extracts all available textures after execution. */
-	All = Depth | CustomDepth
-};
-
-/** Struct containing the scene texture configuration used to create scene textures.  Use InitializeViewFamily to initialize the
- *  SceneTexturesConfig structure in the FViewFamilyInfo.  A global singleton instance is maintained manually with static Set / Get
- *  functions, but will soon be deprecated, in preference of using the structure from the FViewFamilyInfo.
- */
-struct RENDERER_API FSceneTexturesConfig
-{
-	// Sets the persistent global config instance.
-	static void Set(const FSceneTexturesConfig& Config);
-
-	// Gets the persistent global config instance. If unset, will return a default constructed instance.
-	static const FSceneTexturesConfig& Get();
-
-	// Initializes the SceneTexturesConfig in the target view family.
-	static void InitializeViewFamily(FViewFamilyInfo& ViewFamily);
-
-	FORCEINLINE bool IsValid() const
-	{
-		return ShadingPath != EShadingPath::Num;
-	}
-
-	FSceneTexturesConfig()
-		: bRequireMultiView{}
-		, bIsUsingGBuffers{}
-		, bKeepDepthContent{1}
-		, bPreciseDepthAux{}
-		, bSamplesCustomDepthAndStencil{}
-	{}
-
-	// Extractions to queue for after execution of the render graph.
-	ESceneTextureExtracts Extracts = ESceneTextureExtracts::All;
-
-	// Enums describing the shading / feature / platform configurations used to construct the config.
-	EShadingPath ShadingPath = EShadingPath::Num;
-	ERHIFeatureLevel::Type FeatureLevel = ERHIFeatureLevel::SM5;
-	EShaderPlatform ShaderPlatform = SP_PCD3D_SM5;
-
-	// Extent of all full-resolution textures.
-	FIntPoint Extent = FIntPoint::ZeroValue;
-
-	// Extend of the mobile Pixel Projected Reflection texture
-	FIntPoint MobilePixelProjectedReflectionExtent = FIntPoint::ZeroValue;
-
-	// Downsample factors to divide against the full resolution texture extent.
-	uint32 SmallDepthDownsampleFactor = 2;
-
-	// Number of MSAA samples used by color and depth targets.
-	uint32 NumSamples = 1;
-
-	// Number of MSAA sampled used by the editor primitive composition targets.
-	uint32 EditorPrimitiveNumSamples = 1;
-
-	// Pixel format to use when creating scene color.
-	EPixelFormat ColorFormat = PF_Unknown;
-
-	// Optimized clear values to use for color / depth textures.
-	FClearValueBinding ColorClearValue = FClearValueBinding::Black;
-	FClearValueBinding DepthClearValue = FClearValueBinding::DepthFar;
-
-	// (Deferred Shading) Dynamic GBuffer configuration used to control allocation and slotting of base pass textures.
-	FGBufferParams GBufferParams;
-	FGBufferBinding GBufferA;
-	FGBufferBinding GBufferB;
-	FGBufferBinding GBufferC;
-	FGBufferBinding GBufferD;
-	FGBufferBinding GBufferE;
-	FGBufferBinding GBufferVelocity;
-
-	// (VR) True if scene color and depth should be multi-view allocated.
-	uint32 bRequireMultiView : 1;
-
-	// True if platform is using GBuffers.
-	uint32 bIsUsingGBuffers : 1;
-
-	// (Mobile) True if the platform should write depth content back to memory.
-	uint32 bKeepDepthContent : 1;
-	
-	// (Mobile) True if SceneDepthAux should use a precise pixel format
-	uint32 bPreciseDepthAux : 1;
-
-	// (Mobile) True if both CustomDepth and CustomStencil are sampled in a shader
-	uint32 bSamplesCustomDepthAndStencil : 1;
-
-private:
-	static FSceneTexturesConfig GlobalInstance;
-};
+/** Initializes a scene textures config instance from the view family. */
+extern RENDERER_API void InitializeSceneTexturesConfig(FSceneTexturesConfig& Config, const FSceneViewFamily& ViewFamily);
 
 /** RDG struct containing the minimal set of scene textures common across all rendering configurations. */
 struct RENDERER_API FMinimalSceneTextures
diff --git a/Engine/Source/Runtime/Renderer/Public/SceneRenderTargetParameters.h b/Engine/Source/Runtime/Renderer/Public/SceneRenderTargetParameters.h
--- a/Engine/Source/Runtime/Renderer/Public/SceneRenderTargetParameters.h
+++ b/Engine/Source/Runtime/Renderer/Public/SceneRenderTargetParameters.h
@@ -5,40 +5,12 @@
 #if UE_ENABLE_INCLUDE_ORDER_DEPRECATED_IN_5_1
 #include "RenderGraph.h"
 #endif
-#include "RenderGraphDefinitions.h"
-#include "ShaderParameterMacros.h"
 #include "RendererInterface.h"
+#include "SceneTexturesConfig.h"
 
 class FRDGBuilder;
-struct IPooledRenderTarget;
 struct FSceneTextures;
 
-/** A uniform buffer containing common scene textures used by materials or global shaders. */
-BEGIN_GLOBAL_SHADER_PARAMETER_STRUCT(FSceneTextureUniformParameters, RENDERER_API)
-	// Scene Color / Depth
-	SHADER_PARAMETER_RDG_TEXTURE(Texture2D, SceneColorTexture)
-	SHADER_PARAMETER_RDG_TEXTURE(Texture2D, SceneDepthTexture)
-
-	// GBuffer
-	SHADER_PARAMETER_RDG_TEXTURE(Texture2D, GBufferATexture)
-	SHADER_PARAMETER_RDG_TEXTURE(Texture2D, GBufferBTexture)
-	SHADER_PARAMETER_RDG_TEXTURE(Texture2D, GBufferCTexture)
-	SHADER_PARAMETER_RDG_TEXTURE(Texture2D, GBufferDTexture)
-	SHADER_PARAMETER_RDG_TEXTURE(Texture2D, GBufferETexture)
-	SHADER_PARAMETER_RDG_TEXTURE(Texture2D, GBufferFTexture)
-	SHADER_PARAMETER_RDG_TEXTURE(Texture2D, GBufferVelocityTexture)
-
-	// SSAO
-	SHADER_PARAMETER_RDG_TEXTURE(Texture2D, ScreenSpaceAOTexture)
-
-	// Custom Depth / Stencil
-	SHADER_PARAMETER_RDG_TEXTURE(Texture2D, CustomDepthTexture)
-	SHADER_PARAMETER_RDG_TEXTURE_SRV(Texture2D<uint2>, CustomStencilTexture)
-
-	// Misc
-	SHADER_PARAMETER_SAMPLER(SamplerState, PointClampSampler)
-END_GLOBAL_SHADER_PARAMETER_STRUCT()
-
 enum class ESceneTextureSetupMode : uint32
 {
 	None			= 0,
@@ -73,29 +45,6 @@ extern RENDERER_API TRDGUniformBufferRef<FSceneTextureUniformParameters> CreateS
 	ERHIFeatureLevel::Type FeatureLevel,
 	ESceneTextureSetupMode SetupMode = ESceneTextureSetupMode::All);
 
-BEGIN_GLOBAL_SHADER_PARAMETER_STRUCT(FMobileSceneTextureUniformParameters, RENDERER_API)
-	SHADER_PARAMETER_RDG_TEXTURE(Texture2D, SceneColorTexture)
-	SHADER_PARAMETER_SAMPLER(SamplerState, SceneColorTextureSampler)
-	SHADER_PARAMETER_RDG_TEXTURE(Texture2D, SceneDepthTexture)
-	SHADER_PARAMETER_SAMPLER(SamplerState, SceneDepthTextureSampler)
-	SHADER_PARAMETER_RDG_TEXTURE(Texture2D, CustomDepthTexture)
-	SHADER_PARAMETER_SAMPLER(SamplerState, CustomDepthTextureSampler)
-	SHADER_PARAMETER_RDG_TEXTURE_SRV(Texture2D<uint2>, CustomStencilTexture)
-	SHADER_PARAMETER_RDG_TEXTURE(Texture2D, SceneVelocityTexture)
-	SHADER_PARAMETER_SAMPLER(SamplerState, SceneVelocityTextureSampler)
-	// GBuffer
-	SHADER_PARAMETER_RDG_TEXTURE(Texture2D, GBufferATexture)
-	SHADER_PARAMETER_RDG_TEXTURE(Texture2D, GBufferBTexture)
-	SHADER_PARAMETER_RDG_TEXTURE(Texture2D, GBufferCTexture)
-	SHADER_PARAMETER_RDG_TEXTURE(Texture2D, GBufferDTexture)
-	SHADER_PARAMETER_RDG_TEXTURE(Texture2D, SceneDepthAuxTexture)
-	SHADER_PARAMETER_SAMPLER(SamplerState, GBufferATextureSampler)
-	SHADER_PARAMETER_SAMPLER(SamplerState, GBufferBTextureSampler)
-	SHADER_PARAMETER_SAMPLER(SamplerState, GBufferCTextureSampler)
-	SHADER_PARAMETER_SAMPLER(SamplerState, GBufferDTextureSampler)
-	SHADER_PARAMETER_SAMPLER(SamplerState, SceneDepthAuxTextureSampler)
-END_GLOBAL_SHADER_PARAMETER_STRUCT()
-
 enum class EMobileSceneTextureSetupMode : uint32
 {
 	None			= 0,
@@ -126,25 +75,6 @@ extern RENDERER_API TRDGUniformBufferRef<FMobileSceneTextureUniformParameters> C
 	const FSceneTextures* SceneTextures,
 	EMobileSceneTextureSetupMode SetupMode = EMobileSceneTextureSetupMode::All);
 
-BEGIN_SHADER_PARAMETER_STRUCT(FSceneTextureShaderParameters, RENDERER_API)
-	SHADER_PARAMETER_RDG_UNIFORM_BUFFER(FSceneTextureUniformParameters, SceneTextures)
-	SHADER_PARAMETER_RDG_UNIFORM_BUFFER(FMobileSceneTextureUniformParameters, MobileSceneTextures)
-END_SHADER_PARAMETER_STRUCT()
-
-FORCEINLINE FSceneTextureShaderParameters GetSceneTextureShaderParameters(TRDGUniformBufferRef<FSceneTextureUniformParameters> UniformBuffer)
-{
-	FSceneTextureShaderParameters Parameters;
-	Parameters.SceneTextures = UniformBuffer;
-	return Parameters;
-}
-
-FORCEINLINE FSceneTextureShaderParameters GetSceneTextureShaderParameters(TRDGUniformBufferRef<FMobileSceneTextureUniformParameters> UniformBuffer)
-{
-	FSceneTextureShaderParameters Parameters;
-	Parameters.MobileSceneTextures = UniformBuffer;
-	return Parameters;
-}
-
 /** Returns scene texture shader parameters containing the RDG uniform buffer for either mobile or deferred shading. */
 extern RENDERER_API FSceneTextureShaderParameters CreateSceneTextureShaderParameters(
 	FRDGBuilder& GraphBuilder,
