Fixed level instance computation of pivot location at creation of level instance.

diff --git a/Engine/Source/Runtime/Engine/Private/LevelInstance/LevelInstanceSubsystem.cpp b/Engine/Source/Runtime/Engine/Private/LevelInstance/LevelInstanceSubsystem.cpp
--- a/Engine/Source/Runtime/Engine/Private/LevelInstance/LevelInstanceSubsystem.cpp
+++ b/Engine/Source/Runtime/Engine/Private/LevelInstance/LevelInstanceSubsystem.cpp
@@ -876,9 +876,15 @@ ILevelInstanceInterface* ULevelInstanceSubsystem::CreateLevelInstanceFrom(const
 	FBox ActorLocationBox(ForceInit);
 	for (const AActor* ActorToMove : ActorsToMove)
 	{
-		const bool bNonColliding = false;
+		const bool bNonColliding = true;
 		const bool bIncludeChildren = true;
-		ActorLocationBox += ActorToMove->GetComponentsBoundingBox(bNonColliding, bIncludeChildren);
+		FBox LocalActorLocationBox = ActorToMove->GetComponentsBoundingBox(bNonColliding, bIncludeChildren);
+		// In the case where we receive an invalid bounding box, use actor's location if it has a root component
+		if (!LocalActorLocationBox.IsValid && ActorToMove->GetRootComponent())
+		{
+			LocalActorLocationBox = FBox({ ActorToMove->GetActorLocation() });
+		}
+		ActorLocationBox += LocalActorLocationBox;
 
 		FText Reason;
 		if (!CanMoveActorToLevel(ActorToMove, &Reason))
