Add Probe and QueryAndProbe to ECollisionEnabled. Single shpaes or entire objects can now be marked as physics probes, meaning that they will generate constraints, contact data, and hit results, without physically responding.

diff --git a/Engine/Source/Runtime/Core/Public/UObject/FortniteNCBranchObjectVersion.h b/Engine/Source/Runtime/Core/Public/UObject/FortniteNCBranchObjectVersion.h
--- a/Engine/Source/Runtime/Core/Public/UObject/FortniteNCBranchObjectVersion.h
+++ b/Engine/Source/Runtime/Core/Public/UObject/FortniteNCBranchObjectVersion.h
@@ -33,6 +33,9 @@ struct CORE_API FFortniteNCBranchObjectVersion
 		// Serialize property map for actor descriptors
 		WorldPartitionActorDescPropertyMapSerialization,
 
+		// Added ability to mark shapes as probes
+		AddShapeIsProbe,
+
 		// -----<new versions can be added above this line>-------------------------------------------------
 		VersionPlusOne,
 		LatestVersion = VersionPlusOne - 1
diff --git a/Engine/Source/Runtime/Engine/Classes/Engine/EngineTypes.h b/Engine/Source/Runtime/Engine/Classes/Engine/EngineTypes.h
--- a/Engine/Source/Runtime/Engine/Classes/Engine/EngineTypes.h
+++ b/Engine/Source/Runtime/Engine/Classes/Engine/EngineTypes.h
@@ -1382,7 +1382,15 @@ namespace ECollisionEnabled
 		/** Only used only for physics simulation (rigid body, constraints). Cannot be used for spatial queries (raycasts, sweeps, overlaps). Useful for jiggly bits on characters that do not need per bone detection. Performance gains by keeping data out of query tree */
 		PhysicsOnly UMETA(DisplayName="Physics Only (No Query Collision)"),
 		/** Can be used for both spatial queries (raycasts, sweeps, overlaps) and simulation (rigid body, constraints). */
-		QueryAndPhysics UMETA(DisplayName="Collision Enabled (Query and Physics)") 
+		QueryAndPhysics UMETA(DisplayName="Collision Enabled (Query and Physics)"),
+		/** Only used for probing the physics simulation (rigid body, constraints). Cannot be used for spatial queries (raycasts,
+		sweeps, overlaps). Useful for when you want to detect potential physics interactions and pass contact data to hit callbacks
+		or contact modification, but don't want to physically react to these contacts. */
+		ProbeOnly UMETA(DisplayName="Probe Only (Contact Data, No Query or Physics Collision)"),
+		/** Can be used for both spatial queries (raycasts, sweeps, overlaps) and probing the physics simulation (rigid body,
+		constraints). Will not allow for actual physics interaction, but will generate contact data, trigger hit callbacks, and
+		contacts will appear in contact modification. */
+		QueryAndProbe UMETA(DisplayName="Query and Probe (Query Collision and Contact Data, No Physics Collision)")
 	}; 
 } 
 
@@ -1409,23 +1417,61 @@ extern FCollisionEnabledMask ENGINE_API operator|(const ECollisionEnabled::Type
 
 FORCEINLINE bool CollisionEnabledHasPhysics(ECollisionEnabled::Type CollisionEnabled)
 {
-	return (CollisionEnabled == ECollisionEnabled::PhysicsOnly) ||
+	return	(CollisionEnabled == ECollisionEnabled::PhysicsOnly) ||
 			(CollisionEnabled == ECollisionEnabled::QueryAndPhysics);
 }
 
 FORCEINLINE bool CollisionEnabledHasQuery(ECollisionEnabled::Type CollisionEnabled)
 {
-	return (CollisionEnabled == ECollisionEnabled::QueryOnly) ||
-			(CollisionEnabled == ECollisionEnabled::QueryAndPhysics);
+	return	(CollisionEnabled == ECollisionEnabled::QueryOnly) ||
+			(CollisionEnabled == ECollisionEnabled::QueryAndPhysics) ||
+			(CollisionEnabled == ECollisionEnabled::QueryAndProbe);
+}
+
+FORCEINLINE bool CollisionEnabledHasProbe(ECollisionEnabled::Type CollisionEnabled)
+{
+	return (CollisionEnabled == ECollisionEnabled::ProbeOnly) ||
+			(CollisionEnabled == ECollisionEnabled::QueryAndProbe);
 }
 
 FORCEINLINE ECollisionEnabled::Type CollisionEnabledIntersection(ECollisionEnabled::Type CollisionEnabledA, ECollisionEnabled::Type CollisionEnabledB)
 {
-	const bool bHasQuery = (CollisionEnabledHasQuery(CollisionEnabledA) && CollisionEnabledHasQuery(CollisionEnabledB));
-	const bool bHasPhysics = (CollisionEnabledHasPhysics(CollisionEnabledA) && CollisionEnabledHasPhysics(CollisionEnabledB));
-	if (bHasQuery && bHasPhysics) { return ECollisionEnabled::QueryAndPhysics; }
-	if (bHasQuery) { return ECollisionEnabled::QueryOnly; }
-	if (bHasPhysics) { return ECollisionEnabled::PhysicsOnly; }
+	// Combine collision two enabled data.
+	//
+	// The intersection follows the following rules:
+	//
+	// * For the result to have Query, both data must have Query
+	// * For the result to have Probe, either data must have Probe
+	// * For the result to have Physics, both data must have Physics and the result must not have Probe
+	//
+	// This way if an object is query-only, for example, but one of its shapes is query-and-physics,
+	// the object's settings win and the shape ends up being query-only. And if the object is marked
+	// as physics, but the child is marked as probe (or vice versa), use probe.
+	//
+	// The following matrix represents the intersection relationship.
+	// NOTE: The order of types must match the order declared in ECollisionEnabled!
+	using namespace ECollisionEnabled;
+	static constexpr ECollisionEnabled::Type IntersectionMatrix[5][5] = {
+		/*                 |        QueryOnly       PhysicsOnly     QueryAndPhysics     ProbeOnly     QueryAndProbe   */
+		/*-----------------+------------------------------------------------------------------------------------------*/
+		/* QueryOnly       | */ {   QueryOnly,      NoCollision,    QueryOnly,          NoCollision,  QueryOnly       },
+		/* PhysicsOnly     | */ {   NoCollision,    PhysicsOnly,    PhysicsOnly,        ProbeOnly,    ProbeOnly       },
+		/* QueryAndPhysics | */ {   QueryOnly,      PhysicsOnly,    QueryAndPhysics,    ProbeOnly,    QueryAndProbe   },
+		/* ProbeOnly       | */ {   NoCollision,    ProbeOnly,      ProbeOnly,          ProbeOnly,    ProbeOnly       },
+		/* QueryAndProbe   | */ {   QueryOnly,      ProbeOnly,      QueryAndProbe,      ProbeOnly,    QueryAndProbe   }
+	};
+
+	// Subtract 1 because the first index is NoCollision.
+	// If both indices indicate _some_ collision setting (ie, greater than -1),
+	// lookup their intersection setting in the matrix.
+	const int8 IndexA = (int32)CollisionEnabledA - 1;
+	const int8 IndexB = (int32)CollisionEnabledB - 1;
+	if (IndexA >= 0 && IndexB >= 0)
+	{
+		return IntersectionMatrix[IndexA][IndexB];
+	}
+
+	// Either of the two settings were set to NoCollision which trumps all
 	return ECollisionEnabled::NoCollision;
 }
 
diff --git a/Engine/Source/Runtime/Engine/Private/PhysicsEngine/BodyInstance.cpp b/Engine/Source/Runtime/Engine/Private/PhysicsEngine/BodyInstance.cpp
--- a/Engine/Source/Runtime/Engine/Private/PhysicsEngine/BodyInstance.cpp
+++ b/Engine/Source/Runtime/Engine/Private/PhysicsEngine/BodyInstance.cpp
@@ -1039,9 +1039,11 @@ void FBodyInstance::UpdatePhysicsFilterData()
 				const FBodyCollisionFilterData& FilterData = PerShapeCollisionData.CollisionFilterData;
 				const bool bNewQueryShape = CollisionFlags.bEnableQueryCollision && (!bIsStatic || bSyncShape);
 				const bool bNewSimShape = bIsTrimesh ? CollisionFlags.bEnableSimCollisionComplex : CollisionFlags.bEnableSimCollisionSimple;
+				const bool bNewProbeShape = CollisionFlags.bEnableProbeCollision;
 
 				FPhysicsInterface::SetIsQueryShape(InnerShape, bNewQueryShape);
 				FPhysicsInterface::SetIsSimulationShape(InnerShape, bNewSimShape);
+				FPhysicsInterface::SetIsProbeShape(InnerShape, bNewProbeShape);
 
 				// If we changed 'simulation collision' on a shape, we need to recalc mass properties
 				if (bWasSimulationShape != bNewSimShape)
@@ -4163,6 +4165,7 @@ void FBodyInstance::BuildBodyCollisionFlags(FBodyCollisionFlags& OutFlags, EColl
 	OutFlags.bEnableQueryCollision = false;
 	OutFlags.bEnableSimCollisionSimple = false;
 	OutFlags.bEnableSimCollisionComplex = false;
+	OutFlags.bEnableProbeCollision = false;
 
 	if(UseCollisionEnabled != ECollisionEnabled::NoCollision)
 	{
@@ -4170,11 +4173,16 @@ void FBodyInstance::BuildBodyCollisionFlags(FBodyCollisionFlags& OutFlags, EColl
 		OutFlags.bEnableQueryCollision = CollisionEnabledHasQuery(UseCollisionEnabled);
 
 		// Sim collision
-		const bool bSimCollision = SimCollisionEnabled && CollisionEnabledHasPhysics(UseCollisionEnabled);
+		const bool bProbe = CollisionEnabledHasProbe(UseCollisionEnabled);
+		const bool bSimCollision = SimCollisionEnabled && (bProbe || CollisionEnabledHasPhysics(UseCollisionEnabled));
 
 		// Enable sim collision
 		if(bSimCollision)
 		{
+			// Objects marked as probes use sim collision, but don't actually have physical reactions
+			OutFlags.bEnableProbeCollision = bProbe;
+
+			// We use simple sim collision even if we also use complex sim collision?
 			OutFlags.bEnableSimCollisionSimple = true;
 			
 			// on dynamic objects and objects which don't use complex as simple, tri mesh not used for sim
@@ -4201,7 +4209,7 @@ void FBodyInstance::UpdateInterpolateWhenSubStepping()
 			UseCollisionEnabled = GetCollisionEnabled();
 		}
 	
-		bInterpolateWhenSubStepping = UseCollisionEnabled == ECollisionEnabled::PhysicsOnly || UseCollisionEnabled == ECollisionEnabled::QueryAndPhysics;
+		bInterpolateWhenSubStepping = CollisionEnabledHasPhysics(UseCollisionEnabled);
 
 		// If we have a weld parent we should take into account that too as that may be simulating while we are not
 		if(WeldParent)
diff --git a/Engine/Source/Runtime/Experimental/Chaos/Private/Chaos/CCDUtilities.cpp b/Engine/Source/Runtime/Experimental/Chaos/Private/Chaos/CCDUtilities.cpp
--- a/Engine/Source/Runtime/Experimental/Chaos/Private/Chaos/CCDUtilities.cpp
+++ b/Engine/Source/Runtime/Experimental/Chaos/Private/Chaos/CCDUtilities.cpp
@@ -46,6 +46,12 @@ namespace Chaos
 		int32 CCDAxisThresholdMode = 1;
 		FAutoConsoleVariableRef  CVarCCDAxisThresholdMode(TEXT("p.Chaos.CCD.AxisThresholdMode"), CCDAxisThresholdMode , TEXT("Change the mode used to generate CCD axis threshold bounds for particle geometries.\n0: Use object bounds\n1: Find the thinnest object bound on any axis and use it for all CCD axes\n2: On each axis, use the thinnest shape bound on that axis\n3: Find the thinnest shape bound on any axis and use this for all axes"));
 
+		bool bCCDAxisThresholdUsesProbeShapes = false;
+		FAutoConsoleVariableRef  CVarCCDAxisThresholdUsesProbeShapes(TEXT("p.Chaos.CCD.CCDAxisThresholdUsesProbeShapes"), bCCDAxisThresholdUsesProbeShapes , TEXT("When true, probe shapes are considered for CCD axis threshold computation, and can generate contacts in the initial CCD phase."));
+
+		bool bCCDSweepsUseProbeShapes = false;
+		FAutoConsoleVariableRef  CVarCCDSweepsUseProbeShapes(TEXT("p.Chaos.CCD.CCDSweepsUseProbeShapes"), bCCDSweepsUseProbeShapes , TEXT("When true, probe shapes can be swept for more accurate collision detection."));
+
 		extern int32 ChaosSolverDrawCCDInteractions;
 
 #if CHAOS_DEBUG_DRAW
diff --git a/Engine/Source/Runtime/Experimental/Chaos/Private/Chaos/ChaosDebugDraw.cpp b/Engine/Source/Runtime/Experimental/Chaos/Private/Chaos/ChaosDebugDraw.cpp
--- a/Engine/Source/Runtime/Experimental/Chaos/Private/Chaos/ChaosDebugDraw.cpp
+++ b/Engine/Source/Runtime/Experimental/Chaos/Private/Chaos/ChaosDebugDraw.cpp
@@ -875,6 +875,7 @@ namespace Chaos
 
 				for (const FManifoldPoint& ManifoldPoint : Contact.GetManifoldPoints())
 				{
+					const bool bIsProbe = Contact.GetIsProbe();
 					const bool bIsActive = !ManifoldPoint.NetPushOut.IsNearlyZero() || !ManifoldPoint.NetImpulse.IsNearlyZero() || (!Contact.GetUseManifold() && !Contact.AccumulatedImpulse.IsNearlyZero());
 					if (!bIsActive && !bChaosDebugDebugDrawInactiveContacts)
 					{
@@ -906,7 +907,13 @@ namespace Chaos
 					{
 						DiscColor = FColor(150, 200, 0);
 					}
-					if (!bIsActive)
+					if (bIsProbe)
+					{
+						DiscColor = FColor(50, 180, 180);
+						PlaneNormalColor = FColor(50, 180, 180);
+						EdgeNormalColor = FColor(50, 180, 130);
+					}
+					else if (!bIsActive)
 					{
 						DiscColor = FColor(100, 100, 100);
 						PlaneNormalColor = FColor(100, 0, 0);
diff --git a/Engine/Source/Runtime/Experimental/Chaos/Private/Chaos/Collision/PBDCollisionConstraint.cpp b/Engine/Source/Runtime/Experimental/Chaos/Private/Chaos/Collision/PBDCollisionConstraint.cpp
--- a/Engine/Source/Runtime/Experimental/Chaos/Private/Chaos/Collision/PBDCollisionConstraint.cpp
+++ b/Engine/Source/Runtime/Experimental/Chaos/Private/Chaos/Collision/PBDCollisionConstraint.cpp
@@ -277,6 +277,12 @@ namespace Chaos
 		Flags.bUseManifold = bInUseManifold;
 		Flags.bUseIncrementalManifold = bInUseManifold;
 
+		// Initialize the is-probe and is-probe-unmodified flags to the same value.
+		// Contact modification may change bIsProbe but we want to store the unmodified value
+		// so that it can be reset between frames.
+		Flags.bIsProbe = 
+		Flags.bIsProbeUnmodified = Shape[0]->GetIsProbe() || Shape[1]->GetIsProbe();
+
 		const FReal Margin0 = GetImplicit0()->GetMargin();
 		const FReal Margin1 = GetImplicit1()->GetMargin();
 		const EImplicitObjectType ImplicitType0 = GetInnerType(GetImplicit0()->GetCollisionType());
diff --git a/Engine/Source/Runtime/Experimental/Chaos/Private/Chaos/Collision/ParticlePairMidPhase.cpp b/Engine/Source/Runtime/Experimental/Chaos/Private/Chaos/Collision/ParticlePairMidPhase.cpp
--- a/Engine/Source/Runtime/Experimental/Chaos/Private/Chaos/Collision/ParticlePairMidPhase.cpp
+++ b/Engine/Source/Runtime/Experimental/Chaos/Private/Chaos/Collision/ParticlePairMidPhase.cpp
@@ -281,7 +281,7 @@ namespace Chaos
 			// NOTE: these are not used by CCD which continuously moves the particles
 			Constraint->SetShapeWorldTransforms(ShapeWorldTransform0, ShapeWorldTransform1);
 
-			Constraint->ResetMaterial();
+			Constraint->ResetModifications();
 			Constraint->SetCCDEnabled(false);
 
 			// If the constraint was not used last frame, it needs to be reset. 
@@ -311,7 +311,7 @@ namespace Chaos
 				// We will be updating the manifold, if only partially, so update the restore comparison transforms
 				Constraint->SetLastShapeWorldTransforms(ShapeWorldTransform0, ShapeWorldTransform1);
 
-				if (!Context.GetSettings().bDeferNarrowPhase)
+				if (!(Context.GetSettings().bDeferNarrowPhase || Constraint->GetIsProbe()))
 				{
 					// Run the narrow phase
 					Collisions::UpdateConstraint(*Constraint.Get(), ShapeWorldTransform0, ShapeWorldTransform1, Dt);
@@ -348,6 +348,11 @@ namespace Chaos
 		{
 			PHYSICS_CSV_SCOPED_EXPENSIVE(PhysicsVerbose, NarrowPhase_UpdateConstraintCCD);
 
+			if (Constraint->GetIsProbe())
+			{
+				return 0;
+			}
+
 			const FImplicitObject* Implicit0 = Shape0->GetLeafGeometry();
 			const FImplicitObject* Implicit1 = Shape1->GetLeafGeometry();
 			const FRigidTransform3& ShapeWorldTransform0 = Shape0->GetLeafWorldTransform(Particle0);
@@ -358,7 +363,7 @@ namespace Chaos
 			// NOTE: these are not used by CCD which continuously moves the particles
 			Constraint->SetShapeWorldTransforms(ShapeWorldTransform0, ShapeWorldTransform1);
 
-			Constraint->ResetMaterial();
+			Constraint->ResetModifications();
 			Constraint->SetCCDEnabled(true);
 			Constraint->SetCullDistance(CullDistance);
 			Constraint->ResetManifold();
@@ -572,7 +577,7 @@ namespace Chaos
 		const FRigidTransform3 ShapeWorldTransform1 = ShapeRelativeTransform1 * ParticleTransform1;
 		Constraint->SetShapeWorldTransforms(ShapeWorldTransform0, ShapeWorldTransform1);
 
-		Constraint->ResetMaterial();
+		Constraint->ResetModifications();
 
 		NewConstraints.Add(Constraint);
 		return Constraint;
diff --git a/Engine/Source/Runtime/Experimental/Chaos/Private/Chaos/PBDRigidsEvolutionGBF.cpp b/Engine/Source/Runtime/Experimental/Chaos/Private/Chaos/PBDRigidsEvolutionGBF.cpp
--- a/Engine/Source/Runtime/Experimental/Chaos/Private/Chaos/PBDRigidsEvolutionGBF.cpp
+++ b/Engine/Source/Runtime/Experimental/Chaos/Private/Chaos/PBDRigidsEvolutionGBF.cpp
@@ -836,7 +836,7 @@ void FPBDRigidsEvolutionGBF::TestModeResetCollisions()
 	for (FPBDCollisionConstraintHandle* Collision : CollisionConstraints.GetConstraintHandles())
 	{
 		Collision->GetContact().ResetManifold();
-		Collision->GetContact().ResetMaterial();
+		Collision->GetContact().ResetModifications();
 		Collision->GetContact().GetGJKWarmStartData().Reset();
 	}
 }
diff --git a/Engine/Source/Runtime/Experimental/Chaos/Public/Chaos/Collision/PBDCollisionConstraint.h b/Engine/Source/Runtime/Experimental/Chaos/Public/Chaos/Collision/PBDCollisionConstraint.h
--- a/Engine/Source/Runtime/Experimental/Chaos/Public/Chaos/Collision/PBDCollisionConstraint.h
+++ b/Engine/Source/Runtime/Experimental/Chaos/Public/Chaos/Collision/PBDCollisionConstraint.h
@@ -378,8 +378,8 @@ namespace Chaos
 		void SetDisabled(bool bInDisabled) { Flags.bDisabled = bInDisabled; }
 		bool GetDisabled() const { return Flags.bDisabled; }
 
-		void SetIsProbe(bool bInProbe) { Flags.bProbe = bInProbe; }
-		bool GetIsProbe() const { return Flags.bProbe; }
+		void SetIsProbe(bool bInProbe) { Flags.bIsProbe = bInProbe; }
+		bool GetIsProbe() const { return Flags.bIsProbe; }
 
 		virtual void SetIsSleeping(const bool bInIsSleeping) override;
 
@@ -441,9 +441,12 @@ namespace Chaos
 		/**
 		* Reset the material properties to those from the shape materials. Called each frame to reset contact modifications to the material.
 		*/
-		inline void ResetMaterial()
+		inline void ResetModifications()
 		{
 			Material.ResetMaterialModifications();
+
+			// Reset other properties which may have changed in contact modification last frame
+			Flags.bIsProbe = Flags.bIsProbeUnmodified;
 		}
 
 		/**
@@ -799,12 +802,13 @@ namespace Chaos
 			struct
 			{
 				uint32 bDisabled : 1;
-				uint32 bProbe : 1;
 				uint32 bUseManifold : 1;
 				uint32 bUseIncrementalManifold : 1;
 				uint32 bWasManifoldRestored : 1;
 				uint32 bIsQuadratic0 : 1;
 				uint32 bIsQuadratic1 : 1;
+				uint32 bIsProbeUnmodified : 1;  // Is this constraint a probe pre-contact-modification
+				uint32 bIsProbe : 1;            // Is this constraint currently a probe
 			};
 			uint32 Bits;
 		} Flags;
diff --git a/Engine/Source/Runtime/Experimental/Chaos/Public/Chaos/GeometryParticles.h b/Engine/Source/Runtime/Experimental/Chaos/Public/Chaos/GeometryParticles.h
--- a/Engine/Source/Runtime/Experimental/Chaos/Public/Chaos/GeometryParticles.h
+++ b/Engine/Source/Runtime/Experimental/Chaos/Public/Chaos/GeometryParticles.h
@@ -34,6 +34,7 @@ namespace Chaos
 	namespace CVars
 	{
 		CHAOS_API extern int32 CCDAxisThresholdMode;
+		CHAOS_API extern bool bCCDAxisThresholdUsesProbeShapes;
 	}
 
 	/** Data that is associated with geometry. If a union is used an entry is created per internal geometry */
@@ -194,6 +195,12 @@ namespace Chaos
 			CollisionData.Modify(true, DirtyFlags, Proxy, ShapeIdx, [bEnable](FCollisionData& Data){ Data.bSimCollision = bEnable; });
 		}
 
+		bool GetIsProbe() const { return CollisionData.Read().bIsProbe; }
+		void SetIsProbe(const bool bIsProbe)
+		{
+			CollisionData.Modify(true, DirtyFlags, Proxy, ShapeIdx, [bIsProbe](FCollisionData& Data){ Data.bIsProbe = bIsProbe; });
+		}
+
 		EChaosCollisionTraceFlag GetCollisionTraceType() const { return CollisionData.Read().CollisionTraceType; }
 		void SetCollisionTraceType(const EChaosCollisionTraceFlag InTraceFlag)
 		{
@@ -662,7 +669,7 @@ namespace Chaos
 						// Only sim-enabled shapes should ever be swept with CCD, so make sure the
 						// sim-enabled flag is on for each shape before considering it's min bounds
 						// for CCD extents.
-						if (Shape->GetSimEnabled())
+						if (Shape->GetSimEnabled() && (CVars::bCCDAxisThresholdUsesProbeShapes || !Shape->GetIsProbe()))
 						{
 							const TSerializablePtr<FImplicitObject> Geometry = Shape->GetGeometry();
 							if (Geometry->HasBoundingBox())
diff --git a/Engine/Source/Runtime/Experimental/Chaos/Public/Chaos/ParticleDirtyFlags.h b/Engine/Source/Runtime/Experimental/Chaos/Public/Chaos/ParticleDirtyFlags.h
--- a/Engine/Source/Runtime/Experimental/Chaos/Public/Chaos/ParticleDirtyFlags.h
+++ b/Engine/Source/Runtime/Experimental/Chaos/Public/Chaos/ParticleDirtyFlags.h
@@ -14,6 +14,7 @@
 #include "Chaos/RigidParticleControlFlags.h"
 #include "UObject/ExternalPhysicsCustomObjectVersion.h"
 #include "UObject/ExternalPhysicsMaterialCustomObjectVersion.h"
+#include "UObject/FortniteNCBranchObjectVersion.h"
 #include "UObject/PhysicsObjectVersion.h"
 #include "UObject/UE5MainStreamObjectVersion.h"
 #include "UObject/UE5ReleaseStreamObjectVersion.h"
@@ -573,12 +574,14 @@ struct FCollisionData
 	EChaosCollisionTraceFlag CollisionTraceType;
 	uint8 bSimCollision : 1;
 	uint8 bQueryCollision : 1;
+	uint8 bIsProbe : 1;
 
 	FCollisionData()
 	: UserData(nullptr)
 	, CollisionTraceType(EChaosCollisionTraceFlag::Chaos_CTF_UseDefault)
 	, bSimCollision(true)
 	, bQueryCollision(true)
+	, bIsProbe(false)
 	{
 	}
 
@@ -588,6 +591,7 @@ struct FCollisionData
 	{
 		Ar.UsingCustomVersion(FExternalPhysicsCustomObjectVersion::GUID);
 		Ar.UsingCustomVersion(FExternalPhysicsMaterialCustomObjectVersion::GUID);
+		Ar.UsingCustomVersion(FFortniteNCBranchObjectVersion::GUID);
 
 		Ar << QueryData;
 		Ar << SimData;
@@ -622,6 +626,13 @@ struct FCollisionData
 			Ar << Data;
 			CollisionTraceType = (EChaosCollisionTraceFlag)Data;
 		}
+
+		if (Ar.CustomVer(FFortniteNCBranchObjectVersion::GUID) >= FFortniteNCBranchObjectVersion::AddShapeIsProbe)
+		{
+			int8 IsProbe = bIsProbe;
+			Ar << IsProbe;
+			bIsProbe = IsProbe;
+		}
 	}
 };
 
diff --git a/Engine/Source/Runtime/Experimental/ChaosCore/Public/Chaos/Core.h b/Engine/Source/Runtime/Experimental/ChaosCore/Public/Chaos/Core.h
--- a/Engine/Source/Runtime/Experimental/ChaosCore/Public/Chaos/Core.h
+++ b/Engine/Source/Runtime/Experimental/ChaosCore/Public/Chaos/Core.h
@@ -51,7 +51,7 @@ namespace Chaos
 	// NOTE: if you get a merge conflict on the GUID, you must replace it with a new GUID - do not accept the source or target
 	// or you will likely get DDC version conflicts resulting in crashes during load.
 	// Core version string for Chaos data. Any DDC builder dependent on Chaos for serialization should depend on this version
-	static const TCHAR* const ChaosVersionGUID = TEXT("5453DC22-C449-42C0-A74C-4662A08B13E8");
+	static const TCHAR* const ChaosVersionGUID = TEXT("747bf2ef-49be-4f73-b211-74f68a243085");
 
 	inline FString GetChaosVersionStringInner()
 	{
diff --git a/Engine/Source/Runtime/PhysicsCore/Private/ChaosEngineInterface.cpp b/Engine/Source/Runtime/PhysicsCore/Private/ChaosEngineInterface.cpp
--- a/Engine/Source/Runtime/PhysicsCore/Private/ChaosEngineInterface.cpp
+++ b/Engine/Source/Runtime/PhysicsCore/Private/ChaosEngineInterface.cpp
@@ -914,6 +914,11 @@ void FChaosEngineInterface::SetIsSimulationShape(const FPhysicsShapeHandle& InSh
 	InShape.Shape->SetSimEnabled(bIsSimShape);
 }
 
+void FChaosEngineInterface::SetIsProbeShape(const FPhysicsShapeHandle& InShape, bool bIsProbeShape)
+{
+	InShape.Shape->SetIsProbe(bIsProbeShape);
+}
+
 void FChaosEngineInterface::SetIsQueryShape(const FPhysicsShapeHandle& InShape,bool bIsQueryShape)
 {
 	InShape.Shape->SetQueryEnabled(bIsQueryShape);
diff --git a/Engine/Source/Runtime/PhysicsCore/Public/Chaos/ChaosEngineInterface.h b/Engine/Source/Runtime/PhysicsCore/Public/Chaos/ChaosEngineInterface.h
--- a/Engine/Source/Runtime/PhysicsCore/Public/Chaos/ChaosEngineInterface.h
+++ b/Engine/Source/Runtime/PhysicsCore/Public/Chaos/ChaosEngineInterface.h
@@ -537,6 +537,7 @@ public:
 	static void SetSimulationFilter(const FPhysicsShapeHandle& InShape,const FCollisionFilterData& InFilter);
 	static void SetQueryFilter(const FPhysicsShapeHandle& InShape,const FCollisionFilterData& InFilter);
 	static void SetIsSimulationShape(const FPhysicsShapeHandle& InShape,bool bIsSimShape);
+	static void SetIsProbeShape(const FPhysicsShapeHandle& InShape,bool bIsProbeShape);
 	static void SetIsQueryShape(const FPhysicsShapeHandle& InShape,bool bIsQueryShape);
 	static void SetLocalTransform(const FPhysicsShapeHandle& InShape,const FTransform& NewLocalTransform);
 };
diff --git a/Engine/Source/Runtime/PhysicsCore/Public/PhysicsInterfaceTypesCore.h b/Engine/Source/Runtime/PhysicsCore/Public/PhysicsInterfaceTypesCore.h
--- a/Engine/Source/Runtime/PhysicsCore/Public/PhysicsInterfaceTypesCore.h
+++ b/Engine/Source/Runtime/PhysicsCore/Public/PhysicsInterfaceTypesCore.h
@@ -69,12 +69,14 @@ struct FBodyCollisionFlags
 		: bEnableSimCollisionSimple(false)
 		, bEnableSimCollisionComplex(false)
 		, bEnableQueryCollision(false)
+		, bEnableProbeCollision(false)
 	{
 	}
 
 	bool bEnableSimCollisionSimple;
 	bool bEnableSimCollisionComplex;
 	bool bEnableQueryCollision;
+	bool bEnableProbeCollision;
 };
 
 
