UI with miipmap texture groups:
UI Preview Image, UI Icon

diff --git a/Engine/Source/Runtime/Engine/Classes/Engine/TextureDefines.h b/Engine/Source/Runtime/Engine/Classes/Engine/TextureDefines.h
--- a/Engine/Source/Runtime/Engine/Classes/Engine/TextureDefines.h
+++ b/Engine/Source/Runtime/Engine/Classes/Engine/TextureDefines.h
@@ -87,6 +87,8 @@ enum TextureGroup
 	TEXTUREGROUP_Project14 UMETA(DisplayName="ini:Project Group 14"),
 	TEXTUREGROUP_Project15 UMETA(DisplayName="ini:Project Group 15"),
 	TEXTUREGROUP_Project16 UMETA(DisplayName="ini:Project Group 16"),
+	TEXTUREGROUP_Project17 UMETA(DisplayName="ini:Project Group 17"),
+	TEXTUREGROUP_Project18 UMETA(DisplayName="ini:Project Group 18"),
 	TEXTUREGROUP_MAX,
 };
 
