Fixed issue with joining too early if MMs is disabled.

diff --git a/Engine/Source/Runtime/Engine/Classes/Engine/PendingNetGame.h b/Engine/Source/Runtime/Engine/Classes/Engine/PendingNetGame.h
--- a/Engine/Source/Runtime/Engine/Classes/Engine/PendingNetGame.h
+++ b/Engine/Source/Runtime/Engine/Classes/Engine/PendingNetGame.h
@@ -63,6 +63,9 @@ public:
 	 */
 	ENGINE_API void SetEncryptionKey(const FEncryptionKeyResponse& Response);
 
+	bool HasFailedTravel() const {return bFailedTravel; }
+	void SetFailedTravel(bool bInFailedTravel) { bFailedTravel = bInFailedTravel; }
+
 public:
 	/** URL associated with this level. */
 	FURL					URL;
@@ -75,7 +78,10 @@ public:
 
 	/** set when we call LoadMapCompleted */
 	bool					bLoadedMapSuccessfully;
-
+private:
+	/** initialized to true, delaytravel steps can set this to false to indicate error during pendingnetgame travel */
+	bool					bFailedTravel;
+public:
 	/** @todo document */
 	FString					ConnectionError;
 
diff --git a/Engine/Source/Runtime/Engine/Private/PendingNetGame.cpp b/Engine/Source/Runtime/Engine/Private/PendingNetGame.cpp
--- a/Engine/Source/Runtime/Engine/Private/PendingNetGame.cpp
+++ b/Engine/Source/Runtime/Engine/Private/PendingNetGame.cpp
@@ -23,6 +23,7 @@ void UPendingNetGame::Initialize(const FURL& InURL)
 	bSuccessfullyConnected = false;
 	bSentJoinRequest = false;
 	bLoadedMapSuccessfully = false;
+	bFailedTravel = false;
 }
 
 UPendingNetGame::UPendingNetGame(const FObjectInitializer& ObjectInitializer)
diff --git a/Engine/Source/Runtime/Engine/Private/UnrealEngine.cpp b/Engine/Source/Runtime/Engine/Private/UnrealEngine.cpp
--- a/Engine/Source/Runtime/Engine/Private/UnrealEngine.cpp
+++ b/Engine/Source/Runtime/Engine/Private/UnrealEngine.cpp
@@ -14092,7 +14092,12 @@ void UEngine::TickWorldTravel(FWorldContext& Context, float DeltaSeconds)
 		}
 		else if (Context.PendingNetGame && Context.PendingNetGame->bSuccessfullyConnected && !Context.PendingNetGame->bSentJoinRequest && !Context.PendingNetGame->bLoadedMapSuccessfully && (Context.OwningGameInstance == NULL || !Context.OwningGameInstance->DelayPendingNetGameTravel()))
 		{
-			if (!MakeSureMapNameIsValid(Context.PendingNetGame->URL.Map))
+			if (Context.PendingNetGame->HasFailedTravel())
+			{
+				BrowseToDefaultMap(Context);
+				BroadcastTravelFailure(Context.World(), ETravelFailure::TravelFailure, TEXT("Travel failed for unknown reason"));
+			}
+			else if (!MakeSureMapNameIsValid(Context.PendingNetGame->URL.Map))
 			{
 				BrowseToDefaultMap(Context);
 				BroadcastTravelFailure(Context.World(), ETravelFailure::PackageMissing, Context.PendingNetGame->URL.Map);
@@ -14109,7 +14114,7 @@ void UEngine::TickWorldTravel(FWorldContext& Context, float DeltaSeconds)
 					if (!Context.PendingNetGame->LoadMapCompleted(this, Context, bLoadedMapSuccessfully, Error))
 					{
 						BrowseToDefaultMap(Context);
-						BroadcastTravelFailure(Context.World(), ETravelFailure::TravelFailure, Error);
+						BroadcastTravelFailure(Context.World(), ETravelFailure::LoadMapFailure, Error);
 					}
 				}
 				else
@@ -14122,8 +14127,17 @@ void UEngine::TickWorldTravel(FWorldContext& Context, float DeltaSeconds)
 		
 		if (Context.PendingNetGame && Context.PendingNetGame->bLoadedMapSuccessfully && (Context.OwningGameInstance == NULL || !Context.OwningGameInstance->DelayCompletionOfPendingNetGameTravel()))
 		{
-			Context.PendingNetGame->TravelCompleted(this, Context);
-			Context.PendingNetGame = nullptr;
+			if (!Context.PendingNetGame->HasFailedTravel() )
+			{
+				Context.PendingNetGame->TravelCompleted(this, Context);
+				Context.PendingNetGame = nullptr;
+			}
+			else
+			{
+				CancelPending(Context);
+				BrowseToDefaultMap(Context);
+				BroadcastTravelFailure(Context.World(), ETravelFailure::LoadMapFailure, TEXT("Travel failed for unknown reason"));
+			}
 		}
 	}
 	else if (TransitionType == ETransitionType::WaitingToConnect)
