World Partition: remove not implemented function in actor descriptor view.

#rnx

diff --git a/Engine/Source/Runtime/Engine/Public/WorldPartition/WorldPartitionActorDescView.h b/Engine/Source/Runtime/Engine/Public/WorldPartition/WorldPartitionActorDescView.h
--- a/Engine/Source/Runtime/Engine/Public/WorldPartition/WorldPartitionActorDescView.h
+++ b/Engine/Source/Runtime/Engine/Public/WorldPartition/WorldPartitionActorDescView.h
@@ -59,8 +59,6 @@ public:
 
 	bool IsResaveNeeded() const;
 
-	void ResolveRuntimeDataLayers(const UActorDescContainer* InContainer);
-
 	bool operator==(const FWorldPartitionActorDescView& Other) const
 	{
 		return GetGuid() == Other.GetGuid();
