World Partition Blueprint Library:
- Added functions to retrieve the editor and runtime bounds.
- Added functions to load and unload actors descriptors.
- Added class FLoaderAdapterActorList to share code between pinned and BP loaded actors.

#rnx

diff --git a/Engine/Source/Editor/WorldPartitionEditor/Private/WorldPartition/WorldPartitionBlueprintLibrary.cpp b/Engine/Source/Editor/WorldPartitionEditor/Private/WorldPartition/WorldPartitionBlueprintLibrary.cpp
--- a/Engine/Source/Editor/WorldPartitionEditor/Private/WorldPartition/WorldPartitionBlueprintLibrary.cpp
+++ b/Engine/Source/Editor/WorldPartitionEditor/Private/WorldPartition/WorldPartitionBlueprintLibrary.cpp
@@ -32,6 +32,9 @@ FActorDesc::FActorDesc(const FWorldPartitionActorDesc& InActorDesc, const FTrans
 	bActorIsEditorOnly = InActorDesc.GetActorIsEditorOnly();	
 }
 
+TMap<UWorldPartition*, TUniquePtr<FLoaderAdapterActorList>> UWorldPartitionBlueprintLibrary::LoaderAdapterActorListMap;
+FDelegateHandle UWorldPartitionBlueprintLibrary::OnWorldPartitionUninitializedHandle;
+
 UWorld* UWorldPartitionBlueprintLibrary::GetEditorWorld()
 {
 	UUnrealEditorSubsystem* UnrealEditorSubsystem = GEditor->GetEditorSubsystem<UUnrealEditorSubsystem>();
@@ -47,6 +50,16 @@ UWorldPartition* UWorldPartitionBlueprintLibrary::GetWorldPartition()
 	return nullptr;
 }
 
+void UWorldPartitionBlueprintLibrary::OnWorldPartitionUninitialized(UWorldPartition* InWorldPartition)
+{
+	verify(LoaderAdapterActorListMap.Remove(InWorldPartition));
+
+	if (LoaderAdapterActorListMap.IsEmpty())
+	{
+		InWorldPartition->GetWorld()->OnWorldPartitionUninitialized().Remove(OnWorldPartitionUninitializedHandle);
+	}
+}
+
 bool UWorldPartitionBlueprintLibrary::GetActorDescs(const UActorDescContainer* InContainer, const FTransform& InTransform, TArray<FActorDesc>& OutActorDescs)
 {
 	bool bResult = true;
@@ -129,6 +142,50 @@ bool UWorldPartitionBlueprintLibrary::GetIntersectingActorDescs(const UActorDesc
 	return bResult;
 }
 
+FBox UWorldPartitionBlueprintLibrary::GetEditorWorldBounds()
+{
+	if (UWorldPartition* WorldPartition = GetWorldPartition())
+	{
+		return WorldPartition->GetEditorWorldBounds();
+	}
+	return FBox(ForceInit);
+}
+
+FBox UWorldPartitionBlueprintLibrary::GetRuntimeWorldBounds()
+{
+	if (UWorldPartition* WorldPartition = GetWorldPartition())
+	{
+		return WorldPartition->GetRuntimeWorldBounds();
+	}
+	return FBox(ForceInit);
+}
+
+void UWorldPartitionBlueprintLibrary::LoadActors(const TArray<FGuid>& InActorsToLoad)
+{
+	if (UWorldPartition* WorldPartition = GetWorldPartition())
+	{
+		if (LoaderAdapterActorListMap.IsEmpty())
+		{
+			OnWorldPartitionUninitializedHandle = WorldPartition->GetWorld()->OnWorldPartitionUninitialized().AddStatic(&UWorldPartitionBlueprintLibrary::OnWorldPartitionUninitialized);
+		}
+
+		TUniquePtr<FLoaderAdapterActorList>& LoaderAdapterActorList = LoaderAdapterActorListMap.FindOrAdd(WorldPartition, MakeUnique<FLoaderAdapterActorList>(WorldPartition->GetWorld()));
+
+		LoaderAdapterActorList->AddActors(InActorsToLoad);
+	}
+}
+
+void UWorldPartitionBlueprintLibrary::UnloadActors(const TArray<FGuid>& InActorsToLoad)
+{
+	if (UWorldPartition* WorldPartition = GetWorldPartition())
+	{
+		if (TUniquePtr<FLoaderAdapterActorList>* LoaderAdapterActorList = LoaderAdapterActorListMap.Find(WorldPartition))
+		{
+			(*LoaderAdapterActorList)->RemoveActors(InActorsToLoad);
+		}
+	}
+}
+
 bool UWorldPartitionBlueprintLibrary::GetActorDescs(TArray<FActorDesc>& OutActorDescs)
 {
 	if (UWorldPartition* WorldPartition = GetWorldPartition())
diff --git a/Engine/Source/Editor/WorldPartitionEditor/Public/WorldPartition/WorldPartitionBlueprintLibrary.h b/Engine/Source/Editor/WorldPartitionEditor/Public/WorldPartition/WorldPartitionBlueprintLibrary.h
--- a/Engine/Source/Editor/WorldPartitionEditor/Public/WorldPartition/WorldPartitionBlueprintLibrary.h
+++ b/Engine/Source/Editor/WorldPartitionEditor/Public/WorldPartition/WorldPartitionBlueprintLibrary.h
@@ -5,6 +5,7 @@
 #include "Kismet/BlueprintFunctionLibrary.h"
 #include "WorldPartition/ActorDescContainer.h"
 #include "WorldPartition/WorldPartitionActorDesc.h"
+#include "WorldPartition/LoaderAdapter/LoaderAdapterActorList.h"
 #include "WorldPartitionBlueprintLibrary.generated.h"
 
 class FWorldPartitionActorDesc;
@@ -62,10 +63,40 @@ class UWorldPartitionBlueprintLibrary : public UBlueprintFunctionLibrary
 	static UWorld* GetEditorWorld();
 	static UWorldPartition* GetWorldPartition();
 
+	static void OnWorldPartitionUninitialized(UWorldPartition* InWorldPartition);
+	static TMap<UWorldPartition*, TUniquePtr<FLoaderAdapterActorList>> LoaderAdapterActorListMap;
+	static FDelegateHandle OnWorldPartitionUninitializedHandle;
+
 	static bool GetActorDescs(const UActorDescContainer* InContainer, const FTransform& InTransform, TArray<FActorDesc>& OutActorDescs);
 	static bool GetIntersectingActorDescs(const UActorDescContainer* InContainer, const FBox& InBox, const FTransform& InTransform, TArray<FActorDesc>& OutActorDescs);
 
 public:
+	/**
+	 * Gets the editor world bounds, which includes all actor descriptors.
+	 * @return The editor world bounds.
+	 */
+	UFUNCTION(BlueprintCallable, Category="World Partition")
+	static FBox GetEditorWorldBounds();
+
+	/**
+	 * Gets the runtime world bounds, which only includes actor descriptors that aren't editor only.
+	 * @return The runtime world bounds.
+	 */
+	UFUNCTION(BlueprintCallable, Category="World Partition")
+	static FBox GetRuntimeWorldBounds();
+
+	/**
+	 * Load actors
+	 */
+	UFUNCTION(BlueprintCallable, Category="World Partition")
+	static void LoadActors(const TArray<FGuid>& InActorsToLoad);
+
+	/**
+	 * Unload actors
+	 */
+	UFUNCTION(BlueprintCallable, Category="World Partition")
+	static void UnloadActors(const TArray<FGuid>& InActorsToLoad);
+
 	/**
 	 * Gets all the actor descriptors into the provided array, recursing into actor containers.
 	 * @return True if the operation was successful.
diff --git a/Engine/Source/Runtime/Engine/Private/WorldPartition/LoaderAdapter/LoaderAdapterActorList.cpp b/Engine/Source/Runtime/Engine/Private/WorldPartition/LoaderAdapter/LoaderAdapterActorList.cpp
new file mode 100644
--- /dev/null
+++ b/Engine/Source/Runtime/Engine/Private/WorldPartition/LoaderAdapter/LoaderAdapterActorList.cpp
@@ -0,0 +1,90 @@
+// Copyright Epic Games, Inc. All Rights Reserved.
+
+#include "WorldPartition/LoaderAdapter/LoaderAdapterActorList.h"
+#include "WorldPartition/WorldPartition.h"
+#include "Engine/World.h"
+
+#if WITH_EDITOR
+FLoaderAdapterActorList::FLoaderAdapterActorList(UWorld* InWorld)
+	: FLoaderAdapterList(InWorld)
+{
+	Load();
+}
+
+void FLoaderAdapterActorList::AddActors(const TArray<FGuid>& ActorGuids)
+{
+	if (UWorldPartition* WorldPartition = GetWorld()->GetWorldPartition())
+	{
+		TArray<FWorldPartitionHandle> ActorHandles;
+		ActorHandles.Reserve(ActorGuids.Num());
+
+		for (const FGuid& ActorGuid : ActorGuids)
+		{
+			FWorldPartitionHandle ActorHandle(WorldPartition, ActorGuid);
+			if (ActorHandle.IsValid())
+			{
+				ActorHandles.Add(ActorHandle);
+			}
+		}
+
+		AddActors(ActorHandles);
+	}
+}
+
+void FLoaderAdapterActorList::AddActors(const TArray<FWorldPartitionHandle>& ActorHandles)
+{
+	Actors.Append(ActorHandles);
+	RefreshLoadedState();
+}
+
+void FLoaderAdapterActorList::RemoveActors(const TArray<FGuid>& ActorGuids)
+{
+	if (UWorldPartition* WorldPartition = GetWorld()->GetWorldPartition())
+	{
+		TArray<FWorldPartitionHandle> ActorHandles;
+		ActorHandles.Reserve(ActorGuids.Num());
+
+		for (const FGuid& ActorGuid : ActorGuids)
+		{
+			FWorldPartitionHandle ActorHandle(WorldPartition, ActorGuid);
+			if (ActorHandle.IsValid())
+			{
+				ActorHandles.Add(ActorHandle);
+			}
+		}
+
+		RemoveActors(ActorHandles);
+	}
+}
+
+void FLoaderAdapterActorList::RemoveActors(const TArray<FWorldPartitionHandle>& ActorHandles)
+{
+	ActorsToRemove = TSet<FWorldPartitionHandle>(ActorHandles);
+	RefreshLoadedState();
+	Actors = Actors.Difference(ActorsToRemove);
+	ActorsToRemove.Empty();
+}
+
+bool FLoaderAdapterActorList::ContainsActor(const FGuid& ActorGuid) const
+{
+	if (UWorldPartition* WorldPartition = GetWorld()->GetWorldPartition())
+	{
+		FWorldPartitionHandle ActorHandle(WorldPartition, ActorGuid);
+		if(ActorHandle.IsValid())
+		{
+			return ContainsActor(ActorHandle);
+		}
+	}
+	return false;
+}
+
+bool FLoaderAdapterActorList::ContainsActor(const FWorldPartitionHandle& ActorHandle) const
+{
+	return Actors.Contains(ActorHandle);
+}
+
+bool FLoaderAdapterActorList::ShouldActorBeLoaded(const FWorldPartitionHandle& ActorHandle) const
+{
+	return ActorHandle.IsValid() && !ActorsToRemove.Contains(ActorHandle) && PassDataLayersFilter(ActorHandle);
+}
+#endif
diff --git a/Engine/Source/Runtime/Engine/Private/WorldPartition/WorldPartition.cpp b/Engine/Source/Runtime/Engine/Private/WorldPartition/WorldPartition.cpp
--- a/Engine/Source/Runtime/Engine/Private/WorldPartition/WorldPartition.cpp
+++ b/Engine/Source/Runtime/Engine/Private/WorldPartition/WorldPartition.cpp
@@ -46,7 +46,7 @@
 #include "WorldPartition/DataLayer/DataLayerInstance.h"
 #include "WorldPartition/DataLayer/DataLayerSubsystem.h"
 #include "WorldPartition/DataLayer/WorldDataLayers.h"
-#include "WorldPartition/LoaderAdapter/LoaderAdapterList.h"
+#include "WorldPartition/LoaderAdapter/LoaderAdapterActorList.h"
 #include "WorldPartition/LoaderAdapter/LoaderAdapterShape.h"
 #include "WorldPartition/WorldPartitionActorDescViewProxy.h"
 #include "WorldPartition/HLOD/HLODLayer.h"
@@ -200,46 +200,6 @@ public:
 		bIncludeNonSpatiallyLoadedActors = true;
 	}
 };
-
-class FLoaderAdapterPinnedActors: public FLoaderAdapterList
-{
-public:
-	FLoaderAdapterPinnedActors(UWorld* InWorld)
-		: FLoaderAdapterList(InWorld)
-	{
-		Load();
-	}
-
-	void AddActors(const TArray<FWorldPartitionHandle>& ActorHandles)
-	{
-		Actors.Append(ActorHandles);
-		RefreshLoadedState();
-	}
-
-	void RemoveActors(const TArray<FWorldPartitionHandle>& ActorHandles)
-	{
-		ActorsToRemove = TSet<FWorldPartitionHandle>(ActorHandles);
-		RefreshLoadedState();
-		Actors = Actors.Difference(ActorsToRemove);
-		ActorsToRemove.Empty();
-	}
-
-	bool ContainsActor(const FWorldPartitionHandle& ActorHandle) const
-	{
-		return Actors.Contains(ActorHandle);
-	}
-
-protected:
-	//~ Begin IWorldPartitionActorLoaderInterface::ILoaderAdapterList interface
-	virtual bool ShouldActorBeLoaded(const FWorldPartitionHandle& ActorHandle) const override
-	{
-		return ActorHandle.IsValid() && !ActorsToRemove.Contains(ActorHandle) && PassDataLayersFilter(ActorHandle);
-	}
-	//~ End IWorldPartitionActorLoaderInterface::ILoaderAdapterList interface
-
-private:
-	TSet<FWorldPartitionHandle> ActorsToRemove;
-};
 #endif
 
 UWorldPartition::UWorldPartition(const FObjectInitializer& ObjectInitializer)
@@ -442,7 +402,7 @@ void UWorldPartition::Initialize(UWorld* InWorld, const FTransform& InTransform)
 		EditorHash->Initialize();
 
 		AlwaysLoadedActors = new FLoaderAdapterAlwaysLoadedActors(OuterWorld);
-		PinnedActors = new FLoaderAdapterPinnedActors(OuterWorld);
+		PinnedActors = new FLoaderAdapterActorList(OuterWorld);
 	}
 
 	check(RuntimeHash);
@@ -1358,19 +1318,7 @@ void UWorldPartition::PinActors(const TArray<FGuid>& ActorGuids)
 {
 	if (PinnedActors)
 	{
-		TArray<FWorldPartitionHandle> ActorHandles;
-		ActorHandles.Reserve(ActorGuids.Num());
-
-		for (const FGuid& ActorGuid : ActorGuids)
-		{
-			FWorldPartitionHandle ActorHandle(this, ActorGuid);
-			if (ActorHandle.IsValid())
-			{
-				ActorHandles.Add(ActorHandle);
-			}
-		}
-
-		PinnedActors->AddActors(ActorHandles);
+		PinnedActors->AddActors(ActorGuids);
 	}
 }
 
@@ -1378,19 +1326,7 @@ void UWorldPartition::UnpinActors(const TArray<FGuid>& ActorGuids)
 {
 	if (PinnedActors)
 	{
-		TArray<FWorldPartitionHandle> ActorHandles;
-		ActorHandles.Reserve(ActorGuids.Num());
-
-		for (const FGuid& ActorGuid : ActorGuids)
-		{
-			FWorldPartitionHandle ActorHandle(this, ActorGuid);
-			if (ActorHandle.IsValid())
-			{
-				ActorHandles.Add(ActorHandle);
-			}
-		}
-
-		PinnedActors->RemoveActors(ActorHandles);
+		PinnedActors->RemoveActors(ActorGuids);
 	}
 }
 
@@ -1398,11 +1334,7 @@ bool UWorldPartition::IsActorPinned(const FGuid& ActorGuid) const
 {
 	if (PinnedActors)
 	{
-		FWorldPartitionHandle ActorHandle(const_cast<UWorldPartition*>(this), ActorGuid);
-		if(ActorHandle.IsValid())
-		{
-			return PinnedActors->ContainsActor(ActorHandle);
-		}
+		return PinnedActors->ContainsActor(ActorGuid);
 	}
 	return false;
 }
diff --git a/Engine/Source/Runtime/Engine/Public/WorldPartition/LoaderAdapter/LoaderAdapterActorList.h b/Engine/Source/Runtime/Engine/Public/WorldPartition/LoaderAdapter/LoaderAdapterActorList.h
new file mode 100644
--- /dev/null
+++ b/Engine/Source/Runtime/Engine/Public/WorldPartition/LoaderAdapter/LoaderAdapterActorList.h
@@ -0,0 +1,31 @@
+// Copyright Epic Games, Inc. All Rights Reserved.
+#pragma once
+
+#if WITH_EDITOR
+#include "CoreMinimal.h"
+#include "WorldPartition/LoaderAdapter/LoaderAdapterList.h"
+
+/** Loader adapter that contains an actor list. */
+class ENGINE_API FLoaderAdapterActorList: public FLoaderAdapterList
+{
+public:
+	FLoaderAdapterActorList(UWorld* InWorld);
+
+	void AddActors(const TArray<FGuid>& ActorGuids);
+	void AddActors(const TArray<FWorldPartitionHandle>& ActorHandles);
+
+	void RemoveActors(const TArray<FGuid>& ActorGuids);
+	void RemoveActors(const TArray<FWorldPartitionHandle>& ActorHandles);
+
+	bool ContainsActor(const FGuid& ActorGuid) const;
+	bool ContainsActor(const FWorldPartitionHandle& ActorHandle) const;
+
+protected:
+	//~ Begin IWorldPartitionActorLoaderInterface::ILoaderAdapterList interface
+	virtual bool ShouldActorBeLoaded(const FWorldPartitionHandle& ActorHandle) const override;
+	//~ End IWorldPartitionActorLoaderInterface::ILoaderAdapterList interface
+
+private:
+	TSet<FWorldPartitionHandle> ActorsToRemove;
+};
+#endif
\ No newline at end of file
diff --git a/Engine/Source/Runtime/Engine/Public/WorldPartition/WorldPartition.h b/Engine/Source/Runtime/Engine/Public/WorldPartition/WorldPartition.h
--- a/Engine/Source/Runtime/Engine/Public/WorldPartition/WorldPartition.h
+++ b/Engine/Source/Runtime/Engine/Public/WorldPartition/WorldPartition.h
@@ -29,7 +29,7 @@ class UWorldPartitionRuntimeHash;
 class UWorldPartitionStreamingPolicy;
 class IStreamingGenerationErrorHandler;
 class FLoaderAdapterAlwaysLoadedActors;
-class FLoaderAdapterPinnedActors;
+class FLoaderAdapterActorList;
 class FHLODActorDesc;
 class UHLODLayer;
 class UCanvas;
@@ -214,7 +214,7 @@ public:
 	TObjectPtr<UWorldPartitionEditorHash> EditorHash;
 
 	FLoaderAdapterAlwaysLoadedActors* AlwaysLoadedActors;
-	FLoaderAdapterPinnedActors* PinnedActors;
+	FLoaderAdapterActorList* PinnedActors;
 
 	IWorldPartitionEditor* WorldPartitionEditor;
 
