Data Layers: fix resolving data layers for very old maps that weren't resaved since recent changes.

#rnx

diff --git a/Engine/Source/Runtime/Engine/Private/WorldPartition/DataLayer/DataLayerUtils.cpp b/Engine/Source/Runtime/Engine/Private/WorldPartition/DataLayer/DataLayerUtils.cpp
--- a/Engine/Source/Runtime/Engine/Private/WorldPartition/DataLayer/DataLayerUtils.cpp
+++ b/Engine/Source/Runtime/Engine/Private/WorldPartition/DataLayer/DataLayerUtils.cpp
@@ -98,28 +98,47 @@ TArray<FName> FDataLayerUtils::ResolvedDataLayerInstanceNames(const FWorldPartit
 // For performance reasons, this function assumes that InActorDesc's DataLayerInstanceNames was already resolved.
 bool FDataLayerUtils::ResolveRuntimeDataLayerInstanceNames(const FWorldPartitionActorDescView& InActorDescView, const FActorDescViewMap& ActorDescViewMap, TArray<FName>& OutRuntimeDataLayerInstanceNames)
 {
-	TArray<const FWorldPartitionActorDescView*> WorldDataLayerViews = ActorDescViewMap.FindByExactNativeClass<AWorldDataLayers>();
-	
-	if (WorldDataLayerViews.Num())
+	UWorld* World = InActorDescView.GetActorDesc()->GetContainer()->GetWorld();
+	const UDataLayerSubsystem* DataLayerSubsystem = World ? World->GetSubsystem<UDataLayerSubsystem>() : nullptr;
+	if (DataLayerSubsystem && DataLayerSubsystem->CanResolveDataLayers())
 	{
-		check(WorldDataLayerViews.Num() == 1);
-		const FWorldPartitionActorDescView* WorldDataLayersActorDescView = WorldDataLayerViews[0];
-	
-		FWorldDataLayersActorDesc* WorldDataLayersActorDesc = (FWorldDataLayersActorDesc*)WorldDataLayersActorDescView->GetActorDesc();
-
 		for (FName DataLayerInstanceName : InActorDescView.GetDataLayerInstanceNames())
 		{
-			if (const FDataLayerInstanceDesc* DataLayerInstanceDesc = WorldDataLayersActorDesc->GetDataLayerInstanceFromInstanceName(DataLayerInstanceName))
+			const UDataLayerInstance* DataLayerInstance = DataLayerSubsystem->GetDataLayerInstance(DataLayerInstanceName);
+			if (DataLayerInstance && DataLayerInstance->IsRuntime())
 			{
-				if (DataLayerInstanceDesc->GetDataLayerType() == EDataLayerType::Runtime)
-				{
-					OutRuntimeDataLayerInstanceNames.Add(DataLayerInstanceName);
-				}
+				OutRuntimeDataLayerInstanceNames.Add(DataLayerInstanceName);
 			}
 		}
 
 		return true;
 	}
+	else
+	{
+		// Fallback on FWorldDataLayersActorDesc
+		TArray<const FWorldPartitionActorDescView*> WorldDataLayerViews = ActorDescViewMap.FindByExactNativeClass<AWorldDataLayers>();
+	
+		if (WorldDataLayerViews.Num())
+		{
+			check(WorldDataLayerViews.Num() == 1);
+			const FWorldPartitionActorDescView* WorldDataLayersActorDescView = WorldDataLayerViews[0];
+	
+			FWorldDataLayersActorDesc* WorldDataLayersActorDesc = (FWorldDataLayersActorDesc*)WorldDataLayersActorDescView->GetActorDesc();
+
+			for (FName DataLayerInstanceName : InActorDescView.GetDataLayerInstanceNames())
+			{
+				if (const FDataLayerInstanceDesc* DataLayerInstanceDesc = WorldDataLayersActorDesc->GetDataLayerInstanceFromInstanceName(DataLayerInstanceName))
+				{
+					if (DataLayerInstanceDesc->GetDataLayerType() == EDataLayerType::Runtime)
+					{
+						OutRuntimeDataLayerInstanceNames.Add(DataLayerInstanceName);
+					}
+				}
+			}
+
+			return true;
+		}
+	}
 
 	return false;
 }
