Rename/move files to avoid CIS issues with editor-only data in commandlets.

diff --git a/Engine/Source/Editor/WorldPartitionEditor/Private/WorldPartition/WorldPartitionBlueprintLibrary.cpp b/Engine/Source/Editor/WorldPartitionEditor/Private/WorldPartition/WorldPartitionBlueprintLibrary.cpp
deleted file mode 100644
diff --git a/Engine/Source/Editor/WorldPartitionEditor/Public/WorldPartition/WorldPartitionBlueprintLibrary.h b/Engine/Source/Editor/WorldPartitionEditor/Public/WorldPartition/WorldPartitionBlueprintLibrary.h
deleted file mode 100644
diff --git a/Engine/Source/Runtime/Engine/Private/WorldPartitionBlueprintLibrary.cpp b/Engine/Source/Runtime/Engine/Private/WorldPartitionBlueprintLibrary.cpp
new file mode 100644
--- /dev/null
+++ b/Engine/Source/Runtime/Engine/Private/WorldPartitionBlueprintLibrary.cpp
@@ -0,0 +1,222 @@
+// Copyright Epic Games, Inc. All Rights Reserved.
+
+#include "WorldPartition/WorldPartitionBlueprintLibrary.h"
+
+#if WITH_EDITOR
+#include "Subsystems/UnrealEditorSubsystem.h"
+#include "WorldPartition/WorldPartition.h"
+#include "WorldPartition/WorldPartitionActorDesc.h"
+#include "WorldPartition/WorldPartitionHelpers.h"
+#include "Editor.h"
+#endif
+
+FActorDesc::FActorDesc()
+	: Bounds(ForceInit)
+{}
+
+#if WITH_EDITOR
+FActorDesc::FActorDesc(const FWorldPartitionActorDesc& InActorDesc, const FTransform& InTransform)
+{
+	Guid = InActorDesc.GetGuid();
+
+	if (!InActorDesc.GetBaseClass().IsNone())
+	{
+		Class = InActorDesc.GetBaseClass();
+	}
+	else
+	{
+		Class = InActorDesc.GetActorNativeClass();
+	}
+
+	Name = InActorDesc.GetActorName();
+	Label = InActorDesc.GetActorLabel();
+	Bounds = InActorDesc.GetBounds().TransformBy(InTransform);
+	RuntimeGrid = InActorDesc.GetRuntimeGrid();
+	bIsSpatiallyLoaded = InActorDesc.GetIsSpatiallyLoaded();
+	bActorIsEditorOnly = InActorDesc.GetActorIsEditorOnly();	
+}
+
+TMap<UWorldPartition*, TUniquePtr<FLoaderAdapterActorList>> UWorldPartitionBlueprintLibrary::LoaderAdapterActorListMap;
+FDelegateHandle UWorldPartitionBlueprintLibrary::OnWorldPartitionUninitializedHandle;
+
+UWorld* UWorldPartitionBlueprintLibrary::GetEditorWorld()
+{
+	UUnrealEditorSubsystem* UnrealEditorSubsystem = GEditor->GetEditorSubsystem<UUnrealEditorSubsystem>();
+	return UnrealEditorSubsystem ? UnrealEditorSubsystem->GetEditorWorld() : nullptr;
+}
+
+UWorldPartition* UWorldPartitionBlueprintLibrary::GetWorldPartition()
+{
+	if (UWorld* World = GetEditorWorld())
+	{
+		return World->GetWorldPartition();
+	}
+	return nullptr;
+}
+
+void UWorldPartitionBlueprintLibrary::OnWorldPartitionUninitialized(UWorldPartition* InWorldPartition)
+{
+	verify(LoaderAdapterActorListMap.Remove(InWorldPartition));
+
+	if (LoaderAdapterActorListMap.IsEmpty())
+	{
+		InWorldPartition->GetWorld()->OnWorldPartitionUninitialized().Remove(OnWorldPartitionUninitializedHandle);
+	}
+}
+
+bool UWorldPartitionBlueprintLibrary::GetActorDescs(const UActorDescContainer* InContainer, const FTransform& InTransform, TArray<FActorDesc>& OutActorDescs)
+{
+	bool bResult = true;
+
+	for (FActorDescList::TConstIterator<> ActorDescIt(InContainer); ActorDescIt; ++ActorDescIt)
+	{
+		const UActorDescContainer* SubContainer;
+		EContainerClusterMode SubClusterMode;
+		FTransform SubTransform;
+
+		if (ActorDescIt->IsContainerInstance())
+		{
+			if (ActorDescIt->GetContainerInstance(SubContainer, SubTransform, SubClusterMode))
+			{
+				bResult &= GetActorDescs(SubContainer, SubTransform * InTransform, OutActorDescs);
+			}
+			else
+			{
+				bResult = false;
+			}
+		}
+		else
+		{
+			OutActorDescs.Emplace(**ActorDescIt, InTransform);
+		}
+	}
+
+	return bResult;
+}
+
+bool UWorldPartitionBlueprintLibrary::GetIntersectingActorDescs(const UActorDescContainer* InContainer, const FBox& InBox, const FTransform& InTransform, TArray<FActorDesc>& OutActorDescs)
+{
+	bool bResult = true;
+
+	auto HandleActorDesc = [&InBox, &InTransform, &OutActorDescs](const FWorldPartitionActorDesc* ActorDesc)
+	{
+		bool bResult = true;
+		const UActorDescContainer* SubContainer;
+		EContainerClusterMode SubClusterMode;
+		FTransform SubTransform;
+
+		if (ActorDesc->IsContainerInstance())
+		{
+			if (ActorDesc->GetContainerInstance(SubContainer, SubTransform, SubClusterMode))
+			{
+				bResult &= GetIntersectingActorDescs(SubContainer, InBox, SubTransform * InTransform, OutActorDescs);
+			}
+			else
+			{
+				bResult = false;
+			}
+		}
+		else
+		{
+			OutActorDescs.Emplace(*ActorDesc, InTransform);
+		}
+
+		return bResult;
+	};
+
+	if (const UWorldPartition* WorldPartition = Cast<UWorldPartition>(InContainer))
+	{
+		FWorldPartitionHelpers::ForEachIntersectingActorDesc(const_cast<UWorldPartition*>(WorldPartition), InBox, [&bResult, &InBox, &InTransform, &OutActorDescs, &HandleActorDesc](const FWorldPartitionActorDesc* ActorDesc)
+		{
+			bResult &= HandleActorDesc(ActorDesc);
+			return true;
+		});
+	}
+	else
+	{
+		for (FActorDescList::TConstIterator<> ActorDescIt(InContainer); ActorDescIt; ++ActorDescIt)
+		{
+			if (ActorDescIt->GetBounds().Intersect(InBox))
+			{
+				bResult &= HandleActorDesc(*ActorDescIt);
+			}
+		}
+	}
+
+	return bResult;
+}
+#endif
+
+FBox UWorldPartitionBlueprintLibrary::GetEditorWorldBounds()
+{
+#if WITH_EDITOR
+	if (UWorldPartition* WorldPartition = GetWorldPartition())
+	{
+		return WorldPartition->GetEditorWorldBounds();
+	}
+#endif
+	return FBox(ForceInit);
+}
+
+FBox UWorldPartitionBlueprintLibrary::GetRuntimeWorldBounds()
+{
+#if WITH_EDITOR
+	if (UWorldPartition* WorldPartition = GetWorldPartition())
+	{
+		return WorldPartition->GetRuntimeWorldBounds();
+	}
+#endif
+	return FBox(ForceInit);
+}
+
+void UWorldPartitionBlueprintLibrary::LoadActors(const TArray<FGuid>& InActorsToLoad)
+{
+#if WITH_EDITOR
+	if (UWorldPartition* WorldPartition = GetWorldPartition())
+	{
+		if (LoaderAdapterActorListMap.IsEmpty())
+		{
+			OnWorldPartitionUninitializedHandle = WorldPartition->GetWorld()->OnWorldPartitionUninitialized().AddStatic(&UWorldPartitionBlueprintLibrary::OnWorldPartitionUninitialized);
+		}
+
+		TUniquePtr<FLoaderAdapterActorList>& LoaderAdapterActorList = LoaderAdapterActorListMap.FindOrAdd(WorldPartition, MakeUnique<FLoaderAdapterActorList>(WorldPartition->GetWorld()));
+
+		LoaderAdapterActorList->AddActors(InActorsToLoad);
+	}
+#endif
+}
+
+void UWorldPartitionBlueprintLibrary::UnloadActors(const TArray<FGuid>& InActorsToLoad)
+{
+#if WITH_EDITOR
+	if (UWorldPartition* WorldPartition = GetWorldPartition())
+	{
+		if (TUniquePtr<FLoaderAdapterActorList>* LoaderAdapterActorList = LoaderAdapterActorListMap.Find(WorldPartition))
+		{
+			(*LoaderAdapterActorList)->RemoveActors(InActorsToLoad);
+		}
+	}
+#endif
+}
+
+bool UWorldPartitionBlueprintLibrary::GetActorDescs(TArray<FActorDesc>& OutActorDescs)
+{
+#if WITH_EDITOR
+	if (UWorldPartition* WorldPartition = GetWorldPartition())
+	{
+		return GetActorDescs(WorldPartition, FTransform::Identity, OutActorDescs);
+	}
+#endif
+	return false;
+}
+
+bool UWorldPartitionBlueprintLibrary::GetIntersectingActorDescs(const FBox& InBox, TArray<FActorDesc>& OutActorDescs)
+{
+#if WITH_EDITOR
+	if (UWorldPartition* WorldPartition = GetWorldPartition())
+	{
+		return GetIntersectingActorDescs(WorldPartition, InBox, FTransform::Identity, OutActorDescs);
+	}
+#endif
+	return false;
+}
\ No newline at end of file
diff --git a/Engine/Source/Runtime/Engine/Public/WorldPartition/WorldPartitionBlueprintLibrary.h b/Engine/Source/Runtime/Engine/Public/WorldPartition/WorldPartitionBlueprintLibrary.h
new file mode 100644
--- /dev/null
+++ b/Engine/Source/Runtime/Engine/Public/WorldPartition/WorldPartitionBlueprintLibrary.h
@@ -0,0 +1,118 @@
+// Copyright Epic Games, Inc. All Rights Reserved.
+
+#pragma once
+
+#include "Kismet/BlueprintFunctionLibrary.h"
+#include "WorldPartition/ActorDescContainer.h"
+#include "WorldPartition/WorldPartitionActorDesc.h"
+#include "WorldPartition/LoaderAdapter/LoaderAdapterActorList.h"
+#include "WorldPartitionBlueprintLibrary.generated.h"
+
+class FWorldPartitionActorDesc;
+
+/**
+ * Snapshot of an actor descriptor, which represents the state of an actor on disk.
+ * The actor may or may not be loaded.
+ */
+USTRUCT(BlueprintType)
+struct FActorDesc
+{
+	GENERATED_USTRUCT_BODY()
+
+	FActorDesc();
+
+#if WITH_EDITOR
+	FActorDesc(const FWorldPartitionActorDesc& InActorDesc, const FTransform& InTransform);
+#endif
+
+	/** The actor GUID of this descriptor. */
+	UPROPERTY(BlueprintReadOnly, VisibleAnywhere, Category=Actor)
+	FGuid Guid;
+
+	/** Actor class, can point to a native or Blueprint class and may be redirected. */
+	UPROPERTY(BlueprintReadOnly, VisibleAnywhere, Category=Actor)
+	FSoftObjectPath Class;
+
+	/** Internal name of the acgor. */
+	UPROPERTY(BlueprintReadOnly, VisibleAnywhere, Category=Actor)
+	FName Name;
+
+	/** Actor's label. */
+	UPROPERTY(BlueprintReadOnly, VisibleAnywhere, Category=Actor)
+	FName Label;
+
+	/** Streaming bounds of this actor. */
+	UPROPERTY(BlueprintReadOnly, VisibleAnywhere, Category=Actor)
+	FBox Bounds;
+
+	/** Actor's target runtime grid. */
+	UPROPERTY(BlueprintReadOnly, VisibleAnywhere, Category=Actor)
+	FName RuntimeGrid;
+
+	/** Actor's streaming state. */
+	UPROPERTY(BlueprintReadOnly, VisibleAnywhere, Category=Actor)
+	bool bIsSpatiallyLoaded;
+
+	/** Actor's editor-only property. */
+	UPROPERTY(BlueprintReadOnly, VisibleAnywhere, Category=Actor)
+	bool bActorIsEditorOnly;
+};
+
+UCLASS(MinimalAPI)
+class UWorldPartitionBlueprintLibrary : public UBlueprintFunctionLibrary
+{
+	GENERATED_BODY()
+
+#if WITH_EDITOR
+	static UWorld* GetEditorWorld();
+	static UWorldPartition* GetWorldPartition();
+
+	static void OnWorldPartitionUninitialized(UWorldPartition* InWorldPartition);
+	static TMap<UWorldPartition*, TUniquePtr<FLoaderAdapterActorList>> LoaderAdapterActorListMap;
+	static FDelegateHandle OnWorldPartitionUninitializedHandle;
+
+	static bool GetActorDescs(const UActorDescContainer* InContainer, const FTransform& InTransform, TArray<FActorDesc>& OutActorDescs);
+	static bool GetIntersectingActorDescs(const UActorDescContainer* InContainer, const FBox& InBox, const FTransform& InTransform, TArray<FActorDesc>& OutActorDescs);
+#endif
+
+public:
+	/**
+	 * Gets the editor world bounds, which includes all actor descriptors.
+	 * @return The editor world bounds.
+	 */
+	UFUNCTION(BlueprintCallable, Category="World Partition")
+	static FBox GetEditorWorldBounds();
+
+	/**
+	 * Gets the runtime world bounds, which only includes actor descriptors that aren't editor only.
+	 * @return The runtime world bounds.
+	 */
+	UFUNCTION(BlueprintCallable, Category="World Partition")
+	static FBox GetRuntimeWorldBounds();
+
+	/**
+	 * Load actors
+	 */
+	UFUNCTION(BlueprintCallable, Category="World Partition")
+	static void LoadActors(const TArray<FGuid>& InActorsToLoad);
+
+	/**
+	 * Unload actors
+	 */
+	UFUNCTION(BlueprintCallable, Category="World Partition")
+	static void UnloadActors(const TArray<FGuid>& InActorsToLoad);
+
+	/**
+	 * Gets all the actor descriptors into the provided array, recursing into actor containers.
+	 * @return True if the operation was successful.
+	 */
+	UFUNCTION(BlueprintCallable, Category="World Partition")
+	static bool GetActorDescs(TArray<FActorDesc>& OutActorDescs);
+
+	/**
+	 * Gets all the actor descriptors intersecting the provided box into the provided array, recursing into actor containers.
+	 * @return True if the operation was successful.
+	 */
+	UFUNCTION(BlueprintCallable, Category="World Partition")
+	static bool GetIntersectingActorDescs(const FBox& InBox, TArray<FActorDesc>& OutActorDescs);
+};
\ No newline at end of file
