Add RemoveAxisBinding and ClearAxisBindings to InputComponent.

diff --git a/Engine/Source/Runtime/Engine/Classes/Components/InputComponent.h b/Engine/Source/Runtime/Engine/Classes/Components/InputComponent.h
--- a/Engine/Source/Runtime/Engine/Classes/Components/InputComponent.h
+++ b/Engine/Source/Runtime/Engine/Classes/Components/InputComponent.h
@@ -914,6 +914,18 @@ public:
 		return AxisBindings.Last();
 	}
 
+	/**
+	* Removes the axis binding with the specified name.
+	*
+	* @param AxisName the name of the axis to remove.
+	*/
+	void RemoveAxisBinding(const FName AxisName);
+
+	/**
+	* Removes all axis bindings.
+	*/
+	void ClearAxisBindings();
+
 	/**
 	 * Indicates that the InputComponent is interested in knowing the Axis value
 	 * (via GetAxisValue) but does not want a delegate function called each frame.
diff --git a/Engine/Source/Runtime/Engine/Private/Components/InputComponent.cpp b/Engine/Source/Runtime/Engine/Private/Components/InputComponent.cpp
--- a/Engine/Source/Runtime/Engine/Private/Components/InputComponent.cpp
+++ b/Engine/Source/Runtime/Engine/Private/Components/InputComponent.cpp
@@ -380,6 +380,23 @@ void UInputComponent::ClearBindingValues()
 	}
 }
 
+void UInputComponent::RemoveAxisBinding(FName AxisName)
+{
+	for (int32 AxisIdx = AxisBindings.Num() - 1; AxisIdx >= 0; --AxisIdx)
+	{
+		const FInputAxisBinding& Binding = AxisBindings[AxisIdx];
+		if (Binding.AxisName == AxisName)
+		{
+			AxisBindings.RemoveAt(AxisIdx, 1, false);
+		}
+	}
+}
+
+void UInputComponent::ClearAxisBindings()
+{
+	AxisBindings.Reset();
+}
+
 /* Deprecated functions (needed for Blueprints)
  *****************************************************************************/
 
