PR #8757: fix BP action-menu category typo for composure export funcs (Contributed by [REDACTED_PERSON])

diff --git a/Engine/Plugins/Compositing/Composure/Source/Composure/Private/MovieScene/MovieSceneComposureExportSectionTemplate.h b/Engine/Plugins/Compositing/Composure/Source/Composure/Private/MovieScene/MovieSceneComposureExportSectionTemplate.h
--- a/Engine/Plugins/Compositing/Composure/Source/Composure/Private/MovieScene/MovieSceneComposureExportSectionTemplate.h
+++ b/Engine/Plugins/Compositing/Composure/Source/Composure/Private/MovieScene/MovieSceneComposureExportSectionTemplate.h
@@ -28,7 +28,7 @@ public:
 	/**
 	 * Initialize the export to capture the specified named buffer visualization targets from a scene capture
 	 */
-	UFUNCTION(BlueprintCallable, Category="Compsure|Export")
+	UFUNCTION(BlueprintCallable, Category="Composure|Export")
 	void ExportSceneCaptureBuffers(ACompositingElement* CompShotElement, USceneCaptureComponent2D* SceneCapture, const TArray<FString>& BuffersToExport);
 
 	/** Implementation function */
diff --git a/Engine/Plugins/Compositing/Composure/Source/Composure/Public/MovieScene/IMovieSceneComposureExportClient.h b/Engine/Plugins/Compositing/Composure/Source/Composure/Public/MovieScene/IMovieSceneComposureExportClient.h
--- a/Engine/Plugins/Compositing/Composure/Source/Composure/Public/MovieScene/IMovieSceneComposureExportClient.h
+++ b/Engine/Plugins/Compositing/Composure/Source/Composure/Public/MovieScene/IMovieSceneComposureExportClient.h
@@ -30,6 +30,6 @@ public:
 	/**
 	 * Initialize this object for export by setting up any of the necessary scene view extensions with the specified initializer.
 	 */
-	UFUNCTION(BlueprintImplementableEvent, Category="Compsure|Export")
+	UFUNCTION(BlueprintImplementableEvent, Category="Composure|Export")
 	void InitializeForExport(UMovieSceneComposureExportInitializer* ExportInitializer);
 };
\ No newline at end of file
diff --git a/Engine/Source/Runtime/Engine/Private/ShaderCompiler/ShaderCompiler.cpp b/Engine/Source/Runtime/Engine/Private/ShaderCompiler/ShaderCompiler.cpp
--- a/Engine/Source/Runtime/Engine/Private/ShaderCompiler/ShaderCompiler.cpp
+++ b/Engine/Source/Runtime/Engine/Private/ShaderCompiler/ShaderCompiler.cpp
@@ -1924,7 +1924,7 @@ void FShaderCompileUtilities::DoReadTaskResults(const TArray<FShaderCommonCompil
 				FShaderCompileJob* SingleJob = CurrentJob->StageJobs[Index];
 				// cannot reissue a single stage of a pipeline job
 				ReadSingleJob(SingleJob, OutputFile);
-				CurrentJob->bFailedRemovingUnused = CurrentJob->bFailedRemovingUnused | SingleJob->Output.bFailedRemovingUnused;
+				CurrentJob->bFailedRemovingUnused = CurrentJob->bFailedRemovingUnused || SingleJob->Output.bFailedRemovingUnused;
 				CurrentJob->bSucceeded = CurrentJob->bSucceeded && SingleJob->bSucceeded;
 			}
 		}
diff --git a/Engine/Source/Runtime/Engine/Public/ShaderCompiler.h b/Engine/Source/Runtime/Engine/Public/ShaderCompiler.h
--- a/Engine/Source/Runtime/Engine/Public/ShaderCompiler.h
+++ b/Engine/Source/Runtime/Engine/Public/ShaderCompiler.h
@@ -153,13 +153,13 @@ protected:
 	friend class FShaderCompilingManager;
 	friend class FShaderPipelineCompileJob;
 
-	FShaderCommonCompileJob(EShaderCompileJobType InType, uint32 InHash, uint32 InId, EShaderCompileJobPriority InPriroity) :
+	FShaderCommonCompileJob(EShaderCompileJobType InType, uint32 InHash, uint32 InId, EShaderCompileJobPriority InPriority) :
 		NumRefs(0),
 		JobIndex(INDEX_NONE),
 		Hash(InHash),
 		Id(InId),
 		Type(InType),
-		Priority(InPriroity),
+		Priority(InPriority),
 		PendingPriority(EShaderCompileJobPriority::None),
 		CurrentWorker(EShaderCompilerWorkerType::None),
 		bFinalized(false),
@@ -171,7 +171,7 @@ protected:
 		bIsGlobalShader(false)
 
 	{
-		check(InPriroity != EShaderCompileJobPriority::None);
+		check(InPriority != EShaderCompileJobPriority::None);
 	}
 
 	virtual ~FShaderCommonCompileJob() {}
