Stop reading shader compile jobs after corrupted a output file has been detected.

#rnx

diff --git a/Engine/Source/Runtime/Engine/Private/ShaderCompiler/ShaderCompiler.cpp b/Engine/Source/Runtime/Engine/Private/ShaderCompiler/ShaderCompiler.cpp
--- a/Engine/Source/Runtime/Engine/Private/ShaderCompiler/ShaderCompiler.cpp
+++ b/Engine/Source/Runtime/Engine/Private/ShaderCompiler/ShaderCompiler.cpp
@@ -1862,13 +1862,15 @@ void FShaderCompileUtilities::DoReadTaskResults(const TArray<FShaderCommonCompil
 			FString Text = FString::Printf(TEXT("ShaderCompileWorker returned %u single jobs, %u expected"), NumJobs, QueuedSingleJobs.Num());
 			ModalErrorOrLog(Text, OutputFile.Tell(), FileSize);
 		}
-
-		for (int32 JobIndex = 0; JobIndex < NumJobs; JobIndex++)
+		else
 		{
-			auto* CurrentJob = QueuedSingleJobs[JobIndex];
-			if (ReadSingleJob(CurrentJob, OutputFile))
+			for (int32 JobIndex = 0; JobIndex < NumJobs; JobIndex++)
 			{
-				ReissueSourceJobs.Add(CurrentJob);
+				auto* CurrentJob = QueuedSingleJobs[JobIndex];
+				if (ReadSingleJob(CurrentJob, OutputFile))
+				{
+					ReissueSourceJobs.Add(CurrentJob);
+				}
 			}
 		}
 	}
@@ -1890,42 +1892,44 @@ void FShaderCompileUtilities::DoReadTaskResults(const TArray<FShaderCommonCompil
 			FString Text = FString::Printf(TEXT("Worker returned %u pipeline jobs, %u expected"), NumJobs, QueuedPipelineJobs.Num());
 			ModalErrorOrLog(Text, OutputFile.Tell(), FileSize);
 		}
-		for (int32 JobIndex = 0; JobIndex < NumJobs; JobIndex++)
+		else
 		{
-			FShaderPipelineCompileJob* CurrentJob = QueuedPipelineJobs[JobIndex];
-
-			FString PipelineName;
-			OutputFile << PipelineName;
-			if (PipelineName != CurrentJob->Key.ShaderPipeline->GetName())
+			for (int32 JobIndex = 0; JobIndex < NumJobs; JobIndex++)
 			{
-				FString Text = FString::Printf(TEXT("Worker returned Pipeline %s, expected %s!"), *PipelineName, CurrentJob->Key.ShaderPipeline->GetName());
-				ModalErrorOrLog(Text, OutputFile.Tell(), FileSize);
-			}
+				FShaderPipelineCompileJob* CurrentJob = QueuedPipelineJobs[JobIndex];
 
-			check(!CurrentJob->bFinalized);
-			CurrentJob->bFinalized = true;
-			CurrentJob->bFailedRemovingUnused = false;
+				FString PipelineName;
+				OutputFile << PipelineName;
+				if (PipelineName != CurrentJob->Key.ShaderPipeline->GetName())
+				{
+					FString Text = FString::Printf(TEXT("Worker returned Pipeline %s, expected %s!"), *PipelineName, CurrentJob->Key.ShaderPipeline->GetName());
+					ModalErrorOrLog(Text, OutputFile.Tell(), FileSize);
+				}
 
-			int32 NumStageJobs = -1;
-			OutputFile << NumStageJobs;
+				check(!CurrentJob->bFinalized);
+				CurrentJob->bFinalized = true;
+				CurrentJob->bFailedRemovingUnused = false;
 
-			if (NumStageJobs != CurrentJob->StageJobs.Num())
-			{
-				if (NumJobs != QueuedPipelineJobs.Num())
+				int32 NumStageJobs = -1;
+				OutputFile << NumStageJobs;
+
+				if (NumStageJobs != CurrentJob->StageJobs.Num())
 				{
 					FString Text = FString::Printf(TEXT("Worker returned %u stage pipeline jobs, %u expected"), NumStageJobs, CurrentJob->StageJobs.Num());
 					ModalErrorOrLog(Text, OutputFile.Tell(), FileSize);
 				}
-			}
-
-			CurrentJob->bSucceeded = true;
-			for (int32 Index = 0; Index < NumStageJobs; Index++)
-			{
-				FShaderCompileJob* SingleJob = CurrentJob->StageJobs[Index];
-				// cannot reissue a single stage of a pipeline job
-				ReadSingleJob(SingleJob, OutputFile);
-				CurrentJob->bFailedRemovingUnused = CurrentJob->bFailedRemovingUnused || SingleJob->Output.bFailedRemovingUnused;
-				CurrentJob->bSucceeded = CurrentJob->bSucceeded && SingleJob->bSucceeded;
+				else
+				{
+					CurrentJob->bSucceeded = true;
+					for (int32 Index = 0; Index < NumStageJobs; Index++)
+					{
+						FShaderCompileJob* SingleJob = CurrentJob->StageJobs[Index];
+						// cannot reissue a single stage of a pipeline job
+						ReadSingleJob(SingleJob, OutputFile);
+						CurrentJob->bFailedRemovingUnused = CurrentJob->bFailedRemovingUnused || SingleJob->Output.bFailedRemovingUnused;
+						CurrentJob->bSucceeded = CurrentJob->bSucceeded && SingleJob->bSucceeded;
+					}
+				}
 			}
 		}
 	}
