Fix for subobject gc references in demo driver

diff --git a/Engine/Source/Runtime/Engine/Classes/Engine/ActorChannel.h b/Engine/Source/Runtime/Engine/Classes/Engine/ActorChannel.h
--- a/Engine/Source/Runtime/Engine/Classes/Engine/ActorChannel.h
+++ b/Engine/Source/Runtime/Engine/Classes/Engine/ActorChannel.h
@@ -158,8 +158,10 @@ public:
 	/** Release any references this channel is holding to UObjects and object replicators and mark it as broken. */
 	void BreakAndReleaseReferences();
 
+	void ReleaseReferences(bool bKeepReplicators);
+
 	/** UActorChannel interface and accessors. */
-	AActor* GetActor() {return Actor;}
+	AActor* GetActor() const { return Actor; }
 
 	/** Replicate this channel's actor differences. Returns how many bits were replicated (does not include non-bunch packet overhead) */
 	int64 ReplicateActor();
diff --git a/Engine/Source/Runtime/Engine/Private/DataChannel.cpp b/Engine/Source/Runtime/Engine/Private/DataChannel.cpp
--- a/Engine/Source/Runtime/Engine/Private/DataChannel.cpp
+++ b/Engine/Source/Runtime/Engine/Private/DataChannel.cpp
@@ -2163,14 +2163,19 @@ void UActorChannel::SetClosingFlag()
 	UChannel::SetClosingFlag();
 }
 
-void UActorChannel::BreakAndReleaseReferences()
+void UActorChannel::ReleaseReferences(bool bKeepReplicators)
 {
 	Actor = nullptr;
-	Broken = true;
-	CleanupReplicators();
+	CleanupReplicators(bKeepReplicators);
 	CreateSubObjects.Empty();
 }
 
+void UActorChannel::BreakAndReleaseReferences()
+{
+	Broken = true;
+	ReleaseReferences(false);
+}
+
 int64 UActorChannel::Close(EChannelCloseReason Reason)
 {
 	FScopedRepContext RepContext(Connection, Actor);
@@ -2211,8 +2216,7 @@ int64 UActorChannel::Close(EChannelCloseReason Reason)
 			Connection->RemoveActorChannel( Actor );
 		}
 
-		Actor = nullptr;
-		CleanupReplicators( bKeepReplicators );
+		ReleaseReferences(bKeepReplicators);
 	}
 
 	return NumBits;
@@ -4832,7 +4836,8 @@ void UActorChannel::AddedToChannelPool()
 	check(QueuedMustBeMappedGuidsInLastBunch.Num() == 0);
 	check(QueuedExportBunches.Num() == 0);
 
-	Actor = nullptr;
+	ReleaseReferences(false);
+
 	ActorNetGUID = FNetworkGUID();
 	CustomTimeDilation = 0;
 	RelevantTime = 0;
@@ -4845,7 +4850,6 @@ void UActorChannel::AddedToChannelPool()
 	bClearRecentActorRefs = true;
 	QueuedBunchStartTime = 0;
 	bSuppressQueuedBunchWarningsDueToHitches = false;
-	CreateSubObjects.Empty();
 #if !UE_BUILD_SHIPPING
 	bBlockChannelFailure = false;
 #endif
diff --git a/Engine/Source/Runtime/Engine/Private/DemoNetDriver.cpp b/Engine/Source/Runtime/Engine/Private/DemoNetDriver.cpp
--- a/Engine/Source/Runtime/Engine/Private/DemoNetDriver.cpp
+++ b/Engine/Source/Runtime/Engine/Private/DemoNetDriver.cpp
@@ -4068,7 +4068,7 @@ bool UDemoNetDriver::LoadCheckpoint(const FGotoResult& GotoResult)
 			UActorChannel* ActorChannel = Cast<UActorChannel>(OpenChannel);
 			if (ActorChannel != nullptr && KeepAliveActors.Contains(ActorChannel->Actor))
 			{
-				ActorChannel->Actor = nullptr;
+				ActorChannel->ReleaseReferences(false);
 			}
 		}
 	}
@@ -5008,8 +5008,7 @@ void UDemoNetDriver::NotifyActorDestroyed(AActor* Actor, bool IsSeamlessTravel)
 					check(Channel->OpenedLocally);
 					Channel->bClearRecentActorRefs = false;
 					Channel->SetClosingFlag();
-					Channel->Actor = nullptr;
-					Channel->CleanupReplicators(false);
+					Channel->ReleaseReferences(false);
 				}
 
 				Connection->DormantReplicatorMap.Remove(Actor);
@@ -5116,7 +5115,7 @@ void UDemoNetDriver::NotifyActorLevelUnloaded(AActor* Actor)
 		if (UActorChannel* ActorChannel = ServerConnection->FindActorChannelRef(Actor))
 		{
 			ServerConnection->RemoveActorChannel(Actor);
-			ActorChannel->Actor = nullptr;
+			ActorChannel->ReleaseReferences(false);
 			ActorChannel->ConditionalCleanUp(false, EChannelCloseReason::LevelUnloaded);
 		}
 	}
diff --git a/Engine/Source/Runtime/Engine/Private/NetDriver.cpp b/Engine/Source/Runtime/Engine/Private/NetDriver.cpp
--- a/Engine/Source/Runtime/Engine/Private/NetDriver.cpp
+++ b/Engine/Source/Runtime/Engine/Private/NetDriver.cpp
@@ -3588,9 +3588,7 @@ void UNetDriver::NotifyStreamingLevelUnload(ULevel* Level)
 			{
 				UE_LOG(LogNet, Log, TEXT("NotifyStreamingLevelUnload: BREAKING"));
 
-				Channel->Actor = NULL;
-				Channel->Broken = true;
-				Channel->CleanupReplicators();
+				Channel->BreakAndReleaseReferences();
 			}
 		}
 
@@ -3795,6 +3793,14 @@ void UNetDriver::PostGarbageCollect()
 		}
 	}
 	PRAGMA_ENABLE_DEPRECATION_WARNINGS
+
+	for (FObjectReplicator* Replicator : AllOwnedReplicators)
+	{
+		if (!Replicator->GetWeakObjectPtr().IsValid())
+		{
+			Replicator->ReleaseStrongReference();
+		}
+	}
 }
 
 void UNetDriver::AddReferencedObjects(UObject* InThis, FReferenceCollector& Collector)
@@ -3813,7 +3819,11 @@ void UNetDriver::AddReferencedObjects(UObject* InThis, FReferenceCollector& Coll
 	
 	for (FObjectReplicator* Replicator : This->AllOwnedReplicators)
 	{
-		Collector.AddReferencedObject(Replicator->ObjectPtr, This);
+		if (Replicator->GetWeakObjectPtr().IsValid())
+		{
+			Collector.AddReferencedObject(Replicator->ObjectPtr, This);
+		}
+
 		Collector.AddReferencedObject(Replicator->ObjectClass, This);
 	}
 
