Enable the deprecation of the old version InputKey on the base UnrealClient

diff --git a/Engine/Source/Runtime/Engine/Public/UnrealClient.h b/Engine/Source/Runtime/Engine/Public/UnrealClient.h
--- a/Engine/Source/Runtime/Engine/Public/UnrealClient.h
+++ b/Engine/Source/Runtime/Engine/Public/UnrealClient.h
@@ -879,10 +879,12 @@ public:
 	 */
 	virtual bool InputKey(const FInputKeyEventArgs& EventArgs)
 	{
+		PRAGMA_DISABLE_DEPRECATION_WARNINGS
 		return InputKey(EventArgs.Viewport, EventArgs.ControllerId, EventArgs.Key, EventArgs.Event, EventArgs.AmountDepressed, EventArgs.Key.IsGamepadKey());
+		PRAGMA_ENABLE_DEPRECATION_WARNINGS
 	}
 
-	//UE_DEPRECATED(4.21, "Use the new InputKey(const FInputKeyEventArgs& EventArgs) function.")
+	UE_DEPRECATED(5.1, "Use the new InputKey(const FInputKeyEventArgs& EventArgs) function.")
 	virtual bool InputKey(FViewport* Viewport,int32 ControllerId,FKey Key,EInputEvent Event,float AmountDepressed = 1.f,bool bGamepad=false) { return false; }
 
 	/**
@@ -897,7 +899,7 @@ public:
 	 * @param	bGamepad - input came from gamepad (ie xbox controller)
 	 * @return	True to consume the axis movement, false to pass it on.
 	 */
-	// TODO: Deprecate this version
+	UE_DEPRECATED(5.1, "Use the new InputAxis function that takes a FInputDeviceId.")
 	virtual bool InputAxis(FViewport* Viewport,int32 ControllerId,FKey Key,float Delta,float DeltaTime,int32 NumSamples=1,bool bGamepad=false) { return false; }
 
 	/**
