Pass in the input device ID instead of the controller id for the scene viewport. This fixes an issue where the input device would be -1 by default and the player controller would then ignore the input

diff --git a/Engine/Source/Runtime/Engine/Private/Slate/SceneViewport.cpp b/Engine/Source/Runtime/Engine/Private/Slate/SceneViewport.cpp
--- a/Engine/Source/Runtime/Engine/Private/Slate/SceneViewport.cpp
+++ b/Engine/Source/Runtime/Engine/Private/Slate/SceneViewport.cpp
@@ -298,8 +298,9 @@ void FSceneViewport::ProcessAccumulatedPointerInput()
 	if (NumMouseSamplesX > 0 || NumMouseSamplesY > 0)
 	{
 		const float DeltaTime = FApp::GetDeltaTime();
-		ViewportClient->InputAxis( this, 0, EKeys::MouseX, MouseDelta.X, DeltaTime, NumMouseSamplesX );
-		ViewportClient->InputAxis( this, 0, EKeys::MouseY, MouseDelta.Y, DeltaTime, NumMouseSamplesY );
+		FInputDeviceId DefaultInputDevice = IPlatformInputDeviceMapper::Get().GetDefaultInputDevice();
+		ViewportClient->InputAxis( this, DefaultInputDevice, EKeys::MouseX, MouseDelta.X, DeltaTime, NumMouseSamplesX );
+		ViewportClient->InputAxis( this, DefaultInputDevice, EKeys::MouseY, MouseDelta.Y, DeltaTime, NumMouseSamplesY );
 	}
 
 	if ( bCursorHiddenDueToCapture )
@@ -776,7 +777,7 @@ FReply FSceneViewport::OnMouseWheel( const FGeometry& InGeometry, const FPointer
 		// Pressed and released should be sent
 		ViewportClient->InputKey(FInputKeyEventArgs(this, InMouseEvent.GetUserIndex(), ViewportClientKey, IE_Pressed, 1.0f, InMouseEvent.IsTouchEvent()));
 		ViewportClient->InputKey(FInputKeyEventArgs(this, InMouseEvent.GetUserIndex(), ViewportClientKey, IE_Released, 1.0f, InMouseEvent.IsTouchEvent()));
-		ViewportClient->InputAxis(this, InMouseEvent.GetUserIndex(), EKeys::MouseWheelAxis, InMouseEvent.GetWheelDelta(), FApp::GetDeltaTime());
+		ViewportClient->InputAxis(this, InMouseEvent.GetInputDeviceId(), EKeys::MouseWheelAxis, InMouseEvent.GetWheelDelta(), FApp::GetDeltaTime());
 	}
 	return CurrentReplyState;
 }
