Fix for alpha channels of non-high res screenshots being set to 0 when processed at runtime with UI requested, causing the scene to not be captured.

diff --git a/Engine/Source/Runtime/Engine/Private/GameViewportClient.cpp b/Engine/Source/Runtime/Engine/Private/GameViewportClient.cpp
--- a/Engine/Source/Runtime/Engine/Private/GameViewportClient.cpp
+++ b/Engine/Source/Runtime/Engine/Private/GameViewportClient.cpp
@@ -1957,10 +1957,19 @@ bool UGameViewportClient::ProcessScreenShots(FViewport* InViewport)
 				// We need to ignore the high resolution screenshot alpha mask in that case, as the mask isn't relevant when taking a
 				// screenshot of the entire window (and it will trigger an assert).  We don't want to solve this by modifying the
 				// global variables associated with the screenshot feature, as the application may be taking its own screenshots.
-				if (!bIsUI)
+				if (GIsHighResScreenshot && !bIsUI)
 				{
 					GetHighResScreenshotConfig().MergeMaskIntoAlpha(Bitmap, FIntRect(0,0,0,0));
 				}
+				else
+				{
+					// Ensure that all pixels' alpha is set to 255 for a regular screenshot regardless of UI settings
+					// (Regular screenshots return with 0 in the alpha channels)
+					for (auto& Color : Bitmap)
+					{
+						Color.A = 255;
+					}
+				}
 
 				FIntRect SourceRect(0, 0, GScreenshotResolutionX, GScreenshotResolutionY);
 				if (GIsHighResScreenshot)
