Fix issue with ReverbEffects not being updated correctly when the last effect is removed out of the Array

diff --git a/Engine/Source/Runtime/Engine/Private/AudioDevice.cpp b/Engine/Source/Runtime/Engine/Private/AudioDevice.cpp
--- a/Engine/Source/Runtime/Engine/Private/AudioDevice.cpp
+++ b/Engine/Source/Runtime/Engine/Private/AudioDevice.cpp
@@ -338,7 +338,6 @@ FAudioDevice::FAudioDevice()
 	, AudioClock(0.0)
 	, bAllowCenterChannel3DPanning(false)
 	, DeviceDeltaTime(0.0f)
-	, bHasActivatedReverb(false)
 	, bAllowPlayWhenSilent(true)
 	, bUseAttenuationForNonGameWorlds(false)
 	, ConcurrencyManager(this)
@@ -2372,7 +2371,6 @@ void FAudioDevice::UpdateHighestPriorityReverb()
 		const FActivatedReverb& NewActiveReverbRef = ActivatedReverbs.CreateConstIterator().Value();
 		FAudioThread::RunCommandOnAudioThread([AudioDevice, NewActiveReverbRef]()
 		{
-			AudioDevice->bHasActivatedReverb = true;
 			AudioDevice->HighestPriorityActivatedReverb = NewActiveReverbRef;
 		}, GET_STATID(STAT_AudioUpdateHighestPriorityReverb));
 	}
@@ -2380,7 +2378,7 @@ void FAudioDevice::UpdateHighestPriorityReverb()
 	{
 		FAudioThread::RunCommandOnAudioThread([AudioDevice]()
 		{
-			AudioDevice->bHasActivatedReverb = false;
+			AudioDevice->HighestPriorityActivatedReverb = FActivatedReverb();
 		}, GET_STATID(STAT_AudioUpdateHighestPriorityReverb));
 	}
 }
@@ -4591,9 +4589,10 @@ void FAudioDevice::UpdateAudioVolumeEffects()
 
 	bool bHasVolumeSettings = false;
 	FAudioVolumeSettings PlayerAudioVolumeSettings;
-	FAudioVolumeSettings PreviousPlayerAudioVolumeSettings;
 	bool bUsingDefaultReverb = true;
-	bool bReverbChanged = false;
+	bool bEnteredOrExitedAudioVolume = false;
+
+	FAudioVolumeSettings PreviousPlayerAudioVolumeSettings = CurrentAudioVolumeSettings;
 
 	// Gets the current state of the interior settings
 	for (FListener& Listener : Listeners)
@@ -4609,7 +4608,7 @@ void FAudioDevice::UpdateAudioVolumeEffects()
 			bHasVolumeSettings = true;
 			PlayerAudioVolumeSettings = NewPlayerAudioVolumeSettings;
 
-			if (NewPlayerAudioVolumeSettings.AudioVolumeID > 0)
+			if (NewPlayerAudioVolumeSettings.AudioVolumeID > 0 && NewPlayerAudioVolumeSettings.ReverbSettings.bApplyReverb)
 			{
 				bUsingDefaultReverb = false;
 			}
@@ -4621,29 +4620,27 @@ void FAudioDevice::UpdateAudioVolumeEffects()
 
 	if (PlayerAudioVolumeSettings.AudioVolumeID != CurrentAudioVolumeSettings.AudioVolumeID)
 	{
-		PreviousPlayerAudioVolumeSettings = CurrentAudioVolumeSettings;
+		// We're inside a new AudioVolume
 		CurrentAudioVolumeSettings = PlayerAudioVolumeSettings;
-		bReverbChanged = true;
+		bEnteredOrExitedAudioVolume = true;
 	}
 	else if (PlayerAudioVolumeSettings.AudioVolumeID == CurrentAudioVolumeSettings.AudioVolumeID && PlayerAudioVolumeSettings.bChanged)
 	{
+		// Our active AudioVolume was updated
 		CurrentAudioVolumeSettings = PlayerAudioVolumeSettings;
-		bReverbChanged = true;
+		bEnteredOrExitedAudioVolume = true;
 	}
 
 	if (Effects)
 	{
-		// Update the master reverb
-		if (bHasActivatedReverb)
+		// Check if we should be using activated reverb
+		if (HighestPriorityActivatedReverb.Priority > PlayerAudioVolumeSettings.Priority || bUsingDefaultReverb)
 		{
-			if (HighestPriorityActivatedReverb.Priority > PlayerAudioVolumeSettings.Priority || bUsingDefaultReverb)
-			{
-				CurrentAudioVolumeSettings.ReverbSettings = HighestPriorityActivatedReverb.ReverbSettings;
-				bReverbChanged = true;
-			}
+			CurrentAudioVolumeSettings.ReverbSettings = HighestPriorityActivatedReverb.ReverbSettings;
 		}
 
-		if (bReverbChanged)
+		// Update the master reverb if it has changed
+		if (CurrentAudioVolumeSettings.bChanged || (CurrentAudioVolumeSettings.ReverbSettings != PreviousPlayerAudioVolumeSettings.ReverbSettings))
 		{
 			Effects->SetReverbSettings(CurrentAudioVolumeSettings.ReverbSettings);
 		}
@@ -4652,7 +4649,7 @@ void FAudioDevice::UpdateAudioVolumeEffects()
 		Effects->Update();
 
 		// If we any submix override settings apply those overrides to the indicated submixes
-		if (IsAudioMixerEnabled() && bReverbChanged)
+		if (IsAudioMixerEnabled() && bEnteredOrExitedAudioVolume)
 		{
 			// Clear out any previous submix effect chain overrides if the audio volume changed
 			if (PreviousPlayerAudioVolumeSettings.SubmixOverrideSettings.Num() > 0)
@@ -6303,7 +6300,6 @@ void FAudioDevice::Flush(UWorld* WorldToFlush, bool bClearActivatedReverb)
 	if (bClearActivatedReverb)
 	{
 		ActivatedReverbs.Reset();
-		bHasActivatedReverb = false;
 	}
 
 	if (WorldToFlush == nullptr)
diff --git a/Engine/Source/Runtime/Engine/Public/AudioDevice.h b/Engine/Source/Runtime/Engine/Public/AudioDevice.h
--- a/Engine/Source/Runtime/Engine/Public/AudioDevice.h
+++ b/Engine/Source/Runtime/Engine/Public/AudioDevice.h
@@ -2162,9 +2162,6 @@ protected:
 
 private:
 
-	/** Whether the value in HighestPriorityActivatedReverb should be used - Audio Thread owned */
-	uint8 bHasActivatedReverb:1;
-
 	/** Whether or not we're supporting zero volume wave instances */
 	uint8 bAllowPlayWhenSilent:1;
 
