Fixed crash caused by SubsystemCollection::RemoveAndDeinitializeSubsystem not properly updating its internal SubsystemArrayMap

diff --git a/Engine/Source/Runtime/Engine/Private/Subsystems/SubsystemCollection.cpp b/Engine/Source/Runtime/Engine/Private/Subsystems/SubsystemCollection.cpp
--- a/Engine/Source/Runtime/Engine/Private/Subsystems/SubsystemCollection.cpp
+++ b/Engine/Source/Runtime/Engine/Private/Subsystems/SubsystemCollection.cpp
@@ -292,8 +292,16 @@ void FSubsystemCollectionBase::RemoveAndDeinitializeSubsystem(USubsystem* Subsys
 	USubsystem* SubsystemFound = SubsystemMap.FindAndRemoveChecked(Subsystem->GetClass());
 	check(Subsystem == SubsystemFound);
 
-	SubsystemMapWeak.Remove(Subsystem->GetClass());
-	SubsystemArrayMap.Remove(Subsystem->GetClass());
+	const UClass* SubsystemClass = Subsystem->GetClass();
+	SubsystemMapWeak.Remove(SubsystemClass);
+
+	for (auto& Pair : SubsystemArrayMap)
+	{
+		if (SubsystemClass->IsChildOf(Pair.Key))
+		{
+			Pair.Value.Remove(Subsystem);
+		}
+	}
 
 	Subsystem->Deinitialize();
 	Subsystem->InternalOwningSubsystem = nullptr;
