Fix submix volume update logic

diff --git a/Engine/Source/Runtime/Engine/Private/SoundSubmix.cpp b/Engine/Source/Runtime/Engine/Private/SoundSubmix.cpp
--- a/Engine/Source/Runtime/Engine/Private/SoundSubmix.cpp
+++ b/Engine/Source/Runtime/Engine/Private/SoundSubmix.cpp
@@ -59,22 +59,46 @@ void USoundSubmix::Serialize(FArchive& Ar)
 #if WITH_EDITORONLY_DATA
 	if (Ar.IsLoading() || Ar.IsSaving())
 	{
+		// use -96dB as a noise floor when fixing up linear volume settings
+		static constexpr float LinearNeg96dB = 0.0000158489319f;
+
 		// convert any old deprecated values to the new value
-		if (OutputVolume > 0.0f)
+		if (OutputVolume >= 0.0f)
 		{
-			OutputVolumeModulation.Value = Audio::ConvertToDecibels(OutputVolume);
+			if (OutputVolume <= LinearNeg96dB)
+			{
+				OutputVolumeModulation.Value = -96.f;
+			}
+			else
+			{
+				OutputVolumeModulation.Value = Audio::ConvertToDecibels(OutputVolume);
+			}
 			OutputVolume = -1.0f;
 		}
 
-		if (WetLevel > 0.0f)
+		if (WetLevel >= 0.0f)
 		{
-			WetLevelModulation.Value = Audio::ConvertToDecibels(WetLevel);
+			if (WetLevel <= LinearNeg96dB)
+			{
+				WetLevelModulation.Value = -96.f;
+			}
+			else
+			{
+				WetLevelModulation.Value = Audio::ConvertToDecibels(WetLevel);
+			}
 			WetLevel = -1.0f;
 		}
 
-		if (DryLevel > 0.0f)
+		if (DryLevel >= 0.0f)
 		{
-			DryLevelModulation.Value = Audio::ConvertToDecibels(DryLevel);
+			if (DryLevel <= LinearNeg96dB)
+			{
+				DryLevelModulation.Value = -96.f;
+			}
+			else
+			{
+				DryLevelModulation.Value = Audio::ConvertToDecibels(DryLevel);
+			}
 			DryLevel = -1.0f;
 		}
 
