Default bEvaluateWorldPositionOffset to on

diff --git a/Engine/Source/Runtime/Engine/Private/Components/StaticMeshComponent.cpp b/Engine/Source/Runtime/Engine/Private/Components/StaticMeshComponent.cpp
--- a/Engine/Source/Runtime/Engine/Private/Components/StaticMeshComponent.cpp
+++ b/Engine/Source/Runtime/Engine/Private/Components/StaticMeshComponent.cpp
@@ -201,7 +201,7 @@ UStaticMeshComponent::UStaticMeshComponent(const FObjectInitializer& ObjectIniti
 	bForceNavigationObstacle = true;
 	bDisallowMeshPaintPerInstance = false;
 	bDisallowNanite = false;
-	bEvaluateWorldPositionOffset = false;
+	bEvaluateWorldPositionOffset = true;
 	bEvaluateWorldPositionOffsetInRayTracing = false;
 	bInitialEvaluateWorldPositionOffset = false;
 	DistanceFieldIndirectShadowMinVisibility = .1f;
@@ -372,7 +372,7 @@ void UStaticMeshComponent::Serialize(FArchive& Ar)
 	Ar.CustomVer(FFortniteMainBranchObjectVersion::GUID) < FFortniteMainBranchObjectVersion::RemappedEvaluateWorldPositionOffsetInRayTracing)
 	{
 		bEvaluateWorldPositionOffsetInRayTracing = bEvaluateWorldPositionOffset;
-		bEvaluateWorldPositionOffset = false; // Default WPO evaluation off
+		bEvaluateWorldPositionOffset = true; // Default WPO evaluation on
 	}
 
 	NotifyIfStaticMeshChanged();
diff --git a/Engine/Source/Runtime/Foliage/Private/InstancedFoliage.cpp b/Engine/Source/Runtime/Foliage/Private/InstancedFoliage.cpp
--- a/Engine/Source/Runtime/Foliage/Private/InstancedFoliage.cpp
+++ b/Engine/Source/Runtime/Foliage/Private/InstancedFoliage.cpp
@@ -543,7 +543,7 @@ UFoliageType::UFoliageType(const FObjectInitializer& ObjectInitializer)
 	OverriddenLightMapRes = 8;
 	bUseAsOccluder = false;
 	bVisibleInRayTracing = true;
-	bEvaluateWorldPositionOffset = false;
+	bEvaluateWorldPositionOffset = true;
 
 	BodyInstance.SetCollisionProfileName(UCollisionProfile::NoCollision_ProfileName);
 
