PR #7577:  Scene Capture, modify OnRegister so editor-only components aren't created when running a commandlet, cut and paste of identical code from CL 11106828.

diff --git a/Engine/Source/Runtime/Engine/Private/Components/SceneCaptureComponent.cpp b/Engine/Source/Runtime/Engine/Private/Components/SceneCaptureComponent.cpp
--- a/Engine/Source/Runtime/Engine/Private/Components/SceneCaptureComponent.cpp
+++ b/Engine/Source/Runtime/Engine/Private/Components/SceneCaptureComponent.cpp
@@ -171,7 +171,8 @@ USceneCaptureComponent::USceneCaptureComponent(const FObjectInitializer& ObjectI
 void USceneCaptureComponent::OnRegister()
 {
 #if WITH_EDITORONLY_DATA
-	if (AActor* MyOwner = GetOwner())
+	AActor* MyOwner = GetOwner();
+	if ((MyOwner != nullptr) && !IsRunningCommandlet())
 	{
 		if (ProxyMeshComponent == nullptr)
 		{
@@ -510,7 +511,8 @@ void USceneCaptureComponent2D::OnRegister()
 	Super::OnRegister();
 
 #if WITH_EDITORONLY_DATA
-	if (AActor* MyOwner = GetOwner())
+	AActor* MyOwner = GetOwner();
+	if ((MyOwner != nullptr) && !IsRunningCommandlet())
 	{
 		if (DrawFrustum == nullptr)
 		{
@@ -1095,7 +1097,8 @@ void USceneCaptureComponentCube::OnRegister()
 	Super::OnRegister();
 
 #if WITH_EDITORONLY_DATA
-	if (AActor* MyOwner = GetOwner())
+	AActor* MyOwner = GetOwner();
+	if ((MyOwner != nullptr) && !IsRunningCommandlet())
 	{
 		if (DrawFrustum == nullptr)
 		{
