Make InputComponent's "NotBlueprintable" because you shouldn't be making an InputComponent as a blueprint. It won't actually be processed by anything if you do, and there are no ways to remove bound mappings so you can end up with delegates bound to things that don't exist

#rnx

diff --git a/Engine/Source/Runtime/Engine/Classes/Components/InputComponent.h b/Engine/Source/Runtime/Engine/Classes/Components/InputComponent.h
--- a/Engine/Source/Runtime/Engine/Classes/Components/InputComponent.h
+++ b/Engine/Source/Runtime/Engine/Classes/Components/InputComponent.h
@@ -723,7 +723,7 @@ struct FCachedKeyToActionInfo
  *
  * @see https://docs.unrealengine.com/latest/INT/Gameplay/Input/index.html
  */
-UCLASS(transient, config=Input, hidecategories=(Activation, "Components|Activation"))
+UCLASS(NotBlueprintable, transient, config=Input, hidecategories=(Activation, "Components|Activation"))
 class ENGINE_API UInputComponent
 	: public UActorComponent
 {
