Properly support child actor components inside level instances.

diff --git a/Engine/Source/Runtime/Engine/Private/WorldPartition/WorldPartitionLevelHelper.cpp b/Engine/Source/Runtime/Engine/Private/WorldPartition/WorldPartitionLevelHelper.cpp
--- a/Engine/Source/Runtime/Engine/Private/WorldPartition/WorldPartitionLevelHelper.cpp
+++ b/Engine/Source/Runtime/Engine/Private/WorldPartition/WorldPartitionLevelHelper.cpp
@@ -379,6 +379,14 @@ bool FWorldPartitionLevelHelper::LoadActors(UWorld* InOwningWorld, ULevel* InDes
 
 					// Rename through UObject to avoid changing Actor's external packaging and folder properties
 					Actor->UObject::Rename(*FString::Printf(TEXT("%s_%s"), *Actor->GetName(), *PackageObjectMapping->ContainerID.ToString()), InDestLevel, REN_NonTransactional | REN_ForceNoResetLoaders | REN_DoNotDirty | REN_DontCreateRedirectors);
+					// Handle child actors
+					Actor->ForEachComponent<UChildActorComponent>(true, [InDestLevel, PackageObjectMapping](UChildActorComponent* ChildActorComponent)
+					{
+						if (AActor* ChildActor = ChildActorComponent->GetChildActor())
+						{
+							ChildActor->UObject::Rename(*FString::Printf(TEXT("%s_%s"), *ChildActor->GetName(), *PackageObjectMapping->ContainerID.ToString()), InDestLevel, REN_NonTransactional | REN_ForceNoResetLoaders | REN_DoNotDirty | REN_DontCreateRedirectors);
+						}
+					});
 					
 					FLevelUtils::FApplyLevelTransformParams TransformParams(nullptr, PackageObjectMapping->ContainerTransform);
 					TransformParams.Actor = Actor;
