Harden reference reporting against null. Add a check so we can track down the source of the null in development builds.

diff --git a/Engine/Source/Runtime/Engine/Private/UnrealEngine.cpp b/Engine/Source/Runtime/Engine/Private/UnrealEngine.cpp
--- a/Engine/Source/Runtime/Engine/Private/UnrealEngine.cpp
+++ b/Engine/Source/Runtime/Engine/Private/UnrealEngine.cpp
@@ -15466,6 +15466,7 @@ static bool PrintStaleReferenceChainsAndFindReferencingObjects(UObject* ObjectTo
 	return RefChainSearch.PrintResults([&ObjectToFindReferencesTo, &OutGarbageObject, &OutReferencingObject](FReferenceChainSearch::FCallbackParams& Params)
 		{
 			check(Params.Object);
+			check(Params.Referencer);
 			if (!Params.Object->IsValid())
 			{
 				// We may find many chains that lead to the leak but for brevity we only report the first one in the fatal error below
@@ -15582,7 +15583,7 @@ TArray<FString> UEngine::FindAndPrintStaleReferencesToObjects(TConstArrayView<UO
 			PathToCulprit = GetPathToStaleObjectReferencer(ObjectToFindReferencesTo, RefChainSearch);
 			FString GarbageObjectName = OutGarbageObject ? OutGarbageObject->GetFullName() : ObjectToFindReferencesTo->GetFullName();
 			checkf(OutReferencingObject, TEXT("No object referencing %s found even though we have a valid reference chain"), *ObjectToFindReferencesTo->GetPathName());
-			GarbageErrorMessage = FString::Printf(TEXT("Object %s is being referenced by %s"), *GarbageObjectName, *OutReferencingObject->GetFullName());
+			GarbageErrorMessage = FString::Printf(TEXT("Object %s is being referenced by %s"), *GarbageObjectName, OutReferencingObject ? *OutReferencingObject->GetFullName() : TEXT("NULL"));
 		}
 #if ENABLE_GC_HISTORY
 		else
@@ -15603,7 +15604,7 @@ TArray<FString> UEngine::FindAndPrintStaleReferencesToObjects(TConstArrayView<UO
 					PathToCulprit = GetPathToStaleObjectReferencer(ObjectToFindReferencesTo, HistorySearch);
 					FString GarbageObjectName = OutGarbageObject ? OutGarbageObject->GetFullName() : ObjectToFindReferencesTo->GetFullName();
 					checkf(OutReferencingObject, TEXT("No object referencing %s found even though we have a valid reference chain"), *ObjectToFindReferencesTo->GetPathName());
-					GarbageErrorMessage = FString::Printf(TEXT("Garbage object %s was previously being referenced by %s"), *GarbageObjectName, *OutReferencingObject->GetFullName());
+					GarbageErrorMessage = FString::Printf(TEXT("Garbage object %s was previously being referenced by %s"), *GarbageObjectName, OutReferencingObject ? *OutReferencingObject->GetFullName() : TEXT("NULL"));
 				}
 			}
 		}
