Fix incorrect blend mode enforcement for Strata hair BSDF.

diff --git a/Engine/Source/Runtime/Engine/Private/Materials/Material.cpp b/Engine/Source/Runtime/Engine/Private/Materials/Material.cpp
--- a/Engine/Source/Runtime/Engine/Private/Materials/Material.cpp
+++ b/Engine/Source/Runtime/Engine/Private/Materials/Material.cpp
@@ -4572,7 +4572,6 @@ void UMaterial::RebuildShadingModelField()
 			{
 				MaterialDomain = EMaterialDomain::MD_Surface;
 				ShadingModel = MSM_Hair;
-				BlendMode = EBlendMode::BLEND_Opaque;
 			}
 			else if (StrataMaterialInfo.HasOnlyShadingModel(SSM_Eye))
 			{
