Fixed material mask set operation causing the map array to endlessly expand on repeated calls.

diff --git a/Engine/Source/Runtime/Engine/Private/PhysicsEngine/Experimental/PhysInterface_Chaos.cpp b/Engine/Source/Runtime/Engine/Private/PhysicsEngine/Experimental/PhysInterface_Chaos.cpp
--- a/Engine/Source/Runtime/Engine/Private/PhysicsEngine/Experimental/PhysInterface_Chaos.cpp
+++ b/Engine/Source/Runtime/Engine/Private/PhysicsEngine/Experimental/PhysInterface_Chaos.cpp
@@ -909,37 +909,33 @@ void FPhysInterface_Chaos::SetMaterials(const FPhysicsShapeHandle& InShape, cons
 
 		int MaskMapMatIdx = 0;
 
-		InShape.Shape->ModifyMaterialMaskMaps([&](auto& MaterialMaskMaps)
+		for(FPhysicalMaterialMaskParams& MaterialMaskData : InMaterialMasks)
 		{
-			for(FPhysicalMaterialMaskParams& MaterialMaskData : InMaterialMasks)
-		{
-				if(MaterialMaskData.PhysicalMaterialMask && ensure(MaterialMaskData.PhysicalMaterialMap))
+			if(MaterialMaskData.PhysicalMaterialMask && ensure(MaterialMaskData.PhysicalMaterialMap))
 			{
 				NewMaterialMaskHandles.Add(MaterialMaskData.PhysicalMaterialMask->GetPhysicsMaterialMask());
-					for(int i = 0; i < EPhysicalMaterialMaskColor::MAX; i++)
+				for(int i = 0; i < EPhysicalMaterialMaskColor::MAX; i++)
 				{
-						if(UPhysicalMaterial* MapMat = MaterialMaskData.PhysicalMaterialMap->GetPhysicalMaterialFromMap(i))
+					if(UPhysicalMaterial* MapMat = MaterialMaskData.PhysicalMaterialMap->GetPhysicalMaterialFromMap(i))
 					{
-							MaterialMaskMaps.Emplace(MaskMapMatIdx);
+						NewMaterialMaskMaps.Emplace(MaskMapMatIdx);
 						MaskMapMatIdx++;
-						} else
+					} 
+					else
 					{
-							MaterialMaskMaps.Emplace(INDEX_NONE);
+						NewMaterialMaskMaps.Emplace(INDEX_NONE);
+					}
 				}
-			}
-				} else
+			} 
+			else
 			{
 				NewMaterialMaskHandles.Add(Chaos::FMaterialMaskHandle());
-					for(int i = 0; i < EPhysicalMaterialMaskColor::MAX; i++)
+				for(int i = 0; i < EPhysicalMaterialMaskColor::MAX; i++)
 				{
-						MaterialMaskMaps.Emplace(INDEX_NONE);
+					NewMaterialMaskMaps.Emplace(INDEX_NONE);
 				}
 			}
 		}
-
-		});
-		
-
 		
 		if (MaskMapMatIdx > 0)
 		{
@@ -963,6 +959,7 @@ void FPhysInterface_Chaos::SetMaterials(const FPhysicsShapeHandle& InShape, cons
 		}
 
 		InShape.Shape->SetMaterialMasks(NewMaterialMaskHandles);
+		InShape.Shape->SetMaterialMaskMaps(NewMaterialMaskMaps);
 		InShape.Shape->SetMaterialMaskMapMaterials(NewMaterialMaskMaterialHandles);
 	}
 }
