Add new Pose Watch options for offsets, masking bones, and viewing in PIE.
* Pose Watches now appear in PIE when then relevant debug object is attached.
* Which bones a pose watch draws can now be masked with a blend mask created on the skeleton.
* A pose watch can now be given an offset in the Persona/PIE viewport.

diff --git a/Engine/Source/Editor/Persona/Private/AnimationEditorViewportClient.cpp b/Engine/Source/Editor/Persona/Private/AnimationEditorViewportClient.cpp
--- a/Engine/Source/Editor/Persona/Private/AnimationEditorViewportClient.cpp
+++ b/Engine/Source/Editor/Persona/Private/AnimationEditorViewportClient.cpp
@@ -1567,13 +1567,11 @@ void FAnimationViewportClient::DrawWatchedPoses(UDebugSkelMeshComponent * MeshCo
 		{
 			if (Blueprint->GetObjectBeingDebugged())
 			{
-				for (const FAnimNodePoseWatch& AnimNodePoseWatch : AnimBlueprintGeneratedClass->GetAnimBlueprintDebugData().AnimNodePoseWatch)
+				const FAnimBlueprintDebugData& DebugData = AnimBlueprintGeneratedClass->GetAnimBlueprintDebugData();
+				DebugData.ForEachActiveVisiblePoseWatch([MeshComponent, PDI](const FPoseWatchDebugData& PoseWatchDebugData)
 				{
-					if (AnimNodePoseWatch.Object.Get() != nullptr && AnimNodePoseWatch.PoseWatch->GetIsVisible() && AnimNodePoseWatch.PoseWatch->GetIsEnabled())
-					{
-						DrawBonesFromCompactPose(*AnimNodePoseWatch.PoseInfo.Get(), MeshComponent, PDI, AnimNodePoseWatch.PoseWatch->GetColor());
-					}
-				}
+					SkeletalDebugRendering::DrawBonesFromPoseWatch(*PoseWatchDebugData.PoseInfo.Get(), MeshComponent, PDI, PoseWatchDebugData.PoseWatch);
+				});
 			}
 		}
 	}
@@ -1650,135 +1648,38 @@ void FAnimationViewportClient::DrawBones(
 	bool bForceDraw,
 	bool bAddHitProxy) const
 {
-	// first determine which bones to draw, and which to filter out
-	TBitArray<> BonesToDraw(false, RefSkeleton.GetNum());
-	const bool bDrawSelected = GetBoneDrawMode() == EBoneDrawMode::Selected;
-	const bool bDrawSelectedAndParents = GetBoneDrawMode() == EBoneDrawMode::SelectedAndParents;
-	const bool bDrawSelectedAndChildren = GetBoneDrawMode() == EBoneDrawMode::SelectedAndChildren;
-	const bool bDrawSelectedAndParentsAndChildren = GetBoneDrawMode() == EBoneDrawMode::SelectedAndParentsAndChildren;
-	
-	// add selected bones
-	if (bDrawSelected || bDrawSelectedAndParents || bDrawSelectedAndChildren || bDrawSelectedAndParentsAndChildren)
-	{
-		for (int32 BoneIndex : InSelectedBones)
-		{
-			if (BoneIndex != INDEX_NONE)
-			{
-				BonesToDraw[BoneIndex] = true;
-			}
-		}
-	}
-
-	// add children of selected
-	if (bDrawSelectedAndChildren || bDrawSelectedAndParentsAndChildren)
-	{
-		for (int32 BoneIndex = 0; BoneIndex < RefSkeleton.GetNum(); ++BoneIndex)
-		{
-			const int32 ParentIndex = RefSkeleton.GetParentIndex(BoneIndex);
-			if (ParentIndex != INDEX_NONE && BonesToDraw[ParentIndex])
-			{
-				BonesToDraw[BoneIndex] = true;
-			}
-		}
-	}
+	FSkelDebugDrawConfig DrawConfig;
+	DrawConfig.BoneDrawMode = GetBoneDrawMode();
+	DrawConfig.BoneDrawSize = GetBoneDrawSize();
+	DrawConfig.bAddHitProxy = bAddHitProxy;
+	DrawConfig.bForceDraw = bForceDraw;
+	DrawConfig.DefaultBoneColor = GetMutableDefault<UPersonaOptions>()->DefaultBoneColor;
+	DrawConfig.AffectedBoneColor = GetMutableDefault<UPersonaOptions>()->AffectedBoneColor;
+	DrawConfig.SelectedBoneColor = GetMutableDefault<UPersonaOptions>()->SelectedBoneColor;
+	DrawConfig.ParentOfSelectedBoneColor = GetMutableDefault<UPersonaOptions>()->ParentOfSelectedBoneColor;
 
-	// add parents of selected
-	if (bDrawSelectedAndParents || bDrawSelectedAndParentsAndChildren)
-	{
-		for (const int32 BoneIndex : InSelectedBones)
-		{
-			if (BoneIndex != INDEX_NONE)
-			{
-				for (int32 ParentIndex = RefSkeleton.GetParentIndex(BoneIndex); ParentIndex != INDEX_NONE; ParentIndex = RefSkeleton.GetParentIndex(ParentIndex))
-				{
-					BonesToDraw[ParentIndex] = true;
-				}
-			}
-		}
-	}
+	TArray<HHitProxy*> HitProxies;
+	HitProxies.AddUninitialized(RefSkeleton.GetNum());
 
-	// determine which bones are "affected" (these are ALL children of selected bones)
-	TBitArray<> AffectedBones(false, RefSkeleton.GetNum());
-	for (int32 BoneIndex = 0; BoneIndex < RefSkeleton.GetNum(); ++BoneIndex)
+	if (bAddHitProxy)
 	{
-		for (int32 ParentIndex = RefSkeleton.GetParentIndex(BoneIndex); ParentIndex != INDEX_NONE; ParentIndex = RefSkeleton.GetParentIndex(ParentIndex))
+		for (int32 Index = 0; Index < RefSkeleton.GetNum(); ++Index)
 		{
-			if (InSelectedBones.Contains(ParentIndex))
-			{
-				AffectedBones[BoneIndex] = true;
-				break;
-			}
+			HitProxies[Index] = new HPersonaBoneHitProxy(Index, RefSkeleton.GetBoneName(Index));
 		}
 	}
-	
-	// spin through all required bones and render them
-	const float BoneRadius = GetBoneDrawSize();
-	const UPersonaOptions* PersonaOptions = GetDefault<UPersonaOptions>();
-	for ( int32 Index=0; Index<RequiredBones.Num(); ++Index )
-	{
-		const int32 BoneIndex = RequiredBones[Index];
-
-		// skips bones that should not be drawn
-		const bool bDoDraw = bForceDraw || GetBoneDrawMode() == EBoneDrawMode::All || BonesToDraw[BoneIndex];
-		if (!bDoDraw)
-		{
-			continue;
-		}
-
-		// determine color of bone based on selection / affected state
-		const bool bIsSelected = InSelectedBones.Contains(BoneIndex);
-		const bool bIsAffected = AffectedBones[BoneIndex];
-		FLinearColor DefaultBoneColor = BoneColors.IsEmpty() ? PersonaOptions->DefaultBoneColor : BoneColors[Index];
-		FLinearColor BoneColor = bIsAffected ? PersonaOptions->AffectedBoneColor : DefaultBoneColor;
-		BoneColor = bIsSelected ? PersonaOptions->SelectedBoneColor : BoneColor;
-
-		// draw the little coordinate frame inside the bone ONLY if selected or affected
-		const bool bDrawAxesInsideBone = bIsAffected|| bIsSelected;
-
-		// draw cone to each child
-		// but use a different color if this bone is NOT selected, but the child IS selected
-		TArray<FVector> ChildPositions;
-		TArray<FLinearColor> ChildColors;
-		for (int32 ChildIndex = 0; ChildIndex < RefSkeleton.GetNum(); ++ChildIndex)
-		{
-			const int32 ParentIndex = RefSkeleton.GetParentIndex(ChildIndex);
-			if (ParentIndex == BoneIndex)
-			{
-				ChildPositions.Add(WorldTransforms[ChildIndex].GetLocation());
-				FLinearColor ChildLineColor = BoneColor;
-				if (!bIsSelected && InSelectedBones.Contains(ChildIndex))
-				{
-					ChildLineColor = PersonaOptions->ParentOfSelectedBoneColor;
-				}
-				ChildColors.Add(ChildLineColor);
-			}
-		}
 
-		const FName BoneName = RefSkeleton.GetBoneName(BoneIndex);
-		const FTransform BoneTransform = WorldTransforms[BoneIndex];
-
-		if (bAddHitProxy)
-		{
-			PDI->SetHitProxy(new HPersonaBoneHitProxy(BoneIndex, BoneName));
-		}
-		SkeletalDebugRendering::DrawWireBoneAdvanced(
-			PDI,
-			BoneTransform,
-			ChildPositions,
-			ChildColors,
-			BoneColor,
-			SDPG_Foreground,
-			BoneRadius,
-			bDrawAxesInsideBone);
-		if (RefSkeleton.GetParentIndex(BoneIndex) == INDEX_NONE)
-		{
-			SkeletalDebugRendering::DrawRootCone(PDI, BoneTransform, ComponentOrigin, BoneRadius);
-		}
-		if (bAddHitProxy)
-		{
-			PDI->SetHitProxy(nullptr);	
-		}
-	}
+	SkeletalDebugRendering::DrawBones(
+		PDI,
+		ComponentOrigin,
+		RequiredBones,
+		RefSkeleton,
+		WorldTransforms,
+		InSelectedBones,
+		BoneColors,
+		HitProxies,
+		DrawConfig
+	);
 }
 
 void FAnimationViewportClient::DrawMeshSubsetBones(const UDebugSkelMeshComponent* MeshComponent, const TArray<int32>& BonesOfInterest, FPrimitiveDrawInterface* PDI) const
diff --git a/Engine/Source/Editor/Persona/Private/AnimationEditorViewportClient.h b/Engine/Source/Editor/Persona/Private/AnimationEditorViewportClient.h
--- a/Engine/Source/Editor/Persona/Private/AnimationEditorViewportClient.h
+++ b/Engine/Source/Editor/Persona/Private/AnimationEditorViewportClient.h
@@ -13,6 +13,7 @@
 #include "Animation/DebugSkelMeshComponent.h"
 #include "IPersonaPreviewScene.h"
 #include "Preferences/PersonaOptions.h"
+#include "SkeletalDebugRendering.h"
 
 class FCanvas;
 class UPersonaOptions;
@@ -37,23 +38,6 @@ namespace ELocalAxesMode
 	};
 };
 
-//////////////////////////////////////////////////////////////////////////
-// EBoneDrawMode
-
-namespace EBoneDrawMode
-{
-	enum Type
-	{
-		None,
-		Selected,
-		SelectedAndParents,
-		SelectedAndChildren,
-		SelectedAndParentsAndChildren,
-		All,
-		NumDrawModes
-	};
-};
-
 //////////////////////////////////////////////////////////////////////////
 // ELocalAxesMode
 
diff --git a/Engine/Source/Editor/Persona/Private/BlendProfileCustomization.cpp b/Engine/Source/Editor/Persona/Private/BlendProfileCustomization.cpp
--- a/Engine/Source/Editor/Persona/Private/BlendProfileCustomization.cpp
+++ b/Engine/Source/Editor/Persona/Private/BlendProfileCustomization.cpp
@@ -11,6 +11,7 @@
 #include "Kismet2/BlueprintEditorUtils.h"
 #include "Animation/AnimBlueprint.h"
 #include "Modules/ModuleManager.h"
+#include "Engine/PoseWatch.h"
 
 #define LOCTEXT_NAMESPACE "BlendProfileCustomization"
 
@@ -106,6 +107,10 @@ USkeleton* FBlendProfileCustomization::GetSkeletonFromOuter(const UObject* Outer
 	{
 		AnimBlueprint = Cast<UAnimBlueprint>(FBlueprintEditorUtils::FindBlueprintForGraph(OuterEdGraph));
 	}
+	else if (const UPoseWatch* OuterPoseWatch = Cast<UPoseWatch>(Outer))
+	{
+		AnimBlueprint = Cast<UAnimBlueprint>(FBlueprintEditorUtils::FindBlueprintForNode(OuterPoseWatch->Node.Get()));
+	}
 
 	// If outer belongs to an anim blueprint, grab its skeleton.
 	if (AnimBlueprint)
diff --git a/Engine/Source/Editor/Persona/Private/SPoseWatchManager.cpp b/Engine/Source/Editor/Persona/Private/SPoseWatchManager.cpp
--- a/Engine/Source/Editor/Persona/Private/SPoseWatchManager.cpp
+++ b/Engine/Source/Editor/Persona/Private/SPoseWatchManager.cpp
@@ -47,6 +47,7 @@
 #include "Framework/Notifications/NotificationManager.h"
 #include "Widgets/Notifications/SNotificationList.h"
 #include "Framework/Commands/GenericCommands.h"
+#include "SKismetInspector.h"
 
 #define LOCTEXT_NAMESPACE "SPoseWatchManager"
 
@@ -694,8 +695,22 @@ void SPoseWatchManager::OnManagerTreeSelectionChanged(FPoseWatchManagerTreeItemP
 	if (!bIsReentrant)
 	{
 		TGuardValue<bool> ReentrantGuard(bIsReentrant, true);
-		//Mode->OnItemSelectionChanged(TreeItem, SelectInfo, FPoseWatchManagerItemSelection(*PoseWatchManagerTreeView));
 
+		if (TreeItem.IsValid())
+		{
+			if (const FPoseWatchManagerPoseWatchTreeItem* PoseWatchItem = TreeItem.Get()->CastTo<FPoseWatchManagerPoseWatchTreeItem>())
+			{
+				BlueprintEditor->GetInspector()->ShowDetailsForSingleObject(PoseWatchItem->PoseWatch.Get());
+			}
+			else
+			{
+				BlueprintEditor->GetInspector()->ShowDetailsForSingleObject(nullptr);
+			}
+		}
+		else
+		{
+			BlueprintEditor->GetInspector()->ShowDetailsForSingleObject(nullptr);
+		}
 		OnItemSelectionChanged.Broadcast(TreeItem, SelectInfo);
 	}
 }
diff --git a/Engine/Source/Runtime/Engine/Classes/Animation/AnimBlueprintGeneratedClass.h b/Engine/Source/Runtime/Engine/Classes/Animation/AnimBlueprintGeneratedClass.h
--- a/Engine/Source/Runtime/Engine/Classes/Animation/AnimBlueprintGeneratedClass.h
+++ b/Engine/Source/Runtime/Engine/Classes/Animation/AnimBlueprintGeneratedClass.h
@@ -143,6 +143,12 @@ struct FAnimBlueprintDebugData_NodeVisit
 	}
 };
 
+struct FPoseWatchDebugData
+{
+	const UPoseWatch* PoseWatch;
+	TSharedPtr<FCompactHeapPose> PoseInfo;
+};
+
 struct FAnimBlueprintDebugData_AttributeRecord
 {
 	FName Attribute;
@@ -313,6 +319,7 @@ public:
 
 	void AddPoseWatch(int32 NodeID, UPoseWatch*);
 	void RemovePoseWatch(int32 NodeID);
+	void ForEachActiveVisiblePoseWatch(const TFunctionRef<void(FPoseWatchDebugData&)>& InFunction) const;
 
 	TArrayView<const FName> GetNodeAttributes(TWeakObjectPtr<UAnimGraphNode_Base> InAnimGraphNode) const;
 #endif
diff --git a/Engine/Source/Runtime/Engine/Classes/Animation/AnimInstance.h b/Engine/Source/Runtime/Engine/Classes/Animation/AnimInstance.h
--- a/Engine/Source/Runtime/Engine/Classes/Animation/AnimInstance.h
+++ b/Engine/Source/Runtime/Engine/Classes/Animation/AnimInstance.h
@@ -1642,5 +1642,8 @@ public:
 #if WITH_EDITOR
 	// Helper function to handle reinstancing in editor
 	virtual void HandleAnimInstanceReplaced(const TMap<UObject*, UObject*>& OldToNewInstanceMap);
+
+	/** Returns true if this anim instance is actively being debugged, false otherwise */
+	bool IsBeingDebugged() const;
 #endif
 };
diff --git a/Engine/Source/Runtime/Engine/Classes/Components/SkeletalMeshComponent.h b/Engine/Source/Runtime/Engine/Classes/Components/SkeletalMeshComponent.h
--- a/Engine/Source/Runtime/Engine/Classes/Components/SkeletalMeshComponent.h
+++ b/Engine/Source/Runtime/Engine/Classes/Components/SkeletalMeshComponent.h
@@ -1564,6 +1564,10 @@ public:
 	 */
 	void DebugDrawClothingTexts(FCanvas* Canvas, const FSceneView* SceneView);
 
+#if WITH_EDITOR
+	void DebugDrawPoseWatches(FPrimitiveDrawInterface* PDI);
+#endif
+
 	/** Changes the value of bNotifyRigidBodyCollision
 	* @param bNewNotifyRigidBodyCollision - The value to assign to bNotifyRigidBodyCollision
 	*/
diff --git a/Engine/Source/Runtime/Engine/Classes/Engine/PoseWatch.h b/Engine/Source/Runtime/Engine/Classes/Engine/PoseWatch.h
--- a/Engine/Source/Runtime/Engine/Classes/Engine/PoseWatch.h
+++ b/Engine/Source/Runtime/Engine/Classes/Engine/PoseWatch.h
@@ -10,6 +10,7 @@
 
 struct FCompactHeapPose;
 class UAnimBlueprint;
+class UBlendProfile;
 
 struct FAnimNodePoseWatch
 {
@@ -36,7 +37,7 @@ namespace PoseWatchUtil
 	template <typename T>
 	FText FindUniqueNameInFolder(UPoseWatchFolder* InParent, const T* Item, const TArray<TObjectPtr<T>>& Collection);
 
-	/** Returns a color from a palette of predefined colors in round-robin order */
+	/** Returns a new random color */
 	FColor ChoosePoseWatchColor();
 }
 
@@ -227,6 +228,23 @@ private:
 public:
 	UPROPERTY()
 	TWeakObjectPtr<class UEdGraphNode> Node;
+
+	/** Optionally select a Blend Mask to control which bones on the skeleton are rendered. Any non-zero entries are rendered. */
+	UPROPERTY(EditAnywhere, editfixedsize, Category = Default, meta = (UseAsBlendMask = true))
+	TObjectPtr<UBlendProfile> ViewportMask;
+
+	/** Invert which bones are rendered when using a viewport mask */
+	UPROPERTY(EditAnywhere, Category = Default, meta = (EditCondition = "ViewportMask != nullptr"))
+	bool bInvertViewportMask;
+
+	/** The threshold which each bone's blend scale much surpass to be rendered using the viewport mask */
+	UPROPERTY(EditAnywhere, Category = Default, meta = (ClampMin = 0.f, ClampMax = 1.f, EditCondition="ViewportMask != nullptr"))
+	float BlendScaleThreshold;
+
+	/** Offset the rendering of the bones in the viewport. */
+	UPROPERTY(EditAnywhere, Category = Default)
+	FVector3d ViewportOffset;
+
 protected:
 	UPROPERTY()
 	bool bDeleteOnDeselection = false;
diff --git a/Engine/Source/Runtime/Engine/Private/Animation/AnimBlueprintGeneratedClass.cpp b/Engine/Source/Runtime/Engine/Private/Animation/AnimBlueprintGeneratedClass.cpp
--- a/Engine/Source/Runtime/Engine/Private/Animation/AnimBlueprintGeneratedClass.cpp
+++ b/Engine/Source/Runtime/Engine/Private/Animation/AnimBlueprintGeneratedClass.cpp
@@ -260,6 +260,24 @@ void FAnimBlueprintDebugData::RemovePoseWatch(int32 NodeID)
 	}
 }
 
+void FAnimBlueprintDebugData::ForEachActiveVisiblePoseWatch(const TFunctionRef<void(FPoseWatchDebugData&)>& InFunction) const
+{
+	for (const FAnimNodePoseWatch& PoseWatchNode : AnimNodePoseWatch)
+	{
+		if (const UPoseWatch* PoseWatch = PoseWatchNode.PoseWatch.Get())
+		{
+			if (PoseWatch->GetIsEnabled() && PoseWatch->GetIsVisible())
+			{
+				FPoseWatchDebugData DrawData;
+				DrawData.PoseWatch = PoseWatch;
+				DrawData.PoseInfo = PoseWatchNode.PoseInfo;
+
+				InFunction(DrawData);
+			}
+		}
+	}
+}
+
 TArrayView<const FName> FAnimBlueprintDebugData::GetNodeAttributes(TWeakObjectPtr<UAnimGraphNode_Base> InAnimGraphNode) const
 {
 	const TArray<FName>* Attributes = NodeAttributes.Find(InAnimGraphNode);
diff --git a/Engine/Source/Runtime/Engine/Private/Animation/AnimInstance.cpp b/Engine/Source/Runtime/Engine/Private/Animation/AnimInstance.cpp
--- a/Engine/Source/Runtime/Engine/Private/Animation/AnimInstance.cpp
+++ b/Engine/Source/Runtime/Engine/Private/Animation/AnimInstance.cpp
@@ -3860,4 +3860,11 @@ bool UAnimInstance::QueryAndMarkTransitionEvent(int32 MachineIndex, int32 Transi
 	return GetProxyOnAnyThread<FAnimInstanceProxy>().QueryAndMarkTransitionEvent(MachineIndex, TransitionIndex, EventName);
 }
 
+#if WITH_EDITOR
+bool UAnimInstance::IsBeingDebugged() const
+{
+	return GetProxyOnAnyThread<FAnimInstanceProxy>().IsBeingDebugged();
+}
+#endif
+
 #undef LOCTEXT_NAMESPACE 
diff --git a/Engine/Source/Runtime/Engine/Private/Components/SkeletalMeshComponent.cpp b/Engine/Source/Runtime/Engine/Private/Components/SkeletalMeshComponent.cpp
--- a/Engine/Source/Runtime/Engine/Private/Components/SkeletalMeshComponent.cpp
+++ b/Engine/Source/Runtime/Engine/Private/Components/SkeletalMeshComponent.cpp
@@ -27,6 +27,7 @@
 #include "HAL/LowLevelMemTracker.h"
 #include "SkeletalMeshCompiler.h"
 #include "Rendering/SkeletalMeshRenderData.h"
+#include "SkeletalDebugRendering.h"
 
 #include "Logging/MessageLog.h"
 #include "Animation/AnimNode_LinkedInputPose.h"
@@ -4643,4 +4644,27 @@ void USkeletalMeshComponent::HandleObjectsReplaced(const TMap<UObject*, UObject*
 }
 #endif	// #if WITH_EDITOR
 
+
+#if WITH_EDITOR
+void USkeletalMeshComponent::DebugDrawPoseWatches(FPrimitiveDrawInterface* PDI)
+{
+	UAnimInstance* AnimInstance = GetAnimInstance();
+	if (AnimInstance)
+	{
+		if (AnimInstance->IsBeingDebugged())
+		{
+			UAnimBlueprintGeneratedClass* AnimBPGenClass = CastChecked<UAnimBlueprintGeneratedClass>(AnimInstance->GetClass());
+			if (AnimBPGenClass)
+			{
+				const FAnimBlueprintDebugData& DebugData = AnimBPGenClass->GetAnimBlueprintDebugData();
+				DebugData.ForEachActiveVisiblePoseWatch([this, PDI](const FPoseWatchDebugData& PoseWatchDebugData)
+				{
+					SkeletalDebugRendering::DrawBonesFromPoseWatch(*PoseWatchDebugData.PoseInfo.Get(), this, PDI, PoseWatchDebugData.PoseWatch);
+				});
+			}
+		}
+	}
+}
+#endif // #if WITH_EDITOR
+
 #undef LOCTEXT_NAMESPACE
diff --git a/Engine/Source/Runtime/Engine/Private/Engine/PoseWatch.cpp b/Engine/Source/Runtime/Engine/Private/Engine/PoseWatch.cpp
--- a/Engine/Source/Runtime/Engine/Private/Engine/PoseWatch.cpp
+++ b/Engine/Source/Runtime/Engine/Private/Engine/PoseWatch.cpp
@@ -23,6 +23,10 @@ UPoseWatch::UPoseWatch(const FObjectInitializer& ObjectInitializer)
 #if WITH_EDITOR
 	Label = GetDefaultLabel();
 	SetColor(PoseWatchUtil::ChoosePoseWatchColor());
+	ViewportMask = nullptr;
+	bInvertViewportMask = false;
+	BlendScaleThreshold = 0.f;
+	ViewportOffset = FVector3d::ZeroVector;
 #endif
 }
 
diff --git a/Engine/Source/Runtime/Engine/Private/SkeletalDebugRendering.cpp b/Engine/Source/Runtime/Engine/Private/SkeletalDebugRendering.cpp
--- a/Engine/Source/Runtime/Engine/Private/SkeletalDebugRendering.cpp
+++ b/Engine/Source/Runtime/Engine/Private/SkeletalDebugRendering.cpp
@@ -3,6 +3,12 @@
 #include "SkeletalDebugRendering.h"
 #include "DrawDebugHelpers.h"
 #include "SceneManagement.h"
+#include "Engine/PoseWatch.h"
+#include "Components/SkeletalMeshComponent.h"
+#include "BonePose.h"
+#include "Animation/BlendProfile.h"
+#include "HitProxies.h"
+
 
 namespace SkeletalDebugRendering
 {
@@ -165,4 +171,223 @@ void DrawRootCone(
 #endif
 }
 
+void DrawBonesFromPoseWatch(
+	const FCompactHeapPose& Pose,
+	USkeletalMeshComponent* MeshComponent,
+	FPrimitiveDrawInterface* PDI,
+	const UPoseWatch* PoseWatch)
+{
+#if WITH_EDITOR
+	if (Pose.GetNumBones() == 0 ||
+		Pose.GetBoneContainer().GetCompactPoseNumBones() == 0 ||
+		!MeshComponent ||
+		!MeshComponent->GetSkeletalMesh() ||
+		!PoseWatch
+		)
+	{
+		return;
+	}
+
+	// optionally override draw color
+	const FLinearColor BoneColor = FLinearColor(PoseWatch->GetColor());
+
+	TArray<FTransform> WorldTransforms;
+	WorldTransforms.AddUninitialized(Pose.GetBoneContainer().GetNumBones());
+
+	TArray<uint16> RequiredBones;
+	const UBlendProfile* ViewportMask = PoseWatch->ViewportMask;
+
+	const FVector RelativeOffset = MeshComponent->GetComponentRotation().RotateVector(PoseWatch->ViewportOffset);
+
+	// we could cache parent bones as we calculate, but right now I'm not worried about perf issue of this
+	for (FCompactPoseBoneIndex BoneIndex : Pose.ForEachBoneIndex())
+	{
+		FMeshPoseBoneIndex MeshBoneIndex = Pose.GetBoneContainer().MakeMeshPoseIndex(BoneIndex);
+
+		int32 ParentIndex = Pose.GetBoneContainer().GetParentBoneIndex(MeshBoneIndex.GetInt());
+
+		if (ParentIndex == INDEX_NONE)
+		{
+			WorldTransforms[MeshBoneIndex.GetInt()] = Pose[BoneIndex] * MeshComponent->GetComponentTransform();
+			WorldTransforms[MeshBoneIndex.GetInt()].AddToTranslation(RelativeOffset);
+		}
+		else
+		{
+			WorldTransforms[MeshBoneIndex.GetInt()] = Pose[BoneIndex] * WorldTransforms[ParentIndex];
+		}
+
+		if (ViewportMask == nullptr)
+		{
+			RequiredBones.Add(MeshBoneIndex.GetInt());
+		}
+		else
+		{
+			const bool bBoneBlendScaleMeetsThreshold = ViewportMask->GetBoneBlendScale(MeshBoneIndex.GetInt()) > FMath::Clamp(PoseWatch->BlendScaleThreshold, 0.f, 1.f);
+			if ((!PoseWatch->bInvertViewportMask && bBoneBlendScaleMeetsThreshold) || (PoseWatch->bInvertViewportMask && !bBoneBlendScaleMeetsThreshold))
+			{
+				RequiredBones.Add(MeshBoneIndex.GetInt());
+			}
+		}
+	}
+
+	FSkelDebugDrawConfig DrawConfig;
+	DrawConfig.BoneDrawMode = EBoneDrawMode::All;
+	DrawConfig.BoneDrawSize = 1.f;
+	DrawConfig.bAddHitProxy = false;
+	DrawConfig.bForceDraw = true;
+	DrawConfig.DefaultBoneColor = BoneColor;
+	DrawConfig.AffectedBoneColor = BoneColor;
+	DrawConfig.SelectedBoneColor = BoneColor;
+	DrawConfig.ParentOfSelectedBoneColor = BoneColor;
+
+	SkeletalDebugRendering::DrawBones(
+		PDI,
+		MeshComponent->GetComponentLocation() + RelativeOffset,
+		RequiredBones,
+		MeshComponent->GetSkeletalMesh()->GetRefSkeleton(),
+		WorldTransforms,
+		/*SelectedBones*/TArray<int32>(),
+		/*BoneColors*/TArray<FLinearColor>(),
+		/*HitProxies*/TArray<HHitProxy*>(),
+		DrawConfig
+	);
+#endif
+}
+
+void DrawBones(
+	FPrimitiveDrawInterface* PDI,
+	const FVector& ComponentOrigin,
+	const TArray<FBoneIndexType>& RequiredBones,
+	const FReferenceSkeleton& RefSkeleton,
+	const TArray<FTransform>& WorldTransforms,
+	const TArray<int32>& InSelectedBones,
+	const TArray<FLinearColor>& BoneColors,
+	const TArray<HHitProxy*>& HitProxies,
+	const FSkelDebugDrawConfig& DrawConfig)
+{
+	// first determine which bones to draw, and which to filter out
+	TBitArray<> BonesToDraw(false, RefSkeleton.GetNum());
+	const bool bDrawSelected = DrawConfig.BoneDrawMode == EBoneDrawMode::Selected;
+	const bool bDrawSelectedAndParents = DrawConfig.BoneDrawMode == EBoneDrawMode::SelectedAndParents;
+	const bool bDrawSelectedAndChildren = DrawConfig.BoneDrawMode == EBoneDrawMode::SelectedAndChildren;
+	const bool bDrawSelectedAndParentsAndChildren = DrawConfig.BoneDrawMode == EBoneDrawMode::SelectedAndParentsAndChildren;
+
+	// add selected bones
+	if (bDrawSelected || bDrawSelectedAndParents || bDrawSelectedAndChildren || bDrawSelectedAndParentsAndChildren)
+	{
+		for (int32 BoneIndex : InSelectedBones)
+		{
+			if (BoneIndex != INDEX_NONE)
+			{
+				BonesToDraw[BoneIndex] = true;
+			}
+		}
+	}
+
+	// add children of selected
+	if (bDrawSelectedAndChildren || bDrawSelectedAndParentsAndChildren)
+	{
+		for (int32 BoneIndex = 0; BoneIndex < RefSkeleton.GetNum(); ++BoneIndex)
+		{
+			const int32 ParentIndex = RefSkeleton.GetParentIndex(BoneIndex);
+			if (ParentIndex != INDEX_NONE && BonesToDraw[ParentIndex])
+			{
+				BonesToDraw[BoneIndex] = true;
+			}
+		}
+	}
+
+	// add parents of selected
+	if (bDrawSelectedAndParents || bDrawSelectedAndParentsAndChildren)
+	{
+		for (const int32 BoneIndex : InSelectedBones)
+		{
+			if (BoneIndex != INDEX_NONE)
+			{
+				for (int32 ParentIndex = RefSkeleton.GetParentIndex(BoneIndex); ParentIndex != INDEX_NONE; ParentIndex = RefSkeleton.GetParentIndex(ParentIndex))
+				{
+					BonesToDraw[ParentIndex] = true;
+				}
+			}
+		}
+	}
+
+	// determine which bones are "affected" (these are ALL children of selected bones)
+	TBitArray<> AffectedBones(false, RefSkeleton.GetNum());
+	for (int32 BoneIndex = 0; BoneIndex < RefSkeleton.GetNum(); ++BoneIndex)
+	{
+		for (int32 ParentIndex = RefSkeleton.GetParentIndex(BoneIndex); ParentIndex != INDEX_NONE; ParentIndex = RefSkeleton.GetParentIndex(ParentIndex))
+		{
+			if (InSelectedBones.Contains(ParentIndex))
+			{
+				AffectedBones[BoneIndex] = true;
+				break;
+			}
+		}
+	}
+
+	// spin through all required bones and render them
+	const float BoneRadius = DrawConfig.BoneDrawSize;
+	for (int32 Index = 0; Index < RequiredBones.Num(); ++Index)
+	{
+		const int32 BoneIndex = RequiredBones[Index];
+
+		// skips bones that should not be drawn
+		const bool bDoDraw = DrawConfig.bForceDraw || DrawConfig.BoneDrawMode == EBoneDrawMode::All || BonesToDraw[BoneIndex];
+		if (!bDoDraw)
+		{
+			continue;
+		}
+
+		// determine color of bone based on selection / affected state
+		const bool bIsSelected = InSelectedBones.Contains(BoneIndex);
+		const bool bIsAffected = AffectedBones[BoneIndex];
+		FLinearColor DefaultBoneColor = BoneColors.IsEmpty() ? DrawConfig.DefaultBoneColor : BoneColors[Index];
+		FLinearColor BoneColor = bIsAffected ? DrawConfig.AffectedBoneColor : DefaultBoneColor;
+		BoneColor = bIsSelected ? DrawConfig.SelectedBoneColor : BoneColor;
+
+		// draw the little coordinate frame inside the bone ONLY if selected or affected
+		const bool bDrawAxesInsideBone = bIsAffected || bIsSelected;
+
+		// draw cone to each child
+		// but use a different color if this bone is NOT selected, but the child IS selected
+		TArray<FVector> ChildPositions;
+		TArray<FLinearColor> ChildColors;
+		for (int32 ChildIndex = 0; ChildIndex < RefSkeleton.GetNum(); ++ChildIndex)
+		{
+			const int32 ParentIndex = RefSkeleton.GetParentIndex(ChildIndex);
+			if (ParentIndex == BoneIndex)
+			{
+				ChildPositions.Add(WorldTransforms[ChildIndex].GetLocation());
+				FLinearColor ChildLineColor = BoneColor;
+				if (!bIsSelected && InSelectedBones.Contains(ChildIndex))
+				{
+					ChildLineColor = DrawConfig.ParentOfSelectedBoneColor;
+				}
+				ChildColors.Add(ChildLineColor);
+			}
+		}
+
+		const FName BoneName = RefSkeleton.GetBoneName(BoneIndex);
+		const FTransform BoneTransform = WorldTransforms[BoneIndex];
+
+		// Always set new hit proxy to prevent unintentionally using last drawn element's proxy
+		PDI->SetHitProxy(DrawConfig.bAddHitProxy ? HitProxies[BoneIndex] : nullptr);
+		SkeletalDebugRendering::DrawWireBoneAdvanced(
+			PDI,
+			BoneTransform,
+			ChildPositions,
+			ChildColors,
+			BoneColor,
+			SDPG_Foreground,
+			BoneRadius,
+			bDrawAxesInsideBone);
+		if (RefSkeleton.GetParentIndex(BoneIndex) == INDEX_NONE)
+		{
+			SkeletalDebugRendering::DrawRootCone(PDI, BoneTransform, ComponentOrigin, BoneRadius);
+		}
+		PDI->SetHitProxy(nullptr);
+	}
+}
+
 }
\ No newline at end of file
diff --git a/Engine/Source/Runtime/Engine/Private/SkeletalMesh.cpp b/Engine/Source/Runtime/Engine/Private/SkeletalMesh.cpp
--- a/Engine/Source/Runtime/Engine/Private/SkeletalMesh.cpp
+++ b/Engine/Source/Runtime/Engine/Private/SkeletalMesh.cpp
@@ -6379,6 +6379,10 @@ void FSkeletalMeshSceneProxy::GetMeshElementsConditionallySelectable(const TArra
 			{
 				DebugDrawSkeleton(ViewIndex, Collector, ViewFamily.EngineShowFlags);
 			}
+
+#if WITH_EDITOR
+			DebugDrawPoseWatchSkeletons(ViewIndex, Collector, ViewFamily.EngineShowFlags);
+#endif
 		}
 	}
 #endif
@@ -6858,6 +6862,22 @@ void FSkeletalMeshSceneProxy::DebugDrawPhysicsAsset(int32 ViewIndex, FMeshElemen
 	}
 }
 
+#if WITH_EDITOR
+void FSkeletalMeshSceneProxy::DebugDrawPoseWatchSkeletons(int32 ViewIndex, FMeshElementCollector& Collector, const FEngineShowFlags& EngineShowFlags) const
+{
+	FPrimitiveDrawInterface* PDI = Collector.GetPDI(ViewIndex);
+
+	if (Owner)
+	{
+		USkeletalMeshComponent* SkeletalMeshComp = Cast<USkeletalMeshComponent>(Owner->GetComponentByClass(USkeletalMeshComponent::StaticClass()));
+		if (SkeletalMeshComp)
+		{
+			SkeletalMeshComp->DebugDrawPoseWatches(PDI);
+		}
+	}
+}
+#endif
+
 void FSkeletalMeshSceneProxy::DebugDrawSkeleton(int32 ViewIndex, FMeshElementCollector& Collector, const FEngineShowFlags& EngineShowFlags) const
 {
 #if !(UE_BUILD_SHIPPING || UE_BUILD_TEST)
diff --git a/Engine/Source/Runtime/Engine/Public/SkeletalDebugRendering.h b/Engine/Source/Runtime/Engine/Public/SkeletalDebugRendering.h
--- a/Engine/Source/Runtime/Engine/Public/SkeletalDebugRendering.h
+++ b/Engine/Source/Runtime/Engine/Public/SkeletalDebugRendering.h
@@ -4,8 +4,41 @@
 
 #include "CoreMinimal.h"
 #include "Engine/EngineTypes.h"
+#include "BoneIndices.h"
 
+class UPoseWatch;
 class FPrimitiveDrawInterface;
+struct FCompactHeapPose;
+struct FReferenceSkeleton;
+class USkeletalMeshComponent;
+class HHitProxy;
+
+namespace EBoneDrawMode
+{
+	enum Type
+	{
+		None,
+		Selected,
+		SelectedAndParents,
+		SelectedAndChildren,
+		SelectedAndParentsAndChildren,
+		All,
+		NumDrawModes
+	};
+};
+
+struct FSkelDebugDrawConfig
+{
+public:
+	EBoneDrawMode::Type BoneDrawMode;
+	float BoneDrawSize;
+	bool bForceDraw;
+	bool bAddHitProxy;
+	FLinearColor DefaultBoneColor;
+	FLinearColor AffectedBoneColor;
+	FLinearColor SelectedBoneColor;
+	FLinearColor ParentOfSelectedBoneColor;
+};
 
 namespace SkeletalDebugRendering
 {
@@ -77,4 +110,26 @@ ENGINE_API	void DrawRootCone(
 	const FTransform& InBoneTransform,
 	const FVector& ComponentOrigin,
 	const float SphereRadius);
-}
+
+
+ENGINE_API void DrawBonesFromPoseWatch(
+	const FCompactHeapPose& Pose,
+	USkeletalMeshComponent* MeshComponent,
+	FPrimitiveDrawInterface* PDI,
+	const UPoseWatch* PoseWatch);
+
+
+/**
+ * Draw skeleton bones. If BoneColors is empty, the DefaultBoneColor in the draw config will be used.
+ */
+ENGINE_API void DrawBones(
+	FPrimitiveDrawInterface* PDI,
+	const FVector& ComponentOrigin,
+	const TArray<FBoneIndexType>& RequiredBones,
+	const FReferenceSkeleton& RefSkeleton,
+	const TArray<FTransform>& WorldTransforms,
+	const TArray<int32>& InSelectedBones,
+	const TArray<FLinearColor>& BoneColors,
+	const TArray<HHitProxy*>& HitProxies,
+	const FSkelDebugDrawConfig& DrawConfig);
+}
\ No newline at end of file
diff --git a/Engine/Source/Runtime/Engine/Public/SkeletalMeshTypes.h b/Engine/Source/Runtime/Engine/Public/SkeletalMeshTypes.h
--- a/Engine/Source/Runtime/Engine/Public/SkeletalMeshTypes.h
+++ b/Engine/Source/Runtime/Engine/Public/SkeletalMeshTypes.h
@@ -205,6 +205,10 @@ public:
 	/** Render the bones of the skeleton for debug display */ 
 	void DebugDrawSkeleton(int32 ViewIndex, FMeshElementCollector& Collector, const FEngineShowFlags& EngineShowFlags) const;
 
+#if WITH_EDITOR
+	void DebugDrawPoseWatchSkeletons(int32 ViewIndex, FMeshElementCollector& Collector, const FEngineShowFlags& EngineShowFlags) const;
+#endif
+
 	virtual uint32 GetMemoryFootprint( void ) const override { return( sizeof( *this ) + GetAllocatedSize() ); }
 	SIZE_T GetAllocatedSize( void ) const { return( FPrimitiveSceneProxy::GetAllocatedSize() + LODSections.GetAllocatedSize() ); }
 
