Added option to cleanup marked as deleted actor folders that are not referenced by any actor of folder.
When marking a folder as deleted, actually delete the object is it is not referenced.

diff --git a/Engine/Source/Editor/EditorStyle/Private/StarshipStyle.cpp b/Engine/Source/Editor/EditorStyle/Private/StarshipStyle.cpp
--- a/Engine/Source/Editor/EditorStyle/Private/StarshipStyle.cpp
+++ b/Engine/Source/Editor/EditorStyle/Private/StarshipStyle.cpp
@@ -1123,6 +1123,7 @@ void FStarshipEditorStyle::FStyle::SetupGeneralStyles()
 		Set( "SceneOutliner.NewFolderIcon", new CORE_IMAGE_BRUSH_SVG("Starship/Common/folder-plus", Icon16x16 ) );
 		Set( "SceneOutliner.FolderClosed",  new CORE_IMAGE_BRUSH_SVG("Starship/Common/folder-closed", Icon16x16, FStyleColors::AccentFolder ) );
 		Set( "SceneOutliner.FolderOpen",    new CORE_IMAGE_BRUSH_SVG("Starship/Common/folder-open", Icon16x16, FStyleColors::AccentFolder ) );
+		Set( "SceneOutliner.CleanupActorFoldersIcon", new CORE_IMAGE_BRUSH_SVG("Starship/Common/folder-cleanup", Icon16x16));
 		Set( "SceneOutliner.World", 		new CORE_IMAGE_BRUSH_SVG("Starship/Common/world", Icon16x16 ) );
 		Set( "SceneOutliner.ChangedItemHighlight", new FSlateRoundedBoxBrush( FStyleColors::White, 1.0) );
 
diff --git a/Engine/Source/Editor/SceneOutliner/Private/ActorBrowsingMode.cpp b/Engine/Source/Editor/SceneOutliner/Private/ActorBrowsingMode.cpp
--- a/Engine/Source/Editor/SceneOutliner/Private/ActorBrowsingMode.cpp
+++ b/Engine/Source/Editor/SceneOutliner/Private/ActorBrowsingMode.cpp
@@ -474,12 +474,7 @@ void FActorBrowsingMode::RegisterContextMenu()
 				{
 					if (Context->NumSelectedItems == 0)
 					{
-						Section.AddMenuEntry(
-							"CreateFolder",
-							LOCTEXT("CreateFolder", "Create Folder"),
-							FText(),
-							FSlateIcon(FAppStyle::GetAppStyleSetName(), "SceneOutliner.NewFolderIcon"),
-							FUIAction(FExecuteAction::CreateSP(SceneOutliner, &SSceneOutliner::CreateFolder)));
+						FSceneOutlinerMenuHelper::AddMenuEntryCreateFolder(Section, *SceneOutliner);
 					}
 					else
 					{
diff --git a/Engine/Source/Editor/SceneOutliner/Private/ActorTreeItem.cpp b/Engine/Source/Editor/SceneOutliner/Private/ActorTreeItem.cpp
--- a/Engine/Source/Editor/SceneOutliner/Private/ActorTreeItem.cpp
+++ b/Engine/Source/Editor/SceneOutliner/Private/ActorTreeItem.cpp
@@ -451,10 +451,9 @@ void FActorTreeItem::GenerateContextMenu(UToolMenu* Menu, SSceneOutliner& Outlin
 	const ILevelInstanceInterface* LevelInstance = Cast<ILevelInstanceInterface>(ActorPtr);
 	if (LevelInstance && LevelInstance->IsEditing())
 	{
-		auto SharedOutliner = StaticCastSharedRef<SSceneOutliner>(Outliner.AsShared());
-		const FSlateIcon NewFolderIcon(FAppStyle::GetAppStyleSetName(), "SceneOutliner.NewFolderIcon");
 		FToolMenuSection& Section = Menu->AddSection("Section");
-		Section.AddMenuEntry("CreateFolder", LOCTEXT("CreateFolder", "Create Folder"), FText(), NewFolderIcon, FUIAction(FExecuteAction::CreateSP(&Outliner, &SSceneOutliner::CreateFolder)));
+		FSceneOutlinerMenuHelper::AddMenuEntryCreateFolder(Section, Outliner);
+		FSceneOutlinerMenuHelper::AddMenuEntryCleanupFolders(Section, LevelInstance->GetLoadedLevel());
 	}
 }
 
diff --git a/Engine/Source/Editor/SceneOutliner/Private/LevelTreeItem.cpp b/Engine/Source/Editor/SceneOutliner/Private/LevelTreeItem.cpp
--- a/Engine/Source/Editor/SceneOutliner/Private/LevelTreeItem.cpp
+++ b/Engine/Source/Editor/SceneOutliner/Private/LevelTreeItem.cpp
@@ -3,7 +3,9 @@
 #include "LevelTreeItem.h"
 #include "Engine/Level.h"
 #include "LevelUtils.h"
+#include "ToolMenus.h"
 #include "ISceneOutliner.h"
+#include "SSceneOutliner.h"
 #include "SceneOutlinerPublicTypes.h"
 #include "Styling/SlateIconFinder.h"
 #include "Widgets/Images/SImage.h"
@@ -102,4 +104,13 @@ bool FLevelTreeItem::CanInteract() const
 	return Level.IsValid() ? !FLevelUtils::IsLevelLocked(Level.Get()) : true;
 }
 
+void FLevelTreeItem::GenerateContextMenu(UToolMenu* Menu, SSceneOutliner& Outliner)
+{
+	if (Level.IsValid() && Level->IsUsingActorFolders())
+	{
+		FToolMenuSection& Section = Menu->AddSection("Section");
+		FSceneOutlinerMenuHelper::AddMenuEntryCleanupFolders(Section, Level.Get());
+	}
+}
+
 #undef LOCTEXT_NAMESPACE
diff --git a/Engine/Source/Editor/SceneOutliner/Private/SSceneOutliner.cpp b/Engine/Source/Editor/SceneOutliner/Private/SSceneOutliner.cpp
--- a/Engine/Source/Editor/SceneOutliner/Private/SSceneOutliner.cpp
+++ b/Engine/Source/Editor/SceneOutliner/Private/SSceneOutliner.cpp
@@ -2355,4 +2355,23 @@ void SSceneOutliner::SaveConfig()
 	UOutlinerConfig::Get()->SaveEditorConfig();
 }
 
+void FSceneOutlinerMenuHelper::AddMenuEntryCreateFolder(FToolMenuSection& InSection, SSceneOutliner& InOutliner)
+{
+	const FSlateIcon NewFolderIcon(FAppStyle::GetAppStyleSetName(), "SceneOutliner.NewFolderIcon");
+	InSection.AddMenuEntry("CreateFolder", LOCTEXT("CreateFolder", "Create Folder"), LOCTEXT("CreateFolderTooltip", "Create Folder"), NewFolderIcon, FUIAction(FExecuteAction::CreateSP(&InOutliner, &SSceneOutliner::CreateFolder)));
+}
+
+void FSceneOutlinerMenuHelper::AddMenuEntryCleanupFolders(FToolMenuSection& InSection, ULevel* InLevel)
+{
+	if (InLevel && InLevel->IsUsingActorFolders())
+	{
+		const FSlateIcon CleanupFoldersIcon(FAppStyle::GetAppStyleSetName(), "SceneOutliner.CleanupActorFoldersIcon");
+		InSection.AddMenuEntry("CleanupFolders", LOCTEXT("CleanupFolders", "Cleanup Folders"), LOCTEXT("CleanupFoldersTooltip", "Cleanup unreferenced and deleted actor folders"), CleanupFoldersIcon, FExecuteAction::CreateLambda([InLevel]()
+		{
+			const FScopedTransaction Transaction(LOCTEXT("CleanupFolders", "Cleanup Folders"));
+			InLevel->CleanupDeletedAndUnreferencedActorFolders();
+		}));
+	}
+}
+
 #undef LOCTEXT_NAMESPACE
\ No newline at end of file
diff --git a/Engine/Source/Editor/SceneOutliner/Private/WorldTreeItem.cpp b/Engine/Source/Editor/SceneOutliner/Private/WorldTreeItem.cpp
--- a/Engine/Source/Editor/SceneOutliner/Private/WorldTreeItem.cpp
+++ b/Engine/Source/Editor/SceneOutliner/Private/WorldTreeItem.cpp
@@ -208,11 +208,19 @@ void FWorldTreeItem::GenerateContextMenu(UToolMenu* Menu, SSceneOutliner& Outlin
 {
 	auto SharedOutliner = StaticCastSharedRef<SSceneOutliner>(Outliner.AsShared());
 	
-	const FSlateIcon WorldSettingsIcon(FAppStyle::GetAppStyleSetName(), "LevelEditor.WorldProperties.Tab");
-	const FSlateIcon NewFolderIcon(FAppStyle::GetAppStyleSetName(), "SceneOutliner.NewFolderIcon");
 	FToolMenuSection& Section = Menu->AddSection("Section");
-	Section.AddMenuEntry("CreateFolder", LOCTEXT("CreateFolder", "Create Folder"), FText(), NewFolderIcon, FUIAction(FExecuteAction::CreateSP(&Outliner, &SSceneOutliner::CreateFolder)));
+
+	// Create Folder menu entry
+	FSceneOutlinerMenuHelper::AddMenuEntryCreateFolder(Section, Outliner);
+
+	// Open world settings menu entry
+	const FSlateIcon WorldSettingsIcon(FAppStyle::GetAppStyleSetName(), "LevelEditor.WorldProperties.Tab");
 	Section.AddMenuEntry("OpenWorldSettings", LOCTEXT("OpenWorldSettings", "World Settings"), FText(), WorldSettingsIcon, FExecuteAction::CreateSP(this, &FWorldTreeItem::OpenWorldSettings));
+	
+	// Cleanup Folders menu entry
+	UWorld* WorldPtr = World.Get();
+	ULevel* Level = WorldPtr ? WorldPtr->PersistentLevel : nullptr;
+	FSceneOutlinerMenuHelper::AddMenuEntryCleanupFolders(Section, Level);
 }
 
 TSharedRef<SWidget> FWorldTreeItem::GenerateLabelWidget(ISceneOutliner& Outliner, const STableRow<FSceneOutlinerTreeItemPtr>& InRow)
diff --git a/Engine/Source/Editor/SceneOutliner/Public/LevelTreeItem.h b/Engine/Source/Editor/SceneOutliner/Public/LevelTreeItem.h
--- a/Engine/Source/Editor/SceneOutliner/Public/LevelTreeItem.h
+++ b/Engine/Source/Editor/SceneOutliner/Public/LevelTreeItem.h
@@ -41,5 +41,6 @@ public:
 	virtual FString GetDisplayString() const override;
 	virtual bool CanInteract() const override;
 	virtual TSharedRef<SWidget> GenerateLabelWidget(ISceneOutliner& Outliner, const STableRow<FSceneOutlinerTreeItemPtr>& InRow) override;
+	virtual void GenerateContextMenu(UToolMenu* Menu, SSceneOutliner& Outliner) override;
 	/* End ISceneOutlinerTreeItem Implementation */
 };
diff --git a/Engine/Source/Editor/SceneOutliner/Public/SSceneOutliner.h b/Engine/Source/Editor/SceneOutliner/Public/SSceneOutliner.h
--- a/Engine/Source/Editor/SceneOutliner/Public/SSceneOutliner.h
+++ b/Engine/Source/Editor/SceneOutliner/Public/SSceneOutliner.h
@@ -34,6 +34,8 @@ class FMenuBuilder;
 class UToolMenu;
 class ISceneOutlinerColumn;
 class SComboButton;
+class ULevel;
+struct FToolMenuSection;
 
 template<typename ItemType> class STreeView;
 
@@ -906,3 +908,9 @@ private:
 
 	void SaveConfig();
 };
+
+struct SCENEOUTLINER_API FSceneOutlinerMenuHelper
+{
+	static void AddMenuEntryCreateFolder(FToolMenuSection& InSection, SSceneOutliner& InOutliner);
+	static void AddMenuEntryCleanupFolders(FToolMenuSection& InSection, ULevel* InLevel);
+};
diff --git a/Engine/Source/Runtime/Engine/Classes/Engine/Level.h b/Engine/Source/Runtime/Engine/Classes/Engine/Level.h
--- a/Engine/Source/Runtime/Engine/Classes/Engine/Level.h
+++ b/Engine/Source/Runtime/Engine/Classes/Engine/Level.h
@@ -1226,6 +1226,9 @@ public:
 	/** Updates all actors/folders that refer to folders marked as deleted, reparent to valid folder, deletes folders marked as deleted. */
 	ENGINE_API void CreateOrUpdateActorFolders();
 
+	/** Deletes all actor folders marked as deleted and unreferenced by neither an actor nor another actor folder */
+	ENGINE_API void CleanupDeletedAndUnreferencedActorFolders();
+
 	/** Sets if the level uses actor folders mode or not. Returns true if succeeded. */
 	ENGINE_API bool SetUseActorFolders(bool bEnabled, bool bInteractive = false);
 
@@ -1398,6 +1401,9 @@ private:
 	/** Sets the level to use or not the actor folder objects feature. */
 	void SetUseActorFoldersInternal(bool bInEnabled);
 
+	/** Returns unreferenced actor folders that are marked as deleted. */
+	TSet<FGuid> GetDeletedAndUnreferencedActorFolders() const;
+
 	friend struct FLevelActorFoldersHelper;
 	friend class FWorldPartitionLevelHelper;
 	friend class UActorFolder;
diff --git a/Engine/Source/Runtime/Engine/Private/Level.cpp b/Engine/Source/Runtime/Engine/Private/Level.cpp
--- a/Engine/Source/Runtime/Engine/Private/Level.cpp
+++ b/Engine/Source/Runtime/Engine/Private/Level.cpp
@@ -78,6 +78,7 @@ Level.cpp: Level-related functions
 #include "EditorActorFolders.h"
 #include "UObject/MetaData.h"
 #include "UObject/LinkerLoad.h"
+#include "WorldPartition/WorldPartitionHelpers.h"
 #endif
 #include "WorldPartition/WorldPartition.h"
 #include "WorldPartition/WorldPartitionSubsystem.h"
@@ -147,8 +148,11 @@ void FLevelActorFoldersHelper::RenameFolder(ULevel* InLevel, const FFolder& InOl
 	}
 };
 
+static bool GAllowCleanupActorFolders = true;
+
 void FLevelActorFoldersHelper::DeleteFolder(ULevel* InLevel, const FFolder& InFolder)
 {
+	TGuardValue<bool> GuardValue(GAllowCleanupActorFolders, false);
 	check(InLevel);
 	// This implementation can be called both if FActorFolders is or isn't initialized.
 	if (FActorFolders::Get().IsInitializedForWorld(*InLevel->GetWorld()))
@@ -2600,14 +2604,21 @@ void ULevel::RemoveActorFolder(UActorFolder* InActorFolder)
 {
 	Modify(false);
 	check(InActorFolder);
+	check(InActorFolder->IsMarkedAsDeleted());
+	check(GAllowCleanupActorFolders);
+
 	TObjectPtr<UActorFolder> FoundActorFolder;
 	if (ensure(ActorFolders.RemoveAndCopyValue(InActorFolder->GetGuid(), FoundActorFolder)))
 	{
-		if (!InActorFolder->IsMarkedAsDeleted())
-		{
-			FActorFolderSet& Folders = FolderLabelToActorFolders.FindChecked(FoundActorFolder->GetLabel());
-			verify(Folders.Remove(FoundActorFolder));
-		}
+		check(FoundActorFolder == InActorFolder);
+		FActorFolderSet* Folders = FolderLabelToActorFolders.Find(FoundActorFolder->GetLabel());
+		check(!Folders || !Folders->GetActorFolders().Contains(InActorFolder))
+
+		GEngine->BroadcastActorFolderRemoved(InActorFolder);
+		
+		InActorFolder->Modify();
+		InActorFolder->MarkAsGarbage();
+		check(InActorFolder->GetPackage()->IsDirty());
 	}
 }
 
@@ -2624,6 +2635,10 @@ void ULevel::OnFolderMarkAsDeleted(UActorFolder* InActorFolder)
 		{
 			FolderLabelToActorFolders.Remove(InActorFolder->GetLabel());
 		}
+		if (GAllowCleanupActorFolders && GetDeletedAndUnreferencedActorFolders().Contains(InActorFolder->GetGuid()))
+		{
+			RemoveActorFolder(InActorFolder);
+		}
 	}
 }
 
@@ -3518,8 +3533,6 @@ void ULevel::CreateOrUpdateActorFolders()
 	// Remove actor folders marked as deleted
 	for (UActorFolder* ActorFolderToDelete : ActorFoldersToDelete)
 	{
-		GEngine->BroadcastActorFolderRemoved(ActorFolderToDelete);
-
 		RemoveActorFolder(ActorFolderToDelete);
 	}
 
@@ -3531,6 +3544,68 @@ void ULevel::CreateOrUpdateActorFolders()
 	}
 }
 
+TSet<FGuid> ULevel::GetDeletedAndUnreferencedActorFolders() const
+{
+	TSet<FGuid> FoldersToDelete;
+
+	if (!IsUsingActorFolders())
+	{
+		return FoldersToDelete;
+	}
+
+	for (const auto& Pair : ActorFolders)
+	{
+		UActorFolder* ActorFolder = Pair.Value;
+		if (ActorFolder->IsMarkedAsDeleted())
+		{
+			FoldersToDelete.Add(ActorFolder->GetGuid());
+		}
+	}
+
+	for (const auto& Pair : ActorFolders)
+	{
+		UActorFolder* ActorFolder = Pair.Value;
+		UActorFolder* ParentFolder = ActorFolder->GetParent(false);
+		if (ParentFolder)
+		{
+			FoldersToDelete.Remove(ParentFolder->GetGuid());
+		}
+	}
+
+	if (UWorldPartition* WorldPartition = GetWorldPartition())
+	{
+		FWorldPartitionHelpers::ForEachActorDesc(WorldPartition, [&FoldersToDelete](const FWorldPartitionActorDesc* ActorDesc)
+		{
+			AActor* Actor = ActorDesc->GetActor();
+			FoldersToDelete.Remove(Actor ? Actor->GetFolderGuid() : ActorDesc->GetFolderGuid());
+			return true;
+		});
+	}
+
+	// Also loop through all loaded level actors
+	for (AActor* Actor : Actors)
+	{
+		if (IsValid(Actor))
+		{
+			FoldersToDelete.Remove(Actor->GetFolderGuid());
+		}
+	}
+
+	return FoldersToDelete;
+}
+
+void ULevel::CleanupDeletedAndUnreferencedActorFolders()
+{
+	// Remove actor folders marked as deleted
+	for (const FGuid GuidFolderToDelete : GetDeletedAndUnreferencedActorFolders())
+	{
+		if (UActorFolder* ActorFolderToDelete = GetActorFolder(GuidFolderToDelete, false /*bSkipDeleted*/))
+		{
+			RemoveActorFolder(ActorFolderToDelete);
+		}
+	}
+}
+
 void ULevel::ForEachActorFolder(TFunctionRef<bool(UActorFolder*)> Operation, bool bSkipDeleted)
 {
 	if (!IsUsingActorFolders())
diff --git a/Engine/Source/Runtime/Engine/Public/ActorFolder.h b/Engine/Source/Runtime/Engine/Public/ActorFolder.h
--- a/Engine/Source/Runtime/Engine/Public/ActorFolder.h
+++ b/Engine/Source/Runtime/Engine/Public/ActorFolder.h
@@ -22,7 +22,7 @@ public:
 	// Helper that creates a new UActorFolder
 	static UActorFolder* Create(ULevel* InOuter, const FString& InFolderLabel, UActorFolder* InParent);
 
-	static EPackageFlags GetExternalPackageFlags() { return (PKG_EditorOnly | PKG_ContainsMapData); }
+	static EPackageFlags GetExternalPackageFlags() { return (PKG_EditorOnly | PKG_ContainsMapData | PKG_NewlyCreated); }
 	
 	// Returns actor folder info stored in its package
 	static FActorFolderDesc GetAssetRegistryInfoFromPackage(FName PackageName);
