Add WITH_EDITOR around BodySetup factory code

diff --git a/Engine/Source/Runtime/Engine/Classes/PhysicsEngine/BodySetup.h b/Engine/Source/Runtime/Engine/Classes/PhysicsEngine/BodySetup.h
--- a/Engine/Source/Runtime/Engine/Classes/PhysicsEngine/BodySetup.h
+++ b/Engine/Source/Runtime/Engine/Classes/PhysicsEngine/BodySetup.h
@@ -15,7 +15,9 @@
 #include "HAL/ThreadSafeBool.h"
 #include "Async/TaskGraphInterfaces.h"
 #include "BodySetupCore.h"
+#if WITH_EDITOR
 #include "Factories.h"
+#endif
 #include "BodySetup.generated.h"
 
 class ITargetPlatform;
@@ -447,6 +449,7 @@ public:
 
 };
 
+#if WITH_EDITOR
 
 class FBodySetupObjectTextFactory : public FCustomizableTextObjectFactory
 {
@@ -459,3 +462,4 @@ public:
 	TArray<UBodySetup*> NewBodySetups;
 };
 
+#endif // WITH_EDITOR
diff --git a/Engine/Source/Runtime/Engine/Private/PhysicsEngine/BodySetup.cpp b/Engine/Source/Runtime/Engine/Private/PhysicsEngine/BodySetup.cpp
--- a/Engine/Source/Runtime/Engine/Private/PhysicsEngine/BodySetup.cpp
+++ b/Engine/Source/Runtime/Engine/Private/PhysicsEngine/BodySetup.cpp
@@ -2265,6 +2265,8 @@ void FKLevelSetElem::ScaleElem(FVector DeltaSize, float MinSize)
 	SetTransform(ScaledTransform);
 }
 
+#if WITH_EDITOR
+
 bool FBodySetupObjectTextFactory::CanCreateClass(UClass* InObjectClass, bool& bOmitSubObjs) const
 {
 	return (InObjectClass->IsChildOf<UBodySetup>());
@@ -2279,4 +2281,6 @@ void FBodySetupObjectTextFactory::ProcessConstructedObject(UObject* NewObject)
 	}
 }
 
+#endif // WITH_EDITOR
+
 #undef LOCTEXT_NAMESPACE
