Ignore LevelBounds actor when computing LevelInstance bounds

diff --git a/Engine/Source/Runtime/Engine/Private/LevelInstance/LevelInstanceSubsystem.cpp b/Engine/Source/Runtime/Engine/Private/LevelInstance/LevelInstanceSubsystem.cpp
--- a/Engine/Source/Runtime/Engine/Private/LevelInstance/LevelInstanceSubsystem.cpp
+++ b/Engine/Source/Runtime/Engine/Private/LevelInstance/LevelInstanceSubsystem.cpp
@@ -1394,6 +1394,10 @@ void ULevelInstanceSubsystem::FActorDescContainerInstanceManager::FActorDescCont
 	Bounds.Init();
 	for (FActorDescList::TIterator<> ActorDescIt(Container); ActorDescIt; ++ActorDescIt)
 	{
+		if (ActorDescIt->GetActorNativeClass()->IsChildOf<ALevelBounds>())
+		{
+			continue;
+		}
 		Bounds += ActorDescIt->GetBounds();
 	}
 }
