Fixed a race condition when loading a mesh description from the source model of a static mesh.
Nanite Displaced Mesh build is still not safe from other race conditions but this should make the cooking free of race condition while we rework the asset build.

diff --git a/Engine/Source/Runtime/Engine/Classes/Engine/StaticMeshSourceData.h b/Engine/Source/Runtime/Engine/Classes/Engine/StaticMeshSourceData.h
--- a/Engine/Source/Runtime/Engine/Classes/Engine/StaticMeshSourceData.h
+++ b/Engine/Source/Runtime/Engine/Classes/Engine/StaticMeshSourceData.h
@@ -79,6 +79,9 @@ public:
 	 */
 	UPROPERTY()
 	TObjectPtr<UStaticMeshDescriptionBulkData> StaticMeshDescriptionBulkData;
+
+	// Only allow safe read of the mesh description from the game thread while a build is running (not sure about the rest. the UStaticMeshDescription stuff look particularly dangerous)
+	mutable FCriticalSection StaticMeshDescriptionBulkDataCS;
 #endif
 
 #if WITH_EDITOR
@@ -197,7 +200,6 @@ public:
 	void SerializeBulkData(FArchive& Ar, UObject* Owner);
 
 	void ConvertRawMesh(int32 LodIndex);
-
 #endif
 };
 
diff --git a/Engine/Source/Runtime/Engine/Private/StaticMeshSourceData.cpp b/Engine/Source/Runtime/Engine/Private/StaticMeshSourceData.cpp
--- a/Engine/Source/Runtime/Engine/Private/StaticMeshSourceData.cpp
+++ b/Engine/Source/Runtime/Engine/Private/StaticMeshSourceData.cpp
@@ -9,6 +9,7 @@
 #if WITH_EDITOR
 #include "DerivedDataCacheInterface.h"
 #include "RawMesh.h"
+#include "Misc/ScopeLock.h"
 #endif
 
 
@@ -54,15 +55,39 @@
 
 FStaticMeshSourceModel& FStaticMeshSourceModel::operator=(FStaticMeshSourceModel&& Other)
 {
+	if (this == &Other)
+	{
+		return *this;
+	}
+
 #if WITH_EDITOR
 	RawMeshBulkData = Other.RawMeshBulkData;
 	Other.RawMeshBulkData = nullptr;
 #endif
 
 #if WITH_EDITORONLY_DATA
-	StaticMeshDescriptionBulkData = Other.StaticMeshDescriptionBulkData;
-	Other.StaticMeshDescriptionBulkData = nullptr;
+	{
+		FCriticalSection* First;
+		FCriticalSection* Second;
+		// Lock always in the same order
+		if (this < &Other)
+		{
+			First = &StaticMeshDescriptionBulkDataCS;
+			Second = &Other.StaticMeshDescriptionBulkDataCS;
+		}
+		else
+		{
+			First = &Other.StaticMeshDescriptionBulkDataCS;
+			Second = &StaticMeshDescriptionBulkDataCS;
+		}
 
+		FScopeLock StaticMeshDescriptionBulkDataLock(First);
+		FScopeLock OtherStaticMeshDescriptionBulkDataLock(Second);
+
+		StaticMeshDescriptionBulkData = Other.StaticMeshDescriptionBulkData;
+		Other.StaticMeshDescriptionBulkData = nullptr;
+	}
+
 	bImportWithBaseMesh = Other.bImportWithBaseMesh;
 #endif
 
@@ -79,6 +104,8 @@
 void FStaticMeshSourceModel::CreateSubObjects(UStaticMesh* InOwner)
 {
 #if WITH_EDITORONLY_DATA
+	FScopeLock StaticMeshDescriptionBulkDataLock(&StaticMeshDescriptionBulkDataCS);
+
 	if (StaticMeshDescriptionBulkData == nullptr)
 	{
 		// Currently there is a restriction on creating UObjects on a thread while garbage collection is taking place,
@@ -98,15 +125,20 @@
 #if WITH_EDITOR
 bool FStaticMeshSourceModel::IsRawMeshEmpty() const
 {
+	FScopeLock StaticMeshDescriptionBulkDataLock(&StaticMeshDescriptionBulkDataCS);
+
 	// The RawMeshBulkData will always be empty, so the test here is whether we are able to construct a valid non-empty RawMesh
 	// from what exists in the StaticMeshDescription and its bulk data.
 	check(RawMeshBulkData->IsEmpty());
 	check(StaticMeshDescriptionBulkData != nullptr);
+
 	return !StaticMeshDescriptionBulkData->IsBulkDataValid() && !StaticMeshDescriptionBulkData->HasCachedMeshDescription();
 }
 
 void FStaticMeshSourceModel::LoadRawMesh(FRawMesh& OutRawMesh) const
 {
+	FScopeLock StaticMeshDescriptionBulkDataLock(&StaticMeshDescriptionBulkDataCS);
+
 	check(RawMeshBulkData->IsEmpty());
 	check(StaticMeshDescriptionBulkData != nullptr);
 
@@ -133,6 +165,8 @@
 
 	TRACE_CPUPROFILER_EVENT_SCOPE(FStaticMeshSourceModel::SaveRawMesh);
 
+	FScopeLock StaticMeshDescriptionBulkDataLock(&StaticMeshDescriptionBulkDataCS);
+
 	check(RawMeshBulkData->IsEmpty());
 	check(StaticMeshDescriptionBulkData != nullptr);
 
@@ -157,6 +191,9 @@
 bool FStaticMeshSourceModel::LoadMeshDescription(FMeshDescription& OutMeshDescription) const
 {
 	TRACE_CPUPROFILER_EVENT_SCOPE(FStaticMeshSourceModel::LoadMeshDescription);
+
+	FScopeLock StaticMeshDescriptionBulkDataLock(&StaticMeshDescriptionBulkDataCS);
+
 	check(StaticMeshDescriptionBulkData != nullptr);
 
 	// Ensure default MeshDescription result is empty, with no attributes registered
@@ -182,6 +219,8 @@
 {
 	TRACE_CPUPROFILER_EVENT_SCOPE(FStaticMeshSourceModel::CloneMeshDescription);
 
+	FScopeLock StaticMeshDescriptionBulkDataLock(&StaticMeshDescriptionBulkDataCS);
+
 	check(StaticMeshDescriptionBulkData != nullptr);
 	if (StaticMeshDescriptionBulkData->HasCachedMeshDescription())
 	{
@@ -196,6 +235,9 @@
 FMeshDescription* FStaticMeshSourceModel::GetOrCacheMeshDescription()
 {
 	TRACE_CPUPROFILER_EVENT_SCOPE(FStaticMeshSourceModel::GetMeshDescription);
+
+	FScopeLock StaticMeshDescriptionBulkDataLock(&StaticMeshDescriptionBulkDataCS);
+
 	check(StaticMeshDescriptionBulkData != nullptr);
 
 	if (!StaticMeshDescriptionBulkData->HasCachedMeshDescription())
@@ -221,6 +263,8 @@
 
 FMeshDescription* FStaticMeshSourceModel::GetCachedMeshDescription() const
 {
+	FScopeLock StaticMeshDescriptionBulkDataLock(&StaticMeshDescriptionBulkDataCS);
+
 	check(StaticMeshDescriptionBulkData != nullptr);
 	if (StaticMeshDescriptionBulkData->HasCachedMeshDescription())
 	{
@@ -233,6 +277,8 @@
 
 UStaticMeshDescription* FStaticMeshSourceModel::GetCachedStaticMeshDescription() const
 {
+	FScopeLock StaticMeshDescriptionBulkDataLock(&StaticMeshDescriptionBulkDataCS);
+
 	check(StaticMeshDescriptionBulkData != nullptr);
 	return Cast<UStaticMeshDescription>(StaticMeshDescriptionBulkData->GetMeshDescription());
 }
@@ -240,6 +286,7 @@
 
 FMeshDescriptionBulkData* FStaticMeshSourceModel::GetMeshDescriptionBulkData() const
 {
+	FScopeLock StaticMeshDescriptionBulkDataLock(&StaticMeshDescriptionBulkDataCS);
 	check(StaticMeshDescriptionBulkData != nullptr);
 	return &StaticMeshDescriptionBulkData->GetBulkData();
 }
@@ -247,6 +294,8 @@
 
 bool FStaticMeshSourceModel::IsMeshDescriptionValid() const
 {
+	FScopeLock StaticMeshDescriptionBulkDataLock(&StaticMeshDescriptionBulkDataCS);
+
 	// Determine whether a mesh description is valid without requiring it to be loaded first.
 	// If there is a valid MeshDescriptionBulkData, we know this implies a valid mesh description.
 	check(StaticMeshDescriptionBulkData != nullptr);
@@ -257,6 +306,7 @@
 
 FMeshDescription* FStaticMeshSourceModel::CreateMeshDescription()
 {
+	FScopeLock StaticMeshDescriptionBulkDataLock(&StaticMeshDescriptionBulkDataCS);
 	check(StaticMeshDescriptionBulkData != nullptr);
 	return &StaticMeshDescriptionBulkData->CreateMeshDescription()->GetMeshDescription();
 }
@@ -266,6 +316,8 @@
 {
 	TRACE_CPUPROFILER_EVENT_SCOPE(FStaticMeshSourceModel::CommitMeshDescription);
 
+	FScopeLock StaticMeshDescriptionBulkDataLock(&StaticMeshDescriptionBulkDataCS);
+
 	check(RawMeshBulkData->IsEmpty());
 
 	// This part should remain thread-safe so it can be called from any thread
@@ -287,6 +339,8 @@
 {
 	TRACE_CPUPROFILER_EVENT_SCOPE(UStaticMesh::ClearMeshDescription);
 
+	FScopeLock StaticMeshDescriptionBulkDataLock(&StaticMeshDescriptionBulkDataCS);
+
 	check(StaticMeshDescriptionBulkData != nullptr);
 	StaticMeshDescriptionBulkData->RemoveMeshDescription();
 }
