Control local exposure constract scale for highlights and shadows independently.

diff --git a/Engine/Shaders/Private/PostProcessBloom.usf b/Engine/Shaders/Private/PostProcessBloom.usf
--- a/Engine/Shaders/Private/PostProcessBloom.usf
+++ b/Engine/Shaders/Private/PostProcessBloom.usf
@@ -53,7 +53,7 @@ float4 BloomSetupCommon(float2 UV, float2 ExposureScaleMiddleGreyLumValue)
 		float MiddleGreyLumValue = ExposureScaleMiddleGreyLumValue.y;
 
 		float BaseLogLum = CalculateBaseLogLuminance(LogLuminance, EyeAdaptation_LocalExposureBlurredLuminanceBlend, ExposureScale, UV, LumBilateralGrid, BlurredLogLum, LumBilateralGridSampler, BlurredLogLumSampler);
-		float LocalExposure = CalculateLocalExposure(LogLuminance + log2(ExposureScale), BaseLogLum, MiddleGreyLumValue, EyeAdaptation_LocalExposureContrastScale, EyeAdaptation_LocalExposureDetailStrength);
+		float LocalExposure = CalculateLocalExposure(LogLuminance + log2(ExposureScale), BaseLogLum, MiddleGreyLumValue, EyeAdaptation_LocalExposureHighlightContrastScale, EyeAdaptation_LocalExposureShadowContrastScale, EyeAdaptation_LocalExposureDetailStrength);
 		
 		SceneColor.rgb *= LocalExposure;
 	}
diff --git a/Engine/Shaders/Private/PostProcessHistogramCommon.ush b/Engine/Shaders/Private/PostProcessHistogramCommon.ush
--- a/Engine/Shaders/Private/PostProcessHistogramCommon.ush
+++ b/Engine/Shaders/Private/PostProcessHistogramCommon.ush
@@ -16,7 +16,8 @@ float EyeAdaptation_ExposureSpeedDown;
 float EyeAdaptation_HistogramScale;
 float EyeAdaptation_HistogramBias;
 float EyeAdaptation_LuminanceMin;
-float EyeAdaptation_LocalExposureContrastScale;
+float EyeAdaptation_LocalExposureHighlightContrastScale;
+float EyeAdaptation_LocalExposureShadowContrastScale;
 float EyeAdaptation_LocalExposureDetailStrength;
 float EyeAdaptation_LocalExposureBlurredLuminanceBlend;
 float EyeAdaptation_LocalExposureMiddleGreyExposureCompensation;
@@ -228,10 +229,13 @@ float CalculateBaseLogLuminance(float LogLum, float BlurredLumBlend, float Expos
 	return CalculateBaseLogLuminance(BilateralLum, BlurredLum, BlurredLumBlend, ExposureScale);
 }
 
-float CalculateLocalExposure(float LogLum, float BaseLogLum, float LogMiddleGrey, float ContrastReduction, float DetailStrength)
+float CalculateLocalExposure(float LogLum, float BaseLogLum, float LogMiddleGrey, float HighlightContrastScale, float ShadowContrastScale, float DetailStrength)
 {
 	float DetailLogLum = LogLum - BaseLogLum;
 
-	float LogLocalLum = LogMiddleGrey + (BaseLogLum - LogMiddleGrey) * ContrastReduction + DetailLogLum * DetailStrength;
+	float BaseCentered = (BaseLogLum - LogMiddleGrey);
+	float ContrastScale = BaseCentered > 0 ? HighlightContrastScale : ShadowContrastScale;
+
+	float LogLocalLum = LogMiddleGrey + BaseCentered * ContrastScale + DetailLogLum * DetailStrength;
 	return exp2(LogLocalLum - LogLum);
 }
diff --git a/Engine/Shaders/Private/PostProcessLocalExposure.usf b/Engine/Shaders/Private/PostProcessLocalExposure.usf
--- a/Engine/Shaders/Private/PostProcessLocalExposure.usf
+++ b/Engine/Shaders/Private/PostProcessLocalExposure.usf
@@ -46,7 +46,7 @@ void ApplyLocalExposureCS(uint2 DispatchThreadId : SV_DispatchThreadID)
 		float LogLuminance = log2(LuminanceVal);
 
 		float BaseLogLum = CalculateBaseLogLuminance(LogLuminance, EyeAdaptation_LocalExposureBlurredLuminanceBlend, ExposureScale, UV, LumBilateralGrid, BlurredLogLum, TextureSampler, TextureSampler);
-		float LocalExposure = CalculateLocalExposure(LogLuminance + log2(ExposureScale), BaseLogLum, MiddleGreyLumValue, EyeAdaptation_LocalExposureContrastScale, EyeAdaptation_LocalExposureDetailStrength);
+		float LocalExposure = CalculateLocalExposure(LogLuminance + log2(ExposureScale), BaseLogLum, MiddleGreyLumValue, EyeAdaptation_LocalExposureHighlightContrastScale, EyeAdaptation_LocalExposureShadowContrastScale, EyeAdaptation_LocalExposureDetailStrength);
 
 		SceneColor.rgb *= LocalExposure;
 
diff --git a/Engine/Shaders/Private/PostProcessTonemap.usf b/Engine/Shaders/Private/PostProcessTonemap.usf
--- a/Engine/Shaders/Private/PostProcessTonemap.usf
+++ b/Engine/Shaders/Private/PostProcessTonemap.usf
@@ -451,7 +451,7 @@ float4 TonemapCommonPS(
 		float MiddleGreyLumValue = ExposureScaleVignette.y;
 
 		float BaseLogLum = CalculateBaseLogLuminance(LogLuminance, EyeAdaptation_LocalExposureBlurredLuminanceBlend, ExposureScale, FullViewUV, LumBilateralGrid, BlurredLogLum, LumBilateralGridSampler, BlurredLogLumSampler);
-		float LocalExposure = CalculateLocalExposure(LogLuminance + log2(ExposureScale), BaseLogLum, MiddleGreyLumValue, EyeAdaptation_LocalExposureContrastScale, EyeAdaptation_LocalExposureDetailStrength);
+		float LocalExposure = CalculateLocalExposure(LogLuminance + log2(ExposureScale), BaseLogLum, MiddleGreyLumValue, EyeAdaptation_LocalExposureHighlightContrastScale, EyeAdaptation_LocalExposureShadowContrastScale, EyeAdaptation_LocalExposureDetailStrength);
 
 		// Apply before bloom
 		SceneColor.rgb *= LocalExposure;
diff --git a/Engine/Shaders/Private/PostProcessVisualizeLocalExposure.usf b/Engine/Shaders/Private/PostProcessVisualizeLocalExposure.usf
--- a/Engine/Shaders/Private/PostProcessVisualizeLocalExposure.usf
+++ b/Engine/Shaders/Private/PostProcessVisualizeLocalExposure.usf
@@ -35,7 +35,7 @@ float4 MainPS(noperspective float4 UVAndScreenPos : TEXCOORD0, float4 SvPosition
 	float MiddleGreyLumValue = log2(0.18 * ExposureScaleMiddleGrey.y * EyeAdaptation_LocalExposureMiddleGreyExposureCompensation);
 
 	float BaseLogLum = CalculateBaseLogLuminance(LogLuminance, EyeAdaptation_LocalExposureBlurredLuminanceBlend, ExposureScaleMiddleGrey.x, UV, LumBilateralGrid, BlurredLogLum, TextureSampler, TextureSampler);
-	float LocalExposure = CalculateLocalExposure(LogLuminance + log2(ExposureScaleMiddleGrey.x), BaseLogLum, MiddleGreyLumValue, EyeAdaptation_LocalExposureContrastScale, EyeAdaptation_LocalExposureDetailStrength);
+	float LocalExposure = CalculateLocalExposure(LogLuminance + log2(ExposureScaleMiddleGrey.x), BaseLogLum, MiddleGreyLumValue, EyeAdaptation_LocalExposureHighlightContrastScale, EyeAdaptation_LocalExposureShadowContrastScale, EyeAdaptation_LocalExposureDetailStrength);
 	float DetailLogLum = (LogLuminance + log2(ExposureScaleMiddleGrey.x)) - BaseLogLum;
 
 	float3 OutValue;
diff --git a/Engine/Source/Runtime/Engine/Classes/Engine/Scene.h b/Engine/Source/Runtime/Engine/Classes/Engine/Scene.h
--- a/Engine/Source/Runtime/Engine/Classes/Engine/Scene.h
+++ b/Engine/Source/Runtime/Engine/Classes/Engine/Scene.h
@@ -915,8 +915,14 @@ struct FPostProcessSettings
 	UPROPERTY(EditAnywhere, BlueprintReadWrite, Category=Overrides, meta=(PinHiddenByDefault, InlineEditConditionToggle))
 	uint8 bOverride_HistogramLogMax:1;
 
+	UPROPERTY()
+	uint8 bOverride_LocalExposureContrastScale_DEPRECATED:1;
+
+	UPROPERTY(EditAnywhere, BlueprintReadWrite, Category=Overrides, meta=(PinHiddenByDefault, InlineEditConditionToggle))
+	uint8 bOverride_LocalExposureHighlightContrastScale:1;
+
 	UPROPERTY(EditAnywhere, BlueprintReadWrite, Category=Overrides, meta=(PinHiddenByDefault, InlineEditConditionToggle))
-	uint8 bOverride_LocalExposureContrastScale:1;
+	uint8 bOverride_LocalExposureShadowContrastScale:1;
 
 	UPROPERTY(EditAnywhere, BlueprintReadWrite, Category=Overrides, meta=(PinHiddenByDefault, InlineEditConditionToggle))
 	uint8 bOverride_LocalExposureDetailStrength:1;
@@ -1774,8 +1780,26 @@ struct FPostProcessSettings
 	 * Value less than 1 will enable local exposure.
 	 * Good values are usually in the range 0.6 .. 1.0.
 	*/
-	UPROPERTY(interp, BlueprintReadWrite, Category="Lens|Local Exposure", meta=(UIMin = "0.0", UIMax = "1.0", editcondition = "bOverride_LocalExposureContrastScale", DisplayName = "Contrast Scale"))
-	float LocalExposureContrastScale;
+	UPROPERTY()
+	float LocalExposureContrastScale_DEPRECATED;
+
+	/** 
+	 * Local Exposure decomposes luminance of the frame into a base layer and a detail layer.
+	 * Contrast of the base layer is reduced based on this value.
+	 * Value less than 1 will enable local exposure.
+	 * Good values are usually in the range 0.6 .. 1.0.
+	*/
+	UPROPERTY(interp, BlueprintReadWrite, Category="Lens|Local Exposure", meta=(UIMin = "0.0", UIMax = "1.0", editcondition = "bOverride_LocalExposureHighlightContrastScale", DisplayName = "Highlight Contrast Scale"))
+	float LocalExposureHighlightContrastScale;
+
+	/** 
+	 * Local Exposure decomposes luminance of the frame into a base layer and a detail layer.
+	 * Contrast of the base layer is reduced based on this value.
+	 * Value less than 1 will enable local exposure.
+	 * Good values are usually in the range 0.6 .. 1.0.
+	*/
+	UPROPERTY(interp, BlueprintReadWrite, Category="Lens|Local Exposure", meta=(UIMin = "0.0", UIMax = "1.0", editcondition = "bOverride_LocalExposureShadowContrastScale", DisplayName = "Shadow Contrast Scale"))
+	float LocalExposureShadowContrastScale;
 
 	/**
 	 * Local Exposure decomposes luminance of the frame into a base layer and a detail layer.
@@ -2180,6 +2204,17 @@ struct FPostProcessSettings
 		{
 			RayTracingGIType = (ERayTracingGlobalIlluminationType)(RayTracingGI_DEPRECATED == 1);
 		}
+
+		if (bOverride_LocalExposureContrastScale_DEPRECATED)
+		{
+			bOverride_LocalExposureHighlightContrastScale = bOverride_LocalExposureContrastScale_DEPRECATED;
+			bOverride_LocalExposureShadowContrastScale = bOverride_LocalExposureContrastScale_DEPRECATED;
+		}
+		if (LocalExposureContrastScale_DEPRECATED != 1.0f)
+		{
+			LocalExposureHighlightContrastScale = LocalExposureContrastScale_DEPRECATED;
+			LocalExposureShadowContrastScale = LocalExposureContrastScale_DEPRECATED;
+		}
 	}
 #endif
 
diff --git a/Engine/Source/Runtime/Engine/Private/Scene.cpp b/Engine/Source/Runtime/Engine/Private/Scene.cpp
--- a/Engine/Source/Runtime/Engine/Private/Scene.cpp
+++ b/Engine/Source/Runtime/Engine/Private/Scene.cpp
@@ -480,7 +480,9 @@ FPostProcessSettings::FPostProcessSettings()
 
 	AutoExposureApplyPhysicalCameraExposure = 1;
 
-	LocalExposureContrastScale = 1.0f;
+	LocalExposureContrastScale_DEPRECATED = 1.0f;
+	LocalExposureHighlightContrastScale = 1.0f;
+	LocalExposureShadowContrastScale = 1.0f;
 	LocalExposureDetailStrength = 1.0f;
 	LocalExposureBlurredLuminanceBlend = 0.6f;
 	LocalExposureBlurredLuminanceKernelSizePercent = 50.0f;
@@ -698,7 +700,9 @@ FPostProcessSettings::FPostProcessSettings(const FPostProcessSettings& Settings)
 	, bOverride_AutoExposureApplyPhysicalCameraExposure(Settings.bOverride_AutoExposureApplyPhysicalCameraExposure)
 	, bOverride_HistogramLogMin(Settings.bOverride_HistogramLogMin)
 	, bOverride_HistogramLogMax(Settings.bOverride_HistogramLogMax)
-	, bOverride_LocalExposureContrastScale(Settings.bOverride_LocalExposureContrastScale)
+	, bOverride_LocalExposureContrastScale_DEPRECATED(Settings.bOverride_LocalExposureContrastScale_DEPRECATED)
+	, bOverride_LocalExposureHighlightContrastScale(Settings.bOverride_LocalExposureHighlightContrastScale)
+	, bOverride_LocalExposureShadowContrastScale(Settings.bOverride_LocalExposureShadowContrastScale)
 	, bOverride_LocalExposureDetailStrength(Settings.bOverride_LocalExposureDetailStrength)
 	, bOverride_LocalExposureBlurredLuminanceBlend(Settings.bOverride_LocalExposureBlurredLuminanceBlend)
 	, bOverride_LocalExposureBlurredLuminanceKernelSizePercent(Settings.bOverride_LocalExposureBlurredLuminanceKernelSizePercent)
@@ -916,7 +920,9 @@ FPostProcessSettings::FPostProcessSettings(const FPostProcessSettings& Settings)
 	, HistogramLogMin(Settings.HistogramLogMin)
 	, HistogramLogMax(Settings.HistogramLogMax)
 	, AutoExposureCalibrationConstant_DEPRECATED(Settings.AutoExposureCalibrationConstant_DEPRECATED)
-	, LocalExposureContrastScale(Settings.LocalExposureContrastScale)
+	, LocalExposureContrastScale_DEPRECATED(Settings.LocalExposureContrastScale_DEPRECATED)
+	, LocalExposureHighlightContrastScale(Settings.LocalExposureHighlightContrastScale)
+	, LocalExposureShadowContrastScale(Settings.LocalExposureShadowContrastScale)
 	, LocalExposureDetailStrength(Settings.LocalExposureDetailStrength)
 	, LocalExposureBlurredLuminanceBlend(Settings.LocalExposureBlurredLuminanceBlend)
 	, LocalExposureBlurredLuminanceKernelSizePercent(Settings.LocalExposureBlurredLuminanceKernelSizePercent)
diff --git a/Engine/Source/Runtime/Engine/Private/SceneView.cpp b/Engine/Source/Runtime/Engine/Private/SceneView.cpp
--- a/Engine/Source/Runtime/Engine/Private/SceneView.cpp
+++ b/Engine/Source/Runtime/Engine/Private/SceneView.cpp
@@ -1509,7 +1509,9 @@ void FSceneView::OverridePostProcessSettings(const FPostProcessSettings& Src, fl
 		LERP_PP(AutoExposureBias);
 		LERP_PP(HistogramLogMin);
 		LERP_PP(HistogramLogMax);
-		LERP_PP(LocalExposureContrastScale);
+		LERP_PP(LocalExposureContrastScale_DEPRECATED);
+		LERP_PP(LocalExposureHighlightContrastScale);
+		LERP_PP(LocalExposureShadowContrastScale);
 		LERP_PP(LocalExposureDetailStrength);
 		LERP_PP(LocalExposureBlurredLuminanceBlend);
 		LERP_PP(LocalExposureBlurredLuminanceKernelSizePercent);
@@ -2016,7 +2018,8 @@ void FSceneView::EndFinalPostprocessSettings(const FSceneViewInitOptions& ViewIn
 
 		if (Value <= 0)
 		{
-			FinalPostProcessSettings.LocalExposureContrastScale = 1.0f;
+			FinalPostProcessSettings.LocalExposureHighlightContrastScale = 1.0f;
+			FinalPostProcessSettings.LocalExposureShadowContrastScale = 1.0f;
 			FinalPostProcessSettings.LocalExposureDetailStrength = 1.0f;
 		}
 	}
@@ -2039,7 +2042,8 @@ void FSceneView::EndFinalPostprocessSettings(const FSceneViewInitOptions& ViewIn
 
 	if (!Family->EngineShowFlags.LocalExposure)
 	{
-		FinalPostProcessSettings.LocalExposureContrastScale = 1.0f;
+		FinalPostProcessSettings.LocalExposureHighlightContrastScale = 1.0f;
+		FinalPostProcessSettings.LocalExposureShadowContrastScale = 1.0f;
 		FinalPostProcessSettings.LocalExposureDetailStrength = 1.0f;
 	}
 
diff --git a/Engine/Source/Runtime/Renderer/Private/PostProcess/PostProcessEyeAdaptation.cpp b/Engine/Source/Runtime/Renderer/Private/PostProcess/PostProcessEyeAdaptation.cpp
--- a/Engine/Source/Runtime/Renderer/Private/PostProcess/PostProcessEyeAdaptation.cpp
+++ b/Engine/Source/Runtime/Renderer/Private/PostProcess/PostProcessEyeAdaptation.cpp
@@ -427,7 +427,8 @@ FEyeAdaptationParameters GetEyeAdaptationParameters(const FViewInfo& View, ERHIF
 	Parameters.HistogramScale = HistogramScale;
 	Parameters.HistogramBias = HistogramBias;
 	Parameters.LuminanceMin = LuminanceMin;
-	Parameters.LocalExposureContrastScale = Settings.LocalExposureContrastScale;
+	Parameters.LocalExposureHighlightContrastScale = Settings.LocalExposureHighlightContrastScale;
+	Parameters.LocalExposureShadowContrastScale = Settings.LocalExposureShadowContrastScale;
 	Parameters.LocalExposureDetailStrength = Settings.LocalExposureDetailStrength;
 	Parameters.LocalExposureBlurredLuminanceBlend = Settings.LocalExposureBlurredLuminanceBlend;
 	Parameters.LocalExposureMiddleGreyExposureCompensation = LocalExposureMiddleGreyExposureCompensation;
diff --git a/Engine/Source/Runtime/Renderer/Private/PostProcess/PostProcessEyeAdaptation.h b/Engine/Source/Runtime/Renderer/Private/PostProcess/PostProcessEyeAdaptation.h
--- a/Engine/Source/Runtime/Renderer/Private/PostProcess/PostProcessEyeAdaptation.h
+++ b/Engine/Source/Runtime/Renderer/Private/PostProcess/PostProcessEyeAdaptation.h
@@ -65,7 +65,8 @@ BEGIN_SHADER_PARAMETER_STRUCT(FEyeAdaptationParameters, )
 	SHADER_PARAMETER(float, HistogramScale)
 	SHADER_PARAMETER(float, HistogramBias)
 	SHADER_PARAMETER(float, LuminanceMin)
-	SHADER_PARAMETER(float, LocalExposureContrastScale)
+	SHADER_PARAMETER(float, LocalExposureHighlightContrastScale)
+	SHADER_PARAMETER(float, LocalExposureShadowContrastScale)
 	SHADER_PARAMETER(float, LocalExposureDetailStrength)
 	SHADER_PARAMETER(float, LocalExposureBlurredLuminanceBlend)
 	SHADER_PARAMETER(float, LocalExposureMiddleGreyExposureCompensation)
diff --git a/Engine/Source/Runtime/Renderer/Private/PostProcess/PostProcessing.cpp b/Engine/Source/Runtime/Renderer/Private/PostProcess/PostProcessing.cpp
--- a/Engine/Source/Runtime/Renderer/Private/PostProcess/PostProcessing.cpp
+++ b/Engine/Source/Runtime/Renderer/Private/PostProcess/PostProcessing.cpp
@@ -472,7 +472,8 @@ void AddPostProcessingPasses(
 
 		const bool bLocalExposureEnabled =
 			EngineShowFlags.VisualizeLocalExposure || 
-			!FMath::IsNearlyEqual(View.FinalPostProcessSettings.LocalExposureContrastScale, 1.0f) ||
+			!FMath::IsNearlyEqual(View.FinalPostProcessSettings.LocalExposureHighlightContrastScale, 1.0f) ||
+			!FMath::IsNearlyEqual(View.FinalPostProcessSettings.LocalExposureShadowContrastScale, 1.0f) ||
 			!FMath::IsNearlyEqual(View.FinalPostProcessSettings.LocalExposureDetailStrength, 1.0f);
 
 		const bool bBloomEnabled = View.FinalPostProcessSettings.BloomIntensity > 0.0f;
diff --git a/Engine/Source/Runtime/Renderer/Private/SceneRendering.cpp b/Engine/Source/Runtime/Renderer/Private/SceneRendering.cpp
--- a/Engine/Source/Runtime/Renderer/Private/SceneRendering.cpp
+++ b/Engine/Source/Runtime/Renderer/Private/SceneRendering.cpp
@@ -3328,7 +3328,10 @@ void FSceneRenderer::RenderFinish(FRDGBuilder& GraphBuilder, FRDGTextureRef View
 
 			bNaniteEnabledButDisabledInProject = bNaniteEnabledButDisabledInProject || (WouldRenderNanite(Scene, View, /*bCheckForAtomicSupport*/ false, /*bCheckForProjectSetting*/ false) && !WouldRenderNanite(Scene, View, /*bCheckForAtomicSupport*/ false, /*bCheckForProjectSetting*/ true));
 
-			if (IsPostProcessingEnabled(View) && (!FMath::IsNearlyEqual(View.FinalPostProcessSettings.LocalExposureContrastScale, 1.0f) || !FMath::IsNearlyEqual(View.FinalPostProcessSettings.LocalExposureDetailStrength, 1.0f)))
+			if (IsPostProcessingEnabled(View)
+				&& (!FMath::IsNearlyEqual(View.FinalPostProcessSettings.LocalExposureHighlightContrastScale, 1.0f) 
+					|| !FMath::IsNearlyEqual(View.FinalPostProcessSettings.LocalExposureShadowContrastScale, 1.0f) 
+					|| !FMath::IsNearlyEqual(View.FinalPostProcessSettings.LocalExposureDetailStrength, 1.0f)))
 				bLocalExposureEnabledOnAnyView = true;
 		}
 
