[VisualLogger] at all time we maintain list of all children in the whole parent hierarchy. This way the redirection of a given object always has the full list of all inherited children and can be properly updated when changing redirection of one of its child.
This change allows to collect debug snapshots from every redirected objects in case of multiple levels of redirection (A->B->C)

diff --git a/Engine/Source/Runtime/Engine/Private/VisualLogger/VisualLogger.cpp b/Engine/Source/Runtime/Engine/Private/VisualLogger/VisualLogger.cpp
--- a/Engine/Source/Runtime/Engine/Private/VisualLogger/VisualLogger.cpp
+++ b/Engine/Source/Runtime/Engine/Private/VisualLogger/VisualLogger.cpp
@@ -326,13 +326,15 @@ FVisualLogEntry* FVisualLogger::GetEntryToWriteInternal(const UObject* Object, c
 
 		FReadScopeLock RedirectScopeLock(RedirectRWLock);
 		FOwnerToChildrenRedirectionMap& RedirectionMap = GetRedirectionMap(LogOwner);
-		if (RedirectionMap.Contains(LogOwner))
+		if (const IVisualLoggerDebugSnapshotInterface* DebugSnapshotInterface = Cast<const IVisualLoggerDebugSnapshotInterface>(LogOwner))
 		{
-			if (const IVisualLoggerDebugSnapshotInterface* DebugSnapshotInterface = Cast<const IVisualLoggerDebugSnapshotInterface>(LogOwner))
-			{
-				DebugSnapshotInterface->GrabDebugSnapshot(CurrentEntry);
-			}
-			for (TWeakObjectPtr<const UObject>& Child : RedirectionMap[LogOwner])
+			DebugSnapshotInterface->GrabDebugSnapshot(CurrentEntry);
+		}
+		
+		TArray<TWeakObjectPtr<const UObject>>* Children = RedirectionMap.Find(LogOwner);
+		if (Children != nullptr)
+		{
+			for (TWeakObjectPtr<const UObject>& Child : *Children)
 			{
 				if (const IVisualLoggerDebugSnapshotInterface* DebugSnapshotInterface = Cast<const IVisualLoggerDebugSnapshotInterface>(Child.Get()))
 				{
@@ -340,13 +342,6 @@ FVisualLogEntry* FVisualLogger::GetEntryToWriteInternal(const UObject* Object, c
 				}
 			}
 		}
-		else
-		{
-			if (const IVisualLoggerDebugSnapshotInterface* DebugSnapshotInterface = Cast<const IVisualLoggerDebugSnapshotInterface>(LogOwner))
-			{
-				DebugSnapshotInterface->GrabDebugSnapshot(CurrentEntry);
-			}
-		}
 	}
 
 	if (CurrentEntry != nullptr && CurrentEntry->bIsAllowedToLog)
@@ -693,24 +688,59 @@ UObject* FVisualLogger::RedirectInternal(const UObject* FromObject, const UObjec
 	}
 
 	const TWeakObjectPtr<const UObject> FromWeakPtr(FromObject);
-	UObject* OldRedirection = FindRedirectionInternal(FromObject);
-	UObject* NewRedirection = FindRedirectionInternal(ToObject);
 
-	if (OldRedirection != NewRedirection)
+	TWeakObjectPtr<const UObject>& WeakOwnerEntry = ChildToOwnerMap.FindOrAdd(FromObject);
+	const UObject* DirectOwner = WeakOwnerEntry.Get();
+	const UObject* NewDirectOwner = ToObject;
+
+	if (DirectOwner != ToObject)
 	{
 		FOwnerToChildrenRedirectionMap& OwnerToChildrenMap = GetRedirectionMap(FromObject);
+		const TArray<TWeakObjectPtr<const UObject>>* ObjectChildren = OwnerToChildrenMap.Find(FromObject);
 
-		TArray<TWeakObjectPtr<const UObject>>* OldArray = OwnerToChildrenMap.Find(OldRedirection);
-		if (OldArray)
+		// A valid direct owner indicates that FromObject was in one or many list(s)
+		// of children so remove it and its children from the owner hierarchy
+		if (DirectOwner != nullptr)
 		{
-			OldArray->RemoveSingleSwap(FromWeakPtr);
+			const UObject* Owner = DirectOwner;
+			while (Owner != nullptr)
+			{
+				if (TArray<TWeakObjectPtr<const UObject>>* DirectOwnerChildren = OwnerToChildrenMap.Find(Owner))
+				{
+					DirectOwnerChildren->RemoveSingleSwap(FromWeakPtr);
+					if (ObjectChildren != nullptr)
+					{
+						for (const TWeakObjectPtr<const UObject>& Child : *ObjectChildren)
+						{
+							DirectOwnerChildren->RemoveSingleSwap(Child);
+						}
+					}
+				}
+
+				const TWeakObjectPtr<const UObject>* WeakOwner = ChildToOwnerMap.Find(Owner);
+				Owner = WeakOwner ? WeakOwner->Get() : nullptr;
+			}
 		}
 
-		OwnerToChildrenMap.FindOrAdd(NewRedirection).AddUnique(FromWeakPtr);
+		// Add this child with its children (if any) to the owner hierarchy
+		const UObject* Owner = NewDirectOwner;
+		while (Owner != nullptr)
+		{
+			TArray<TWeakObjectPtr<const UObject>>& NewHighestOwnerChildren = OwnerToChildrenMap.FindOrAdd(Owner);
+			NewHighestOwnerChildren.Add(FromWeakPtr);
+			if (ObjectChildren != nullptr)
+			{
+				NewHighestOwnerChildren.Append(*ObjectChildren);
+			}
+
+			const TWeakObjectPtr<const UObject>* WeakOwner = ChildToOwnerMap.Find(Owner);
+			Owner = WeakOwner ? WeakOwner->Get() : nullptr;
+		}
 	}
 
-	ChildToOwnerMap.FindOrAdd(FromWeakPtr.Get(true/*bEvenIfPendingKill*/)) = ToObject;
+	WeakOwnerEntry = ToObject;
 
+	UObject* NewRedirection = FindRedirectionInternal(ToObject);
 	return NewRedirection;
 }
 
diff --git a/Engine/Source/Runtime/Engine/Private/VisualLogger/VisualLoggerAutomationTests.cpp b/Engine/Source/Runtime/Engine/Private/VisualLogger/VisualLoggerAutomationTests.cpp
--- a/Engine/Source/Runtime/Engine/Private/VisualLogger/VisualLoggerAutomationTests.cpp
+++ b/Engine/Source/Runtime/Engine/Private/VisualLogger/VisualLoggerAutomationTests.cpp
@@ -337,6 +337,238 @@ bool FVisualLogEventsTest::RunTest(const FString& Parameters)
 	return true;
 }
 
+IMPLEMENT_SIMPLE_AUTOMATION_TEST(FVisualLogRedirectionsCleanupTest, "System.Engine.VisualLogger.Redirections.Cleanup", EAutomationTestFlags::ClientContext | EAutomationTestFlags::EditorContext | EAutomationTestFlags::EngineFilter)
+bool FVisualLogRedirectionsCleanupTest::RunTest(const FString& Parameters)
+{
+	UWorld* World = GetSimpleEngineAutomationTestWorld(GetTestFlags());
+	CHECK_NOT_NULL(World);
+
+	FVisualLogger& Logger = FVisualLogger::Get();
+	Logger.Cleanup(World);
+	FVisualLogger::FOwnerToChildrenRedirectionMap& RedirectionMap = Logger.GetRedirectionMap(World);
+	CHECK_SUCCESS(RedirectionMap.Num() == 0);
+
+	const UObject* ObjA = NewObject<AActor>(World->GetCurrentLevel(), TEXT("VLogTestObjA"), RF_Transient);
+	const UObject* ObjB = NewObject<AActor>(World->GetCurrentLevel(), TEXT("VLogTestObjB"), RF_Transient);
+	const UObject* ObjC = NewObject<AActor>(World->GetCurrentLevel(), TEXT("VLogTestObjC"), RF_Transient);
+
+	REDIRECT_OBJECT_TO_VLOG(ObjB, ObjA);
+	REDIRECT_OBJECT_TO_VLOG(ObjC, ObjB);
+	// C -> B -> A
+	CHECK_SUCCESS(RedirectionMap.Num() == 2);
+	CHECK_NOT_NULL(RedirectionMap.Find(ObjA));
+	CHECK_NOT_NULL(RedirectionMap.Find(ObjB));
+	CHECK_SUCCESS(RedirectionMap.Find(ObjC) == nullptr);
+	
+	Logger.Cleanup(World);
+	CHECK_SUCCESS(RedirectionMap.Num() == 0);
+
+	return true;
+}
+
+IMPLEMENT_SIMPLE_AUTOMATION_TEST(FVisualLogRedirectionsMultipleChildrenTest, "System.Engine.VisualLogger.Redirections.MultipleChildren", EAutomationTestFlags::ClientContext | EAutomationTestFlags::EditorContext | EAutomationTestFlags::EngineFilter)
+bool FVisualLogRedirectionsMultipleChildrenTest::RunTest(const FString& Parameters)
+{
+	UWorld* World = GetSimpleEngineAutomationTestWorld(GetTestFlags());
+	CHECK_NOT_NULL(World);
+
+	FVisualLogger& Logger = FVisualLogger::Get();
+	Logger.Cleanup(World);
+	FVisualLogger::FOwnerToChildrenRedirectionMap& RedirectionMap = Logger.GetRedirectionMap(World);
+	CHECK_SUCCESS(RedirectionMap.Num() == 0);
+
+	const UObject* ObjA = NewObject<AActor>(World->GetCurrentLevel(), TEXT("VLogTestObjA"), RF_Transient);
+	const UObject* ObjB = NewObject<AActor>(World->GetCurrentLevel(), TEXT("VLogTestObjB"), RF_Transient);
+	const UObject* ObjC = NewObject<AActor>(World->GetCurrentLevel(), TEXT("VLogTestObjC"), RF_Transient);
+	const UObject* ObjD = NewObject<AActor>(World->GetCurrentLevel(), TEXT("VLogTestObjD"), RF_Transient);
+	const UObject* ObjE = NewObject<AActor>(World->GetCurrentLevel(), TEXT("VLogTestObjE"), RF_Transient);
+	const UObject* ObjF = NewObject<AActor>(World->GetCurrentLevel(), TEXT("VLogTestObjF"), RF_Transient);
+
+	REDIRECT_OBJECT_TO_VLOG(ObjB, ObjA);
+	// B -> A
+	TArray<TWeakObjectPtr<const UObject>>* ChildrenOfA = RedirectionMap.Find(ObjA);
+	CHECK_NOT_NULL(ChildrenOfA);
+	CHECK_SUCCESS(ChildrenOfA->Num() == 1);
+	CHECK_SUCCESS(Logger.FindRedirection(ObjB) == ObjA);
+
+	REDIRECT_OBJECT_TO_VLOG(ObjC, ObjA);
+	// B -> A
+	// C -> A
+	CHECK_SUCCESS(ChildrenOfA->Num() == 2);
+	CHECK_SUCCESS(Logger.FindRedirection(ObjC) == ObjA);
+
+	REDIRECT_OBJECT_TO_VLOG(ObjD, ObjA);
+	// B -> A
+	// C -> A
+	// D -> A
+	CHECK_SUCCESS(ChildrenOfA->Num() == 3);
+	CHECK_SUCCESS(Logger.FindRedirection(ObjD) == ObjA);
+
+	return true;
+}
+
+IMPLEMENT_SIMPLE_AUTOMATION_TEST(FVisualLogRedirectionsCreationOrderTest, "System.Engine.VisualLogger.Redirections.CreationOrder", EAutomationTestFlags::ClientContext | EAutomationTestFlags::EditorContext | EAutomationTestFlags::EngineFilter)
+bool FVisualLogRedirectionsCreationOrderTest::RunTest(const FString& Parameters)
+{
+	UWorld* World = GetSimpleEngineAutomationTestWorld(GetTestFlags());
+	CHECK_NOT_NULL(World);
+
+	FVisualLogger& Logger = FVisualLogger::Get();
+	Logger.Cleanup(World);
+	FVisualLogger::FOwnerToChildrenRedirectionMap& RedirectionMap = Logger.GetRedirectionMap(World);
+	CHECK_SUCCESS(RedirectionMap.Num() == 0);
+
+	const UObject* ObjA = NewObject<AActor>(World->GetCurrentLevel(), TEXT("VLogTestObjA"), RF_Transient);
+	const UObject* ObjB = NewObject<AActor>(World->GetCurrentLevel(), TEXT("VLogTestObjB"), RF_Transient);
+	const UObject* ObjC = NewObject<AActor>(World->GetCurrentLevel(), TEXT("VLogTestObjC"), RF_Transient);
+	const UObject* ObjD = NewObject<AActor>(World->GetCurrentLevel(), TEXT("VLogTestObjD"), RF_Transient);
+
+	// Validate that redirection creation order doesn't affect the final list of children
+	REDIRECT_OBJECT_TO_VLOG(ObjB, ObjA);
+	// B -> A
+	REDIRECT_OBJECT_TO_VLOG(ObjC, ObjB);
+	// C -> B -> A
+	REDIRECT_OBJECT_TO_VLOG(ObjD, ObjC);
+	// D -> C -> B -> A
+	TArray<TWeakObjectPtr<const UObject>>* ChildrenOfA = RedirectionMap.Find(ObjA);
+	TArray<TWeakObjectPtr<const UObject>>* ChildrenOfB = RedirectionMap.Find(ObjB);
+	TArray<TWeakObjectPtr<const UObject>>* ChildrenOfC = RedirectionMap.Find(ObjC);
+	CHECK_NOT_NULL(ChildrenOfA);
+	CHECK_NOT_NULL(ChildrenOfB);
+	CHECK_NOT_NULL(ChildrenOfC);
+	CHECK_SUCCESS(ChildrenOfA->Num() == 3);
+	CHECK_SUCCESS(ChildrenOfB->Num() == 2);
+	CHECK_SUCCESS(ChildrenOfC->Num() == 1);
+	CHECK_SUCCESS(Logger.FindRedirection(ObjB) == ObjA);
+	CHECK_SUCCESS(Logger.FindRedirection(ObjC) == ObjA);
+	CHECK_SUCCESS(Logger.FindRedirection(ObjD) == ObjA);
+
+	Logger.Cleanup(World);
+
+	REDIRECT_OBJECT_TO_VLOG(ObjD, ObjC);
+	// D -> C
+	REDIRECT_OBJECT_TO_VLOG(ObjC, ObjB);
+	// D -> C -> B
+	REDIRECT_OBJECT_TO_VLOG(ObjB, ObjA);
+	// D -> C -> B -> A
+	ChildrenOfA = RedirectionMap.Find(ObjA);
+	ChildrenOfB = RedirectionMap.Find(ObjB);
+	ChildrenOfC = RedirectionMap.Find(ObjC);
+	CHECK_NOT_NULL(ChildrenOfA);
+	CHECK_NOT_NULL(ChildrenOfB);
+	CHECK_NOT_NULL(ChildrenOfC);
+	CHECK_SUCCESS(ChildrenOfA->Num() == 3);
+	CHECK_SUCCESS(ChildrenOfB->Num() == 2);
+	CHECK_SUCCESS(ChildrenOfC->Num() == 1);
+	CHECK_SUCCESS(Logger.FindRedirection(ObjB) == ObjA);
+	CHECK_SUCCESS(Logger.FindRedirection(ObjC) == ObjA);
+	CHECK_SUCCESS(Logger.FindRedirection(ObjD) == ObjA);
+
+	return true;
+}
+
+IMPLEMENT_SIMPLE_AUTOMATION_TEST(FVisualLogRedirectionsWithinHierachyTest, "System.Engine.VisualLogger.Redirections.WithinHierachy", EAutomationTestFlags::ClientContext | EAutomationTestFlags::EditorContext | EAutomationTestFlags::EngineFilter)
+bool FVisualLogRedirectionsWithinHierachyTest::RunTest(const FString& Parameters)
+{
+	UWorld* World = GetSimpleEngineAutomationTestWorld(GetTestFlags());
+	CHECK_NOT_NULL(World);
+
+	FVisualLogger& Logger = FVisualLogger::Get();
+	Logger.Cleanup(World);
+	FVisualLogger::FOwnerToChildrenRedirectionMap& RedirectionMap = Logger.GetRedirectionMap(World);
+	CHECK_SUCCESS(RedirectionMap.Num() == 0);
+
+	const UObject* ObjA = NewObject<AActor>(World->GetCurrentLevel(), TEXT("VLogTestObjA"), RF_Transient);
+	const UObject* ObjB = NewObject<AActor>(World->GetCurrentLevel(), TEXT("VLogTestObjB"), RF_Transient);
+	const UObject* ObjC = NewObject<AActor>(World->GetCurrentLevel(), TEXT("VLogTestObjC"), RF_Transient);
+	const UObject* ObjD = NewObject<AActor>(World->GetCurrentLevel(), TEXT("VLogTestObjD"), RF_Transient);
+	const UObject* ObjE = NewObject<AActor>(World->GetCurrentLevel(), TEXT("VLogTestObjE"), RF_Transient);
+
+	REDIRECT_OBJECT_TO_VLOG(ObjE, ObjD);
+	REDIRECT_OBJECT_TO_VLOG(ObjD, ObjB);
+	REDIRECT_OBJECT_TO_VLOG(ObjB, ObjA);
+	REDIRECT_OBJECT_TO_VLOG(ObjC, ObjA);
+	// E -> D -> B -> A
+	//           C -> A
+	TArray<TWeakObjectPtr<const UObject>>* ChildrenOfA = RedirectionMap.Find(ObjA);
+	TArray<TWeakObjectPtr<const UObject>>* ChildrenOfB = RedirectionMap.Find(ObjB);
+	TArray<TWeakObjectPtr<const UObject>>* ChildrenOfD = RedirectionMap.Find(ObjD);
+	CHECK_NOT_NULL(ChildrenOfA);
+	CHECK_NOT_NULL(ChildrenOfB);
+	CHECK_NOT_NULL(ChildrenOfD);
+	CHECK_SUCCESS(ChildrenOfA->Num() == 4);
+	CHECK_SUCCESS(ChildrenOfB->Num() == 2);
+	CHECK_SUCCESS(ChildrenOfD->Num() == 1);
+	CHECK_SUCCESS(Logger.FindRedirection(ObjB) == ObjA);
+	CHECK_SUCCESS(Logger.FindRedirection(ObjC) == ObjA);
+	CHECK_SUCCESS(Logger.FindRedirection(ObjD) == ObjA);
+	CHECK_SUCCESS(Logger.FindRedirection(ObjE) == ObjA);
+
+	REDIRECT_OBJECT_TO_VLOG(ObjD, ObjC);
+	//           B -> A
+	// E -> D -> C -> A
+	TArray<TWeakObjectPtr<const UObject>>* ChildrenOfC = RedirectionMap.Find(ObjC);
+	CHECK_NOT_NULL(ChildrenOfC);
+	CHECK_SUCCESS(ChildrenOfA->Num() == 4);
+	CHECK_SUCCESS(ChildrenOfB->Num() == 0);
+	CHECK_SUCCESS(ChildrenOfC->Num() == 2);
+	CHECK_SUCCESS(ChildrenOfD->Num() == 1);
+	CHECK_SUCCESS(Logger.FindRedirection(ObjB) == ObjA);
+	CHECK_SUCCESS(Logger.FindRedirection(ObjC) == ObjA);
+	CHECK_SUCCESS(Logger.FindRedirection(ObjD) == ObjA);
+	CHECK_SUCCESS(Logger.FindRedirection(ObjE) == ObjA);
+
+	return true;
+}
+
+IMPLEMENT_SIMPLE_AUTOMATION_TEST(FVisualLogRedirectionsDeepHierarchyToNewParentTest, "System.Engine.VisualLogger.Redirections.DeepHierarchyToNewParent", EAutomationTestFlags::ClientContext | EAutomationTestFlags::EditorContext | EAutomationTestFlags::EngineFilter)
+bool FVisualLogRedirectionsDeepHierarchyToNewParentTest::RunTest(const FString& Parameters)
+{
+	UWorld* World = GetSimpleEngineAutomationTestWorld(GetTestFlags());
+	CHECK_NOT_NULL(World);
+
+	FVisualLogger& Logger = FVisualLogger::Get();
+	Logger.Cleanup(World);
+	FVisualLogger::FOwnerToChildrenRedirectionMap& RedirectionMap = Logger.GetRedirectionMap(World);
+
+	const UObject* ObjA = NewObject<AActor>(World->GetCurrentLevel(), TEXT("VLogTestObjA"), RF_Transient);
+	const UObject* ObjB = NewObject<AActor>(World->GetCurrentLevel(), TEXT("VLogTestObjB"), RF_Transient);
+	const UObject* ObjC = NewObject<AActor>(World->GetCurrentLevel(), TEXT("VLogTestObjC"), RF_Transient);
+	const UObject* ObjD = NewObject<AActor>(World->GetCurrentLevel(), TEXT("VLogTestObjD"), RF_Transient);
+	const UObject* ObjE = NewObject<AActor>(World->GetCurrentLevel(), TEXT("VLogTestObjE"), RF_Transient);
+	const UObject* ObjF = NewObject<AActor>(World->GetCurrentLevel(), TEXT("VLogTestObjF"), RF_Transient);
+
+	REDIRECT_OBJECT_TO_VLOG(ObjB, ObjA);
+	REDIRECT_OBJECT_TO_VLOG(ObjC, ObjB);
+	REDIRECT_OBJECT_TO_VLOG(ObjD, ObjC);
+	REDIRECT_OBJECT_TO_VLOG(ObjE, ObjD);
+	// E -> D -> C -> B -> A
+	TArray<TWeakObjectPtr<const UObject>>* ChildrenOfA = RedirectionMap.Find(ObjA);
+	TArray<TWeakObjectPtr<const UObject>>* ChildrenOfB = RedirectionMap.Find(ObjB);
+	TArray<TWeakObjectPtr<const UObject>>* ChildrenOfC = RedirectionMap.Find(ObjC);
+	TArray<TWeakObjectPtr<const UObject>>* ChildrenOfD = RedirectionMap.Find(ObjD);
+	CHECK_NOT_NULL(ChildrenOfA);
+	CHECK_NOT_NULL(ChildrenOfB);
+	CHECK_NOT_NULL(ChildrenOfC);
+	CHECK_NOT_NULL(ChildrenOfD);
+	CHECK_SUCCESS(ChildrenOfA->Num() == 4);
+	CHECK_SUCCESS(ChildrenOfB->Num() == 3);
+	CHECK_SUCCESS(ChildrenOfC->Num() == 2);
+	CHECK_SUCCESS(ChildrenOfD->Num() == 1);
+	
+	REDIRECT_OBJECT_TO_VLOG(ObjB, ObjF);
+	// E -> D -> C -> B -> F
+	TArray<TWeakObjectPtr<const UObject>>* ChildrenOfF = RedirectionMap.Find(ObjF);
+	CHECK_NOT_NULL(ChildrenOfF);
+	CHECK_SUCCESS(ChildrenOfA->Num() == 0);
+	CHECK_SUCCESS(ChildrenOfF->Num() == 4);
+	CHECK_SUCCESS(ChildrenOfB->Num() == 3);
+	CHECK_SUCCESS(ChildrenOfC->Num() == 2);
+	CHECK_SUCCESS(ChildrenOfD->Num() == 1);
+
+	return true;
+}
+
 #undef CHECK_SUCCESS
 #undef CHECK_FAIL
 
