Adding MetaSound bus type to metasounds

diff --git a/Engine/Plugins/Runtime/Metasound/Source/MetasoundEngine/Private/MetasoundAudioBus.cpp b/Engine/Plugins/Runtime/Metasound/Source/MetasoundEngine/Private/MetasoundAudioBus.cpp
new file mode 100644
--- /dev/null
+++ b/Engine/Plugins/Runtime/Metasound/Source/MetasoundEngine/Private/MetasoundAudioBus.cpp
@@ -0,0 +1,20 @@
+// Copyright Epic Games, Inc. All Rights Reserved.
+
+#include "MetasoundAudioBus.h"
+
+#include "AudioDevice.h"
+#include "DecoderInputFactory.h"
+
+namespace Metasound
+{
+	FAudioBusAsset::FAudioBusAsset(const TUniquePtr<Audio::IProxyData>& InInitData)
+	{
+		if (InInitData.IsValid())
+		{
+			if (InInitData->CheckTypeCast<FAudioBusProxy>())
+			{
+				AudioBusProxy = MakeShared<FAudioBusProxy, ESPMode::ThreadSafe>(InInitData->GetAs<FAudioBusProxy>());
+			}
+		}
+	}
+}
diff --git a/Engine/Plugins/Runtime/Metasound/Source/MetasoundEngine/Private/MetasoundEngineModule.cpp b/Engine/Plugins/Runtime/Metasound/Source/MetasoundEngine/Private/MetasoundEngineModule.cpp
--- a/Engine/Plugins/Runtime/Metasound/Source/MetasoundEngine/Private/MetasoundEngineModule.cpp
+++ b/Engine/Plugins/Runtime/Metasound/Source/MetasoundEngine/Private/MetasoundEngineModule.cpp
@@ -11,6 +11,7 @@
 #include "MetasoundUObjectRegistry.h"
 #include "MetasoundWave.h"
 #include "MetasoundWaveTable.h"
+#include "MetasoundAudioBus.h"
 #include "Modules/ModuleManager.h"
 #include "Sound/AudioSettings.h"
 
@@ -19,6 +20,7 @@ DEFINE_LOG_CATEGORY(LogMetasoundEngine);
 REGISTER_METASOUND_DATATYPE(Metasound::FWaveAsset, "WaveAsset", Metasound::ELiteralType::UObjectProxy, USoundWave);
 REGISTER_METASOUND_DATATYPE(WaveTable::FWaveTable, "WaveTable", Metasound::ELiteralType::FloatArray)
 REGISTER_METASOUND_DATATYPE(Metasound::FWaveTableBankAsset, "WaveTableBankAsset", Metasound::ELiteralType::UObjectProxy, UWaveTableBank);
+REGISTER_METASOUND_DATATYPE(Metasound::FAudioBusAsset, "AudioBusAsset", Metasound::ELiteralType::UObjectProxy, UAudioBus);
 
 class FMetasoundEngineModule : public IMetasoundEngineModule
 {
diff --git a/Engine/Plugins/Runtime/Metasound/Source/MetasoundEngine/Public/MetasoundAudioBus.h b/Engine/Plugins/Runtime/Metasound/Source/MetasoundEngine/Public/MetasoundAudioBus.h
new file mode 100644
--- /dev/null
+++ b/Engine/Plugins/Runtime/Metasound/Source/MetasoundEngine/Public/MetasoundAudioBus.h
@@ -0,0 +1,37 @@
+// Copyright Epic Games, Inc. All Rights Reserved.
+
+#pragma once
+
+#include "MetasoundDataReference.h"
+#include "MetasoundDataTypeRegistrationMacro.h"
+#include "IAudioProxyInitializer.h"
+#include "Sound/AudioBus.h"
+
+namespace Metasound
+{
+	// Forward declare ReadRef
+	class FAudioBus;
+
+	// Metasound data type that holds onto a weak ptr. Mostly used as a placeholder until we have a proper proxy type.
+	class METASOUNDENGINE_API FAudioBusAsset
+	{
+		FAudioBusProxyPtr AudioBusProxy;
+		
+	public:
+
+		FAudioBusAsset() = default;
+		FAudioBusAsset(const FAudioBusAsset&) = default;
+		FAudioBusAsset& operator=(const FAudioBusAsset& Other) = default;
+
+		FAudioBusAsset(const TUniquePtr<Audio::IProxyData>& InInitData);
+
+		const FAudioBusProxyPtr& GetAudioBusProxy() const
+		{
+			return AudioBusProxy;
+		}
+		
+	};
+
+	DECLARE_METASOUND_DATA_REFERENCE_TYPES(FAudioBusAsset, METASOUNDENGINE_API, FAudioBusAssetTypeInfo, FAudioBusAssetReadRef, FAudioBusAssetWriteRef)
+}
+ 
diff --git a/Engine/Source/Runtime/Engine/Classes/Sound/AudioBus.h b/Engine/Source/Runtime/Engine/Classes/Sound/AudioBus.h
--- a/Engine/Source/Runtime/Engine/Classes/Sound/AudioBus.h
+++ b/Engine/Source/Runtime/Engine/Classes/Sound/AudioBus.h
@@ -2,8 +2,8 @@
 #pragma once
 
 #include "CoreMinimal.h"
-#include "UObject/ObjectMacros.h"
 #include "UObject/Object.h"
+#include "IAudioProxyInitializer.h"
 #include "AudioBus.generated.h"
 
 class FAudioDevice;
@@ -19,13 +19,36 @@ enum class EAudioBusChannels : uint8
 	SevenPointOne = 7
 };
 
+class FAudioBusProxy;
+
+using FAudioBusProxyPtr = TSharedPtr<FAudioBusProxy, ESPMode::ThreadSafe>;
+class ENGINE_API FAudioBusProxy final : public Audio::TProxyData<FAudioBusProxy>
+{
+public:
+	IMPL_AUDIOPROXY_CLASS(FAudioBusProxy);
+
+	explicit FAudioBusProxy(UAudioBus* InAudioBus);
+
+	FAudioBusProxy(const FAudioBusProxy& Other) = default;
+
+	virtual ~FAudioBusProxy() override {}
+
+	virtual Audio::IProxyDataPtr Clone() const override
+	{
+		return MakeUnique<FAudioBusProxy>(*this);
+	}
+
+	uint32 AudioBusId = INDEX_NONE;
+	int32 NumChannels = INDEX_NONE;
+};
+
 // Function to retrieve an audio bus buffer given a handle
 // static float* GetAudioBusBuffer(const FAudioBusHandle& AudioBusHandle);
 
 // An audio bus is an object which represents an audio patch cord. Audio can be sent to it. It can be sonified using USoundSourceBuses.
 // Instances of the audio bus are created in the audio engine. 
 UCLASS(ClassGroup = Sound, meta = (BlueprintSpawnableComponent))
-class ENGINE_API UAudioBus : public UObject
+class ENGINE_API UAudioBus : public UObject, public IAudioProxyDataFactory
 {
 	GENERATED_UCLASS_BODY()
 
@@ -33,7 +56,7 @@ public:
 
 	/** Number of channels to use for the Audio Bus. */
 	UPROPERTY(EditAnywhere, Category = BusProperties)
-	EAudioBusChannels AudioBusChannels;
+	EAudioBusChannels AudioBusChannels = EAudioBusChannels::Mono;
 
 	//~ Begin UObject
 	virtual void BeginDestroy() override;
@@ -47,4 +70,8 @@ public:
 	virtual void PostEditChangeProperty(FPropertyChangedEvent& PropertyChangedEvent) override;
 #endif // WITH_EDITOR
 	//~ End UObject Interface.
+
+	//~ Begin IAudioProxy Interface
+	virtual TUniquePtr<Audio::IProxyData> CreateNewProxyData(const Audio::FProxyDataInitParams& InitParams) override;
+	//~ End IAudioProxy Interface
 };
diff --git a/Engine/Source/Runtime/Engine/Private/AudioBus.cpp b/Engine/Source/Runtime/Engine/Private/AudioBus.cpp
--- a/Engine/Source/Runtime/Engine/Private/AudioBus.cpp
+++ b/Engine/Source/Runtime/Engine/Private/AudioBus.cpp
@@ -3,13 +3,20 @@
 
 #include "Sound/AudioBus.h"
 #include "AudioDeviceManager.h"
-#include "EngineGlobals.h"
 #include "Engine/Engine.h"
-#include "UObject/UObjectIterator.h"
-#include "ActiveSound.h"
 #include "AudioDevice.h"
 #include "AudioMixerDevice.h"
 
+FAudioBusProxy::FAudioBusProxy(UAudioBus* InAudioBus)
+{
+	if (InAudioBus)
+	{
+		AudioBusId = InAudioBus->GetUniqueID();
+		NumChannels = InAudioBus->GetNumChannels();
+	}
+}
+
+
 UAudioBus::UAudioBus(const FObjectInitializer& ObjectInitializer)
 	: Super(ObjectInitializer)
 {
@@ -50,46 +57,26 @@ void UAudioBus::PostEditChangeProperty(FPropertyChangedEvent& PropertyChangedEve
 		return;
 	}
 
-	FName PropertyName = PropertyChangedEvent.Property->GetFName();
-	if (PropertyName != GET_MEMBER_NAME_CHECKED(UAudioBus, AudioBusChannels))
+	const FName PropertyName = PropertyChangedEvent.Property->GetFName();
+	if (PropertyName == GET_MEMBER_NAME_CHECKED(UAudioBus, AudioBusChannels))
 	{
-		return;
-	}
-
-	const UAudioSettings* AudioSettings = GetDefault<UAudioSettings>();
-	if (!AudioSettings)
-	{
-		return;
-	}
-
-	for (const FDefaultAudioBusSettings& DefaultBusSettings : AudioSettings->DefaultAudioBuses)
-	{
-		if (DefaultBusSettings.AudioBus.ResolveObject() != this)
+		if (FAudioDeviceManager* DeviceManager = FAudioDeviceManager::Get())
 		{
-			continue;
-		}
-
-		FAudioDeviceManager* DeviceManager = FAudioDeviceManager::Get();
-		if (!DeviceManager)
-		{
-			continue;
-		}
-
-		// Restart bus with new channel count
-		UObject* AudioBus = DefaultBusSettings.AudioBus.ResolveObject();
-		if (this != AudioBus)
-		{
-			continue;
-		}
-
-		DeviceManager->IterateOverAllDevices([BusId = AudioBus->GetUniqueID(), NumChannels = AudioBusChannels](Audio::FDeviceId, FAudioDevice* InDevice)
-		{
-			if (Audio::FMixerDevice* MixerDevice = static_cast<Audio::FMixerDevice*>(InDevice))
+			DeviceManager->IterateOverAllDevices([BusId = GetUniqueID(), NumChannels = AudioBusChannels](Audio::FDeviceId, FAudioDevice* InDevice)
 			{
-				MixerDevice->StopAudioBus(BusId);
-				MixerDevice->StartAudioBus(BusId, (int32)NumChannels + 1, false /* bInIsAutomatic */);
-			}
-		});
+				if (Audio::FMixerDevice* MixerDevice = static_cast<Audio::FMixerDevice*>(InDevice))
+				{
+					MixerDevice->StopAudioBus(BusId);
+					MixerDevice->StartAudioBus(BusId, (int32)NumChannels + 1, false /* bInIsAutomatic */);
+				}
+			});
+		}
 	}
 }
 #endif // WITH_EDITOR
+
+ TUniquePtr<Audio::IProxyData> UAudioBus::CreateNewProxyData(const Audio::FProxyDataInitParams& InitParams)
+{
+	return MakeUnique<FAudioBusProxy>(this);
+}
+
