Remove Nanite static mesh memory from incompatible platforms

diff --git a/Engine/Source/Developer/MeshBuilder/Private/MeshBuilderModule.cpp b/Engine/Source/Developer/MeshBuilder/Private/MeshBuilderModule.cpp
--- a/Engine/Source/Developer/MeshBuilder/Private/MeshBuilderModule.cpp
+++ b/Engine/Source/Developer/MeshBuilder/Private/MeshBuilderModule.cpp
@@ -26,7 +26,7 @@ public:
 		// Unregister any modular features here
 	}
 
-	virtual bool BuildMesh(FStaticMeshRenderData& OutRenderData, UObject* Mesh, const FStaticMeshLODGroup& LODGroup, bool bGenerateCoarseMeshStreamingLODs) override;
+	virtual bool BuildMesh(FStaticMeshRenderData& OutRenderData, UObject* Mesh, const FStaticMeshLODGroup& LODGroup, bool bGenerateCoarseMeshStreamingLODs, bool bAllowNanite) override;
 
 	virtual bool BuildMeshVertexPositions(
 		UObject* StaticMesh,
@@ -41,13 +41,13 @@ private:
 
 IMPLEMENT_MODULE(FMeshBuilderModule, MeshBuilder );
 
-bool FMeshBuilderModule::BuildMesh(FStaticMeshRenderData& OutRenderData, class UObject* Mesh, const FStaticMeshLODGroup& LODGroup, bool bGenerateCoarseMeshStreamingLODs)
+bool FMeshBuilderModule::BuildMesh(FStaticMeshRenderData& OutRenderData, class UObject* Mesh, const FStaticMeshLODGroup& LODGroup, bool bGenerateCoarseMeshStreamingLODs, bool bAllowNanite)
 {
 	UStaticMesh* StaticMesh = Cast<UStaticMesh>(Mesh);
 	if (StaticMesh != nullptr)
 	{
 		//Call the static mesh builder
-		return FStaticMeshBuilder().Build(OutRenderData, StaticMesh, LODGroup, bGenerateCoarseMeshStreamingLODs);
+		return FStaticMeshBuilder().Build(OutRenderData, StaticMesh, LODGroup, bGenerateCoarseMeshStreamingLODs, bAllowNanite);
 	}
 	return false;
 }
diff --git a/Engine/Source/Developer/MeshBuilder/Private/StaticMeshBuilder.cpp b/Engine/Source/Developer/MeshBuilder/Private/StaticMeshBuilder.cpp
--- a/Engine/Source/Developer/MeshBuilder/Private/StaticMeshBuilder.cpp
+++ b/Engine/Source/Developer/MeshBuilder/Private/StaticMeshBuilder.cpp
@@ -260,9 +260,9 @@ static bool BuildNanite(
 }
 
 
-bool FStaticMeshBuilder::Build(FStaticMeshRenderData& StaticMeshRenderData, UStaticMesh* StaticMesh, const FStaticMeshLODGroup& LODGroup, bool bGenerateCoarseMeshStreamingLODs)
+bool FStaticMeshBuilder::Build(FStaticMeshRenderData& StaticMeshRenderData, UStaticMesh* StaticMesh, const FStaticMeshLODGroup& LODGroup, bool bGenerateCoarseMeshStreamingLODs, bool bAllowNanite)
 {
-	const bool bNaniteBuildEnabled = StaticMesh->NaniteSettings.bEnabled;
+	const bool bNaniteBuildEnabled = StaticMesh->NaniteSettings.bEnabled && bAllowNanite;
 	const bool bHaveHiResSourceModel = StaticMesh->IsHiResMeshDescriptionValid();
 	int32 NumTasks = (bNaniteBuildEnabled && bHaveHiResSourceModel) ? (StaticMesh->GetNumSourceModels() + 1) : (StaticMesh->GetNumSourceModels());
 	FScopedSlowTask SlowTask(NumTasks, NSLOCTEXT("StaticMeshEditor", "StaticMeshBuilderBuild", "Building static mesh render data."));
diff --git a/Engine/Source/Developer/MeshBuilder/Public/IMeshBuilderModule.h b/Engine/Source/Developer/MeshBuilder/Public/IMeshBuilderModule.h
--- a/Engine/Source/Developer/MeshBuilder/Public/IMeshBuilderModule.h
+++ b/Engine/Source/Developer/MeshBuilder/Public/IMeshBuilderModule.h
@@ -29,7 +29,7 @@ public:
 
 	virtual void AppendToDDCKey(FString& DDCKey, bool bSkeletal) { }
 
-	virtual bool BuildMesh(class FStaticMeshRenderData& OutRenderData, class UObject* Mesh, const class FStaticMeshLODGroup& LODGroup, bool bGenerateCoarseMeshStreamingLODs) = 0;
+	virtual bool BuildMesh(class FStaticMeshRenderData& OutRenderData, class UObject* Mesh, const class FStaticMeshLODGroup& LODGroup, bool bGenerateCoarseMeshStreamingLODs, bool bAllowNanite) = 0;
 
 	virtual bool BuildMeshVertexPositions(class UObject* StaticMesh, TArray<uint32>& Indices, TArray<FVector3f>& Vertices) = 0;
 
diff --git a/Engine/Source/Developer/MeshBuilder/Public/MeshBuilder.h b/Engine/Source/Developer/MeshBuilder/Public/MeshBuilder.h
--- a/Engine/Source/Developer/MeshBuilder/Public/MeshBuilder.h
+++ b/Engine/Source/Developer/MeshBuilder/Public/MeshBuilder.h
@@ -25,7 +25,13 @@ public:
 	/**
 	 * Build function should be override and is the starting point for static mesh builders
 	 */
-	virtual bool Build(FStaticMeshRenderData& OutRenderData, UStaticMesh* StaticMesh, const FStaticMeshLODGroup& LODGroup, bool bGenerateCoarseMeshStreamingLODs) = 0;
+	virtual bool Build(
+		FStaticMeshRenderData& OutRenderData,
+		UStaticMesh* StaticMesh,
+		const FStaticMeshLODGroup& LODGroup,
+		bool bGenerateCoarseMeshStreamingLODs,
+		bool bAllowNanite) = 0;
+
 	virtual bool BuildMeshVertexPositions(
 		UStaticMesh* StaticMesh,
 		TArray<uint32>& Indices,
diff --git a/Engine/Source/Developer/MeshBuilder/Public/SkeletalMeshBuilder.h b/Engine/Source/Developer/MeshBuilder/Public/SkeletalMeshBuilder.h
--- a/Engine/Source/Developer/MeshBuilder/Public/SkeletalMeshBuilder.h
+++ b/Engine/Source/Developer/MeshBuilder/Public/SkeletalMeshBuilder.h
@@ -18,7 +18,7 @@ public:
 	FSkeletalMeshBuilder();
 
 	//No support for static mesh build in this class
-	virtual bool Build(FStaticMeshRenderData& OutRenderData, UStaticMesh* StaticMesh, const FStaticMeshLODGroup& LODGroup, bool bGenerateCoarseMeshStreamingLODs) override
+	virtual bool Build(FStaticMeshRenderData& OutRenderData, UStaticMesh* StaticMesh, const FStaticMeshLODGroup& LODGroup, bool bGenerateCoarseMeshStreamingLODs, bool bAllowNanite) override
 	{
 		bool No_Support_For_StaticMesh_Build_In_FSkeletalMeshBuilder_Class = false;
 		check(No_Support_For_StaticMesh_Build_In_FSkeletalMeshBuilder_Class);
diff --git a/Engine/Source/Developer/MeshBuilder/Public/StaticMeshBuilder.h b/Engine/Source/Developer/MeshBuilder/Public/StaticMeshBuilder.h
--- a/Engine/Source/Developer/MeshBuilder/Public/StaticMeshBuilder.h
+++ b/Engine/Source/Developer/MeshBuilder/Public/StaticMeshBuilder.h
@@ -22,7 +22,12 @@ public:
 	FStaticMeshBuilder();
 	virtual ~FStaticMeshBuilder() {}
 
-	virtual bool Build(FStaticMeshRenderData& OutRenderData, UStaticMesh* StaticMesh, const FStaticMeshLODGroup& LODGroup, bool bGenerateCoarseMeshStreamingLODs) override;
+	virtual bool Build(
+		FStaticMeshRenderData& OutRenderData,
+		UStaticMesh* StaticMesh,
+		const FStaticMeshLODGroup& LODGroup,
+		bool bGenerateCoarseMeshStreamingLODs,
+		bool bAllowNanite) override;
 
 	//No support for skeletal mesh build in this class
 	virtual bool Build(const struct FSkeletalMeshBuildParameters& SkeletalMeshBuildParameters) override
diff --git a/Engine/Source/Runtime/Engine/Private/StaticMesh.cpp b/Engine/Source/Runtime/Engine/Private/StaticMesh.cpp
--- a/Engine/Source/Runtime/Engine/Private/StaticMesh.cpp
+++ b/Engine/Source/Runtime/Engine/Private/StaticMesh.cpp
@@ -174,6 +174,24 @@ namespace StaticMeshCookStats
 }
 #endif
 
+static bool DoesTargetPlatformSupportNanite(const ITargetPlatform* TargetPlatform)
+{
+	if (TargetPlatform != nullptr)
+	{
+		TArray<FName> DesiredShaderFormats;
+		TargetPlatform->GetAllTargetedShaderFormats(DesiredShaderFormats);
+		for (int32 FormatIndex = 0; FormatIndex < DesiredShaderFormats.Num(); FormatIndex++)
+		{
+			const EShaderPlatform ShaderPlatform = ShaderFormatToLegacyShaderPlatform(DesiredShaderFormats[FormatIndex]);
+			if (DoesPlatformSupportNanite(ShaderPlatform))
+			{
+				return true;
+			}
+		}
+	}
+
+	return false;
+}
 
 #if WITH_EDITOR
 static void FillMaterialName(const TArray<FStaticMaterial>& StaticMaterials, TMap<int32, FName>& OutMaterialMap)
@@ -2923,10 +2941,13 @@ void FStaticMeshRenderData::Cache(const ITargetPlatform* TargetPlatform, UStatic
 
 				// Only build the LODs if coarse mesh streaming is enabled
 				static auto* VarNaniteCoarseMeshStreaming = IConsoleManager::Get().FindConsoleVariable(TEXT("r.Nanite.CoarseMeshStreaming"));
-				const bool bNaniteCoareMeshStreamingEnabled = !VarNaniteCoarseMeshStreaming || VarNaniteCoarseMeshStreaming->GetInt() != 0;
+				const bool bNaniteCoarseMeshStreamingEnabled = !VarNaniteCoarseMeshStreaming || VarNaniteCoarseMeshStreaming->GetInt() != 0;
+				const bool bGenerateStreamingLODs = TargetPlatform->SupportsFeature(ETargetPlatformFeatures::MeshLODStreaming) && bNaniteCoarseMeshStreamingEnabled && LODGroup.IsLODStreamingSupported();
+
+				// Check if the target platform supports Nanite at all
+				const bool bAllowNanite = DoesTargetPlatformSupportNanite(TargetPlatform);
 
-				const bool bGenerateStreamingLODs = TargetPlatform->SupportsFeature(ETargetPlatformFeatures::MeshLODStreaming) && bNaniteCoareMeshStreamingEnabled && LODGroup.IsLODStreamingSupported();
-				if (!MeshBuilderModule.BuildMesh(*this, Owner, LODGroup, bGenerateStreamingLODs))
+				if (!MeshBuilderModule.BuildMesh(*this, Owner, LODGroup, bGenerateStreamingLODs, bAllowNanite))
 				{
 					UE_LOG(LogStaticMesh, Error, TEXT("Failed to build static mesh. See previous line(s) for details."));
 					return;
