Chaos
- add PhysicsAsset option to select linear or non-linear joint solver for RBAN (old assets default to non-linear, new to linear)
- optimize joint SLerp drive (asin approximation and quaternion to axes)
- add unit tests for AsinEst

diff --git a/Engine/Source/Editor/PhysicsAssetEditor/Private/PhysicsAssetEditorAnimInstanceProxy.cpp b/Engine/Source/Editor/PhysicsAssetEditor/Private/PhysicsAssetEditorAnimInstanceProxy.cpp
--- a/Engine/Source/Editor/PhysicsAssetEditor/Private/PhysicsAssetEditorAnimInstanceProxy.cpp
+++ b/Engine/Source/Editor/PhysicsAssetEditor/Private/PhysicsAssetEditorAnimInstanceProxy.cpp
@@ -85,6 +85,7 @@ bool FPhysicsAssetEditorAnimInstanceProxy::Evaluate_WithRoot(FPoseContext& Outpu
 			SolverSettings.FixedTimeStep,
 			SolverSettings.CullDistance,
 			SolverSettings.MaxDepenetrationVelocity,
+			SolverSettings.bUseLinearJointSolver,
 			SolverSettings.PositionIterations,
 			SolverSettings.VelocityIterations,
 			SolverSettings.ProjectionIterations);
diff --git a/Engine/Source/Programs/HeadlessChaos/Private/HeadlessChaosTestMath.cpp b/Engine/Source/Programs/HeadlessChaos/Private/HeadlessChaosTestMath.cpp
--- a/Engine/Source/Programs/HeadlessChaos/Private/HeadlessChaosTestMath.cpp
+++ b/Engine/Source/Programs/HeadlessChaos/Private/HeadlessChaosTestMath.cpp
@@ -30,6 +30,81 @@ namespace ChaosTest
 		}
 	}
 
+	template<int Order>
+	void TestAsinEst(const FReal ExpectedMaxError, const FReal SmallError, const FReal SmallErrorX)
+	{
+		const int32 NumSteps = 21;
+		for (int32 I = 0; I < NumSteps; ++I)
+		{
+			// X from 0 to +1 in steps
+			const FReal X = FMath::Clamp(FReal(I) / FReal(NumSteps - 1), FReal(0), FReal(1));
+			const FReal Y = Utilities::AsinEst<FReal, Order>(X);
+			const FReal YNeg = Utilities::AsinEst<FReal, Order>(-X);
+			const FReal YExact = FMath::Asin(X);
+
+			// Asin(X) == -Asin(-X)
+			EXPECT_NEAR(Y, -YNeg, UE_SMALL_NUMBER);
+
+			// Total error percent less than expected
+			const FReal Error = (YExact > 0) ? FMath::Abs(Y - YExact) / YExact : UE_SMALL_NUMBER;
+			EXPECT_LT(Error, ExpectedMaxError);
+
+			// Error percent should be less than SmallError for X < ExpectedXAtSmallError
+			if (Error > SmallError)
+			{
+				EXPECT_GT(X, SmallErrorX - UE_SMALL_NUMBER);
+			}
+		}
+	}
+
+	TEST(MathTests, TestAsinEst3)
+	{
+		const FReal MaxError = FReal(0.26);			// expected error - 26% at X=1
+		const FReal SmallError = FReal(0.01);		// expected error - 1% at X=0.6
+		const FReal SmallErrorX = FReal(0.6);
+
+		TestAsinEst<3>(MaxError, SmallError, SmallErrorX);
+	}
+
+	TEST(MathTests, TestAsinEst5)
+	{
+		const FReal MaxError = FReal(0.21);			// expected error - 21% at X=1
+		const FReal SmallError = FReal(0.01);		// expected error - 1% at X=0.75
+		const FReal SmallErrorX = FReal(0.75);
+
+		TestAsinEst<5>(MaxError, SmallError, SmallErrorX);
+	}
+
+	TEST(MathTests, TestAsinEst7)
+	{
+		const FReal MaxError = FReal(0.19);			// expected error - 19% at X=1
+		const FReal SmallError = FReal(0.01);		// expected error - 1% at X=0.8
+		const FReal SmallErrorX = FReal(0.8);
+
+		TestAsinEst<7>(MaxError, SmallError, SmallErrorX);
+	}
+
+	TEST(MathTests, TestAsinEstCrossover)
+	{
+		const FReal MaxError = FReal(0.01);	// 1%
+		const int32 NumSteps = 21;
+		for (int32 I = 0; I < NumSteps; ++I)
+		{
+			// X from 0 to +1 in steps
+			const FReal X = FMath::Clamp(FReal(I) / FReal(NumSteps - 1), FReal(0), FReal(1));
+			const FReal Y = Utilities::AsinEstCrossover(X);
+			const FReal YNeg = Utilities::AsinEstCrossover(-X);
+			const FReal YExact = FMath::Asin(X);
+
+			// F(X) == -F(-X)
+			EXPECT_NEAR(Y, -YNeg, UE_SMALL_NUMBER);
+
+			// Total error less than expected
+			const FReal Error = (YExact > 0) ? FMath::Abs(Y - YExact) / YExact : UE_SMALL_NUMBER;
+			EXPECT_LT(Error, MaxError);
+		}
+	}
+
 	Chaos::TVector<double, 4> 
 	ToVec4(const TArray<double>& x)
 	{
diff --git a/Engine/Source/Runtime/AnimGraphRuntime/Private/BoneControllers/AnimNode_RigidBody.cpp b/Engine/Source/Runtime/AnimGraphRuntime/Private/BoneControllers/AnimNode_RigidBody.cpp
--- a/Engine/Source/Runtime/AnimGraphRuntime/Private/BoneControllers/AnimNode_RigidBody.cpp
+++ b/Engine/Source/Runtime/AnimGraphRuntime/Private/BoneControllers/AnimNode_RigidBody.cpp
@@ -844,6 +844,7 @@ void FAnimNode_RigidBody::EvaluateSkeletalControl_AnyThread(FComponentSpacePoseC
 				SolverSettings.FixedTimeStep,
 				SolverSettings.CullDistance,
 				SolverSettings.MaxDepenetrationVelocity,
+				SolverSettings.bUseLinearJointSolver,
 				SolverSettings.PositionIterations,
 				SolverSettings.VelocityIterations,
 				SolverSettings.ProjectionIterations);
@@ -1274,6 +1275,7 @@ void FAnimNode_RigidBody::InitPhysics(const UAnimInstance* InAnimInstance)
 			SolverSettings.FixedTimeStep,
 			SolverSettings.CullDistance,
 			SolverSettings.MaxDepenetrationVelocity,
+			SolverSettings.bUseLinearJointSolver,
 			SolverSettings.PositionIterations,
 			SolverSettings.VelocityIterations,
 			SolverSettings.ProjectionIterations);
diff --git a/Engine/Source/Runtime/Core/Public/UObject/FortniteNCBranchObjectVersion.h b/Engine/Source/Runtime/Core/Public/UObject/FortniteNCBranchObjectVersion.h
--- a/Engine/Source/Runtime/Core/Public/UObject/FortniteNCBranchObjectVersion.h
+++ b/Engine/Source/Runtime/Core/Public/UObject/FortniteNCBranchObjectVersion.h
@@ -36,6 +36,9 @@ struct CORE_API FFortniteNCBranchObjectVersion
 		// Added ability to mark shapes as probes
 		AddShapeIsProbe,
 
+		// Transfer PhysicsAsset SolverSettings (iteration counts etc) to new structure
+		PhysicsAssetNewSolverSettings,
+
 		// -----<new versions can be added above this line>-------------------------------------------------
 		VersionPlusOne,
 		LatestVersion = VersionPlusOne - 1
diff --git a/Engine/Source/Runtime/Engine/Classes/PhysicsEngine/PhysicsAsset.h b/Engine/Source/Runtime/Engine/Classes/PhysicsEngine/PhysicsAsset.h
--- a/Engine/Source/Runtime/Engine/Classes/PhysicsEngine/PhysicsAsset.h
+++ b/Engine/Source/Runtime/Engine/Classes/PhysicsEngine/PhysicsAsset.h
@@ -70,6 +70,15 @@ struct ENGINE_API FPhysicsAssetSolverSettings
 	 */
 	UPROPERTY(EditAnywhere, BlueprintReadWrite, Category = SolverSettings, meta = (ClampMin = 0))
 	float FixedTimeStep;
+
+	/**
+	 * Whether to use the linear or non-linear solver for RBAN Joints. The linear solver is significantly cheaper than
+	 * the non-linear solver when you are running multiple iterations, but is more likely to suffer from jitter. 
+	 * In general you should try to use the linear solver and increase the PositionIterations to improve stability if 
+	 * possible, only using the non-linear solver as a last resort.
+	*/
+	UPROPERTY(EditAnywhere, BlueprintReadWrite, Category = SolverSettings)
+	bool bUseLinearJointSolver;
 };
 
 
diff --git a/Engine/Source/Runtime/Engine/Private/PhysicsEngine/ImmediatePhysics/ImmediatePhysicsChaos/ImmediatePhysicsSimulation_Chaos.cpp b/Engine/Source/Runtime/Engine/Private/PhysicsEngine/ImmediatePhysics/ImmediatePhysicsChaos/ImmediatePhysicsSimulation_Chaos.cpp
--- a/Engine/Source/Runtime/Engine/Private/PhysicsEngine/ImmediatePhysics/ImmediatePhysicsChaos/ImmediatePhysicsSimulation_Chaos.cpp
+++ b/Engine/Source/Runtime/Engine/Private/PhysicsEngine/ImmediatePhysics/ImmediatePhysicsChaos/ImmediatePhysicsSimulation_Chaos.cpp
@@ -164,8 +164,8 @@ int32 ChaosImmediate_SolverType = (int32)Chaos::EConstraintSolverType::QuasiPbd;
 FAutoConsoleVariableRef CVarChaosImmPhysSolverType(TEXT("p.Chaos.ImmPhys.SolverType"), ChaosImmediate_SolverType, TEXT("0 = None; 1 = GbfPbd; 2 = Pbd; 3 = QuasiPbd"));
 
 // Whether to use the linear joint solver which is significantly faster than the non-linear one but less accurate. Only applies to the QuasiPBD Solver
-bool bChaosImmediate_Joint_UseLinearSolver = true;
-FAutoConsoleVariableRef CVarChaosImmPhysJointUseCachedSolver(TEXT("p.Chaos.ImmPhys.Joint.UseLinearSolver"), bChaosImmediate_Joint_UseLinearSolver, TEXT("Use linear version of joint solver. (default is true"));
+int32 ChaosImmediate_Joint_UseLinearSolver = -1;
+FAutoConsoleVariableRef CVarChaosImmPhysJointUseCachedSolver(TEXT("p.Chaos.ImmPhys.Joint.UseLinearSolver"), ChaosImmediate_Joint_UseLinearSolver, TEXT("Force use of linear or non-linear joint solver. (-1 to use PhysicsAsset setting)"));
 
 //
 // end remove when finished
@@ -742,7 +742,7 @@ namespace ImmediatePhysics_Chaos
 		SimSpaceSettings.EulerAlpha = ChaosImmediate_Evolution_SimSpaceEulerAlpha;
 	}
 
-	void FSimulation::SetSolverSettings(const FReal FixedDt, const FReal CullDistance, const FReal MaxDepenetrationVelocity, const int32 PositionIts, const int32 VelocityIts, const int32 ProjectionIts)
+	void FSimulation::SetSolverSettings(const FReal FixedDt, const FReal CullDistance, const FReal MaxDepenetrationVelocity, const int32 UseLinearJointSolver, const int32 PositionIts, const int32 VelocityIts, const int32 ProjectionIts)
 	{
 		if (FixedDt >= FReal(0))
 		{
@@ -773,6 +773,11 @@ namespace ImmediatePhysics_Chaos
 		{
 			Implementation->Evolution.SetNumProjectionIterations(ProjectionIts);
 		}
+
+		if (UseLinearJointSolver >= 0)
+		{
+			Implementation->Joints.SetUseLinearJointSolver(UseLinearJointSolver != 0);
+		}
 	}
 
 	void FSimulation::SetLegacySolverSettings(const int32 SolverIts, const int32 JointIts, const int32 CollisionIts, const int32 SolverPushOutIts, const int32 JointPushOutIts, const int32 CollisionPushOutIts)
@@ -849,6 +854,7 @@ namespace ImmediatePhysics_Chaos
 					ChaosImmediate_Evolution_FixedStepTime,
 					ChaosImmediate_Collision_CullDistance,
 					ChaosImmediate_Collision_MaxDepenetrationVelocity,
+					ChaosImmediate_Joint_UseLinearSolver,
 					ChaosImmediate_Evolution_PositionIterations,
 					ChaosImmediate_Evolution_VelocityIterations,
 					ChaosImmediate_Evolution_ProjectionIterations);
@@ -889,11 +895,10 @@ namespace ImmediatePhysics_Chaos
 			JointsSettings.LinearDriveDampingOverride = ChaosImmediate_Joint_LinearDriveDamping;
 			JointsSettings.AngularDriveStiffnessOverride = ChaosImmediate_Joint_AngularDriveStiffness;
 			JointsSettings.AngularDriveDampingOverride = ChaosImmediate_Joint_AngularDriveDamping;
-			JointsSettings.bUseLinearSolver = false;
 			if (SolverType == EConstraintSolverType::QuasiPbd)
 			{
+				// NOTE: bUseLinearSolver is set via SetSolverSettings above for QPBD
 				JointsSettings.NumShockPropagationIterations = ChaosImmediate_Joint_NumShockPropagationIterations;
-				JointsSettings.bUseLinearSolver = bChaosImmediate_Joint_UseLinearSolver;
 			}
 			else
 			{
diff --git a/Engine/Source/Runtime/Engine/Private/PhysicsEngine/PhysicsAsset.cpp b/Engine/Source/Runtime/Engine/Private/PhysicsEngine/PhysicsAsset.cpp
--- a/Engine/Source/Runtime/Engine/Private/PhysicsEngine/PhysicsAsset.cpp
+++ b/Engine/Source/Runtime/Engine/Private/PhysicsEngine/PhysicsAsset.cpp
@@ -15,6 +15,7 @@
 #include "UObject/ReleaseObjectVersion.h"
 #include "Logging/MessageLog.h"
 #include "UObject/UObjectIterator.h"
+#include "UObject/FortniteNCBranchObjectVersion.h"
 
 #if WITH_EDITOR
 #include "Misc/MessageDialog.h"
@@ -32,6 +33,7 @@ FPhysicsAssetSolverSettings::FPhysicsAssetSolverSettings()
 	, CullDistance(3.0f)
 	, MaxDepenetrationVelocity(0.0f)
 	, FixedTimeStep(0.0f)
+	, bUseLinearJointSolver(true)
 {
 }
 
@@ -207,6 +209,8 @@ void UPhysicsAsset::PostLoad()
 
 void UPhysicsAsset::Serialize(FArchive& Ar)
 {
+	Ar.UsingCustomVersion(FFortniteNCBranchObjectVersion::GUID);
+
 	Super::Serialize(Ar);
 	Ar << CollisionDisableTable;
 
@@ -218,14 +222,17 @@ void UPhysicsAsset::Serialize(FArchive& Ar)
 	}
 #endif
 
-	// If the legacy iteration defaults were changed, but the new settings were not, pass the legacy changes on to the new
-	if ((SolverSettings.PositionIterations == 6) && ((SolverIterations.SolverIterations != 3) || (SolverIterations.JointIterations != 2)))
+	// Transfer the legacy solver iteration counts to the new ones. These settings are used by the RBAN solver.
+	// These settings are intended to be roughly equivalent to the previous settings (the new solver adds VelocityIterations
+	// and provides the linear/non-linear joint solver option). Any new PhysicsAssets will get the defaults (new linear solver etc)
+	const bool bRequiresSettingsTransfer = (Ar.CustomVer(FFortniteNCBranchObjectVersion::GUID) < FFortniteNCBranchObjectVersion::PhysicsAssetNewSolverSettings);
+	if (bRequiresSettingsTransfer)
 	{
 		SolverSettings.PositionIterations = SolverIterations.SolverIterations * SolverIterations.JointIterations;
-	}
-	if ((SolverSettings.VelocityIterations == 1) && ((SolverIterations.SolverPushOutIterations != 1) || (SolverIterations.JointIterations != 1)))
-	{
-		SolverSettings.VelocityIterations = SolverIterations.SolverPushOutIterations * SolverIterations.JointIterations;
+		SolverSettings.VelocityIterations = 1;
+		SolverSettings.ProjectionIterations = SolverIterations.SolverPushOutIterations;
+		SolverSettings.bUseLinearJointSolver = false;
+		SolverSettings.CullDistance = 1.0f;
 	}
 
 	Ar.UsingCustomVersion(FFrameworkObjectVersion::GUID);
diff --git a/Engine/Source/Runtime/Engine/Public/Physics/ImmediatePhysics/ImmediatePhysicsChaos/ImmediatePhysicsSimulation_Chaos.h b/Engine/Source/Runtime/Engine/Public/Physics/ImmediatePhysics/ImmediatePhysicsChaos/ImmediatePhysicsSimulation_Chaos.h
--- a/Engine/Source/Runtime/Engine/Public/Physics/ImmediatePhysics/ImmediatePhysicsChaos/ImmediatePhysicsSimulation_Chaos.h
+++ b/Engine/Source/Runtime/Engine/Public/Physics/ImmediatePhysics/ImmediatePhysicsChaos/ImmediatePhysicsSimulation_Chaos.h
@@ -81,6 +81,7 @@ namespace ImmediatePhysics_Chaos
 			const FReal FixedDt,
 			const FReal CullDistance,
 			const FReal MaxDepenetrationVelocity,
+			const int32 UseLinearJointSolver,
 			const int32 PositionIts,
 			const int32 VelocityIts,
 			const int32 ProjectionIts);
diff --git a/Engine/Source/Runtime/Experimental/Chaos/Private/Chaos/Joint/PBDJointCachedSolverGaussSeidel.cpp b/Engine/Source/Runtime/Experimental/Chaos/Private/Chaos/Joint/PBDJointCachedSolverGaussSeidel.cpp
--- a/Engine/Source/Runtime/Experimental/Chaos/Private/Chaos/Joint/PBDJointCachedSolverGaussSeidel.cpp
+++ b/Engine/Source/Runtime/Experimental/Chaos/Private/Chaos/Joint/PBDJointCachedSolverGaussSeidel.cpp
@@ -1700,27 +1700,27 @@ void FPBDJointCachedSolver::InitSLerpDrive(
 	TargetAngPos.EnforceShortestArcWith(R01);
 	const FRotation3 R1Error = TargetAngPos.Inverse() * R01;
 
-	// @todo(chaos): try approximation for Asin. Try (x), (x + (1/6) x^3), or (x + (1/2) x^7)
-	const FReal AxisAngles[3] = 
-	{ 
-		2.0f * FMath::Asin(R1Error.X), 
-		2.0f * FMath::Asin(R1Error.Y), 
-		2.0f * FMath::Asin(R1Error.Z) 
+	FReal AxisAngles[3] =
+	{
+		2.0f * Utilities::AsinEst(R1Error.X),
+		2.0f * Utilities::AsinEst(R1Error.Y),
+		2.0f * Utilities::AsinEst(R1Error.Z)
 	};
 
-	const FRotation3& AxesRotation = ConnectorRs[1];
-	const FVec3 Axes[3] = {
-		AxesRotation.GetAxisX(),
-		AxesRotation.GetAxisY(),
-		AxesRotation.GetAxisZ()
-	};
+	FVec3 Axes[3];
+	ConnectorRs[1].ToMatrixAxes(Axes[0], Axes[1], Axes[2]);
 
-	const FVec3 TargetAngVel = ConnectorRs[0] * JointSettings.AngularDriveVelocityTarget;
+	InitRotationConstraintDrive(0, Axes[0], Dt, AxisAngles[0]);
+	InitRotationConstraintDrive(1, Axes[1], Dt, AxisAngles[1]);
+	InitRotationConstraintDrive(2, Axes[2], Dt, AxisAngles[2]);
 
-	for (int32 AxisIndex = 0; AxisIndex < 3; ++AxisIndex)
+	// @todo(chaos): pass constraint target velocity into InitRotationConstraintDrive (it currently sets it ConstraintVX to 0)
+	if (!JointSettings.AngularDriveVelocityTarget.IsNearlyZero())
 	{
-		InitRotationConstraintDrive(AxisIndex, Axes[AxisIndex],  Dt, AxisAngles[AxisIndex]);
-		RotationDrives.ConstraintVX[AxisIndex] = FVec3::DotProduct(TargetAngVel, RotationDrives.ConstraintAxis[AxisIndex]);
+		const FVec3 TargetAngVel = ConnectorRs[0] * JointSettings.AngularDriveVelocityTarget;
+		RotationDrives.ConstraintVX[0] = FVec3::DotProduct(TargetAngVel, Axes[0]);
+		RotationDrives.ConstraintVX[1] = FVec3::DotProduct(TargetAngVel, Axes[1]);
+		RotationDrives.ConstraintVX[2] = FVec3::DotProduct(TargetAngVel, Axes[2]);
 	}
 }
 
diff --git a/Engine/Source/Runtime/Experimental/Chaos/Private/Chaos/Joint/PBDJointSolverGaussSeidel.cpp b/Engine/Source/Runtime/Experimental/Chaos/Private/Chaos/Joint/PBDJointSolverGaussSeidel.cpp
--- a/Engine/Source/Runtime/Experimental/Chaos/Private/Chaos/Joint/PBDJointSolverGaussSeidel.cpp
+++ b/Engine/Source/Runtime/Experimental/Chaos/Private/Chaos/Joint/PBDJointSolverGaussSeidel.cpp
@@ -1950,7 +1950,6 @@ namespace Chaos
 		}
 	}
 
-
 	void FPBDJointSolver::ApplySLerpDrive(
 		const FReal Dt,
 		const FPBDJointSolverSettings& SolverSettings,
@@ -1960,34 +1959,33 @@ namespace Chaos
 		const FReal AngularDriveDamping = FPBDJointUtilities::GetAngularSLerpDriveDamping(SolverSettings, JointSettings);
 		const bool bAccelerationMode = FPBDJointUtilities::GetDriveAccelerationMode(SolverSettings, JointSettings);
 
-		// NOTE: Slerp target velocity only works properly if we have a stiffness of zero.
 		const FRotation3 R01 = ConnectorRs[0].Inverse() * ConnectorRs[1];
 		FRotation3 TargetAngPos = JointSettings.AngularDrivePositionTarget;
 		TargetAngPos.EnforceShortestArcWith(R01);
 		const FRotation3 R1Error = TargetAngPos.Inverse() * R01;
-		FReal AxisAngles[3] = 
-		{ 
-			2.0f * FMath::Asin(R1Error.X), 
-			2.0f * FMath::Asin(R1Error.Y), 
-			2.0f * FMath::Asin(R1Error.Z) 
-		};
 
-		const FRotation3& AxesRotation = ConnectorRs[1];
-		const FVec3 Axes[3] = {
-			AxesRotation.GetAxisX(),
-			AxesRotation.GetAxisY(),
-			AxesRotation.GetAxisZ()
+		FReal AxisAngles[3] =
+		{ 
+			2.0f * Utilities::AsinEst(R1Error.X),
+			2.0f * Utilities::AsinEst(R1Error.Y),
+			2.0f * Utilities::AsinEst(R1Error.Z)
 		};
 
-		const FVec3 TargetAngVel = ConnectorRs[0] * JointSettings.AngularDriveVelocityTarget;
+		FVec3 Axes[3];
+		ConnectorRs[1].ToMatrixAxes(Axes[0], Axes[1], Axes[2]);
 
-		for (int32 AxisIndex = 0; AxisIndex < 3; ++AxisIndex)
+		FReal AxisAngVel[3] = {0, 0, 0};
+		if (!JointSettings.AngularDriveVelocityTarget.IsNearlyZero())
 		{
-			FReal ReturnedLambda = RotationDriveLambdas[AxisIndex];
-			const FReal AxisAngVel = FVec3::DotProduct(TargetAngVel, Axes[AxisIndex]);
-			ApplyRotationConstraintSoft(Dt, AngularDriveStiffness, AngularDriveDamping, bAccelerationMode, Axes[AxisIndex], AxisAngles[AxisIndex], AxisAngVel, ReturnedLambda);
-			RotationDriveLambdas[AxisIndex] = ReturnedLambda;
+			const FVec3 TargetAngVel = ConnectorRs[0] * JointSettings.AngularDriveVelocityTarget;
+			AxisAngVel[0] = FVec3::DotProduct(TargetAngVel, Axes[0]);
+			AxisAngVel[1] = FVec3::DotProduct(TargetAngVel, Axes[1]);
+			AxisAngVel[2] = FVec3::DotProduct(TargetAngVel, Axes[2]);
 		}
+
+		ApplyRotationConstraintSoft(Dt, AngularDriveStiffness, AngularDriveDamping, bAccelerationMode, Axes[0], AxisAngles[0], AxisAngVel[0], RotationDriveLambdas[0]);
+		ApplyRotationConstraintSoft(Dt, AngularDriveStiffness, AngularDriveDamping, bAccelerationMode, Axes[1], AxisAngles[1], AxisAngVel[1], RotationDriveLambdas[1]);
+		ApplyRotationConstraintSoft(Dt, AngularDriveStiffness, AngularDriveDamping, bAccelerationMode, Axes[2], AxisAngles[2], AxisAngVel[2], RotationDriveLambdas[2]);
 	}
 
 
diff --git a/Engine/Source/Runtime/Experimental/Chaos/Public/Chaos/PBDJointConstraints.h b/Engine/Source/Runtime/Experimental/Chaos/Public/Chaos/PBDJointConstraints.h
--- a/Engine/Source/Runtime/Experimental/Chaos/Public/Chaos/PBDJointConstraints.h
+++ b/Engine/Source/Runtime/Experimental/Chaos/Public/Chaos/PBDJointConstraints.h
@@ -121,6 +121,8 @@ namespace Chaos
 		const FPBDJointSolverSettings& GetSettings() const;
 		void SetSettings(const FPBDJointSolverSettings& InSettings);
 
+		void SetUseLinearJointSolver(const bool bInEnable) { Settings.bUseLinearSolver = bInEnable; }
+
 		void SetNumPairIterations(const int32 NumPairIterationss) { Settings.ApplyPairIterations = NumPairIterationss; }
 		void SetNumPushOutPairIterations(const int32 NumPairIterationss) { Settings.ApplyPushOutPairIterations = NumPairIterationss; }
 
diff --git a/Engine/Source/Runtime/Experimental/Chaos/Public/Chaos/Rotation.h b/Engine/Source/Runtime/Experimental/Chaos/Public/Chaos/Rotation.h
--- a/Engine/Source/Runtime/Experimental/Chaos/Public/Chaos/Rotation.h
+++ b/Engine/Source/Runtime/Experimental/Chaos/Public/Chaos/Rotation.h
@@ -112,6 +112,21 @@ namespace Chaos
 			return false;
 		}
 
+		/**
+		* Convert to a matrix and return as the 3 matrix axes.
+		*/
+		inline void ToMatrixAxes(TVector<FRealSingle, 3>& OutX, TVector<FRealSingle, 3>& OutY, TVector<FRealSingle, 3>& OutZ)
+		{
+			const FRealSingle x2 = X + X;    const FRealSingle y2 = Y + Y;    const FRealSingle z2 = Z + Z;
+			const FRealSingle xx = X * x2;   const FRealSingle xy = X * y2;   const FRealSingle xz = X * z2;
+			const FRealSingle yy = Y * y2;   const FRealSingle yz = Y * z2;   const FRealSingle zz = Z * z2;
+			const FRealSingle wx = W * x2;   const FRealSingle wy = W * y2;   const FRealSingle wz = W * z2;
+
+			OutX = TVector<FRealSingle, 3>(1.0f - (yy + zz), xy + wz, xz - wy);
+			OutY = TVector<FRealSingle, 3>(xy - wz, 1.0f - (xx + zz), yz + wx);
+			OutZ = TVector<FRealSingle, 3>(xz + wy, yz - wx, 1.0f - (xx + yy));
+		}
+
 		/**
 		 * Extract the Swing and Twist rotations, assuming that the Twist Axis is (1,0,0).
 		 * /see ToSwingTwist
@@ -316,6 +331,21 @@ namespace Chaos
 			return false;
 		}
 
+		/**
+		* Convert to a matrix and return as the 3 matrix axes.
+		*/
+		inline void ToMatrixAxes(TVector<FRealDouble, 3>& OutX, TVector<FRealDouble, 3>& OutY, TVector<FRealDouble, 3>& OutZ)
+		{
+			const FRealDouble x2 = X + X;    const FRealDouble y2 = Y + Y;    const FRealDouble z2 = Z + Z;
+			const FRealDouble xx = X * x2;   const FRealDouble xy = X * y2;   const FRealDouble xz = X * z2;
+			const FRealDouble yy = Y * y2;   const FRealDouble yz = Y * z2;   const FRealDouble zz = Z * z2;
+			const FRealDouble wx = W * x2;   const FRealDouble wy = W * y2;   const FRealDouble wz = W * z2;
+
+			OutX = TVector<FRealDouble, 3>(1.0f - (yy + zz), xy + wz, xz - wy);
+			OutY = TVector<FRealDouble, 3>(xy - wz, 1.0f - (xx + zz), yz + wx);
+			OutZ = TVector<FRealDouble, 3>(xz + wy, yz - wx, 1.0f - (xx + yy));
+		}
+
 		/**
 		 * Extract the Swing and Twist rotations, assuming that the Twist Axis is (1,0,0).
 		 * /see ToSwingTwist
diff --git a/Engine/Source/Runtime/Experimental/Chaos/Public/Chaos/Utilities.h b/Engine/Source/Runtime/Experimental/Chaos/Public/Chaos/Utilities.h
--- a/Engine/Source/Runtime/Experimental/Chaos/Public/Chaos/Utilities.h
+++ b/Engine/Source/Runtime/Experimental/Chaos/Public/Chaos/Utilities.h
@@ -688,6 +688,101 @@ namespace Chaos
 			}
 		}
 
+		/**
+		 * Approximate Asin(X) to 1st order ~= X + ...
+		 * Returns an approximation for all valid input range [-1, 1] with an error that increases as the input magnitude approaches 1.
+		 */
+		template<typename T>
+		inline T AsinEst1(const T X)
+		{
+			return X /*+...*/;
+		}
+
+
+		/**
+		 * Approximate Asin(X) to 3rd order ~= X + (1/6)X^3 + ...
+		 * Returns an approximation for all valid input range [-1, 1] with an error that increases as the input magnitude approaches 1.
+		 */
+		template<typename T>
+		inline T AsinEst3(const T X)
+		{
+			constexpr T C0 = T(1.0);
+			constexpr T C1 = T(1.0 / 6.0);
+			const T X2 = X * X;
+			return X * (C0 + X2 * (C1 /*+...*/));
+		}
+
+		/**
+		 * Approximate Asin(X) to 5th order ~= X + (1/6)X^3 + (3/40)X^5 + ...
+		 * Returns an approximation for all valid input range [-1, 1] with an error that increases as the input magnitude approaches 1.
+		 */
+		template<typename T>
+		inline T AsinEst5(const T X)
+		{
+			constexpr T C0 = T(1.0);
+			constexpr T C1 = T(1.0 / 6.0);
+			constexpr T C2 = T(3.0 / 40.0);
+			const T X2 = X * X;
+			return X * (C0 + X2 * (C1 + X2 * (C2 /*+...*/)));
+
+		}
+
+		/**
+		 * Approximate Asin(X) to 7th order ~= X + (1/6)X^3 + (3/40)X^5 + (15/336)X^7 + ...
+		 * Returns an approximation for all valid input range [-1, 1] with an error that increases as the input magnitude approaches 1.
+		 */
+		template<typename T>
+		inline T AsinEst7(const T X)
+		{
+			constexpr T C0 = T(1.0);
+			constexpr T C1 = T(1.0 / 6.0);
+			constexpr T C2 = T(3.0 / 40.0);
+			constexpr T C3 = T(15.0 / 336.0);
+			const T X2 = X * X;
+			return X * (C0 + X2 * (C1 + X2 * (C2 + X2 * (C3 /*+...*/))));
+		}
+
+		/**
+		 * Approximate Asin using expansion to the specified Order (must be 1, 3, 5, or 7). Defaults to 5th order.
+		*/
+		template<typename T, int Order = 5>
+		inline T AsinEst(const T X)
+		{
+			static_assert((Order == 1) || (Order == 3) || (Order == 5) || (Order == 7), "AsinEst: Only 1, 3, 5, or 7 is supported for the Order");
+			if (Order == 1)
+			{
+				return AsinEst1(X);
+			}
+			else if (Order == 3)
+			{
+				return AsinEst3(X);
+			}
+			else if (Order == 5)
+			{
+				return AsinEst5(X);
+			}
+			else
+			{
+				return AsinEst7(X);
+			}
+		}
+
+		/**
+		 * Approximate Asin(X). Like AsinApprox but with a crossover to FMath::Asin for larger input magnitudes.
+		 */
+		template<typename T, int Order = 5>
+		inline T AsinEstCrossover(const T X, const T Crossover = T(0.7))
+		{
+			if (FMath::Abs(X) < Crossover)
+			{
+				return AsinEst<T, Order>(X);
+			}
+			else
+			{
+				return FMath::Asin(X);
+			}
+		}
+
 		template <class T, class TV, class TV_INT4>
 		void TetMeshFromGrid(const TUniformGrid<T, 3>& Grid, TArray<TV_INT4>& Mesh, TArray<TV>& X)
 		{
