[KillPendingKill] Release pointers to garbage actors pre-GC.

diff --git a/Engine/Source/Runtime/Engine/Classes/Components/SceneCaptureComponent.h b/Engine/Source/Runtime/Engine/Classes/Components/SceneCaptureComponent.h
--- a/Engine/Source/Runtime/Engine/Classes/Components/SceneCaptureComponent.h
+++ b/Engine/Source/Runtime/Engine/Classes/Components/SceneCaptureComponent.h
@@ -126,6 +126,8 @@ public:
 	UPROPERTY(EditAnywhere, interp, Category = SceneCapture)
 	FString ProfilingEventName;
 
+	virtual void BeginDestroy() override;
+
 	//~ Begin UActorComponent Interface
 	virtual void OnComponentDestroyed(bool bDestroyingHierarchy) override;
 	virtual void OnRegister() override;
@@ -202,6 +204,8 @@ protected:
 	/** Update the show flags from our show flags settings (ideally, you'd be able to set this more directly, but currently unable to make FEngineShowFlags a UStruct to use it as a FProperty...) */
 	void UpdateShowFlags();
 
+	void ReleaseGarbageReferences();
+
 	virtual void UpdateSceneCaptureContents(FSceneInterface* Scene) {};
 
 	/**
diff --git a/Engine/Source/Runtime/Engine/Private/Components/SceneCaptureComponent.cpp b/Engine/Source/Runtime/Engine/Private/Components/SceneCaptureComponent.cpp
--- a/Engine/Source/Runtime/Engine/Private/Components/SceneCaptureComponent.cpp
+++ b/Engine/Source/Runtime/Engine/Private/Components/SceneCaptureComponent.cpp
@@ -166,6 +166,35 @@ USceneCaptureComponent::USceneCaptureComponent(const FObjectInitializer& ObjectI
 	ShowFlags.SetSeparateTranslucency(0);
 	ShowFlags.SetHMDDistortion(0);
 	ShowFlags.SetOnScreenDebug(0);
+
+	if (!HasAnyFlags(RF_ClassDefaultObject))
+	{
+		FCoreUObjectDelegates::GetPreGarbageCollectDelegate().AddUObject(this, &USceneCaptureComponent::ReleaseGarbageReferences);
+	}
+}
+
+void USceneCaptureComponent::BeginDestroy()
+{
+	Super::BeginDestroy();
+
+	FCoreUObjectDelegates::GetPreGarbageCollectDelegate().RemoveAll(this);
+}
+
+// Because HiddenActors and ShowOnlyActors are serialized they are not easily refactored into weak pointers.
+// Before GC runs, release pointers to any actors that have been explicitly marked garbage.
+void USceneCaptureComponent::ReleaseGarbageReferences()
+{
+	// We only want to remove resolved actors that are explicitly marked as garbage, not unresolved pointers.
+	auto Predicate = [](TObjectPtr<AActor> Ptr)
+	{
+		if (AActor* Actor = Ptr.Get())
+		{
+			return Actor->HasAnyInternalFlags(EInternalObjectFlags::Garbage);
+		}
+		return false;
+	};
+	HiddenActors.RemoveAll(Predicate);
+	ShowOnlyActors.RemoveAll(Predicate);
 }
 
 void USceneCaptureComponent::OnRegister()
