Allow FMemberReference to optionally follow CoreRedirect data at runtime

This can be useful to resolve cached member references within saved UGC data, where those references may have become stale since the UGC was saved

#rnx

diff --git a/Engine/Source/Runtime/Engine/Classes/Engine/MemberReference.h b/Engine/Source/Runtime/Engine/Classes/Engine/MemberReference.h
--- a/Engine/Source/Runtime/Engine/Classes/Engine/MemberReference.h
+++ b/Engine/Source/Runtime/Engine/Classes/Engine/MemberReference.h
@@ -353,15 +353,13 @@ public:
 	/** 
 	 *	Returns the member FProperty/UFunction this reference is pointing to, or NULL if it no longer exists 
 	 *	Derives 'self' scope from supplied Blueprint node if required
-	 *	Will check for redirects and fix itself up if one is found.
+	 *	Will check for redirects and fix itself up if one is found (when WITH_EDITOR, or when bAlwaysFollowRedirects is true).
 	 */
 	template<class TFieldType>
-	TFieldType* ResolveMember(UClass* SelfScope = nullptr) const
+	TFieldType* ResolveMember(UClass* SelfScope = nullptr, const bool bAlwaysFollowRedirects = false) const
 	{
 		TFieldType* ReturnField = nullptr;
 
-		bool bUseUpToDateClass = SelfScope && SelfScope->GetAuthoritativeClass() != SelfScope;
-
 		if(bSelfContext && SelfScope == nullptr)
 		{
 			UE_LOG(LogBlueprint, Warning, TEXT("FMemberReference::ResolveMember (%s) bSelfContext == true, but no scope supplied!"), *MemberName.ToString() );
@@ -389,10 +387,15 @@ public:
 		}
 		else
 		{
+#if WITH_EDITOR
+			const bool bCanFollowRedirects = !GIsSavingPackage;
+#else
+			const bool bCanFollowRedirects = bAlwaysFollowRedirects;
+#endif
+
 			// Look for remapped member
 			UClass* TargetScope = GetScope(SelfScope);
-#if WITH_EDITOR
-			if( TargetScope != nullptr &&  !GIsSavingPackage )
+			if (bCanFollowRedirects && TargetScope)
 			{
 				ReturnField = FindRemappedField<TFieldType>(TargetScope, MemberName, true);
 			}
@@ -404,26 +407,38 @@ public:
 				MemberParent = Cast<UClass>(GetFieldOuter(static_cast<typename TFieldType::BaseFieldClass*>(ReturnField)));
 
 				MemberGuid.Invalidate();
+#if WITH_EDITOR
 				UBlueprint::GetGuidFromClassByFieldName<TFieldType>(TargetScope, MemberName, MemberGuid);
+#endif
 
 				if (UClass* ParentAsClass = GetMemberParentClass())
 				{
+#if WITH_EDITOR
 					ParentAsClass = ParentAsClass->GetAuthoritativeClass();
+#endif
 					MemberParent  = ParentAsClass;
 
 					// Re-evaluate self-ness against the redirect if we were given a valid SelfScope
 					// For functions and multicast delegates we don't want to go from not-self to self as the target pin type should remain consistent
 					if (SelfScope != nullptr && (bSelfContext || (!CompareClassesHelper(TFieldType::StaticClass(), UFunction::StaticClass()) && !CompareClassesHelper(TFieldType::StaticClass(), FMulticastDelegateProperty::StaticClass()))))
 					{
-						SetGivenSelfScope(MemberName, MemberGuid, ParentAsClass, SelfScope);
+#if WITH_EDITOR
+						bSelfContext = SelfScope->IsChildOf(ParentAsClass) || SelfScope->ClassGeneratedBy == ParentAsClass->ClassGeneratedBy;
+#else
+						bSelfContext = SelfScope->IsChildOf(ParentAsClass);
+#endif
+
+						if (bSelfContext)
+						{
+							MemberParent = nullptr;
+						}
 					}
 				}	
 			}
-			else
-#endif
-			if (TargetScope != nullptr)
+			else if (TargetScope != nullptr)
 			{
 #if WITH_EDITOR
+				bool bUseUpToDateClass = SelfScope && SelfScope->GetAuthoritativeClass() != SelfScope;
 				TargetScope = GetClassToUse(TargetScope, bUseUpToDateClass);
 				if (TargetScope)
 #endif
@@ -494,7 +509,6 @@ public:
 		return ResolveMember<TFieldType>(SelfScope->SkeletonGeneratedClass);
 	}
 
-#if WITH_EDITOR
 	/**
 	 * Searches the field redirect map for the specified named field in the scope, and returns the remapped field if found
 	 *
@@ -514,6 +528,7 @@ public:
 		return FFieldVariant(FindRemappedField(TFieldType::StaticClass(), InitialScope, InitialName, bInitialScopeMustBeOwnerOfField)).Get<TFieldType>();
 	}
 
+#if WITH_EDITOR
 	/** Init the field redirect map (if not already done) from .ini file entries */
 	ENGINE_API static void InitFieldRedirectMap();
 
diff --git a/Engine/Source/Runtime/Engine/Private/MemberReference.cpp b/Engine/Source/Runtime/Engine/Private/MemberReference.cpp
--- a/Engine/Source/Runtime/Engine/Private/MemberReference.cpp
+++ b/Engine/Source/Runtime/Engine/Private/MemberReference.cpp
@@ -276,12 +276,16 @@ UClass* FMemberReference::GetClassToUse(UClass* InClass, bool bUseUpToDateClass)
 	}
 }
 
+#endif
+
 template <typename TFieldType>
 TFieldType* FindRemappedFieldImpl(FName FieldClassOutermostName, FName FieldClassName, UClass* InitialScope, FName InitialName, bool bInitialScopeMustBeOwnerOfField)
 {
 	DECLARE_SCOPE_CYCLE_COUNTER(TEXT("FMemberReference::FindRemappedField"), STAT_LinkerLoad_FindRemappedField, STATGROUP_LoadTimeVerbose);
 
+#if WITH_EDITOR
 	FMemberReference::InitFieldRedirectMap();
+#endif
 
 	// In the case of a bifurcation of a variable (e.g. moved from a parent into certain children), verify that we don't also define the variable in the current scope first
 	if (FindUFieldOrFProperty<TFieldType>(InitialScope, InitialName))
@@ -367,5 +371,3 @@ FField* FMemberReference::FindRemappedField(FFieldClass* FieldClass, UClass* Ini
 {
 	return FindRemappedFieldImpl<FField>(GLongCoreUObjectPackageName, FieldClass->GetFName(), InitialScope, InitialName, bInitialScopeMustBeOwnerOfField);
 }
-
-#endif
