Add debug recording of camera shakes

diff --git a/Engine/Source/Runtime/Engine/Classes/Camera/CameraModifier_CameraShake.h b/Engine/Source/Runtime/Engine/Classes/Camera/CameraModifier_CameraShake.h
--- a/Engine/Source/Runtime/Engine/Classes/Camera/CameraModifier_CameraShake.h
+++ b/Engine/Source/Runtime/Engine/Classes/Camera/CameraModifier_CameraShake.h
@@ -16,6 +16,9 @@
 class UCameraShakeBase;
 class UCameraShakeSourceComponent;
 
+/**
+ * List of reusable camera shake instances that we store in a recycling pool.
+ */
 USTRUCT()
 struct FPooledCameraShakes
 {
@@ -25,6 +28,9 @@ struct FPooledCameraShakes
 	TArray<TObjectPtr<UCameraShakeBase>> PooledShakes;
 };
 
+/**
+ * Information about an active camera shake.
+ */
 USTRUCT()
 struct FActiveCameraShakeInfo
 {
@@ -36,15 +42,61 @@ struct FActiveCameraShakeInfo
 		, bIsCustomInitialized(false)
 	{}
 
+	/** The camera shake instance currently running */
 	UPROPERTY()
 	TObjectPtr<UCameraShakeBase> ShakeInstance;
 
+	/** An optional source that the shake is running from (otherwise it's global and un-attenuated) */
 	UPROPERTY()
 	TWeakObjectPtr<const UCameraShakeSourceComponent> ShakeSource;
 
+	/** Whether the shake was initialized with some custom callback code */
 	UPROPERTY()
 	bool bIsCustomInitialized;
+
+#if UE_ENABLE_DEBUG_DRAWING
+	/** Index of the debug graph display data for this shake */
+	int32 DebugDataIndex = INDEX_NONE;
+#endif
+};
+
+#if UE_ENABLE_DEBUG_DRAWING
+/**
+ * A data point for the debug graph display.
+ */
+struct FCameraShakeDebugDataPoint
+{
+	/** Time since start of the shake */
+	float AccumulatedTime;
+	/** Location movement of the shake for this frame */
+	FVector DeltaLocation;
+	/** Rotation movement of the shake for this frame */
+	FRotator DeltaRotation;
+};
+#endif
+
+#if UE_ENABLE_DEBUG_DRAWING
+/** 
+ * Information about active or recently ended camera shakes that we want to show in the 
+ * debug graph display. Shakes stay alive a little longer in this struct (compared to the
+ * ActiveShakes list) because after a shake ended we want to keep displaying its debug graph
+ * for a little while instead of having it disappear immediately. We also draw again in
+ * the same graph if the shake comes back soon enough.
+ */
+struct FCameraShakeDebugData
+{
+	/** The class of the shake that is/was running */
+	TSubclassOf<UCameraShakeBase> ShakeClass;
+	/** The name of the shake that is/was running */
+	FString ShakeInstanceName;
+	/** Data points for the debug graph display */
+	TArray<FCameraShakeDebugDataPoint> DataPoints;
+	/** How long ago did the shake end */
+	float TimeInactive = 0.f;
+	/** Is the shake active */
+	bool bIsInactive = false;
 };
+#endif
 
 DECLARE_DELEGATE_OneParam(FOnInitializeCameraShake, UCameraShakeBase*);
 
@@ -171,4 +223,16 @@ protected:
 	/** Scaling factor applied to all camera shakes in when in splitscreen mode. Normally used to reduce shaking, since shakes feel more intense in a smaller viewport. */
 	UPROPERTY(EditAnywhere, Category = CameraModifier_CameraShake)
 	float SplitScreenShakeScale;
+
+#if UE_ENABLE_DEBUG_DRAWING
+	void AddCameraShakeDebugData(FActiveCameraShakeInfo& ShakeInfo);
+	void RemoveCameraShakeDebugData(const FActiveCameraShakeInfo& ShakeInfo);
+	void DisplayDebugGraphs(class UCanvas* Canvas, const FDebugDisplayInfo& DebugDisplay);
+
+	/** List of debug data for camera shakes */
+	TArray<FCameraShakeDebugData> DebugShakes;
+	/** Whether we are currently recording camera shake debug info */
+	bool bRecordDebugData;
+#endif
 };
+
diff --git a/Engine/Source/Runtime/Engine/Private/Camera/CameraModifier_CameraShake.cpp b/Engine/Source/Runtime/Engine/Private/Camera/CameraModifier_CameraShake.cpp
--- a/Engine/Source/Runtime/Engine/Private/Camera/CameraModifier_CameraShake.cpp
+++ b/Engine/Source/Runtime/Engine/Private/Camera/CameraModifier_CameraShake.cpp
@@ -19,8 +19,32 @@ DEFINE_LOG_CATEGORY_STATIC(LogCameraShake, Warning, All);
 
 DECLARE_CYCLE_STAT(TEXT("AddCameraShake"), STAT_AddCameraShake, STATGROUP_Game);
 
+TAutoConsoleVariable<bool> GShowCameraShakeDebugCVar(
+	TEXT("r.CameraShakeDebug"),
+	false,
+	TEXT("Show extra debug info for camera shakes (requires `showdebug CAMERA`)"));
+TAutoConsoleVariable<bool> GShowCameraShakeDebugLocationCVar(
+	TEXT("r.CameraShakeDebug.Location"),
+	true,
+	TEXT("Whether to show camera shakes' location modifications (defaults to true)"));
+TAutoConsoleVariable<bool> GShowCameraShakeDebugRotationCVar(
+	TEXT("r.CameraShakeDebug.Rotation"),
+	true,
+	TEXT("Whether to show camera shakes' rotation modifications (defaults to true)"));
+TAutoConsoleVariable<bool> GCameraShakeDebugLargeGraphCVar(
+	TEXT("r.CameraShakeDebug.LargeGraph"),
+	false,
+	TEXT("Draws larger graphs for camera shake debug info"));
+TAutoConsoleVariable<float> GCameraShakeDebugInfoRecordLimitCVar(
+	TEXT("r.CameraShakeDebug.InfoRecordLimit"),
+	2,
+	TEXT("How many seconds to keep while recording camera shake debug info (defaults to 2 seconds)"));
+
 UCameraModifier_CameraShake::UCameraModifier_CameraShake(const FObjectInitializer& ObjectInitializer)
 	: Super(ObjectInitializer)
+#if UE_ENABLE_DEBUG_DRAWING
+	, bRecordDebugData(false)
+#endif
 {
 	SplitScreenShakeScale = 0.5f;
 }
@@ -37,6 +61,10 @@ bool UCameraModifier_CameraShake::ModifyCamera(float DeltaTime, FMinimalViewInfo
 		return false;
 	}
 
+#if UE_ENABLE_DEBUG_DRAWING
+	const float DebugDataLimit = GCameraShakeDebugInfoRecordLimitCVar.GetValueOnGameThread();
+#endif
+
 	// Update and apply active shakes
 	if( ActiveShakes.Num() > 0 )
 	{
@@ -44,6 +72,11 @@ bool UCameraModifier_CameraShake::ModifyCamera(float DeltaTime, FMinimalViewInfo
 		{
 			if (ShakeInfo.ShakeInstance != nullptr)
 			{
+#if UE_ENABLE_DEBUG_DRAWING
+				const FVector OrigLocation = InOutPOV.Location;
+				const FRotator OrigRotation = InOutPOV.Rotation;
+#endif
+
 				// Compute the scale of this shake for this frame according to the location
 				// of its source.
 				float CurShakeAlpha = Alpha;
@@ -55,6 +88,24 @@ bool UCameraModifier_CameraShake::ModifyCamera(float DeltaTime, FMinimalViewInfo
 				}
 
 				ShakeInfo.ShakeInstance->UpdateAndApplyCameraShake(DeltaTime, CurShakeAlpha, InOutPOV);
+
+#if UE_ENABLE_DEBUG_DRAWING
+				if (bRecordDebugData && ShakeInfo.DebugDataIndex != INDEX_NONE)
+				{
+					// Record the delta between the original camera location/rotation, and the result
+					// of this camera shake.
+					check(DebugShakes.IsValidIndex(ShakeInfo.DebugDataIndex));
+					FCameraShakeDebugData& DebugShake = DebugShakes[ShakeInfo.DebugDataIndex];
+					const float PreviousAccumulatedTime = (DebugShake.DataPoints.Num() > 0) ?
+						DebugShake.DataPoints.Last().AccumulatedTime : 0.f;
+
+					const FVector DeltaLocation = InOutPOV.Location - OrigLocation;
+					const FRotator DeltaRotation = (InOutPOV.Rotation - OrigRotation).GetNormalized();
+					const float AccumulatedTime = PreviousAccumulatedTime + DeltaTime;
+					const FCameraShakeDebugDataPoint FrameData { AccumulatedTime, DeltaLocation, DeltaRotation };
+					DebugShake.DataPoints.Add(FrameData);
+				}
+#endif
 			}
 		}
 
@@ -77,6 +128,59 @@ bool UCameraModifier_CameraShake::ModifyCamera(float DeltaTime, FMinimalViewInfo
 		}
 	}
 
+#if UE_ENABLE_DEBUG_DRAWING
+	// Add empty data points for the inactive shakes, and get rid of those who expired.
+	for (int32 Index = DebugShakes.Num() - 1; Index >= 0; --Index)
+	{
+		FCameraShakeDebugData& DebugShake = DebugShakes[Index];
+		if (DebugShake.bIsInactive)
+		{
+			DebugShake.TimeInactive += DeltaTime;
+			if (DebugShake.TimeInactive >= DebugDataLimit)
+			{
+				// We can remove this debug data, it's been kept in the debug display long enough.
+				DebugShakes.RemoveAt(Index);
+
+				// Update the debug indices for active shakes.
+				for (FActiveCameraShakeInfo& ShakeInfo : ActiveShakes)
+				{
+					check(ShakeInfo.DebugDataIndex != Index);
+					if (ShakeInfo.DebugDataIndex > Index)
+					{
+						--ShakeInfo.DebugDataIndex;
+					}
+				}
+			}
+			else
+			{
+				// Add empty data points for a while.
+				const float PreviousAccumulatedTime = (DebugShake.DataPoints.Num() > 0) ?
+					DebugShake.DataPoints.Last().AccumulatedTime : 0.f;
+				const float AccumulatedTime = PreviousAccumulatedTime + DeltaTime;
+				const FCameraShakeDebugDataPoint FrameData{ AccumulatedTime, FVector::ZeroVector, FRotator::ZeroRotator };
+				DebugShake.DataPoints.Add(FrameData);
+			}
+		}
+	}
+
+	// Remove old data points.
+	for (FCameraShakeDebugData& DebugShake : DebugShakes)
+	{
+		if (DebugShake.DataPoints.Num() > 0)
+		{
+			const float CutoffTime = DebugShake.DataPoints.Last().AccumulatedTime - DebugDataLimit;
+			for (int32 Index = 0; Index < DebugShake.DataPoints.Num(); ++Index)
+			{
+				if (DebugShake.DataPoints[Index].AccumulatedTime > CutoffTime)
+				{
+					DebugShake.DataPoints.RemoveAt(0, Index);
+					break;
+				}
+			}
+		}
+	}
+#endif
+
 	// If ModifyCamera returns true, exit loop
 	// Allows high priority things to dictate if they are
 	// the last modifier to be applied
@@ -174,6 +278,10 @@ UCameraShakeBase* UCameraModifier_CameraShake::AddCameraShake(TSubclassOf<UCamer
 					ShakeInfo.ShakeSource = SourceComponent;
 					ShakeInfo.bIsCustomInitialized = bIsCustomInitialized;
 					bReplacedNull = true;
+
+#if UE_ENABLE_DEBUG_DRAWING
+					AddCameraShakeDebugData(ShakeInfo);
+#endif
 				}
 			}
 
@@ -186,6 +294,10 @@ UCameraShakeBase* UCameraModifier_CameraShake::AddCameraShake(TSubclassOf<UCamer
 				ShakeInfo.bIsCustomInitialized = bIsCustomInitialized;
 				ActiveShakes.Emplace(ShakeInfo);
 				UE_LOG(LogCameraShake, Verbose, TEXT("UCameraModifier_CameraShake::AddCameraShake %s Active Instance Added"), *GetNameSafe(ShakeInfo.ShakeInstance));
+
+#if UE_ENABLE_DEBUG_DRAWING
+				AddCameraShakeDebugData(ActiveShakes.Last());
+#endif
 			}
 		}
 
@@ -210,6 +322,10 @@ void UCameraModifier_CameraShake::SaveShakeInExpiredPoolIfPossible(const FActive
 	{
 		SaveShakeInExpiredPool(ShakeInfo.ShakeInstance);
 	}
+
+#if UE_ENABLE_DEBUG_DRAWING
+	RemoveCameraShakeDebugData(ShakeInfo);
+#endif
 }
 
 UCameraShakeBase* UCameraModifier_CameraShake::ReclaimShakeFromExpiredPool(TSubclassOf<UCameraShakeBase> CameraShakeClass)
@@ -337,7 +453,7 @@ void UCameraModifier_CameraShake::DisplayDebug(UCanvas* Canvas, const FDebugDisp
 {
 	Canvas->SetDrawColor(FColor::Yellow);
 	UFont* DrawFont = GEngine->GetSmallFont();
-	
+
 	int Indentation = 1;
 	int LineNumber = FMath::CeilToInt(YPos / YL);
 
@@ -350,12 +466,288 @@ void UCameraModifier_CameraShake::DisplayDebug(UCanvas* Canvas, const FDebugDisp
 
 		if (ShakeInfo.ShakeInstance != nullptr)
 		{
-			const FString DurationString = !ShakeInfo.ShakeInstance->GetCameraShakeDuration().IsInfinite() ? FString::SanitizeFloat(ShakeInfo.ShakeInstance->GetCameraShakeDuration().Get()) : TEXT("Infinite");
-			Canvas->DrawText(DrawFont, FString::Printf(TEXT("[%d] %s Source:%s Duration: %s Elapsed: %f"), i, *GetNameSafe(ShakeInfo.ShakeInstance), *GetNameSafe(ShakeInfo.ShakeSource.Get()), *DurationString, ShakeInfo.ShakeInstance->GetElapsedTime()), Indentation* YL, (LineNumber++)* YL);
+			const FString DurationString = !ShakeInfo.ShakeInstance->GetCameraShakeDuration().IsInfinite() ? 
+				FString::SanitizeFloat(ShakeInfo.ShakeInstance->GetCameraShakeDuration().Get()) : TEXT("Infinite");
+			Canvas->DrawText(DrawFont,
+					FString::Printf(TEXT("[%d] %s Source:%s Duration: %s Elapsed: %f"), 
+						i, *GetNameSafe(ShakeInfo.ShakeInstance), *GetNameSafe(ShakeInfo.ShakeSource.Get()), 
+						*DurationString, ShakeInfo.ShakeInstance->GetElapsedTime()),
+					Indentation * YL, (LineNumber++) * YL);
 		}
 	}
 
 	YPos = LineNumber * YL;
 
+#if UE_ENABLE_DEBUG_DRAWING
+	DisplayDebugGraphs(Canvas, DebugDisplay);
+#endif
+
 	Super::DisplayDebug(Canvas, DebugDisplay, YL, YPos);
 }
+
+#if UE_ENABLE_DEBUG_DRAWING
+
+void UCameraModifier_CameraShake::AddCameraShakeDebugData(FActiveCameraShakeInfo& ShakeInfo)
+{
+	FString ShakeInstanceName = GetNameSafe(ShakeInfo.ShakeInstance);
+	TSubclassOf<UCameraShakeBase> ShakeClass = ShakeInfo.ShakeInstance->GetClass();
+
+	for (int32 Index = 0; Index < DebugShakes.Num(); ++Index)
+	{
+		FCameraShakeDebugData& DebugShake = DebugShakes[Index];
+
+		if (DebugShake.bIsInactive && 
+				DebugShake.ShakeClass == ShakeClass && 
+				DebugShake.ShakeInstanceName == ShakeInstanceName)
+		{
+			// Found a matching debug info, re-use it.
+			ShakeInfo.DebugDataIndex = Index;
+			DebugShake.bIsInactive = false;
+			DebugShake.TimeInactive = 0.f;
+			return;
+		}
+	}
+
+	FCameraShakeDebugData NewDebugShake;
+	NewDebugShake.ShakeClass = ShakeClass;
+	NewDebugShake.ShakeInstanceName = ShakeInstanceName;
+	ShakeInfo.DebugDataIndex = DebugShakes.Emplace(NewDebugShake);
+}
+
+void UCameraModifier_CameraShake::RemoveCameraShakeDebugData(const FActiveCameraShakeInfo& ShakeInfo)
+{
+	if (ShakeInfo.DebugDataIndex != INDEX_NONE)
+	{
+		// Just make the debug info as inactive. It will get culled once it has been displayed
+		// long enough after that.
+		ensure(DebugShakes.IsValidIndex(ShakeInfo.DebugDataIndex));
+		ensure(DebugShakes[ShakeInfo.DebugDataIndex].bIsInactive == false);
+		DebugShakes[ShakeInfo.DebugDataIndex].bIsInactive = true;
+		DebugShakes[ShakeInfo.DebugDataIndex].TimeInactive = 0.f;
+	}
+}
+
+void UCameraModifier_CameraShake::DisplayDebugGraphs(class UCanvas* Canvas, const FDebugDisplayInfo& DebugDisplay)
+{
+	const bool bPreviousRecordDebugData = bRecordDebugData;
+
+	// See if we should be recording and display shake data.
+	const bool bShowLocationDebug = GShowCameraShakeDebugCVar.GetValueOnGameThread() && 
+		GShowCameraShakeDebugLocationCVar.GetValueOnGameThread();
+	const bool bShowRotationDebug = GShowCameraShakeDebugCVar.GetValueOnGameThread() && 
+		GShowCameraShakeDebugRotationCVar.GetValueOnGameThread();
+	bRecordDebugData = bShowLocationDebug || bShowRotationDebug;
+	
+	if (!bRecordDebugData)
+	{
+		return;
+	}
+
+	if (!bPreviousRecordDebugData)
+	{
+		// If there are any already running shakes when we start displaying debug info, let's add them
+		// to the debug list.
+		for (FActiveCameraShakeInfo& ShakeInfo : ActiveShakes)
+		{
+			if (ShakeInfo.DebugDataIndex == INDEX_NONE)
+			{
+				AddCameraShakeDebugData(ShakeInfo);
+			}
+		}
+	}
+
+	Canvas->SetDrawColor(FColor::Yellow);
+	UFont* DrawFont = GEngine->GetSmallFont();
+	const float FontHeight = DrawFont->GetMaxCharHeight();
+	
+	const bool bDrawLargeGraphs = GCameraShakeDebugLargeGraphCVar.GetValueOnGameThread();
+
+	const float GraphWidth = bDrawLargeGraphs ? 320.f : 160.f;
+	const float GraphHeight = bDrawLargeGraphs ? 120.f : 60.0f;
+	const float GraphRightXPos = Canvas->SizeX - 20;
+	const float GraphLeftXPos = GraphRightXPos - GraphWidth;
+
+	const float DebugDataLimit = GCameraShakeDebugInfoRecordLimitCVar.GetValueOnGameThread();
+	const float PixelsPerSecond = GraphWidth / (DebugDataLimit > 0 ? DebugDataLimit : 3.f);
+
+	for (int i = 0; i < DebugShakes.Num(); i++)
+	{
+		const FCameraShakeDebugData& DebugShake = DebugShakes[i];
+		if (DebugShake.DataPoints.IsEmpty())
+		{
+			continue;
+		}
+
+		const float GraphTopYPos = 20 + (i * (GraphHeight + FontHeight + 20));
+
+		// Draw background.
+		FLinearColor BackgroundColor = FLinearColor(0.0f, 0.0f, 0.0f, 0.7f);
+		FCanvasTileItem BackgroundTile(
+				FVector2D(GraphLeftXPos, GraphTopYPos),
+				FVector2D(GraphWidth, GraphHeight),
+				BackgroundColor);
+		BackgroundTile.BlendMode = SE_BLEND_AlphaBlend;
+		Canvas->DrawItem(BackgroundTile);
+
+		// Draw shake name.
+		Canvas->DrawText(DrawFont, *DebugShake.ShakeInstanceName, GraphLeftXPos, GraphTopYPos + GraphHeight + 10);
+
+		FBatchedElements* BatchedElements = Canvas->Canvas->GetBatchedElements(FCanvas::ET_Line);
+		FHitProxyId HitProxyId = Canvas->Canvas->GetHitProxyId();
+		const float GraphBottomYPos = GraphTopYPos + GraphHeight;
+
+		// Compute vertical scale based on min/max values.
+		float MaxLocationValue = 0.f;
+		float MinLocationValue = 0.f;
+		float MaxRotationValue = 0.f;
+		float MinRotationValue = 0.f;
+
+		for (int32 DataIndex = 0; DataIndex < DebugShake.DataPoints.Num(); ++DataIndex)
+		{
+			const FCameraShakeDebugDataPoint& DataPoint = DebugShake.DataPoints[DataIndex];
+
+			MaxLocationValue = FMath::Max(MaxLocationValue, DataPoint.DeltaLocation.X);
+			MaxLocationValue = FMath::Max(MaxLocationValue, DataPoint.DeltaLocation.Y);
+			MaxLocationValue = FMath::Max(MaxLocationValue, DataPoint.DeltaLocation.Z);
+			MinLocationValue = FMath::Min(MinLocationValue, DataPoint.DeltaLocation.X);
+			MinLocationValue = FMath::Min(MinLocationValue, DataPoint.DeltaLocation.Y);
+			MinLocationValue = FMath::Min(MinLocationValue, DataPoint.DeltaLocation.Z);
+
+			MaxRotationValue = FMath::Max(MaxRotationValue, DataPoint.DeltaRotation.Yaw);
+			MaxRotationValue = FMath::Max(MaxRotationValue, DataPoint.DeltaRotation.Pitch);
+			MaxRotationValue = FMath::Max(MaxRotationValue, DataPoint.DeltaRotation.Roll);
+			MinRotationValue = FMath::Min(MinRotationValue, DataPoint.DeltaRotation.Yaw);
+			MinRotationValue = FMath::Min(MinRotationValue, DataPoint.DeltaRotation.Pitch);
+			MinRotationValue = FMath::Min(MinRotationValue, DataPoint.DeltaRotation.Roll);
+		}
+
+		const float Threshold = 0.005f;
+		if (MaxLocationValue - MinLocationValue < Threshold * 2.f)
+		{
+			MaxLocationValue += Threshold;
+			MinLocationValue -= Threshold;
+		}
+		if (MaxRotationValue - MinRotationValue < Threshold * 2.f)
+		{
+			MaxRotationValue += Threshold;
+			MinRotationValue -= Threshold;
+		}
+
+		// Draw min/max.
+		if (bShowLocationDebug && bShowRotationDebug)
+		{
+			int MaxLabelHeight, MaxLabelWidth;
+			FString MaxLabel = FString::Printf(TEXT("L=%.02f R=%.02f"), MaxLocationValue, MaxRotationValue);
+			DrawFont->GetStringHeightAndWidth(MaxLabel, MaxLabelHeight, MaxLabelWidth);
+
+			int MinLabelHeight, MinLabelWidth;
+			FString MinLabel = FString::Printf(TEXT("L=%.02f R=%.02f"), MinLocationValue, MinRotationValue);
+			DrawFont->GetStringHeightAndWidth(MinLabel, MinLabelHeight, MinLabelWidth);
+
+			Canvas->DrawText(DrawFont, MaxLabel, GraphLeftXPos - MaxLabelWidth - 10, GraphTopYPos);
+			Canvas->DrawText(DrawFont, MinLabel, GraphLeftXPos - MinLabelWidth - 10, GraphTopYPos + GraphHeight - MinLabelHeight);
+		}
+		else if (bShowLocationDebug)
+		{
+			int MaxLabelHeight, MaxLabelWidth;
+			FString MaxLabel = FString::Printf(TEXT("%.02f"), MaxLocationValue);
+			DrawFont->GetStringHeightAndWidth(MaxLabel, MaxLabelHeight, MaxLabelWidth);
+
+			int MinLabelHeight, MinLabelWidth;
+			FString MinLabel = FString::Printf(TEXT("%.02f"), MinLocationValue);
+			DrawFont->GetStringHeightAndWidth(MinLabel, MinLabelHeight, MinLabelWidth);
+
+			Canvas->DrawText(DrawFont, MaxLabel, GraphLeftXPos - MaxLabelWidth - 10, GraphTopYPos);
+			Canvas->DrawText(DrawFont, MinLabel, GraphLeftXPos - MinLabelWidth - 10, GraphTopYPos + GraphHeight - MinLabelHeight);
+		}
+		else if (bShowRotationDebug)
+		{
+			int MaxLabelHeight, MaxLabelWidth;
+			FString MaxLabel = FString::Printf(TEXT("%.02f"), MaxRotationValue);
+			DrawFont->GetStringHeightAndWidth(MaxLabel, MaxLabelHeight, MaxLabelWidth);
+
+			int MinLabelHeight, MinLabelWidth;
+			FString MinLabel = FString::Printf(TEXT("%.02f"), MinRotationValue);
+			DrawFont->GetStringHeightAndWidth(MinLabel, MinLabelHeight, MinLabelWidth);
+
+			Canvas->DrawText(DrawFont, MaxLabel, GraphLeftXPos - MaxLabelWidth - 10, GraphTopYPos);
+			Canvas->DrawText(DrawFont, MinLabel, GraphLeftXPos - MinLabelWidth - 10, GraphTopYPos + GraphHeight - MinLabelHeight);
+		}
+
+		// Draw curves.
+		const float LocationUnitsPerPixel = (MaxLocationValue - MinLocationValue) / GraphHeight;
+		const float RotationUnitsPerPixel = (MaxRotationValue - MinRotationValue) / GraphHeight;
+
+		FLinearColor LocationXStatColor(1.0f, 0.1f, 0.1f);		// Red
+		FLinearColor LocationYStatColor(0.1f, 1.0f, 0.1f);		// Green
+		FLinearColor LocationZStatColor(0.1f, 0.1f, 1.0f);		// Blue
+
+		FLinearColor RotationXStatColor(1.0f, 1.0f, 0.1f);		// Yellow
+		FLinearColor RotationYStatColor(1.0f, 0.5f, 0.1f);		// Orange
+		FLinearColor RotationZStatColor(0.1f, 0.8f, 0.8f);		// Cyan
+
+		const float TimeOffset = DebugShake.DataPoints[0].AccumulatedTime;
+
+		for (int32 DataIndex = 1; DataIndex < DebugShake.DataPoints.Num(); ++DataIndex)
+		{
+			const FCameraShakeDebugDataPoint& PrevDataPoint = DebugShake.DataPoints[DataIndex - 1];
+			const float PrevOffsetX = (PrevDataPoint.AccumulatedTime - TimeOffset) * PixelsPerSecond;
+
+			const FCameraShakeDebugDataPoint& NextDataPoint = DebugShake.DataPoints[DataIndex];
+			const float NextOffsetX = (NextDataPoint.AccumulatedTime - TimeOffset) * PixelsPerSecond;
+
+			if (bShowLocationDebug)
+			{
+				const FVector PrevLocationPixels(
+					(PrevDataPoint.DeltaLocation.X - MinLocationValue) / LocationUnitsPerPixel,
+					(PrevDataPoint.DeltaLocation.Y - MinLocationValue) / LocationUnitsPerPixel,
+					(PrevDataPoint.DeltaLocation.Z - MinLocationValue) / LocationUnitsPerPixel);
+
+				const FVector NextLocationPixels(
+					(NextDataPoint.DeltaLocation.X - MinLocationValue) / LocationUnitsPerPixel,
+					(NextDataPoint.DeltaLocation.Y - MinLocationValue) / LocationUnitsPerPixel,
+					(NextDataPoint.DeltaLocation.Z - MinLocationValue) / LocationUnitsPerPixel);
+
+				const FVector LineStartLocationX(GraphLeftXPos + PrevOffsetX, GraphBottomYPos - PrevLocationPixels.X, 0.0f);
+				const FVector LineStartLocationY(GraphLeftXPos + PrevOffsetX, GraphBottomYPos - PrevLocationPixels.Y, 0.0f);
+				const FVector LineStartLocationZ(GraphLeftXPos + PrevOffsetX, GraphBottomYPos - PrevLocationPixels.Z, 0.0f);
+
+				const FVector LineEndLocationX(GraphLeftXPos + NextOffsetX, GraphBottomYPos - NextLocationPixels.X, 0.0f);
+				const FVector LineEndLocationY(GraphLeftXPos + NextOffsetX, GraphBottomYPos - NextLocationPixels.Y, 0.0f);
+				const FVector LineEndLocationZ(GraphLeftXPos + NextOffsetX, GraphBottomYPos - NextLocationPixels.Z, 0.0f);
+
+				BatchedElements->AddLine(LineStartLocationX, LineEndLocationX, LocationXStatColor, HitProxyId);
+				BatchedElements->AddLine(LineStartLocationY, LineEndLocationY, LocationYStatColor, HitProxyId);
+				BatchedElements->AddLine(LineStartLocationZ, LineEndLocationZ, LocationZStatColor, HitProxyId);
+			}
+
+			if (bShowRotationDebug)
+			{
+				const FVector PrevRotationPixels( 
+					(PrevDataPoint.DeltaRotation.Yaw - MinRotationValue) / RotationUnitsPerPixel, 
+					(PrevDataPoint.DeltaRotation.Pitch - MinRotationValue) / RotationUnitsPerPixel,
+					(PrevDataPoint.DeltaRotation.Roll - MinRotationValue) / RotationUnitsPerPixel);
+			
+				const FVector NextRotationPixels( 
+					(NextDataPoint.DeltaRotation.Yaw - MinRotationValue) / RotationUnitsPerPixel,
+					(NextDataPoint.DeltaRotation.Pitch - MinRotationValue) / RotationUnitsPerPixel,
+					(NextDataPoint.DeltaRotation.Roll - MinRotationValue) / RotationUnitsPerPixel);
+
+				const FVector LineStartRotationX(GraphLeftXPos + PrevOffsetX, GraphBottomYPos - PrevRotationPixels.X, 0.0f);
+				const FVector LineStartRotationY(GraphLeftXPos + PrevOffsetX, GraphBottomYPos - PrevRotationPixels.Y, 0.0f);
+				const FVector LineStartRotationZ(GraphLeftXPos + PrevOffsetX, GraphBottomYPos - PrevRotationPixels.Z, 0.0f);
+
+				const FVector LineEndRotationX(GraphLeftXPos + NextOffsetX, GraphBottomYPos - NextRotationPixels.X, 0.0f);
+				const FVector LineEndRotationY(GraphLeftXPos + NextOffsetX, GraphBottomYPos - NextRotationPixels.Y, 0.0f);
+				const FVector LineEndRotationZ(GraphLeftXPos + NextOffsetX, GraphBottomYPos - NextRotationPixels.Z, 0.0f);
+
+				BatchedElements->AddLine(LineStartRotationX, LineEndRotationX, RotationXStatColor, HitProxyId);
+				BatchedElements->AddLine(LineStartRotationY, LineEndRotationY, RotationYStatColor, HitProxyId);
+				BatchedElements->AddLine(LineStartRotationZ, LineEndRotationZ, RotationZStatColor, HitProxyId);
+			}
+		}
+	}
+}
+
+#endif
+
