Add SSS profile support for Strata eye BSDF.

diff --git a/Engine/Shaders/Private/Strata/Strata.ush b/Engine/Shaders/Private/Strata/Strata.ush
--- a/Engine/Shaders/Private/Strata/Strata.ush
+++ b/Engine/Shaders/Private/Strata/Strata.ush
@@ -446,6 +446,7 @@ struct FStrataBSDF
 #define EYE_ROUGHNESS(X)				X.VGPRs[0].w
 #define EYE_IRISDISTANCE(X)				X.VGPRs[1].x
 #define EYE_IRISMASK(X)					X.VGPRs[1].y
+#define EYE_SSSPROFILEID(X)				X.VGPRs[1].z
 #define EYE_IRISNORMAL(X)				X.VGPRs[2].xyz
 
 ///////////////////////////////////////////////////////////////////////////////
@@ -1425,7 +1426,7 @@ FStrataData GetStrataHairBSDF(float3 BaseColor, float Scatter, float Specular, f
 	return StrataData;
 }
 
-FStrataData GetStrataEyeBSDF(float3 DiffuseAlbedo, float Roughness, float IrisMask, float IrisDistance, float3 IrisNormal, float3 Emissive, uint SharedLocalBasisIndex)
+FStrataData GetStrataEyeBSDF(float3 DiffuseAlbedo, float Roughness, float IrisMask, float IrisDistance, float3 IrisNormal, float SSSProfileId, float3 Emissive, uint SharedLocalBasisIndex)
 {
 	FStrataData StrataData = GetInitialisedStrataData();
 	BSDF_SETTYPE					(StrataData.InlinedBSDF, STRATA_BSDF_TYPE_EYE);
@@ -1436,6 +1437,10 @@ FStrataData GetStrataEyeBSDF(float3 DiffuseAlbedo, float Roughness, float IrisMa
 	EYE_IRISDISTANCE				(StrataData.InlinedBSDF) = IrisDistance;
 	EYE_IRISMASK					(StrataData.InlinedBSDF) = IrisMask;
 	EYE_IRISNORMAL					(StrataData.InlinedBSDF) = IrisNormal;
+	EYE_SSSPROFILEID				(StrataData.InlinedBSDF) = SSSProfileId;
+
+	const bool bHasSSS = SSSProfileId != 0;
+	BSDF_SETSSSTYPE(StrataData.InlinedBSDF, bHasSSS ? SSS_TYPE_DIFFUSION_PROFILE : SSS_TYPE_INVALID);
 
 #if STRATA_INLINE_SHADING
 	StrataData.InlinedBSDF.Coverage = 1.0f;
diff --git a/Engine/Shaders/Private/Strata/StrataExport.ush b/Engine/Shaders/Private/Strata/StrataExport.ush
--- a/Engine/Shaders/Private/Strata/StrataExport.ush
+++ b/Engine/Shaders/Private/Strata/StrataExport.ush
@@ -247,6 +247,8 @@ void PackStrataOut(
 						TopLayerTotalWeight += BSDF.TopLayerDataWeight;
 						TopLayerData.Roughness += BSDF.TopLayerDataWeight * FullyRough;
 
+						const bool bHasSSS = BSDF_GETSSSTYPE(BSDF) == SSS_TYPE_DIFFUSION_PROFILE;
+						SetHasSubsurface(StrataPixelHeader, bHasSSS);
 						SetMaterialMode(StrataPixelHeader, HEADER_MATERIALMODE_EYE);
 						break;
 					}
@@ -554,6 +556,15 @@ void PackStrataOut(
 							STRATA_STORE_UINT1(PackColorLinearToGamma2AlphaLinear(float4(EYE_DIFFUSEALBEDO(BSDF), EYE_ROUGHNESS(BSDF))));
 							STRATA_STORE_UINT1(PackRGBA8(float4(EYE_IRISMASK(BSDF), EYE_IRISDISTANCE(BSDF), EncodedIrisNormal * 0.5f + 0.5f)));
 							// 8 bytes
+
+							// Note: we don't store the SSS profile ID into the Eye data, as it is not needed there. The profile is only stored into the SSS data.
+							const bool bHasSSS = BSDF_GETSSSTYPE(BSDF) == SSS_TYPE_DIFFUSION_PROFILE;
+							if (bHasSSS)
+							{
+								StrataSubSurfaceHeaderSetSSSType(SSSData.Header, SSS_TYPE_DIFFUSION_PROFILE);
+								StrataSubSurfaceHeaderSetProfile(SSSData.Header, 1.0f, StrataSubsurfaceProfileIdTo8bits(EYE_SSSPROFILEID(BSDF)));
+								StrataSubsurfaceExtrasSetBaseColor(SSSData.Extras, EYE_DIFFUSEALBEDO(BSDF));
+							}
 						}
 						break;
 						case STRATA_BSDF_TYPE_SINGLELAYERWATER:
diff --git a/Engine/Shaders/Private/Strata/StrataLegacyConversion.ush b/Engine/Shaders/Private/Strata/StrataLegacyConversion.ush
--- a/Engine/Shaders/Private/Strata/StrataLegacyConversion.ush
+++ b/Engine/Shaders/Private/Strata/StrataLegacyConversion.ush
@@ -550,12 +550,13 @@ FStrataData StrataConvertLegacyMaterialStatic(
 
 		FStrataData EyeMat = PromoteParameterBlendedBSDFToOperator(
 			GetStrataEyeBSDF(
-				BaseColor,	   // DiffuseAlbedo
-				Roughness,	   // Roughness
-				IrisMask,	   // IrisMask
-				IrisDistance,  // IrisDistance
-				IrisNormal,	   // IrisNormal
-				Emissive,	   // EmissiveColor
+				BaseColor,	         // DiffuseAlbedo
+				Roughness,	         // Roughness
+				IrisMask,	         // IrisMask
+				IrisDistance,        // IrisDistance
+				IrisNormal,	         // IrisNormal
+				SubSurfaceProfileId, // SSS profile
+				Emissive,	         // EmissiveColor
 				SharedLocalBasisIndex),
 			StrataTree,
 			0		/*OperatorIndex*/,
diff --git a/Engine/Shaders/Private/Strata/StrataVisualize.usf b/Engine/Shaders/Private/Strata/StrataVisualize.usf
--- a/Engine/Shaders/Private/Strata/StrataVisualize.usf
+++ b/Engine/Shaders/Private/Strata/StrataVisualize.usf
@@ -296,9 +296,10 @@ void PrintBSDF(inout FShaderPrintContext Context, inout FStrataAddressing Strata
 		{
 			Print(Context, TEXT("Hair"), FontBSDFType);
 			Newline(Context, RectMax);
-			Print(Context, TEXT(" NormalID Aniso IsTopLayer LuminanceWeight"), FontBSDF);
+			Print(Context, TEXT(" NormalID Aniso IsTopLayer LuminanceWeight"), FontBSDFStateName);
 			Newline(Context, RectMax);
 
+			Print(Context, TEXT(" "), FontSilver);
 			Print(Context, BSDF_GETSHAREDLOCALBASISID(BSDF), FontBSDF);
 			PrintBool(Context, BSDF_GETHASANISOTROPY(BSDF), FontBSDF);
 			PrintBool(Context, BSDF_GETISTOPLAYER(BSDF), FontBSDF);
@@ -327,9 +328,10 @@ void PrintBSDF(inout FShaderPrintContext Context, inout FStrataAddressing Strata
 		{
 			Print(Context, TEXT("Eye"), FontBSDFType);
 			Newline(Context, RectMax);
-			Print(Context, TEXT(" NormalID IsTopLayer LuminanceWeight"), FontBSDF);
+			Print(Context, TEXT(" NormalID IsTopLayer LuminanceWeight"), FontBSDFStateName);
 			Newline(Context, RectMax);
 
+			Print(Context, TEXT(" "), FontSilver);
 			Print(Context, BSDF_GETSHAREDLOCALBASISID(BSDF), FontBSDF);
 			PrintBool(Context, BSDF_GETISTOPLAYER(BSDF), FontBSDF);
 			if (BSDF_GETHASGREYWEIGHT_V(BSDF))
@@ -353,6 +355,16 @@ void PrintBSDF(inout FShaderPrintContext Context, inout FStrataAddressing Strata
 			// Tangent basis for the iris
 			const float3x3 IrisTangentBasis = GetTangentBasis(EYE_IRISNORMAL(BSDF));
 			DrawReferentialTWS(WorldPosition, IrisTangentBasis[0], IrisTangentBasis[1], IrisTangentBasis[2], float3(0, 1, 1));
+
+			const bool bHasSSS = HasSubsurface(Header);
+			const FStrataSubsurfaceHeader SSSHeader = StrataLoadSubsurfaceHeader(Strata.SSSTexture, StrataAddressing.PixelCoords);
+			const bool bIsValid = StrataSubSurfaceHeaderGetIsValid(SSSHeader);
+			const bool bIsProfile = StrataSubSurfaceHeaderGetIsProfile(SSSHeader);
+			if (bHasSSS && bIsValid && bIsProfile)
+			{
+				const uint SubsurfaceProfileInt = StrataSubSurfaceHeaderGetProfileId(SSSHeader);
+				Print(Context, TEXT(" SSS ID.           "), FontBSDFPropName);	Print(Context, SubsurfaceProfileInt, FontBSDFPropValu);			Newline(Context, RectMax);
+			}
 		}
 		break;
 
diff --git a/Engine/Source/Runtime/Engine/Classes/Materials/MaterialExpressionStrata.h b/Engine/Source/Runtime/Engine/Classes/Materials/MaterialExpressionStrata.h
--- a/Engine/Source/Runtime/Engine/Classes/Materials/MaterialExpressionStrata.h
+++ b/Engine/Source/Runtime/Engine/Classes/Materials/MaterialExpressionStrata.h
@@ -598,6 +598,10 @@ class UMaterialExpressionStrataEyeBSDF : public UMaterialExpressionStrataBSDF
 	UPROPERTY()
 	FExpressionInput EmissiveColor;
 
+	/** SubsurfaceProfile, for Subsurface Scattering diffusion. */
+	UPROPERTY(EditAnywhere, BlueprintReadOnly, Category = Material, meta = (DisplayName = "Subsurface Profile"))
+	TObjectPtr<class USubsurfaceProfile> SubsurfaceProfile;
+
 	//~ Begin UMaterialExpression Interface
 #if WITH_EDITOR
 	virtual int32 Compile(class FMaterialCompiler* Compiler, int32 OutputIndex) override;
diff --git a/Engine/Source/Runtime/Engine/Private/Materials/HLSLMaterialTranslator.cpp b/Engine/Source/Runtime/Engine/Private/Materials/HLSLMaterialTranslator.cpp
--- a/Engine/Source/Runtime/Engine/Private/Materials/HLSLMaterialTranslator.cpp
+++ b/Engine/Source/Runtime/Engine/Private/Materials/HLSLMaterialTranslator.cpp
@@ -10490,7 +10490,7 @@ int32 FHLSLMaterialTranslator::StrataHairBSDF(int32 BaseColor, int32 Scatter, in
 	);
 }
 
-int32 FHLSLMaterialTranslator::StrataEyeBSDF(int32 DiffuseAlbedo, int32 Roughness, int32 IrisMask, int32 IrisDistance, int32 EmissiveColor, int32 CorneaNormal, int32 IrisNormal, const FString& SharedLocalBasisIndexMacro)
+int32 FHLSLMaterialTranslator::StrataEyeBSDF(int32 DiffuseAlbedo, int32 Roughness, int32 IrisMask, int32 IrisDistance, int32 EmissiveColor, int32 CorneaNormal, int32 IrisNormal, int32 SSSProfileId, const FString& SharedLocalBasisIndexMacro)
 {
 	return AddCodeChunk(
 		MCT_Strata, TEXT("GetStrataEyeBSDF(%s, %s, %s, %s, %s, %s, %s) /* Cornea:%s Iris:%s */"),
@@ -10499,6 +10499,7 @@ int32 FHLSLMaterialTranslator::StrataEyeBSDF(int32 DiffuseAlbedo, int32 Roughnes
 		*StrataGetCastParameterCode(IrisMask, MCT_Float),
 		*StrataGetCastParameterCode(IrisDistance, MCT_Float),
 		*StrataGetCastParameterCode(IrisNormal, MCT_Float),
+		*StrataGetCastParameterCode(SSSProfileId, MCT_Float),
 		*StrataGetCastParameterCode(EmissiveColor, MCT_Float3),
 		*SharedLocalBasisIndexMacro,
 		*GetParameterCode(CorneaNormal),
diff --git a/Engine/Source/Runtime/Engine/Private/Materials/HLSLMaterialTranslator.h b/Engine/Source/Runtime/Engine/Private/Materials/HLSLMaterialTranslator.h
--- a/Engine/Source/Runtime/Engine/Private/Materials/HLSLMaterialTranslator.h
+++ b/Engine/Source/Runtime/Engine/Private/Materials/HLSLMaterialTranslator.h
@@ -1046,7 +1046,7 @@ protected:
 	virtual int32 StrataVolumetricFogCloudBSDF(int32 Albedo, int32 Extinction, int32 EmissiveColor, int32 AmbientOcclusion) override;
 	virtual int32 StrataUnlitBSDF(int32 EmissiveColor, int32 TransmittanceColor) override;
 	virtual int32 StrataHairBSDF(int32 BaseColor, int32 Scatter, int32 Specular, int32 Roughness, int32 Backlit, int32 EmissiveColor, int32 Tangent, const FString& SharedLocalBasisIndexMacro) override;
-	virtual int32 StrataEyeBSDF(int32 DiffuseColor, int32 Roughness, int32 IrisMask, int32 IrisDistance, int32 EmissiveColor, int32 CorneaNormal, int32 IrisNormal, const FString& SharedLocalBasisIndexMacro) override;
+	virtual int32 StrataEyeBSDF(int32 DiffuseColor, int32 Roughness, int32 IrisMask, int32 IrisDistance, int32 EmissiveColor, int32 CorneaNormal, int32 IrisNormal, int32 SSSProfileId, const FString& SharedLocalBasisIndexMacro) override;
 	virtual int32 StrataSingleLayerWaterBSDF(
 		int32 BaseColor, int32 Metallic, int32 Specular, int32 Roughness, 
 		int32 EmissiveColor, int32 TopMaterialOpacity, int32 WaterAlbedo, int32 WaterExtinction, int32 WaterPhaseG, 
diff --git a/Engine/Source/Runtime/Engine/Private/Materials/Material.cpp b/Engine/Source/Runtime/Engine/Private/Materials/Material.cpp
--- a/Engine/Source/Runtime/Engine/Private/Materials/Material.cpp
+++ b/Engine/Source/Runtime/Engine/Private/Materials/Material.cpp
@@ -4617,7 +4617,7 @@ void UMaterial::RebuildShadingModelField()
 			ShadingModels.AddShadingModel(ShadingModel);
 
 			// Subsurface profil
-			if (StrataMaterialInfo.HasOnlyShadingModel(SSM_SubsurfaceLit) && StrataMaterialInfo.CountSubsurfaceProfiles() > 0)
+			if ((StrataMaterialInfo.HasOnlyShadingModel(SSM_Eye) || StrataMaterialInfo.HasOnlyShadingModel(SSM_SubsurfaceLit)) && StrataMaterialInfo.CountSubsurfaceProfiles() > 0)
 			{
 				if (StrataMaterialInfo.CountSubsurfaceProfiles() > 1)
 				{
diff --git a/Engine/Source/Runtime/Engine/Private/Materials/MaterialExpressions.cpp b/Engine/Source/Runtime/Engine/Private/Materials/MaterialExpressions.cpp
--- a/Engine/Source/Runtime/Engine/Private/Materials/MaterialExpressions.cpp
+++ b/Engine/Source/Runtime/Engine/Private/Materials/MaterialExpressions.cpp
@@ -23307,12 +23307,19 @@ int32 UMaterialExpressionStrataEyeBSDF::Compile(class FMaterialCompiler* Compile
 	int32 CorneaNormalCodeChunk = CompileWithDefaultTangentWS(Compiler, CorneaNormal);
 	const FStrataRegisteredSharedLocalBasis NewRegisteredSharedLocalBasis = StrataCompilationInfoCreateSharedLocalBasis(Compiler, CorneaNormalCodeChunk);
 
+	int32 SSSProfileCodeChunk = INDEX_NONE;
+	if (SubsurfaceProfile != nullptr)
+	{
+		SSSProfileCodeChunk = Compiler->ForceCast(Compiler->ScalarParameter(GetSubsurfaceProfileParameterName(), 1.0f), MCT_Float1);
+	}
+
 	int32 OutputCodeChunk = Compiler->StrataEyeBSDF(
 		CompileWithDefaultFloat3(Compiler, DiffuseColor, 0.0f, 0.0f, 0.0f),
 		CompileWithDefaultFloat1(Compiler, Roughness,	 0.5f),
 		CompileWithDefaultFloat1(Compiler, IrisMask,	 0.0f),
 		CompileWithDefaultFloat1(Compiler, IrisDistance, 0.0f),
 		CompileWithDefaultFloat3(Compiler, IrisNormal,   0.0f, 0.0f, 1.0f),
+		SSSProfileCodeChunk != INDEX_NONE ? SSSProfileCodeChunk : Compiler->Constant(0.0f),
 		CompileWithDefaultFloat3(Compiler, EmissiveColor,0.0f, 0.0f, 0.0f),
 		CorneaNormalCodeChunk,
 		Compiler->GetStrataSharedLocalBasisIndexMacro(NewRegisteredSharedLocalBasis));
@@ -23362,7 +23369,10 @@ void UMaterialExpressionStrataEyeBSDF::GatherStrataMaterialInfo(FStrataMaterialI
 	if (IrisMask.IsConnected())			{ StrataMaterialInfo.AddPropertyConnected(MP_CustomData0); }
 	if (IrisDistance.IsConnected())		{ StrataMaterialInfo.AddPropertyConnected(MP_CustomData1); }
 	if (EmissiveColor.IsConnected())	{ StrataMaterialInfo.AddPropertyConnected(MP_EmissiveColor); }
-
+	if (SubsurfaceProfile)
+	{
+		StrataMaterialInfo.AddSubsurfaceProfile(SubsurfaceProfile);
+	}
 	StrataMaterialInfo.AddShadingModel(SSM_Eye);
 }
 
diff --git a/Engine/Source/Runtime/Engine/Public/MaterialCompiler.h b/Engine/Source/Runtime/Engine/Public/MaterialCompiler.h
--- a/Engine/Source/Runtime/Engine/Public/MaterialCompiler.h
+++ b/Engine/Source/Runtime/Engine/Public/MaterialCompiler.h
@@ -589,7 +589,7 @@ public:
 	virtual int32 StrataVolumetricFogCloudBSDF(int32 Albedo, int32 Extinction, int32 EmissiveColor, int32 AmbientOcclusion) = 0;
 	virtual int32 StrataUnlitBSDF(int32 EmissiveColor, int32 TransmittanceColor) = 0;
 	virtual int32 StrataHairBSDF(int32 BaseColor, int32 Scatter, int32 Specular, int32 Roughness, int32 Backlit, int32 EmissiveColor, int32 Tangent, const FString& SharedLocalBasisIndexMacro) = 0;
-	virtual int32 StrataEyeBSDF(int32 DiffuseAlbedo, int32 Roughness, int32 IrisMask, int32 IrisDistance, int32 EmissiveColor, int32 CorneaNormal, int32 IrisNormal, const FString& SharedLocalBasisIndexMacro) = 0;
+	virtual int32 StrataEyeBSDF(int32 DiffuseAlbedo, int32 Roughness, int32 IrisMask, int32 IrisDistance, int32 EmissiveColor, int32 CorneaNormal, int32 IrisNormal, int32 SSSProfileId, const FString& SharedLocalBasisIndexMacro) = 0;
 	virtual int32 StrataSingleLayerWaterBSDF(
 		int32 BaseColor, int32 Metallic, int32 Specular, int32 Roughness, int32 EmissiveColor, int32 TopMaterialOpacity, 
 		int32 WaterAlbedo, int32 WaterExtinction, int32 WaterPhaseG, int32 ColorScaleBehindWater, int32 Normal, const FString& SharedLocalBasisIndexMacro, FStrataOperator* PromoteToOperator) = 0;
@@ -1182,9 +1182,9 @@ public:
 		return Compiler->StrataHairBSDF(BaseColor, Scatter, Specular, Roughness, Backlit, EmissiveColor, Tangent, SharedLocalBasisIndexMacro);
 	}
 
-	virtual int32 StrataEyeBSDF(int32 DiffuseAlbedo, int32 Roughness, int32 IrisMask, int32 IrisDistance, int32 EmissiveColor, int32 CorneaNormal, int32 IrisNormal, const FString& SharedLocalBasisIndexMacro) override
+	virtual int32 StrataEyeBSDF(int32 DiffuseAlbedo, int32 Roughness, int32 IrisMask, int32 IrisDistance, int32 EmissiveColor, int32 CorneaNormal, int32 IrisNormal, int32 SSSProfileId, const FString& SharedLocalBasisIndexMacro) override
 	{
-		return Compiler->StrataEyeBSDF(DiffuseAlbedo, Roughness, IrisMask, IrisDistance, EmissiveColor, CorneaNormal, IrisNormal, SharedLocalBasisIndexMacro);
+		return Compiler->StrataEyeBSDF(DiffuseAlbedo, Roughness, IrisMask, IrisDistance, EmissiveColor, CorneaNormal, IrisNormal, SSSProfileId, SharedLocalBasisIndexMacro);
 	}
 
 	virtual int32 StrataSingleLayerWaterBSDF(
