Add custom tangent support for Strata legacy conversion node.

diff --git a/Engine/Shaders/Private/Strata/StrataLegacyConversion.ush b/Engine/Shaders/Private/Strata/StrataLegacyConversion.ush
--- a/Engine/Shaders/Private/Strata/StrataLegacyConversion.ush
+++ b/Engine/Shaders/Private/Strata/StrataLegacyConversion.ush
@@ -144,6 +144,7 @@ FStrataData StrataConvertLegacyMaterialDynamic(
 	float3 RawNormal,
 	float3 RawTangent,
 	float3 RawClearCoatNormal,
+	float3 RawCustomTangent,
 	uint SharedLocalBasisIndex,
 	uint ClearCoatBottomNormal_SharedLocalBasisIndex,
 	inout uint SharedLocalBasisTypes,
@@ -477,6 +478,7 @@ FStrataData StrataConvertLegacyMaterialStatic(
 	float3 RawNormal,
 	float3 RawTangent,
 	float3 RawClearCoatNormal,
+	float3 RawCustomTangent,
 	uint SharedLocalBasisIndex,
 	uint ClearCoatBottomNormal_SharedLocalBasisIndex,
 	inout uint SharedLocalBasisTypes,
@@ -546,7 +548,7 @@ FStrataData StrataConvertLegacyMaterialStatic(
 
 		float3 IrisNormal = RawNormal;
 		float3 CausticNormal = RawNormal;
-		GetEyeNormals(IrisMask, IrisDistance, RawNormal, RawClearCoatNormal, RawClearCoatNormal /*RawCustomTangent*/, IrisNormal, CausticNormal);
+		GetEyeNormals(IrisMask, IrisDistance, RawNormal, RawClearCoatNormal, RawCustomTangent, IrisNormal, CausticNormal);
 
 		FStrataData EyeMat = PromoteParameterBlendedBSDFToOperator(
 			GetStrataEyeBSDF(
@@ -610,6 +612,7 @@ FStrataData StrataConvertLegacyMaterialStatic(
 			RawNormal,
 			RawTangent,
 			RawClearCoatNormal,
+			RawCustomTangent,
 			SharedLocalBasisIndex,
 			ClearCoatBottomNormal_SharedLocalBasisIndex,
 			SharedLocalBasisTypes,
diff --git a/Engine/Source/Runtime/Engine/Classes/Materials/MaterialExpressionStrata.h b/Engine/Source/Runtime/Engine/Classes/Materials/MaterialExpressionStrata.h
--- a/Engine/Source/Runtime/Engine/Classes/Materials/MaterialExpressionStrata.h
+++ b/Engine/Source/Runtime/Engine/Classes/Materials/MaterialExpressionStrata.h
@@ -152,6 +152,12 @@ class UMaterialExpressionStrataLegacyConversion : public UMaterialExpressionStra
 	UPROPERTY()
 	FExpressionInput ClearCoatNormal;
 
+	/**
+	 * Take the tangent output node as input. The tangent is considered tangent or world space according to the space properties on the main material node. (type = float3, unit = unitless, defaults to vertex tangent)
+	 */
+	UPROPERTY()
+	FExpressionInput CustomTangent;
+
 	/**
 	 * Shading models
 	 */
diff --git a/Engine/Source/Runtime/Engine/Private/Materials/HLSLMaterialTranslator.cpp b/Engine/Source/Runtime/Engine/Private/Materials/HLSLMaterialTranslator.cpp
--- a/Engine/Source/Runtime/Engine/Private/Materials/HLSLMaterialTranslator.cpp
+++ b/Engine/Source/Runtime/Engine/Private/Materials/HLSLMaterialTranslator.cpp
@@ -10411,7 +10411,8 @@ int32 FHLSLMaterialTranslator::StrataConversionFromLegacy(
 	int32 WaterScatteringCoefficients, int32 WaterAbsorptionCoefficients, int32 WaterPhaseG, int32 ColorScaleBehindWater,
 	int32 ShadingModel,
 	int32 Normal, int32 Tangent, const FString& SharedLocalBasisIndexMacro,
-	int32 ClearCoat_Normal, int32 ClearCoat_Tangent, const FString& ClearCoat_SharedLocalBasisIndexMacro)
+	int32 ClearCoat_Normal, int32 ClearCoat_Tangent, const FString& ClearCoat_SharedLocalBasisIndexMacro,
+	int32 CustomTangent_Tangent)
 {
 	const FString NormalCode = GetParameterCode(Normal);
 	const FString TangentCode = Tangent != INDEX_NONE ? *GetParameterCode(Tangent) : TEXT("NONE");
@@ -10421,7 +10422,7 @@ int32 FHLSLMaterialTranslator::StrataConversionFromLegacy(
 	const FString ClearCoat_TangentCode = Tangent != INDEX_NONE ? *GetParameterCode(ClearCoat_Tangent) : TEXT("NONE");
 
 	return AddCodeChunk(
-		MCT_Strata, TEXT("StrataConvertLegacyMaterial%s(Parameters.StrataPixelFootprint, %s, %s, %s, %s, %s, %s, %s, %s, %s, %s, %s, %s, %s, %s, %s, %s, %s, %s, %s, %s, %s, %s, Parameters.SharedLocalBases.Types, Parameters.StrataTree) /* Normal = %s ; Tangent = %s ; ClearCoat_Normal = %s ; ClearCoat_Tangent = %s */"),
+		MCT_Strata, TEXT("StrataConvertLegacyMaterial%s(Parameters.StrataPixelFootprint, %s, %s, %s, %s, %s, %s, %s, %s, %s, %s, %s, %s, %s, %s, %s, %s, %s, %s, %s, %s, %s, %s, %s, Parameters.SharedLocalBases.Types, Parameters.StrataTree) /* Normal = %s ; Tangent = %s ; ClearCoat_Normal = %s ; ClearCoat_Tangent = %s */"),
 		bHasDynamicShadingModels ? TEXT("Dynamic") : TEXT("Static"),
 		*StrataGetCastParameterCode(BaseColor,						MCT_Float3),
 		*StrataGetCastParameterCode(Specular,						MCT_Float),
@@ -10440,10 +10441,11 @@ int32 FHLSLMaterialTranslator::StrataConversionFromLegacy(
 		*StrataGetCastParameterCode(WaterPhaseG,					MCT_Float),
 		*StrataGetCastParameterCode(ColorScaleBehindWater,			MCT_Float3),
 		*GetParameterCode(ShadingModel),
-		// Raw access to Normal/Tanget/ClearCoatNormal for conversion purpose
+		// Raw access to Normal/Tangent/ClearCoatNormal/CustomTangent for conversion purpose
 		*StrataGetCastParameterCode(Normal,							MCT_Float3),
 		*StrataGetCastParameterCode(RawTangent,						MCT_Float3),
 		*StrataGetCastParameterCode(ClearCoat_Normal,				MCT_Float3),
+		*StrataGetCastParameterCode(CustomTangent_Tangent,			MCT_Float3),
 		*SharedLocalBasisIndexMacro,
 		*ClearCoat_SharedLocalBasisIndexMacro,
 		// Regular normal basis
diff --git a/Engine/Source/Runtime/Engine/Private/Materials/HLSLMaterialTranslator.h b/Engine/Source/Runtime/Engine/Private/Materials/HLSLMaterialTranslator.h
--- a/Engine/Source/Runtime/Engine/Private/Materials/HLSLMaterialTranslator.h
+++ b/Engine/Source/Runtime/Engine/Private/Materials/HLSLMaterialTranslator.h
@@ -1042,7 +1042,8 @@ protected:
 		int32 WaterScatteringCoefficients, int32 WaterAbsorptionCoefficients, int32 WaterPhaseG, int32 ColorScaleBehindWater,
 		int32 ShadingModel,
 		int32 Normal, int32 Tangent, const FString& SharedLocalBasisIndexMacro,
-		int32 ClearCoat_Normal, int32 ClearCoat_Tangent, const FString& ClearCoat_SharedLocalBasisIndexMacro) override;
+		int32 ClearCoat_Normal, int32 ClearCoat_Tangent, const FString& ClearCoat_SharedLocalBasisIndexMacro, 
+		int32 CustomTangent_Tangent) override;
 	virtual int32 StrataVolumetricFogCloudBSDF(int32 Albedo, int32 Extinction, int32 EmissiveColor, int32 AmbientOcclusion) override;
 	virtual int32 StrataUnlitBSDF(int32 EmissiveColor, int32 TransmittanceColor) override;
 	virtual int32 StrataHairBSDF(int32 BaseColor, int32 Scatter, int32 Specular, int32 Roughness, int32 Backlit, int32 EmissiveColor, int32 Tangent, const FString& SharedLocalBasisIndexMacro) override;
diff --git a/Engine/Source/Runtime/Engine/Private/Materials/Material.cpp b/Engine/Source/Runtime/Engine/Private/Materials/Material.cpp
--- a/Engine/Source/Runtime/Engine/Private/Materials/Material.cpp
+++ b/Engine/Source/Runtime/Engine/Private/Materials/Material.cpp
@@ -3198,7 +3198,7 @@ void UMaterial::ConvertMaterialToStrataMaterial()
 					check(EditorOnly->ShadingModelFromMaterialExpression.IsConnected());
 
 					// Reconnect the shading model expression
-					MoveConnectionTo(EditorOnly->ShadingModelFromMaterialExpression, ConvertNode, 18);
+					MoveConnectionTo(EditorOnly->ShadingModelFromMaterialExpression, ConvertNode, 19);
 
 					// Store strata shading model of the converted material. 
 					GatherCustomNodes();
@@ -3355,7 +3355,11 @@ void UMaterial::ConvertMaterialToStrataMaterial()
 		}
 		if (ClearCoatBottomNormalOutput)
 		{
-			CopyConnectionTo(*ClearCoatBottomNormalOutput->GetInput(0), ConvertNode, 17);		 // ClearCoatNormal
+			CopyConnectionTo(*ClearCoatBottomNormalOutput->GetInput(0), ConvertNode, 17); // ClearCoatNormal
+		}
+		if (TangentOutput)
+		{
+			CopyConnectionTo(*TangentOutput->GetInput(0), ConvertNode, 18);	// TangentOutput
 		}
 	}
 
diff --git a/Engine/Source/Runtime/Engine/Private/Materials/MaterialExpressions.cpp b/Engine/Source/Runtime/Engine/Private/Materials/MaterialExpressions.cpp
--- a/Engine/Source/Runtime/Engine/Private/Materials/MaterialExpressions.cpp
+++ b/Engine/Source/Runtime/Engine/Private/Materials/MaterialExpressions.cpp
@@ -21654,6 +21654,18 @@ int32 UMaterialExpressionStrataLegacyConversion::Compile(class FMaterialCompiler
 		ClearCoat_BasisIndexMacro = BasisIndexMacro;
 	}
 
+	// Custom tangent. No need to register it as a local basis, as it is only used for eye shading internal conversion
+	int32 CustomTangent_TangentCodeChunk = INDEX_NONE;
+	const bool bHasCustomTangent = CustomTangent.IsConnected();
+	if (bHasCustomTangent)
+	{
+		CustomTangent_TangentCodeChunk = CompileWithDefaultNormalWS(Compiler, CustomTangent);
+	}
+	else
+	{
+		CustomTangent_TangentCodeChunk = NormalCodeChunk;
+	}
+
 	int32 SSSProfileCodeChunk = INDEX_NONE;
 	const bool bHasSSS = HasSSS();
 	if (bHasSSS)
@@ -21694,7 +21706,8 @@ int32 UMaterialExpressionStrataLegacyConversion::Compile(class FMaterialCompiler
 		BasisIndexMacro,
 		ClearCoat_NormalCodeChunk,
 		ClearCoat_TangentCodeChunk,
-		ClearCoat_BasisIndexMacro);
+		ClearCoat_BasisIndexMacro,
+		CustomTangent_TangentCodeChunk);
 
 	// Now update the local shared basis
 
@@ -21763,6 +21776,7 @@ const TArray<FExpressionInput*> UMaterialExpressionStrataLegacyConversion::GetIn
 	Result.Add(&WaterPhaseG);
 	Result.Add(&ColorScaleBehindWater);
 	Result.Add(&ClearCoatNormal);
+	Result.Add(&CustomTangent);
 	Result.Add(&ShadingModel);
 	return Result;
 }
@@ -21806,7 +21820,8 @@ uint32 UMaterialExpressionStrataLegacyConversion::GetInputType(int32 InputIndex)
 	else if (InputIndex == 15) return MCT_Float1; // WaterPhaseG
 	else if (InputIndex == 16) return MCT_Float3; // ColorScaleBehindWater
 	else if (InputIndex == 17) return MCT_Float3; // ClearCoatNormal
-	else if (InputIndex == 18) return MCT_Float1; // ShadingModel
+	else if (InputIndex == 18) return MCT_Float3; // CustomTangent
+	else if (InputIndex == 19) return MCT_Float1; // ShadingModel
 
 	check(false);
 	return MCT_Float1;
@@ -21832,7 +21847,8 @@ FName UMaterialExpressionStrataLegacyConversion::GetInputName(int32 InputIndex)
 	else if (InputIndex == 15)	return TEXT("Water Phase G");
 	else if (InputIndex == 16)	return TEXT("Color Scale BehindWater");
 	else if (InputIndex == 17)	return TEXT("Clear Coat Normal");
-	else if (InputIndex == 18)	return TEXT("Shading Model");
+	else if (InputIndex == 18)	return TEXT("Custom Tangent");
+	else if (InputIndex == 19)	return TEXT("Shading Model");
 	return TEXT("Unknown");
 }
 
diff --git a/Engine/Source/Runtime/Engine/Public/MaterialCompiler.h b/Engine/Source/Runtime/Engine/Public/MaterialCompiler.h
--- a/Engine/Source/Runtime/Engine/Public/MaterialCompiler.h
+++ b/Engine/Source/Runtime/Engine/Public/MaterialCompiler.h
@@ -585,7 +585,8 @@ public:
 		int32 WaterScatteringCoefficients, int32 WaterAbsorptionCoefficients, int32 WaterPhaseG, int32 ColorScaleBehindWater,
 		int32 ShadingModel,
 		int32 Normal, int32 Tangent, const FString& SharedLocalBasisIndexMacro,
-		int32 ClearCoat_Normal, int32 ClearCoat_Tangent, const FString& ClearCoat_SharedLocalBasisIndexMacro) = 0;
+		int32 ClearCoat_Normal, int32 ClearCoat_Tangent, const FString& ClearCoat_SharedLocalBasisIndexMacro,
+		int32 CustomTangent_Tangent) = 0;
 	virtual int32 StrataVolumetricFogCloudBSDF(int32 Albedo, int32 Extinction, int32 EmissiveColor, int32 AmbientOcclusion) = 0;
 	virtual int32 StrataUnlitBSDF(int32 EmissiveColor, int32 TransmittanceColor) = 0;
 	virtual int32 StrataHairBSDF(int32 BaseColor, int32 Scatter, int32 Specular, int32 Roughness, int32 Backlit, int32 EmissiveColor, int32 Tangent, const FString& SharedLocalBasisIndexMacro) = 0;
@@ -1150,7 +1151,8 @@ public:
 		int32 WaterScatteringCoefficients, int32 WaterAbsorptionCoefficients, int32 WaterPhaseG, int32 ColorScaleBehindWater,
 		int32 ShadingModel,
 		int32 Normal, int32 Tangent, const FString& SharedLocalBasisIndexMacro,
-		int32 ClearCoat_Normal, int32 ClearCoat_Tangent, const FString& ClearCoat_SharedLocalBasisIndexMacro) override
+		int32 ClearCoat_Normal, int32 ClearCoat_Tangent, const FString& ClearCoat_SharedLocalBasisIndexMacro,
+		int32 CustomTangent_Tangent) override
 	{
 		return Compiler->StrataConversionFromLegacy(
 			bHasDynamicShadingModels,
@@ -1164,7 +1166,8 @@ public:
 			WaterScatteringCoefficients, WaterAbsorptionCoefficients, WaterPhaseG, ColorScaleBehindWater,
 			ShadingModel,
 			Normal, Tangent, SharedLocalBasisIndexMacro,
-			ClearCoat_Normal, ClearCoat_Tangent, ClearCoat_SharedLocalBasisIndexMacro);
+			ClearCoat_Normal, ClearCoat_Tangent, ClearCoat_SharedLocalBasisIndexMacro,
+			CustomTangent_Tangent);
 	}
 
 	virtual int32 StrataVolumetricFogCloudBSDF(int32 Albedo, int32 Extinction, int32 EmissiveColor, int32 AmbientOcclusion) override
