Fix for pose watch viewport masks drawing the wrong bones.

diff --git a/Engine/Source/Runtime/Engine/Private/SkeletalDebugRendering.cpp b/Engine/Source/Runtime/Engine/Private/SkeletalDebugRendering.cpp
--- a/Engine/Source/Runtime/Engine/Private/SkeletalDebugRendering.cpp
+++ b/Engine/Source/Runtime/Engine/Private/SkeletalDebugRendering.cpp
@@ -222,7 +222,8 @@ void DrawBonesFromPoseWatch(
 		}
 		else
 		{
-			const bool bBoneBlendScaleMeetsThreshold = ViewportMask->GetBoneBlendScale(MeshBoneIndex.GetInt()) > FMath::Clamp(PoseWatch->BlendScaleThreshold, 0.f, 1.f);
+			const FName& BoneName = MeshComponent->GetBoneName(MeshBoneIndex.GetInt());
+			const bool bBoneBlendScaleMeetsThreshold = ViewportMask->GetBoneBlendScale(BoneName) > PoseWatch->BlendScaleThreshold;
 			if ((!PoseWatch->bInvertViewportMask && bBoneBlendScaleMeetsThreshold) || (PoseWatch->bInvertViewportMask && !bBoneBlendScaleMeetsThreshold))
 			{
 				RequiredBones.Add(MeshBoneIndex.GetInt());
