Allow formatting strings to be binary-compatible char types, not only exact types.

diff --git a/Engine/Source/Developer/GameplayDebugger/Public/GameplayDebuggerTypes.h b/Engine/Source/Developer/GameplayDebugger/Public/GameplayDebuggerTypes.h
--- a/Engine/Source/Developer/GameplayDebugger/Public/GameplayDebuggerTypes.h
+++ b/Engine/Source/Developer/GameplayDebugger/Public/GameplayDebuggerTypes.h
@@ -4,6 +4,8 @@
 
 #include "CoreMinimal.h"
 #include "Engine/EngineTypes.h"
+#include "Templates/IsArrayOrRefOfTypeByPredicate.h"
+#include "Traits/IsCharEncodingCompatibleWith.h"
 
 class FCanvasItem;
 class UCanvas;
@@ -61,55 +63,55 @@ public:
 	template <typename FmtType, typename... Types>
 	void Printf(const FmtType& Fmt, Types... Args)
 	{
-		static_assert(TIsArrayOrRefOfType<FmtType, TCHAR>::Value, "Formatting string must be a TCHAR array.");
+		static_assert(TIsArrayOrRefOfTypeByPredicate<FmtType, TIsCharEncodingCompatibleWithTCHAR>::Value, "Formatting string must be a TCHAR array.");
 		static_assert(TAnd<TIsValidVariadicFunctionArg<Types>...>::Value, "Invalid argument(s) passed to FGameplayDebuggerCanvasContext::Printf");
 
-		PrintfImpl(FColor::White, 1.0f, Fmt, Args...);
+		PrintfImpl(FColor::White, 1.0f, (const TCHAR*)Fmt, Args...);
 	}
 	
 	template <typename FmtType, typename... Types>
 	void Printf(const FColor& Color, const FmtType& Fmt, Types... Args)
 	{
-		static_assert(TIsArrayOrRefOfType<FmtType, TCHAR>::Value, "Formatting string must be a TCHAR array.");
+		static_assert(TIsArrayOrRefOfTypeByPredicate<FmtType, TIsCharEncodingCompatibleWithTCHAR>::Value, "Formatting string must be a TCHAR array.");
 		static_assert(TAnd<TIsValidVariadicFunctionArg<Types>...>::Value, "Invalid argument(s) passed to FGameplayDebuggerCanvasContext::Printf");
 
-		PrintfImpl(Color, 1.0f, Fmt, Args...);
+		PrintfImpl(Color, 1.0f, (const TCHAR*)Fmt, Args...);
 	}
 	
 	template <typename FmtType, typename... Types>
 	void Printf(const FColor& Color, const float Alpha, const FmtType& Fmt, Types... Args)
 	{
-		static_assert(TIsArrayOrRefOfType<FmtType, TCHAR>::Value, "Formatting string must be a TCHAR array.");
+		static_assert(TIsArrayOrRefOfTypeByPredicate<FmtType, TIsCharEncodingCompatibleWithTCHAR>::Value, "Formatting string must be a TCHAR array.");
 		static_assert(TAnd<TIsValidVariadicFunctionArg<Types>...>::Value, "Invalid argument(s) passed to FGameplayDebuggerCanvasContext::Printf");
 
-		PrintfImpl(Color, Alpha, Fmt, Args...);
+		PrintfImpl(Color, Alpha, (const TCHAR*)Fmt, Args...);
 	}
 	
 	template <typename FmtType, typename... Types>
 	void PrintfAt(float PosX, float PosY, const FmtType& Fmt, Types... Args)
 	{
-		static_assert(TIsArrayOrRefOfType<FmtType, TCHAR>::Value, "Formatting string must be a TCHAR array.");
+		static_assert(TIsArrayOrRefOfTypeByPredicate<FmtType, TIsCharEncodingCompatibleWithTCHAR>::Value, "Formatting string must be a TCHAR array.");
 		static_assert(TAnd<TIsValidVariadicFunctionArg<Types>...>::Value, "Invalid argument(s) passed to FGameplayDebuggerCanvasContext::PrintfAt");
 
-		PrintfAtImpl(PosX, PosY, FColor::White, 1.0f, Fmt, Args...);
+		PrintfAtImpl(PosX, PosY, FColor::White, 1.0f, (const TCHAR*)Fmt, Args...);
 	}
 	
 	template <typename FmtType, typename... Types>
 	void PrintfAt(float PosX, float PosY, const FColor& Color, const FmtType& Fmt, Types... Args)
 	{
-		static_assert(TIsArrayOrRefOfType<FmtType, TCHAR>::Value, "Formatting string must be a TCHAR array.");
+		static_assert(TIsArrayOrRefOfTypeByPredicate<FmtType, TIsCharEncodingCompatibleWithTCHAR>::Value, "Formatting string must be a TCHAR array.");
 		static_assert(TAnd<TIsValidVariadicFunctionArg<Types>...>::Value, "Invalid argument(s) passed to FGameplayDebuggerCanvasContext::PrintfAt");
 
-		PrintfAtImpl(PosX, PosY, Color, 1.0f, Fmt, Args...);
+		PrintfAtImpl(PosX, PosY, Color, 1.0f, (const TCHAR*)Fmt, Args...);
 	}
 	
 	template <typename FmtType, typename... Types>
 	void PrintfAt(float PosX, float PosY, const FColor& Color, const float Alpha, const FmtType& Fmt, Types... Args)
 	{
-		static_assert(TIsArrayOrRefOfType<FmtType, TCHAR>::Value, "Formatting string must be a TCHAR array.");
+		static_assert(TIsArrayOrRefOfTypeByPredicate<FmtType, TIsCharEncodingCompatibleWithTCHAR>::Value, "Formatting string must be a TCHAR array.");
 		static_assert(TAnd<TIsValidVariadicFunctionArg<Types>...>::Value, "Invalid argument(s) passed to FGameplayDebuggerCanvasContext::PrintfAt");
 
-		PrintfAtImpl(PosX, PosY, Color, Alpha, Fmt, Args...);
+		PrintfAtImpl(PosX, PosY, Color, Alpha, (const TCHAR*)Fmt, Args...);
 	}
 
 private:
diff --git a/Engine/Source/Programs/UnrealLightmass/Private/ImportExport/LightmassSwarm.h b/Engine/Source/Programs/UnrealLightmass/Private/ImportExport/LightmassSwarm.h
--- a/Engine/Source/Programs/UnrealLightmass/Private/ImportExport/LightmassSwarm.h
+++ b/Engine/Source/Programs/UnrealLightmass/Private/ImportExport/LightmassSwarm.h
@@ -13,7 +13,8 @@
 #include "Editor/SwarmInterface/Public/SwarmInterface.h"
 #include "Templates/IsValidVariadicFunctionArg.h"
 #include "Templates/AndOrNot.h"
-#include "Templates/IsArrayOrRefOfType.h"
+#include "Templates/IsArrayOrRefOfTypeByPredicate.h"
+#include "Traits/IsCharEncodingCompatibleWith.h"
 #include "ImportExport.h"
 
 namespace Lightmass
@@ -234,10 +235,10 @@ public:
 	template <typename FmtType, typename... Types>
 	void SendTextMessage(const FmtType& Fmt, Types... Args)
 	{
-		static_assert(TIsArrayOrRefOfType<FmtType, TCHAR>::Value, "Formatting string must be a TCHAR array.");
+		static_assert(TIsArrayOrRefOfTypeByPredicate<FmtType, TIsCharEncodingCompatibleWithTCHAR>::Value, "Formatting string must be a TCHAR array.");
 		static_assert(TAnd<TIsValidVariadicFunctionArg<Types>...>::Value, "Invalid argument(s) passed to FString::Printf");
 
-		SendTextMessageImpl(Fmt, Args...);
+		SendTextMessageImpl((const TCHAR*)Fmt, Args...);
 	}
 
 private:
diff --git a/Engine/Source/Runtime/Core/Public/Containers/UnrealString.h b/Engine/Source/Runtime/Core/Public/Containers/UnrealString.h
--- a/Engine/Source/Runtime/Core/Public/Containers/UnrealString.h
+++ b/Engine/Source/Runtime/Core/Public/Containers/UnrealString.h
@@ -22,10 +22,11 @@
 #include "Templates/Invoke.h"
 #include "Templates/IsValidVariadicFunctionArg.h"
 #include "Templates/AndOrNot.h"
-#include "Templates/IsArrayOrRefOfType.h"
+#include "Templates/IsArrayOrRefOfTypeByPredicate.h"
 #include "Templates/TypeHash.h"
 #include "Templates/IsFloatingPoint.h"
 #include "Traits/IsCharType.h"
+#include "Traits/IsCharEncodingCompatibleWith.h"
 #include "Traits/IsCharEncodingSimplyConvertibleTo.h"
 
 struct FStringFormatArg;
@@ -1237,10 +1238,10 @@ public:
 	template <typename FmtType, typename... Types>
 	UE_NODISCARD static FString Printf(const FmtType& Fmt, Types... Args)
 	{
-		static_assert(TIsArrayOrRefOfType<FmtType, TCHAR>::Value, "Formatting string must be a TCHAR array.");
+		static_assert(TIsArrayOrRefOfTypeByPredicate<FmtType, TIsCharEncodingCompatibleWithTCHAR>::Value, "Formatting string must be a TCHAR array.");
 		static_assert(TAnd<TIsValidVariadicFunctionArg<Types>...>::Value, "Invalid argument(s) passed to FString::Printf");
 
-		return PrintfImpl(Fmt, Args...);
+		return PrintfImpl((const TCHAR*)Fmt, Args...);
 	}
 
 	/**
@@ -1250,10 +1251,10 @@ public:
 	template <typename FmtType, typename... Types>
 	FString& Appendf(const FmtType& Fmt, Types... Args)
 	{
-		static_assert(TIsArrayOrRefOfType<FmtType, TCHAR>::Value, "Formatting string must be a TCHAR array.");
+		static_assert(TIsArrayOrRefOfTypeByPredicate<FmtType, TIsCharEncodingCompatibleWithTCHAR>::Value, "Formatting string must be a TCHAR array.");
 		static_assert(TAnd<TIsValidVariadicFunctionArg<Types>...>::Value, "Invalid argument(s) passed to FString::Appendf");
 
-		AppendfImpl(*this, Fmt, Args...);
+		AppendfImpl(*this, (const TCHAR*)Fmt, Args...);
 		return *this;
 	}
 
diff --git a/Engine/Source/Runtime/Core/Public/Logging/LogMacros.h b/Engine/Source/Runtime/Core/Public/Logging/LogMacros.h
--- a/Engine/Source/Runtime/Core/Public/Logging/LogMacros.h
+++ b/Engine/Source/Runtime/Core/Public/Logging/LogMacros.h
@@ -11,7 +11,8 @@
 #include "Logging/LogTrace.h"
 #include "Templates/IsValidVariadicFunctionArg.h"
 #include "Templates/AndOrNot.h"
-#include "Templates/IsArrayOrRefOfType.h"
+#include "Templates/IsArrayOrRefOfTypeByPredicate.h"
+#include "Traits/IsCharEncodingCompatibleWith.h"
 
 
 /*----------------------------------------------------------------------------
@@ -34,30 +35,30 @@ struct CORE_API FMsg
 	template <typename FmtType, typename... Types>
 	static void SendNotificationStringf(const FmtType& Fmt, Types... Args)
 	{
-		static_assert(TIsArrayOrRefOfType<FmtType, TCHAR>::Value, "Formatting string must be a TCHAR array.");
+		static_assert(TIsArrayOrRefOfTypeByPredicate<FmtType, TIsCharEncodingCompatibleWithTCHAR>::Value, "Formatting string must be a TCHAR array.");
 		static_assert(TAnd<TIsValidVariadicFunctionArg<Types>...>::Value, "Invalid argument(s) passed to FMsg::SendNotificationStringf");
 
-		SendNotificationStringfImpl(Fmt, Args...);
+		SendNotificationStringfImpl((const TCHAR*)Fmt, Args...);
 	}
 
 	/** Log function */
 	template <typename FmtType, typename... Types>
 	static void Logf(const ANSICHAR* File, int32 Line, const FLogCategoryName& Category, ELogVerbosity::Type Verbosity, const FmtType& Fmt, Types... Args)
 	{
-		static_assert(TIsArrayOrRefOfType<FmtType, TCHAR>::Value, "Formatting string must be a TCHAR array.");
+		static_assert(TIsArrayOrRefOfTypeByPredicate<FmtType, TIsCharEncodingCompatibleWithTCHAR>::Value, "Formatting string must be a TCHAR array.");
 		static_assert(TAnd<TIsValidVariadicFunctionArg<Types>...>::Value, "Invalid argument(s) passed to FMsg::Logf");
 
-		LogfImpl(File, Line, Category, Verbosity, Fmt, Args...);
+		LogfImpl(File, Line, Category, Verbosity, (const TCHAR*)Fmt, Args...);
 	}
 
 	/** Internal version of log function. Should be used only in logging macros, as it relies on caller to call assert on fatal error */
 	template <typename FmtType, typename... Types>
 	static void Logf_Internal(const ANSICHAR* File, int32 Line, const FLogCategoryName& Category, ELogVerbosity::Type Verbosity, const FmtType& Fmt, Types... Args)
 	{
-		static_assert(TIsArrayOrRefOfType<FmtType, TCHAR>::Value, "Formatting string must be a TCHAR array.");
+		static_assert(TIsArrayOrRefOfTypeByPredicate<FmtType, TIsCharEncodingCompatibleWithTCHAR>::Value, "Formatting string must be a TCHAR array.");
 		static_assert(TAnd<TIsValidVariadicFunctionArg<Types>...>::Value, "Invalid argument(s) passed to FMsg::Logf_Internal");
 
-		Logf_InternalImpl(File, Line, Category, Verbosity, Fmt, Args...);
+		Logf_InternalImpl(File, Line, Category, Verbosity, (const TCHAR*)Fmt, Args...);
 	}
 
 private:
@@ -188,7 +189,7 @@ private:
 	 * INTERNAL IMPLEMENTATION. DO NOT CALL DIRECTLY!
 	**/
 	#define UE_INTERNAL_LOG_IMPL(CategoryName, Verbosity, Format, ...) \
-		static_assert(TIsArrayOrRefOfType<decltype(Format), TCHAR>::Value, "Formatting string must be a TCHAR array."); \
+		static_assert(TIsArrayOrRefOfTypeByPredicate<decltype(Format), TIsCharEncodingCompatibleWithTCHAR>::Value, "Formatting string must be a TCHAR array."); \
 		static_assert((ELogVerbosity::Verbosity & ELogVerbosity::VerbosityMask) < ELogVerbosity::NumVerbosity && ELogVerbosity::Verbosity > 0, "Verbosity must be constant and in range."); \
 		UE_LOG_EXPAND_IS_FATAL(Verbosity, PREPROCESSOR_NOTHING, if (!CategoryName.IsSuppressed(ELogVerbosity::Verbosity))) \
 			{ \
@@ -246,7 +247,7 @@ private:
 	 ***/
 	#define UE_LOG_CLINKAGE(CategoryName, Verbosity, Format, ...) \
 	{ \
-		static_assert(TIsArrayOrRefOfType<decltype(Format), TCHAR>::Value, "Formatting string must be a TCHAR array."); \
+		static_assert(TIsArrayOrRefOfTypeByPredicate<decltype(Format), TIsCharEncodingCompatibleWithTCHAR>::Value, "Formatting string must be a TCHAR array."); \
 		static_assert((ELogVerbosity::Verbosity & ELogVerbosity::VerbosityMask) < ELogVerbosity::NumVerbosity && ELogVerbosity::Verbosity > 0, "Verbosity must be constant and in range."); \
 		CA_CONSTANT_IF((ELogVerbosity::Verbosity & ELogVerbosity::VerbosityMask) <= ELogVerbosity::COMPILED_IN_MINIMUM_VERBOSITY && (ELogVerbosity::Warning & ELogVerbosity::VerbosityMask) <= FLogCategory##CategoryName::CompileTimeVerbosity) \
 		{ \
@@ -276,7 +277,7 @@ private:
 	#define UE_SECURITY_LOG(NetConnection, SecurityEventType, Format, ...) \
 	{ \
 		DEPRECATED_MACRO(5.1, "UE_SECURITY_LOG has been deprecated in favor of UE_LOG") \
-		static_assert(TIsArrayOrRefOfType<decltype(Format), TCHAR>::Value, "Formatting string must be a TCHAR array."); \
+		static_assert(TIsArrayOrRefOfTypeByPredicate<decltype(Format), TIsCharEncodingCompatibleWithTCHAR>::Value, "Formatting string must be a TCHAR array."); \
 		check(NetConnection != nullptr); \
 		CA_CONSTANT_IF((ELogVerbosity::Warning & ELogVerbosity::VerbosityMask) <= ELogVerbosity::COMPILED_IN_MINIMUM_VERBOSITY && (ELogVerbosity::Warning & ELogVerbosity::VerbosityMask) <= FLogCategoryLogSecurity::CompileTimeVerbosity) \
 		{ \
@@ -290,7 +291,7 @@ private:
 	// Conditional logging. Will only log if Condition is met.
 	#define UE_CLOG(Condition, CategoryName, Verbosity, Format, ...) \
 	{ \
-		static_assert(TIsArrayOrRefOfType<decltype(Format), TCHAR>::Value, "Formatting string must be a TCHAR array."); \
+		static_assert(TIsArrayOrRefOfTypeByPredicate<decltype(Format), TIsCharEncodingCompatibleWithTCHAR>::Value, "Formatting string must be a TCHAR array."); \
 		static_assert((ELogVerbosity::Verbosity & ELogVerbosity::VerbosityMask) < ELogVerbosity::NumVerbosity && ELogVerbosity::Verbosity > 0, "Verbosity must be constant and in range."); \
 		CA_CONSTANT_IF((ELogVerbosity::Verbosity & ELogVerbosity::VerbosityMask) <= ELogVerbosity::COMPILED_IN_MINIMUM_VERBOSITY && (ELogVerbosity::Warning & ELogVerbosity::VerbosityMask) <= FLogCategory##CategoryName::CompileTimeVerbosity) \
 		{ \
@@ -405,7 +406,7 @@ private:
 #define CLOSE_CONNECTION_DUE_TO_SECURITY_VIOLATION_INNER(NetConnection, SecurityEventType, Format, ...) \
 { \
 	DEPRECATED_MACRO(5.1, "CLOSE_CONNECTION_DUE_TO_SECURITY_VIOLATION has been deprecated") \
-	static_assert(TIsArrayOrRefOfType<decltype(Format), TCHAR>::Value, "Formatting string must be a TCHAR array."); \
+	static_assert(TIsArrayOrRefOfTypeByPredicate<decltype(Format), TIsCharEncodingCompatibleWithTCHAR>::Value, "Formatting string must be a TCHAR array."); \
 	check(NetConnection != nullptr); \
 	FString SecurityPrint = FString::Printf(Format, ##__VA_ARGS__); \
 	UE_LOG(LogSecurity, Warning, TEXT("%s: %s: %s"), *(NetConnection)->RemoteAddressToString(), ToString(SecurityEventType), SecurityEventType, *SecurityPrint); \
diff --git a/Engine/Source/Runtime/Core/Public/Logging/LogTrace.h b/Engine/Source/Runtime/Core/Public/Logging/LogTrace.h
--- a/Engine/Source/Runtime/Core/Public/Logging/LogTrace.h
+++ b/Engine/Source/Runtime/Core/Public/Logging/LogTrace.h
@@ -4,6 +4,8 @@
 
 #include "CoreTypes.h"
 #include "Trace/Config.h"
+#include "Templates/IsArrayOrRefOfTypeByPredicate.h"
+#include "Traits/IsCharEncodingCompatibleWith.h"
 
 #if !defined(LOGTRACE_ENABLED)
 #if UE_TRACE_ENABLED && !UE_BUILD_SHIPPING
@@ -62,12 +64,13 @@ private:
 	}
 #else
 #define TRACE_LOG_MESSAGE(Category, Verbosity, Format, ...) \
+	static_assert(TIsArrayOrRefOfTypeByPredicate<decltype(Format), TIsCharEncodingCompatibleWithTCHAR>::Value, "Formatting string must be a TCHAR array."); \
 	if (UE_TRACE_CHANNELEXPR_IS_ENABLED(LogChannel)) \
 	{ \
 		static bool PREPROCESSOR_JOIN(__LogPoint, __LINE__); \
 		if (!PREPROCESSOR_JOIN(__LogPoint, __LINE__)) \
 		{ \
-			FLogTrace::OutputLogMessageSpec(&PREPROCESSOR_JOIN(__LogPoint, __LINE__), &Category, ELogVerbosity::Verbosity, __FILE__, __LINE__, Format); \
+			FLogTrace::OutputLogMessageSpec(&PREPROCESSOR_JOIN(__LogPoint, __LINE__), &Category, ELogVerbosity::Verbosity, __FILE__, __LINE__, (const TCHAR*)Format); \
 			PREPROCESSOR_JOIN(__LogPoint, __LINE__) = true; \
 		} \
 		FLogTrace::OutputLogMessage(&PREPROCESSOR_JOIN(__LogPoint, __LINE__), ##__VA_ARGS__); \
diff --git a/Engine/Source/Runtime/Core/Public/Misc/AssertionMacros.h b/Engine/Source/Runtime/Core/Public/Misc/AssertionMacros.h
--- a/Engine/Source/Runtime/Core/Public/Misc/AssertionMacros.h
+++ b/Engine/Source/Runtime/Core/Public/Misc/AssertionMacros.h
@@ -8,8 +8,9 @@
 #include "HAL/PreprocessorHelpers.h"
 #include "Templates/AndOrNot.h"
 #include "Templates/EnableIf.h"
-#include "Templates/IsArrayOrRefOfType.h"
+#include "Templates/IsArrayOrRefOfTypeByPredicate.h"
 #include "Templates/IsValidVariadicFunctionArg.h"
+#include "Traits/IsCharEncodingCompatibleWith.h"
 #include "Misc/VarArgs.h"
 
 #if (DO_CHECK || DO_GUARD_SLOW || DO_ENSURE) && !PLATFORM_CPU_ARM_FAMILY
@@ -111,12 +112,12 @@ public:
 
 	/** Failed assertion handler.  Warning: May be called at library startup time. */
 	template <typename FmtType, typename... Types>
-	static FORCEINLINE typename TEnableIf<TIsArrayOrRefOfType<FmtType, TCHAR>::Value, bool>::Type OptionallyLogFormattedEnsureMessageReturningFalse(bool bLog, const ANSICHAR* Expr, const ANSICHAR* File, int32 Line, void* ProgramCounter, const FmtType& FormattedMsg, Types... Args)
+	static FORCEINLINE bool OptionallyLogFormattedEnsureMessageReturningFalse(bool bLog, const ANSICHAR* Expr, const ANSICHAR* File, int32 Line, void* ProgramCounter, const FmtType& FormattedMsg, Types... Args)
 	{
-		static_assert(TIsArrayOrRefOfType<FmtType, TCHAR>::Value, "Formatting string must be a TCHAR array.");
+		static_assert(TIsArrayOrRefOfTypeByPredicate<FmtType, TIsCharEncodingCompatibleWithTCHAR>::Value, "Formatting string must be a TCHAR array.");
 		static_assert(TAnd<TIsValidVariadicFunctionArg<Types>...>::Value, "Invalid argument(s) passed to ensureMsgf");
 
-		return OptionallyLogFormattedEnsureMessageReturningFalseImpl(bLog, Expr, File, Line, ProgramCounter, FormattedMsg, Args...);
+		return OptionallyLogFormattedEnsureMessageReturningFalseImpl(bLog, Expr, File, Line, ProgramCounter, (const TCHAR*)FormattedMsg, Args...);
 	}
 
 #endif // DO_CHECK || DO_GUARD_SLOW
@@ -154,9 +155,9 @@ public:
 //		const FmtType& Format,
 //		Types... Args)
 //	{
-//		static_assert(TIsArrayOrRefOfType<FmtType, TCHAR>::Value, "Formatting string must be a TCHAR array.");
+//		static_assert(TIsArrayOrRefOfTypeByPredicate<FmtType, TIsCharEncodingCompatibleWithTCHAR>::Value, "Formatting string must be a TCHAR array.");
 //		static_assert(TAnd<TIsValidVariadicFunctionArg<Types>...>::Value, "Invalid argument(s) passed to CheckVerifyFailed()");
-//		return CheckVerifyFailedImpl(Expr, File, Line, ProgramCounter, Format, Args...);
+//		return CheckVerifyFailedImpl(Expr, File, Line, ProgramCounter, (const TCHAR*)Format, Args...);
 //	}
 //#endif
 
@@ -430,11 +431,11 @@ CORE_API void VARARGS LowLevelFatalErrorHandler(const ANSICHAR* File, int32 Line
 
 #define LowLevelFatalError(Format, ...) \
 	{ \
-		static_assert(TIsArrayOrRefOfType<decltype(Format), TCHAR>::Value, "Formatting string must be a TCHAR array."); \
+		static_assert(TIsArrayOrRefOfTypeByPredicate<decltype(Format), TIsCharEncodingCompatibleWithTCHAR>::Value, "Formatting string must be a TCHAR array."); \
 		DispatchCheckVerify([] (const auto& LFormat, const auto&... UE_LOG_Args) UE_DEBUG_SECTION \
 		{ \
 			void* ProgramCounter = PLATFORM_RETURN_ADDRESS(); \
-			LowLevelFatalErrorHandler(__FILE__, __LINE__, ProgramCounter, LFormat, UE_LOG_Args...); \
+			LowLevelFatalErrorHandler(__FILE__, __LINE__, ProgramCounter, (const TCHAR*)LFormat, UE_LOG_Args...); \
 			UE_DEBUG_BREAK_AND_PROMPT_FOR_REMOTE(); \
 			FDebug::ProcessFatalError(ProgramCounter); \
 		}, Format, ##__VA_ARGS__); \
diff --git a/Engine/Source/Runtime/Core/Public/Misc/CString.h b/Engine/Source/Runtime/Core/Public/Misc/CString.h
--- a/Engine/Source/Runtime/Core/Public/Misc/CString.h
+++ b/Engine/Source/Runtime/Core/Public/Misc/CString.h
@@ -9,7 +9,8 @@
 #include "HAL/PlatformString.h"
 #include "Templates/IsValidVariadicFunctionArg.h"
 #include "Templates/AndOrNot.h"
-#include "Templates/IsArrayOrRefOfType.h"
+#include "Templates/IsArrayOrRefOfTypeByPredicate.h"
+#include "Traits/IsCharEncodingCompatibleWith.h"
 
 #define MAX_SPRINTF 1024
 
@@ -369,18 +370,21 @@ private:
 	static int32 VARARGS SprintfImpl(CharType* Dest, const CharType* Fmt, ...);
 	static int32 VARARGS SnprintfImpl(CharType* Dest, int32 DestSize, const CharType* Fmt, ...);
 
+	template <typename SrcEncoding>
+	using TIsCharEncodingCompatibleWithCharType = TIsCharEncodingCompatibleWith<SrcEncoding, CharType>;
+
 public:
-	/** 
+	/**
 	* Standard string formatted print. 
 	* @warning: make sure code using FCString::Sprintf allocates enough (>= MAX_SPRINTF) memory for the destination buffer
 	*/
 	template <typename FmtType, typename... Types>
 	static int32 Sprintf(CharType* Dest, const FmtType& Fmt, Types... Args)
 	{
-		static_assert(TIsArrayOrRefOfType<FmtType, CharType>::Value, "Formatting string must be a literal string of the same character type as template.");
+		static_assert(TIsArrayOrRefOfTypeByPredicate<FmtType, TIsCharEncodingCompatibleWithCharType>::Value, "Formatting string must be a literal string of a char type compatible with the TCString type.");
 		static_assert(TAnd<TIsValidVariadicFunctionArg<Types>...>::Value, "Invalid argument(s) passed to TCString::Sprintf");
 
-		return SprintfImpl(Dest, Fmt, Args...);
+		return SprintfImpl(Dest, (const CharType*)Fmt, Args...);
 	}
 
 	/** 
@@ -389,10 +393,10 @@ public:
 	template <typename FmtType, typename... Types>
 	static int32 Snprintf(CharType* Dest, int32 DestSize, const FmtType& Fmt, Types... Args)
 	{
-		static_assert(TIsArrayOrRefOfType<FmtType, CharType>::Value, "Formatting string must be a literal string of the same character type as template.");
+		static_assert(TIsArrayOrRefOfTypeByPredicate<FmtType, TIsCharEncodingCompatibleWithCharType>::Value, "Formatting string must be a literal string of a char type compatible with the TCString type.");
 		static_assert(TAnd<TIsValidVariadicFunctionArg<Types>...>::Value, "Invalid argument(s) passed to TCString::Snprintf");
 
-		return SnprintfImpl(Dest, DestSize, Fmt, Args...);
+		return SnprintfImpl(Dest, DestSize, (const CharType*)Fmt, Args...);
 	}
 
 	/**
diff --git a/Engine/Source/Runtime/Core/Public/Misc/OutputDevice.h b/Engine/Source/Runtime/Core/Public/Misc/OutputDevice.h
--- a/Engine/Source/Runtime/Core/Public/Misc/OutputDevice.h
+++ b/Engine/Source/Runtime/Core/Public/Misc/OutputDevice.h
@@ -9,7 +9,8 @@
 #include "Misc/VarArgs.h"
 #include "Templates/IsValidVariadicFunctionArg.h"
 #include "Templates/AndOrNot.h"
-#include "Templates/IsArrayOrRefOfType.h"
+#include "Templates/IsArrayOrRefOfTypeByPredicate.h"
+#include "Traits/IsCharEncodingCompatibleWith.h"
 
 #ifndef USE_DEBUG_LOGGING
 #define USE_DEBUG_LOGGING 1
@@ -226,35 +227,35 @@ public:
 	template <typename FmtType>
 	void Logf(const FmtType& Fmt)
 	{
-		static_assert(TIsArrayOrRefOfType<FmtType, TCHAR>::Value, "Formatting string must be a TCHAR array.");
-		return Log(Fmt);
+		static_assert(TIsArrayOrRefOfTypeByPredicate<FmtType, TIsCharEncodingCompatibleWithTCHAR>::Value, "Formatting string must be a TCHAR array.");
+		return Log((const TCHAR*)Fmt);
 	}
 
 	template <typename FmtType, typename... Types>
 	FORCEINLINE void Logf(const FmtType& Fmt, Types... Args)
 	{
-		static_assert(TIsArrayOrRefOfType<FmtType, TCHAR>::Value, "Formatting string must be a TCHAR array.");
+		static_assert(TIsArrayOrRefOfTypeByPredicate<FmtType, TIsCharEncodingCompatibleWithTCHAR>::Value, "Formatting string must be a TCHAR array.");
 		static_assert(TAnd<TIsValidVariadicFunctionArg<Types>...>::Value, "Invalid argument(s) passed to FOutputDevice::Logf");
 
-		LogfImpl(Fmt, Args...);
+		LogfImpl((const TCHAR*)Fmt, Args...);
 	}
 
 	template <typename FmtType, typename... Types>
 	FORCEINLINE void Logf(ELogVerbosity::Type Verbosity, const FmtType& Fmt, Types... Args)
 	{
-		static_assert(TIsArrayOrRefOfType<FmtType, TCHAR>::Value, "Formatting string must be a TCHAR array.");
+		static_assert(TIsArrayOrRefOfTypeByPredicate<FmtType, TIsCharEncodingCompatibleWithTCHAR>::Value, "Formatting string must be a TCHAR array.");
 		static_assert(TAnd<TIsValidVariadicFunctionArg<Types>...>::Value, "Invalid argument(s) passed to FOutputDevice::Logf");
 
-		LogfImpl(Verbosity, Fmt, Args...);
+		LogfImpl(Verbosity, (const TCHAR*)Fmt, Args...);
 	}
 
 	template <typename FmtType, typename... Types>
 	FORCEINLINE void CategorizedLogf(const FName& Category, ELogVerbosity::Type Verbosity, const FmtType& Fmt, Types... Args)
 	{
-		static_assert(TIsArrayOrRefOfType<FmtType, TCHAR>::Value, "Formatting string must be a TCHAR array.");
+		static_assert(TIsArrayOrRefOfTypeByPredicate<FmtType, TIsCharEncodingCompatibleWithTCHAR>::Value, "Formatting string must be a TCHAR array.");
 		static_assert(TAnd<TIsValidVariadicFunctionArg<Types>...>::Value, "Invalid argument(s) passed to FOutputDevice::CategorizedLogf");
 
-		CategorizedLogfImpl(Category, Verbosity, Fmt, Args...);
+		CategorizedLogfImpl(Category, Verbosity, (const TCHAR*)Fmt, Args...);
 	}
 
 protected:
diff --git a/Engine/Source/Runtime/Core/Public/Misc/StringBuilder.h b/Engine/Source/Runtime/Core/Public/Misc/StringBuilder.h
--- a/Engine/Source/Runtime/Core/Public/Misc/StringBuilder.h
+++ b/Engine/Source/Runtime/Core/Public/Misc/StringBuilder.h
@@ -8,10 +8,11 @@
 #include "Misc/CString.h"
 #include "Templates/AndOrNot.h"
 #include "Templates/EnableIf.h"
-#include "Templates/IsArrayOrRefOfType.h"
+#include "Templates/IsArrayOrRefOfTypeByPredicate.h"
 #include "Templates/IsValidVariadicFunctionArg.h"
 #include "Templates/UnrealTemplate.h"
 #include "Traits/IsCharType.h"
+#include "Traits/IsCharEncodingCompatibleWith.h"
 #include "Traits/IsCharEncodingSimplyConvertibleTo.h"
 #include "Traits/IsContiguousContainer.h"
 #include <type_traits>
@@ -310,18 +311,22 @@ public:
 		return *this;
 	}
 
+private:
+	template <typename SrcEncoding>
+	using TIsCharEncodingCompatibleWithCharType = TIsCharEncodingCompatibleWith<SrcEncoding, CharType>;
+
+public:
 	/**
 	 * Appends to the string builder similarly to how classic sprintf works.
 	 *
 	 * @param Format A format string that specifies how to format the additional arguments. Refer to standard printf format.
 	 */
 	template <typename FmtType, typename... Types,
-		std::enable_if_t<TIsArrayOrRefOfType<FmtType, CharType>::Value>* = nullptr>
+		std::enable_if_t<TIsArrayOrRefOfTypeByPredicate<FmtType, TIsCharEncodingCompatibleWithCharType>::Value>* = nullptr>
 	BuilderType& Appendf(const FmtType& Fmt, Types... Args)
 	{
-		static_assert(TIsArrayOrRefOfType<FmtType, CharType>::Value, "Formatting string must be a character array.");
 		static_assert(TAnd<TIsValidVariadicFunctionArg<Types>...>::Value, "Invalid argument(s) passed to Appendf.");
-		return AppendfImpl(*this, Fmt, Forward<Types>(Args)...);
+		return AppendfImpl(*this, (const CharType*)Fmt, Forward<Types>(Args)...);
 	}
 
 private:
diff --git a/Engine/Source/Runtime/Core/Public/ProfilingDebugging/CsvProfiler.h b/Engine/Source/Runtime/Core/Public/ProfilingDebugging/CsvProfiler.h
--- a/Engine/Source/Runtime/Core/Public/ProfilingDebugging/CsvProfiler.h
+++ b/Engine/Source/Runtime/Core/Public/ProfilingDebugging/CsvProfiler.h
@@ -17,6 +17,8 @@
 #include "ProfilingDebugging/MiscTrace.h"
 #include "ProfilingDebugging/CsvProfilerConfig.h"
 #include "ProfilingDebugging/CsvProfilerTrace.h"
+#include "Templates/IsArrayOrRefOfTypeByPredicate.h"
+#include "Traits/IsCharEncodingCompatibleWith.h"
 
 #include <atomic>
 
@@ -232,9 +234,9 @@ public:
 	template <typename FmtType, typename... Types>
 	FORCEINLINE static void RecordEventf(int32 CategoryIndex, const FmtType& Fmt, Types... Args)
 	{
-		static_assert(TIsArrayOrRefOfType<FmtType, TCHAR>::Value, "Formatting string must be a TCHAR array.");
+		static_assert(TIsArrayOrRefOfTypeByPredicate<FmtType, TIsCharEncodingCompatibleWithTCHAR>::Value, "Formatting string must be a TCHAR array.");
 		static_assert(TAnd<TIsValidVariadicFunctionArg<Types>...>::Value, "Invalid argument(s) passed to FCsvProfiler::RecordEventf");
-		RecordEventfInternal(CategoryIndex, Fmt, Args...);
+		RecordEventfInternal(CategoryIndex, (const TCHAR*)Fmt, Args...);
 	}
 
 	CORE_API static void BeginSetWaitStat(const char * StatName);
diff --git a/Engine/Source/Runtime/Core/Public/Serialization/Archive.h b/Engine/Source/Runtime/Core/Public/Serialization/Archive.h
--- a/Engine/Source/Runtime/Core/Public/Serialization/Archive.h
+++ b/Engine/Source/Runtime/Core/Public/Serialization/Archive.h
@@ -14,10 +14,11 @@
 #include "Templates/AndOrNot.h"
 #include "Templates/EnableIf.h"
 #include "Templates/Function.h"
-#include "Templates/IsArrayOrRefOfType.h"
+#include "Templates/IsArrayOrRefOfTypeByPredicate.h"
 #include "Templates/IsEnumClass.h"
 #include "Templates/IsSigned.h"
 #include "Templates/IsValidVariadicFunctionArg.h"
+#include "Traits/IsCharEncodingCompatibleWith.h"
 #include "UObject/ObjectVersion.h"
 
 class FArchive;
@@ -1820,7 +1821,7 @@ public:
 	template <typename FmtType, typename... Types>
 	void Logf(const FmtType& Fmt, Types... Args)
 	{
-		static_assert(TIsArrayOrRefOfType<FmtType, TCHAR>::Value, "Formatting string must be a TCHAR array.");
+		static_assert(TIsArrayOrRefOfTypeByPredicate<FmtType, TIsCharEncodingCompatibleWithTCHAR>::Value, "Formatting string must be a TCHAR array.");
 		static_assert(TAnd<TIsValidVariadicFunctionArg<Types>...>::Value, "Invalid argument(s) passed to FArchive::Logf");
 
 		LogfImpl(Fmt, Args...);
diff --git a/Engine/Source/Runtime/Core/Public/Traits/IsCharEncodingCompatibleWith.h b/Engine/Source/Runtime/Core/Public/Traits/IsCharEncodingCompatibleWith.h
--- a/Engine/Source/Runtime/Core/Public/Traits/IsCharEncodingCompatibleWith.h
+++ b/Engine/Source/Runtime/Core/Public/Traits/IsCharEncodingCompatibleWith.h
@@ -67,3 +67,6 @@ struct TIsCharEncodingCompatibleWith
 
 template <typename SrcEncoding, typename DestEncoding>
 constexpr inline bool TIsCharEncodingCompatibleWith_V = TIsCharEncodingCompatibleWith<SrcEncoding, DestEncoding>::Value;
+
+template <typename SrcEncoding>
+using TIsCharEncodingCompatibleWithTCHAR = TIsCharEncodingCompatibleWith<SrcEncoding, TCHAR>;
diff --git a/Engine/Source/Runtime/CoreUObject/Public/UObject/UObjectGlobals.h b/Engine/Source/Runtime/CoreUObject/Public/UObject/UObjectGlobals.h
--- a/Engine/Source/Runtime/CoreUObject/Public/UObject/UObjectGlobals.h
+++ b/Engine/Source/Runtime/CoreUObject/Public/UObject/UObjectGlobals.h
@@ -16,7 +16,8 @@
 #include "Containers/ArrayView.h"
 #include "Containers/StringView.h"
 #include "Templates/Function.h"
-#include "Templates/IsArrayOrRefOfType.h"
+#include "Templates/IsArrayOrRefOfTypeByPredicate.h"
+#include "Traits/IsCharEncodingCompatibleWith.h"
 #include "Serialization/ArchiveUObject.h"
 #include "UObject/TopLevelAssetPath.h"
 
@@ -2867,18 +2868,18 @@ struct FAssetMsg
 	 */
 	#define UE_ASSET_LOG(CategoryName, Verbosity, Asset, Format, ...) \
 	{ \
-		static_assert(TIsArrayOrRefOfType<decltype(Format), TCHAR>::Value, "Formatting string must be a TCHAR array."); \
+		static_assert(TIsArrayOrRefOfTypeByPredicate<decltype(Format), TIsCharEncodingCompatibleWithTCHAR>::Value, "Formatting string must be a TCHAR array."); \
 		static_assert((ELogVerbosity::Verbosity & ELogVerbosity::VerbosityMask) < ELogVerbosity::NumVerbosity && ELogVerbosity::Verbosity > 0, "Verbosity must be constant and in range."); \
 		CA_CONSTANT_IF((ELogVerbosity::Verbosity & ELogVerbosity::VerbosityMask) <= ELogVerbosity::COMPILED_IN_MINIMUM_VERBOSITY && (ELogVerbosity::Warning & ELogVerbosity::VerbosityMask) <= FLogCategory##CategoryName::CompileTimeVerbosity) \
 		{ \
 			UE_LOG_EXPAND_IS_FATAL(Verbosity, PREPROCESSOR_NOTHING, if (!CategoryName.IsSuppressed(ELogVerbosity::Verbosity))) \
 			{ \
 				FString FormatPath = FAssetMsg::FormatPathForAssetLog(Asset);\
-				FMsg::Logf_Internal(__FILE__, __LINE__, CategoryName.GetCategoryName(), ELogVerbosity::Verbosity, ASSET_LOG_FORMAT_STRING Format, *FormatPath, ##__VA_ARGS__); \
+				FMsg::Logf_Internal(__FILE__, __LINE__, CategoryName.GetCategoryName(), ELogVerbosity::Verbosity, TEXT(ASSET_LOG_FORMAT_STRING_ANSI "%s"), *FormatPath, *FString::Printf(Format, ##__VA_ARGS__)); \
 				UE_LOG_EXPAND_IS_FATAL(Verbosity, \
 					{ \
 						UE_DEBUG_BREAK_AND_PROMPT_FOR_REMOTE(); \
-						FDebug::AssertFailed("", __FILE__, __LINE__, TEXT("%s: ") Format, *FormatPath, ##__VA_ARGS__); \
+						FDebug::AssertFailed("", __FILE__, __LINE__, TEXT("%s: %s"), *FormatPath, *FString::Printf(Format, ##__VA_ARGS__)); \
 						CA_ASSUME(false); \
 					}, \
 					PREPROCESSOR_NOTHING \
diff --git a/Engine/Source/Runtime/Engine/Public/VisualLogger/VisualLogger.h b/Engine/Source/Runtime/Engine/Public/VisualLogger/VisualLogger.h
--- a/Engine/Source/Runtime/Engine/Public/VisualLogger/VisualLogger.h
+++ b/Engine/Source/Runtime/Engine/Public/VisualLogger/VisualLogger.h
@@ -10,7 +10,8 @@
 #include "EngineStats.h"
 #include "Templates/IsValidVariadicFunctionArg.h"
 #include "Templates/AndOrNot.h"
-#include "Templates/IsArrayOrRefOfType.h"
+#include "Templates/IsArrayOrRefOfTypeByPredicate.h"
+#include "Traits/IsCharEncodingCompatibleWith.h"
 #include "UObject/ObjectKey.h"
 #include "Containers/Ticker.h"
 
@@ -191,53 +192,53 @@ public:
 	template <typename FmtType, typename... Types>
 	static void CategorizedLogf(const UObject* LogOwner, const FLogCategoryBase& Category, ELogVerbosity::Type Verbosity, const FmtType& Fmt, Types... Args)
 	{
-		static_assert(TIsArrayOrRefOfType<FmtType, TCHAR>::Value, "Formatting string must be a TCHAR array.");
+		static_assert(TIsArrayOrRefOfTypeByPredicate<FmtType, TIsCharEncodingCompatibleWithTCHAR>::Value, "Formatting string must be a TCHAR array.");
 		static_assert(TAnd<TIsValidVariadicFunctionArg<Types>...>::Value, "Invalid argument(s) passed to FVisualLogger::CategorizedLogf");
 
-		CategorizedLogfImpl(LogOwner, Category, Verbosity, Fmt, Args...);
+		CategorizedLogfImpl(LogOwner, Category, Verbosity, (const TCHAR*)Fmt, Args...);
 	}
 	template <typename FmtType, typename... Types>
 	static void CategorizedLogf(const UObject* LogOwner, const FName& CategoryName, ELogVerbosity::Type Verbosity, const FmtType& Fmt, Types... Args)
 	{
-		static_assert(TIsArrayOrRefOfType<FmtType, TCHAR>::Value, "Formatting string must be a TCHAR array.");
+		static_assert(TIsArrayOrRefOfTypeByPredicate<FmtType, TIsCharEncodingCompatibleWithTCHAR>::Value, "Formatting string must be a TCHAR array.");
 		static_assert(TAnd<TIsValidVariadicFunctionArg<Types>...>::Value, "Invalid argument(s) passed to FVisualLogger::CategorizedLogf");
 
-		CategorizedLogfImpl(LogOwner, CategoryName, Verbosity, Fmt, Args...);
+		CategorizedLogfImpl(LogOwner, CategoryName, Verbosity, (const TCHAR*)Fmt, Args...);
 	}
 
 	// Segment log
 	template <typename FmtType, typename... Types>
 	static void GeometryShapeLogf(const UObject* LogOwner, const FLogCategoryBase& Category, ELogVerbosity::Type Verbosity, const FVector& Start, const FVector& End, const FColor& Color, const uint16 Thickness, const FmtType& Fmt, Types... Args)
 	{
-		static_assert(TIsArrayOrRefOfType<FmtType, TCHAR>::Value, "Formatting string must be a TCHAR array.");
+		static_assert(TIsArrayOrRefOfTypeByPredicate<FmtType, TIsCharEncodingCompatibleWithTCHAR>::Value, "Formatting string must be a TCHAR array.");
 		static_assert(TAnd<TIsValidVariadicFunctionArg<Types>...>::Value, "Invalid argument(s) passed to FVisualLogger::GeometryShapeLogf");
 
-		GeometryShapeLogfImpl(LogOwner, Category, Verbosity, Start, End, Color, Thickness, Fmt, Args...);
+		GeometryShapeLogfImpl(LogOwner, Category, Verbosity, Start, End, Color, Thickness, (const TCHAR*)Fmt, Args...);
 	}
 	template <typename FmtType, typename... Types>
 	static void GeometryShapeLogf(const UObject* LogOwner, const FName& CategoryName, ELogVerbosity::Type Verbosity, const FVector& Start, const FVector& End, const FColor& Color, const uint16 Thickness, const FmtType& Fmt, Types... Args)
 	{
-		static_assert(TIsArrayOrRefOfType<FmtType, TCHAR>::Value, "Formatting string must be a TCHAR array.");
+		static_assert(TIsArrayOrRefOfTypeByPredicate<FmtType, TIsCharEncodingCompatibleWithTCHAR>::Value, "Formatting string must be a TCHAR array.");
 		static_assert(TAnd<TIsValidVariadicFunctionArg<Types>...>::Value, "Invalid argument(s) passed to FVisualLogger::GeometryShapeLogf");
 
-		GeometryShapeLogfImpl(LogOwner, CategoryName, Verbosity, Start, End, Color, Thickness, Fmt, Args...);
+		GeometryShapeLogfImpl(LogOwner, CategoryName, Verbosity, Start, End, Color, Thickness, (const TCHAR*)Fmt, Args...);
 	}
 
 	// Arrow
 	template <typename FmtType, typename... Types>
 	static void ArrowLogf(const UObject* LogOwner, const FLogCategoryBase& Category, ELogVerbosity::Type Verbosity, const FVector& Start, const FVector& End, const FColor& Color, const FmtType& Fmt, Types... Args)
 	{
-		static_assert(TIsArrayOrRefOfType<FmtType, TCHAR>::Value, "Formatting string must be a TCHAR array.");
+		static_assert(TIsArrayOrRefOfTypeByPredicate<FmtType, TIsCharEncodingCompatibleWithTCHAR>::Value, "Formatting string must be a TCHAR array.");
 		static_assert(TAnd<TIsValidVariadicFunctionArg<Types>...>::Value, "Invalid argument(s) passed to FVisualLogger::ArrowLogf");
-		ArrowLogfImpl(LogOwner, Category, Verbosity, Start, End, Color, Fmt, Args...);
+		ArrowLogfImpl(LogOwner, Category, Verbosity, Start, End, Color, (const TCHAR*)Fmt, Args...);
 	}
 	template <typename FmtType, typename... Types>
 	static void ArrowLogf(const UObject* LogOwner, const FName& CategoryName, ELogVerbosity::Type Verbosity, const FVector& Start, const FVector& End, const FColor& Color, const FmtType& Fmt, Types... Args)
 	{
-		static_assert(TIsArrayOrRefOfType<FmtType, TCHAR>::Value, "Formatting string must be a TCHAR array.");
+		static_assert(TIsArrayOrRefOfTypeByPredicate<FmtType, TIsCharEncodingCompatibleWithTCHAR>::Value, "Formatting string must be a TCHAR array.");
 		static_assert(TAnd<TIsValidVariadicFunctionArg<Types>...>::Value, "Invalid argument(s) passed to FVisualLogger::ArrowLogf");
 
-		ArrowLogfImpl(LogOwner, CategoryName, Verbosity, Start, End, Color, Fmt, Args...);
+		ArrowLogfImpl(LogOwner, CategoryName, Verbosity, Start, End, Color, (const TCHAR*)Fmt, Args...);
 	}
 
 
@@ -245,180 +246,180 @@ public:
 	template <typename FmtType, typename... Types>
 	static void CircleLogf(const UObject* LogOwner, const FLogCategoryBase& Category, ELogVerbosity::Type Verbosity, const FVector& Center, const FVector& UpAxis, const float Radius, const FColor& Color, const uint16 Thickness, const FmtType& Fmt, Types... Args)
 	{
-		static_assert(TIsArrayOrRefOfType<FmtType, TCHAR>::Value, "Formatting string must be a TCHAR array.");
+		static_assert(TIsArrayOrRefOfTypeByPredicate<FmtType, TIsCharEncodingCompatibleWithTCHAR>::Value, "Formatting string must be a TCHAR array.");
 		static_assert(TAnd<TIsValidVariadicFunctionArg<Types>...>::Value, "Invalid argument(s) passed to FVisualLogger::CircleLogf");
 
-		CircleLogfImpl(LogOwner, Category, Verbosity, Center, UpAxis, Radius, Color, Thickness, Fmt, Args...);
+		CircleLogfImpl(LogOwner, Category, Verbosity, Center, UpAxis, Radius, Color, Thickness, (const TCHAR*)Fmt, Args...);
 	}
 	template <typename FmtType, typename... Types>
 	static void CircleLogf(const UObject* LogOwner, const FName& CategoryName, ELogVerbosity::Type Verbosity, const FVector& Center, const FVector& UpAxis, const float Radius, const FColor& Color, const uint16 Thickness, const FmtType& Fmt, Types... Args)
 	{
-		static_assert(TIsArrayOrRefOfType<FmtType, TCHAR>::Value, "Formatting string must be a TCHAR array.");
+		static_assert(TIsArrayOrRefOfTypeByPredicate<FmtType, TIsCharEncodingCompatibleWithTCHAR>::Value, "Formatting string must be a TCHAR array.");
 		static_assert(TAnd<TIsValidVariadicFunctionArg<Types>...>::Value, "Invalid argument(s) passed to FVisualLogger::CircleLogf");
 
-		CircleLogfImpl(LogOwner, CategoryName, Verbosity, Center, UpAxis, Radius, Color, Thickness, Fmt, Args...);
+		CircleLogfImpl(LogOwner, CategoryName, Verbosity, Center, UpAxis, Radius, Color, Thickness, (const TCHAR*)Fmt, Args...);
 	}
 
 	// Location/Sphere log
 	template <typename FmtType, typename... Types>
 	static void GeometryShapeLogf(const UObject* LogOwner, const FLogCategoryBase& Category, ELogVerbosity::Type Verbosity, const FVector& Location, float Radius, const FColor& Color, const FmtType& Fmt, Types... Args)
 	{
-		static_assert(TIsArrayOrRefOfType<FmtType, TCHAR>::Value, "Formatting string must be a TCHAR array.");
+		static_assert(TIsArrayOrRefOfTypeByPredicate<FmtType, TIsCharEncodingCompatibleWithTCHAR>::Value, "Formatting string must be a TCHAR array.");
 		static_assert(TAnd<TIsValidVariadicFunctionArg<Types>...>::Value, "Invalid argument(s) passed to FVisualLogger::GeometryShapeLogf");
 
-		GeometryShapeLogfImpl(LogOwner, Category, Verbosity, Location, Radius, Color, Fmt, Args...);
+		GeometryShapeLogfImpl(LogOwner, Category, Verbosity, Location, Radius, Color, (const TCHAR*)Fmt, Args...);
 	}
 	template <typename FmtType, typename... Types>
 	static void GeometryShapeLogf(const UObject* LogOwner, const FName& CategoryName, ELogVerbosity::Type Verbosity, const FVector& Location, float Radius, const FColor& Color, const FmtType& Fmt, Types... Args)
 	{
-		static_assert(TIsArrayOrRefOfType<FmtType, TCHAR>::Value, "Formatting string must be a TCHAR array.");
+		static_assert(TIsArrayOrRefOfTypeByPredicate<FmtType, TIsCharEncodingCompatibleWithTCHAR>::Value, "Formatting string must be a TCHAR array.");
 		static_assert(TAnd<TIsValidVariadicFunctionArg<Types>...>::Value, "Invalid argument(s) passed to FVisualLogger::GeometryShapeLogf");
 
-		GeometryShapeLogfImpl(LogOwner, CategoryName, Verbosity, Location, Radius, Color, Fmt, Args...);
+		GeometryShapeLogfImpl(LogOwner, CategoryName, Verbosity, Location, Radius, Color, (const TCHAR*)Fmt, Args...);
 	}
 
 	// Box log
 	template <typename FmtType, typename... Types>
 	static void GeometryBoxLogf(const UObject* LogOwner, const FLogCategoryBase& Category, ELogVerbosity::Type Verbosity, const FBox& Box, const FMatrix& Matrix, const FColor& Color, const FmtType& Fmt, Types... Args)
 	{
-		static_assert(TIsArrayOrRefOfType<FmtType, TCHAR>::Value, "Formatting string must be a TCHAR array.");
+		static_assert(TIsArrayOrRefOfTypeByPredicate<FmtType, TIsCharEncodingCompatibleWithTCHAR>::Value, "Formatting string must be a TCHAR array.");
 		static_assert(TAnd<TIsValidVariadicFunctionArg<Types>...>::Value, "Invalid argument(s) passed to FVisualLogger::GeometryBoxLogf");
 
-		GeometryBoxLogfImpl(LogOwner, Category, Verbosity, Box, Matrix, Color, Fmt, Args...);
+		GeometryBoxLogfImpl(LogOwner, Category, Verbosity, Box, Matrix, Color, (const TCHAR*)Fmt, Args...);
 	}
 	template <typename FmtType, typename... Types>
 	static void GeometryBoxLogf(const UObject* LogOwner, const FName& CategoryName, ELogVerbosity::Type Verbosity, const FBox& Box, const FMatrix& Matrix, const FColor& Color, const FmtType& Fmt, Types... Args)
 	{
-		static_assert(TIsArrayOrRefOfType<FmtType, TCHAR>::Value, "Formatting string must be a TCHAR array.");
+		static_assert(TIsArrayOrRefOfTypeByPredicate<FmtType, TIsCharEncodingCompatibleWithTCHAR>::Value, "Formatting string must be a TCHAR array.");
 		static_assert(TAnd<TIsValidVariadicFunctionArg<Types>...>::Value, "Invalid argument(s) passed to FVisualLogger::GeometryBoxLogf");
 
-		GeometryBoxLogfImpl(LogOwner, CategoryName, Verbosity, Box, Matrix, Color, Fmt, Args...);
+		GeometryBoxLogfImpl(LogOwner, CategoryName, Verbosity, Box, Matrix, Color, (const TCHAR*)Fmt, Args...);
 	}
 
 	// Cone log
 	template <typename FmtType, typename... Types>
 	static void GeometryShapeLogf(const UObject* LogOwner, const FLogCategoryBase& Category, ELogVerbosity::Type Verbosity, const FVector& Origin, const FVector& Direction, const float Length, const float Angle, const FColor& Color, const FmtType& Fmt, Types... Args)
 	{
-		static_assert(TIsArrayOrRefOfType<FmtType, TCHAR>::Value, "Formatting string must be a TCHAR array.");
+		static_assert(TIsArrayOrRefOfTypeByPredicate<FmtType, TIsCharEncodingCompatibleWithTCHAR>::Value, "Formatting string must be a TCHAR array.");
 		static_assert(TAnd<TIsValidVariadicFunctionArg<Types>...>::Value, "Invalid argument(s) passed to FVisualLogger::GeometryShapeLogf");
 
-		GeometryShapeLogfImpl(LogOwner, Category, Verbosity, Origin, Direction, Length, Angle, Color, Fmt, Args...);
+		GeometryShapeLogfImpl(LogOwner, Category, Verbosity, Origin, Direction, Length, Angle, Color, (const TCHAR*)Fmt, Args...);
 	}
 	template <typename FmtType, typename... Types>
 	static void GeometryShapeLogf(const UObject* LogOwner, const FName& CategoryName, ELogVerbosity::Type Verbosity, const FVector& Origin, const FVector& Direction, const float Length, const float Angle, const FColor& Color, const FmtType& Fmt, Types... Args)
 	{
-		static_assert(TIsArrayOrRefOfType<FmtType, TCHAR>::Value, "Formatting string must be a TCHAR array.");
+		static_assert(TIsArrayOrRefOfTypeByPredicate<FmtType, TIsCharEncodingCompatibleWithTCHAR>::Value, "Formatting string must be a TCHAR array.");
 		static_assert(TAnd<TIsValidVariadicFunctionArg<Types>...>::Value, "Invalid argument(s) passed to FVisualLogger::GeometryShapeLogf");
 
-		GeometryShapeLogfImpl(LogOwner, CategoryName, Verbosity, Origin, Direction, Length, Angle, Color, Fmt, Args...);
+		GeometryShapeLogfImpl(LogOwner, CategoryName, Verbosity, Origin, Direction, Length, Angle, Color, (const TCHAR*)Fmt, Args...);
 	}
 
 	// Cylinder log
 	template <typename FmtType, typename... Types>
 	static void GeometryShapeLogf(const UObject* LogOwner, const FLogCategoryBase& Category, ELogVerbosity::Type Verbosity, const FVector& Start, const FVector& End, const float Radius, const FColor& Color, const FmtType& Fmt, Types... Args)
 	{
-		static_assert(TIsArrayOrRefOfType<FmtType, TCHAR>::Value, "Formatting string must be a TCHAR array.");
+		static_assert(TIsArrayOrRefOfTypeByPredicate<FmtType, TIsCharEncodingCompatibleWithTCHAR>::Value, "Formatting string must be a TCHAR array.");
 		static_assert(TAnd<TIsValidVariadicFunctionArg<Types>...>::Value, "Invalid argument(s) passed to FVisualLogger::GeometryShapeLogf");
 
-		GeometryShapeLogfImpl(LogOwner, Category, Verbosity, Start, End, Radius, Color, Fmt, Args...);
+		GeometryShapeLogfImpl(LogOwner, Category, Verbosity, Start, End, Radius, Color, (const TCHAR*)Fmt, Args...);
 	}
 	template <typename FmtType, typename... Types>
 	static void GeometryShapeLogf(const UObject* LogOwner, const FName& CategoryName, ELogVerbosity::Type Verbosity, const FVector& Start, const FVector& End, const float Radius, const FColor& Color, const FmtType& Fmt, Types... Args)
 	{
-		static_assert(TIsArrayOrRefOfType<FmtType, TCHAR>::Value, "Formatting string must be a TCHAR array.");
+		static_assert(TIsArrayOrRefOfTypeByPredicate<FmtType, TIsCharEncodingCompatibleWithTCHAR>::Value, "Formatting string must be a TCHAR array.");
 		static_assert(TAnd<TIsValidVariadicFunctionArg<Types>...>::Value, "Invalid argument(s) passed to FVisualLogger::GeometryShapeLogf");
 
-		GeometryShapeLogfImpl(LogOwner, CategoryName, Verbosity, Start, End, Radius, Color, Fmt, Args...);
+		GeometryShapeLogfImpl(LogOwner, CategoryName, Verbosity, Start, End, Radius, Color, (const TCHAR*)Fmt, Args...);
 	}
 
 	// Capsule log
 	template <typename FmtType, typename... Types>
 	static void GeometryShapeLogf(const UObject* LogOwner, const FLogCategoryBase& Category, ELogVerbosity::Type Verbosity, const FVector& Base, float HalfHeight, float Radius, const FQuat & Rotation, const FColor& Color, const FmtType& Fmt, Types... Args)
 	{
-		static_assert(TIsArrayOrRefOfType<FmtType, TCHAR>::Value, "Formatting string must be a TCHAR array.");
+		static_assert(TIsArrayOrRefOfTypeByPredicate<FmtType, TIsCharEncodingCompatibleWithTCHAR>::Value, "Formatting string must be a TCHAR array.");
 		static_assert(TAnd<TIsValidVariadicFunctionArg<Types>...>::Value, "Invalid argument(s) passed to FVisualLogger::GeometryShapeLogf");
 
-		GeometryShapeLogfImpl(LogOwner, Category, Verbosity, Base, HalfHeight, Radius, Rotation, Color, Fmt, Args...);
+		GeometryShapeLogfImpl(LogOwner, Category, Verbosity, Base, HalfHeight, Radius, Rotation, Color, (const TCHAR*)Fmt, Args...);
 	}
 	template <typename FmtType, typename... Types>
 	static void GeometryShapeLogf(const UObject* LogOwner, const FName& CategoryName, ELogVerbosity::Type Verbosity, const FVector& Base, float HalfHeight, float Radius, const FQuat & Rotation, const FColor& Color, const FmtType& Fmt, Types... Args)
 	{
-		static_assert(TIsArrayOrRefOfType<FmtType, TCHAR>::Value, "Formatting string must be a TCHAR array.");
+		static_assert(TIsArrayOrRefOfTypeByPredicate<FmtType, TIsCharEncodingCompatibleWithTCHAR>::Value, "Formatting string must be a TCHAR array.");
 		static_assert(TAnd<TIsValidVariadicFunctionArg<Types>...>::Value, "Invalid argument(s) passed to FVisualLogger::GeometryShapeLogf");
 
-		GeometryShapeLogfImpl(LogOwner, CategoryName, Verbosity, Base, HalfHeight, Radius, Rotation, Color, Fmt, Args...);
+		GeometryShapeLogfImpl(LogOwner, CategoryName, Verbosity, Base, HalfHeight, Radius, Rotation, Color, (const TCHAR*)Fmt, Args...);
 	}
 
 	// NavArea/Extruded convex log
 	template <typename FmtType, typename... Types>
 	static void NavAreaShapeLogf(const UObject* LogOwner, const FLogCategoryBase& Category, ELogVerbosity::Type Verbosity, const TArray<FVector>& ConvexPoints, float MinZ, float MaxZ, const FColor& Color, const FmtType& Fmt, Types... Args)
 	{
-		static_assert(TIsArrayOrRefOfType<FmtType, TCHAR>::Value, "Formatting string must be a TCHAR array.");
+		static_assert(TIsArrayOrRefOfTypeByPredicate<FmtType, TIsCharEncodingCompatibleWithTCHAR>::Value, "Formatting string must be a TCHAR array.");
 		static_assert(TAnd<TIsValidVariadicFunctionArg<Types>...>::Value, "Invalid argument(s) passed to FVisualLogger::NavAreaShapeLogf");
 
-		NavAreaShapeLogfImpl(LogOwner, Category, Verbosity, ConvexPoints, MinZ, MaxZ, Color, Fmt, Args...);
+		NavAreaShapeLogfImpl(LogOwner, Category, Verbosity, ConvexPoints, MinZ, MaxZ, Color, (const TCHAR*)Fmt, Args...);
 	}
 	template <typename FmtType, typename... Types>
 	static void NavAreaShapeLogf(const UObject* LogOwner, const FName& CategoryName, ELogVerbosity::Type Verbosity, const TArray<FVector>& ConvexPoints, float MinZ, float MaxZ, const FColor& Color, const FmtType& Fmt, Types... Args)
 	{
-		static_assert(TIsArrayOrRefOfType<FmtType, TCHAR>::Value, "Formatting string must be a TCHAR array.");
+		static_assert(TIsArrayOrRefOfTypeByPredicate<FmtType, TIsCharEncodingCompatibleWithTCHAR>::Value, "Formatting string must be a TCHAR array.");
 		static_assert(TAnd<TIsValidVariadicFunctionArg<Types>...>::Value, "Invalid argument(s) passed to FVisualLogger::NavAreaShapeLogf");
 
-		NavAreaShapeLogfImpl(LogOwner, CategoryName, Verbosity, ConvexPoints, MinZ, MaxZ, Color, Fmt, Args...);
+		NavAreaShapeLogfImpl(LogOwner, CategoryName, Verbosity, ConvexPoints, MinZ, MaxZ, Color, (const TCHAR*)Fmt, Args...);
 	}
 
 	// 3d Mesh log
 	template <typename FmtType, typename... Types>
 	static void GeometryShapeLogf(const UObject* LogOwner, const FLogCategoryBase& Category, ELogVerbosity::Type Verbosity, const TArray<FVector>& Vertices, const TArray<int32>& Indices, const FColor& Color, const FmtType& Fmt, Types... Args)
 	{
-		static_assert(TIsArrayOrRefOfType<FmtType, TCHAR>::Value, "Formatting string must be a TCHAR array.");
+		static_assert(TIsArrayOrRefOfTypeByPredicate<FmtType, TIsCharEncodingCompatibleWithTCHAR>::Value, "Formatting string must be a TCHAR array.");
 		static_assert(TAnd<TIsValidVariadicFunctionArg<Types>...>::Value, "Invalid argument(s) passed to FVisualLogger::GeometryShapeLogf");
 
-		GeometryShapeLogfImpl(LogOwner, Category, Verbosity, Vertices, Indices, Color, Fmt, Args...);
+		GeometryShapeLogfImpl(LogOwner, Category, Verbosity, Vertices, Indices, Color, (const TCHAR*)Fmt, Args...);
 	}
 	template <typename FmtType, typename... Types>
 	static void GeometryShapeLogf(const UObject* LogOwner, const FName& CategoryName, ELogVerbosity::Type Verbosity, const TArray<FVector>& Vertices, const TArray<int32>& Indices, const FColor& Color, const FmtType& Fmt, Types... Args)
 	{
-		static_assert(TIsArrayOrRefOfType<FmtType, TCHAR>::Value, "Formatting string must be a TCHAR array.");
+		static_assert(TIsArrayOrRefOfTypeByPredicate<FmtType, TIsCharEncodingCompatibleWithTCHAR>::Value, "Formatting string must be a TCHAR array.");
 		static_assert(TAnd<TIsValidVariadicFunctionArg<Types>...>::Value, "Invalid argument(s) passed to FVisualLogger::GeometryShapeLogf");
 
-		GeometryShapeLogfImpl(LogOwner, CategoryName, Verbosity, Vertices, Indices, Color, Fmt, Args...);
+		GeometryShapeLogfImpl(LogOwner, CategoryName, Verbosity, Vertices, Indices, Color, (const TCHAR*)Fmt, Args...);
 	}
 
 	// 2d Convex shape
 	template <typename FmtType, typename... Types>
 	static void GeometryConvexLogf(const UObject* LogOwner, const FLogCategoryBase& Category, ELogVerbosity::Type Verbosity, const TArray<FVector>& Points, const FColor& Color, const FmtType& Fmt, Types... Args)
 	{
-		static_assert(TIsArrayOrRefOfType<FmtType, TCHAR>::Value, "Formatting string must be a TCHAR array.");
+		static_assert(TIsArrayOrRefOfTypeByPredicate<FmtType, TIsCharEncodingCompatibleWithTCHAR>::Value, "Formatting string must be a TCHAR array.");
 		static_assert(TAnd<TIsValidVariadicFunctionArg<Types>...>::Value, "Invalid argument(s) passed to FVisualLogger::GeometryConvexLogf");
 
-		GeometryConvexLogfImpl(LogOwner, Category, Verbosity, Points, Color, Fmt, Args...);
+		GeometryConvexLogfImpl(LogOwner, Category, Verbosity, Points, Color, (const TCHAR*)Fmt, Args...);
 	}
 	template <typename FmtType, typename... Types>
 	static void GeometryConvexLogf(const UObject* LogOwner, const FName& CategoryName, ELogVerbosity::Type Verbosity, const TArray<FVector>& Points, const FColor& Color, const FmtType& Fmt, Types... Args)
 	{
-		static_assert(TIsArrayOrRefOfType<FmtType, TCHAR>::Value, "Formatting string must be a TCHAR array.");
+		static_assert(TIsArrayOrRefOfTypeByPredicate<FmtType, TIsCharEncodingCompatibleWithTCHAR>::Value, "Formatting string must be a TCHAR array.");
 		static_assert(TAnd<TIsValidVariadicFunctionArg<Types>...>::Value, "Invalid argument(s) passed to FVisualLogger::GeometryConvexLogf");
 
-		GeometryConvexLogfImpl(LogOwner, CategoryName, Verbosity, Points, Color, Fmt, Args...);
+		GeometryConvexLogfImpl(LogOwner, CategoryName, Verbosity, Points, Color, (const TCHAR*)Fmt, Args...);
 	}
 
 	//Histogram data
 	template <typename FmtType, typename... Types>
 	static void HistogramDataLogf(const UObject* LogOwner, const FLogCategoryBase& Category, ELogVerbosity::Type Verbosity, FName GraphName, FName DataName, const FVector2D& Data, const FColor& Color, const FmtType& Fmt, Types... Args)
 	{
-		static_assert(TIsArrayOrRefOfType<FmtType, TCHAR>::Value, "Formatting string must be a TCHAR array.");
+		static_assert(TIsArrayOrRefOfTypeByPredicate<FmtType, TIsCharEncodingCompatibleWithTCHAR>::Value, "Formatting string must be a TCHAR array.");
 		static_assert(TAnd<TIsValidVariadicFunctionArg<Types>...>::Value, "Invalid argument(s) passed to FVisualLogger::HistogramDataLogf");
 
-		HistogramDataLogfImpl(LogOwner, Category, Verbosity, GraphName, DataName, Data, Color, Fmt, Args...);
+		HistogramDataLogfImpl(LogOwner, Category, Verbosity, GraphName, DataName, Data, Color, (const TCHAR*)Fmt, Args...);
 	}
 	template <typename FmtType, typename... Types>
 	static void HistogramDataLogf(const UObject* LogOwner, const FName& CategoryName, ELogVerbosity::Type Verbosity, FName GraphName, FName DataName, const FVector2D& Data, const FColor& Color, const FmtType& Fmt, Types... Args)
 	{
-		static_assert(TIsArrayOrRefOfType<FmtType, TCHAR>::Value, "Formatting string must be a TCHAR array.");
+		static_assert(TIsArrayOrRefOfTypeByPredicate<FmtType, TIsCharEncodingCompatibleWithTCHAR>::Value, "Formatting string must be a TCHAR array.");
 		static_assert(TAnd<TIsValidVariadicFunctionArg<Types>...>::Value, "Invalid argument(s) passed to FVisualLogger::HistogramDataLogf");
 
-		HistogramDataLogfImpl(LogOwner, CategoryName, Verbosity, GraphName, DataName, Data, Color, Fmt, Args...);
+		HistogramDataLogfImpl(LogOwner, CategoryName, Verbosity, GraphName, DataName, Data, Color, (const TCHAR*)Fmt, Args...);
 	}
 
 	// Navigation data debug snapshot
diff --git a/Engine/Source/Runtime/LiveLinkInterface/Public/LiveLinkLog.h b/Engine/Source/Runtime/LiveLinkInterface/Public/LiveLinkLog.h
--- a/Engine/Source/Runtime/LiveLinkInterface/Public/LiveLinkLog.h
+++ b/Engine/Source/Runtime/LiveLinkInterface/Public/LiveLinkLog.h
@@ -5,6 +5,8 @@
 #include "CoreMinimal.h"
 #include "LiveLinkTypes.h"
 #include "Logging/TokenizedMessage.h"
+#include "Templates/IsArrayOrRefOfTypeByPredicate.h"
+#include "Traits/IsCharEncodingCompatibleWith.h"
 
 #define ENABLE_LIVELINK_LOGGING (1 && !NO_LOGGING && !(UE_BUILD_SHIPPING || UE_BUILD_TEST))
 
@@ -16,7 +18,7 @@ public:
 	template<typename FormatType, typename... ArgsType>
 	static void Error(const FormatType& Format, ArgsType... Args)
 	{
-		static_assert(TIsArrayOrRefOfType<FormatType, TCHAR>::Value, "Formatting string must be a TCHAR array.");
+		static_assert(TIsArrayOrRefOfTypeByPredicate<FormatType, TIsCharEncodingCompatibleWithTCHAR>::Value, "Formatting string must be a TCHAR array.");
 		Log_Internal(EMessageSeverity::Error, NAME_None, FLiveLinkSubjectKey(), Format, Args...);
 	}
 
@@ -27,7 +29,7 @@ public:
 	template<typename FormatType, typename... ArgsType>
 	static void ErrorOnce(FName MessageID, const FLiveLinkSubjectKey& SubjectKey, const FormatType& Format, ArgsType... Args)
 	{
-		static_assert(TIsArrayOrRefOfType<FormatType, TCHAR>::Value, "Formatting string must be a TCHAR array.");
+		static_assert(TIsArrayOrRefOfTypeByPredicate<FormatType, TIsCharEncodingCompatibleWithTCHAR>::Value, "Formatting string must be a TCHAR array.");
 		Log_Internal(EMessageSeverity::Error, MessageID, SubjectKey, Format, Args...);
 	}
 	
@@ -35,7 +37,7 @@ public:
 	template<typename FormatType, typename... ArgsType>
 	static void Warning(const FormatType& Format, ArgsType... Args)
 	{
-		static_assert(TIsArrayOrRefOfType<FormatType, TCHAR>::Value, "Formatting string must be a TCHAR array.");
+		static_assert(TIsArrayOrRefOfTypeByPredicate<FormatType, TIsCharEncodingCompatibleWithTCHAR>::Value, "Formatting string must be a TCHAR array.");
 		Log_Internal(EMessageSeverity::Warning, NAME_None, FLiveLinkSubjectKey(), Format, Args...);
 	}
 
@@ -46,7 +48,7 @@ public:
 	template<typename FormatType, typename... ArgsType>
 	static void WarningOnce(FName MessageID, const FLiveLinkSubjectKey& SubjectKey, const FormatType& Format, ArgsType... Args)
 	{
-		static_assert(TIsArrayOrRefOfType<FormatType, TCHAR>::Value, "Formatting string must be a TCHAR array.");
+		static_assert(TIsArrayOrRefOfTypeByPredicate<FormatType, TIsCharEncodingCompatibleWithTCHAR>::Value, "Formatting string must be a TCHAR array.");
 		Log_Internal(EMessageSeverity::Warning, MessageID, SubjectKey, Format, Args...);
 	}
 	
@@ -54,7 +56,7 @@ public:
 	template<typename FormatType, typename... ArgsType>
 	static void Info(const FormatType& Format, ArgsType... Args)
 	{
-		static_assert(TIsArrayOrRefOfType<FormatType, TCHAR>::Value, "Formatting string must be a TCHAR array.");
+		static_assert(TIsArrayOrRefOfTypeByPredicate<FormatType, TIsCharEncodingCompatibleWithTCHAR>::Value, "Formatting string must be a TCHAR array.");
 		Log_Internal(EMessageSeverity::Info, NAME_None, FLiveLinkSubjectKey(), Format, Args...);
 	}
 
@@ -65,7 +67,7 @@ public:
 	template<typename FormatType, typename... ArgsType>
 	static void InfoOnce(FName MessageID, const FLiveLinkSubjectKey& SubjectKey, const FormatType& Format, ArgsType... Args)
 	{
-		static_assert(TIsArrayOrRefOfType<FormatType, TCHAR>::Value, "Formatting string must be a TCHAR array.");
+		static_assert(TIsArrayOrRefOfTypeByPredicate<FormatType, TIsCharEncodingCompatibleWithTCHAR>::Value, "Formatting string must be a TCHAR array.");
 		Log_Internal(EMessageSeverity::Info, MessageID, SubjectKey, Format, Args...);
 	}
 
@@ -76,7 +78,7 @@ public:
 	template<typename FormatType, typename... ArgsType>
 	static TSharedPtr<FTokenizedMessage> TokenizedMessage(EMessageSeverity::Type Severity, const FormatType& Format, ArgsType... Args)
 	{
-		static_assert(TIsArrayOrRefOfType<FormatType, TCHAR>::Value, "Formatting string must be a TCHAR array.");
+		static_assert(TIsArrayOrRefOfTypeByPredicate<FormatType, TIsCharEncodingCompatibleWithTCHAR>::Value, "Formatting string must be a TCHAR array.");
 		return CreateTokenizedMessage_Internal(Severity, NAME_None, FLiveLinkSubjectKey(), Format, Args...);
 	}
 
@@ -88,7 +90,7 @@ public:
 	template<typename FormatType, typename... ArgsType>
 	static TSharedPtr<FTokenizedMessage> TokenizedMessageOnce(EMessageSeverity::Type Severity, FName MessageID, const FLiveLinkSubjectKey& SubjectKey, const FormatType& Format, ArgsType... Args)
 	{
-		static_assert(TIsArrayOrRefOfType<FormatType, TCHAR>::Value, "Formatting string must be a TCHAR array.");
+		static_assert(TIsArrayOrRefOfTypeByPredicate<FormatType, TIsCharEncodingCompatibleWithTCHAR>::Value, "Formatting string must be a TCHAR array.");
 		return CreateTokenizedMessage_Internal(Severity, MessageID, SubjectKey, Format, Args...);
 	}
 
diff --git a/Engine/Source/Runtime/RenderCore/Public/ShaderParameterMacros.h b/Engine/Source/Runtime/RenderCore/Public/ShaderParameterMacros.h
--- a/Engine/Source/Runtime/RenderCore/Public/ShaderParameterMacros.h
+++ b/Engine/Source/Runtime/RenderCore/Public/ShaderParameterMacros.h
@@ -10,7 +10,8 @@
 #include "ShaderParameterMetadata.h"
 #include "RenderGraphAllocator.h"
 #include "Algo/Reverse.h"
-
+#include "Templates/IsArrayOrRefOfTypeByPredicate.h"
+#include "Traits/IsCharEncodingCompatibleWith.h"
 
 class FRDGTexture;
 class FRDGTextureSRV;
@@ -1294,13 +1295,13 @@ struct TShaderParameterStructTypeInfo<StructType[InNumElements]>
 		struct zzNextMemberId##MemberName { enum { HasDeclaredResource = zzMemberId##MemberName::HasDeclaredResource || !TypeInfo::bIsStoredInConstantBuffer }; }; \
 		static zzFuncPtr zzAppendMemberGetPrev(zzNextMemberId##MemberName, TArray<FShaderParametersMetadata::FMember>* Members) \
 		{ \
-			static_assert(TypeInfo::bIsStoredInConstantBuffer || TIsArrayOrRefOfType<decltype(OptionalShaderType), TCHAR>::Value, "No shader type for " #MemberName "."); \
+			static_assert(TypeInfo::bIsStoredInConstantBuffer || TIsArrayOrRefOfTypeByPredicate<decltype(OptionalShaderType), TIsCharEncodingCompatibleWithTCHAR>::Value, "No shader type for " #MemberName "."); \
 			static_assert(\
 				(STRUCT_OFFSET(zzTThisStruct, MemberName) & (TypeInfo::Alignment - 1)) == 0, \
 				"Misaligned uniform buffer struct member " #MemberName "."); \
 			Members->Add(FShaderParametersMetadata::FMember( \
 				TEXT(#MemberName), \
-				OptionalShaderType, \
+				(const TCHAR*)OptionalShaderType, \
 				__LINE__, \
 				STRUCT_OFFSET(zzTThisStruct,MemberName), \
 				EUniformBufferBaseType(BaseType), \
