Fixed uninitialized properties in FActorDesc.
#rnx

diff --git a/Engine/Source/Runtime/Engine/Private/WorldPartitionBlueprintLibrary.cpp b/Engine/Source/Runtime/Engine/Private/WorldPartitionBlueprintLibrary.cpp
--- a/Engine/Source/Runtime/Engine/Private/WorldPartitionBlueprintLibrary.cpp
+++ b/Engine/Source/Runtime/Engine/Private/WorldPartitionBlueprintLibrary.cpp
@@ -12,6 +12,8 @@
 
 FActorDesc::FActorDesc()
 	: Bounds(ForceInit)
+	, bIsSpatiallyLoaded(false)
+	, bActorIsEditorOnly(false)
 {}
 
 #if WITH_EDITOR
