Fixed incorrect HISM culling in ray tracing because bounds are not always initialized when the culling happens.

[REDACTED_HANDLE] at Nvidia

diff --git a/Engine/Source/Runtime/Engine/Classes/Engine/InstancedStaticMesh.h b/Engine/Source/Runtime/Engine/Classes/Engine/InstancedStaticMesh.h
--- a/Engine/Source/Runtime/Engine/Classes/Engine/InstancedStaticMesh.h
+++ b/Engine/Source/Runtime/Engine/Classes/Engine/InstancedStaticMesh.h
@@ -421,7 +421,7 @@ struct FPerInstanceRenderData
 	TSharedPtr<FStaticMeshInstanceData, ESPMode::ThreadSafe> InstanceBuffer_GameThread;
 
 	/** Get data for culling ray tracing instances */
-	const TArray<FVector4f>& GetPerInstanceBounds();
+	const TArray<FVector4f>& GetPerInstanceBounds(FBox CurrentBounds);
 
 	/** Get cached CPU-friendly instance transforms */
 	const TArray<FRenderTransform>& GetPerInstanceTransforms();
@@ -437,7 +437,7 @@ private:
 	TArray<FVector4f> PerInstanceBounds;
 	TArray<FRenderTransform> PerInstanceTransforms;
 	FGraphEventRef UpdateBoundsTask;
-	const FBox InstanceLocalBounds;
+	FBox InstanceLocalBounds;
 	bool bTrackBounds;
 	bool bBoundsTransformsDirty;
 };
diff --git a/Engine/Source/Runtime/Engine/Private/InstancedStaticMesh.cpp b/Engine/Source/Runtime/Engine/Private/InstancedStaticMesh.cpp
--- a/Engine/Source/Runtime/Engine/Private/InstancedStaticMesh.cpp
+++ b/Engine/Source/Runtime/Engine/Private/InstancedStaticMesh.cpp
@@ -1208,11 +1208,18 @@ void FPerInstanceRenderData::EnsureInstanceDataUpdated(bool bForceUpdate)
 	}
 }
 
-const TArray<FVector4f>& FPerInstanceRenderData::GetPerInstanceBounds()
+const TArray<FVector4f>& FPerInstanceRenderData::GetPerInstanceBounds(FBox CurrentBounds)
 {
 	// if bounds tracking has not been enabled, it needs to be enabled and instance data must be forced to update
 	bool bForceUpdate = !bTrackBounds;
 	bTrackBounds = true;
+
+	if (CurrentBounds != InstanceLocalBounds)
+	{
+		// bounding box changed, need to force an update
+		bForceUpdate = true;
+		InstanceLocalBounds = CurrentBounds;
+	}
 	EnsureInstanceDataUpdated(bForceUpdate);
 	return PerInstanceBounds;
 }
@@ -1844,8 +1851,9 @@ void FInstancedStaticMeshSceneProxy::GetDynamicRayTracingInstances(struct FRayTr
 		DistanceToInstanceStart = (DistanceToInstanceCenter - InstanceRadius) * LocalToWorldScale;
 	};
 
+	const FBox CurrentBounds = StaticMeshBounds.GetBox();
 	const TArray<FRenderTransform>& PerInstanceTransforms = InstancedRenderData.PerInstanceRenderData->GetPerInstanceTransforms();
-	const TArray<FVector4f>& PerInstanceBounds = InstancedRenderData.PerInstanceRenderData->GetPerInstanceBounds();
+	const TArray<FVector4f>& PerInstanceBounds = InstancedRenderData.PerInstanceRenderData->GetPerInstanceBounds(CurrentBounds);
 	if (CVarRayTracingRenderInstancesCulling.GetValueOnRenderThread() > 0 && PerInstanceBounds.Num())
 	{
 		// whether to use angular culling instead of distance, angle is halved as it is compared against the projection of the radius rather than the diameter
@@ -4147,10 +4155,16 @@ void UInstancedStaticMeshComponent::InitPerInstanceRenderData(bool InitializeFro
 	ERHIFeatureLevel::Type FeatureLevel = World != nullptr ? World->FeatureLevel.GetValue() : GMaxRHIFeatureLevel;
 
 	bool KeepInstanceBufferCPUAccess = UseGPUScene(GetFeatureLevelShaderPlatform(FeatureLevel), FeatureLevel) || GIsEditor || InRequireCPUAccess || ComponentRequestsCPUAccess(this, FeatureLevel);
-	bool bTrackBounds = IsRayTracingEnabled() && bVisibleInRayTracing;
 
 	FBox LocalBounds;
-	GetLocalBounds(LocalBounds.Min, LocalBounds.Max);
+	if (GetStaticMesh())
+	{
+		FVector BoundsMin, BoundsMax;
+		GetLocalBounds(BoundsMin, BoundsMax);
+		LocalBounds = FBox(BoundsMin, BoundsMax);
+	}
+
+	bool bTrackBounds = IsRayTracingEnabled() && bVisibleInRayTracing && LocalBounds.IsValid;
 
 	// If Nanite is used, we should defer the upload to GPU as the Nanite proxy simply will skip this step.
 	// We can't just disable the upload, because at this point we can't know whether the Nanite proxy will be created in the end
