Allow a world to contains multiple WorldDataLayers
Calls to retrieve data layers from the subsystem are now forwarded to a collection of WorldDataLayers
The collection forwards the call to the relevent WorldDataLayers

#test AncientGame, CitySample, Cook, PIE, -Game, Hlods

diff --git a/Engine/Source/Editor/DataLayerEditor/Private/DataLayer/DataLayerEditorSubsystem.cpp b/Engine/Source/Editor/DataLayerEditor/Private/DataLayer/DataLayerEditorSubsystem.cpp
--- a/Engine/Source/Editor/DataLayerEditor/Private/DataLayer/DataLayerEditorSubsystem.cpp
+++ b/Engine/Source/Editor/DataLayerEditor/Private/DataLayer/DataLayerEditorSubsystem.cpp
@@ -33,8 +33,7 @@ DEFINE_LOG_CATEGORY_STATIC(LogDataLayerEditorSubsystem, All, All);
 
 FDataLayerCreationParameters::FDataLayerCreationParameters()
 	:DataLayerAsset(nullptr),
-	ParentDataLayer(nullptr),
-	WorlDataLayers(nullptr)
+	WorldDataLayers(nullptr)
 {
 
 }
@@ -56,6 +55,8 @@ private:
 	void OnObjectPostEditChange(UObject* Object, FPropertyChangedEvent& PropertyChangedEvent);
 	void OnLevelActorsAdded(AActor* InActor) { DataLayerEditorSubsystem->InitializeNewActorDataLayers(InActor); }
 	void OnLevelSelectionChanged(UObject* InObject) { DataLayerEditorSubsystem->OnSelectionChanged(); }
+	void OnWorldDataLayersPostRegister(AWorldDataLayers* WorldDatalayers) { DataLayerEditorSubsystem->OnWorldDataLayersPostRegister(WorldDatalayers); }
+	void OnWorldDataLayersPreUnregister(AWorldDataLayers* WorldDatalayers) { DataLayerEditorSubsystem->OnWorldDataLayersPreUnregister(WorldDatalayers); }
 
 	UDataLayerEditorSubsystem* DataLayerEditorSubsystem;
 	bool bIsInitialized;
@@ -87,6 +88,8 @@ void FDataLayersBroadcast::Deinitialize()
 		}
 		USelection::SelectionChangedEvent.RemoveAll(this);
 		USelection::SelectObjectEvent.RemoveAll(this);
+		UDataLayerSubsystem::OnWorldDataLayerPostRegister.RemoveAll(this);
+		UDataLayerSubsystem::OnWorldDataLayerPreUnregister.RemoveAll(this);
 	}
 }
 
@@ -101,6 +104,8 @@ void FDataLayersBroadcast::Initialize()
 		GEngine->OnLevelActorAdded().AddRaw(this, &FDataLayersBroadcast::OnLevelActorsAdded);
 		USelection::SelectionChangedEvent.AddRaw(this, &FDataLayersBroadcast::OnLevelSelectionChanged);
 		USelection::SelectObjectEvent.AddRaw(this, &FDataLayersBroadcast::OnLevelSelectionChanged);
+		UDataLayerSubsystem::OnWorldDataLayerPostRegister.AddRaw(this, &FDataLayersBroadcast::OnWorldDataLayersPostRegister);
+		UDataLayerSubsystem::OnWorldDataLayerPreUnregister.AddRaw(this, &FDataLayersBroadcast::OnWorldDataLayersPreUnregister);
 	}
 }
 
@@ -291,6 +296,16 @@ void UDataLayerEditorSubsystem::RemoveFromActorEditorContext(UDataLayerInstance*
 	}
 }
 
+void UDataLayerEditorSubsystem::OnWorldDataLayersPostRegister(AWorldDataLayers* WorldDataLayers)
+{
+	EditorRefreshDataLayerBrowser();
+}
+
+void UDataLayerEditorSubsystem::OnWorldDataLayersPreUnregister(AWorldDataLayers* WorldDataLayers)
+{
+	EditorRefreshDataLayerBrowser();
+}
+
 void UDataLayerEditorSubsystem::EditorMapChange()
 {
 	if (UWorld * World = GetWorld())
@@ -409,6 +424,9 @@ bool UDataLayerEditorSubsystem::AddActorsToDataLayers(const TArray<AActor*>& Act
 	{
 		GEditor->GetSelectedActors()->BeginBatchSelectOperation();
 
+		UDataLayerSubsystem* DataLayerSubsystem = UWorld::GetSubsystem<UDataLayerSubsystem>(GetWorld());
+		UE_CLOG(!DataLayerSubsystem, LogDataLayerEditorSubsystem, Error, TEXT("%s - Adding actors to data layers while the world is null."), ANSI_TO_TCHAR(__FUNCTION__));
+
 		for (AActor* Actor : Actors)
 		{
 			if (!IsActorValidForDataLayer(Actor))
@@ -417,9 +435,28 @@ bool UDataLayerEditorSubsystem::AddActorsToDataLayers(const TArray<AActor*>& Act
 			}
 
 			bool bActorWasModified = false;
-			for (const UDataLayerInstance* DataLayer : DataLayers)
+			for (const UDataLayerInstance* DataLayerInstance : DataLayers)
 			{
-				if (Actor->AddDataLayer(DataLayer))
+				if (const UDataLayerInstanceWithAsset* DataLayerInstanceWithAsset = CastChecked<UDataLayerInstanceWithAsset>(DataLayerInstance))
+				{
+					// If actor's level doesn't match this DataLayerInstance outer level, 
+					// Make sure that a DataLayer Instance for this Data Layer Asset exists in the Actor's level.
+					if (Actor->GetLevel() != DataLayerInstance->GetTypedOuter<ULevel>())
+					{
+						if (DataLayerSubsystem != nullptr)
+						{
+							DataLayerInstance = DataLayerSubsystem->GetDataLayerInstance(DataLayerInstanceWithAsset->GetAsset(), Actor->GetLevel());
+
+							bool bDataLayerInstanceExistsInActorLevel = DataLayerInstance != nullptr;
+							if (!bDataLayerInstanceExistsInActorLevel)
+							{
+								DataLayerInstance = CreateDataLayerInstance<UDataLayerInstanceWithAsset>(Actor->GetLevel()->GetWorldDataLayers(), DataLayerInstanceWithAsset->GetAsset());
+							}
+						}
+					}
+				}
+
+				if (Actor->AddDataLayer(DataLayerInstance))
 				{
 					bActorWasModified = true;
 					BroadcastActorDataLayersChanged(Actor);
@@ -753,12 +790,16 @@ bool UDataLayerEditorSubsystem::UpdateAllActorsVisibility(const bool bNotifySele
 
 	bool bSelectionChanged = false;
 	bool bChangesOccurred = false;
-	for (AActor* Actor : FActorRange(GetWorld()))
+
+	if (UWorld* World = GetWorld())
 	{
-		bool bActorModified = false;
-		bool bActorSelectionChanged = false;
-		bChangesOccurred |= UpdateActorVisibility(Actor, bActorSelectionChanged, bActorModified, /*bActorNotifySelectionChange*/false, /*bActorRedrawViewports*/false);
-		bSelectionChanged |= bActorSelectionChanged;
+		for (AActor* Actor : FActorRange(World))
+		{
+			bool bActorModified = false;
+			bool bActorSelectionChanged = false;
+			bChangesOccurred |= UpdateActorVisibility(Actor, bActorSelectionChanged, bActorModified, /*bActorNotifySelectionChange*/false, /*bActorRedrawViewports*/false);
+			bSelectionChanged |= bActorSelectionChanged;
+		}
 	}
 
 	if (bNotifySelectionChange && bSelectionChanged)
@@ -1056,23 +1097,21 @@ UDataLayerInstance* UDataLayerEditorSubsystem::CreateDataLayerInstance(const FDa
 	UDataLayerInstance* NewDataLayer = nullptr;
 
 	UE_CLOG(!GetWorld(), LogDataLayerEditorSubsystem, Error, TEXT("%s - Failed because world in null."), ANSI_TO_TCHAR(__FUNCTION__));
-	if (AWorldDataLayers* WorldDataLayers = Parameters.WorlDataLayers != nullptr ? Parameters.WorlDataLayers.Get() : (GetWorld() ? GetWorld()->GetWorldDataLayers() : nullptr))
-	{
-		check(!Parameters.ParentDataLayer || !Parameters.WorlDataLayers || (Parameters.ParentDataLayer->GetOuterAWorldDataLayers() == Parameters.WorlDataLayers));
-	if (!WorldDataLayers->HasDeprecatedDataLayers())
-	{
-			NewDataLayer = WorldDataLayers->CreateDataLayer<UDataLayerInstanceWithAsset>(Parameters.DataLayerAsset);
-	}
-	else
+	if (AWorldDataLayers* WorldDataLayers = Parameters.WorldDataLayers != nullptr ? Parameters.WorldDataLayers.Get() : (GetWorld() ? GetWorld()->GetWorldDataLayers() : nullptr))
 	{
-		NewDataLayer = WorldDataLayers->CreateDataLayer<UDeprecatedDataLayerInstance>();
+		if (!WorldDataLayers->HasDeprecatedDataLayers())
+		{
+			NewDataLayer = CreateDataLayerInstance<UDataLayerInstanceWithAsset>(WorldDataLayers, Parameters.DataLayerAsset);
+		}
+		else
+		{
+			NewDataLayer = CreateDataLayerInstance<UDeprecatedDataLayerInstance>(WorldDataLayers);
+		}
 	}
 
 	if (NewDataLayer != nullptr)
 	{
 		BroadcastDataLayerChanged(EDataLayerAction::Add, NewDataLayer, NAME_None);
-			SetParentDataLayer(NewDataLayer, Parameters.ParentDataLayer);
-		}
 	}
 	
 	return NewDataLayer;
@@ -1083,7 +1122,7 @@ void UDataLayerEditorSubsystem::DeleteDataLayers(const TArray<UDataLayerInstance
 	UE_CLOG(!GetWorld(), LogDataLayerEditorSubsystem, Error, TEXT("%s - Failed because world in null."), ANSI_TO_TCHAR(__FUNCTION__));
 	if (UDataLayerSubsystem* DataLayerSubsystem = UWorld::GetSubsystem<UDataLayerSubsystem>(GetWorld()))
 	{
-		if (DataLayerSubsystem->RemoveDataLayers(DataLayersToDelete))
+		if (DataLayerSubsystem->RemoveDataLayers(DataLayersToDelete) > 0)
 		{
 			for (UDataLayerInstance* DataLayerInstance : DataLayersToDelete)
 			{
diff --git a/Engine/Source/Editor/DataLayerEditor/Private/DataLayer/DataLayerHierarchy.cpp b/Engine/Source/Editor/DataLayerEditor/Private/DataLayer/DataLayerHierarchy.cpp
--- a/Engine/Source/Editor/DataLayerEditor/Private/DataLayer/DataLayerHierarchy.cpp
+++ b/Engine/Source/Editor/DataLayerEditor/Private/DataLayer/DataLayerHierarchy.cpp
@@ -148,11 +148,6 @@ UWorld* FDataLayerHierarchy::GetOwningWorld() const
 
 void FDataLayerHierarchy::CreateItems(TArray<FSceneOutlinerTreeItemPtr>& OutItems) const
 {
-	if (!GetOwningWorld()->GetWorldDataLayers())
-	{
-		return;
-	}
-
 	const UDataLayerSubsystem* DataLayerSubsystem = UWorld::GetSubsystem<UDataLayerSubsystem>(GetOwningWorld());
 
 	auto IsDataLayerShown = [this](const UDataLayerInstance* DataLayerInstance)
@@ -176,11 +171,14 @@ void FDataLayerHierarchy::CreateItems(TArray<FSceneOutlinerTreeItemPtr>& OutItem
 		}, OuterLevel);
 	};
 
+	CreateDataLayerTreeItems(GetOwningWorld()->PersistentLevel);
+
 	// CurrentLevel represents the current level if different than the PersistentLevel
 	ULevel* CurrentLevel = (GetOwningWorld()->GetCurrentLevel() && !GetOwningWorld()->GetCurrentLevel()->IsPersistentLevel()) ? GetOwningWorld()->GetCurrentLevel() : nullptr;
-
-	CreateDataLayerTreeItems(GetOwningWorld()->PersistentLevel);
-	CreateDataLayerTreeItems(CurrentLevel);
+	if (CurrentLevel != nullptr)
+	{
+		CreateDataLayerTreeItems(CurrentLevel);
+	}
 
 	if (bShowDataLayerActors)
 	{
diff --git a/Engine/Source/Editor/DataLayerEditor/Private/DataLayer/DataLayerMode.cpp b/Engine/Source/Editor/DataLayerEditor/Private/DataLayer/DataLayerMode.cpp
--- a/Engine/Source/Editor/DataLayerEditor/Private/DataLayer/DataLayerMode.cpp
+++ b/Engine/Source/Editor/DataLayerEditor/Private/DataLayer/DataLayerMode.cpp
@@ -558,6 +558,11 @@ FSceneOutlinerDragValidationInfo FDataLayerMode::ValidateDrop(const ISceneOutlin
 						{
 							return FSceneOutlinerDragValidationInfo(ESceneOutlinerDropCompatibility::IncompatibleGeneric, LOCTEXT("CantMoveToSameDataLayer", "Can't move Data Layer to same Data Layer"));
 						}
+
+						if (ParentDataLayer->GetOuterAWorldDataLayers() != DataLayerInstance->GetOuterAWorldDataLayers())
+						{
+							return FSceneOutlinerDragValidationInfo(ESceneOutlinerDropCompatibility::IncompatibleGeneric, LOCTEXT("CantMoveToDifferentOuterWorldDataLayer", "Can't move Data Layer to a different World Data Layer"));
+						}
 					}
 				
 					if (ParentDataLayer->IsLocked())
@@ -945,6 +950,16 @@ TArray<UDataLayerInstance*> FDataLayerMode::GetSelectedDataLayers(SSceneOutliner
 	return MoveTemp(ValidSelectedDataLayers);
 }
 
+TArray<AWorldDataLayers*> FDataLayerMode::GetSelectedWorldDataLayers(SSceneOutliner* InSceneOutliner) const
+{
+	FSceneOutlinerItemSelection ItemSelection(InSceneOutliner->GetSelection());
+	TArray<FWorldDataLayersTreeItem*> SelectedWorldDataLayersItems;
+	ItemSelection.Get<FWorldDataLayersTreeItem>(SelectedWorldDataLayersItems);
+	TArray<AWorldDataLayers*> ValidSelectedDataLayers;
+	Algo::TransformIf(SelectedWorldDataLayersItems, ValidSelectedDataLayers, [](const auto Item) { return Item && Item->GetWorldDataLayers(); }, [](const auto Item) { return Item->GetWorldDataLayers(); });
+	return MoveTemp(ValidSelectedDataLayers);
+}
+
 void FDataLayerMode::CreateDataLayerPicker(UToolMenu* InMenu, FOnDataLayerPicked OnDataLayerPicked, bool bInShowRoot /*= false*/)
 {
 	if (bInShowRoot)
@@ -1016,12 +1031,15 @@ void FDataLayerMode::RegisterContextMenu()
 					Mode->SelectedDataLayersSet.Empty();
 					Mode->SelectedDataLayerActors.Empty();
 
+					TArray<AWorldDataLayers*> SelectedWorldDataLayers = Mode->GetSelectedWorldDataLayers(SceneOutliner);
+
 					FDataLayerCreationParameters CreationParams;
 					CreationParams.DataLayerAsset = const_cast<UDataLayerAsset*>(InDataLayerAsset);
-					CreationParams.ParentDataLayer = InParentDataLayer;
-					CreationParams.WorlDataLayers = InParentDataLayer ? InParentDataLayer->GetOuterAWorldDataLayers() : Mode->GetOwningWorldAWorldDataLayers();
+					CreationParams.WorldDataLayers = InParentDataLayer ? InParentDataLayer->GetOuterAWorldDataLayers() : (SelectedWorldDataLayers.Num() == 1 ? SelectedWorldDataLayers[0] : Mode->GetOwningWorldAWorldDataLayers());
 					if (UDataLayerInstance* NewDataLayerInstance = UDataLayerEditorSubsystem::Get()->CreateDataLayerInstance(CreationParams))
 					{
+						UDataLayerEditorSubsystem::Get()->SetParentDataLayer(NewDataLayerInstance, InParentDataLayer);
+
 						Mode->SelectedDataLayersSet.Add(NewDataLayerInstance);
 						// Select it and open a rename when it gets refreshed
 						SceneOutliner->OnItemAdded(NewDataLayerInstance, SceneOutliner::ENewItemAction::Select | SceneOutliner::ENewItemAction::Rename);
@@ -1117,7 +1135,7 @@ void FDataLayerMode::RegisterContextMenu()
 								const FScopedTransaction Transaction(LOCTEXT("AddSelectedActorsToNewDataLayer", "Add Selected Actors to New Data Layer"));
 								FDataLayerCreationParameters CreationParams;
 								CreationParams.DataLayerAsset = DataLayerAsset;
-								CreationParams.WorlDataLayers = Mode->GetOwningWorldAWorldDataLayers();
+								CreationParams.WorldDataLayers = Mode->GetOwningWorldAWorldDataLayers();
 								if (UDataLayerInstance* NewDataLayerInstance = UDataLayerEditorSubsystem::Get()->CreateDataLayerInstance(CreationParams))
 								{
 									UDataLayerEditorSubsystem::Get()->AddSelectedActorsToDataLayer(NewDataLayerInstance);
diff --git a/Engine/Source/Editor/DataLayerEditor/Private/DataLayer/DataLayerMode.h b/Engine/Source/Editor/DataLayerEditor/Private/DataLayer/DataLayerMode.h
--- a/Engine/Source/Editor/DataLayerEditor/Private/DataLayer/DataLayerMode.h
+++ b/Engine/Source/Editor/DataLayerEditor/Private/DataLayer/DataLayerMode.h
@@ -109,6 +109,7 @@ private:
 	TArray<AActor*> GetActorsFromOperation(const FDragDropOperation& Operation, bool bOnlyFindFirst = false) const;
 	TArray<UDataLayerInstance*> GetDataLayersFromOperation(const FDragDropOperation& Operation, bool bOnlyFindFirst = false) const;
 	TArray<UDataLayerInstance*> GetSelectedDataLayers(SSceneOutliner* InSceneOutliner) const;
+	TArray<AWorldDataLayers*> GetSelectedWorldDataLayers(SSceneOutliner* InSceneOutliner) const;
 	void SetParentDataLayer(const TArray<UDataLayerInstance*> DataLayers, UDataLayerInstance* ParentDataLayer) const;
 	void OnLevelSelectionChanged(UObject* Obj);
 	static void CreateDataLayerPicker(UToolMenu* InMenu, FOnDataLayerPicked OnDataLayerPicked, bool bInShowRoot = false);
diff --git a/Engine/Source/Editor/DataLayerEditor/Private/DataLayer/WorldDataLayersTreeItem.cpp b/Engine/Source/Editor/DataLayerEditor/Private/DataLayer/WorldDataLayersTreeItem.cpp
--- a/Engine/Source/Editor/DataLayerEditor/Private/DataLayer/WorldDataLayersTreeItem.cpp
+++ b/Engine/Source/Editor/DataLayerEditor/Private/DataLayer/WorldDataLayersTreeItem.cpp
@@ -88,7 +88,14 @@ FString FWorldDataLayersTreeItem::GetDisplayString() const
 	if (const AWorldDataLayers* WorldDataLayersPtr = WorldDataLayers.Get())
 	{
 		check(WorldDataLayersPtr->GetLevel());
-		return SceneOutliner::GetWorldDescription(WorldDataLayersPtr->GetLevel()->GetTypedOuter<UWorld>()).ToString();
+		FString DisplayString = SceneOutliner::GetWorldDescription(WorldDataLayersPtr->GetLevel()->GetTypedOuter<UWorld>()).ToString();
+
+		if (!WorldDataLayersPtr->GetActorLabel().IsEmpty())
+		{
+			DisplayString += TEXT(" - ") + WorldDataLayersPtr->GetActorLabel();
+		}
+
+		return DisplayString;
 	}
 	return LOCTEXT("UnknownWorldDataLayers", "Unknown").ToString();
 }
@@ -102,7 +109,7 @@ bool FWorldDataLayersTreeItem::IsReadOnly() const
 {
 	const AWorldDataLayers* WorldDataLayersPtr = WorldDataLayers.Get();
 	UWorld* WorldPtr = WorldDataLayersPtr ? WorldDataLayersPtr->GetLevel()->GetTypedOuter<UWorld>() : nullptr;
-	return (WorldDataLayersPtr && !WorldDataLayersPtr->IsMainWorldDataLayers()) || !WorldPtr || (WorldPtr->WorldType != EWorldType::Editor);
+	return (WorldDataLayersPtr && WorldDataLayersPtr->IsSubWorldDataLayers()) || !WorldPtr || (WorldPtr->WorldType != EWorldType::Editor);
 }
 
 TSharedRef<SWidget> FWorldDataLayersTreeItem::GenerateLabelWidget(ISceneOutliner& Outliner, const STableRow<FSceneOutlinerTreeItemPtr>& InRow)
diff --git a/Engine/Source/Editor/DataLayerEditor/Public/DataLayer/DataLayerEditorSubsystem.h b/Engine/Source/Editor/DataLayerEditor/Public/DataLayer/DataLayerEditorSubsystem.h
--- a/Engine/Source/Editor/DataLayerEditor/Public/DataLayer/DataLayerEditorSubsystem.h
+++ b/Engine/Source/Editor/DataLayerEditor/Public/DataLayer/DataLayerEditorSubsystem.h
@@ -8,6 +8,7 @@
 #include "WorldPartition/DataLayer/ActorDataLayer.h"
 #include "WorldPartition/DataLayer/DataLayerInstance.h"
 #include "IActorEditorContextClient.h"
+#include "WorldPartition/DataLayer/WorldDataLayers.h"
 #include "DataLayerEditorSubsystem.generated.h"
 
 class AActor;
@@ -30,13 +31,9 @@ struct FDataLayerCreationParameters
 	UPROPERTY(EditAnywhere, BlueprintReadWrite, Category = "Data Layer")
 	TObjectPtr<UDataLayerAsset> DataLayerAsset;
 
-	// Optional. Will create the Data Layer at root level is unset.
-	UPROPERTY()
-	TObjectPtr<UDataLayerInstance> ParentDataLayer;
-
 	// Optional. Will default at the level WorldDataLayers if unset.
-	UPROPERTY()
-	TObjectPtr<AWorldDataLayers> WorlDataLayers;
+	UPROPERTY(EditAnywhere, BlueprintReadWrite, Category = "Data Layer")
+	TObjectPtr<AWorldDataLayers> WorldDataLayers;
 };
 
 UCLASS()
@@ -80,6 +77,12 @@ public:
 	void AddToActorEditorContext(UDataLayerInstance* InDataLayerInstance);
 	void RemoveFromActorEditorContext(UDataLayerInstance* InDataLayerInstance);
 
+	void OnWorldDataLayersPostRegister(AWorldDataLayers* WorldDataLayers);
+	void OnWorldDataLayersPreUnregister(AWorldDataLayers* WorldDataLayers);
+
+	template<class DataLayerType, typename ...Args>
+	UDataLayerInstance* CreateDataLayerInstance(AWorldDataLayers* WorldDataLayers, Args&&... InArgs);
+
 	/* Broadcasts whenever one or more DataLayers are modified
 	 *
 	 * Actions
@@ -695,3 +698,21 @@ private:
 
 	friend class FDataLayersBroadcast;
 };
+
+template<class DataLayerType, typename ...Args>
+UDataLayerInstance* UDataLayerEditorSubsystem::CreateDataLayerInstance(AWorldDataLayers* WorldDataLayers, Args&&... InArgs)
+{
+	UDataLayerInstance* NewDataLayer = nullptr;
+
+	if (WorldDataLayers)
+	{
+		NewDataLayer = WorldDataLayers->CreateDataLayer<DataLayerType>(Forward<Args>(InArgs)...);
+	}
+
+	if (NewDataLayer != nullptr)
+	{
+		BroadcastDataLayerChanged(EDataLayerAction::Add, NewDataLayer, NAME_None);
+	}
+
+	return NewDataLayer;
+}
\ No newline at end of file
diff --git a/Engine/Source/Runtime/Engine/Private/WorldPartition/DataLayer/DataLayerInstance.cpp b/Engine/Source/Runtime/Engine/Private/WorldPartition/DataLayer/DataLayerInstance.cpp
--- a/Engine/Source/Runtime/Engine/Private/WorldPartition/DataLayer/DataLayerInstance.cpp
+++ b/Engine/Source/Runtime/Engine/Private/WorldPartition/DataLayer/DataLayerInstance.cpp
@@ -269,6 +269,7 @@ void UDataLayerInstance::ForEachChild(TFunctionRef<bool(const UDataLayerInstance
 
 void UDataLayerInstance::AddChild(UDataLayerInstance* InDataLayer)
 {
+	check(InDataLayer->GetOuterAWorldDataLayers() == GetOuterAWorldDataLayers())
 	Modify();
 	checkSlow(!Children.Contains(InDataLayer));
 	Children.Add(InDataLayer);
diff --git a/Engine/Source/Runtime/Engine/Private/WorldPartition/DataLayer/DataLayerInstanceWithAsset.cpp b/Engine/Source/Runtime/Engine/Private/WorldPartition/DataLayer/DataLayerInstanceWithAsset.cpp
--- a/Engine/Source/Runtime/Engine/Private/WorldPartition/DataLayer/DataLayerInstanceWithAsset.cpp
+++ b/Engine/Source/Runtime/Engine/Private/WorldPartition/DataLayer/DataLayerInstanceWithAsset.cpp
@@ -32,7 +32,7 @@ void UDataLayerInstanceWithAsset::OnCreated(const UDataLayerAsset* Asset)
 
 bool UDataLayerInstanceWithAsset::IsReadOnly() const
 {
-	return !GetOuterAWorldDataLayers()->IsMainWorldDataLayers();
+	return GetOuterAWorldDataLayers()->IsSubWorldDataLayers();
 }
 
 bool UDataLayerInstanceWithAsset::IsLocked() const
@@ -45,24 +45,10 @@ bool UDataLayerInstanceWithAsset::IsLocked() const
 }
 
 bool UDataLayerInstanceWithAsset::AddActor(AActor* Actor) const
-{
-	// If actor's level WorldDataLayers doesn't match this DataLayerInstance outer WorldDataLayers, 
-	// Make sure that an instance for this Data Layer asset exists in the outer level AWorldDataLayer and forward the call on it.
-	AWorldDataLayers* OuterWorldDataLayers = Actor->GetLevel()->GetWorldDataLayers();
-	if (GetOuterAWorldDataLayers() != OuterWorldDataLayers)
-	{
-		check(OuterWorldDataLayers);
-		UDataLayerInstanceWithAsset* OuterDataLayerInstanceWithAsset = Cast<UDataLayerInstanceWithAsset>(const_cast<UDataLayerInstance*>(OuterWorldDataLayers->GetDataLayerInstance(DataLayerAsset)));
-		if (!OuterDataLayerInstanceWithAsset)
-		{
-			OuterDataLayerInstanceWithAsset = OuterWorldDataLayers->CreateDataLayer<UDataLayerInstanceWithAsset>(DataLayerAsset);
-		}
-		return OuterDataLayerInstanceWithAsset->AddActor(Actor);
-	}
-	else
-	{
-		return Actor->AddDataLayer(DataLayerAsset);
-	}
+{	
+	check(GetWorld()->GetSubsystem<UDataLayerSubsystem>()->GetDataLayerInstance(DataLayerAsset) != nullptr);
+	check(GetTypedOuter<ULevel>() == Actor->GetLevel()); // Make sure the instance is part of the same world as the actor.
+	return Actor->AddDataLayer(DataLayerAsset);
 }
 
 bool UDataLayerInstanceWithAsset::RemoveActor(AActor* Actor) const
diff --git a/Engine/Source/Runtime/Engine/Private/WorldPartition/DataLayer/DataLayerSubsystem.cpp b/Engine/Source/Runtime/Engine/Private/WorldPartition/DataLayer/DataLayerSubsystem.cpp
--- a/Engine/Source/Runtime/Engine/Private/WorldPartition/DataLayer/DataLayerSubsystem.cpp
+++ b/Engine/Source/Runtime/Engine/Private/WorldPartition/DataLayer/DataLayerSubsystem.cpp
@@ -11,6 +11,7 @@
 #include "Engine/Engine.h"
 #include "Engine/World.h"
 #include "Engine/Canvas.h"
+#include "EngineUtils.h"
 #include "Algo/Transform.h"
 #if WITH_EDITOR
 #include "Editor.h"
@@ -23,6 +24,11 @@
 
 extern int32 GDrawDataLayersLoadTime;
 
+#if WITH_EDITOR
+FOnWorldDataLayersPostRegister UDataLayerSubsystem::OnWorldDataLayerPostRegister;
+FOnWorldDataLayersPreUnregister UDataLayerSubsystem::OnWorldDataLayerPreUnregister;
+#endif
+
 static FAutoConsoleCommandWithOutputDevice GDumpDataLayersCmd(
 	TEXT("wp.DumpDataLayers"),
 	TEXT("Dumps data layers to the log"),
@@ -72,6 +78,8 @@ void UDataLayerSubsystem::Initialize(FSubsystemCollectionBase& Collection)
 		FModuleManager::LoadModuleChecked<IDataLayerEditorModule>("DataLayerEditor");
 	}
 
+	RegisterEarlyLoadedWorldDataLayers();
+
 	UActorDescContainer::OnActorDescContainerInitialized.AddUObject(this, &UDataLayerSubsystem::OnActorDescContainerInitialized);
 #endif
 }
@@ -80,6 +88,8 @@ void UDataLayerSubsystem::Deinitialize()
 {
 	Super::Deinitialize();
 
+	WorldDataLayerCollection.Empty();
+
 #if WITH_EDITOR
 	UActorDescContainer::OnActorDescContainerInitialized.RemoveAll(this);
 #endif
@@ -95,35 +105,31 @@ bool UDataLayerSubsystem::DoesSupportWorldType(const EWorldType::Type WorldType)
 
 bool UDataLayerSubsystem::CanResolveDataLayers() const
 {
-	return GetWorld()->GetWorldDataLayers() != nullptr;
+	return !WorldDataLayerCollection.IsEmpty();
 }
 
-bool UDataLayerSubsystem::RemoveDataLayer(const UDataLayerInstance* InDataLayer)
+bool UDataLayerSubsystem::RemoveDataLayer(const UDataLayerInstance* InDataLayerInstance)
 {
-	if (AWorldDataLayers* WorldDataLayers = GetWorld()->GetWorldDataLayers())
+	bool bIsDataLayerRemoved = false;
+	WorldDataLayerCollection.ForEachWorldDataLayers([InDataLayerInstance, &bIsDataLayerRemoved](AWorldDataLayers* WorldDataLayer)
 	{
-		return WorldDataLayers->RemoveDataLayer(InDataLayer);;
-	}
-	else
-	{
-		UE_LOG(LogWorldPartition, Warning, TEXT("Removing a Data Layer without a World Partition"));
-	}
+		bIsDataLayerRemoved = WorldDataLayer->RemoveDataLayer(InDataLayerInstance);
+		return !bIsDataLayerRemoved;
+	});
 
-	return false;
+	return bIsDataLayerRemoved;
 }
 
-bool UDataLayerSubsystem::RemoveDataLayers(const TArray<UDataLayerInstance*>& InDataLayerInstances)
+int32 UDataLayerSubsystem::RemoveDataLayers(const TArray<UDataLayerInstance*>& InDataLayerInstances)
 {
-	if (AWorldDataLayers* WorldDataLayers = GetWorld()->GetWorldDataLayers())
+	int32 DataLayerRemovedCount = 0;
+	WorldDataLayerCollection.ForEachWorldDataLayers([InDataLayerInstances, &DataLayerRemovedCount](AWorldDataLayers* WorldDataLayer)
 	{
-		return WorldDataLayers->RemoveDataLayers(InDataLayerInstances);
-	}
-	else
-	{
-		UE_LOG(LogWorldPartition, Warning, TEXT("Removing Data Layers without a World Partition"));
-	}
+		DataLayerRemovedCount += WorldDataLayer->RemoveDataLayers(InDataLayerInstances);
+		return DataLayerRemovedCount != InDataLayerInstances.Num();
+	});
 
-	return false;
+	return DataLayerRemovedCount;
 }
 
 void UDataLayerSubsystem::UpdateDataLayerEditorPerProjectUserSettings() const
@@ -176,18 +182,58 @@ void UDataLayerSubsystem::GetUserLoadedInEditorStates(TArray<FName>& OutDataLaye
 }
 #endif
 
-const TSet<FName>& UDataLayerSubsystem::GetEffectiveActiveDataLayerNames() const
+TSet<FName> UDataLayerSubsystem::GetEffectiveActiveDataLayerNames() const
+{
+	TSet<FName> EffectiveActiveDataLayerNames;
+	WorldDataLayerCollection.ForEachWorldDataLayers([&EffectiveActiveDataLayerNames](AWorldDataLayers* WorldDataLayers)
+	{
+		EffectiveActiveDataLayerNames.Append(WorldDataLayers->GetEffectiveActiveDataLayerNames());
+	});
+	
+	return EffectiveActiveDataLayerNames;
+}
+
+TSet<FName> UDataLayerSubsystem::GetEffectiveLoadedDataLayerNames() const
+{
+	TSet<FName> EffectiveLoadedDataLayerNames;
+	WorldDataLayerCollection.ForEachWorldDataLayers([&EffectiveLoadedDataLayerNames](AWorldDataLayers* WorldDataLayers)
+	{
+		EffectiveLoadedDataLayerNames.Append(WorldDataLayers->GetEffectiveLoadedDataLayerNames());
+	});
+
+	return EffectiveLoadedDataLayerNames;
+}
+
+void UDataLayerSubsystem::RegisterWorldDataLayer(AWorldDataLayers* WorldDataLayers)
+{
+	if (WorldDataLayerCollection.RegisterWorldDataLayer(WorldDataLayers))
+	{
+#if WITH_EDITOR
+		OnWorldDataLayerPostRegister.Broadcast(WorldDataLayers);
+#endif
+	}
+}
+
+void UDataLayerSubsystem::UnregisterWorldDataLayer(AWorldDataLayers* WorldDataLayers)
+{
+	if (WorldDataLayerCollection.Contains(WorldDataLayers))
+	{
+#if WITH_EDITOR
+		OnWorldDataLayerPreUnregister.Broadcast(WorldDataLayers);
+#endif
+
+		WorldDataLayerCollection.UnregisterWorldDataLayer(WorldDataLayers);
+	}
+}
+
+void UDataLayerSubsystem::ForEachWorldDataLayer(TFunctionRef<bool(AWorldDataLayers*)> Func)
 {
-	static TSet<FName> EmptySet;
-	const AWorldDataLayers* WorldDataLayers = GetWorld()->GetWorldDataLayers();
-	return WorldDataLayers ? WorldDataLayers->GetEffectiveActiveDataLayerNames() : EmptySet;
+	WorldDataLayerCollection.ForEachWorldDataLayersBreakable(Func);
 }
 
-const TSet<FName>& UDataLayerSubsystem::GetEffectiveLoadedDataLayerNames() const
+void UDataLayerSubsystem::ForEachWorldDataLayer(TFunctionRef<bool(AWorldDataLayers*)> Func) const
 {
-	static TSet<FName> EmptySet;
-	const AWorldDataLayers* WorldDataLayers = GetWorld()->GetWorldDataLayers();
-	return WorldDataLayers ? WorldDataLayers->GetEffectiveLoadedDataLayerNames() : EmptySet;
+	const_cast<UDataLayerSubsystem*>(this)->ForEachWorldDataLayer(Func);
 }
 
 UDataLayerInstance* UDataLayerSubsystem::GetDataLayerInstanceFromAsset(const UDataLayerAsset* InDataLayerAsset) const
@@ -483,40 +529,28 @@ TArray<UDataLayerInstance*> UDataLayerSubsystem::ConvertArgsToDataLayers(UWorld*
 
 void UDataLayerSubsystem::DumpDataLayers(FOutputDevice& OutputDevice) const
 {
-	if (AWorldDataLayers* WorldDataLayers = GetWorld()->GetWorldDataLayers())
+	WorldDataLayerCollection.ForEachWorldDataLayers([&OutputDevice](AWorldDataLayers* WorldDataLayers)
 	{
 		WorldDataLayers->DumpDataLayers(OutputDevice);
-	}
-	else
-	{
-		UE_LOG(LogWorldPartition, Warning, TEXT("Dumping Data Layers without a World Partition"));
-	}
+	});
 }
 
 const UDataLayerInstance* UDataLayerSubsystem::GetDataLayerInstanceFromAssetName(const FName& InDataLayerAssetFullName) const
 {
-	if (AWorldDataLayers* WorldDataLayers = GetWorld()->GetWorldDataLayers())
+	const UDataLayerInstance* DataLayerInstance = nullptr;
+
+	WorldDataLayerCollection.ForEachWorldDataLayers([&DataLayerInstance, &InDataLayerAssetFullName](const AWorldDataLayers* WorldDataLayers)
 	{
-		const UDataLayerInstance* DataLayerInstance = WorldDataLayers->GetDataLayerInstanceFromAssetName(InDataLayerAssetFullName);
-		return const_cast<UDataLayerInstance*>(DataLayerInstance);
-	}
+		DataLayerInstance = WorldDataLayers->GetDataLayerInstanceFromAssetName(InDataLayerAssetFullName);
+		return DataLayerInstance == nullptr;
+	});
 
-	return nullptr;
+	return const_cast<UDataLayerInstance*>(DataLayerInstance);
 }
 
 void UDataLayerSubsystem::ForEachDataLayer(TFunctionRef<bool(UDataLayerInstance*)> Func, const ULevel* InLevelContext /* = nullptr */)
 {
-	if (!InLevelContext || InLevelContext == GetWorld()->PersistentLevel)
-	{
-		if (AWorldDataLayers* WorldDataLayers = GetWorld()->GetWorldDataLayers())
-		{
-			WorldDataLayers->ForEachDataLayer(Func);
-		}
-	}
-	else if (const AWorldDataLayers* CurrentLevelWorldDataLayers = (InLevelContext && !InLevelContext->IsPersistentLevel()) ? InLevelContext->GetWorldDataLayers() : nullptr)
-	{
-		CurrentLevelWorldDataLayers->ForEachDataLayer(Func);
-	}
+	WorldDataLayerCollection.ForEachDataLayerBreakable(Func, InLevelContext);
 }
 
 void UDataLayerSubsystem::ForEachDataLayer(TFunctionRef<bool(UDataLayerInstance*)> Func, const ULevel* InLevelContext /* = nullptr */) const
@@ -602,30 +636,31 @@ void UDataLayerSubsystem::GetDataLayerDebugColors(TMap<FName, FColor>& OutMappin
 void UDataLayerSubsystem::PushActorEditorContext() const
 {
 	++DataLayerActorEditorContextID;
-	if (AWorldDataLayers* WorldDataLayers = GetWorld()->GetWorldDataLayers())
+	WorldDataLayerCollection.ForEachWorldDataLayers([this](AWorldDataLayers* WorldDataLayers)
 	{
 		WorldDataLayers->PushActorEditorContext(DataLayerActorEditorContextID);
-	}
+	});
 }
 
 void UDataLayerSubsystem::PopActorEditorContext() const
 {
 	check(DataLayerActorEditorContextID > 0);
-	if (AWorldDataLayers* WorldDataLayers = GetWorld()->GetWorldDataLayers())
+	WorldDataLayerCollection.ForEachWorldDataLayers([this](AWorldDataLayers* WorldDataLayers)
 	{
 		WorldDataLayers->PopActorEditorContext(DataLayerActorEditorContextID);
-	}
+	});
 	--DataLayerActorEditorContextID;
 }
 
 TArray<UDataLayerInstance*> UDataLayerSubsystem::GetActorEditorContextDataLayers() const
 {
-	if(AWorldDataLayers* WorldDataLayers = GetWorld()->GetWorldDataLayers())
+	TArray<UDataLayerInstance*> DataLayersInstances;
+	WorldDataLayerCollection.ForEachWorldDataLayers([&DataLayersInstances](AWorldDataLayers* WorldDataLayers)
 	{
-		return WorldDataLayers->GetActorEditorContextDataLayers();
-	}
+		DataLayersInstances += WorldDataLayers->GetActorEditorContextDataLayers();
+	});
 
-	return TArray<UDataLayerInstance*>();
+	return DataLayersInstances;
 }
 
 uint32 UDataLayerSubsystem::GetDataLayerEditorContextHash() const
@@ -642,8 +677,9 @@ void UDataLayerSubsystem::OnActorDescContainerInitialized(UActorDescContainer* I
 {
 	check(InActorDescContainer);
 
-	auto GetWorldDataLayersActorDesc = [](const UActorDescContainer* InContainer) -> FWorldDataLayersActorDesc*
+	auto GetWorldDataLayersActorDesc = [](const UActorDescContainer* InContainer) -> TArray<const FWorldDataLayersActorDesc*>
 	{
+		TArray<const FWorldDataLayersActorDesc*> WorldDataLayersActorDecs;
 		if (InContainer)
 		{
 			for (FActorDescList::TIterator<AWorldDataLayers> Iterator(const_cast<UActorDescContainer*>(InContainer)); Iterator; ++Iterator)
@@ -651,25 +687,29 @@ void UDataLayerSubsystem::OnActorDescContainerInitialized(UActorDescContainer* I
 				// IsValid will return false for maps that don't have a valid WorldDataLayers actors
 				if (Iterator->IsValid())
 				{
-					FWorldDataLayersActorDesc* WorldDataLayersActorDesc = *Iterator;
-					return WorldDataLayersActorDesc;
+					WorldDataLayersActorDecs.Add(*Iterator);
 				}
-				// No need to iterate (we assume that there's only 1 AWorldDataLayer for now)
-				return nullptr;
 			}
 		}
-		return nullptr;
+
+		return WorldDataLayersActorDecs;
 	};
 
-	const FWorldDataLayersActorDesc* WorldDataLayersActorDesc = GetWorldDataLayersActorDesc(InActorDescContainer);
+	const TArray<const FWorldDataLayersActorDesc*> WorldDataLayersActorDescs = GetWorldDataLayersActorDesc(InActorDescContainer);
 	for (FActorDescList::TIterator<> Iterator(InActorDescContainer); Iterator; ++Iterator)
 	{
 		FWorldPartitionActorDesc* ActorDesc = *Iterator;
 		check(ActorDesc->GetContainer() == InActorDescContainer);
-		ActorDesc->SetDataLayerInstanceNames(FDataLayerUtils::ResolvedDataLayerInstanceNames(ActorDesc, WorldDataLayersActorDesc));
+		ActorDesc->SetDataLayerInstanceNames(FDataLayerUtils::ResolvedDataLayerInstanceNames(ActorDesc, WorldDataLayersActorDescs));
 	}
 }
 
+void UDataLayerSubsystem::RegisterEarlyLoadedWorldDataLayers()
+{
+	// ULevel::WorldDatalayers is loaded before the subsystem is created and miss their registration.
+	RegisterWorldDataLayer(GetWorld()->GetWorldDataLayers());
+}
+
 #endif
 
 //~ Begin Deprecated
@@ -685,7 +725,7 @@ UDataLayerInstance* UDataLayerSubsystem::GetDataLayerFromLabel(FName InDataLayer
 {
 	if (AWorldDataLayers* WorldDataLayers = GetWorld()->GetWorldDataLayers())
 	{
-		const UDataLayerInstance* DataLayerInstance = GetWorld()->GetWorldDataLayers()->GetDataLayerFromLabel(InDataLayerLabel);
+		const UDataLayerInstance* DataLayerInstance = WorldDataLayers->GetDataLayerFromLabel(InDataLayerLabel);
 		return const_cast<UDataLayerInstance*>(DataLayerInstance);
 	}
 
diff --git a/Engine/Source/Runtime/Engine/Private/WorldPartition/DataLayer/DataLayerUtils.cpp b/Engine/Source/Runtime/Engine/Private/WorldPartition/DataLayer/DataLayerUtils.cpp
--- a/Engine/Source/Runtime/Engine/Private/WorldPartition/DataLayer/DataLayerUtils.cpp
+++ b/Engine/Source/Runtime/Engine/Private/WorldPartition/DataLayer/DataLayerUtils.cpp
@@ -10,8 +10,11 @@
 #include "WorldPartition/WorldPartitionActorDesc.h"
 #include "WorldPartition/ActorDescContainer.h"
 #include "WorldPartition/DataLayer/DataLayerSubsystem.h"
+#include "WorldPartition/WorldPartitionLog.h"
+#include "Algo/Transform.h"
+#include "Algo/AllOf.h"
 
-TArray<FName> FDataLayerUtils::ResolvedDataLayerInstanceNames(const FWorldPartitionActorDesc* InActorDesc, const FWorldDataLayersActorDesc* InWorldDataLayersActorDesc, UWorld* InWorld, bool* bOutIsResultValid)
+TArray<FName> FDataLayerUtils::ResolvedDataLayerInstanceNames(const FWorldPartitionActorDesc* InActorDesc, const TArray<const FWorldDataLayersActorDesc*>& InWorldDataLayersActorDescs, UWorld* InWorld, bool* bOutIsResultValid)
 {
 	bool bLocalIsSuccess = true;
 	bool& bIsSuccess = bOutIsResultValid ? *bOutIsResultValid : bLocalIsSuccess;
@@ -24,6 +27,7 @@ TArray<FName> FDataLayerUtils::ResolvedDataLayerInstanceNames(const FWorldPartit
 		World = InActorDesc->GetContainer() ? InActorDesc->GetContainer()->GetWorld() : nullptr;
 	}
 	const UDataLayerSubsystem* DataLayerSubsystem = UWorld::GetSubsystem<UDataLayerSubsystem>(World);
+	check(DataLayerSubsystem != nullptr);
 
 	// DataLayers not using DataLayer Assets represent DataLayerInstanceNames
 	if (!InActorDesc->IsUsingDataLayerAsset())
@@ -41,12 +45,12 @@ TArray<FName> FDataLayerUtils::ResolvedDataLayerInstanceNames(const FWorldPartit
 			return Result;
 		}
 		// Fallback on FWorldDataLayersActorDesc
-		else if (InWorldDataLayersActorDesc)
+		else if (!InWorldDataLayersActorDescs.IsEmpty() && AreWorldDataLayersActorDescsSane(InWorldDataLayersActorDescs))
 		{
 			TArray<FName> Result;
 			for (const FName& DataLayerInstanceName : InActorDesc->GetDataLayers())
 			{
-				if (const FDataLayerInstanceDesc* DataLayerInstanceDesc = InWorldDataLayersActorDesc->GetDataLayerInstanceFromInstanceName(DataLayerInstanceName))
+				if (GetDataLayerInstanceDescFromInstanceName(InWorldDataLayersActorDescs, DataLayerInstanceName) != nullptr)
 				{
 					Result.Add(DataLayerInstanceName);
 				}
@@ -77,12 +81,12 @@ TArray<FName> FDataLayerUtils::ResolvedDataLayerInstanceNames(const FWorldPartit
 			return Result;
 		}
 		// Fallback on FWorldDataLayersActorDesc
-		else if (InWorldDataLayersActorDesc)
+		else if (!InWorldDataLayersActorDescs.IsEmpty() && AreWorldDataLayersActorDescsSane(InWorldDataLayersActorDescs))
 		{
 			TArray<FName> Result;
 			for (const FName& DataLayerAssetPath : InActorDesc->GetDataLayers())
 			{
-				if (const FDataLayerInstanceDesc* DataLayerInstanceDesc = InWorldDataLayersActorDesc->GetDataLayerInstanceFromAssetPath(DataLayerAssetPath))
+				if (const FDataLayerInstanceDesc* DataLayerInstanceDesc = GetDataLayerInstanceDescFromAssetPath(InWorldDataLayersActorDescs, DataLayerAssetPath))
 				{
 					Result.Add(DataLayerInstanceDesc->GetName());
 				}
@@ -100,6 +104,7 @@ bool FDataLayerUtils::ResolveRuntimeDataLayerInstanceNames(const FWorldPartition
 {
 	UWorld* World = InActorDescView.GetActorDesc()->GetContainer()->GetWorld();
 	const UDataLayerSubsystem* DataLayerSubsystem = World ? World->GetSubsystem<UDataLayerSubsystem>() : nullptr;
+
 	if (DataLayerSubsystem && DataLayerSubsystem->CanResolveDataLayers())
 	{
 		for (FName DataLayerInstanceName : InActorDescView.GetDataLayerInstanceNames())
@@ -115,19 +120,15 @@ bool FDataLayerUtils::ResolveRuntimeDataLayerInstanceNames(const FWorldPartition
 	}
 	else
 	{
+		TArray<const FWorldDataLayersActorDesc*> WorldDataLayersActorDescs;
+		FindWorldDataLayerActorDescs(ActorDescViewMap, WorldDataLayersActorDescs);
+
 		// Fallback on FWorldDataLayersActorDesc
-		TArray<const FWorldPartitionActorDescView*> WorldDataLayerViews = ActorDescViewMap.FindByExactNativeClass<AWorldDataLayers>();
-	
-		if (WorldDataLayerViews.Num())
+		if (WorldDataLayersActorDescs.Num() && AreWorldDataLayersActorDescsSane(WorldDataLayersActorDescs))
 		{
-			check(WorldDataLayerViews.Num() == 1);
-			const FWorldPartitionActorDescView* WorldDataLayersActorDescView = WorldDataLayerViews[0];
-	
-			FWorldDataLayersActorDesc* WorldDataLayersActorDesc = (FWorldDataLayersActorDesc*)WorldDataLayersActorDescView->GetActorDesc();
-
 			for (FName DataLayerInstanceName : InActorDescView.GetDataLayerInstanceNames())
 			{
-				if (const FDataLayerInstanceDesc* DataLayerInstanceDesc = WorldDataLayersActorDesc->GetDataLayerInstanceFromInstanceName(DataLayerInstanceName))
+				if (const FDataLayerInstanceDesc* DataLayerInstanceDesc = GetDataLayerInstanceDescFromInstanceName(WorldDataLayersActorDescs, DataLayerInstanceName))
 				{
 					if (DataLayerInstanceDesc->GetDataLayerType() == EDataLayerType::Runtime)
 					{
@@ -143,4 +144,43 @@ bool FDataLayerUtils::ResolveRuntimeDataLayerInstanceNames(const FWorldPartition
 	return false;
 }
 
+const FDataLayerInstanceDesc* FDataLayerUtils::GetDataLayerInstanceDescFromInstanceName(const TArray<const FWorldDataLayersActorDesc*>& InWorldDataLayersActorDescs, const FName& DataLayerInstanceName)
+{
+	for (const FWorldDataLayersActorDesc* WorldDataLayerActorDesc : InWorldDataLayersActorDescs)
+	{
+		if (const FDataLayerInstanceDesc* DataLayerInstanceDesc = WorldDataLayerActorDesc->GetDataLayerInstanceFromInstanceName(DataLayerInstanceName))
+		{
+			return DataLayerInstanceDesc;
+		}
+	}
+
+	return nullptr;
+}
+
+const FDataLayerInstanceDesc* FDataLayerUtils::GetDataLayerInstanceDescFromAssetPath(const TArray<const FWorldDataLayersActorDesc*>& InWorldDataLayersActorDescs, const FName& DataLayerAssetPath)
+{
+	for (const FWorldDataLayersActorDesc* WorldDataLayerActorDesc : InWorldDataLayersActorDescs)
+	{
+		if (const FDataLayerInstanceDesc* DataLayerInstanceDesc = WorldDataLayerActorDesc->GetDataLayerInstanceFromAssetPath(DataLayerAssetPath))
+		{
+			return DataLayerInstanceDesc;
+		}
+	}
+
+	return nullptr;
+}
+
+bool FDataLayerUtils::FindWorldDataLayerActorDescs(const FActorDescViewMap& ActorDescViewMap, TArray<const FWorldDataLayersActorDesc*>& OutWorldDataLayersActorDescs)
+{
+	TArray<const FWorldPartitionActorDescView*> WorldDataLayerViews = ActorDescViewMap.FindByExactNativeClass<AWorldDataLayers>();
+	Algo::Transform(WorldDataLayerViews, OutWorldDataLayersActorDescs, [](const FWorldPartitionActorDescView* WorldDataLayersActorDescView) { return  (FWorldDataLayersActorDesc*)WorldDataLayersActorDescView->GetActorDesc(); });
+	return !OutWorldDataLayersActorDescs.IsEmpty();
+}
+
+bool FDataLayerUtils::AreWorldDataLayersActorDescsSane(const TArray<const FWorldDataLayersActorDesc*>& InWorldDataLayersActorDescs)
+{
+	// Deprecation handling: Case where pre 5.1 map, without WorldDataLayersActorDescs, are resolving their data layer outside of their world (changelist validation). They need proper Actor descs to resolve.
+	return Algo::AllOf(InWorldDataLayersActorDescs, [](const FWorldDataLayersActorDesc* WorldDataLayerActorDescs) { return WorldDataLayerActorDescs->IsValid(); });
+}
+
 #endif
\ No newline at end of file
diff --git a/Engine/Source/Runtime/Engine/Private/WorldPartition/DataLayer/WorldDataLayers.cpp b/Engine/Source/Runtime/Engine/Private/WorldPartition/DataLayer/WorldDataLayers.cpp
--- a/Engine/Source/Runtime/Engine/Private/WorldPartition/DataLayer/WorldDataLayers.cpp
+++ b/Engine/Source/Runtime/Engine/Private/WorldPartition/DataLayer/WorldDataLayers.cpp
@@ -363,7 +363,7 @@ void AWorldDataLayers::DumpDataLayerRecursively(const UDataLayerInstance* DataLa
 void AWorldDataLayers::DumpDataLayers(FOutputDevice& OutputDevice) const
 {
 	OutputDevice.Logf(TEXT("===================================================="));
-	OutputDevice.Logf(TEXT(" Data Layers for World %s"), *GetWorld()->GetName());
+	OutputDevice.Logf(TEXT(" Data Layers for %s"), *GetName());
 	OutputDevice.Logf(TEXT("===================================================="));
 	OutputDevice.Logf(TEXT(""));
 
@@ -417,12 +417,21 @@ TUniquePtr<class FWorldPartitionActorDesc> AWorldDataLayers::CreateClassActorDes
 
 AWorldDataLayers* AWorldDataLayers::Create(UWorld* World, FName InWorldDataLayerName)
 {
-	check(World);
-	
+	FActorSpawnParameters SpawnParameters;
+	SpawnParameters.Name =  InWorldDataLayerName.IsNone() ? AWorldDataLayers::StaticClass()->GetFName() : InWorldDataLayerName;
+	SpawnParameters.OverrideLevel = World->PersistentLevel;
+	return Create(SpawnParameters);
+}
+
+AWorldDataLayers* AWorldDataLayers::Create(const FActorSpawnParameters& SpawnParameters)
+{
+	check(SpawnParameters.OverrideLevel != nullptr && SpawnParameters.OverrideLevel->IsPersistentLevel());
+	check(SpawnParameters.Name != NAME_None);
+	check(SpawnParameters.NameMode == FActorSpawnParameters::ESpawnActorNameMode::Required_Fatal);
+
 	AWorldDataLayers* WorldDataLayers = nullptr;
-	
-	FName WorldDataLayersName = InWorldDataLayerName.IsNone() ? AWorldDataLayers::StaticClass()->GetFName() : InWorldDataLayerName;
-	if (UObject* ExistingObject = StaticFindObject(nullptr, World->PersistentLevel, *WorldDataLayersName.ToString()))
+
+	if (UObject* ExistingObject = StaticFindObject(nullptr, SpawnParameters.OverrideLevel, *SpawnParameters.Name.ToString()))
 	{
 		WorldDataLayers = CastChecked<AWorldDataLayers>(ExistingObject);
 		if (!IsValidChecked(WorldDataLayers))
@@ -432,31 +441,22 @@ AWorldDataLayers* AWorldDataLayers::Create(UWorld* World, FName InWorldDataLayer
 			WorldDataLayers = nullptr;
 		}
 	}
-	
+
 	if (!WorldDataLayers)
 	{
-		FActorSpawnParameters SpawnParams;
-		SpawnParams.OverrideLevel = World->PersistentLevel;
-		SpawnParams.Name = WorldDataLayersName;
-		SpawnParams.NameMode = FActorSpawnParameters::ESpawnActorNameMode::Required_Fatal;
-		WorldDataLayers = World->SpawnActor<AWorldDataLayers>(AWorldDataLayers::StaticClass(), SpawnParams);
+		UWorld* World = SpawnParameters.OverrideLevel->GetWorld();
+		WorldDataLayers = World->SpawnActor<AWorldDataLayers>(AWorldDataLayers::StaticClass(), SpawnParameters);
 	}
 	else
 	{
-		UE_LOG(LogWorldPartition, Error, TEXT("Failed to create WorldDataLayers Actor. There is already a WorldDataLayer Actor named \"%s\" "), *WorldDataLayersName.ToString())
+		UE_LOG(LogWorldPartition, Error, TEXT("Failed to create WorldDataLayers Actor. There is already a WorldDataLayer Actor named \"%s\" "), *SpawnParameters.Name.ToString())
 	}
-	
+
 	check(WorldDataLayers);
 
 	return WorldDataLayers;
 }
 
-bool AWorldDataLayers::IsMainWorldDataLayers() const
-{
-	check(GetWorld());
-	return GetWorld()->GetWorldDataLayers() == this;
-}
-
 TArray<FName> AWorldDataLayers::GetDataLayerInstanceNames(const TArray<TObjectPtr<const UDataLayerAsset>>& InDataLayersAssets) const
 {
 	TArray<FName> OutDataLayerNames;
@@ -486,9 +486,9 @@ TArray<const UDataLayerInstance*> AWorldDataLayers::GetDataLayerInstances(const
 	return OutDataLayers;
 }
 
-bool AWorldDataLayers::RemoveDataLayers(const TArray<UDataLayerInstance*>& InDataLayerInstances)
+int32 AWorldDataLayers::RemoveDataLayers(const TArray<UDataLayerInstance*>& InDataLayerInstances)
 {
-	bool bIsModified = false;
+	int32 RemovedCount = 0;
 
 	for (UDataLayerInstance* DataLayerInstance : InDataLayerInstances)
 	{
@@ -501,16 +501,16 @@ bool AWorldDataLayers::RemoveDataLayers(const TArray<UDataLayerInstance*>& InDat
 			{
 				DeprecatedDataLayerNameToDataLayerInstance.Remove(DataLayerInstance->GetDataLayerFName());
 			}
-			bIsModified = true;
+			RemovedCount++;
 		}
 	}
 
-	if (bIsModified)
+	if (RemovedCount > 0)
 	{
 		UpdateContainsDeprecatedDataLayers();
 	}
 
-	return bIsModified;
+	return RemovedCount;
 }
 
 bool AWorldDataLayers::RemoveDataLayer(const UDataLayerInstance* InDataLayerInstance)
@@ -736,7 +736,12 @@ void AWorldDataLayers::PostLoad()
 
 	GetLevel()->ConditionalPostLoad();
 
-	GetTypedOuter<UWorld>()->SetWorldDataLayers(this);
+	// Patch WorldDataLayer in UWorld.
+	// Only the "main" world data Layer is named AWorldDataLayers::StaticClass()->GetFName() for a given world.
+	if ((GetTypedOuter<UWorld>()->GetWorldDataLayers() == nullptr) && (GetFName() == GetWorldPartionWorldDataLayersName()))
+	{
+		GetTypedOuter<UWorld>()->SetWorldDataLayers(this);
+	}
 
 #if WITH_EDITOR
 	ConvertDataLayerToInstancces();
@@ -803,6 +808,37 @@ void AWorldDataLayers::PostLoad()
 #endif
 
 	InitializeDataLayerRuntimeStates();
+
+	if (UDataLayerSubsystem* DataLayerSubsystem = UWorld::GetSubsystem<UDataLayerSubsystem>(GetWorld()))
+	{
+		DataLayerSubsystem->RegisterWorldDataLayer(this);
+	}
+}
+
+void AWorldDataLayers::PreRegisterAllComponents()
+{
+	Super::PreRegisterAllComponents();
+
+	if (UDataLayerSubsystem* DataLayerSubsystem = UWorld::GetSubsystem<UDataLayerSubsystem>(GetWorld()))
+	{
+		DataLayerSubsystem->RegisterWorldDataLayer(this);
+	}
+}
+
+void AWorldDataLayers::PostUnregisterAllComponents()
+{
+	if (UDataLayerSubsystem* DataLayerSubsystem = UWorld::GetSubsystem<UDataLayerSubsystem>(GetWorld()))
+	{
+		DataLayerSubsystem->UnregisterWorldDataLayer(this);
+	}
+
+	Super::PostUnregisterAllComponents();
+}
+
+bool AWorldDataLayers::IsSubWorldDataLayers() const
+{
+	check(GetWorld());
+	return GetTypedOuter<UWorld>() != GetWorld();
 }
 
 #if WITH_EDITOR
diff --git a/Engine/Source/Runtime/Engine/Private/WorldPartition/DataLayer/WorldDataLayersCollection.cpp b/Engine/Source/Runtime/Engine/Private/WorldPartition/DataLayer/WorldDataLayersCollection.cpp
new file mode 100644
--- /dev/null
+++ b/Engine/Source/Runtime/Engine/Private/WorldPartition/DataLayer/WorldDataLayersCollection.cpp
@@ -0,0 +1,98 @@
+// Copyright Epic Games, Inc. All Rights Reserved.
+#include "WorldPartition/DataLayer/WorldDataLayersCollection.h"
+
+bool FWorldDataLayersCollection::RegisterWorldDataLayer(AWorldDataLayers* InWorldDataLayers)
+{
+	// Registration and Unregistration for WorldDataLayers occurs during PreRegisterAllComponents and PostUnregisterAllComponents.
+	// Those 2 functions are not symmetrical, so we need to check contains to avoid regisering twice.
+	if (!WorldDataLayers.Contains(InWorldDataLayers) && InWorldDataLayers != nullptr)
+	{
+		WorldDataLayers.Add(InWorldDataLayers);
+		return true;
+	}
+
+	return false;
+}
+
+bool FWorldDataLayersCollection::UnregisterWorldDataLayer(AWorldDataLayers* InWorldDataLayers)
+{
+	int32 OldCount = WorldDataLayers.Num();
+	int32 NewCount = WorldDataLayers.RemoveSwap(InWorldDataLayers);
+	check(!Contains(InWorldDataLayers));
+	return OldCount > NewCount;
+}
+
+bool FWorldDataLayersCollection::Contains(const AWorldDataLayers* InWorldDataLayers) const
+{
+	return WorldDataLayers.Contains(InWorldDataLayers);
+}
+
+void FWorldDataLayersCollection::ForEachWorldDataLayersBreakable(TFunctionRef<bool(AWorldDataLayers*)> Func, const ULevel* InLevelContext /* = nullptr */) const
+{
+	for (const TWeakObjectPtr<AWorldDataLayers>& WorldDataLayer : WorldDataLayers)
+	{
+		if (ensure(WorldDataLayer.IsValid()))
+		{
+			if (ShoudIterateWorldDataLayer(WorldDataLayer.Get(), InLevelContext))
+			{
+				if (!Func(WorldDataLayer.Get()))
+				{
+					break;
+				}
+			}
+		}
+	}
+}
+
+void FWorldDataLayersCollection::ForEachWorldDataLayersBreakable(TFunctionRef<bool(AWorldDataLayers*)> Func, const ULevel* InLevelContext /* = nullptr */)
+{
+	const_cast<const FWorldDataLayersCollection*>(this)->ForEachWorldDataLayersBreakable(Func, InLevelContext);
+}
+
+void FWorldDataLayersCollection::ForEachWorldDataLayers(TFunctionRef<void(AWorldDataLayers*)> Func, const ULevel* InLevelContext /* = nullptr */) const
+{
+	for (const TWeakObjectPtr<AWorldDataLayers>& WorldDataLayer : WorldDataLayers)
+	{
+		if (ensure(WorldDataLayer.IsValid()))
+		{
+			if (ShoudIterateWorldDataLayer(WorldDataLayer.Get(), InLevelContext))
+			{
+				Func(WorldDataLayer.Get());
+			}
+		}
+	}
+}
+
+void FWorldDataLayersCollection::ForEachWorldDataLayers(TFunctionRef<void(AWorldDataLayers*)> Func, const ULevel* InLevelContext /* = nullptr */)
+{
+	const_cast<const FWorldDataLayersCollection*>(this)->ForEachWorldDataLayers(Func, InLevelContext);
+}
+
+void FWorldDataLayersCollection::ForEachDataLayerBreakable(TFunctionRef<bool(UDataLayerInstance*)> Func, const ULevel* InLevelContext /* = nullptr */) const
+{
+	bool bShouldContinue = false;
+	auto CallAndSetContinueFunc = [Func, &bShouldContinue](UDataLayerInstance* DataLayerInstance)
+	{
+		bShouldContinue = Func(DataLayerInstance);
+		return bShouldContinue;
+	};
+
+	auto BreakableWorldDataLayerFunc = [CallAndSetContinueFunc, &bShouldContinue](const AWorldDataLayers* WorldDataLayer)
+	{
+		WorldDataLayer->ForEachDataLayer(CallAndSetContinueFunc);
+		return bShouldContinue;
+	};
+
+	ForEachWorldDataLayersBreakable(BreakableWorldDataLayerFunc, InLevelContext);
+}
+
+void FWorldDataLayersCollection::ForEachDataLayerBreakable(TFunctionRef<bool(UDataLayerInstance*)> Func, const ULevel* InLevelContext /* = nullptr */)
+{
+	const_cast<const FWorldDataLayersCollection*>(this)->ForEachDataLayerBreakable(Func, InLevelContext);
+}
+
+bool FWorldDataLayersCollection::ShoudIterateWorldDataLayer(const AWorldDataLayers* WorldDataLayer, const ULevel* InLevelContext) const
+{
+	return (InLevelContext == nullptr && !WorldDataLayer->IsSubWorldDataLayers())
+		|| (InLevelContext != nullptr && WorldDataLayer->GetTypedOuter<ULevel>() == InLevelContext);
+}
\ No newline at end of file
diff --git a/Engine/Source/Runtime/Engine/Private/WorldPartition/WorldPartitionActorDesc.cpp b/Engine/Source/Runtime/Engine/Private/WorldPartition/WorldPartitionActorDesc.cpp
--- a/Engine/Source/Runtime/Engine/Private/WorldPartition/WorldPartitionActorDesc.cpp
+++ b/Engine/Source/Runtime/Engine/Private/WorldPartition/WorldPartitionActorDesc.cpp
@@ -94,7 +94,7 @@ void FWorldPartitionActorDesc::Init(const AActor* InActor)
 		DataLayers = bIsUsingDataLayerAsset ? MoveTemp(LocalDataLayerAssetPaths) : MoveTemp(LocalDataLayerInstanceNames);
 
 		// Init DataLayers transient info
-		DataLayerInstanceNames = FDataLayerUtils::ResolvedDataLayerInstanceNames(this, nullptr, InActor->GetWorld());
+		DataLayerInstanceNames = FDataLayerUtils::ResolvedDataLayerInstanceNames(this, TArray<const FWorldDataLayersActorDesc*>(), InActor->GetWorld());
 	}
 
 	Tags = InActor->Tags;
diff --git a/Engine/Source/Runtime/Engine/Private/WorldPartition/WorldPartitionRuntimeCell.cpp b/Engine/Source/Runtime/Engine/Private/WorldPartition/WorldPartitionRuntimeCell.cpp
--- a/Engine/Source/Runtime/Engine/Private/WorldPartition/WorldPartitionRuntimeCell.cpp
+++ b/Engine/Source/Runtime/Engine/Private/WorldPartition/WorldPartitionRuntimeCell.cpp
@@ -138,9 +138,9 @@ FLinearColor UWorldPartitionRuntimeCell::GetDebugStreamingPriorityColor() const
 
 TArray<const UDataLayerInstance*> UWorldPartitionRuntimeCell::GetDataLayerInstances() const
 {
-	if (const AWorldDataLayers* WorldDataLayers = GetWorld()->GetWorldDataLayers())
+	if (const UDataLayerSubsystem* DataLayerSubsystem = GetWorld()->GetSubsystem<UDataLayerSubsystem>())
 	{
-		return WorldDataLayers->GetDataLayerInstances(GetDataLayers());
+		return DataLayerSubsystem->GetDataLayerInstances(GetDataLayers());
 	}
 
 	return TArray<const UDataLayerInstance*>();
@@ -148,9 +148,9 @@ TArray<const UDataLayerInstance*> UWorldPartitionRuntimeCell::GetDataLayerInstan
 
 bool UWorldPartitionRuntimeCell::ContainsDataLayer(const UDataLayerAsset* DataLayerAsset) const
 {
-	if (const AWorldDataLayers* WorldDataLayers = GetWorld()->GetWorldDataLayers())
+	if (const UDataLayerSubsystem* DataLayerSubsystem = GetWorld()->GetSubsystem<UDataLayerSubsystem>())
 	{
-		if (const UDataLayerInstance* DataLayerInstance = WorldDataLayers->GetDataLayerInstance(DataLayerAsset))
+		if (const UDataLayerInstance* DataLayerInstance = DataLayerSubsystem->GetDataLayerInstance(DataLayerAsset))
 		{
 			return ContainsDataLayer(DataLayerInstance);
 		}
diff --git a/Engine/Source/Runtime/Engine/Public/WorldPartition/DataLayer/DataLayerSubsystem.h b/Engine/Source/Runtime/Engine/Public/WorldPartition/DataLayer/DataLayerSubsystem.h
--- a/Engine/Source/Runtime/Engine/Public/WorldPartition/DataLayer/DataLayerSubsystem.h
+++ b/Engine/Source/Runtime/Engine/Public/WorldPartition/DataLayer/DataLayerSubsystem.h
@@ -6,7 +6,7 @@
 #include "WorldPartition/DataLayer/ActorDataLayer.h"
 #include "WorldPartition/DataLayer/DataLayerInstance.h"
 #include "WorldPartition/DataLayer/DataLayerEditorContext.h"
-#include "WorldPartition/DataLayer/WorldDataLayers.h"
+#include "WorldPartition/DataLayer/WorldDataLayersCollection.h"
 #include "DataLayerSubsystem.generated.h"
 
 /**
@@ -18,6 +18,11 @@ class UCanvas;
 
 DECLARE_DYNAMIC_MULTICAST_DELEGATE_TwoParams(FOnDataLayerRuntimeStateChanged, const UDataLayerInstance*, DataLayer, EDataLayerRuntimeState, State);
 
+#if WITH_EDITOR
+DECLARE_MULTICAST_DELEGATE_OneParam(FOnWorldDataLayersPostRegister, AWorldDataLayers*);
+DECLARE_MULTICAST_DELEGATE_OneParam(FOnWorldDataLayersPreUnregister, AWorldDataLayers*);
+#endif
+
 #if WITH_EDITOR
 class ENGINE_API FDataLayersEditorBroadcast
 {
@@ -80,7 +85,7 @@ public:
 	bool CanResolveDataLayers() const;
 
 	bool RemoveDataLayer(const UDataLayerInstance* InDataLayer);
-	bool RemoveDataLayers(const TArray<UDataLayerInstance*>& InDataLayerInstances);
+	int32 RemoveDataLayers(const TArray<UDataLayerInstance*>& InDataLayerInstances);
 
 	void UpdateDataLayerEditorPerProjectUserSettings() const;
 	void GetUserLoadedInEditorStates(TArray<FName>& OutDataLayersLoadedInEditor, TArray<FName>& OutDataLayersNotLoadedInEditor) const;
@@ -111,8 +116,22 @@ public:
 	EDataLayerRuntimeState GetDataLayerEffectiveRuntimeState(const UDataLayerInstance* InDataLayer) const;
 	EDataLayerRuntimeState GetDataLayerEffectiveRuntimeStateByName(const FName& InDataLayerName) const;
 	bool IsAnyDataLayerInEffectiveRuntimeState(const TArray<FName>& InDataLayerNames, EDataLayerRuntimeState InState) const;
-	const TSet<FName>& GetEffectiveActiveDataLayerNames() const;
-	const TSet<FName>& GetEffectiveLoadedDataLayerNames() const;
+	TSet<FName> GetEffectiveActiveDataLayerNames() const;
+	TSet<FName> GetEffectiveLoadedDataLayerNames() const;
+
+#if WITH_EDITOR
+	/** Called once a AWorldDataLayers has been registered. */
+	static FOnWorldDataLayersPostRegister OnWorldDataLayerPostRegister;
+
+	/** Called before a AWorldDataLayers has been unregistered. */
+	static FOnWorldDataLayersPreUnregister OnWorldDataLayerPreUnregister;
+#endif
+
+	void RegisterWorldDataLayer(AWorldDataLayers* WorldDataLayers);
+	void UnregisterWorldDataLayer(AWorldDataLayers* WorldDataLayers);
+
+	void ForEachWorldDataLayer(TFunctionRef<bool(AWorldDataLayers*)> Func);
+	void ForEachWorldDataLayer(TFunctionRef<bool(AWorldDataLayers*)> Func) const;
 
 	void GetDataLayerDebugColors(TMap<FName, FColor>& OutMapping) const;
 	void DrawDataLayersStatus(UCanvas* Canvas, FVector2D& Offset) const;
@@ -153,11 +172,11 @@ public:
 
 	UE_DEPRECATED(5.0, "GetActiveDataLayerNames will be removed.")
 	UFUNCTION(BlueprintCallable, Category = DataLayers, meta = (DeprecatedFunction, DeprecationMessage = "GetActiveDataLayerNames will be removed."))
-	const TSet<FName>& GetActiveDataLayerNames() const { return GetEffectiveActiveDataLayerNames(); }
+	const TSet<FName>& GetActiveDataLayerNames() const { static TSet<FName> EmptySet; return EmptySet; }
 
 	UE_DEPRECATED(5.0, "GetLoadedDataLayerNames will be removed.")
 	UFUNCTION(BlueprintCallable, Category = DataLayers, meta = (DeprecatedFunction, DeprecationMessage = "GetLoadedDataLayerNames will be removed."))
-	const TSet<FName>& GetLoadedDataLayerNames() const { return GetEffectiveLoadedDataLayerNames(); }
+	const TSet<FName>& GetLoadedDataLayerNames() const { static TSet<FName> EmptySet; return EmptySet; }
 
 	UE_DEPRECATED(5.1, "Use GetDataLayerFromAsset() instead.")
 	UFUNCTION(BlueprintCallable, Category = DataLayers)
@@ -199,6 +218,7 @@ public:
 
 private:
 #if WITH_EDITOR
+	void RegisterEarlyLoadedWorldDataLayers();
 	void OnActorDescContainerInitialized(UActorDescContainer* InActorDescContainer);
 
 	mutable int32 DataLayerActorEditorContextID;
@@ -212,48 +232,24 @@ private:
 
 	/** Console command used to set Runtime DataLayer state*/
 	static class FAutoConsoleCommand SetDataLayerRuntimeStateCommand;
+
+	FWorldDataLayersCollection WorldDataLayerCollection;
 };
 
 template<class T>
 UDataLayerInstance* UDataLayerSubsystem::GetDataLayerInstance(const T& InDataLayerIdentifier, const ULevel* InLevelContext /* = nullptr */) const
 {
-	PRAGMA_DISABLE_DEPRECATION_WARNINGS
-	const AWorldDataLayers* WorldDataLayers = InLevelContext ? InLevelContext->GetWorldDataLayers() : GetWorld()->GetWorldDataLayers();
-	const UDataLayerInstance * DataLayerInstance = WorldDataLayers ? WorldDataLayers->GetDataLayerInstance(InDataLayerIdentifier) : nullptr;
-	PRAGMA_ENABLE_DEPRECATION_WARNINGS
-	return const_cast<UDataLayerInstance*>(DataLayerInstance);
+	return WorldDataLayerCollection.GetDataLayerInstance(InDataLayerIdentifier, InLevelContext);
 }
 
 template<class T>
 TArray<FName> UDataLayerSubsystem::GetDataLayerInstanceNames(const TArray<T>& InDataLayerIdentifiers, const ULevel* InLevelContext /* = nullptr */) const
 {
-	TArray<FName> DataLayerInstanceNames;
-	DataLayerInstanceNames.Reserve(InDataLayerIdentifiers.Num());
-
-	PRAGMA_DISABLE_DEPRECATION_WARNINGS
-	// Non-partitioned worlds don't have an AWorldDataLayers.
-	// This function can be called by a partitioned sub-level that contains Data Layers.
-	if (const AWorldDataLayers* WorldDataLayers = InLevelContext ? InLevelContext->GetWorldDataLayers() : GetWorld()->GetWorldDataLayers())
-	{
-		DataLayerInstanceNames = WorldDataLayers->GetDataLayerInstanceNames(InDataLayerIdentifiers);
-	}
-	PRAGMA_ENABLE_DEPRECATION_WARNINGS
-	return DataLayerInstanceNames;
+	return WorldDataLayerCollection.GetDataLayerInstanceNames(InDataLayerIdentifiers, InLevelContext);
 }
 
 template<class T>
 TArray<const UDataLayerInstance*> UDataLayerSubsystem::GetDataLayerInstances(const TArray<T>& InDataLayerIdentifiers, const ULevel* InLevelContext /* = nullptr */) const
 {
-	TArray<const UDataLayerInstance*> DataLayerInstances;
-	DataLayerInstances.Reserve(InDataLayerIdentifiers.Num());
-
-	PRAGMA_DISABLE_DEPRECATION_WARNINGS
-	// Non-partitioned worlds don't have an AWorldDataLayers.
-	// This function can be called by a partitioned sub-level that contains Data Layers.
-	if (const AWorldDataLayers* WorldDataLayers = InLevelContext ? InLevelContext->GetWorldDataLayers() : GetWorld()->GetWorldDataLayers())
-	{
-		DataLayerInstances = WorldDataLayers->GetDataLayerInstances(InDataLayerIdentifiers);
-	}
-	PRAGMA_ENABLE_DEPRECATION_WARNINGS
-	return DataLayerInstances;
+	return WorldDataLayerCollection.GetDataLayerInstances(InDataLayerIdentifiers, InLevelContext);
 }
\ No newline at end of file
diff --git a/Engine/Source/Runtime/Engine/Public/WorldPartition/DataLayer/DataLayerUtils.h b/Engine/Source/Runtime/Engine/Public/WorldPartition/DataLayer/DataLayerUtils.h
--- a/Engine/Source/Runtime/Engine/Public/WorldPartition/DataLayer/DataLayerUtils.h
+++ b/Engine/Source/Runtime/Engine/Public/WorldPartition/DataLayer/DataLayerUtils.h
@@ -10,6 +10,7 @@
 #include "WorldPartition/WorldPartitionActorDescView.h"
 #include "WorldPartition/WorldPartitionStreamingGeneration.h"
 
+class FDataLayerInstanceDesc;
 class FWorldDataLayersActorDesc;
 class FWorldPartitionActorDesc;
 class UActorDescContainer;
@@ -25,9 +26,17 @@ public:
 		return IconNameByType[static_cast<uint32>(DataLayerType)];
 	}
 
-	static TArray<FName> ResolvedDataLayerInstanceNames(const FWorldPartitionActorDesc* InActorDesc, const FWorldDataLayersActorDesc* InWorldDataLayersActorDesc = nullptr, UWorld* InWorld = nullptr, bool* bOutIsResultValid = nullptr);
+	static TArray<FName> ResolvedDataLayerInstanceNames(const FWorldPartitionActorDesc* InActorDesc, const TArray<const FWorldDataLayersActorDesc*>& InWorldDataLayersActorDescs = TArray<const FWorldDataLayersActorDesc*>(), UWorld* InWorld = nullptr, bool* bOutIsResultValid = nullptr);
 	
 	static bool ResolveRuntimeDataLayerInstanceNames(const FWorldPartitionActorDescView& InActorDescView, const FActorDescViewMap& ActorDescViewMap, TArray<FName>& OutRuntimeDataLayerInstanceNames);
+
+	static const FDataLayerInstanceDesc* GetDataLayerInstanceDescFromInstanceName(const TArray<const FWorldDataLayersActorDesc*>& InWorldDataLayersActorDescs, const FName& DataLayerInstanceName);
+
+	static const FDataLayerInstanceDesc* GetDataLayerInstanceDescFromAssetPath(const TArray<const FWorldDataLayersActorDesc*>& InWorldDataLayersActorDescs, const FName& DataLayerAssetPath);
+
+	static bool FindWorldDataLayerActorDescs(const FActorDescViewMap& ActorDescViewMap, TArray<const FWorldDataLayersActorDesc*>& OutWorldDataLayersActorDescs);
+
+	static bool AreWorldDataLayersActorDescsSane(const TArray<const FWorldDataLayersActorDesc*>& InWorldDataLayersActorDescs);
 #endif
 
 #if DATALAYER_TO_INSTANCE_RUNTIME_CONVERSION_ENABLED
diff --git a/Engine/Source/Runtime/Engine/Public/WorldPartition/DataLayer/WorldDataLayers.h b/Engine/Source/Runtime/Engine/Public/WorldPartition/DataLayer/WorldDataLayers.h
--- a/Engine/Source/Runtime/Engine/Public/WorldPartition/DataLayer/WorldDataLayers.h
+++ b/Engine/Source/Runtime/Engine/Public/WorldPartition/DataLayer/WorldDataLayers.h
@@ -34,7 +34,7 @@ struct FActorPlacementDataLayers
 };
 
 /**
- * Actor containing all data layers for a world
+ * Actor containing data layers instances within a world.
  */
 UCLASS(hidecategories = (Actor, HLOD, Cooking, Transform, Advanced, Display, Events, Object, Physics, Attachment, Info, Input, Blueprint, Layers, Tags, Replication), notplaceable)
 class ENGINE_API AWorldDataLayers : public AInfo
@@ -45,6 +45,8 @@ class ENGINE_API AWorldDataLayers : public AInfo
 
 public:
 	virtual void PostLoad() override;
+	virtual void PreRegisterAllComponents() override;
+	virtual void PostUnregisterAllComponents() override;
 	virtual void RewindForReplay() override;
 	virtual void BeginPlay() override;
 
@@ -56,6 +58,7 @@ public:
 	virtual TUniquePtr<class FWorldPartitionActorDesc> CreateClassActorDesc() const override;
 
 	static AWorldDataLayers* Create(UWorld* World, FName InWorldDataLayerName = NAME_None);
+	static AWorldDataLayers* Create(const FActorSpawnParameters& SpawnParameters);
 
 	bool IsEmpty() const { return DataLayerInstances.IsEmpty(); }
 	bool HasDeprecatedDataLayers() const { return bHasDeprecatedDataLayers; }
@@ -64,7 +67,7 @@ public:
 	DataLayerInstanceType* CreateDataLayer(CreationsArgs... InCreationArgs);
 
 	bool RemoveDataLayer(const UDataLayerInstance* InDataLayer);
-	bool RemoveDataLayers(const TArray<UDataLayerInstance*>& InDataLayerInstances);
+	int32 RemoveDataLayers(const TArray<UDataLayerInstance*>& InDataLayerInstances);
 	void SetAllowRuntimeDataLayerEditing(bool bInAllowRuntimeDataLayerEditing);
 	bool GetAllowRuntimeDataLayerEditing() const { return bAllowRuntimeDataLayerEditing; }
 
@@ -83,9 +86,10 @@ public:
 	// Allows overriding of DataLayers with PlayFromHere
 	template<class T>
 	void OverwriteDataLayerRuntimeStates(const TArray<T>* InActiveDataLayers, const TArray<T>* InLoadedDataLayers );
-
-	bool IsMainWorldDataLayers() const;
 #endif
+
+	bool IsSubWorldDataLayers() const;
+	static FName GetWorldPartionWorldDataLayersName() { return FName(TEXT("WorldDataLayers")); } // reserved for ULevel::WorldDataLayers
 	
 	void DumpDataLayers(FOutputDevice& OutputDevice) const;
 	bool ContainsDataLayer(const UDataLayerInstance* InDataLayer) const;
diff --git a/Engine/Source/Runtime/Engine/Public/WorldPartition/DataLayer/WorldDataLayersCollection.h b/Engine/Source/Runtime/Engine/Public/WorldPartition/DataLayer/WorldDataLayersCollection.h
new file mode 100644
--- /dev/null
+++ b/Engine/Source/Runtime/Engine/Public/WorldPartition/DataLayer/WorldDataLayersCollection.h
@@ -0,0 +1,99 @@
+// Copyright Epic Games, Inc. All Rights Reserved.
+
+#pragma once
+
+#include "CoreMinimal.h"
+#include "Containers/Array.h"
+
+#include "UObject/ObjectMacros.h"
+#include "UObject/UObjectGlobals.h"
+#include "UObject/Object.h"
+#include "WorldPartition/DataLayer/WorldDataLayers.h"
+
+class ENGINE_API FWorldDataLayersCollection
+{
+public:
+	bool RegisterWorldDataLayer(AWorldDataLayers* InWorldDataLayers);
+	bool UnregisterWorldDataLayer(AWorldDataLayers* InWorldDataLayers);
+	bool Contains(const AWorldDataLayers* InWorldDataLayers) const;
+
+	void ForEachWorldDataLayersBreakable(TFunctionRef<bool(AWorldDataLayers*)> Func, const ULevel* InLevelContext = nullptr);
+	void ForEachWorldDataLayersBreakable(TFunctionRef<bool(AWorldDataLayers*)> Func, const ULevel* InLevelContext = nullptr) const;
+
+	void ForEachWorldDataLayers(TFunctionRef<void(AWorldDataLayers*)> Func, const ULevel* InLevelContext = nullptr);
+	void ForEachWorldDataLayers(TFunctionRef<void(AWorldDataLayers*)> Func, const ULevel* InLevelContext = nullptr) const;
+
+	void ForEachDataLayerBreakable(TFunctionRef<bool(UDataLayerInstance*)> Func, const ULevel* InLevelContext = nullptr);
+	void ForEachDataLayerBreakable(TFunctionRef<bool(UDataLayerInstance*)> Func, const ULevel* InLevelContext = nullptr) const;
+
+	bool IsEmpty() const { return WorldDataLayers.IsEmpty(); }
+
+	void Empty() { return WorldDataLayers.Empty(); }
+
+	template<class T>
+	UDataLayerInstance* GetDataLayerInstance(const T& InDataLayerIdentifier, const ULevel* InLevelContext = nullptr) const;
+
+	template<class T>
+	TArray<FName> GetDataLayerInstanceNames(const TArray<T>& InDataLayerIdentifiers, const ULevel* InLevelContext = nullptr) const;
+
+	template<class T>
+	TArray<const UDataLayerInstance*> GetDataLayerInstances(const TArray<T>& InDataLayerIdentifiers, const ULevel* InLevelContext = nullptr) const;
+
+private:
+	bool ShoudIterateWorldDataLayer(const AWorldDataLayers* WorldDataLayers, const ULevel* InLevelContext) const;
+
+	TArray<TWeakObjectPtr<AWorldDataLayers>> WorldDataLayers;
+};
+
+template<class T>
+UDataLayerInstance* FWorldDataLayersCollection::GetDataLayerInstance(const T& InDataLayerIdentifier, const ULevel* InLevelContext /* = nullptr */) const
+{
+	const UDataLayerInstance* DataLayerInstance = nullptr;
+
+	PRAGMA_DISABLE_DEPRECATION_WARNINGS
+	ForEachWorldDataLayersBreakable([&InDataLayerIdentifier, &DataLayerInstance](AWorldDataLayers* InWorldDataLayers)
+	{
+		DataLayerInstance = InWorldDataLayers->GetDataLayerInstance(InDataLayerIdentifier);
+		return DataLayerInstance == nullptr;
+
+	}, InLevelContext);
+	PRAGMA_ENABLE_DEPRECATION_WARNINGS
+
+	return const_cast<UDataLayerInstance*>(DataLayerInstance);
+}
+
+template<class T>
+TArray<FName> FWorldDataLayersCollection::GetDataLayerInstanceNames(const TArray<T>& InDataLayerIdentifiers, const ULevel* InLevelContext /* = nullptr */) const
+{
+	TArray<FName> DataLayerInstanceNames;
+	DataLayerInstanceNames.Reserve(InDataLayerIdentifiers.Num());
+
+	PRAGMA_DISABLE_DEPRECATION_WARNINGS
+	ForEachWorldDataLayersBreakable([&InDataLayerIdentifiers, &DataLayerInstanceNames](AWorldDataLayers* InWorldDataLayers)
+	{
+		DataLayerInstanceNames.Append(InWorldDataLayers->GetDataLayerInstanceNames(InDataLayerIdentifiers));
+		return DataLayerInstanceNames.Num() != InDataLayerIdentifiers.Num();
+
+	}, InLevelContext);
+	PRAGMA_ENABLE_DEPRECATION_WARNINGS
+	
+	return DataLayerInstanceNames;
+}
+
+template<class T>
+TArray<const UDataLayerInstance*> FWorldDataLayersCollection::GetDataLayerInstances(const TArray<T>& InDataLayerIdentifiers, const ULevel* InLevelContext /* = nullptr */) const
+{
+	TArray<const UDataLayerInstance*> DataLayerInstances;
+	DataLayerInstances.Reserve(InDataLayerIdentifiers.Num());
+
+	PRAGMA_DISABLE_DEPRECATION_WARNINGS
+	ForEachWorldDataLayersBreakable([&InDataLayerIdentifiers, &DataLayerInstances](AWorldDataLayers* InWorldDataLayers)
+	{
+		DataLayerInstances.Append(InWorldDataLayers->GetDataLayerInstances(InDataLayerIdentifiers));
+		return DataLayerInstances.Num() != InDataLayerIdentifiers.Num();
+
+	}, InLevelContext);
+	PRAGMA_ENABLE_DEPRECATION_WARNINGS
+	
+	return DataLayerInstances;
+}
\ No newline at end of file
