VLogger fix for heap-use-after-free (ASAN)

diff --git a/Engine/Source/Runtime/Engine/Private/VisualLogger/VisualLogger.cpp b/Engine/Source/Runtime/Engine/Private/VisualLogger/VisualLogger.cpp
--- a/Engine/Source/Runtime/Engine/Private/VisualLogger/VisualLogger.cpp
+++ b/Engine/Source/Runtime/Engine/Private/VisualLogger/VisualLogger.cpp
@@ -726,11 +726,19 @@ UObject* FVisualLogger::RedirectInternal(const UObject* FromObject, const UObjec
 		const UObject* Owner = NewDirectOwner;
 		while (Owner != nullptr)
 		{
-			TArray<TWeakObjectPtr<const UObject>>& NewHighestOwnerChildren = OwnerToChildrenMap.FindOrAdd(Owner);
-			NewHighestOwnerChildren.Add(FromWeakPtr);
+			TArray<TWeakObjectPtr<const UObject>>* OwnerChildren = OwnerToChildrenMap.Find(Owner);
+			if (OwnerChildren == nullptr)
+			{
+				// Entry not found, add it and refresh current object children since it might have been reallocated
+				OwnerChildren = &OwnerToChildrenMap.Emplace(Owner);
+				ObjectChildren = ObjectChildren ? OwnerToChildrenMap.Find(FromObject) : nullptr;
+			}
+
+			check(OwnerChildren);
+			OwnerChildren->Add(FromWeakPtr);
 			if (ObjectChildren != nullptr)
 			{
-				NewHighestOwnerChildren.Append(*ObjectChildren);
+				OwnerChildren->Append(*ObjectChildren);
 			}
 
 			const TWeakObjectPtr<const UObject>* WeakOwner = ChildToOwnerMap.Find(Owner);
