Fixed AActor::Rename to avoid dirtying actor's level package when actor uses external package and Rename doesn't change actor's outer object.

diff --git a/Engine/Source/Runtime/Engine/Private/Actor.cpp b/Engine/Source/Runtime/Engine/Private/Actor.cpp
--- a/Engine/Source/Runtime/Engine/Private/Actor.cpp
+++ b/Engine/Source/Runtime/Engine/Private/Actor.cpp
@@ -1297,11 +1297,19 @@ bool AActor::Rename( const TCHAR* InName, UObject* NewOuter, ERenameFlags Flags
 			if (bExternalActor)
 			{
 				bShouldSetPackageExternal = bChangingOuters || IsMainPackageActor();
-				// If we are changing outers always change the external package because the outer chain as an impact on the filename
+				// If we are changing outers always change the external package because the outer chain has an impact on the filename
 				// but if we are not changing outers only change the external flag if the actor is the main actor in a package. (this is to avoid Child Actors creating packages)
 				if (bShouldSetPackageExternal)
 				{
+					bool bLevelPackageWasDirty = MyLevel->GetPackage()->IsDirty();
 					SetPackageExternal(false, MyLevel->IsUsingExternalActors());
+
+					// If we are not changing the outer, SetPackageExternal will still mark dirty the new - temporary - outer (which is the level)
+					// Restore its backed-up dirty state
+					if (!bChangingOuters && !bLevelPackageWasDirty)
+					{
+						MyLevel->GetPackage()->SetDirtyFlag(false);
+					}
 				}
 			}
 			if (bChangingOuters && MyLevel->IsUsingActorFolders())
@@ -1343,7 +1351,9 @@ bool AActor::Rename( const TCHAR* InName, UObject* NewOuter, ERenameFlags Flags
 	{
 		if (ULevel* MyLevel = GetLevel())
 		{
-			SetPackageExternal(true, MyLevel->IsUsingExternalActors());
+			// Don't dirty the current package (which is the level package) if the outer doesn't change
+			check(bChangingOuters || (MyLevel->GetPackage() == GetPackage()));
+			SetPackageExternal(true, bChangingOuters && MyLevel->IsUsingExternalActors());
 		}
 	}
 #endif
