UPlayerInput:
- Fixed issue where analog values weren't cleared when flushing pressed keys.

diff --git a/Engine/Source/Runtime/Engine/Private/UserInterface/PlayerInput.cpp b/Engine/Source/Runtime/Engine/Private/UserInterface/PlayerInput.cpp
--- a/Engine/Source/Runtime/Engine/Private/UserInterface/PlayerInput.cpp
+++ b/Engine/Source/Runtime/Engine/Private/UserInterface/PlayerInput.cpp
@@ -116,7 +116,9 @@ void UPlayerInput::FlushPressedKeys()
 			for (int32 KeyIndex = 0; KeyIndex < PressedKeys.Num(); KeyIndex++)
 			{
 				FKey& Key = PressedKeys[KeyIndex];
-				InputKey(FInputKeyParams(Key, IE_Released, 0.0));
+				FInputKeyParams Params(Key, IE_Released, 0.0);
+				Params.NumSamples = 1;
+				InputKey(Params);
 			}
 		}
 	}
