Fix a bug where NetInitial properties may not be sent for newly relevant objects when Push Model is used.

diff --git a/Engine/Source/Runtime/Engine/Private/RepLayout.cpp b/Engine/Source/Runtime/Engine/Private/RepLayout.cpp
--- a/Engine/Source/Runtime/Engine/Private/RepLayout.cpp
+++ b/Engine/Source/Runtime/Engine/Private/RepLayout.cpp
@@ -1476,7 +1476,11 @@ static void CompareParentProperties(
 		{
 			for (int32 ParentIndex = 0; ParentIndex < SharedParams.Parents.Num(); ++ParentIndex)
 			{
-				const bool bIsPropertyDirty = UE_RepLayout_Private::IsPropertyDirty(ParentIndex, bRecentlyCollectedGarbage, SharedParams, StackParams);
+				const bool bRecompareInitialProperties = SharedParams.bIsInitial && SharedParams.Parents[ParentIndex].Condition == COND_InitialOnly;
+
+				const bool bIsPropertyDirty = bRecompareInitialProperties ||
+												UE_RepLayout_Private::IsPropertyDirty(ParentIndex, bRecentlyCollectedGarbage, SharedParams, StackParams);
+				
 				const bool bDidPropertyChange = UE_RepLayout_Private::CompareParentPropertyHelper(ParentIndex, SharedParams, StackParams);
 
 				ensureAlwaysMsgf(!bDidPropertyChange || bIsPropertyDirty, TEXT("Push Model Property changed value, but was not marked dirty! Property=%s"), *SharedParams.Parents[ParentIndex].Property->GetPathName());
@@ -1484,6 +1488,43 @@ static void CompareParentProperties(
 		}
 #endif // WITH_PUSH_VALIDATION_SUPPORT
 
+		else if (UNLIKELY(SharedParams.bIsInitial && EnumHasAnyFlags(SharedParams.Flags, ERepLayoutFlags::HasInitialOnlyProperties)))
+		{
+			/*
+				Most replication conditions don't actually effect whether or not we compare properties,
+				because we share comparisons across connections, and typically only want to do one comparison
+				per frame for all connections, when possible.
+				
+				This means that we will end up performing the comparison, ignoring replication conditions,
+				even if we might not have to send the property to any connections on a given frame frame.
+				
+				COND_InitialOnly is special, though. In theory, COND_InitialOnly properties should only ever be
+				compared a single time, the first time an object is ever replicated, so it's a waste to compare
+				them every frame.
+
+				However, in practice, there are a number of cases where an object can be replicated multiple
+				times, and still have its bNetInitial flag set, because of how we handle multiple connections.
+				E.G., it may be the first time an Object is replicating to a Specific Connection, and we would
+				consider that an Initial Replication, even though the object may have replicated previously.
+
+				If the property's value has changed, and the object was replicated again BEFORE a new
+				Net Initial replication occurs, AND the object is a Push Model Property, then we will end up
+				clearing the fact that the property was ever changed, and the new value will never be sent
+				to any connections.
+
+				To get around that, if we detect that this is an Initial Replication, and we have Initial Only
+				properties, we will consider them dirty.
+			*/
+			for (int32 ParentIndex = 0; ParentIndex < SharedParams.Parents.Num(); ++ParentIndex)
+			{
+				if (SharedParams.Parents[ParentIndex].Condition == COND_InitialOnly ||
+					UE_RepLayout_Private::IsPropertyDirty(ParentIndex, bRecentlyCollectedGarbage, SharedParams, StackParams))
+				{
+					UE_RepLayout_Private::CompareParentPropertyHelper(ParentIndex, SharedParams, StackParams);
+				}
+			}
+		}
+
 		// If we have full push model property support, then we only need to check properties that are actually dirty.
 		else if (EnumHasAnyFlags(SharedParams.Flags, ERepLayoutFlags::FullPushProperties) && !bRecentlyCollectedGarbage)
 		{
@@ -6131,6 +6172,10 @@ void FRepLayout::InitFromClass(
 			{
 				Parents[ParentIndex].Flags &= ~ERepParentFlags::IsConditional;
 			}
+			else if (LifetimeProps[i].Condition == COND_InitialOnly)
+			{
+				Flags |= ERepLayoutFlags::HasInitialOnlyProperties;
+			}
 
 			if (EnumHasAnyFlags(Parents[ParentIndex].Flags, ERepParentFlags::HasNetSerializeProperties | ERepParentFlags::HasObjectProperties))
 			{
diff --git a/Engine/Source/Runtime/Engine/Public/Net/RepLayout.h b/Engine/Source/Runtime/Engine/Public/Net/RepLayout.h
--- a/Engine/Source/Runtime/Engine/Public/Net/RepLayout.h
+++ b/Engine/Source/Runtime/Engine/Public/Net/RepLayout.h
@@ -1042,6 +1042,7 @@ enum class ERepLayoutFlags : uint8
 	HasObjectOrNetSerializeProperties	= (1 << 3),	//! Will be set for any RepLayout that contains Object or Net Serialize property commands.
 	NoReplicatedProperties				= (1 << 4), //! Will be set if the RepLayout has no lifetime properties, or they are all disabled.
 	FullPushProperties					= (1 << 5), //! All properties in this RepLayout use Push Model.
+	HasInitialOnlyProperties			= (1 << 6), //! There is at least 1 Initial Only Lifetime property on this RepLayout.
 };
 ENUM_CLASS_FLAGS(ERepLayoutFlags);
 
