DrawDebug: Call FlushPersistentDebugLines on each client when called on the server

diff --git a/Engine/Source/Runtime/Engine/Private/DrawDebugHelpers.cpp b/Engine/Source/Runtime/Engine/Private/DrawDebugHelpers.cpp
--- a/Engine/Source/Runtime/Engine/Private/DrawDebugHelpers.cpp
+++ b/Engine/Source/Runtime/Engine/Private/DrawDebugHelpers.cpp
@@ -64,9 +64,16 @@ bool CanDrawServerDebugInContext(const FWorldContext& WorldContext)
 
 void FlushPersistentDebugLines( const UWorld* InWorld )
 {
-	if(InWorld && InWorld->PersistentLineBatcher)
+	if (GEngine->GetNetMode(InWorld) != NM_DedicatedServer)
+	{
+		if (InWorld && InWorld->PersistentLineBatcher)
+		{
+			InWorld->PersistentLineBatcher->Flush();
+		}
+	}
+	else
 	{
-		InWorld->PersistentLineBatcher->Flush();
+		UE_DRAW_SERVER_DEBUG_ON_EACH_CLIENT(FlushPersistentDebugLines);
 	}
 }
 
