DrawDebug: Fix compilation on Linux, do not use macro

diff --git a/Engine/Source/Runtime/Engine/Private/DrawDebugHelpers.cpp b/Engine/Source/Runtime/Engine/Private/DrawDebugHelpers.cpp
--- a/Engine/Source/Runtime/Engine/Private/DrawDebugHelpers.cpp
+++ b/Engine/Source/Runtime/Engine/Private/DrawDebugHelpers.cpp
@@ -71,10 +71,21 @@ void FlushPersistentDebugLines( const UWorld* InWorld )
 			InWorld->PersistentLineBatcher->Flush();
 		}
 	}
+#if WITH_EDITOR
 	else
 	{
-		UE_DRAW_SERVER_DEBUG_ON_EACH_CLIENT(FlushPersistentDebugLines);
+		if (GIsEditor)
+		{
+			for (const FWorldContext& WorldContext : GEngine->GetWorldContexts())
+			{ 
+				if (CanDrawServerDebugInContext(WorldContext))
+				{
+					FlushPersistentDebugLines(WorldContext.World());
+				}
+			} 
+		}
 	}
+#endif
 }
 
 ULineBatchComponent* GetDebugLineBatcher( const UWorld* InWorld, bool bPersistentLines, float LifeTime, bool bDepthIsForeground )
