[KillPendingKill] Don't reattach garbage components when applying cached data when rerunning construction scripts. Emit an error here so GC dangling-garbage-refrerence errors have context.

diff --git a/Engine/Source/Runtime/Engine/Private/Components/SceneComponent.cpp b/Engine/Source/Runtime/Engine/Private/Components/SceneComponent.cpp
--- a/Engine/Source/Runtime/Engine/Private/Components/SceneComponent.cpp
+++ b/Engine/Source/Runtime/Engine/Private/Components/SceneComponent.cpp
@@ -2363,9 +2363,17 @@ void FSceneComponentInstanceData::ApplyToComponent(UActorComponent* Component, c
 	for (const auto& ChildComponentPair : AttachedInstanceComponents)
 	{
 		USceneComponent* ChildComponent = ChildComponentPair.Key;
+		if (!IsValid(ChildComponent))
+		{
+			UE_LOG(LogSceneComponent, Error, TEXT("Attempting to reattach destroyed child component '%s' to '%s' while reapplying instance data."),
+				*GetPathNameSafe(ChildComponent), *GetPathNameSafe(SceneComponent));
+			continue;
+		}
+
 		// If the ChildComponent now has a "good" attach parent it was set by the transaction and it means we are undoing/redoing attachment
 		// and so the rebuilt component should not take back attachment ownership
-		if (ChildComponent && !IsValid(ChildComponent->GetAttachParent()))
+		// Make sure also that we don't reattach trashed components
+		if (!IsValid(ChildComponent->GetAttachParent()))
 		{
 			ChildComponent->SetRelativeTransform_Direct(ChildComponentPair.Value);
 			ChildComponent->AttachToComponent(SceneComponent, FAttachmentTransformRules::KeepRelativeTransform);
