PR #9336: Add a default console key for Hungarian keyboard layout (Contributed by [REDACTED_PERSON])

diff --git a/Engine/Source/Runtime/Engine/Private/UserInterface/InputSettings.cpp b/Engine/Source/Runtime/Engine/Private/UserInterface/InputSettings.cpp
--- a/Engine/Source/Runtime/Engine/Private/UserInterface/InputSettings.cpp
+++ b/Engine/Source/Runtime/Engine/Private/UserInterface/InputSettings.cpp
@@ -83,6 +83,7 @@ void UInputSettings::PostInitProperties()
 		switch(PRIMARYLANGID(LOWORD(GetKeyboardLayout(0))))
 		{
 		case LANG_FRENCH:
+		case LANG_HUNGARIAN:
 			DefaultConsoleKey = FInputKeyManager::Get().GetKeyFromCodes(VK_OEM_7, 0);
 			break;
 
