[KillPendingKill] Clean up references between scene components when actors are destroyed or stream out to avoid holding reference to garbage objects.
On destruction, clear attachments between components owned by different actors.
On uninitialize component, clear attachments between components in different outer hierarchies, i.e. different streaming levels.

Note that this changes behavior so that attachments between scene components in different levels will be broken when either level is streamed out (made not visible) even if the level is not unloaded & garbage collected.

diff --git a/Engine/Source/Runtime/Engine/Classes/Components/SceneComponent.h b/Engine/Source/Runtime/Engine/Classes/Components/SceneComponent.h
--- a/Engine/Source/Runtime/Engine/Classes/Components/SceneComponent.h
+++ b/Engine/Source/Runtime/Engine/Classes/Components/SceneComponent.h
@@ -905,6 +905,8 @@ public:
 	virtual void GetLifetimeReplicatedProps(TArray<FLifetimeProperty>& OutLifetimeProps) const override;
 	virtual void OnRegister() override;
 	virtual void OnUnregister() override;
+	virtual void EndPlay(EEndPlayReason::Type Reason) override;
+	virtual void UninitializeComponent() override;
 	virtual bool ShouldCreateRenderState() const override { return true; }
 	virtual void UpdateComponentToWorld(EUpdateTransformFlags UpdateTransformFlags = EUpdateTransformFlags::None, ETeleportType Teleport = ETeleportType::None) override final
 	{
diff --git a/Engine/Source/Runtime/Engine/Private/Components/SceneComponent.cpp b/Engine/Source/Runtime/Engine/Private/Components/SceneComponent.cpp
--- a/Engine/Source/Runtime/Engine/Private/Components/SceneComponent.cpp
+++ b/Engine/Source/Runtime/Engine/Private/Components/SceneComponent.cpp
@@ -750,6 +750,62 @@ void USceneComponent::OnUnregister()
 	Super::OnUnregister();
 }
 
+void USceneComponent::EndPlay(EEndPlayReason::Type Reason)
+{
+	Super::EndPlay(Reason);
+
+	// Detach all parent/child relationships with other actors to allow GC to clean us and our owner up.
+	// End play reasons other than destruction will be handled in UninitializeComponent and detach cross-level references.
+	if (Reason == EEndPlayReason::Destroyed)
+	{
+		if (AttachParent && AttachParent->GetOwner() != GetOwner())
+		{
+			DetachFromComponent(FDetachmentTransformRules::KeepRelativeTransform); // No need to fix up to world transform because we're being destroyed.
+		}
+
+		TInlineComponentArray<USceneComponent*> ChildrenToDetach;
+		for (int32 i = 0; i < AttachChildren.Num(); ++i)
+		{
+			if (USceneComponent* AttachChild = AttachChildren[i].Get(); AttachChild && AttachChild->GetOwner() != GetOwner())
+			{
+				ChildrenToDetach.Add(AttachChild);
+			}
+		}
+
+		for (USceneComponent* Child : ChildrenToDetach)
+		{
+			Child->DetachFromComponent(FDetachmentTransformRules::KeepWorldTransform);
+		}
+	}
+}
+
+void USceneComponent::UninitializeComponent()
+{
+	Super::UninitializeComponent();
+
+	// Detach components which are in different streaming levels so that this level can be properly garbage collected.
+	// Note that we explicitly want to check the outer hierarchy and not the owning package because we want references that participate in GC.
+	UObject* Outermost = GetOutermostObject();
+	if (AttachParent && AttachParent->GetOutermostObject() != Outermost)
+	{
+		DetachFromComponent(FDetachmentTransformRules::KeepWorldTransform);
+	}
+
+	TInlineComponentArray<USceneComponent*> ChildrenToDetach;
+	for (int32 i = 0; i < AttachChildren.Num(); ++i)
+	{
+		if (USceneComponent* AttachChild = AttachChildren[i].Get(); AttachChild && AttachChild->GetOutermostObject() != Outermost)
+		{
+			ChildrenToDetach.Add(AttachChild);
+		}
+	}
+
+	for (USceneComponent* Child : ChildrenToDetach)
+	{
+		Child->DetachFromComponent(FDetachmentTransformRules::KeepWorldTransform);
+	}
+}
+
 void USceneComponent::PropagateTransformUpdate(bool bTransformChanged, EUpdateTransformFlags UpdateTransformFlags, ETeleportType Teleport)
 {
 	//QUICK_SCOPE_CYCLE_COUNTER(STAT_USceneComponent_PropagateTransformUpdate);
