Fix problems with the GetUsedMaterials and GetMaterialsAffectedByTexture methods of FObjectCacheContext:
* when render state changes and material used notifications were sent from SendAllEndOfFrameUpdates, they were added to a stack instead of a queue, so they were applied in reverse order. This meant that the component -> material mapping was first set and then cleared, instead of the other way around, so components which took this path were missing from the cache.
* dynamic material instances were not adding themselves to the textures -> materials map, so GetUsedMaterials was not taking them into account.

The result was that we were not recreating the mesh draw commands for primitives missed due to these reasons, leading to crashes from using stale uniform buffer pointers.

#rnx

diff --git a/Engine/Source/Runtime/Engine/Private/Materials/MaterialInstanceDynamic.cpp b/Engine/Source/Runtime/Engine/Private/Materials/MaterialInstanceDynamic.cpp
--- a/Engine/Source/Runtime/Engine/Private/Materials/MaterialInstanceDynamic.cpp
+++ b/Engine/Source/Runtime/Engine/Private/Materials/MaterialInstanceDynamic.cpp
@@ -49,6 +49,10 @@ void UMaterialInstanceDynamic::UpdateCachedDataDynamic()
 	{
 		Resource->GameThread_UpdateCachedData(*CachedData);
 	}
+
+#if WITH_EDITOR
+	FObjectCacheEventSink::NotifyReferencedTextureChanged_Concurrent(this);
+#endif // WITH_EDITOR
 }
 
 #if WITH_EDITOR
diff --git a/Engine/Source/Runtime/Engine/Private/ObjectCacheContext.cpp b/Engine/Source/Runtime/Engine/Private/ObjectCacheContext.cpp
--- a/Engine/Source/Runtime/Engine/Private/ObjectCacheContext.cpp
+++ b/Engine/Source/Runtime/Engine/Private/ObjectCacheContext.cpp
@@ -716,7 +716,7 @@ struct FObjectCacheEventSinkPrivate
 	static std::atomic<bool> bShouldQueueSinkEvents;
 
 	using TEventAllocator = TLockFreeClassAllocator<FNotifyEvent, PLATFORM_CACHE_LINE_SIZE>;
-	using TEventList = TLockFreePointerListUnordered<FNotifyEvent, PLATFORM_CACHE_LINE_SIZE>;
+	using TEventList = TLockFreePointerListFIFO<FNotifyEvent, PLATFORM_CACHE_LINE_SIZE>;
 	static TEventAllocator& GetAllocator();
 	static TEventList& GetNotifyEvents();
 	static void AddNotifyEvent(FObjectCacheEventSinkPrivate::ECacheEventType EventType, UMaterialInterface* MaterialInterface, const UPrimitiveComponent* PrimitiveComponent, UStaticMeshComponent* StaticMeshComponent, const TArray<UMaterialInterface*>* UsedMaterials);
