Fix Trace.Screenshot not working when called from non-PIE command line.
Move screenshot tracing to an async task.
Add timing scopes to screenshot tracing code.

diff --git a/Engine/Source/Runtime/Engine/Private/ProfilingDebugging/TraceScreenshot.cpp b/Engine/Source/Runtime/Engine/Private/ProfilingDebugging/TraceScreenshot.cpp
--- a/Engine/Source/Runtime/Engine/Private/ProfilingDebugging/TraceScreenshot.cpp
+++ b/Engine/Source/Runtime/Engine/Private/ProfilingDebugging/TraceScreenshot.cpp
@@ -11,6 +11,7 @@
 #include "Misc/DateTime.h"
 #include "Misc/Paths.h"
 #include "ProfilingDebugging/MiscTrace.h"
+#include "Tasks/Task.h"
 #include "Trace/Trace.inl"
 
 static void TraceScreenshotCommandCallback(const TArray<FString>& Args)
@@ -42,7 +43,6 @@ void FTraceScreenshot::RequestScreenshot(FString Name)
 	{
 		Name = FDateTime::Now().ToString(TEXT("Screenshot_%Y%m%d_%H%M%S"));
 	}
-
 	constexpr bool bShowUI = false;
 	constexpr bool bAddUniqueSuffix = false;
 	FScreenshotRequest::RequestScreenshot(Name, bShowUI, bAddUniqueSuffix);
@@ -50,6 +50,7 @@ void FTraceScreenshot::RequestScreenshot(FString Name)
 
 void FTraceScreenshot::TraceScreenshot(int32 InSizeX, int32 InSizeY, const TArray<FColor>& InImageData, const FString& InScreenshotName, int32 DesiredX)
 {
+	TRACE_CPUPROFILER_EVENT_SCOPE(ScreenshotTracing_Prepare);
 	FString ScreenshotName = FPaths::GetBaseFilename(InScreenshotName);
 
 	if (ScreenshotName.IsEmpty())
@@ -57,27 +58,43 @@ void FTraceScreenshot::TraceScreenshot(int32 InSizeX, int32 InSizeY, const TArra
 		ScreenshotName = FDateTime::Now().ToString(TEXT("Screenshot_%Y%m%d_%H%M%S"));
 	}
 
-	uint64 Cycles = FPlatformTime::Cycles64();
 	UE_LOG(LogCore, Display, TEXT("Tracing Screenshot \"%s\" taken with size: %d x %d"), *ScreenshotName, InSizeX, InSizeY);
-
-	TArray64<uint8> CompressedBitmap;
-	if (DesiredX > 0 && InSizeX != DesiredX)
-	{
-		int32 ResizedX = FMath::Min(640, InSizeX);
-		int32 ResizedY = (InSizeY * ResizedX) / InSizeX;
-
-		TArray<FColor> ResizedImage;
-		ResizedImage.SetNum(ResizedX * ResizedY);
-		FImageUtils::ImageResize(InSizeX, InSizeY, InImageData, ResizedX, ResizedY, ResizedImage, false);
-
-		FImageUtils::PNGCompressImageArray(ResizedX, ResizedY, TArrayView64<const FColor>(ResizedImage.GetData(), ResizedImage.Num()), CompressedBitmap);
-		TRACE_SCREENSHOT(*ScreenshotName, Cycles, ResizedX, ResizedY, CompressedBitmap);
-	}
-	else
-	{
-		FImageUtils::PNGCompressImageArray(InSizeX, InSizeY, TArrayView64<const FColor>(InImageData.GetData(), InImageData.Num()), CompressedBitmap);
-		TRACE_SCREENSHOT(*ScreenshotName, Cycles, InSizeX, InSizeY, CompressedBitmap);
-	}
+	
+	uint64 Cycles = FPlatformTime::Cycles64();
+	TArray<FColor> *ImageCopy = new TArray<FColor>(InImageData);
+	
+	UE::Tasks::Launch(UE_SOURCE_LOCATION, 
+		[ImageCopy, Cycles, ScreenshotName, InSizeX, InSizeY, DesiredX]
+		{
+			TRACE_CPUPROFILER_EVENT_SCOPE(ScreenshotTracing_Execute);
+
+			// Set full alpha on the bitmap
+			for (FColor& Pixel : *ImageCopy)
+			{
+				Pixel.A = 255;
+			}
+
+			TArray64<uint8> CompressedBitmap;
+			if (DesiredX > 0 && InSizeX != DesiredX)
+			{
+				int32 ResizedX = FMath::Min(640, InSizeX);
+				int32 ResizedY = (InSizeY * ResizedX) / InSizeX;
+
+				TArray<FColor> ResizedImage;
+				ResizedImage.SetNum(ResizedX * ResizedY);
+				FImageUtils::ImageResize(InSizeX, InSizeY, *ImageCopy, ResizedX, ResizedY, ResizedImage, false);
+
+				FImageUtils::PNGCompressImageArray(ResizedX, ResizedY, TArrayView64<const FColor>(ResizedImage.GetData(), ResizedImage.Num()), CompressedBitmap);
+				TRACE_SCREENSHOT(*ScreenshotName, Cycles, ResizedX, ResizedY, CompressedBitmap);
+			}
+			else
+			{
+				FImageUtils::PNGCompressImageArray(InSizeX, InSizeY, TArrayView64<const FColor>(ImageCopy->GetData(), ImageCopy->Num()), CompressedBitmap);
+				TRACE_SCREENSHOT(*ScreenshotName, Cycles, InSizeX, InSizeY, CompressedBitmap);
+			}
+
+			delete ImageCopy;
+		});
 
 	Reset();
 }
