Fix VolumeTexture no longer being lazy-updated in PostLoad
redundant, gets updated by reimport in EditorFactories
was incorrectly using LightingGuid instead of Source.Id
was not marking package dirty so it never got saved
which made it rebuild every time
so user saw a DDC cache miss on every run

diff --git a/Engine/Source/Editor/UnrealEd/Private/Factories/EditorFactories.cpp b/Engine/Source/Editor/UnrealEd/Private/Factories/EditorFactories.cpp
--- a/Engine/Source/Editor/UnrealEd/Private/Factories/EditorFactories.cpp
+++ b/Engine/Source/Editor/UnrealEd/Private/Factories/EditorFactories.cpp
@@ -4398,9 +4398,9 @@ UObject* UTextureFactory::FactoryCreateBinary
 	// Invalidate any materials using the newly imported texture. (occurs if you import over an existing texture)
 	Texture->PostEditChange();
 
-	// Invalidate any volume texture that was built on this texture.
 	if (Texture2D)
 	{
+		// Invalidate any volume texture that was built on this texture.
 		for (TObjectIterator<UVolumeTexture> It; It; ++It)
 		{
 			UVolumeTexture* VolumeTexture = *It;
@@ -4410,11 +4410,8 @@ UObject* UTextureFactory::FactoryCreateBinary
 				VolumeTexture->UpdateResource();
 			}
 		}
-	}
 
-	// Invalidate any Texture2DArrays that use the updated texture.
-	if (Texture2D) 
-	{
+		// Invalidate any Texture2DArrays that use the updated texture.
 		for (TObjectIterator<UTexture2DArray> It; It; ++It) 
 		{
 			UTexture2DArray* TextureArray = *It;
@@ -4426,6 +4423,28 @@ UObject* UTextureFactory::FactoryCreateBinary
 					{
 						// Update the entire texture array.
 						TextureArray->UpdateSourceFromSourceTextures(false);
+						// UpdateSourceFromSourceTextures calls UpdateResource
+						break;
+					}
+				}
+			}
+		}
+	}
+
+	{
+		// Invalidate any TextureCubeArrays that use the updated texture.
+		for (TObjectIterator<UTextureCubeArray> It; It; ++It) 
+		{
+			UTextureCubeArray* TextureArray = *It;
+			if (TextureArray) 
+			{
+				for (int32 SourceIndex = 0; SourceIndex < TextureArray->SourceTextures.Num(); ++SourceIndex) 
+				{
+					if (TextureArray->SourceTextures[SourceIndex] == Texture) 
+					{
+						// Update the entire texture array.
+						TextureArray->UpdateSourceFromSourceTextures(false);
+						// UpdateSourceFromSourceTextures calls UpdateResource
 						break;
 					}
 				}
diff --git a/Engine/Source/Runtime/Engine/Classes/Engine/Texture2DArray.cpp b/Engine/Source/Runtime/Engine/Classes/Engine/Texture2DArray.cpp
--- a/Engine/Source/Runtime/Engine/Classes/Engine/Texture2DArray.cpp
+++ b/Engine/Source/Runtime/Engine/Classes/Engine/Texture2DArray.cpp
@@ -140,7 +140,6 @@ uint32 UTexture2DArray::CalcTextureMemorySizeEnum(ETextureMipCount Enum) const
 }
 
 #if WITH_EDITOR
-
 ENGINE_API bool UTexture2DArray::CheckArrayTexturesCompatibility()
 {
 	bool bError = false;
@@ -199,6 +198,8 @@ ENGINE_API bool UTexture2DArray::UpdateSourceFromSourceTextures(bool bCreatingNe
 
 	if (SourceTextures.Num() > 0)
 	{
+		Modify();
+
 		FTextureSource& InitialSource = SourceTextures[0]->Source;
 		// Format and format size.
 		ETextureSourceFormat Format = InitialSource.GetFormat();
@@ -254,7 +255,7 @@ ENGINE_API bool UTexture2DArray::UpdateSourceFromSourceTextures(bool bCreatingNe
 
 		ValidateSettingsAfterImportOrEdit();
 		SetLightingGuid();
-		UpdateResource();
+		UpdateResource();		
 	}
 
 	return true;
@@ -262,6 +263,8 @@ ENGINE_API bool UTexture2DArray::UpdateSourceFromSourceTextures(bool bCreatingNe
 
 ENGINE_API void UTexture2DArray::InvadiateTextureSource()
 {
+	Modify();
+
 	if (PlatformData) 
 	{
 		delete PlatformData;
@@ -270,7 +273,7 @@ ENGINE_API void UTexture2DArray::InvadiateTextureSource()
 
 	Source = FTextureSource();
 	Source.SetOwner(this);
-
+	
 	UpdateResource();
 }
 #endif
@@ -302,8 +305,8 @@ void UTexture2DArray::PostLoad()
 {
 #if WITH_EDITOR
 	FinishCachePlatformData();
-
 #endif // #if WITH_EDITOR
+
 	Super::PostLoad();
 };
 
@@ -377,7 +380,7 @@ ENGINE_API void UTexture2DArray::PostEditChangeProperty(FPropertyChangedEvent &
 			}
 		}
 	}
-
+	
 	if (PropertyName == GET_MEMBER_NAME_CHECKED(UTexture2DArray, AddressX)
 		|| PropertyName == GET_MEMBER_NAME_CHECKED(UTexture2DArray, AddressY)
 		|| PropertyName == GET_MEMBER_NAME_CHECKED(UTexture2DArray, AddressZ))
diff --git a/Engine/Source/Runtime/Engine/Classes/Engine/VolumeTexture.h b/Engine/Source/Runtime/Engine/Classes/Engine/VolumeTexture.h
--- a/Engine/Source/Runtime/Engine/Classes/Engine/VolumeTexture.h
+++ b/Engine/Source/Runtime/Engine/Classes/Engine/VolumeTexture.h
@@ -25,10 +25,11 @@ public:
 #if WITH_EDITORONLY_DATA
 	/** A (optional) reference texture from which the volume texture was built */
 	UPROPERTY(EditAnywhere, Category=Source2D, meta=(DisplayName="Source Texture"))
-	TObjectPtr<UTexture2D> Source2DTexture;
-	/** The lighting Guid of the source 2D texture, used to trigger rebuild when the source changes. */
+	TObjectPtr<UTexture2D> Source2DTexture;	
+
 	UPROPERTY()
-	FGuid SourceLightingGuid;
+	FGuid SourceLightingGuid_DEPRECATED;
+
 	UPROPERTY(EditAnywhere, Category=Source2D, meta=(DisplayName="Tile Size X"))
 	/** The reference texture tile size X */
 	int32 Source2DTileSizeX;
diff --git a/Engine/Source/Runtime/Engine/Private/Texture.cpp b/Engine/Source/Runtime/Engine/Private/Texture.cpp
--- a/Engine/Source/Runtime/Engine/Private/Texture.cpp
+++ b/Engine/Source/Runtime/Engine/Private/Texture.cpp
@@ -692,12 +692,14 @@ void UTexture::PostEditChangeProperty(FPropertyChangedEvent& PropertyChangedEven
 		for (TObjectIterator<UTexture> It; It; ++It) // walk all textures in the world
 		{
 			UTexture* Tex = *It;
-
+			
+			// @@!! should this just be done in reimport, not for every property change?
 			if(Tex != this && Tex->CompositeTexture == this && Tex->CompositeTextureMode != CTM_Disabled)
 			{
 				TexturesThatUseThisTexture.Add(Tex);
 			}
 		}
+		// there is a potential infinite loop here if two textures depend on each other (or in a ring)
 		for (int32 i = 0; i < TexturesThatUseThisTexture.Num(); ++i)
 		{
 			TexturesThatUseThisTexture[i]->PostEditChange();
@@ -713,6 +715,7 @@ void UTexture::PostEditChangeProperty(FPropertyChangedEvent& PropertyChangedEven
 		}
 	}
 }
+
 #endif // WITH_EDITOR
 
 #if WITH_EDITORONLY_DATA
@@ -2619,23 +2622,28 @@ FGuid FTextureSource::GetId() const
 	{
 		return Id;
 	}
-
+	
 	UE::DerivedData::FBuildVersionBuilder IdBuilder;
 	IdBuilder << BaseBlockX;
 	IdBuilder << BaseBlockY;
 	IdBuilder << SizeX;
-	IdBuilder << SizeY;
+	IdBuilder << SizeY;	
 	IdBuilder << NumSlices;
 	IdBuilder << NumMips;
-	IdBuilder << NumLayers;
+	IdBuilder << NumLayers;	
 	IdBuilder << bLongLatCubemap;
 	IdBuilder << CompressionFormat;
-	IdBuilder << bGuidIsHash;
+	IdBuilder << bGuidIsHash; // always true here
 	IdBuilder << static_cast<uint8>(Format.GetValue());
+	
+	// @@!! this looks wrong
+	// .Num() is number of items, not number of bytes
 	IdBuilder.Serialize(const_cast<void*>(static_cast<const void*>(LayerFormat.GetData())), LayerFormat.Num());
 	IdBuilder.Serialize(const_cast<void*>(static_cast<const void*>(Blocks.GetData())), Blocks.Num());
 	IdBuilder.Serialize(const_cast<void*>(static_cast<const void*>(BlockDataOffsets.GetData())), BlockDataOffsets.Num());
+	
 	IdBuilder << const_cast<FGuid&>(Id);
+
 	return IdBuilder.Build();
 }
 
@@ -3523,5 +3531,5 @@ bool FTextureSource::EnsureBlocksAreSorted()
 
 	return true;
 }
-
 #endif //WITH_EDITOR
+
diff --git a/Engine/Source/Runtime/Engine/Private/TextureCubeArray.cpp b/Engine/Source/Runtime/Engine/Private/TextureCubeArray.cpp
--- a/Engine/Source/Runtime/Engine/Private/TextureCubeArray.cpp
+++ b/Engine/Source/Runtime/Engine/Private/TextureCubeArray.cpp
@@ -364,7 +364,6 @@ uint32 UTextureCubeArray::CalcTextureMemorySizeEnum(ETextureMipCount Enum) const
 }
 
 #if WITH_EDITOR
-
 ENGINE_API bool UTextureCubeArray::CheckArrayTexturesCompatibility()
 {
 	if (SourceTextures.Num() == 0)
@@ -430,7 +429,8 @@ ENGINE_API bool UTextureCubeArray::UpdateSourceFromSourceTextures(bool bCreating
 	{
 		return false;
 	}
-
+	
+	Modify();
 
 	const FTextureSource& InitialSource = SourceTextures[0]->Source;
 	// Format and format size.
@@ -499,12 +499,14 @@ ENGINE_API bool UTextureCubeArray::UpdateSourceFromSourceTextures(bool bCreating
 	SetLightingGuid();
 	ValidateSettingsAfterImportOrEdit();
 	UpdateResource();
-
+	
 	return true;
 }
 
 ENGINE_API void UTextureCubeArray::InvadiateTextureSource()
 {
+	Modify();
+
 	if (PlatformData)
 	{
 		delete PlatformData;
@@ -514,6 +516,7 @@ ENGINE_API void UTextureCubeArray::InvadiateTextureSource()
 	Source = FTextureSource();
 	Source.SetOwner(this);
 	UpdateResource();
+	
 }
 #endif
 
@@ -543,9 +546,9 @@ void UTextureCubeArray::Serialize(FArchive& Ar)
 void UTextureCubeArray::PostLoad()
 {
 #if WITH_EDITOR
-	FinishCachePlatformData();
-
+	FinishCachePlatformData();	
 #endif // #if WITH_EDITOR
+
 	Super::PostLoad();
 };
 
diff --git a/Engine/Source/Runtime/Engine/Private/VolumeTexture.cpp b/Engine/Source/Runtime/Engine/Private/VolumeTexture.cpp
--- a/Engine/Source/Runtime/Engine/Private/VolumeTexture.cpp
+++ b/Engine/Source/Runtime/Engine/Private/VolumeTexture.cpp
@@ -40,6 +40,9 @@ bool UVolumeTexture::UpdateSourceFromSourceTexture()
 	bool bSourceValid = false;
 
 #if WITH_EDITOR
+
+	Modify(true);
+
 	if (Source2DTexture && Source2DTileSizeX > 0 && Source2DTileSizeY > 0)
 	{
 		FTextureSource& InitialSource = Source2DTexture->Source;
@@ -77,7 +80,6 @@ bool UVolumeTexture::UpdateSourceFromSourceTexture()
 					}
 
 					Source.Init(Source2DTileSizeX, Source2DTileSizeY, TileSizeZ, 1, InitialSource.GetFormat(), NewData);
-					SourceLightingGuid = Source2DTexture->GetLightingGuid();
 					bSourceValid = true;
 
 					FMemory::Free(NewData);
@@ -94,7 +96,7 @@ bool UVolumeTexture::UpdateSourceFromSourceTexture()
 	{
 		Source = FTextureSource();
 		Source.SetOwner(this);
-		SourceLightingGuid.Invalidate();
+		// Source2DTexture = nullptr; // ??
 
 		if (PlatformData)
 		{
@@ -118,6 +120,8 @@ ENGINE_API bool UVolumeTexture::UpdateSourceFromFunction(TFunction<void(int32, i
 		UE_LOG(LogTexture, Warning, TEXT("%s UpdateSourceFromFunction size in x,y, and z must be greater than zero"), *GetFullName());
 		return false;
 	}
+	
+	Modify(true);
 
 	// First clear up the existing source with the requested TextureSourceFormat
 	Source.Init(0, 0, 0, 1, Format, nullptr);
@@ -194,11 +198,6 @@ void UVolumeTexture::PostLoad()
 {
 #if WITH_EDITOR
 	FinishCachePlatformData();
-
-	if (Source2DTexture && SourceLightingGuid != Source2DTexture->GetLightingGuid())
-	{
-		UpdateSourceFromSourceTexture();
-	}
 #endif // #if WITH_EDITOR
 
 	Super::PostLoad();
@@ -369,8 +368,9 @@ void UVolumeTexture::PostEditChangeProperty(FPropertyChangedEvent& PropertyChang
 			UpdateSourceFromSourceTexture();
 		}
 	}
-	
+		
 	Super::PostEditChangeProperty(PropertyChangedEvent);
+
 }
 
 
@@ -438,3 +438,4 @@ bool UVolumeTexture::StreamIn(int32 NewMipCount, bool bHighPrio)
 	}
 	return false;
 }
+
