Ensure online game session exists before calling player registration functions.

diff --git a/Engine/Source/Runtime/Engine/Private/GameSession.cpp b/Engine/Source/Runtime/Engine/Private/GameSession.cpp
--- a/Engine/Source/Runtime/Engine/Private/GameSession.cpp
+++ b/Engine/Source/Runtime/Engine/Private/GameSession.cpp
@@ -248,7 +248,8 @@ void AGameSession::UnregisterPlayer(FName InSessionName, const FUniqueNetIdRepl&
 {
 	UWorld* World = GetWorld();
 	if (GetNetMode() != NM_Standalone &&
-		UniqueId.IsValid())
+		UniqueId.IsValid() &&
+		UOnlineEngineInterface::Get()->DoesSessionExist(World, InSessionName))
 	{
 		// Remove the player from the session
 		UOnlineEngineInterface::Get()->UnregisterPlayer(World, InSessionName, UniqueId);
@@ -259,7 +260,8 @@ void AGameSession::UnregisterPlayers(FName InSessionName, const TArray<FUniqueNe
 {
 	UWorld* World = GetWorld();
 	if (GetNetMode() != NM_Standalone &&
-		Players.Num() > 0)
+		Players.Num() > 0 &&
+		UOnlineEngineInterface::Get()->DoesSessionExist(World, InSessionName))
 	{
 		// Remove the player from the session
 		TArray<FUniqueNetIdWrapper> PlayerIdsAsWrappers;
@@ -275,7 +277,8 @@ void AGameSession::UnregisterPlayers(FName InSessionName, const TArray<FUniqueNe
 {
 	UWorld * World = GetWorld();
 	if (GetNetMode() != NM_Standalone &&
-		Players.Num() > 0)
+		Players.Num() > 0 &&
+		UOnlineEngineInterface::Get()->DoesSessionExist(World, InSessionName))
 	{
 		// Remove the player from the session
 		TArray<FUniqueNetIdWrapper> PlayerIdsAsWrappers;
diff --git a/Engine/Source/Runtime/Engine/Private/PlayerState.cpp b/Engine/Source/Runtime/Engine/Private/PlayerState.cpp
--- a/Engine/Source/Runtime/Engine/Private/PlayerState.cpp
+++ b/Engine/Source/Runtime/Engine/Private/PlayerState.cpp
@@ -399,7 +399,10 @@ void APlayerState::RegisterPlayerWithSession(bool bWasFromInvite)
 		{
 			// Register the player as part of the session
 			const APlayerState* PlayerState = GetDefault<APlayerState>();
-			UOnlineEngineInterface::Get()->RegisterPlayer(GetWorld(), PlayerState->SessionName, GetUniqueId(), bWasFromInvite);
+			if (UOnlineEngineInterface::Get()->DoesSessionExist(GetWorld(), PlayerState->SessionName))
+			{
+				UOnlineEngineInterface::Get()->RegisterPlayer(GetWorld(), PlayerState->SessionName, GetUniqueId(), bWasFromInvite);
+			}
 		}
 	}
 }
@@ -411,7 +414,10 @@ void APlayerState::UnregisterPlayerWithSession()
 		const APlayerState* PlayerState = GetDefault<APlayerState>();
 		if (PlayerState->SessionName != NAME_None)
 		{
-			UOnlineEngineInterface::Get()->UnregisterPlayer(GetWorld(), PlayerState->SessionName, GetUniqueId());
+			if (UOnlineEngineInterface::Get()->DoesSessionExist(GetWorld(), PlayerState->SessionName))
+			{
+				UOnlineEngineInterface::Get()->UnregisterPlayer(GetWorld(), PlayerState->SessionName, GetUniqueId());
+			}
 		}
 	}
 }
