Adding mechanism to catch and handle crashes and abnormal shutdowns triggered by terminate() and abort()

diff --git a/Engine/Source/Runtime/Core/Private/Windows/WindowsPlatformCrashContext.cpp b/Engine/Source/Runtime/Core/Private/Windows/WindowsPlatformCrashContext.cpp
--- a/Engine/Source/Runtime/Core/Private/Windows/WindowsPlatformCrashContext.cpp
+++ b/Engine/Source/Runtime/Core/Private/Windows/WindowsPlatformCrashContext.cpp
@@ -30,6 +30,7 @@
 #include "BuildSettings.h"
 #include "CoreGlobals.h"
 #include <atomic>
+#include <signal.h>
 
 #include "Windows/WindowsHWrapper.h"
 #include "Windows/AllowWindowsPlatformTypes.h"
@@ -57,6 +58,10 @@
 #define DISABLE_CRC_SUBMIT_AND_RESTART_BUTTON 0
 #endif
 
+#ifndef HANDLE_ABORT_SIGNALS
+#define HANDLE_ABORT_SIGNALS 1
+#endif
+
 #define CR_CLIENT_MAX_PATH_LEN 265
 #define CR_CLIENT_MAX_ARGS_LEN 256
 
@@ -969,6 +974,25 @@ void CreateExceptionInfoString(EXCEPTION_RECORD* ExceptionRecord, TCHAR* OutErro
 #undef HANDLE_CASE
 }
 
+#if HANDLE_ABORT_SIGNALS
+static void AbortHandler(int Signal)
+{
+#if !PLATFORM_SEH_EXCEPTIONS_DISABLED && !NOINITCRASHREPORTER
+	__try
+	{
+		// Force an exception to invoke the crash reporter
+		FAssertInfo Info(TEXT("Abort signal received"), PLATFORM_RETURN_ADDRESS());
+		ULONG_PTR Arguments[] = { (ULONG_PTR)&Info };
+		::RaiseException(EnsureExceptionCode, 0, UE_ARRAY_COUNT(Arguments), Arguments);
+	}
+	__except(ReportCrash(GetExceptionInformation()))
+	{
+		GIsCriticalError = true;
+		FPlatformMisc::RequestExit(true);
+	}
+#endif
+}
+#endif
 
 
 
@@ -1129,6 +1153,11 @@ public:
 		// Register an exception handler for exceptions that aren't handled by any vectored exception handlers or structured exception handlers (__try/__except),
 		// especially to capture crash in non-engine-wrapped threads (like native threads) that are usually not guarded with structured exception handling.
 		FCoreDelegates::GetPreMainInitDelegate().AddRaw(this, &FCrashReportingThread::RegisterUnhandledExceptionHandler);
+
+#if HANDLE_ABORT_SIGNALS
+		// Register a signal handler for abort and terminate occurrences so that abnormal or third party shutdowns can be reported.
+		FCoreDelegates::GetPreMainInitDelegate().AddRaw(this, &FCrashReportingThread::RegisterAbortSignalHandler);
+#endif
 	}
 
 	FORCENOINLINE ~FCrashReportingThread()
@@ -1167,6 +1196,22 @@ public:
 #endif
 	}
 
+#if HANDLE_ABORT_SIGNALS
+	void RegisterAbortSignalHandler()
+	{
+#if !NOINITCRASHREPORTER
+		if (::signal(SIGABRT, AbortHandler) != SIG_ERR)
+		{
+			UE_LOG(LogWindows, Log, TEXT("Custom abort handler registered for crash reporting."));
+		}
+		else
+		{
+			UE_LOG(LogWindows, Error, TEXT("Unable to register custom abort handler for crash reporting."));
+		}
+#endif
+	}
+#endif
+
 	DWORD GetReporterThreadId() const
 	{
 		return ThreadId;
diff --git a/Engine/Source/Runtime/Engine/Private/UnrealEngine.cpp b/Engine/Source/Runtime/Engine/Private/UnrealEngine.cpp
--- a/Engine/Source/Runtime/Engine/Private/UnrealEngine.cpp
+++ b/Engine/Source/Runtime/Engine/Private/UnrealEngine.cpp
@@ -10011,6 +10011,68 @@ bool UEngine::PerformError(const TCHAR* Cmd, FOutputDevice& Ar)
 		UE_LOG(LogEngine, Fatal, TEXT("%s"), TEXT("Crashing the gamethread at your request"));
 		return true;
 	}
+	else if (FParse::Command(&Cmd, TEXT("TERMINATE")))
+	{
+		UE_LOG(LogEngine, Warning, TEXT("Printed warning to log."));
+		FGenericCrashContext::SetCrashTrigger(ECrashTrigger::Debug);
+		UE_LOG(LogEngine, Warning, TEXT("%s"), TEXT("Terminating from the gamethread at your request"));
+		terminate();
+		return true;
+	}
+	else if (FParse::Command(&Cmd, TEXT("ABORT")))
+	{
+		UE_LOG(LogEngine, Warning, TEXT("Printed warning to log."));
+		FGenericCrashContext::SetCrashTrigger(ECrashTrigger::Debug);
+		UE_LOG(LogEngine, Warning, TEXT("%s"), TEXT("Aborting from the gamethread at your request"));
+		abort();
+		return true;
+	}
+	else if (FParse::Command(&Cmd, TEXT("THROWINNOTHROW")))
+	{
+		struct FThrowCrash
+		{
+			static void ThrowInNoThrow() noexcept(true)
+			{
+				ThrowException();
+			}
+
+		private:
+			
+			static void ThrowException()
+			{
+				throw 1;
+			}
+		};
+		UE_LOG(LogEngine, Warning, TEXT("Printed warning to log."));
+		FGenericCrashContext::SetCrashTrigger(ECrashTrigger::Debug);
+		UE_LOG(LogEngine, Warning, TEXT("%s"), TEXT("Throwing in a nothrow function on the gamethread at your request"));
+		FThrowCrash::ThrowInNoThrow();
+		return true;
+	}
+	else if (FParse::Command(&Cmd, TEXT("EXIT")))
+	{
+		UE_LOG(LogEngine, Warning, TEXT("Printed warning to log."));
+		FGenericCrashContext::SetCrashTrigger(ECrashTrigger::Debug);
+		UE_LOG(LogEngine, Warning, TEXT("%s"), TEXT("Exiting from the gamethread at your request with 'exit()'"));
+		exit(1);
+		return true;
+	}
+	else if (FParse::Command(&Cmd, TEXT("QUICKEXIT")))
+	{
+		UE_LOG(LogEngine, Warning, TEXT("Printed warning to log."));
+		FGenericCrashContext::SetCrashTrigger(ECrashTrigger::Debug);
+		UE_LOG(LogEngine, Warning, TEXT("%s"), TEXT("Exiting from the gamethread at your request with 'quick_exit()'"));
+		quick_exit(1);
+		return true;
+	}
+	else if (FParse::Command(&Cmd, TEXT("_EXIT")))
+	{
+		UE_LOG(LogEngine, Warning, TEXT("Printed warning to log."));
+		FGenericCrashContext::SetCrashTrigger(ECrashTrigger::Debug);
+		UE_LOG(LogEngine, Warning, TEXT("%s"), TEXT("Exiting from the gamethread at your request with '_Exit()'"));
+		_Exit(1);
+		return true;
+	}
 	else if (FParse::Command(&Cmd, TEXT("CHECK")))
 	{
 		UE_LOG(LogEngine, Warning, TEXT("Printed warning to log."));
