Adding mechanism to catch and handle crashes and abnormal shutdowns triggered by terminate() and abort()

diff --git a/Engine/Source/Runtime/Engine/Private/UnrealEngine.cpp b/Engine/Source/Runtime/Engine/Private/UnrealEngine.cpp
--- a/Engine/Source/Runtime/Engine/Private/UnrealEngine.cpp
+++ b/Engine/Source/Runtime/Engine/Private/UnrealEngine.cpp
@@ -10016,7 +10016,7 @@ bool UEngine::PerformError(const TCHAR* Cmd, FOutputDevice& Ar)
 		UE_LOG(LogEngine, Warning, TEXT("Printed warning to log."));
 		FGenericCrashContext::SetCrashTrigger(ECrashTrigger::Debug);
 		UE_LOG(LogEngine, Warning, TEXT("%s"), TEXT("Terminating from the gamethread at your request"));
-		terminate();
+		std::terminate();
 		return true;
 	}
 	else if (FParse::Command(&Cmd, TEXT("ABORT")))
@@ -10024,29 +10024,7 @@ bool UEngine::PerformError(const TCHAR* Cmd, FOutputDevice& Ar)
 		UE_LOG(LogEngine, Warning, TEXT("Printed warning to log."));
 		FGenericCrashContext::SetCrashTrigger(ECrashTrigger::Debug);
 		UE_LOG(LogEngine, Warning, TEXT("%s"), TEXT("Aborting from the gamethread at your request"));
-		abort();
-		return true;
-	}
-	else if (FParse::Command(&Cmd, TEXT("THROWINNOTHROW")))
-	{
-		struct FThrowCrash
-		{
-			static void ThrowInNoThrow() noexcept(true)
-			{
-				ThrowException();
-			}
-
-		private:
-			
-			static void ThrowException()
-			{
-				throw 1;
-			}
-		};
-		UE_LOG(LogEngine, Warning, TEXT("Printed warning to log."));
-		FGenericCrashContext::SetCrashTrigger(ECrashTrigger::Debug);
-		UE_LOG(LogEngine, Warning, TEXT("%s"), TEXT("Throwing in a nothrow function on the gamethread at your request"));
-		FThrowCrash::ThrowInNoThrow();
+		std::abort();
 		return true;
 	}
 	else if (FParse::Command(&Cmd, TEXT("EXIT")))
