Work around for asset registry bug (UE-158163) that causes the AR to not immediatly add new external actors.
Re-enable project sync using asset filter.
Make UAssetManager::GatherPublicAssetsForPackage a bit simpler.

diff --git a/Engine/Source/Runtime/Engine/Private/AssetManager.cpp b/Engine/Source/Runtime/Engine/Private/AssetManager.cpp
--- a/Engine/Source/Runtime/Engine/Private/AssetManager.cpp
+++ b/Engine/Source/Runtime/Engine/Private/AssetManager.cpp
@@ -4062,42 +4062,42 @@ void UAssetManager::ModifyDLCCook(const FString& DLCName, TArray<FName>& Package
 
 void UAssetManager::GatherPublicAssetsForPackage(FName PackagePath, TArray<FName>& PackagesToCook) const
 {
-	// query asset registry for package assets
-	TArray<FAssetData> OutAssetData;
-	GetAssetRegistry().GetAssetsByPath(PackagePath, OutAssetData, true, true);
+	FARFilter Filter;
+	Filter.PackagePaths.Add(PackagePath);
+	Filter.bRecursivePaths = true;	
+	Filter.bIncludeOnlyOnDiskAssets = true;
+	Filter.WithoutPackageFlags |= PKG_NotExternallyReferenceable;
 
-	for (const FAssetData& AssetData : OutAssetData)
+	GetAssetRegistry().EnumerateAssets(Filter, [&PackagesToCook](const FAssetData& AssetData)
 	{
-		// determine if this package can be externally referenced
-		if (0 == (AssetData.PackageFlags & PKG_NotExternallyReferenceable))
+		// this package can be externally referenced; include it in the cook
+		if (const UPackage* const OuterPackage = AssetData.GetPackage())
 		{
-			// this package can be externally referenced; include it in the cook
-			if (const UPackage* const OuterPackage = AssetData.GetPackage())
+			FName PackageFName = OuterPackage->GetLoadedPath().GetPackageFName();
+			if (!PackageFName.IsNone())
 			{
-				FName PackageFName = OuterPackage->GetLoadedPath().GetPackageFName();
-				if (!PackageFName.IsNone())
-				{
-					UE_LOG(LogAssetManager, Verbose,
-						TEXT("GatherPublicAssetsForPackage: Adding public package [%s] (instigator: [%s]"), 
-						*PackageFName.ToString(),
-						*AssetData.AssetName.ToString());
-					PackagesToCook.AddUnique(PackageFName);
-				}
-				else
-				{
-					UE_LOG(LogAssetManager, Error, 
-						TEXT("GatherPublicAssetsForPackage: Failed to resolve package name for asset [%s]"), 
-						*AssetData.AssetName.ToString());
-				}
+				UE_LOG(LogAssetManager, Verbose,
+					TEXT("GatherPublicAssetsForPackage: Adding public package [%s] (instigator: [%s]"),
+					*PackageFName.ToString(),
+					*AssetData.AssetName.ToString());
+				PackagesToCook.AddUnique(PackageFName);
 			}
 			else
 			{
 				UE_LOG(LogAssetManager, Error,
-					TEXT("GatherPublicAssetsForPackage: Failed to resolve package for asset [%s]"), 
+					TEXT("GatherPublicAssetsForPackage: Failed to resolve package name for asset [%s]"),
 					*AssetData.AssetName.ToString());
 			}
 		}
-	}
+		else
+		{
+			UE_LOG(LogAssetManager, Error,
+				TEXT("GatherPublicAssetsForPackage: Failed to resolve package for asset [%s]"),
+				*AssetData.AssetName.ToString());
+		}
+
+		return true;
+	});
 }
 
 bool UAssetManager::ShouldCookForPlatform(const UPackage* Package, const ITargetPlatform* TargetPlatform)
