Change ClientVisibleActorOuters map to use FNames, removes strong ref to levels
Fix ServerNotifyLoadedWorld to use NetworkRemapPath

diff --git a/Engine/Source/Runtime/Engine/Classes/Engine/NetConnection.h b/Engine/Source/Runtime/Engine/Classes/Engine/NetConnection.h
--- a/Engine/Source/Runtime/Engine/Classes/Engine/NetConnection.h
+++ b/Engine/Source/Runtime/Engine/Classes/Engine/NetConnection.h
@@ -1488,7 +1488,7 @@ private:
 	FName PlayerOnlinePlatformName;
 
 	/** This is an acceleration set that is derived from ClientWorldPackageName and ClientVisibleLevelNames. We use this to quickly test an AActor*'s visibility while replicating. */
-	mutable TMap<UObject*, bool> ClientVisibleActorOuters;
+	mutable TMap<FName, bool> ClientVisibleActorOuters;
 
 	/** This is used to capture visibility updates while the server is in transition and deffer the update until the server has completed the transition */
 	TMap<FName, FUpdateLevelVisibilityLevelInfo> PendingUpdateLevelVisibility;
@@ -1502,7 +1502,7 @@ private:
 	void UpdateLevelVisibilityInternal(const struct FUpdateLevelVisibilityLevelInfo& LevelVisibility);
 
 	/** Called internally to update cached acceleration map */
-	bool UpdateCachedLevelVisibility(ULevel* Level) const;
+	bool UpdateCachedLevelVisibility(const FName& PackageName) const;
 
 	/** Updates entire cached LevelVisibility map */
 	void UpdateAllCachedLevelVisibility() const;
diff --git a/Engine/Source/Runtime/Engine/Private/NetConnection.cpp b/Engine/Source/Runtime/Engine/Private/NetConnection.cpp
--- a/Engine/Source/Runtime/Engine/Private/NetConnection.cpp
+++ b/Engine/Source/Runtime/Engine/Private/NetConnection.cpp
@@ -1291,12 +1291,6 @@ void UNetConnection::AddReferencedObjects(UObject* InThis, FReferenceCollector&
 		}
 	}
 
-	// ClientVisibleActorOuters acceleration map
-	for (auto& MapIt : This->ClientVisibleActorOuters)
-	{
-		Collector.AddReferencedObject(MapIt.Key, This);
-	}
-
 	Super::AddReferencedObjects(This, Collector);
 }
 
@@ -1346,7 +1340,7 @@ FNetLevelVisibilityTransactionId UNetConnection::UpdateLevelStreamStatusChangedT
 	if (bShouldBeVisible == false)
 	{
 		ClientVisibleLevelNames.Remove(PackageName);
-		UpdateAllCachedLevelVisibility();
+		UpdateCachedLevelVisibility(PackageName);
 	}
 
 	return TransactionId;
@@ -1361,47 +1355,46 @@ bool UNetConnection::ClientHasInitializedLevelFor(const AActor* TestActor) const
 	// Each client connection has a different visibility state (what levels are currently loaded for them).
 	// Actor's outer is what we need
 
-	// Note: we are calling GetOuter() here instead of GetLevel() to avoid an unreal Cast<>: we justt need the memory address for the lookup.
-	UObject* ActorOuter = TestActor->GetOuter();
-	if (const bool* bIsVisible = ClientVisibleActorOuters.Find(ActorOuter))
+	const FName PackageName = TestActor->GetOutermost()->GetFName();
+
+	if (const bool* bIsVisible = ClientVisibleActorOuters.Find(PackageName))
 	{
 		return *bIsVisible;
 	}
 
 	// The actor's outer was not in the acceleration map so we perform the "legacy" function and 
 	// cache the result so that we don't do this every time:
-	return UpdateCachedLevelVisibility(Cast<ULevel>(ActorOuter));
+	return UpdateCachedLevelVisibility(PackageName);
 }
 
-bool UNetConnection::UpdateCachedLevelVisibility(ULevel* Level) const
+bool UNetConnection::UpdateCachedLevelVisibility(const FName& PackageName) const
 {
-	bool IsVisibile = false;
-	if (Level == nullptr)
+	bool bVisible = false;
+	
+	if (PackageName.IsNone())
 	{
-		IsVisibile = true;
+		bVisible = true;
 	}
-	else if (Level->IsPersistentLevel() && Driver->GetWorldPackage()->GetFName() == ClientWorldPackageName)
+	else if ((PackageName == ClientWorldPackageName) && (Driver->GetWorldPackage()->GetFName() == ClientWorldPackageName))
 	{
-		IsVisibile = true;
+		bVisible = true;
 	}
 	else
 	{
-		IsVisibile = ClientVisibleLevelNames.Contains(Level->GetOutermost()->GetFName());
+		bVisible = ClientVisibleLevelNames.Contains(PackageName);
 	}
 
-	ClientVisibleActorOuters.FindOrAdd(Level) = IsVisibile;
-	return IsVisibile;
+	ClientVisibleActorOuters.FindOrAdd(PackageName) = bVisible;
+
+	return bVisible;
 }
 
 void UNetConnection::UpdateAllCachedLevelVisibility() const
 {
 	// Update our acceleration map
-	for (auto& MapIt : ClientVisibleActorOuters)
+	for (const TPair<FName, bool>& VisPair : ClientVisibleActorOuters)
 	{
-		if (ULevel* Level = Cast<ULevel>(MapIt.Key))
-		{
-			UpdateCachedLevelVisibility(Level);
-		}
+		UpdateCachedLevelVisibility(VisPair.Key);
 	}
 }
 
@@ -1453,6 +1446,8 @@ void UNetConnection::UpdateLevelVisibilityInternal(const FUpdateLevelVisibilityL
 			ClientVisibleLevelNames.Add(LevelVisibility.PackageName);
 			UE_LOG(LogPlayerController, Verbose, TEXT("ServerUpdateLevelVisibility() Added '%s'"), *LevelVisibility.PackageName.ToString());
 
+			UpdateCachedLevelVisibility(LevelVisibility.PackageName);
+
 			QUICK_USE_CYCLE_STAT(NetUpdateLevelVisibility_UpdateDormantActors, STATGROUP_Net);
 
 			// Any destroyed actors that were destroyed prior to the streaming level being unloaded for the client will not be in the connections
@@ -1535,9 +1530,9 @@ void UNetConnection::UpdateLevelVisibilityInternal(const FUpdateLevelVisibilityL
 				Channel->Close(EChannelCloseReason::LevelUnloaded);
 			}
 		}
-	}
 
-	UpdateAllCachedLevelVisibility();
+		UpdateCachedLevelVisibility(LevelVisibility.PackageName);
+	}
 }
 
 void UNetConnection::SetClientWorldPackageName(FName NewClientWorldPackageName)
diff --git a/Engine/Source/Runtime/Engine/Private/PlayerController.cpp b/Engine/Source/Runtime/Engine/Private/PlayerController.cpp
--- a/Engine/Source/Runtime/Engine/Private/PlayerController.cpp
+++ b/Engine/Source/Runtime/Engine/Private/PlayerController.cpp
@@ -573,6 +573,8 @@ bool APlayerController::ServerNotifyLoadedWorld_Validate(FName WorldPackageName)
 
 void APlayerController::ServerNotifyLoadedWorld_Implementation(FName WorldPackageName)
 {
+	WorldPackageName = NetworkRemapPath(WorldPackageName, true);
+
 	UE_LOG(LogPlayerController, Verbose, TEXT("APlayerController::ServerNotifyLoadedWorld_Implementation: Client loaded %s"), *WorldPackageName.ToString());
 
 	UWorld* CurWorld = GetWorld();
diff --git a/Engine/Source/Runtime/Engine/Private/World.cpp b/Engine/Source/Runtime/Engine/Private/World.cpp
--- a/Engine/Source/Runtime/Engine/Private/World.cpp
+++ b/Engine/Source/Runtime/Engine/Private/World.cpp
@@ -7369,7 +7369,9 @@ UWorld* FSeamlessTravelHandler::Tick()
 					UE_LOG(LogNet, Verbose, TEXT("NotifyLoadedWorld Begin"));
 #endif
 					It->PlayerController->NotifyLoadedWorld(LoadedWorldName, bSwitchedToDefaultMap);
-					It->PlayerController->ServerNotifyLoadedWorld(LoadedWorldName);
+
+					const FName SendLoadedWorldName = It->PlayerController->NetworkRemapPath(LoadedWorldName, false);
+					It->PlayerController->ServerNotifyLoadedWorld(SendLoadedWorldName);
 #if !UE_BUILD_SHIPPING
 					UE_LOG(LogNet, Verbose, TEXT("NotifyLoadedWorld End"));
 #endif
