fix FTextureSource::GetId was not right for layers/blocks
was serializing arrays with element count when it meant to use size in bytes
broken serialization left in place for common case where it is benign so ids aren't changed

diff --git a/Engine/Source/Runtime/Engine/Private/Texture.cpp b/Engine/Source/Runtime/Engine/Private/Texture.cpp
--- a/Engine/Source/Runtime/Engine/Private/Texture.cpp
+++ b/Engine/Source/Runtime/Engine/Private/Texture.cpp
@@ -2616,6 +2616,12 @@ void FTextureSource::UseHashAsGuid()
 	}
 }
 
+template <typename ArrayType>
+static SIZE_T ArraySizeBytes(const ArrayType & Array)
+{
+	return (SIZE_T)( Array.Num() * Array.GetTypeSize() );
+}
+
 FGuid FTextureSource::GetId() const
 {
 	if (!bGuidIsHash)
@@ -2636,12 +2642,32 @@ FGuid FTextureSource::GetId() const
 	IdBuilder << bGuidIsHash; // always true here
 	IdBuilder << static_cast<uint8>(Format.GetValue());
 	
-	// @@!! this looks wrong
-	// .Num() is number of items, not number of bytes
-	IdBuilder.Serialize(const_cast<void*>(static_cast<const void*>(LayerFormat.GetData())), LayerFormat.Num());
-	IdBuilder.Serialize(const_cast<void*>(static_cast<const void*>(Blocks.GetData())), Blocks.Num());
-	IdBuilder.Serialize(const_cast<void*>(static_cast<const void*>(BlockDataOffsets.GetData())), BlockDataOffsets.Num());
-	
+	if ( GetNumLayers() == 1 && GetNumBlocks() == 1 )
+	{
+		// preserve broken code for common case so Ids don't change :
+		// was serializing using array Num (element count) instead of byte count
+		// the broken serialize here only takes 1 byte from these arrays
+		// but that's benign because they don't really need to be hashed anyway (they are redundant in this case)
+		
+		IdBuilder.Serialize((void *)LayerFormat.GetData(), LayerFormat.Num());
+		IdBuilder.Serialize((void *)Blocks.GetData(), Blocks.Num());
+		IdBuilder.Serialize((void *)BlockDataOffsets.GetData(), BlockDataOffsets.Num());
+	}
+	else
+	{
+		// better version :
+
+		if ( GetNumLayers() > 1 )
+		{
+			IdBuilder.Serialize((void *)LayerFormat.GetData(), ArraySizeBytes(LayerFormat));
+		}
+		if ( GetNumBlocks() > 1 )
+		{
+			IdBuilder.Serialize((void *)Blocks.GetData(), ArraySizeBytes(Blocks));
+			IdBuilder.Serialize((void *)BlockDataOffsets.GetData(), ArraySizeBytes(BlockDataOffsets));
+		}
+	}
+
 	IdBuilder << const_cast<FGuid&>(Id);
 
 	return IdBuilder.Build();
@@ -3450,12 +3476,14 @@ void FTextureSource::InitBlockedImpl(const ETextureSourceFormat* InLayerFormats,
 	NumLayers = InNumLayers;
 	Format = InLayerFormats[0];
 
+	// Blocks is of size NumBlocks-1 , and 0th block is in the TextureSource
 	Blocks.Reserve(InNumBlocks - 1);
 	for (int32 BlockIndex = 1; BlockIndex < InNumBlocks; ++BlockIndex)
 	{
 		Blocks.Add(InBlocks[BlockIndex]);
 	}
 
+	// LayerFormat is of size NumLayers, and Format == LayerFormat[0]
 	LayerFormat.SetNum(InNumLayers, true);
 	for (int i = 0; i < InNumLayers; ++i)
 	{
@@ -3484,6 +3512,8 @@ inline bool operator<(const FSortedTextureSourceBlock& Lhs, const FSortedTexture
 
 bool FTextureSource::EnsureBlocksAreSorted()
 {
+	// BlockDataOffsets is of size NumBlocks, even when NumBlocks==1
+	// and BlockDataOffsets[0] == 0
 	const int32 NumBlocks = GetNumBlocks();
 	if (BlockDataOffsets.Num() == NumBlocks)
 	{
