Added lock to TResourcePool to make it thread safe.

diff --git a/Engine/Source/Runtime/Engine/Public/ResourcePool.h b/Engine/Source/Runtime/Engine/Public/ResourcePool.h
--- a/Engine/Source/Runtime/Engine/Public/ResourcePool.h
+++ b/Engine/Source/Runtime/Engine/Public/ResourcePool.h
@@ -55,16 +55,19 @@ public:
 		// Find the appropriate bucket based on size
 		const uint32 BucketIndex = Policy.GetPoolBucketIndex(Args);
 		TArray<FPooledResource>& PoolBucket = ResourceBuckets[BucketIndex];
-		if (PoolBucket.Num() > 0)
-		{
-			// Reuse the last entry in this size bucket
-			return PoolBucket.Pop().Resource;
-		}
-		else
+
 		{
-			// Nothing usable was found in the free pool, create a new resource
-			return Policy.CreateResource(Args);
+			FScopeLock Lock(&CS);
+
+			if (PoolBucket.Num() > 0)
+			{
+				// Reuse the last entry in this size bucket
+				return PoolBucket.Pop().Resource;
+			}
 		}
+
+		// Nothing usable was found in the free pool, create a new resource
+		return Policy.CreateResource(Args);
 	}
 	
 	/** Release a resource back into the pool.
@@ -80,7 +83,9 @@ public:
 		// Add to this frame's array of free resources
 		const int32 SafeFrameIndex = GFrameNumberRenderThread % ResourcePoolPolicy::NumSafeFrames;
 		const uint32 BucketIndex = Policy.GetPoolBucketIndex(NewEntry.CreationArguments);
-		
+
+		FScopeLock Lock(&CS);
+	
 		SafeResourceBuckets[SafeFrameIndex][BucketIndex].Add(NewEntry);
 	}
 	
@@ -91,7 +96,9 @@ public:
 	{
 		uint32 NumToCleanThisFrame = ResourcePoolPolicy::NumToDrainPerFrame;
 		uint32 CullAfterFramesNum = ResourcePoolPolicy::CullAfterFramesNum;
-		
+
+		FScopeLock Lock(&CS);
+
 		if(!bForceDrainAll)
 		{
 			// Index of the bucket that is now old enough to be reused
@@ -160,7 +167,9 @@ private:
 		/** The frame the resource was freed */
 		uint32 FrameFreed;
 	};
-	
+
+	FCriticalSection CS;
+
 	// Pool of free Resources, indexed by bucket for constant size search time.
 	TArray<FPooledResource> ResourceBuckets[ResourcePoolPolicy::NumPoolBuckets];
 	
@@ -199,8 +208,6 @@ public:
 	 */
 	ResourceType CreatePooledResource(ResourceCreationArguments Args)
 	{
-		ensure(IsInRenderingThread());
-
 		if (IsInitialized())
 		{
 			return TResourcePool<ResourceType, ResourcePoolPolicy, ResourceCreationArguments>::CreatePooledResource(Args);
@@ -216,8 +223,6 @@ public:
 	 */
 	void ReleasePooledResource(ResourceType Resource)
 	{
-		ensure(IsInRenderingThread());
-
 		if (IsInitialized())
 		{
 			TResourcePool<ResourceType, ResourcePoolPolicy, ResourceCreationArguments>::ReleasePooledResource(Resource);
