Adding mechanism to catch and handle crashes and abnormal shutdowns triggered by terminate() and abort()

diff --git a/Engine/Source/Runtime/Engine/Private/UnrealEngine.cpp b/Engine/Source/Runtime/Engine/Private/UnrealEngine.cpp
--- a/Engine/Source/Runtime/Engine/Private/UnrealEngine.cpp
+++ b/Engine/Source/Runtime/Engine/Private/UnrealEngine.cpp
@@ -10027,30 +10027,6 @@ bool UEngine::PerformError(const TCHAR* Cmd, FOutputDevice& Ar)
 		std::abort();
 		return true;
 	}
-	else if (FParse::Command(&Cmd, TEXT("EXIT")))
-	{
-		UE_LOG(LogEngine, Warning, TEXT("Printed warning to log."));
-		FGenericCrashContext::SetCrashTrigger(ECrashTrigger::Debug);
-		UE_LOG(LogEngine, Warning, TEXT("%s"), TEXT("Exiting from the gamethread at your request with 'exit()'"));
-		exit(1);
-		return true;
-	}
-	else if (FParse::Command(&Cmd, TEXT("QUICKEXIT")))
-	{
-		UE_LOG(LogEngine, Warning, TEXT("Printed warning to log."));
-		FGenericCrashContext::SetCrashTrigger(ECrashTrigger::Debug);
-		UE_LOG(LogEngine, Warning, TEXT("%s"), TEXT("Exiting from the gamethread at your request with 'quick_exit()'"));
-		quick_exit(1);
-		return true;
-	}
-	else if (FParse::Command(&Cmd, TEXT("_EXIT")))
-	{
-		UE_LOG(LogEngine, Warning, TEXT("Printed warning to log."));
-		FGenericCrashContext::SetCrashTrigger(ECrashTrigger::Debug);
-		UE_LOG(LogEngine, Warning, TEXT("%s"), TEXT("Exiting from the gamethread at your request with '_Exit()'"));
-		_Exit(1);
-		return true;
-	}
 	else if (FParse::Command(&Cmd, TEXT("CHECK")))
 	{
 		UE_LOG(LogEngine, Warning, TEXT("Printed warning to log."));
