Improved PhysicsAsset tooltips
Hide old solver settings by default and prevent changes.

diff --git a/Engine/Source/Runtime/Engine/Classes/PhysicsEngine/PhysicsAsset.h b/Engine/Source/Runtime/Engine/Classes/PhysicsEngine/PhysicsAsset.h
--- a/Engine/Source/Runtime/Engine/Classes/PhysicsEngine/PhysicsAsset.h
+++ b/Engine/Source/Runtime/Engine/Classes/PhysicsEngine/PhysicsAsset.h
@@ -29,42 +29,44 @@ struct ENGINE_API FPhysicsAssetSolverSettings
 	FPhysicsAssetSolverSettings();
 
 	/**
-	 * The number of position iterations to run. The position solve is responsible for depenetration and friction.
+	 * RBAN: The number of position iterations to run. The position solve is responsible for depenetration.
 	 * Increasing this will improve simulation stability, but increase the cost.
 	*/
 	UPROPERTY(EditAnywhere, BlueprintReadWrite, Category = SolverSettings)
 	int32 PositionIterations;
 
 	/**
-	 * The number of velocity iterations to run. The velocity solve is responsible for restitution.
+	 * RBAN: The number of velocity iterations to run. The velocity solve is responsible for restitution (bounce) and friction. 
+	 * This should usually be 1, but could be 0 if you don't care about friction and restitution.
 	*/
 	UPROPERTY(EditAnywhere, BlueprintReadWrite, Category = SolverSettings)
 	int32 VelocityIterations;
 
 	/**
-	 * The number of projection iterations to run. The projection phase is a final pass over the constraints, applying
+	 * RBAN: The number of projection iterations to run. The projection phase is a final pass over the constraints, applying
 	 * a semi-physical correction to any joint errors remaining after the position and velocity solves. It can be
 	 * very helpful to stabilize joint chains, but can cause issues with collision response. The projection magnitude
 	 * can be controlled per-constraint in the constraint settings (assuming ProjectionIteration is not zero).
+	 * This should be left as 1 in almost all cases.
 	*/
 	UPROPERTY(EditAnywhere, BlueprintReadWrite, Category = SolverSettings)
 	int32 ProjectionIterations;
 
 	/**
-	 * The distance at which collisions are ignored. In general you need this to be a bit larger than the typical relative body
+	 * RBAN: The distance at which collisions are ignored. In general you need this to be a bit larger than the typical relative body
 	 * movement in your simulation, but small enough so that we don't have to speculatively create too many unused collisions.
 	*/
 	UPROPERTY(EditAnywhere, BlueprintReadWrite, Category = SolverSettings, meta = (ClampMin = 0))
 	float CullDistance;
 
 	/**
-	 * When bodies are penetrating, this is the maximum velocity delta that can be applied in one frame.
+	 * RBAN: When bodies are penetrating, this is the maximum velocity delta that can be applied in one frame.
 	*/
 	UPROPERTY(EditAnywhere, BlueprintReadWrite, Category = SolverSettings, meta = (ClampMin = 0))
 	float MaxDepenetrationVelocity;
 
 	/**
-	 * The recommended fixed timestep for the RBAN solver. Set to 0 to run with variable timestep (default).
+	 * RBAN: The recommended fixed timestep for the RBAN solver. Set to 0 to run with variable timestep (default).
 	 * NOTE: If this value is non-zero and less than the current frame time, the simulation will step multiple times
 	 * which increases the cost.
 	 */
@@ -72,7 +74,7 @@ struct ENGINE_API FPhysicsAssetSolverSettings
 	float FixedTimeStep;
 
 	/**
-	 * Whether to use the linear or non-linear solver for RBAN Joints. The linear solver is significantly cheaper than
+	 * RBAN: Whether to use the linear or non-linear solver for RBAN Joints. The linear solver is significantly cheaper than
 	 * the non-linear solver when you are running multiple iterations, but is more likely to suffer from jitter. 
 	 * In general you should try to use the linear solver and increase the PositionIterations to improve stability if 
 	 * possible, only using the non-linear solver as a last resort.
@@ -153,7 +155,7 @@ enum class EPhysicsAssetSolverType: uint8
  * @see USkeletalMesh
  */
 
-UCLASS(hidecategories=Object, BlueprintType, MinimalAPI, Config=Game, PerObjectConfig)
+UCLASS(hidecategories=Object, BlueprintType, MinimalAPI, Config=Game, PerObjectConfig, AutoCollapseCategories=(OldSolverSettings))
 class UPhysicsAsset : public UObject, public IInterface_PreviewMeshProvider
 {
 	GENERATED_UCLASS_BODY()
@@ -210,11 +212,14 @@ public:
 	FPhysicsAssetSolverSettings SolverSettings;
 
 	/** 
-	 * Solver settings for RBAN simulations with the non-linear RBAN solver.
-	 * These are only used for the non-linear solver which can be enabled with a config variable.
-	 * They will eventually be deprecated and removed.
+	 * Old solver settings shown for reference. These will be removed at some point.
+	 * When you open an old asset you should see that the settings were transferred to "SolverSettings" above. 
+	 * You should usually see:
+	 * SolverSettings.PositionIterations = OldSettings.SolverIterations * OldSetting.JointIterations;
+	 * SolverSettings.VelocityIterations = 1;
+	 * SolverSettings.ProjectionIterations = 1;
 	*/
-	UPROPERTY(EditAnywhere, Category = LegacySolverSettings, Config)
+	UPROPERTY(VisibleAnywhere, Category = OldSolverSettings, Config, Meta = (DisplayName="[Not Used] Old Solver Settings"))
 	FSolverIterations SolverIterations;
 
 	/** 
