Networking changes for PIE seamless travel
Adding cvar net.AllowPIESeamlessTravel, disabled by default

diff --git a/Engine/Source/Editor/UnrealEd/Private/EditorEngine.cpp b/Engine/Source/Editor/UnrealEd/Private/EditorEngine.cpp
--- a/Engine/Source/Editor/UnrealEd/Private/EditorEngine.cpp
+++ b/Engine/Source/Editor/UnrealEd/Private/EditorEngine.cpp
@@ -7066,13 +7066,21 @@ void UEditorEngine::VerifyLoadMapWorldCleanup(FWorldContext* ForContext)
 			for (int32 idx=0; idx < WorldList.Num(); ++idx)
 			{
 				FWorldContext& WorldContext = WorldList[idx];
-				if (World == WorldContext.SeamlessTravelHandler.GetLoadedWorld())
+
+				if (const UWorld* TravelWorld = WorldContext.SeamlessTravelHandler.GetLoadedWorld())
 				{
-					// World valid, but not loaded yet
-					ValidWorld = true;
-					break;
+					TArray<UWorld*> TravelWorlds;
+					EditorLevelUtils::GetWorlds((UWorld*)TravelWorld, TravelWorlds, true, false);
+
+					if (TravelWorlds.Contains(World))
+					{
+						// World valid, but not loaded yet
+						ValidWorld = true;
+						break;
+					}
 				}
-				else if (WorldContext.World())
+				
+				if (WorldContext.World())
 				{
 					TArray<UWorld*> OtherWorlds;
 					EditorLevelUtils::GetWorlds(WorldContext.World(), OtherWorlds, true, false);
diff --git a/Engine/Source/Runtime/Engine/Private/GameModeBase.cpp b/Engine/Source/Runtime/Engine/Private/GameModeBase.cpp
--- a/Engine/Source/Runtime/Engine/Private/GameModeBase.cpp
+++ b/Engine/Source/Runtime/Engine/Private/GameModeBase.cpp
@@ -41,6 +41,17 @@ FGameModeEvents::FGameModePostLoginEvent FGameModeEvents::GameModePostLoginEvent
 FGameModeEvents::FGameModeLogoutEvent FGameModeEvents::GameModeLogoutEvent;
 FGameModeEvents::FGameModeMatchStateSetEvent FGameModeEvents::GameModeMatchStateSetEvent;
 
+namespace UE::GameModeBase::Private
+{
+	static bool bAllowPIESeamlessTravel = false;
+	static FAutoConsoleVariableRef CVarAllowPIESeamlessTravel(
+		TEXT("net.AllowPIESeamlessTravel"),
+		bAllowPIESeamlessTravel,
+		TEXT("When true, allow seamless travels in single process PIE.")
+	);
+}
+
+
 AGameModeBase::AGameModeBase(const FObjectInitializer& ObjectInitializer)
 	: Super(ObjectInitializer.DoNotCreateDefaultSubobject(TEXT("Sprite")))
 {
@@ -405,8 +416,11 @@ bool AGameModeBase::CanServerTravel(const FString& FURL, bool bAbsolute)
 	// There are a few issues with seamless travel using single process PIE, so we're disabling that for now while working on a fix
 	if (World->WorldType == EWorldType::PIE && bUseSeamlessTravel && !FParse::Param(FCommandLine::Get(), TEXT("MultiprocessOSS")))
 	{
-		UE_LOG(LogGameMode, Warning, TEXT("CanServerTravel: Seamless travel currently NOT supported in single process PIE."));
-		return false;
+		if (!UE::GameModeBase::Private::bAllowPIESeamlessTravel)
+		{
+			UE_LOG(LogGameMode, Warning, TEXT("CanServerTravel: Seamless travel currently NOT supported in single process PIE."));
+			return false;
+		}
 	}
 
 	if (FURL.Contains(TEXT("%")))
diff --git a/Engine/Source/Runtime/Engine/Private/NetDriver.cpp b/Engine/Source/Runtime/Engine/Private/NetDriver.cpp
--- a/Engine/Source/Runtime/Engine/Private/NetDriver.cpp
+++ b/Engine/Source/Runtime/Engine/Private/NetDriver.cpp
@@ -1718,7 +1718,9 @@ ENetMode UNetDriver::GetNetMode() const
 #if WITH_EDITOR
 	if (World && World->WorldType == EWorldType::PIE && IsServer())
 	{
-		if ( GEngine->GetWorldContextFromWorldChecked(World).RunAsDedicated )
+		//@todo: world context won't be valid during seamless travel CopyWorldData
+		FWorldContext* WorldContext = GEngine->GetWorldContextFromWorld(World);
+		if (WorldContext && WorldContext->RunAsDedicated)
 		{
 			return NM_DedicatedServer;
 		}
diff --git a/Engine/Source/Runtime/Engine/Private/World.cpp b/Engine/Source/Runtime/Engine/Private/World.cpp
--- a/Engine/Source/Runtime/Engine/Private/World.cpp
+++ b/Engine/Source/Runtime/Engine/Private/World.cpp
@@ -6904,12 +6904,17 @@ void FSeamlessTravelHandler::StartLoadingDestination()
 		if (GIsEditor)
 		{
 			FWorldContext &WorldContext = GEngine->GetWorldContextFromHandleChecked(WorldContextHandle);
+
+			PIEInstanceID = WorldContext.PIEInstance;
+			URLMapPackageName = UWorld::ConvertToPIEPackageName(URLMapPackageName, PIEInstanceID);
+
 			if (WorldContext.WorldType == EWorldType::PIE)
 			{
 				PackageFlags |= PKG_PlayInEditor;
+
+				// Prepare soft object paths for fixup
+				FSoftObjectPath::AddPIEPackageName(FName(*URLMapPackageName));
 			}
-			PIEInstanceID = WorldContext.PIEInstance;
-			URLMapPackageName = UWorld::ConvertToPIEPackageName(URLMapPackageName, PIEInstanceID);
 		}
 #endif
 
