Optimizations to Material_UpdateDeferredCachedUniformExpressions / EvaluatePreshader.  Savings per platform ranges from 44% to 58% (55% on Win64)

* Material parameter lookup uses hash table instead of binary search to find parameters by name.  Custom hash table using 16-bit entries and a 16-bit next index per item to save memory relative to a TMap.
* Parameter key structure optimized to fit in 16 bytes versus 20, saving 4, 8, or 16 bytes per parameter, depending on the alignment of the parameter.  More than pays for memory of hash table, and better for cache perf.
* Most preshader operations now operate in-place on the FPreshaderStack, rather than using intermediate FValues.  Radically reduces moving of memory and bookkeeping costs for operations.
* To support in-place operations, in-place wrappers were added for most shader value operations -- in place means they operate on loose types and arrays of components, and write output in place over the input.
* Other specific preshader operations were ported to work in place, like append and swizzle.  All operations with non-trivial usage based on a sampling profile were ported...
* Type information lookup was optimized using an array rather than switch statement.  The component size was added to the structure (in what was previously alignment padding), avoiding a separate lookup.
* Fast path operations added specific to common types, bypassing conditional logic, casting, and levels of indirection.
* Other minor optimizations here and there (for example, faster version of PopValue for popping the result, using TNonRelocatableInlineAllocator for preshader stack).

#rnx

diff --git a/Engine/Source/Runtime/Engine/Private/Materials/MaterialInstance.cpp b/Engine/Source/Runtime/Engine/Private/Materials/MaterialInstance.cpp
--- a/Engine/Source/Runtime/Engine/Private/Materials/MaterialInstance.cpp
+++ b/Engine/Source/Runtime/Engine/Private/Materials/MaterialInstance.cpp
@@ -316,28 +316,47 @@ void FMaterialInstanceResource::GameThread_UpdateCachedData(const FMaterialInsta
 		});
 }
 
+// Matches GetTypeHash for FMemoryImageMaterialParameterInfo
+static uint32 GetTypeHashLegacy(const FHashedMaterialParameterInfoPacked& Value)
+{
+	return HashCombine(HashCombine(GetTypeHash(Value.Name), (int32)Value.Index), (uint32)Value.Association);
+}
+
 template <typename TInstanceType>
-static bool SortMaterialInstanceParametersPredicate(const FMaterialInstanceResource::TNamedParameter<TInstanceType>& Left, const FMaterialInstanceResource::TNamedParameter<TInstanceType>& Right)
+static bool SortMaterialInstanceParametersPredicate(
+	const typename THashedMaterialParameterMap<TInstanceType>::TNamedParameter& Left,
+	const typename THashedMaterialParameterMap<TInstanceType>::TNamedParameter& Right)
 {
-	return GetTypeHash(Left.Info) < GetTypeHash(Right.Info);
+	// To keep the array sort the same as it has been historically, sort by the legacy type hash, not THashedMaterialParameterMap::TypeHash.
+	// This only matters for duplicate items, where earlier duplicates take precedence over later ones.  Duplicates are exceedingly rare, and
+	// possibly due to bugs, but we want to err on the side of preserving existing behavior.
+	return GetTypeHashLegacy(Left.Info) < GetTypeHashLegacy(Right.Info);
 }
 
 void FMaterialInstanceResource::InitMIParameters(FMaterialInstanceParameterSet& ParameterSet)
 {
 	InvalidateUniformExpressionCache(false);
 
-	// Sort the parameters so that a binary lookup and be used.
+	// Sort the parameters.  Originally this was done so a binary lookup could be used.  We now have a hash table, but we're trying to preserve
+	// the sort order logic to maintain consistent behavior where duplicate items occur.
 	ParameterSet.ScalarParameters.Sort(SortMaterialInstanceParametersPredicate<float>);
 	ParameterSet.VectorParameters.Sort(SortMaterialInstanceParametersPredicate<FLinearColor>);
 	ParameterSet.DoubleVectorParameters.Sort(SortMaterialInstanceParametersPredicate<FVector4d>);
 	ParameterSet.TextureParameters.Sort(SortMaterialInstanceParametersPredicate<const UTexture*>);
 	ParameterSet.RuntimeVirtualTextureParameters.Sort(SortMaterialInstanceParametersPredicate<const URuntimeVirtualTexture*>);
 
-	Swap(ScalarParameterArray, ParameterSet.ScalarParameters);
-	Swap(VectorParameterArray, ParameterSet.VectorParameters);
-	Swap(DoubleVectorParameterArray, ParameterSet.DoubleVectorParameters);
-	Swap(TextureParameterArray, ParameterSet.TextureParameters);
-	Swap(RuntimeVirtualTextureParameterArray, ParameterSet.RuntimeVirtualTextureParameters);
+	Swap(ScalarParameterArray.Array, ParameterSet.ScalarParameters);
+	Swap(VectorParameterArray.Array, ParameterSet.VectorParameters);
+	Swap(DoubleVectorParameterArray.Array, ParameterSet.DoubleVectorParameters);
+	Swap(TextureParameterArray.Array, ParameterSet.TextureParameters);
+	Swap(RuntimeVirtualTextureParameterArray.Array, ParameterSet.RuntimeVirtualTextureParameters);
+
+	// Build hash tables.
+	ScalarParameterArray.HashAddAllItems();
+	VectorParameterArray.HashAddAllItems();
+	DoubleVectorParameterArray.HashAddAllItems();
+	TextureParameterArray.HashAddAllItems();
+	RuntimeVirtualTextureParameterArray.HashAddAllItems();
 }
 
 /**
@@ -3761,11 +3780,18 @@ void UMaterialInstance::GetResourceSizeEx(FResourceSizeEx& CumulativeResourceSiz
 	if (Resource)
 	{
 		CumulativeResourceSize.AddDedicatedSystemMemoryBytes(sizeof(FMaterialInstanceResource));
-		CumulativeResourceSize.AddDedicatedSystemMemoryBytes(ScalarParameterValues.Num() * sizeof(FMaterialInstanceResource::TNamedParameter<float>));
-		CumulativeResourceSize.AddDedicatedSystemMemoryBytes(VectorParameterValues.Num() * sizeof(FMaterialInstanceResource::TNamedParameter<FLinearColor>));
-		CumulativeResourceSize.AddDedicatedSystemMemoryBytes(TextureParameterValues.Num() * sizeof(FMaterialInstanceResource::TNamedParameter<const UTexture*>));
-		CumulativeResourceSize.AddDedicatedSystemMemoryBytes(RuntimeVirtualTextureParameterValues.Num() * sizeof(FMaterialInstanceResource::TNamedParameter<const URuntimeVirtualTexture*>));
-		CumulativeResourceSize.AddDedicatedSystemMemoryBytes(FontParameterValues.Num() * sizeof(FMaterialInstanceResource::TNamedParameter<const UTexture*>));
+		CumulativeResourceSize.AddDedicatedSystemMemoryBytes(ScalarParameterValues.Num() * sizeof(THashedMaterialParameterMap<float>::TNamedParameter));
+		CumulativeResourceSize.AddDedicatedSystemMemoryBytes(VectorParameterValues.Num() * sizeof(THashedMaterialParameterMap<FLinearColor>::TNamedParameter));
+		CumulativeResourceSize.AddDedicatedSystemMemoryBytes(TextureParameterValues.Num() * sizeof(THashedMaterialParameterMap<const UTexture*>::TNamedParameter));
+		CumulativeResourceSize.AddDedicatedSystemMemoryBytes(RuntimeVirtualTextureParameterValues.Num() * sizeof(THashedMaterialParameterMap<const URuntimeVirtualTexture*>::TNamedParameter));
+		CumulativeResourceSize.AddDedicatedSystemMemoryBytes(FontParameterValues.Num() * sizeof(THashedMaterialParameterMap<const UTexture*>::TNamedParameter));
+
+		// Record space for hash tables as well..
+		if (ScalarParameterValues.Num()) CumulativeResourceSize.AddDedicatedSystemMemoryBytes(FDefaultSetAllocator::GetNumberOfHashBuckets(ScalarParameterValues.Num()) * sizeof(uint16));
+		if (VectorParameterValues.Num()) CumulativeResourceSize.AddDedicatedSystemMemoryBytes(FDefaultSetAllocator::GetNumberOfHashBuckets(VectorParameterValues.Num()) * sizeof(uint16));
+		if (TextureParameterValues.Num()) CumulativeResourceSize.AddDedicatedSystemMemoryBytes(FDefaultSetAllocator::GetNumberOfHashBuckets(TextureParameterValues.Num()) * sizeof(uint16));
+		if (RuntimeVirtualTextureParameterValues.Num()) CumulativeResourceSize.AddDedicatedSystemMemoryBytes(FDefaultSetAllocator::GetNumberOfHashBuckets(RuntimeVirtualTextureParameterValues.Num()) * sizeof(uint16));
+		if (FontParameterValues.Num()) CumulativeResourceSize.AddDedicatedSystemMemoryBytes(FDefaultSetAllocator::GetNumberOfHashBuckets(FontParameterValues.Num()) * sizeof(uint16));
 	}
 }
 
diff --git a/Engine/Source/Runtime/Engine/Private/Materials/MaterialInstanceSupport.h b/Engine/Source/Runtime/Engine/Private/Materials/MaterialInstanceSupport.h
--- a/Engine/Source/Runtime/Engine/Private/Materials/MaterialInstanceSupport.h
+++ b/Engine/Source/Runtime/Engine/Private/Materials/MaterialInstanceSupport.h
@@ -60,20 +60,148 @@ private:
 };
 
 /**
-* The resource used to render a UMaterialInstance.
-*/
-class FMaterialInstanceResource: public FMaterialRenderProxy
+ * Structure that packs FHashedMaterialParameterInfo plus a "HashNext" value into 16 bytes to save memory.  Allows key/value
+ * pairs of 16-byte aligned items like vectors to fit in 32 bytes, rather 48, saving significant space on padding.  This
+ * introduces limitations that Index cannot exceed 32767, and a given material can't have more than 16384 named parameters
+ * of a given type.  Both of those numbers are at least a couple orders of magnitude larger than seen in practice.
+ */
+struct FHashedMaterialParameterInfoPacked
+{
+	FScriptName	Name;
+	int16		Index = 0;
+	uint16		HashNext : 14;
+	uint16		Association : 2;
+
+	FHashedMaterialParameterInfoPacked()
+	{
+		HashNext = 0;
+		Association = 0;
+	}
+
+	FORCEINLINE FHashedMaterialParameterInfoPacked(const struct FHashedMaterialParameterInfoPacked& Rhs)
+		: Name(Rhs.Name), Index(Rhs.Index), HashNext(0), Association(Rhs.Association)
+	{
+	}
+
+	FHashedMaterialParameterInfoPacked& operator=(const FHashedMaterialParameterInfo& InInfo)
+	{
+		check((InInfo.Index >= std::numeric_limits<int16>::min()) && (InInfo.Index <= std::numeric_limits<int16>::max()));
+		check(InInfo.Association <= 3);
+
+		Name = InInfo.Name;
+		Index = (int16)InInfo.Index;
+		Association = InInfo.Association;
+		return *this;
+	}
+
+	/** HashNext ignored in equality comparisons */
+	bool operator==(const FHashedMaterialParameterInfoPacked& Other) const
+	{
+		return Name == Other.Name && Index == Other.Index && Association == Other.Association;
+	}
+
+	bool operator==(const FHashedMaterialParameterInfo& Other) const
+	{
+		return Name == Other.Name && Index == Other.Index && Association == Other.Association;
+	}
+};
+
+template <typename ValueType>
+class THashedMaterialParameterMap
 {
 public:
 
 	/** Material instances store pairs of names and values in arrays to look up parameters at run time. */
-	template <typename ValueType>
 	struct TNamedParameter
 	{
-		FHashedMaterialParameterInfo Info;
+		FHashedMaterialParameterInfoPacked Info;
 		ValueType Value;
 	};
 
+	TArray<TNamedParameter> Array;
+	TArray<uint16> HashTable;				// Indices are plus one, so zero represents invalid index
+
+	template<typename ParameterInfo>
+	static uint32 TypeHash(const ParameterInfo& Info)
+	{
+		// We expect the Name to be the only part of the key to vary in 99.995% of keys.  If this changes, we could make
+		// the hash take into account more fields than just Name, at some cost in perf generating the hash.  It's important
+		// to mix the key bits, because GetTypeHash(FScriptName) just returns an index, which is not good for a hash table.
+		int32 Key = GetTypeHash(Info.Name);
+		Key += ~(Key << 15);
+		Key ^= (Key >> 10);
+		Key += (Key << 3);
+		Key ^= (Key >> 6);
+		Key += ~(Key << 11);
+		Key ^= (Key >> 16);
+		return (uint32)Key;
+	}
+
+	// Add an item from the Array to the hash table.  Hash table must already be allocated.
+	void HashAddOneItem(int32 AddedIndex)
+	{
+		// Clear next index -- we will be adding the item at the end of the bucket it ends up in.
+		Array[AddedIndex].Info.HashNext = 0;
+
+		check(!HashTable.IsEmpty());
+		uint32 HashIndex = THashedMaterialParameterMap::TypeHash(Array[AddedIndex].Info) & (HashTable.Num() - 1);
+		int32 ItemIndex = HashTable[HashIndex];
+
+		if (!ItemIndex)
+		{
+			// Add to root table.  Indices are plus one, so zero represents invalid index.
+			HashTable[HashIndex] = AddedIndex + 1;
+		}
+		else
+		{
+			// Items in table are incremented by one, get the actual item index
+			ItemIndex--;
+
+			// Find the last item in the chain
+			for (int32 NextItem = Array[ItemIndex].Info.HashNext; NextItem; NextItem = Array[NextItem].Info.HashNext)
+			{
+				ItemIndex = NextItem;
+			}
+			Array[ItemIndex].Info.HashNext = (uint16)AddedIndex;
+		}
+	}
+
+	// Add all items to hash table.  Can be used to rehash the full array.
+	void HashAddAllItems(int32 NumHashBuckets = 0)
+	{
+		// If zero is passed in (or anything invalid), choose a default size
+		if (NumHashBuckets <= 0)
+		{
+			NumHashBuckets = FDefaultSetAllocator::GetNumberOfHashBuckets(Array.Num());
+		}
+		if (Array.IsEmpty())
+		{
+			HashTable.Empty();
+		}
+		else
+		{
+			HashTable.SetNumUninitialized(NumHashBuckets);
+			FMemory::Memset(HashTable.GetData(), 0, HashTable.Num() * HashTable.GetTypeSize());
+			for (int32 Index = 0; Index < Array.Num(); Index++)
+			{
+				HashAddOneItem(Index);
+			}
+		}
+	}
+
+	void Empty()
+	{
+		Array.Empty();
+		HashTable.Empty();
+	}
+};
+
+/**
+* The resource used to render a UMaterialInstance.
+*/
+class FMaterialInstanceResource: public FMaterialRenderProxy
+{
+public:
 	/** Initialization constructor. */
 	FMaterialInstanceResource(UMaterialInstance* InOwner);
 
@@ -129,20 +257,32 @@ public:
 		LLM_SCOPE(ELLMTag::MaterialInstance);
 
 		InvalidateUniformExpressionCache(false);
-		bool bWasFound;
-		TArray<TNamedParameter<ValueType> >& ValueArray = GetValueArray<ValueType>();
-		int Index = RenderThread_FindParameterByNameInternal<ValueType>(ParameterInfo, bWasFound);
+		THashedMaterialParameterMap<ValueType>& ValueArray = GetValueArray<ValueType>();
+		int32 Index = RenderThread_FindParameterByNameInternal<ValueType>(ParameterInfo);
 
-		if (bWasFound)
+		if (Index != INDEX_NONE)
 		{
-			ValueArray[Index].Value = Value;
+			ValueArray.Array[Index].Value = Value;
 		}
 		else
 		{
-			TNamedParameter<ValueType> NewParameter;
+			typename THashedMaterialParameterMap<ValueType>::TNamedParameter NewParameter;
 			NewParameter.Info = ParameterInfo;
 			NewParameter.Value = Value;
-			ValueArray.Insert(NewParameter, Index);
+			ValueArray.Array.Add(NewParameter);
+			
+			// FHashedMaterialParameterInfoPacked::HashNext is limited to 14 bits
+			check(ValueArray.Array.Num() <= (1<<14));
+
+			const int32 NumHashBuckets = FDefaultSetAllocator::GetNumberOfHashBuckets(ValueArray.Array.Num());
+			if (NumHashBuckets == ValueArray.HashTable.Num())
+			{
+				ValueArray.HashAddOneItem(ValueArray.Array.Num() - 1);
+			}
+			else
+			{
+				ValueArray.HashAddAllItems(NumHashBuckets);
+			}
 		}
 	}
 
@@ -152,13 +292,12 @@ public:
 	template <typename ValueType>
 	bool RenderThread_GetParameterValue(const FHashedMaterialParameterInfo& ParameterInfo, FMaterialParameterValue& OutValue) const
 	{
-		bool bWasFound;
-		const TArray<TNamedParameter<ValueType> >& ValueArray = GetValueArray<ValueType>();
-		int Index = RenderThread_FindParameterByNameInternal<ValueType>(ParameterInfo, bWasFound);
+		const THashedMaterialParameterMap<ValueType>& ValueArray = GetValueArray<ValueType>();
+		int Index = RenderThread_FindParameterByNameInternal<ValueType>(ParameterInfo);
 
-		if (bWasFound && IsValidParameterValue(ValueArray[Index].Value))
+		if (Index != INDEX_NONE && IsValidParameterValue(ValueArray.Array[Index].Value))
 		{
-			OutValue = ValueArray[Index].Value;
+			OutValue = ValueArray.Array[Index].Value;
 			return true;
 		}
 		return false;
@@ -168,8 +307,8 @@ private:
 	/**
 	 * Retrieves the array of values for a given type.
 	 */
-	template <typename ValueType> TArray<TNamedParameter<ValueType> >& GetValueArray() { return ScalarParameterArray; }
-	template <typename ValueType> const TArray<TNamedParameter<ValueType> >& GetValueArray() const { return ScalarParameterArray; }
+	template <typename ValueType> THashedMaterialParameterMap<ValueType>& GetValueArray() { return ScalarParameterArray; }
+	template <typename ValueType> const THashedMaterialParameterMap<ValueType>& GetValueArray() const { return ScalarParameterArray; }
 
 	static bool IsValidParameterValue(float) { return true; }
 	static bool IsValidParameterValue(const FLinearColor&) { return true; }
@@ -177,36 +316,67 @@ private:
 	static bool IsValidParameterValue(const UTexture* Value) { return Value != nullptr; }
 	static bool IsValidParameterValue(const URuntimeVirtualTexture* Value) { return Value != nullptr; }
 
+	// Validation logic that exhaustively searches the array, to compare with what the hash table logic finds
+	#define HASHED_MATERIAL_PARAMETER_MAP_VALIDATE 0
+
 	template <typename ValueType>
-	int RenderThread_FindParameterByNameInternal(const FHashedMaterialParameterInfo& ParameterInfo, bool& OutWasFound) const
+	int32 RenderThread_FindParameterByNameInternal(const FHashedMaterialParameterInfo& ParameterInfo) const
 	{
-		const TArray<TNamedParameter<ValueType> >& ValueArray = GetValueArray<ValueType>();
-
-		TNamedParameter<ValueType> SearchParam;
-		SearchParam.Info = ParameterInfo;
+		const THashedMaterialParameterMap<ValueType>& ValueArray = GetValueArray<ValueType>();
 
-		int Index = Algo::LowerBound(ValueArray, SearchParam,
-			[](const TNamedParameter<ValueType>& Left, const TNamedParameter<ValueType>& Right)
+#if HASHED_MATERIAL_PARAMETER_MAP_VALIDATE
+		int32 ValidateItemIndex = INDEX_NONE;
+		for (int32 Index = 0; Index < ValueArray.Array.Num(); Index++)
+		{
+			if (ValueArray.Array[Index].Info == ParameterInfo)
 			{
-				return GetTypeHash(Left.Info) < GetTypeHash(Right.Info);
-			});
-
-		uint32 SearchHash = GetTypeHash(ParameterInfo);
+				ValidateItemIndex = Index;
+				break;
+			}
+		}
+#endif  // HASHED_MATERIAL_PARAMETER_MAP_VALIDATE
 
-		while (ValueArray.IsValidIndex(Index) && GetTypeHash(ValueArray[Index].Info) == SearchHash)
+		uint32 HashTableSize = ValueArray.HashTable.Num();
+		if (HashTableSize)
 		{
-			if (ValueArray[Index].Info == ParameterInfo)
+			uint32 HashValue = THashedMaterialParameterMap<ValueType>::TypeHash(ParameterInfo);
+			uint32 HashIndex = HashValue & (HashTableSize - 1);
+			int32 ItemIndex = (int32)ValueArray.HashTable[HashIndex];
+
+			if (ItemIndex)
 			{
-				OutWasFound = true;
-				return Index;
-			}
+				// Items in table are incremented by one, get the actual item index
+				ItemIndex--;
+				do
+				{
+					if (ValueArray.Array[ItemIndex].Info == ParameterInfo)
+					{
+#if HASHED_MATERIAL_PARAMETER_MAP_VALIDATE
+						check(ValidateItemIndex == ItemIndex);
+#endif
+						return ItemIndex;
+					}
+
+					// Additional items in a chain must have a non-zero index, because items in the array are linked to buckets
+					// in order, so the zeroeth item in the array must be first in its bucket, and zero can't be a valid next index.
+					ItemIndex = (int32)ValueArray.Array[ItemIndex].Info.HashNext;
 
-			Index++;
+				} while (ItemIndex);
+			}
+		}
+		else
+		{
+			// No hash table -- item array should also be empty
+			check(ValueArray.Array.IsEmpty());
 		}
 
-		OutWasFound = false;
-		return Index;
+#if HASHED_MATERIAL_PARAMETER_MAP_VALIDATE
+		check(ValidateItemIndex == INDEX_NONE);
+#endif
+		return INDEX_NONE;
 	}
+	
+	#undef HASHED_MATERIAL_PARAMETER_MAP_VALIDATE
 
 	/** The parent of the material instance. */
 	UMaterialInterface* Parent;
@@ -218,37 +388,37 @@ private:
 	UMaterialInterface* GameThreadParent;
 	
 	/** Vector parameters for this material instance. */
-	TArray<TNamedParameter<FLinearColor> > VectorParameterArray;
+	THashedMaterialParameterMap<FLinearColor> VectorParameterArray;
 	/** DoubleVector parameters for this material instance. */
-	TArray<TNamedParameter<FVector4d> > DoubleVectorParameterArray;
+	THashedMaterialParameterMap<FVector4d> DoubleVectorParameterArray;
 	/** Scalar parameters for this material instance. */
-	TArray<TNamedParameter<float> > ScalarParameterArray;
+	THashedMaterialParameterMap<float> ScalarParameterArray;
 	/** Texture parameters for this material instance. */
-	TArray<TNamedParameter<const UTexture*> > TextureParameterArray;
+	THashedMaterialParameterMap<const UTexture*> TextureParameterArray;
 	/** Runtime Virtual Texture parameters for this material instance. */
-	TArray<TNamedParameter<const URuntimeVirtualTexture*> > RuntimeVirtualTextureParameterArray; 
+	THashedMaterialParameterMap<const URuntimeVirtualTexture*> RuntimeVirtualTextureParameterArray;
 	/** Remap layer indices for parent */
 	TArray<int32> ParentLayerIndexRemap;
 };
 
-template <> FORCEINLINE TArray<FMaterialInstanceResource::TNamedParameter<float> >& FMaterialInstanceResource::GetValueArray() { return ScalarParameterArray; }
-template <> FORCEINLINE TArray<FMaterialInstanceResource::TNamedParameter<FLinearColor> >& FMaterialInstanceResource::GetValueArray() { return VectorParameterArray; }
-template <> FORCEINLINE TArray<FMaterialInstanceResource::TNamedParameter<FVector4d> >& FMaterialInstanceResource::GetValueArray() { return DoubleVectorParameterArray; }
-template <> FORCEINLINE TArray<FMaterialInstanceResource::TNamedParameter<const UTexture*> >& FMaterialInstanceResource::GetValueArray() { return TextureParameterArray; }
-template <> FORCEINLINE TArray<FMaterialInstanceResource::TNamedParameter<const URuntimeVirtualTexture*> >& FMaterialInstanceResource::GetValueArray() { return RuntimeVirtualTextureParameterArray; }
-template <> FORCEINLINE const TArray<FMaterialInstanceResource::TNamedParameter<float> >& FMaterialInstanceResource::GetValueArray() const { return ScalarParameterArray; }
-template <> FORCEINLINE const TArray<FMaterialInstanceResource::TNamedParameter<FLinearColor> >& FMaterialInstanceResource::GetValueArray() const { return VectorParameterArray; }
-template <> FORCEINLINE const TArray<FMaterialInstanceResource::TNamedParameter<FVector4d> >& FMaterialInstanceResource::GetValueArray() const { return DoubleVectorParameterArray; }
-template <> FORCEINLINE const TArray<FMaterialInstanceResource::TNamedParameter<const UTexture*> >& FMaterialInstanceResource::GetValueArray() const { return TextureParameterArray; }
-template <> FORCEINLINE const TArray<FMaterialInstanceResource::TNamedParameter<const URuntimeVirtualTexture*> >& FMaterialInstanceResource::GetValueArray() const { return RuntimeVirtualTextureParameterArray; }
+template <> FORCEINLINE THashedMaterialParameterMap<float>& FMaterialInstanceResource::GetValueArray() { return ScalarParameterArray; }
+template <> FORCEINLINE THashedMaterialParameterMap<FLinearColor>& FMaterialInstanceResource::GetValueArray() { return VectorParameterArray; }
+template <> FORCEINLINE THashedMaterialParameterMap<FVector4d>& FMaterialInstanceResource::GetValueArray() { return DoubleVectorParameterArray; }
+template <> FORCEINLINE THashedMaterialParameterMap<const UTexture*>& FMaterialInstanceResource::GetValueArray() { return TextureParameterArray; }
+template <> FORCEINLINE THashedMaterialParameterMap<const URuntimeVirtualTexture*>& FMaterialInstanceResource::GetValueArray() { return RuntimeVirtualTextureParameterArray; }
+template <> FORCEINLINE const THashedMaterialParameterMap<float>& FMaterialInstanceResource::GetValueArray() const { return ScalarParameterArray; }
+template <> FORCEINLINE const THashedMaterialParameterMap<FLinearColor>& FMaterialInstanceResource::GetValueArray() const { return VectorParameterArray; }
+template <> FORCEINLINE const THashedMaterialParameterMap<FVector4d>& FMaterialInstanceResource::GetValueArray() const { return DoubleVectorParameterArray; }
+template <> FORCEINLINE const THashedMaterialParameterMap<const UTexture*>& FMaterialInstanceResource::GetValueArray() const { return TextureParameterArray; }
+template <> FORCEINLINE const THashedMaterialParameterMap<const URuntimeVirtualTexture*>& FMaterialInstanceResource::GetValueArray() const { return RuntimeVirtualTextureParameterArray; }
 
 struct FMaterialInstanceParameterSet
 {
-	TArray<FMaterialInstanceResource::TNamedParameter<float> > ScalarParameters;
-	TArray<FMaterialInstanceResource::TNamedParameter<FLinearColor> > VectorParameters;
-	TArray<FMaterialInstanceResource::TNamedParameter<FVector4d> > DoubleVectorParameters;
-	TArray<FMaterialInstanceResource::TNamedParameter<const UTexture*> > TextureParameters;
-	TArray<FMaterialInstanceResource::TNamedParameter<const URuntimeVirtualTexture*> > RuntimeVirtualTextureParameters;
+	TArray<THashedMaterialParameterMap<float>::TNamedParameter>							ScalarParameters;
+	TArray<THashedMaterialParameterMap<FLinearColor>::TNamedParameter>					VectorParameters;
+	TArray<THashedMaterialParameterMap<FVector4d>::TNamedParameter>						DoubleVectorParameters;
+	TArray<THashedMaterialParameterMap<const UTexture*>::TNamedParameter>				TextureParameters;
+	TArray<THashedMaterialParameterMap<const URuntimeVirtualTexture*>::TNamedParameter>	RuntimeVirtualTextureParameters;
 };
 	
 /** Finds a parameter by name from the game thread. */
diff --git a/Engine/Source/Runtime/Engine/Private/Materials/MaterialUniformExpressions.cpp b/Engine/Source/Runtime/Engine/Private/Materials/MaterialUniformExpressions.cpp
--- a/Engine/Source/Runtime/Engine/Private/Materials/MaterialUniformExpressions.cpp
+++ b/Engine/Source/Runtime/Engine/Private/Materials/MaterialUniformExpressions.cpp
@@ -749,90 +749,124 @@ void FUniformExpressionSet::FillUniformBuffer(const FMaterialRenderContext& Mate
 			FPreshaderDataContext PreshaderContext(PreshaderBaseContext, Preshader.OpcodeOffset, Preshader.OpcodeSize);
 			const FPreshaderValue Result = EvaluatePreshader(this, MaterialRenderContext, PreshaderStack, PreshaderContext);
 
-			for (uint32 FieldIndex = 0u; FieldIndex < Preshader.NumFields; ++FieldIndex)
+			// Fast path for non-structure float1 to float4 fields, which represents the vast majority of cases
+			EValueType FirstFieldType = UniformPreshaderFields[Preshader.FieldIndex].Type;
+			EValueType ResultType = Result.Type.ValueType;
+
+			// This min-max logic assumes that Float1 through Float4 are sequential in the EValueType enum, so we can do just two comparisons
+			// to detect if both Result and the destination field are float types.
+			static_assert((int32)EValueType::Float2 == (int32)EValueType::Float1 + 1);
+			static_assert((int32)EValueType::Float3 == (int32)EValueType::Float2 + 1);
+			static_assert((int32)EValueType::Float4 == (int32)EValueType::Float3 + 1);
+			
+			int32 MinValueType = FMath::Min((int32)FirstFieldType, (int32)ResultType);
+			int32 MaxValueType = FMath::Max((int32)FirstFieldType, (int32)ResultType);
+
+			if (Preshader.NumFields == 1 && MinValueType >= (int32)EValueType::Float1 && MaxValueType <= (int32)EValueType::Float4)
 			{
-				const FMaterialUniformPreshaderField& PreshaderField = UniformPreshaderFields[Preshader.FieldIndex + FieldIndex];
-				const FValueTypeDescription UniformTypeDesc = GetValueTypeDescription(PreshaderField.Type);
-				const int32 NumFieldComponents = UniformTypeDesc.NumComponents;
-
-				const int32 NumResultComponents = FMath::Min<int32>(NumFieldComponents, Result.Component.Num() - PreshaderField.ComponentIndex);
-				check(NumResultComponents > 0);
-
-				// The type generated by the preshader might not match the expected type
-				// In the future, with new HLSLTree, preshader could potentially include explicit cast opcodes, and avoid implicit conversions
-				// Making this change is difficult while also maintaining backwards compatibility however
-				FValue FieldValue(Result.Type.GetComponentType(PreshaderField.ComponentIndex), NumResultComponents);
-				for (int32 i = 0; i < NumResultComponents; ++i)
+				float* DestAddress = PreshaderBuffer + UniformPreshaderFields[Preshader.FieldIndex].BufferOffset;
+
+				// Copy components that exist in both result and destination buffer
+				int32 NumCopyComponents = MinValueType - ((int32)EValueType::Float1 - 1);
+				for (int32 i = 0; i < NumCopyComponents; ++i)
 				{
-					FieldValue.Component[i] = Result.Component[PreshaderField.ComponentIndex + i];
+					*DestAddress++ = Result.Component[i].Float;
 				}
 
-				if (UniformTypeDesc.ComponentType == EValueComponentType::Float)
+				// Zero out any additional components that exist in the destination buffer
+				int32 NumDestComponents = (int32)FirstFieldType - ((int32)EValueType::Float1 - 1);
+				for (int32 i = NumCopyComponents; i < NumDestComponents; ++i)
 				{
-					const FFloatValue FloatValue = FieldValue.AsFloat();
-					float* DestAddress = PreshaderBuffer + PreshaderField.BufferOffset;
-					for (int32 i = 0; i < NumFieldComponents; ++i)
-					{
-						*DestAddress++ = FloatValue[i];
-					}
+					*DestAddress++ = 0.0f;
 				}
-				else if (UniformTypeDesc.ComponentType == EValueComponentType::Int)
+			}
+			else
+			{
+				for (uint32 FieldIndex = 0u; FieldIndex < Preshader.NumFields; ++FieldIndex)
 				{
-					const FIntValue IntValue = FieldValue.AsInt();
-					int32* DestAddress = (int32*)PreshaderBuffer + PreshaderField.BufferOffset;
-					for (int32 i = 0; i < NumFieldComponents; ++i)
+					const FMaterialUniformPreshaderField& PreshaderField = UniformPreshaderFields[Preshader.FieldIndex + FieldIndex];
+					const FValueTypeDescription UniformTypeDesc = GetValueTypeDescription(PreshaderField.Type);
+					const int32 NumFieldComponents = UniformTypeDesc.NumComponents;
+
+					const int32 NumResultComponents = FMath::Min<int32>(NumFieldComponents, Result.Component.Num() - PreshaderField.ComponentIndex);
+					check(NumResultComponents > 0);
+
+					// The type generated by the preshader might not match the expected type
+					// In the future, with new HLSLTree, preshader could potentially include explicit cast opcodes, and avoid implicit conversions
+					// Making this change is difficult while also maintaining backwards compatibility however
+					FValue FieldValue(Result.Type.GetComponentType(PreshaderField.ComponentIndex), NumResultComponents);
+					for (int32 i = 0; i < NumResultComponents; ++i)
 					{
-						*DestAddress++ = IntValue[i];
+						FieldValue.Component[i] = Result.Component[PreshaderField.ComponentIndex + i];
 					}
-				}
-				else if (UniformTypeDesc.ComponentType == EValueComponentType::Bool)
-				{
-					const FBoolValue BoolValue = FieldValue.AsBool();
-					uint32 Mask = 0u;
-					for (int32 i = 0; i < NumFieldComponents; ++i)
+
+					if (UniformTypeDesc.ComponentType == EValueComponentType::Float)
 					{
-						if (BoolValue[i])
+						const FFloatValue FloatValue = FieldValue.AsFloat();
+						float* DestAddress = PreshaderBuffer + PreshaderField.BufferOffset;
+						for (int32 i = 0; i < NumFieldComponents; ++i)
 						{
-							Mask |= (1u << i);
+							*DestAddress++ = FloatValue[i];
 						}
 					}
-
-					const uint32 BufferOffset = PreshaderField.BufferOffset / 32u;
-					const uint32 BufferBitOffset = PreshaderField.BufferOffset % 32u;
-					check(BufferBitOffset + NumFieldComponents <= 32u);
-
-					uint32* DestAddress = (uint32*)PreshaderBuffer + BufferOffset;
-					if (BufferBitOffset == 0u)
+					else if (UniformTypeDesc.ComponentType == EValueComponentType::Int)
 					{
-						// First update to a group of bits needs to initialize memory
-						*DestAddress = Mask;
+						const FIntValue IntValue = FieldValue.AsInt();
+						int32* DestAddress = (int32*)PreshaderBuffer + PreshaderField.BufferOffset;
+						for (int32 i = 0; i < NumFieldComponents; ++i)
+						{
+							*DestAddress++ = IntValue[i];
+						}
 					}
-					else
+					else if (UniformTypeDesc.ComponentType == EValueComponentType::Bool)
 					{
-						// Combine with any previous bits
-						*DestAddress |= (Mask << BufferBitOffset);
+						const FBoolValue BoolValue = FieldValue.AsBool();
+						uint32 Mask = 0u;
+						for (int32 i = 0; i < NumFieldComponents; ++i)
+						{
+							if (BoolValue[i])
+							{
+								Mask |= (1u << i);
+							}
+						}
+
+						const uint32 BufferOffset = PreshaderField.BufferOffset / 32u;
+						const uint32 BufferBitOffset = PreshaderField.BufferOffset % 32u;
+						check(BufferBitOffset + NumFieldComponents <= 32u);
+
+						uint32* DestAddress = (uint32*)PreshaderBuffer + BufferOffset;
+						if (BufferBitOffset == 0u)
+						{
+							// First update to a group of bits needs to initialize memory
+							*DestAddress = Mask;
+						}
+						else
+						{
+							// Combine with any previous bits
+							*DestAddress |= (Mask << BufferBitOffset);
+						}
 					}
-				}
-				else if (UniformTypeDesc.ComponentType == EValueComponentType::Double)
-				{
-					const FDoubleValue DoubleValue = FieldValue.AsDouble();
+					else if (UniformTypeDesc.ComponentType == EValueComponentType::Double)
+					{
+						const FDoubleValue DoubleValue = FieldValue.AsDouble();
 
-					float TileValue[4];
-					float OffsetValue[4];
-					for (int32 i = 0; i < NumFieldComponents; ++i)
+						float TileValue[4];
+						float OffsetValue[4];
+						for (int32 i = 0; i < NumFieldComponents; ++i)
+						{
+							const FLargeWorldRenderScalar Value(DoubleValue[i]);
+							TileValue[i] = Value.GetTile();
+							OffsetValue[i] = Value.GetOffset();
+						}
+
+						float* DestAddress = PreshaderBuffer + PreshaderField.BufferOffset;
+						for (int32 i = 0; i < NumFieldComponents; ++i) *DestAddress++ = TileValue[i];
+						for (int32 i = 0; i < NumFieldComponents; ++i) *DestAddress++ = OffsetValue[i];
+					}
+					else
 					{
-						const FLargeWorldRenderScalar Value(DoubleValue[i]);
-						TileValue[i] = Value.GetTile();
-						OffsetValue[i] = Value.GetOffset();
+						ensure(false);
 					}
-
-					float* DestAddress = PreshaderBuffer + PreshaderField.BufferOffset;
-					for (int32 i = 0; i < NumFieldComponents; ++i) *DestAddress++ = TileValue[i];
-					for (int32 i = 0; i < NumFieldComponents; ++i) *DestAddress++ = OffsetValue[i];
-				}
-				else
-				{
-					ensure(false);
 				}
 			}
 		}
diff --git a/Engine/Source/Runtime/Engine/Private/Shader/Preshader.cpp b/Engine/Source/Runtime/Engine/Private/Shader/Preshader.cpp
--- a/Engine/Source/Runtime/Engine/Private/Shader/Preshader.cpp
+++ b/Engine/Source/Runtime/Engine/Private/Shader/Preshader.cpp
@@ -37,7 +37,7 @@ EValueComponentType FPreshaderType::GetComponentType(int32 Index) const
 	}
 	else
 	{
-		const FValueTypeDescription TypeDesc = GetValueTypeDescription(ValueType);
+		const FValueTypeDescription& TypeDesc = GetValueTypeDescription(ValueType);
 		return (Index >= 0 && Index < TypeDesc.NumComponents) ? TypeDesc.ComponentType : EValueComponentType::Void;
 	}
 }
@@ -71,6 +71,14 @@ TArrayView<FValueComponent> FPreshaderStack::PushEmptyValue(const FPreshaderType
 	return MakeArrayView(Components.GetData() + ComponentIndex, NumComponents);
 }
 
+FValueComponent* FPreshaderStack::PushEmptyValue(EValueType InType, int32 NumComponents)
+{
+	Values.Add(InType);
+	checkSlow(NumComponents == Values.Last().GetNumComponents());
+	const int32 ComponentIndex = Components.AddZeroed(NumComponents);
+	return Components.GetData() + ComponentIndex;
+}
+
 FPreshaderValue FPreshaderStack::PopValue()
 {
 	FPreshaderValue Value;
@@ -85,13 +93,36 @@ FPreshaderValue FPreshaderStack::PopValue()
 	return Value;
 }
 
+void FPreshaderStack::PopResult(FPreshaderValue& OutValue)
+{
+	if (Values.Num() > 0)
+	{
+		check(Values.Num() == 1);
+		check(Values[0].GetNumComponents() == Components.Num());
+
+		OutValue.Type = Values[0];
+		OutValue.Component = Components;
+
+		// Call Reset instead of Empty, so it doesn't change the memory allocation
+		Values.Reset();
+		Components.Reset();
+	}
+}
+
 FPreshaderValue FPreshaderStack::PeekValue(int32 Offset)
 {
 	FPreshaderValue Value;
 	Value.Type = Values.Last(Offset);
 
+	int32 OffsetComponents = 0;
+	for (int32 OffsetIndex = 0; OffsetIndex < Offset; OffsetIndex++)
+	{
+		OffsetComponents += Values.Last(OffsetIndex).GetNumComponents();
+	}
+
 	const int32 NumComponents = Value.Type.GetNumComponents();
-	const int32 ComponentIndex = Components.Num() - NumComponents;
+	const int32 ComponentIndex = Components.Num() - NumComponents - OffsetComponents;
+	check(ComponentIndex >= 0);
 	Value.Component = MakeArrayView(Components.GetData() + ComponentIndex, NumComponents);
 	return Value;
 }
@@ -295,11 +326,24 @@ static void EvaluateConstant(FPreshaderStack& Stack, FPreshaderDataContext& REST
 {
 	const FPreshaderType Type = ReadPreshaderValue<FPreshaderType>(Data);
 	TArrayView<FValueComponent> Component = Stack.PushEmptyValue(Type);
-	for (int32 Index = 0; Index < Component.Num(); ++Index)
+	if (!Type.IsStruct())
+	{
+		// Common case:  not a struct
+		const FValueTypeDescription& TypeDesc = GetValueTypeDescription(Type.ValueType);
+		const int32 ComponmentSizeInBytes = TypeDesc.ComponentSizeInBytes;
+		for (int32 Index = 0; Index < Component.Num(); ++Index)
+		{
+			ReadPreshaderData(Data, ComponmentSizeInBytes, &Component[Index].Packed);
+		}
+	}
+	else
 	{
-		const EValueComponentType ComponentType = Type.GetComponentType(Index);
-		const int32 ComponmentSizeInBytes = GetComponentTypeSizeInBytes(ComponentType);
-		ReadPreshaderData(Data, ComponmentSizeInBytes, &Component[Index].Packed);
+		for (int32 Index = 0; Index < Component.Num(); ++Index)
+		{
+			const EValueComponentType ComponentType = Type.GetComponentType(Index);
+			const int32 ComponmentSizeInBytes = GetComponentTypeSizeInBytes(ComponentType);
+			ReadPreshaderData(Data, ComponmentSizeInBytes, &Component[Index].Packed);
+		}
 	}
 }
 
@@ -412,7 +456,133 @@ static void EvaluateParameter(FPreshaderStack& Stack, const FUniformExpressionSe
 	// Default value
 	if (!bFoundParameter)
 	{
-		Stack.PushValue(UniformExpressionSet->GetDefaultParameterValue(Parameter.ParameterType, Parameter.DefaultValueOffset));
+		// Fast path for numeric parameters (common case).  Writes data in place to FPreshaderStack, and has a single
+		// lookup table to fetch type and component information.
+		if (Parameter.ParameterType <= EMaterialParameterType::DoubleVector)
+		{
+			// Validate that it's OK to do the comparison above, and that our table is correct
+			static_assert((int32)EMaterialParameterType::Scalar == 0);
+			static_assert((int32)EMaterialParameterType::Vector == 1);
+			static_assert((int32)EMaterialParameterType::DoubleVector == 2);
+			static_assert(sizeof(FValueComponent) == sizeof(double));
+
+			// Lookup includes type, number of components, and size to copy per component.
+			static const int8 GParameterCopyLookup[3][3] =
+			{
+				{ (int8)EValueType::Float1, 1, sizeof(float) },		// EMaterialParameterType::Scalar
+				{ (int8)EValueType::Float4, 4, sizeof(float) },		// EMaterialParameterType::Vector
+				{ (int8)EValueType::Double4, 4, sizeof(double) },	// EMaterialParameterType::DoubleVector
+			};
+			const int8* ParameterLookup = GParameterCopyLookup[(int32)Parameter.ParameterType];
+			int32 NumComponents = ParameterLookup[1];
+			int32 ComponentSizeInBytes = ParameterLookup[2];
+
+			FValueComponent* Components = Stack.PushEmptyValue((EValueType)ParameterLookup[0], NumComponents);
+			const uint8* ParameterData = UniformExpressionSet->GetDefaultParameterData(Parameter.DefaultValueOffset);
+
+			// Would be nice to be able to do hard coded copies instead of memcpy per component, but it's possible
+			// the source data may not be aligned.
+			for (int32 ComponentIndex = 0; ComponentIndex < NumComponents; ComponentIndex++)
+			{
+				FMemory::Memcpy(&Components[ComponentIndex].Packed, ParameterData, ComponentSizeInBytes);
+				ParameterData += ComponentSizeInBytes;
+			}
+		}
+		else
+		{
+			Stack.PushValue(UniformExpressionSet->GetDefaultParameterValue(Parameter.ParameterType, Parameter.DefaultValueOffset));
+		}
+	}
+}
+
+// Utility function to normalize a vector in place, for use with EvaluateUnaryOpInPlace
+EValueType NormalizeInPlace(EValueType Type, TArrayView<FValueComponent> Component)
+{
+	const FValueTypeDescription& TypeDesc = GetValueTypeDescription(Type);
+	if (TypeDesc.ComponentType == EValueComponentType::Double)
+	{
+		// Compute magnitude squared (via dot product)
+		double MagnitudeSquared = 0.0;
+		for (int32 ComponentIndex = 0; ComponentIndex < TypeDesc.NumComponents; ComponentIndex++)
+		{
+			double Value = Component[ComponentIndex].Double;
+			MagnitudeSquared += Value * Value;
+		}
+
+		// Make the magnitude safe to invert.  Similar to logic in GetSafeDivisor, which clamps a value
+		// to UE_DELTA, but in this case we are using the delta squared to produce the same result we
+		// would get if we did a square root before dividing.
+		MagnitudeSquared = FMath::Max(MagnitudeSquared, UE_DOUBLE_DELTA * UE_DOUBLE_DELTA);
+
+		double InverseMagnitude = FMath::InvSqrt(MagnitudeSquared);
+
+		// And apply that back to the components
+		for (int32 ComponentIndex = 0; ComponentIndex < TypeDesc.NumComponents; ComponentIndex++)
+		{
+			Component[ComponentIndex].Double *= InverseMagnitude;
+		}
+
+		// Type didn't change, started as Double, still Double
+		return Type;
+	}
+	else
+	{
+		// Convert the type to float, if it's not already.
+		if (TypeDesc.ComponentType != EValueComponentType::Float)
+		{
+			if (TypeDesc.ComponentType == EValueComponentType::Int)
+			{
+				for (int32 ComponentIndex = 0; ComponentIndex < TypeDesc.NumComponents; ComponentIndex++)
+				{
+					Component[ComponentIndex].Float = (float)Component[ComponentIndex].Int;
+				}
+
+				// Translate type.  Note that int doesn't support matrices, so we know it should be [Int1..Int4].
+				check(Type >= EValueType::Int1 && Type <= EValueType::Int4);
+				Type = (EValueType)((int32)Type - (int32)EValueType::Int1 + (int32)EValueType::Float1);
+			}
+			else if (TypeDesc.ComponentType == EValueComponentType::Bool)
+			{
+				for (int32 ComponentIndex = 0; ComponentIndex < TypeDesc.NumComponents; ComponentIndex++)
+				{
+					Component[ComponentIndex].Float = (float)Component[ComponentIndex].Bool;
+				}
+
+				// Translate type.  Note that bool doesn't support matrices, so we know it should be [Bool1..Bool4].
+				check(Type >= EValueType::Bool1 && Type <= EValueType::Bool4);
+				Type = (EValueType)((int32)Type - (int32)EValueType::Bool1 + (int32)EValueType::Float1);
+			}
+			else
+			{
+				// Anything else, we'll just bail and not modify the input.  It's probably a bug for the bytecode to
+				// generate a normalize operation on a non numeric type.
+				return Type;
+			}
+		}
+
+		// Compute magnitude squared (via dot product)
+		float MagnitudeSquared = 0.0;
+		for (int32 ComponentIndex = 0; ComponentIndex < TypeDesc.NumComponents; ComponentIndex++)
+		{
+			float Value = Component[ComponentIndex].Float;
+			MagnitudeSquared += Value * Value;
+		}
+
+		// Make the magnitude safe to invert.  Similar to logic in GetSafeDivisor, which clamps a value
+		// to UE_DELTA, but in this case we are using the delta squared to produce the same result we
+		// would get if we did a square root before dividing.
+		MagnitudeSquared = FMath::Max(MagnitudeSquared, UE_DELTA * UE_DELTA);
+
+		float InverseMagnitude = FMath::InvSqrt(MagnitudeSquared);
+
+		// And apply that back to the components
+		for (int32 ComponentIndex = 0; ComponentIndex < TypeDesc.NumComponents; ComponentIndex++)
+		{
+			Component[ComponentIndex].Float *= InverseMagnitude;
+		}
+
+		// We updated the type to a Float variation above, if it wasn't already Float, so return that
+		return Type;
 	}
 }
 
@@ -423,6 +593,21 @@ static inline void EvaluateUnaryOp(FPreshaderStack& Stack, const Operation& Op)
 	Stack.PushValue(Op(Value));
 }
 
+template<typename Operation>
+static inline void EvaluateUnaryOpInPlace(FPreshaderStack& Stack, const Operation& Op)
+{
+	// The original EvaluateUnaryOp pops a value from the top of the stack, applies an operation to it, then pushes a new value back
+	// on the stack.  This involves moving a lot of memory around, as well as costly bookkeeping.  What we do here instead is keep
+	// the top value on the stack, and modify it in place.  We use different variations of shader operations which can be applied
+	// to a loose EValueType and array of FValueComponent, as opposed to requiring those two things to be copied into an FValue first.
+	// It's possible for any given operation to modify the type, so we override the type of the top item based on the return value.
+	FPreshaderValue Top = Stack.PeekValue();
+	if (!Top.Type.IsStruct())
+	{
+		Stack.OverrideTopType(Op(Top.Type.ValueType, Top.Component));
+	}
+}
+
 template<typename Operation>
 static inline void EvaluateBinaryOp(FPreshaderStack& Stack, const Operation& Op)
 {
@@ -431,6 +616,25 @@ static inline void EvaluateBinaryOp(FPreshaderStack& Stack, const Operation& Op)
 	Stack.PushValue(Op(Value0, Value1));
 }
 
+template<typename Operation, typename OperationFallback>
+static inline void EvaluateBinaryOpInPlace(FPreshaderStack& Stack, const Operation& Op, const OperationFallback& FallbackOp)
+{
+	// Lhs is one below top of stack, Rhs is top
+	FPreshaderValue Lhs = Stack.PeekValue(1);
+	FPreshaderValue Rhs = Stack.PeekValue(0);
+	if (Lhs.Type.IsStruct() || Rhs.Type.IsStruct())
+	{
+		// Throw up our hands if one of the inputs is a struct, and run the old logic.  Probably a bug for a struct to be passed to a BinaryOp.
+		EvaluateBinaryOp(Stack, FallbackOp);
+	}
+	else
+	{
+		int32 ComponentsConsumed;
+		EValueType ResultType = Op(Lhs.Type.ValueType, Rhs.Type.ValueType, TArrayView<FValueComponent>(Lhs.Component.GetData(), Lhs.Component.Num() + Rhs.Component.Num()), ComponentsConsumed);
+		Stack.MergeTopTwoValues(ResultType, ComponentsConsumed);
+	}
+}
+
 template<typename Operation>
 static inline void EvaluateTernaryOp(FPreshaderStack& Stack, const Operation& Op)
 {
@@ -440,6 +644,7 @@ static inline void EvaluateTernaryOp(FPreshaderStack& Stack, const Operation& Op
 	Stack.PushValue(Op(Value0, Value1, Value2));
 }
 
+// Runs the swizzle in place on the top item on the stack
 static void EvaluateComponentSwizzle(FPreshaderStack& Stack, FPreshaderDataContext& RESTRICT Data)
 {
 	const uint8 NumElements = ReadPreshaderValue<uint8>(Data);
@@ -448,27 +653,63 @@ static void EvaluateComponentSwizzle(FPreshaderStack& Stack, FPreshaderDataConte
 	const uint8 IndexB = ReadPreshaderValue<uint8>(Data);
 	const uint8 IndexA = ReadPreshaderValue<uint8>(Data);
 
-	FValue Value = Stack.PopValue().AsShaderValue();
-	FValue Result(MakeValueType(Value.Type, NumElements));
+	// Get original type and adjust the count and type to the new type.  Adjusting the count doesn't destroy
+	// the elements, so we can still access them, but saves time over having to re-fetch the component pointer
+	// in case the component array gets resized between the read and write.
+	const FPreshaderType& OriginalType = Stack.PeekType();
+	check(OriginalType.IsStruct() == false);
+
+	int32 OriginalNumComponents = GetValueTypeDescription(OriginalType.ValueType).NumComponents;
+
+	Stack.AdjustComponentCount(NumElements - OriginalNumComponents);
+	Stack.OverrideTopType(UE::Shader::MakeValueType(OriginalType.ValueType, NumElements));
+
+	// We've already adjusted the component count for the stack top to the new count, so that's the offset we
+	// need to use to get the Component pointer.
+	FValueComponent* Components = Stack.PeekComponents(NumElements);
+
 	switch (NumElements)
 	{
 	case 4:
-		Result.Component[3] = Value.GetComponent(IndexA);
-		// Fallthrough...
+		{
+			uint64 PackedR = IndexR < OriginalNumComponents ? Components[IndexR].Packed : 0;
+			uint64 PackedG = IndexG < OriginalNumComponents ? Components[IndexG].Packed : 0;
+			uint64 PackedB = IndexB < OriginalNumComponents ? Components[IndexB].Packed : 0;
+			uint64 PackedA = IndexA < OriginalNumComponents ? Components[IndexA].Packed : 0;
+			Components[0].Packed = PackedR;
+			Components[1].Packed = PackedG;
+			Components[2].Packed = PackedB;
+			Components[3].Packed = PackedA;
+		}
+		break;
 	case 3:
-		Result.Component[2] = Value.GetComponent(IndexB);
-		// Fallthrough...
+		{
+			uint64 PackedR = IndexR < OriginalNumComponents ? Components[IndexR].Packed : 0;
+			uint64 PackedG = IndexG < OriginalNumComponents ? Components[IndexG].Packed : 0;
+			uint64 PackedB = IndexB < OriginalNumComponents ? Components[IndexB].Packed : 0;
+			Components[0].Packed = PackedR;
+			Components[1].Packed = PackedG;
+			Components[2].Packed = PackedB;
+		}
+		break;
 	case 2:
-		Result.Component[1] = Value.GetComponent(IndexG);
-		// Fallthrough...
+		{
+			uint64 PackedR = IndexR < OriginalNumComponents ? Components[IndexR].Packed : 0;
+			uint64 PackedG = IndexG < OriginalNumComponents ? Components[IndexG].Packed : 0;
+			Components[0].Packed = PackedR;
+			Components[1].Packed = PackedG;
+		}
+		break;
 	case 1:
-		Result.Component[0] = Value.GetComponent(IndexR);
+		{
+			uint64 PackedR = IndexR < OriginalNumComponents ? Components[IndexR].Packed : 0;
+			Components[0].Packed = PackedR;
+		}
 		break;
 	default:
 		UE_LOG(LogMaterial, Fatal, TEXT("Invalid number of swizzle elements: %d"), NumElements);
 		break;
 	}
-	Stack.PushValue(Result);
 }
 
 static const UTexture* GetTextureParameter(const FMaterialRenderContext& Context, FPreshaderDataContext& RESTRICT Data)
@@ -602,45 +843,45 @@ FPreshaderValue EvaluatePreshader(const FUniformExpressionSet* UniformExpression
 			break;
 		case EPreshaderOpcode::PushValue: EvaluatePushValue(Stack, Data); break;
 		case EPreshaderOpcode::Assign: EvaluateAssign(Stack); break;
-		case EPreshaderOpcode::Add: EvaluateBinaryOp(Stack, Add); break;
-		case EPreshaderOpcode::Sub: EvaluateBinaryOp(Stack, Sub); break;
-		case EPreshaderOpcode::Mul: EvaluateBinaryOp(Stack, Mul); break;
-		case EPreshaderOpcode::Div: EvaluateBinaryOp(Stack, Div); break;
+		case EPreshaderOpcode::Add: EvaluateBinaryOpInPlace(Stack, AddInPlace, Add); break;
+		case EPreshaderOpcode::Sub: EvaluateBinaryOpInPlace(Stack, SubInPlace, Sub); break;
+		case EPreshaderOpcode::Mul: EvaluateBinaryOpInPlace(Stack, MulInPlace, Mul); break;
+		case EPreshaderOpcode::Div: EvaluateBinaryOpInPlace(Stack, DivInPlace, Div); break;
 		case EPreshaderOpcode::Less: EvaluateBinaryOp(Stack, Less); break;
 		case EPreshaderOpcode::Greater: EvaluateBinaryOp(Stack, Greater); break;
 		case EPreshaderOpcode::LessEqual: EvaluateBinaryOp(Stack, LessEqual); break;
 		case EPreshaderOpcode::GreaterEqual: EvaluateBinaryOp(Stack, GreaterEqual); break;
-		case EPreshaderOpcode::Fmod: EvaluateBinaryOp(Stack, Fmod); break;
-		case EPreshaderOpcode::Min: EvaluateBinaryOp(Stack, Min); break;
-		case EPreshaderOpcode::Max: EvaluateBinaryOp(Stack, Max); break;
+		case EPreshaderOpcode::Fmod: EvaluateBinaryOpInPlace(Stack, FmodInPlace, Fmod); break;
+		case EPreshaderOpcode::Min: EvaluateBinaryOpInPlace(Stack, MinInPlace, Min); break;
+		case EPreshaderOpcode::Max: EvaluateBinaryOpInPlace(Stack, MaxInPlace, Max); break;
 		case EPreshaderOpcode::Clamp: EvaluateTernaryOp(Stack, Clamp); break;
 		case EPreshaderOpcode::Dot: EvaluateBinaryOp(Stack, Dot); break;
 		case EPreshaderOpcode::Cross: EvaluateBinaryOp(Stack, Cross); break;
-		case EPreshaderOpcode::Neg: EvaluateUnaryOp(Stack, Neg); break;
-		case EPreshaderOpcode::Sqrt: EvaluateUnaryOp(Stack, Sqrt); break;
-		case EPreshaderOpcode::Rcp: EvaluateUnaryOp(Stack, Rcp); break;
+		case EPreshaderOpcode::Neg: EvaluateUnaryOpInPlace(Stack, NegInPlace); break;
+		case EPreshaderOpcode::Sqrt: EvaluateUnaryOpInPlace(Stack, SqrtInPlace); break;
+		case EPreshaderOpcode::Rcp: EvaluateUnaryOpInPlace(Stack, RcpInPlace); break;
 		case EPreshaderOpcode::Length: EvaluateUnaryOp(Stack, [](const FValue& Value) { return Sqrt(Dot(Value, Value)); }); break;
-		case EPreshaderOpcode::Normalize: EvaluateUnaryOp(Stack, [](const FValue& Value) { return Div(Value, Sqrt(Dot(Value, Value))); }); break;
-		case EPreshaderOpcode::Sin: EvaluateUnaryOp(Stack, Sin); break;
-		case EPreshaderOpcode::Cos: EvaluateUnaryOp(Stack, Cos); break;
-		case EPreshaderOpcode::Tan: EvaluateUnaryOp(Stack, Tan); break;
-		case EPreshaderOpcode::Asin: EvaluateUnaryOp(Stack, Asin); break;;
-		case EPreshaderOpcode::Acos: EvaluateUnaryOp(Stack, Acos); break;
-		case EPreshaderOpcode::Atan: EvaluateUnaryOp(Stack, Atan); break;
-		case EPreshaderOpcode::Atan2: EvaluateBinaryOp(Stack, Atan2); break;
-		case EPreshaderOpcode::Abs: EvaluateUnaryOp(Stack, Abs); break;
-		case EPreshaderOpcode::Saturate: EvaluateUnaryOp(Stack, Saturate); break;
-		case EPreshaderOpcode::Floor: EvaluateUnaryOp(Stack, Floor); break;
-		case EPreshaderOpcode::Ceil: EvaluateUnaryOp(Stack, Ceil); break;
-		case EPreshaderOpcode::Round: EvaluateUnaryOp(Stack, Round); break;
-		case EPreshaderOpcode::Trunc: EvaluateUnaryOp(Stack, Trunc); break;
-		case EPreshaderOpcode::Sign: EvaluateUnaryOp(Stack, Sign); break;
-		case EPreshaderOpcode::Frac: EvaluateUnaryOp(Stack, Frac); break;
-		case EPreshaderOpcode::Fractional: EvaluateUnaryOp(Stack, Fractional); break;
-		case EPreshaderOpcode::Log2: EvaluateUnaryOp(Stack, Log2); break;
-		case EPreshaderOpcode::Log10: EvaluateUnaryOp(Stack, Log10); break;
+		case EPreshaderOpcode::Normalize: EvaluateUnaryOpInPlace(Stack, NormalizeInPlace); break;
+		case EPreshaderOpcode::Sin: EvaluateUnaryOpInPlace(Stack, SinInPlace); break;
+		case EPreshaderOpcode::Cos: EvaluateUnaryOpInPlace(Stack, CosInPlace); break;
+		case EPreshaderOpcode::Tan: EvaluateUnaryOpInPlace(Stack, TanInPlace); break;
+		case EPreshaderOpcode::Asin: EvaluateUnaryOpInPlace(Stack, AsinInPlace); break;;
+		case EPreshaderOpcode::Acos: EvaluateUnaryOpInPlace(Stack, AcosInPlace); break;
+		case EPreshaderOpcode::Atan: EvaluateUnaryOpInPlace(Stack, AtanInPlace); break;
+		case EPreshaderOpcode::Atan2: EvaluateBinaryOpInPlace(Stack, Atan2InPlace, Atan2); break;
+		case EPreshaderOpcode::Abs: EvaluateUnaryOpInPlace(Stack, AbsInPlace); break;
+		case EPreshaderOpcode::Saturate: EvaluateUnaryOpInPlace(Stack, SaturateInPlace); break;
+		case EPreshaderOpcode::Floor: EvaluateUnaryOpInPlace(Stack, FloorInPlace); break;
+		case EPreshaderOpcode::Ceil: EvaluateUnaryOpInPlace(Stack, CeilInPlace); break;
+		case EPreshaderOpcode::Round: EvaluateUnaryOpInPlace(Stack, RoundInPlace); break;
+		case EPreshaderOpcode::Trunc: EvaluateUnaryOpInPlace(Stack, TruncInPlace); break;
+		case EPreshaderOpcode::Sign: EvaluateUnaryOpInPlace(Stack, SignInPlace); break;
+		case EPreshaderOpcode::Frac: EvaluateUnaryOpInPlace(Stack, FracInPlace); break;
+		case EPreshaderOpcode::Fractional: EvaluateUnaryOpInPlace(Stack, FractionalInPlace); break;
+		case EPreshaderOpcode::Log2: EvaluateUnaryOpInPlace(Stack, Log2InPlace); break;
+		case EPreshaderOpcode::Log10: EvaluateUnaryOpInPlace(Stack, Log10InPlace); break;
 		case EPreshaderOpcode::ComponentSwizzle: EvaluateComponentSwizzle(Stack, Data); break;
-		case EPreshaderOpcode::AppendVector: EvaluateBinaryOp(Stack, Append); break;
+		case EPreshaderOpcode::AppendVector: EvaluateBinaryOpInPlace(Stack, AppendInPlace, Append); break;
 		case EPreshaderOpcode::TextureSize: EvaluateTextureSize(Context, Stack, Data); break;
 		case EPreshaderOpcode::TexelSize: EvaluateTexelSize(Context, Stack, Data); break;
 		case EPreshaderOpcode::ExternalTextureCoordinateScaleRotation: EvaluateExternalTextureCoordinateScaleRotation(Context, Stack, Data); break;
@@ -656,11 +897,7 @@ FPreshaderValue EvaluatePreshader(const FUniformExpressionSet* UniformExpression
 	check(Data.Ptr == DataEnd);
 
 	FPreshaderValue Result;
-	if (Stack.Num() > 0)
-	{
-		Result = Stack.PopValue();
-	}
-	Stack.CheckEmpty();
+	Stack.PopResult(Result);
 	return Result;
 }
 
diff --git a/Engine/Source/Runtime/Engine/Private/Shader/ShaderValue.cpp b/Engine/Source/Runtime/Engine/Private/Shader/ShaderValue.cpp
--- a/Engine/Source/Runtime/Engine/Private/Shader/ShaderValue.cpp
+++ b/Engine/Source/Runtime/Engine/Private/Shader/ShaderValue.cpp
@@ -11,6 +11,69 @@ namespace UE
 namespace Shader
 {
 
+const FValueTypeDescription GValueTypeDescriptions[] =
+{
+	{ TEXT("void"),			EValueType::Void,		EValueComponentType::Void,		0, 0 },
+	{ TEXT("float"),		EValueType::Float1,		EValueComponentType::Float,		1, sizeof(float) },
+	{ TEXT("float2"),		EValueType::Float2,		EValueComponentType::Float,		2, sizeof(float) },
+	{ TEXT("float3"),		EValueType::Float3,		EValueComponentType::Float,		3, sizeof(float) },
+	{ TEXT("float4"),		EValueType::Float4,		EValueComponentType::Float,		4, sizeof(float) },
+	{ TEXT("FLWCScalar"),	EValueType::Double1,	EValueComponentType::Double,	1, sizeof(double) },
+	{ TEXT("FLWCVector2"),	EValueType::Double2,	EValueComponentType::Double,	2, sizeof(double) },
+	{ TEXT("FLWCVector3"),	EValueType::Double3,	EValueComponentType::Double,	3, sizeof(double) },
+	{ TEXT("FLWCVector4"),	EValueType::Double4,	EValueComponentType::Double,	4, sizeof(double) },
+	{ TEXT("int"),			EValueType::Int1,		EValueComponentType::Int,		1, sizeof(int32) },
+	{ TEXT("int2"),			EValueType::Int2,		EValueComponentType::Int,		2, sizeof(int32) },
+	{ TEXT("int3"),			EValueType::Int3,		EValueComponentType::Int,		3, sizeof(int32) },
+	{ TEXT("int4"),			EValueType::Int4,		EValueComponentType::Int,		4, sizeof(int32) },
+	{ TEXT("bool"),			EValueType::Bool1,		EValueComponentType::Bool,		1, 1 },
+	{ TEXT("bool2"),		EValueType::Bool2,		EValueComponentType::Bool,		2, 1 },
+	{ TEXT("bool3"),		EValueType::Bool3,		EValueComponentType::Bool,		3, 1 },
+	{ TEXT("bool4"),		EValueType::Bool4,		EValueComponentType::Bool,		4, 1 },
+	{ TEXT("Numeric1"),		EValueType::Numeric1,	EValueComponentType::Numeric,	1, sizeof(double) },
+	{ TEXT("Numeric2"),		EValueType::Numeric2,	EValueComponentType::Numeric,	2, sizeof(double) },
+	{ TEXT("Numeric3"),		EValueType::Numeric3,	EValueComponentType::Numeric,	3, sizeof(double) },
+	{ TEXT("Numeric4"),		EValueType::Numeric4,	EValueComponentType::Numeric,	4, sizeof(double) },
+	{ TEXT("float4x4"),		EValueType::Float4x4,	EValueComponentType::Float,		16, sizeof(float) },
+	{ TEXT("FLWCMatrix"),	EValueType::Double4x4,	EValueComponentType::Double,	16, sizeof(double) },
+	{ TEXT("FLWCInverseMatrix"), EValueType::DoubleInverse4x4, EValueComponentType::Double, 16, sizeof(double) },
+	{ TEXT("Numeric4x4"),	EValueType::Numeric4x4, EValueComponentType::Numeric,	16, sizeof(double) },
+	{ TEXT("struct"),		EValueType::Struct,		EValueComponentType::Void,		0, 0 },
+	{ TEXT("object"),		EValueType::Object,		EValueComponentType::Void,		0, 0 },
+	{ TEXT("Any"),			EValueType::Any,		EValueComponentType::Void,		0, 0 },
+	{ TEXT("<INVALID>"),	EValueType::Num,		EValueComponentType::Void,		0, 0 },
+};
+
+static_assert(UE_ARRAY_COUNT(GValueTypeDescriptions) == (int32)EValueType::Num + 1, "Missing entry from shader value description table");
+
+#if DO_CHECK
+// Startup validation logic for table above
+struct FValidateShaderValueTypeDescriptionTable
+{
+	FValidateShaderValueTypeDescriptionTable()
+	{
+		for (int32 DescriptionIndex = 0; DescriptionIndex <= (int32)EValueType::Num; DescriptionIndex++)
+		{
+			// Make sure element location in array matches its ValueType (validates that the array matches the enum order)
+			check((int32)GValueTypeDescriptions[DescriptionIndex].ValueType == DescriptionIndex);
+
+			// Make sure component size matches the component type
+			check(GValueTypeDescriptions[DescriptionIndex].ComponentSizeInBytes == GetComponentTypeSizeInBytes(GValueTypeDescriptions[DescriptionIndex].ComponentType));
+		}
+	}
+};
+
+FValidateShaderValueTypeDescriptionTable GValidateShaderValueTypeDescriptionTable;
+#endif  // DO_CHECK
+
+const FValueTypeDescription& GetValueTypeDescription(EValueType Type)
+{
+	uint32 TypeIndex = (uint32)Type;
+	check((uint32)Type < NumValueTypes);
+
+	return GValueTypeDescriptions[FMath::Min(TypeIndex, (uint32)NumValueTypes)];
+}
+
 const TCHAR* FType::GetName() const
 {
 	if (IsStruct())
@@ -21,7 +84,7 @@ const TCHAR* FType::GetName() const
 
 	// TODO - do we need specific object names?
 	check(ValueType != EValueType::Struct);
-	const FValueTypeDescription TypeDesc = GetValueTypeDescription(ValueType);
+	const FValueTypeDescription& TypeDesc = GetValueTypeDescription(ValueType);
 	return TypeDesc.Name;
 }
 
@@ -54,7 +117,7 @@ int32 FType::GetNumComponents() const
 	}
 
 	check(ValueType != EValueType::Struct);
-	const FValueTypeDescription TypeDesc = GetValueTypeDescription(ValueType);
+	const FValueTypeDescription& TypeDesc = GetValueTypeDescription(ValueType);
 	return TypeDesc.NumComponents;
 }
 
@@ -87,7 +150,7 @@ EValueComponentType FType::GetComponentType(int32 Index) const
 	}
 	else if(IsNumeric())
 	{
-		const FValueTypeDescription TypeDesc = GetValueTypeDescription(ValueType);
+		const FValueTypeDescription& TypeDesc = GetValueTypeDescription(ValueType);
 		// Scalar types replicate xyzw
 		if ((TypeDesc.NumComponents == 1 && Index < 4) || Index < TypeDesc.NumComponents)
 		{
@@ -124,8 +187,8 @@ FType CombineTypes(const FType& Lhs, const FType& Rhs)
 
 	if (Lhs.IsNumericVector() && Rhs.IsNumericVector())
 	{
-		const FValueTypeDescription LhsDesc = GetValueTypeDescription(Lhs);
-		const FValueTypeDescription RhsDesc = GetValueTypeDescription(Rhs);
+		const FValueTypeDescription& LhsDesc = GetValueTypeDescription(Lhs);
+		const FValueTypeDescription& RhsDesc = GetValueTypeDescription(Rhs);
 		const EValueComponentType ComponentType = CombineComponentTypes(LhsDesc.ComponentType, RhsDesc.ComponentType);
 		if (ComponentType == EValueComponentType::Void)
 		{
@@ -228,7 +291,7 @@ template<typename Operation, typename ResultType>
 void AsType(const Operation& Op, const FValue& Value, ResultType& OutResult)
 {
 	using FComponentType = typename Operation::FComponentType;
-	const FValueTypeDescription TypeDesc = GetValueTypeDescription(Value.Type);
+	const FValueTypeDescription& TypeDesc = GetValueTypeDescription(Value.Type);
 	if (TypeDesc.NumComponents == 1)
 	{
 		const FComponentType Component = Op(TypeDesc.ComponentType, Value.Component[0]);
@@ -251,12 +314,40 @@ void AsType(const Operation& Op, const FValue& Value, ResultType& OutResult)
 	}
 }
 
+// Identical to "AsType" above, except Type and Component are loose, rather than pulled from an FValue structure
+template<typename Operation, typename ResultType>
+void AsTypeInPlace(const Operation& Op, EValueType Type, TArrayView<FValueComponent> Component, ResultType& OutResult)
+{
+	using FComponentType = typename Operation::FComponentType;
+	const FValueTypeDescription& TypeDesc = GetValueTypeDescription(Type);
+	if (TypeDesc.NumComponents == 1)
+	{
+		const FComponentType ComponentCast = Op(TypeDesc.ComponentType, Component[0]);
+		for (int32 i = 0; i < 4; ++i)
+		{
+			OutResult[i] = ComponentCast;
+		}
+	}
+	else
+	{
+		const int32 NumComponents = FMath::Min<int32>(TypeDesc.NumComponents, 4);
+		for (int32 i = 0; i < NumComponents; ++i)
+		{
+			OutResult[i] = Op(TypeDesc.ComponentType, Component[i]);
+		}
+		for (int32 i = NumComponents; i < 4; ++i)
+		{
+			OutResult[i] = (FComponentType)0;
+		}
+	}
+}
+
 template<typename Operation>
 void Cast(const Operation& Op, const FValue& Value, FValue& OutResult)
 {
 	using FComponentType = typename Operation::FComponentType;
-	const FValueTypeDescription ValueTypeDesc = GetValueTypeDescription(Value.Type);
-	const FValueTypeDescription ResultTypeDesc = GetValueTypeDescription(OutResult.Type);
+	const FValueTypeDescription& ValueTypeDesc = GetValueTypeDescription(Value.Type);
+	const FValueTypeDescription& ResultTypeDesc = GetValueTypeDescription(OutResult.Type);
 	const int32 NumCopyComponents = FMath::Min<int32>(ValueTypeDesc.NumComponents, ResultTypeDesc.NumComponents);
 	for (int32 i = 0; i < NumCopyComponents; ++i)
 	{
@@ -307,11 +398,11 @@ void FormatComponent_Double(double Value, int32 NumComponents, EValueStringForma
 FValue FValue::FromMemoryImage(EValueType Type, const void* Data, uint32* OutSizeInBytes)
 {
 	check(IsNumericType(Type));
-	const FValueTypeDescription TypeDesc = GetValueTypeDescription(Type);
+	const FValueTypeDescription& TypeDesc = GetValueTypeDescription(Type);
 
 	FValue Result(TypeDesc.ComponentType, TypeDesc.NumComponents);
 	const uint8* Bytes = static_cast<const uint8*>(Data);
-	const uint32 ComponentSizeInBytes = GetComponentTypeSizeInBytes(TypeDesc.ComponentType);
+	const uint32 ComponentSizeInBytes = TypeDesc.ComponentSizeInBytes;
 	if (ComponentSizeInBytes > 0u)
 	{
 		for (int32 i = 0u; i < TypeDesc.NumComponents; ++i)
@@ -330,11 +421,11 @@ FValue FValue::FromMemoryImage(EValueType Type, const void* Data, uint32* OutSiz
 FMemoryImageValue FValue::AsMemoryImage() const
 {
 	check(Type.IsNumeric());
-	const FValueTypeDescription TypeDesc = GetValueTypeDescription(Type);
+	const FValueTypeDescription& TypeDesc = GetValueTypeDescription(Type);
 
 	FMemoryImageValue Result;
 	uint8* Bytes = Result.Bytes;
-	const uint32 ComponentSizeInBytes = GetComponentTypeSizeInBytes(TypeDesc.ComponentType);
+	const uint32 ComponentSizeInBytes = TypeDesc.ComponentSizeInBytes;
 	if (ComponentSizeInBytes > 0u)
 	{
 		for (int32 i = 0u; i < TypeDesc.NumComponents; ++i)
@@ -502,7 +593,7 @@ const TCHAR* FValue::ToString(EValueStringFormat Format, FStringBuilderBase& Out
 		}
 		else
 		{
-			const FValueTypeDescription TypeDesc = GetValueTypeDescription(Type.ValueType);
+			const FValueTypeDescription& TypeDesc = GetValueTypeDescription(Type.ValueType);
 			check(TypeDesc.ComponentType != EValueComponentType::Numeric);
 			if (TypeDesc.ComponentType != EValueComponentType::Double)
 			{
@@ -539,48 +630,12 @@ const TCHAR* FValue::ToString(EValueStringFormat Format, FStringBuilderBase& Out
 	return OutString.ToString();
 }
 
-FValueTypeDescription GetValueTypeDescription(EValueType Type)
-{
-	switch (Type)
-	{
-	case EValueType::Void: return FValueTypeDescription(TEXT("void"), EValueComponentType::Void, 0);
-	case EValueType::Float1: return FValueTypeDescription(TEXT("float"), EValueComponentType::Float, 1);
-	case EValueType::Float2: return FValueTypeDescription(TEXT("float2"), EValueComponentType::Float, 2);
-	case EValueType::Float3: return FValueTypeDescription(TEXT("float3"), EValueComponentType::Float, 3);
-	case EValueType::Float4: return FValueTypeDescription(TEXT("float4"), EValueComponentType::Float, 4);
-	case EValueType::Double1: return FValueTypeDescription(TEXT("FLWCScalar"), EValueComponentType::Double, 1);
-	case EValueType::Double2: return FValueTypeDescription(TEXT("FLWCVector2"), EValueComponentType::Double, 2);
-	case EValueType::Double3: return FValueTypeDescription(TEXT("FLWCVector3"), EValueComponentType::Double, 3);
-	case EValueType::Double4: return FValueTypeDescription(TEXT("FLWCVector4"), EValueComponentType::Double, 4);
-	case EValueType::Int1: return FValueTypeDescription(TEXT("int"), EValueComponentType::Int, 1);
-	case EValueType::Int2: return FValueTypeDescription(TEXT("int2"), EValueComponentType::Int, 2);
-	case EValueType::Int3: return FValueTypeDescription(TEXT("int3"), EValueComponentType::Int, 3);
-	case EValueType::Int4: return FValueTypeDescription(TEXT("int4"), EValueComponentType::Int, 4);
-	case EValueType::Bool1: return FValueTypeDescription(TEXT("bool"), EValueComponentType::Bool, 1);
-	case EValueType::Bool2: return FValueTypeDescription(TEXT("bool2"), EValueComponentType::Bool, 2);
-	case EValueType::Bool3: return FValueTypeDescription(TEXT("bool3"), EValueComponentType::Bool, 3);
-	case EValueType::Bool4: return FValueTypeDescription(TEXT("bool4"), EValueComponentType::Bool, 4);
-	case EValueType::Numeric1: return FValueTypeDescription(TEXT("Numeric1"), EValueComponentType::Numeric, 1);
-	case EValueType::Numeric2: return FValueTypeDescription(TEXT("Numeric2"), EValueComponentType::Numeric, 2);
-	case EValueType::Numeric3: return FValueTypeDescription(TEXT("Numeric3"), EValueComponentType::Numeric, 3);
-	case EValueType::Numeric4: return FValueTypeDescription(TEXT("Numeric4"), EValueComponentType::Numeric, 4);
-	case EValueType::Float4x4: return FValueTypeDescription(TEXT("float4x4"), EValueComponentType::Float, 16);
-	case EValueType::Double4x4: return FValueTypeDescription(TEXT("FLWCMatrix"), EValueComponentType::Double, 16);
-	case EValueType::DoubleInverse4x4: return FValueTypeDescription(TEXT("FLWCInverseMatrix"), EValueComponentType::Double, 16);
-	case EValueType::Numeric4x4: return FValueTypeDescription(TEXT("Numeric4x4"), EValueComponentType::Numeric, 16);
-	case EValueType::Struct: return FValueTypeDescription(TEXT("struct"), EValueComponentType::Void, 0);
-	case EValueType::Object: return FValueTypeDescription(TEXT("object"), EValueComponentType::Void, 0);
-	case EValueType::Any: return FValueTypeDescription(TEXT("Any"), EValueComponentType::Void, 0);
-	default: checkNoEntry(); return FValueTypeDescription(TEXT("<INVALID>"), EValueComponentType::Void, 0);
-	}
-}
-
 EValueType FindValueType(FName Name)
 {
 	for (int32 TypeIndex = 1; TypeIndex < NumValueTypes; ++TypeIndex)
 	{
 		const EValueType Type = (EValueType)TypeIndex;
-		const FValueTypeDescription TypeDesc = GetValueTypeDescription(Type);
+		const FValueTypeDescription& TypeDesc = GetValueTypeDescription(Type);
 		if (Name == TypeDesc.Name)
 		{
 			return Type;
@@ -668,13 +723,13 @@ EValueType MakeValueType(EValueType BaseType, int32 NumComponents)
 
 EValueType MakeValueTypeWithRequestedNumComponents(EValueType BaseType, int8 RequestedNumComponents)
 {
-	const FValueTypeDescription TypeDesc = GetValueTypeDescription(BaseType);
+	const FValueTypeDescription& TypeDesc = GetValueTypeDescription(BaseType);
 	return MakeValueType(TypeDesc.ComponentType, FMath::Min(TypeDesc.NumComponents, RequestedNumComponents));
 }
 
 EValueType MakeNonLWCType(EValueType Type)
 {
-	const FValueTypeDescription TypeDesc = GetValueTypeDescription(Type);
+	const FValueTypeDescription& TypeDesc = GetValueTypeDescription(Type);
 	check(IsNumericType(TypeDesc.ComponentType));
 	if (TypeDesc.ComponentType == EValueComponentType::Double)
 	{
@@ -685,7 +740,7 @@ EValueType MakeNonLWCType(EValueType Type)
 
 EValueType MakeConcreteType(EValueType Type)
 {
-	const FValueTypeDescription TypeDesc = GetValueTypeDescription(Type);
+	const FValueTypeDescription& TypeDesc = GetValueTypeDescription(Type);
 	check(IsNumericType(TypeDesc.ComponentType));
 	if (TypeDesc.ComponentType == EValueComponentType::Numeric)
 	{
@@ -696,7 +751,7 @@ EValueType MakeConcreteType(EValueType Type)
 
 EValueType MakeDerivativeType(EValueType Type)
 {
-	const FValueTypeDescription TypeDesc = GetValueTypeDescription(Type);
+	const FValueTypeDescription& TypeDesc = GetValueTypeDescription(Type);
 	if (IsNumericType(TypeDesc.ComponentType))
 	{
 		return MakeValueType(EValueComponentType::Float, TypeDesc.NumComponents);
@@ -749,7 +804,7 @@ void SetFieldType(EValueType* FieldTypes, EValueComponentType* ComponentTypes, i
 	else
 	{
 		FieldTypes[FieldIndex] = Type.ValueType;
-		const FValueTypeDescription TypeDesc = GetValueTypeDescription(Type.ValueType);
+		const FValueTypeDescription& TypeDesc = GetValueTypeDescription(Type.ValueType);
 		for (int32 i = 0; i < TypeDesc.NumComponents; ++i)
 		{
 			ComponentTypes[ComponentIndex + i] = TypeDesc.ComponentType;
@@ -957,7 +1012,7 @@ inline FValue UnaryOp(const Operation& Op, const FValue& Value)
 	{
 		return FValue();
 	}
-	const FValueTypeDescription TypeDesc = GetValueTypeDescription(Value.Type);
+	const FValueTypeDescription& TypeDesc = GetValueTypeDescription(Value.Type);
 	check(TypeDesc.ComponentType != EValueComponentType::Numeric);
 	const int8 NumComponents = TypeDesc.NumComponents;
 	
@@ -999,6 +1054,67 @@ inline FValue UnaryOp(const Operation& Op, const FValue& Value)
 	return Result;
 }
 
+// Similar to "UnaryOp", but reads and writes results to a loose FValueComponent array with a specified initial type.
+// Certain ops may change the type, so a new type is returned, but it will always have the same number of components.
+template<typename Operation>
+inline EValueType UnaryOpInPlace(const Operation& Op, EValueType Type, TArrayView<FValueComponent>& Component)
+{
+	const FValueTypeDescription& TypeDesc = GetValueTypeDescription(Type);
+	check(TypeDesc.ComponentType != EValueComponentType::Numeric);
+	const int8 NumComponents = TypeDesc.NumComponents;
+
+	EValueType ResultType;
+
+	// Check for Float input first, as it's by far the most common case (96%)
+	if (TypeDesc.ComponentType == EValueComponentType::Float)
+	{
+		// Float to float, no need to recompute type or cast
+		ResultType = Type;
+		for (int32 i = 0; i < NumComponents; ++i)
+		{
+			Component[i].Float = Op(Component[i].Float);
+		}
+		return ResultType;
+	}
+
+	if constexpr (Operation::SupportsDouble)
+	{
+		if (TypeDesc.ComponentType == EValueComponentType::Double)
+		{
+			// Double to double, no need to recompute type or cast
+			ResultType = Type;
+			for (int32 i = 0; i < NumComponents; ++i)
+			{
+				Component[i].Double = Op(Component[i].Double);
+			}
+			return ResultType;
+		}
+	}
+
+	if constexpr (Operation::SupportsInt)
+	{
+		// Cast to integer
+		ResultType = MakeValueType(EValueComponentType::Int, NumComponents);
+		const FCastInt Cast;
+		for (int32 i = 0; i < NumComponents; ++i)
+		{
+			Component[i].Int = Op(Cast(TypeDesc.ComponentType,Component[i]));
+		}
+		return ResultType;
+	}
+	else
+	{
+		// Cast to float
+		ResultType = MakeValueType(EValueComponentType::Float, NumComponents);
+		const FCastFloat Cast;
+		for (int32 i = 0; i < NumComponents; ++i)
+		{
+			Component[i].Float = Op(Cast(TypeDesc.ComponentType, Component[i]));
+		}
+		return ResultType;
+	}
+}
+
 inline int8 GetNumComponentsResult(int8 Lhs, int8 Rhs)
 {
 	// operations between scalar and non-scalar will splat the scalar value
@@ -1013,8 +1129,8 @@ inline FValue BinaryOp(const Operation& Op, const FValue& Lhs, const FValue& Rhs
 	{
 		return FValue();
 	}
-	const FValueTypeDescription LhsDesc = GetValueTypeDescription(Lhs.Type);
-	const FValueTypeDescription RhsDesc = GetValueTypeDescription(Rhs.Type);
+	const FValueTypeDescription& LhsDesc = GetValueTypeDescription(Lhs.Type);
+	const FValueTypeDescription& RhsDesc = GetValueTypeDescription(Rhs.Type);
 	check(LhsDesc.ComponentType != EValueComponentType::Numeric);
 	check(RhsDesc.ComponentType != EValueComponentType::Numeric);
 	const int8 NumComponents = GetNumComponentsResult(LhsDesc.NumComponents, RhsDesc.NumComponents);
@@ -1060,6 +1176,102 @@ inline FValue BinaryOp(const Operation& Op, const FValue& Lhs, const FValue& Rhs
 	return Result;
 }
 
+// Similar to "BinaryOp", but reads and writes results to a loose FValueComponent array with specified types for Lhs and Rhs.
+// The components are assumed to be sequential in the array, with Lhs first and Rhs second.  The result overwrites the start
+// of the Component array, and may have a different (but always smaller) number of components and type from the original values.
+// The new type is returned, and the number of components consumed by the operation is stored in OutComponentsConsumed.  The
+// consumed number can be removed from the Component array by the caller to produce a final result array.
+template<typename Operation>
+inline EValueType BinaryOpInPlace(const Operation& Op, EValueType LhsType, EValueType RhsType, TArrayView<FValueComponent>& Component, int32& OutComponentsConsumed)
+{
+	const FValueTypeDescription& LhsDesc = GetValueTypeDescription(LhsType);
+	const FValueTypeDescription& RhsDesc = GetValueTypeDescription(RhsType);
+	check(LhsDesc.ComponentType != EValueComponentType::Numeric);
+	check(RhsDesc.ComponentType != EValueComponentType::Numeric);
+	const int8 NumComponents = GetNumComponentsResult(LhsDesc.NumComponents, RhsDesc.NumComponents);
+
+	OutComponentsConsumed = LhsDesc.NumComponents + RhsDesc.NumComponents - NumComponents;
+	check(OutComponentsConsumed >= 0);
+
+	EValueType ResultType;
+
+	// Fast path for both input and output float (common case)
+	if (LhsDesc.ComponentType == EValueComponentType::Float && RhsDesc.ComponentType == EValueComponentType::Float)
+	{
+		ResultType = MakeValueType(EValueComponentType::Float, NumComponents);
+
+		// Handle splatting of scalars, by choosing a zero or one increment
+		int32 LhsIncrement = (LhsDesc.NumComponents == 1) ? 0 : 1;
+		int32 RhsIncrement = (RhsDesc.NumComponents == 1) ? 0 : 1;
+		int32 LhsComponent = 0;
+		int32 RhsComponent = LhsDesc.NumComponents;
+
+		// Need to write to temporary, since output can overlap with input
+		FFloatValue TempResult;
+		for (int32 i = 0; i < NumComponents; i++, LhsComponent += LhsIncrement, RhsComponent += RhsIncrement)
+		{
+			TempResult.Component[i] = Op(Component[LhsComponent].Float, Component[RhsComponent].Float);
+		}
+		for (int32 i = 0; i < NumComponents; i++)
+		{
+			Component[i].Float = TempResult.Component[i];
+		}
+
+		return ResultType;
+	}
+
+	if constexpr (Operation::SupportsDouble)
+	{
+		if (LhsDesc.ComponentType == EValueComponentType::Double || RhsDesc.ComponentType == EValueComponentType::Double)
+		{
+			ResultType = MakeValueType(EValueComponentType::Double, NumComponents);
+
+			FDoubleValue LhsCast;
+			FDoubleValue RhsCast;
+			Private::AsTypeInPlace(Private::FCastDouble(), LhsType, TArrayView<FValueComponent>(Component.GetData(),  LhsDesc.NumComponents), LhsCast);
+			Private::AsTypeInPlace(Private::FCastDouble(), RhsType, TArrayView<FValueComponent>(Component.GetData() + LhsDesc.NumComponents, RhsDesc.NumComponents), RhsCast);
+
+			for (int32 i = 0; i < NumComponents; ++i)
+			{
+				Component[i].Double = Op(LhsCast.Component[i], RhsCast.Component[i]);
+			}
+			return ResultType;
+		}
+	}
+
+	if constexpr (Operation::SupportsInt)
+	{
+		if (LhsDesc.ComponentType != EValueComponentType::Float && RhsDesc.ComponentType != EValueComponentType::Float)
+		{
+			ResultType = MakeValueType(EValueComponentType::Int, NumComponents);
+
+			FIntValue LhsCast;
+			FIntValue RhsCast;
+			Private::AsTypeInPlace(Private::FCastInt(), LhsType, TArrayView<FValueComponent>(Component.GetData(),  LhsDesc.NumComponents), LhsCast);
+			Private::AsTypeInPlace(Private::FCastInt(), RhsType, TArrayView<FValueComponent>(Component.GetData() + LhsDesc.NumComponents, RhsDesc.NumComponents), RhsCast);
+
+			for (int32 i = 0; i < NumComponents; ++i)
+			{
+				Component[i].Int = Op(LhsCast.Component[i], RhsCast.Component[i]);
+			}
+			return ResultType;
+		}
+	}
+
+	ResultType = MakeValueType(EValueComponentType::Float, NumComponents);
+
+	FFloatValue LhsCast;
+	FFloatValue RhsCast;
+	Private::AsTypeInPlace(Private::FCastFloat(), LhsType, TArrayView<FValueComponent>(Component.GetData(),  LhsDesc.NumComponents), LhsCast);
+	Private::AsTypeInPlace(Private::FCastFloat(), RhsType, TArrayView<FValueComponent>(Component.GetData() + LhsDesc.NumComponents, RhsDesc.NumComponents), RhsCast);
+
+	for (int32 i = 0; i < NumComponents; ++i)
+	{
+		Component[i].Float = Op(LhsCast.Component[i], RhsCast.Component[i]);
+	}
+	return ResultType;
+}
+
 template<typename Operation>
 inline FValue CompareOp(const Operation& Op, const FValue& Lhs, const FValue& Rhs)
 {
@@ -1067,8 +1279,8 @@ inline FValue CompareOp(const Operation& Op, const FValue& Lhs, const FValue& Rh
 	{
 		return FValue();
 	}
-	const FValueTypeDescription LhsDesc = GetValueTypeDescription(Lhs.Type);
-	const FValueTypeDescription RhsDesc = GetValueTypeDescription(Rhs.Type);
+	const FValueTypeDescription& LhsDesc = GetValueTypeDescription(Lhs.Type);
+	const FValueTypeDescription& RhsDesc = GetValueTypeDescription(Rhs.Type);
 	check(LhsDesc.ComponentType != EValueComponentType::Numeric);
 	check(RhsDesc.ComponentType != EValueComponentType::Numeric);
 	const int8 NumComponents = GetNumComponentsResult(LhsDesc.NumComponents, RhsDesc.NumComponents);
@@ -1334,8 +1546,8 @@ FValue Dot(const FValue& Lhs, const FValue& Rhs)
 	{
 		return FValue();
 	}
-	const FValueTypeDescription LhsDesc = GetValueTypeDescription(Lhs.Type);
-	const FValueTypeDescription RhsDesc = GetValueTypeDescription(Rhs.Type);
+	const FValueTypeDescription& LhsDesc = GetValueTypeDescription(Lhs.Type);
+	const FValueTypeDescription& RhsDesc = GetValueTypeDescription(Rhs.Type);
 	const int8 NumComponents = Private::GetNumComponentsResult(LhsDesc.NumComponents, RhsDesc.NumComponents);
 
 	FValue Result;
@@ -1384,8 +1596,8 @@ FValue Cross(const FValue& Lhs, const FValue& Rhs)
 	{
 		return FValue();
 	}
-	const FValueTypeDescription LhsDesc = GetValueTypeDescription(Lhs.Type);
-	const FValueTypeDescription RhsDesc = GetValueTypeDescription(Rhs.Type);
+	const FValueTypeDescription& LhsDesc = GetValueTypeDescription(Lhs.Type);
+	const FValueTypeDescription& RhsDesc = GetValueTypeDescription(Rhs.Type);
 
 	FValue Result;
 	if (LhsDesc.ComponentType == EValueComponentType::Double || RhsDesc.ComponentType == EValueComponentType::Double)
@@ -1427,8 +1639,8 @@ FValue Append(const FValue& Lhs, const FValue& Rhs)
 	{
 		return FValue();
 	}
-	const FValueTypeDescription LhsDesc = GetValueTypeDescription(Lhs.Type);
-	const FValueTypeDescription RhsDesc = GetValueTypeDescription(Rhs.Type);
+	const FValueTypeDescription& LhsDesc = GetValueTypeDescription(Lhs.Type);
+	const FValueTypeDescription& RhsDesc = GetValueTypeDescription(Rhs.Type);
 
 	FValue Result;
 	const int32 NumComponents = FMath::Min<int32>(LhsDesc.NumComponents + RhsDesc.NumComponents, 4);
@@ -1515,5 +1727,219 @@ FValue Cast(const FValue& Value, EValueType Type)
 	return Result;
 }
 
+
+EValueType NegInPlace(EValueType Type, TArrayView<FValueComponent> Component)
+{
+	return Private::UnaryOpInPlace(Private::FOpNeg(), Type, Component);
+}
+
+EValueType AbsInPlace(EValueType Type, TArrayView<FValueComponent> Component)
+{
+	return Private::UnaryOpInPlace(Private::FOpAbs(), Type, Component);
+}
+
+EValueType SaturateInPlace(EValueType Type, TArrayView<FValueComponent> Component)
+{
+	return Private::UnaryOpInPlace(Private::FOpSaturate(), Type, Component);
+}
+
+EValueType FloorInPlace(EValueType Type, TArrayView<FValueComponent> Component)
+{
+	return Private::UnaryOpInPlace(Private::FOpFloor(), Type, Component);
+}
+
+EValueType CeilInPlace(EValueType Type, TArrayView<FValueComponent> Component)
+{
+	return Private::UnaryOpInPlace(Private::FOpCeil(), Type, Component);
+}
+
+EValueType RoundInPlace(EValueType Type, TArrayView<FValueComponent> Component)
+{
+	return Private::UnaryOpInPlace(Private::FOpRound(), Type, Component);
+}
+
+EValueType TruncInPlace(EValueType Type, TArrayView<FValueComponent> Component)
+{
+	return Private::UnaryOpInPlace(Private::FOpTrunc(), Type, Component);
+}
+
+EValueType SignInPlace(EValueType Type, TArrayView<FValueComponent> Component)
+{
+	return Private::UnaryOpInPlace(Private::FOpSign(), Type, Component);
+}
+
+EValueType FracInPlace(EValueType Type, TArrayView<FValueComponent> Component)
+{
+	return Private::UnaryOpInPlace(Private::FOpFrac(), Type, Component);
+}
+
+EValueType FractionalInPlace(EValueType Type, TArrayView<FValueComponent> Component)
+{
+	return Private::UnaryOpInPlace(Private::FOpFractional(), Type, Component);
+}
+
+EValueType SqrtInPlace(EValueType Type, TArrayView<FValueComponent> Component)
+{
+	return Private::UnaryOpInPlace(Private::FOpSqrt(), Type, Component);
+}
+
+EValueType RcpInPlace(EValueType Type, TArrayView<FValueComponent> Component)
+{
+	return Private::UnaryOpInPlace(Private::FOpRcp(), Type, Component);
+}
+
+EValueType Log2InPlace(EValueType Type, TArrayView<FValueComponent> Component)
+{
+	return Private::UnaryOpInPlace(Private::FOpLog2(), Type, Component);
+}
+
+EValueType Log10InPlace(EValueType Type, TArrayView<FValueComponent> Component)
+{
+	return Private::UnaryOpInPlace(Private::FOpLog10(), Type, Component);
+}
+
+EValueType SinInPlace(EValueType Type, TArrayView<FValueComponent> Component)
+{
+	return Private::UnaryOpInPlace(Private::FOpSin(), Type, Component);
+}
+
+EValueType CosInPlace(EValueType Type, TArrayView<FValueComponent> Component)
+{
+	return Private::UnaryOpInPlace(Private::FOpCos(), Type, Component);
+}
+
+EValueType TanInPlace(EValueType Type, TArrayView<FValueComponent> Component)
+{
+	return Private::UnaryOpInPlace(Private::FOpTan(), Type, Component);
+}
+
+EValueType AsinInPlace(EValueType Type, TArrayView<FValueComponent> Component)
+{
+	return Private::UnaryOpInPlace(Private::FOpAsin(), Type, Component);
+}
+
+EValueType AcosInPlace(EValueType Type, TArrayView<FValueComponent> Component)
+{
+	return Private::UnaryOpInPlace(Private::FOpAcos(), Type, Component);
+}
+
+EValueType AtanInPlace(EValueType Type, TArrayView<FValueComponent> Component)
+{
+	return Private::UnaryOpInPlace(Private::FOpAtan(), Type, Component);
+}
+
+EValueType AddInPlace(EValueType LhsType, EValueType RhsType, TArrayView<FValueComponent> Component, int32& OutComponentsConsumed)
+{
+	return Private::BinaryOpInPlace(Private::FOpAdd(), LhsType, RhsType, Component, OutComponentsConsumed);
+}
+
+EValueType SubInPlace(EValueType LhsType, EValueType RhsType, TArrayView<FValueComponent> Component, int32& OutComponentsConsumed)
+{
+	return Private::BinaryOpInPlace(Private::FOpSub(), LhsType, RhsType, Component, OutComponentsConsumed);
+}
+
+EValueType MulInPlace(EValueType LhsType, EValueType RhsType, TArrayView<FValueComponent> Component, int32& OutComponentsConsumed)
+{
+	return Private::BinaryOpInPlace(Private::FOpMul(), LhsType, RhsType, Component, OutComponentsConsumed);
+}
+
+EValueType DivInPlace(EValueType LhsType, EValueType RhsType, TArrayView<FValueComponent> Component, int32& OutComponentsConsumed)
+{
+	return Private::BinaryOpInPlace(Private::FOpDiv(), LhsType, RhsType, Component, OutComponentsConsumed);
+}
+
+EValueType MinInPlace(EValueType LhsType, EValueType RhsType, TArrayView<FValueComponent> Component, int32& OutComponentsConsumed)
+{
+	return Private::BinaryOpInPlace(Private::FOpMin(), LhsType, RhsType, Component, OutComponentsConsumed);
+}
+
+EValueType MaxInPlace(EValueType LhsType, EValueType RhsType, TArrayView<FValueComponent> Component, int32& OutComponentsConsumed)
+{
+	return Private::BinaryOpInPlace(Private::FOpMax(), LhsType, RhsType, Component, OutComponentsConsumed);
+}
+
+EValueType FmodInPlace(EValueType LhsType, EValueType RhsType, TArrayView<FValueComponent> Component, int32& OutComponentsConsumed)
+{
+	return Private::BinaryOpInPlace(Private::FOpFmod(), LhsType, RhsType, Component, OutComponentsConsumed);
+}
+
+EValueType Atan2InPlace(EValueType LhsType, EValueType RhsType, TArrayView<FValueComponent> Component, int32& OutComponentsConsumed)
+{
+	return Private::BinaryOpInPlace(Private::FOpAtan2(), LhsType, RhsType, Component, OutComponentsConsumed);
+}
+
+EValueType AppendInPlace(EValueType LhsType, EValueType RhsType, TArrayView<FValueComponent> Component, int32& OutComponentsConsumed)
+{
+	const FValueTypeDescription& LhsDesc = GetValueTypeDescription(LhsType);
+	const FValueTypeDescription& RhsDesc = GetValueTypeDescription(RhsType);
+	const int32 NumComponents = FMath::Min<int32>(LhsDesc.NumComponents + RhsDesc.NumComponents, 4);
+
+	OutComponentsConsumed = LhsDesc.NumComponents + RhsDesc.NumComponents - NumComponents;
+
+	if (LhsDesc.ComponentType == RhsDesc.ComponentType)
+	{
+		return MakeValueType(LhsDesc.ComponentType, NumComponents);
+	}
+	else if (LhsDesc.ComponentType == EValueComponentType::Double || RhsDesc.ComponentType == EValueComponentType::Double)
+	{
+		// One of them is double, cast whichever one is not
+		const Private::FCastDouble Cast;
+		if (LhsDesc.ComponentType != EValueComponentType::Double)
+		{
+			// Could original LhsDesc.NumComponents have been greater than 4?  Clamp just in case...
+			int32 LhsClampedComponents = FMath::Min(LhsDesc.NumComponents, 4);
+			for (int32 i = 0; i < LhsClampedComponents; ++i)
+			{
+				Component[i].Double = Cast(LhsDesc.ComponentType, Component[i]);
+			}
+		}
+		else
+		{
+			for (int32 i = LhsDesc.NumComponents; i < NumComponents; ++i)
+			{
+				Component[i].Double = Cast(RhsDesc.ComponentType, Component[i]);
+			}
+		}
+		return MakeValueType(EValueComponentType::Double, NumComponents);
+	}
+	else if (LhsDesc.ComponentType == EValueComponentType::Float || RhsDesc.ComponentType == EValueComponentType::Float)
+	{
+		// One of them is float, cast whichever one is not
+		const Private::FCastFloat Cast;
+		if (LhsDesc.ComponentType != EValueComponentType::Float)
+		{
+			// Could original LhsDesc.NumComponents have been greater than 4?  Clamp just in case...
+			int32 LhsClampedComponents = FMath::Min(LhsDesc.NumComponents, 4);
+			for (int32 i = 0; i < LhsClampedComponents; ++i)
+			{
+				Component[i].Float = Cast(LhsDesc.ComponentType, Component[i]);
+			}
+		}
+		else
+		{
+			for (int32 i = LhsDesc.NumComponents; i < NumComponents; ++i)
+			{
+				Component[i].Float = Cast(RhsDesc.ComponentType, Component[i]);
+			}
+		}
+		return MakeValueType(EValueComponentType::Float, NumComponents);
+	}
+	else
+	{
+		// Cast everything else to int
+		const Private::FCastInt Cast;
+		int32 LhsClampedComponents = FMath::Min(LhsDesc.NumComponents, 4);
+		for (int32 i = 0; i < LhsClampedComponents; ++i)
+		{
+			Component[i].Int = Cast(LhsDesc.ComponentType, Component[i]);
+		}
+		for (int32 i = LhsDesc.NumComponents; i < NumComponents; ++i)
+		{
+			Component[i].Int = Cast(RhsDesc.ComponentType, Component[i]);
+		}
+		return MakeValueType(EValueComponentType::Int, NumComponents);
+	}
+}
+
 } // namespace Shader
 } // namespace UE
diff --git a/Engine/Source/Runtime/Engine/Public/MaterialShared.h b/Engine/Source/Runtime/Engine/Public/MaterialShared.h
--- a/Engine/Source/Runtime/Engine/Public/MaterialShared.h
+++ b/Engine/Source/Runtime/Engine/Public/MaterialShared.h
@@ -584,6 +584,11 @@ public:
 		return UE::Shader::FValue::FromMemoryImage(GetShaderValueType(Type), DefaultValues.GetData() + Offset);
 	}
 
+	const uint8* GetDefaultParameterData(uint32 Offset) const
+	{
+		return DefaultValues.GetData() + Offset;
+	}
+
 	inline const FMaterialNumericParameterInfo& GetNumericParameter(uint32 Index) const { return UniformNumericParameters[Index]; }
 	inline const FMaterialTextureParameterInfo& GetTextureParameter(EMaterialTextureParameterType Type, int32 Index) const { return UniformTextureParameters[(uint32)Type][Index]; }
 
diff --git a/Engine/Source/Runtime/Engine/Public/Shader/PreshaderTypes.h b/Engine/Source/Runtime/Engine/Public/Shader/PreshaderTypes.h
--- a/Engine/Source/Runtime/Engine/Public/Shader/PreshaderTypes.h
+++ b/Engine/Source/Runtime/Engine/Public/Shader/PreshaderTypes.h
@@ -53,11 +53,53 @@ public:
 	void PushValue(const FPreshaderValue& InValue);
 	void PushValue(const FPreshaderType& InType, TArrayView<const FValueComponent> InComponents);
 	TArrayView<FValueComponent> PushEmptyValue(const FPreshaderType& InType);
+	FValueComponent* PushEmptyValue(EValueType InType, int32 NumComponents);
 
 	/** Returned values are invalidated when anything is pushed onto the stack */
 	FPreshaderValue PopValue();
 	FPreshaderValue PeekValue(int32 Offset = 0);
 
+	/** Utility functions used for in-place swizzle operation */
+	inline const FPreshaderType& PeekType()
+	{
+		return Values.Last();
+	}
+	inline FValueComponent* PeekComponents(int32 NumComponents)
+	{
+		return &Components[Components.Num() - NumComponents];
+	}
+
+	/** Pops the final result from the stack, with reduced overhead relative to PopValue -- must be zero or one items left! */
+	void PopResult(FPreshaderValue& OutValue);
+
+	/** Override top value type (after an in place unary operation) */
+	inline void OverrideTopType(EValueType ValueType)
+	{
+		check(!Values.Last().IsStruct());
+		Values.Last().ValueType = ValueType;
+	}
+
+	/** Merge top two values (after an in place binary operation) */
+	inline void MergeTopTwoValues(EValueType ValueType, int32 ComponentsConsumed)
+	{
+		Values.RemoveAt(Values.Num() - 1, 1, false);
+		Components.RemoveAt(Components.Num() - ComponentsConsumed, ComponentsConsumed, false);
+		Values.Last().ValueType = ValueType;
+	}
+
+	// Adjust component count, by adding or removing components by the specified amount
+	inline void AdjustComponentCount(int32 NumComponentChange)
+	{
+		if (NumComponentChange > 0)
+		{
+			Components.AddUninitialized(NumComponentChange);
+		}
+		else if (NumComponentChange < 0)
+		{
+			Components.RemoveAt(Components.Num() + NumComponentChange, -NumComponentChange, false);
+		}
+	}
+
 	void Reset()
 	{
 		Values.Reset();
@@ -65,8 +107,8 @@ public:
 	}
 
 private:
-	TArray<FPreshaderType, TInlineAllocator<8>> Values;
-	TArray<FValueComponent, TInlineAllocator<64>> Components;
+	TArray<FPreshaderType, TNonRelocatableInlineAllocator<8>> Values;
+	TArray<FValueComponent, TNonRelocatableInlineAllocator<64>> Components;
 };
 
 } // namespace Shader
diff --git a/Engine/Source/Runtime/Engine/Public/Shader/ShaderTypes.h b/Engine/Source/Runtime/Engine/Public/Shader/ShaderTypes.h
--- a/Engine/Source/Runtime/Engine/Public/Shader/ShaderTypes.h
+++ b/Engine/Source/Runtime/Engine/Public/Shader/ShaderTypes.h
@@ -141,15 +141,14 @@ static constexpr int32 NumValueTypes = (int32)EValueType::Num;
 
 struct FValueTypeDescription
 {
-	FValueTypeDescription() : Name(nullptr), ComponentType(EValueComponentType::Void), NumComponents(0) {}
-	FValueTypeDescription(const TCHAR* InName, EValueComponentType InComponentType, int8 InNumComponents) : Name(InName), ComponentType(InComponentType), NumComponents(InNumComponents) {}
-
 	const TCHAR* Name;
+	EValueType ValueType;					// Useful for validation
 	EValueComponentType ComponentType;
 	int8 NumComponents;
+	int8 ComponentSizeInBytes;
 };
 
-ENGINE_API FValueTypeDescription GetValueTypeDescription(EValueType Type);
+ENGINE_API const FValueTypeDescription& GetValueTypeDescription(EValueType Type);
 ENGINE_API EValueType FindValueType(FName Name);
 inline bool IsLWCType(EValueType Type) { return IsLWCType(GetValueTypeDescription(Type).ComponentType); }
 inline bool IsGenericType(EValueType Type) { return Type == EValueType::Any || IsGenericType(GetValueTypeDescription(Type).ComponentType); }
@@ -574,6 +573,39 @@ ENGINE_API FValue Append(const FValue& Lhs, const FValue& Rhs);
 
 ENGINE_API FValue Cast(const FValue& Value, EValueType Type);
 
+// In place versions of certain functions.  Only works on numeric types.  Overwrites input components in place,
+// possibly changing their type (returns the new type).
+ENGINE_API EValueType NegInPlace(EValueType Type, TArrayView<FValueComponent> Component);
+ENGINE_API EValueType AbsInPlace(EValueType Type, TArrayView<FValueComponent> Component);
+ENGINE_API EValueType SaturateInPlace(EValueType Type, TArrayView<FValueComponent> Component);
+ENGINE_API EValueType FloorInPlace(EValueType Type, TArrayView<FValueComponent> Component);
+ENGINE_API EValueType CeilInPlace(EValueType Type, TArrayView<FValueComponent> Component);
+ENGINE_API EValueType RoundInPlace(EValueType Type, TArrayView<FValueComponent> Component);
+ENGINE_API EValueType TruncInPlace(EValueType Type, TArrayView<FValueComponent> Component);
+ENGINE_API EValueType SignInPlace(EValueType Type, TArrayView<FValueComponent> Component);
+ENGINE_API EValueType FracInPlace(EValueType Type, TArrayView<FValueComponent> Component);
+ENGINE_API EValueType FractionalInPlace(EValueType Type, TArrayView<FValueComponent> Component);
+ENGINE_API EValueType SqrtInPlace(EValueType Type, TArrayView<FValueComponent> Component);
+ENGINE_API EValueType RcpInPlace(EValueType Type, TArrayView<FValueComponent> Component);
+ENGINE_API EValueType Log2InPlace(EValueType Type, TArrayView<FValueComponent> Component);
+ENGINE_API EValueType Log10InPlace(EValueType Type, TArrayView<FValueComponent> Component);
+ENGINE_API EValueType SinInPlace(EValueType Type, TArrayView<FValueComponent> Component);
+ENGINE_API EValueType CosInPlace(EValueType Type, TArrayView<FValueComponent> Component);
+ENGINE_API EValueType TanInPlace(EValueType Type, TArrayView<FValueComponent> Component);
+ENGINE_API EValueType AsinInPlace(EValueType Type, TArrayView<FValueComponent> Component);
+ENGINE_API EValueType AcosInPlace(EValueType Type, TArrayView<FValueComponent> Component);
+ENGINE_API EValueType AtanInPlace(EValueType Type, TArrayView<FValueComponent> Component);
+
+ENGINE_API EValueType AddInPlace(EValueType LhsType, EValueType RhsType, TArrayView<FValueComponent> Component, int32& OutComponentsConsumed);
+ENGINE_API EValueType SubInPlace(EValueType LhsType, EValueType RhsType, TArrayView<FValueComponent> Component, int32& OutComponentsConsumed);
+ENGINE_API EValueType MulInPlace(EValueType LhsType, EValueType RhsType, TArrayView<FValueComponent> Component, int32& OutComponentsConsumed);
+ENGINE_API EValueType DivInPlace(EValueType LhsType, EValueType RhsType, TArrayView<FValueComponent> Component, int32& OutComponentsConsumed);
+ENGINE_API EValueType MinInPlace(EValueType LhsType, EValueType RhsType, TArrayView<FValueComponent> Component, int32& OutComponentsConsumed);
+ENGINE_API EValueType MaxInPlace(EValueType LhsType, EValueType RhsType, TArrayView<FValueComponent> Component, int32& OutComponentsConsumed);
+ENGINE_API EValueType FmodInPlace(EValueType LhsType, EValueType RhsType, TArrayView<FValueComponent> Component, int32& OutComponentsConsumed);
+ENGINE_API EValueType Atan2InPlace(EValueType LhsType, EValueType RhsType, TArrayView<FValueComponent> Component, int32& OutComponentsConsumed);
+ENGINE_API EValueType AppendInPlace(EValueType LhsType, EValueType RhsType, TArrayView<FValueComponent> Component, int32& OutComponentsConsumed);
+
 } // namespace UE::Shader
 
 DECLARE_INTRINSIC_TYPE_LAYOUT(UE::Shader::EValueType);
