Cloth asset - Add synchronous DDC render data support and serialization of sk mesh model

diff --git a/Engine/Source/Developer/MeshBuilder/Public/IMeshBuilderModule.h b/Engine/Source/Developer/MeshBuilder/Public/IMeshBuilderModule.h
--- a/Engine/Source/Developer/MeshBuilder/Public/IMeshBuilderModule.h
+++ b/Engine/Source/Developer/MeshBuilder/Public/IMeshBuilderModule.h
@@ -35,5 +35,5 @@ public:
 
 	virtual bool BuildSkeletalMesh(const struct FSkeletalMeshBuildParameters& SkeletalMeshBuildParameters) = 0;
 
-	virtual void PostBuildSkeletalMesh(class FSkeletalMeshRenderData* SkeletalMeshRenderData, class USkeletalMesh* SkeletalMesh) { }
+	virtual void PostBuildSkeletalMesh(class FSkeletalMeshRenderData* SkeletalMeshRenderData, class USkinnedAsset* SkinnedAsset) { }
 };
diff --git a/Engine/Source/Runtime/Engine/Classes/Engine/SkeletalMesh.h b/Engine/Source/Runtime/Engine/Classes/Engine/SkeletalMesh.h
--- a/Engine/Source/Runtime/Engine/Classes/Engine/SkeletalMesh.h
+++ b/Engine/Source/Runtime/Engine/Classes/Engine/SkeletalMesh.h
@@ -763,7 +763,7 @@ namespace NSSkeletalMeshSourceFileLabels
  * @see https://docs.unrealengine.com/latest/INT/Engine/Content/Types/SkeletalMeshes/
  */
 UCLASS(hidecategories=Object, BlueprintType)
-class ENGINE_API USkeletalMesh : public USkinnedAsset, public IInterface_CollisionDataProvider, public IInterface_AssetUserData, public INodeMappingProviderInterface, public IInterface_AsyncCompilation
+class ENGINE_API USkeletalMesh : public USkinnedAsset, public IInterface_CollisionDataProvider, public IInterface_AssetUserData, public INodeMappingProviderInterface
 {
 	GENERATED_UCLASS_BODY()
 
@@ -902,7 +902,7 @@ public:
 	int32 GetNumImportedVertices() const;
 
 	/** Get the imported data for this skeletal mesh. */
-	FORCEINLINE FSkeletalMeshModel* GetImportedModel() const 
+	virtual FSkeletalMeshModel* GetImportedModel() const override
 	{
 		WaitUntilAsyncPropertyReleased(ESkeletalMeshAsyncProperties::ImportedModel);
 		PRAGMA_DISABLE_DEPRECATION_WARNINGS
@@ -1092,8 +1092,8 @@ public:
 		return false;
 	}
 
-	/* Return true if this skeletalmesh was never build since its creation. */
-	bool IsInitialBuildDone() const;
+	/* Return true if this skeletalmesh was never build since its creation. USkinnedAsset interface. */
+	virtual bool IsInitialBuildDone() const override;
 
 	/* Return true if the reduction settings are setup to reduce a LOD*/
 	bool IsReductionActive(int32 LODIndex) const;
@@ -1230,8 +1230,8 @@ public:
 		PRAGMA_ENABLE_DEPRECATION_WARNINGS
 	}
 
-	/** Check the QualitLevel property is enabled for MinLod. */
-	bool IsMinLodQualityLevelEnable() const;
+	/** Check the QualitLevel property is enabled for MinLod. USkinnedAsset interface. */
+	virtual bool IsMinLodQualityLevelEnable() const override;
 
 	/** USkinnedAsset interface */
 	virtual int32 GetPlatformMinLODIdx(const ITargetPlatform* TargetPlatform) const override;
@@ -1256,7 +1256,8 @@ public:
 		PRAGMA_ENABLE_DEPRECATION_WARNINGS
 	}
 
-	const FPerPlatformInt& GetMinLod() const
+	/** USkinnedAsset interface. */
+	virtual const FPerPlatformInt& GetMinLod() const override
 	{
 		WaitUntilAsyncPropertyReleased(ESkeletalMeshAsyncProperties::MinLod, EAsyncPropertyLockType::ReadOnly);
 		PRAGMA_DISABLE_DEPRECATION_WARNINGS
@@ -1475,6 +1476,9 @@ public:
 
 	/** Get the maximum number of optional LODs. Do not use MaxNumOptionalLODs directly. Call this method instead. USkinnedAsset Interface. */
 	virtual int32 GetMaxNumOptionalLODs(const class ITargetPlatform* TargetPlatform) const override;
+
+	/* Build a LOD model before creating its render data. USkinnedAsset Interface. */
+	virtual void BuildLODModel(const ITargetPlatform* TargetPlatform, int32 LODIndex) override;
 #endif
 
 	UFUNCTION(BlueprintSetter)
@@ -2672,11 +2676,6 @@ public:
 	 */
 	virtual const FMeshUVChannelInfo* GetUVChannelData(int32 MaterialIndex) const override;
 
-	/**
-	 * Computes flags for building vertex buffers.
-	 */
-	uint32 GetVertexBufferFlags() const;
-
 	//~ Begin UObject Interface.
 #if WITH_EDITOR
 private:
@@ -2992,6 +2991,9 @@ public:
 	*/
 	void InvalidateDeriveDataCacheGUID();
 
+	/** Generate the derived data key for the given platform. USkinnedAsset interface. */
+	virtual FString BuildDerivedDataKey(const ITargetPlatform* TargetPlatform) override;
+
 	/** Generate the derived data key used to fetch derived data */
 	FString GetDerivedDataKey();
 #endif 
@@ -3272,7 +3274,8 @@ public:
 		PRAGMA_ENABLE_DEPRECATION_WARNINGS
 	}
 
-	bool GetUseLegacyMeshDerivedDataKey() const
+	// USkinnedAsset interface
+	virtual bool GetUseLegacyMeshDerivedDataKey() const override
 	{
 		WaitUntilAsyncPropertyReleased(ESkeletalMeshAsyncProperties::UseLegacyMeshDerivedDataKey, EAsyncPropertyLockType::ReadOnly);
 		PRAGMA_DISABLE_DEPRECATION_WARNINGS
diff --git a/Engine/Source/Runtime/Engine/Classes/Engine/SkinnedAsset.h b/Engine/Source/Runtime/Engine/Classes/Engine/SkinnedAsset.h
--- a/Engine/Source/Runtime/Engine/Classes/Engine/SkinnedAsset.h
+++ b/Engine/Source/Runtime/Engine/Classes/Engine/SkinnedAsset.h
@@ -9,6 +9,7 @@
 #include "CoreMinimal.h"
 #include "UObject/ObjectMacros.h"
 #include "Engine/StreamableRenderAsset.h"
+#include "Interfaces/Interface_AsyncCompilation.h"
 #include "ReferenceSkeleton.h"
 #include "PerPlatformProperties.h"
 #include "SkinnedAsset.generated.h"
@@ -21,7 +22,7 @@ struct FSkeletalMaterial;
 struct FSkeletalMeshLODInfo;
 
 UCLASS(hidecategories = Object, config = Engine, editinlinenew, abstract)
-class ENGINE_API USkinnedAsset : public UStreamableRenderAsset
+class ENGINE_API USkinnedAsset : public UStreamableRenderAsset, public IInterface_AsyncCompilation
 {
 	GENERATED_BODY()
 
@@ -97,6 +98,12 @@ public:
 	virtual int32 GetDefaultMinLod() const
 	PURE_VIRTUAL(USkinnedAsset::GetDefaultMinLod, return 0;);
 
+	virtual const FPerPlatformInt& GetMinLod() const
+	PURE_VIRTUAL(USkinnedAsset::GetMinLod, static const FPerPlatformInt Dummy;  return Dummy;);
+
+	/** Check the QualitLevel property is enabled for MinLod. */
+	virtual bool IsMinLodQualityLevelEnable() const { return false; }
+
 	virtual UPhysicsAsset* GetPhysicsAsset() const
 	PURE_VIRTUAL(USkinnedAsset::GetPhysicsAsset, return nullptr;);
 
@@ -161,6 +168,10 @@ public:
 
 	virtual void SetSkinWeightProfilesData(int32 LODIndex, struct FSkinWeightProfilesData& SkinWeightProfilesData) {}
 
+	/** Computes flags for building vertex buffers. */
+	virtual uint32 GetVertexBufferFlags() const
+	{ return GetHasVertexColors() ? ESkeletalMeshVertexFlags::HasVertexColors : ESkeletalMeshVertexFlags::None; }
+
 	/**
 	* UObject Interface
 	* This will return detail info about this specific object. (e.g. AudioComponent will return the name of the cue,
@@ -172,6 +183,20 @@ public:
 	{ return GetPathName(nullptr); }
 
 #if WITH_EDITOR
+	/** IInterface_AsyncCompilation begin*/
+	virtual bool IsCompiling() const override { return false; }
+	/** IInterface_AsyncCompilation end*/
+
+	virtual FString BuildDerivedDataKey(const ITargetPlatform* TargetPlatform)
+	PURE_VIRTUAL(USkinnedAsset::BuildDerivedDataKey, return TEXT(""););
+
+	/* Return true if this asset was never build since its creation. */
+	virtual bool IsInitialBuildDone() const
+	PURE_VIRTUAL(USkinnedAsset::IsInitialBuildDone, return false;);
+
+	/* Build a LOD model before creating its render data. */
+	virtual void BuildLODModel(const ITargetPlatform* TargetPlatform, int32 LODIndex) {}
+
 	/** Get whether this mesh should use LOD streaming for the given platform. */
 	virtual bool GetEnableLODStreaming(const ITargetPlatform* TargetPlatform) const
 	PURE_VIRTUAL(USkinnedAsset::GetEnableLODStreaming, return false;);
@@ -182,6 +207,14 @@ public:
 
 	virtual int32 GetMaxNumOptionalLODs(const ITargetPlatform* TargetPlatform) const
 	PURE_VIRTUAL(USkinnedAsset::GetMaxNumOptionalLODs, return 0;);
-#endif
+#endif // WITH_EDITOR
+
+#if WITH_EDITORONLY_DATA
+	virtual bool GetUseLegacyMeshDerivedDataKey() const	{ return false; }
+
+	/** Get the source mesh data. */
+	virtual class FSkeletalMeshModel* GetImportedModel() const
+	PURE_VIRTUAL(USkinnedAsset::GetImportedModel, return nullptr;);
+#endif // WITH_EDITORONLY_DATA
 };
 
diff --git a/Engine/Source/Runtime/Engine/Private/SkeletalMesh.cpp b/Engine/Source/Runtime/Engine/Private/SkeletalMesh.cpp
--- a/Engine/Source/Runtime/Engine/Private/SkeletalMesh.cpp
+++ b/Engine/Source/Runtime/Engine/Private/SkeletalMesh.cpp
@@ -71,7 +71,7 @@
 #include "SkeletalMeshCompiler.h"
 #include "MeshUtilities.h"
 #include "Engine/SkeletalMeshEditorData.h"
-
+#include "DerivedDataCacheInterface.h"
 #include "IMeshReductionManagerModule.h"
 #include "SkeletalMeshReductionSettings.h"
 #include "Engine/RendererSettings.h"
@@ -415,7 +415,11 @@ void FSkeletalMeshAsyncBuildWorker::DoWork()
 	}
 }
 
-extern const FString& GetSkeletalMeshDerivedDataVersion();
+const FString& GetSkeletalMeshDerivedDataVersion()
+{
+	static FString CachedVersionString = FDevSystemGuids::GetSystemGuid(FDevSystemGuids::Get().SkeletalMeshDerivedDataVersion).ToString();
+	return CachedVersionString;
+}
 
 namespace SkeletalMeshImpl
 {
@@ -1540,11 +1544,9 @@ void CachePlatform(USkeletalMesh* Mesh, const ITargetPlatform* TargetPlatform, F
 	PlatformRenderData->Cache(TargetPlatform, Mesh, &Context);
 }
 
-FString BuildSkeletalMeshDerivedDataKey(const ITargetPlatform* TargetPlatform, USkeletalMesh* SkelMesh);
-
 static FSkeletalMeshRenderData& GetPlatformSkeletalMeshRenderData(USkeletalMesh* Mesh, const ITargetPlatform* TargetPlatform, const bool bIsSerializeSaving)
 {
-	const FString PlatformDerivedDataKey = BuildSkeletalMeshDerivedDataKey(TargetPlatform, Mesh);
+	const FString PlatformDerivedDataKey = Mesh->BuildDerivedDataKey(TargetPlatform);
 	FSkeletalMeshRenderData* PlatformRenderData = Mesh->GetResourceForRendering();
 	if (Mesh->GetOutermost()->bIsCookedForEditor)
 	{
@@ -1881,16 +1883,6 @@ void USkeletalMesh::FlushRenderState()
 	ReleaseResourcesFence.Wait();
 }
 
-uint32 USkeletalMesh::GetVertexBufferFlags() const
-{
-	uint32 VertexFlags = ESkeletalMeshVertexFlags::None;
-	if (GetHasVertexColors())
-	{
-		VertexFlags |= ESkeletalMeshVertexFlags::HasVertexColors;
-	}
-	return VertexFlags;
-}
-
 #if WITH_EDITOR
 
 thread_local const USkeletalMesh* FSkeletalMeshAsyncBuildScope::SkeletalMeshBeingAsyncCompiled = nullptr;
@@ -4513,21 +4505,120 @@ void USkeletalMesh::ClearAllCachedCookedPlatformData()
 		//Normally this call should be able to read values out of ddc rather than rebuilding, because the ddc for the running platform was cached when we loaded the asset.
 		ITargetPlatform* RunningPlatform = GetTargetPlatformManagerRef().GetRunningTargetPlatform();
 		check(RunningPlatform != NULL);
-		const FString RunningPlatformDerivedDataKey = BuildSkeletalMeshDerivedDataKey(RunningPlatform, this);
+		const FString RunningPlatformDerivedDataKey = BuildDerivedDataKey(RunningPlatform);
 		FSkeletalMeshRenderData RunningPlatformSkeletalMeshRenderData;
 		constexpr bool bIsSerializeSaving = false;
 		CachePlatform(this, RunningPlatform, &RunningPlatformSkeletalMeshRenderData, bIsSerializeSaving);
 	}
 }
 
-extern FString BuildSkeletalMeshDerivedDataKey(const ITargetPlatform* TargetPlatform, USkeletalMesh* SkelMesh);
+//Serialize the LODInfo and append the result to the KeySuffix to build the LODInfo part of the DDC KEY
+//Note: this serializer is only used to build the mesh DDC key, no versioning is required
+static void SerializeLODInfoForDDC(USkeletalMesh* SkeletalMesh, FString& KeySuffix)
+{
+	TArray<FSkeletalMeshLODInfo>& LODInfos = SkeletalMesh->GetLODInfoArray();
+	const bool bIs16BitfloatBufferSupported = GVertexElementTypeSupport.IsSupported(VET_Half2);
+	for (int32 LODIndex = 0; LODIndex < SkeletalMesh->GetLODNum(); ++LODIndex)
+	{
+		check(LODInfos.IsValidIndex(LODIndex));
+		FSkeletalMeshLODInfo& LODInfo = LODInfos[LODIndex];
+		bool bValidLODSettings = false;
+		if (SkeletalMesh->GetLODSettings() != nullptr)
+		{
+			const int32 NumSettings = FMath::Min(SkeletalMesh->GetLODSettings()->GetNumberOfSettings(), SkeletalMesh->GetLODNum());
+			if (LODIndex < NumSettings)
+			{
+				bValidLODSettings = true;
+			}
+		}
+		const FSkeletalMeshLODGroupSettings* SkeletalMeshLODGroupSettings = bValidLODSettings ? &SkeletalMesh->GetLODSettings()->GetSettingsForLODLevel(LODIndex) : nullptr;
+		LODInfo.BuildGUID = LODInfo.ComputeDeriveDataCacheKey(SkeletalMeshLODGroupSettings);
+		KeySuffix += LODInfo.BuildGUID.ToString(EGuidFormats::Digits);
+	}
+}
+
+extern int32 GStripSkeletalMeshLodsDuringCooking;
+extern int32 GSkeletalMeshKeepMobileMinLODSettingOnDesktop;
+
+FString USkeletalMesh::BuildDerivedDataKey(const ITargetPlatform* TargetPlatform)
+{
+	FString KeySuffix(TEXT(""));
+
+	if (GetUseLegacyMeshDerivedDataKey())
+	{
+		//Old asset will have the same LOD settings for bUseFullPrecisionUVs. We can use the LOD 0
+		const FSkeletalMeshLODInfo* BaseLODInfo = GetLODInfo(0);
+		bool bUseFullPrecisionUVs = BaseLODInfo ? BaseLODInfo->BuildSettings.bUseFullPrecisionUVs : false;
+		KeySuffix += GetImportedModel()->GetIdString();
+		KeySuffix += (bUseFullPrecisionUVs || !GVertexElementTypeSupport.IsSupported(VET_Half2)) ? "1" : "0";
+
+		//Dummy call to update the FSkeletalMeshLODModel::BuildStringID members.
+		GetImportedModel()->GetLODModelIdString();
+
+		//Dummy call to update the LODInfo::BuildGUID members.
+		{
+			FString TmpString;
+			SerializeLODInfoForDDC(this, TmpString);
+		}
+	}
+	else
+	{
+		FString TmpPartialKeySuffix;
+		//Synchronize the user data that are part of the key
+		GetImportedModel()->SyncronizeLODUserSectionsData();
+		TmpPartialKeySuffix = GetImportedModel()->GetIdString();
+		KeySuffix += TmpPartialKeySuffix;
+		TmpPartialKeySuffix = GetImportedModel()->GetLODModelIdString();
+		KeySuffix += TmpPartialKeySuffix;
+
+		//Add the max gpu bone per section
+		const int32 MaxGPUSkinBones = FGPUBaseSkinVertexFactory::GetMaxGPUSkinBones(TargetPlatform);
+		KeySuffix += FString::FromInt(MaxGPUSkinBones);
+
+		TmpPartialKeySuffix = TEXT("");
+		SerializeLODInfoForDDC(this, TmpPartialKeySuffix);
+		KeySuffix += TmpPartialKeySuffix;
+	}
+
+	KeySuffix += GetHasVertexColors() ? "1" : "0";
+	KeySuffix += GetVertexColorGuid().ToString(EGuidFormats::Digits);
+
+	if (GetEnableLODStreaming(TargetPlatform))
+	{
+		const int32 MaxStreamedLODs = GetMaxNumStreamedLODs(TargetPlatform);
+		const int32 MaxOptionalLODs = GetMaxNumOptionalLODs(TargetPlatform);
+		KeySuffix += *FString::Printf(TEXT("1%08x%08x"), MaxStreamedLODs, MaxOptionalLODs);
+	}
+	else
+	{
+		KeySuffix += TEXT("0zzzzzzzzzzzzzzzz");
+	}
+
+	if (TargetPlatform->GetPlatformInfo().PlatformGroupName == TEXT("Desktop")
+		&& GStripSkeletalMeshLodsDuringCooking != 0
+		&& GSkeletalMeshKeepMobileMinLODSettingOnDesktop != 0)
+	{
+		KeySuffix += TEXT("_MinMLOD");
+	}
+
+	IMeshBuilderModule::GetForPlatform(TargetPlatform).AppendToDDCKey(KeySuffix, true);
+	const bool bUnlimitedBoneInfluences = FGPUBaseSkinVertexFactory::GetUnlimitedBoneInfluences();
+	KeySuffix += bUnlimitedBoneInfluences ? "1" : "0";
+
+	return FDerivedDataCacheInterface::BuildCacheKey(
+		TEXT("SKELETALMESH"),
+		*GetSkeletalMeshDerivedDataVersion(),
+		*KeySuffix
+	);
+}
+
 FString USkeletalMesh::GetDerivedDataKey()
 {
 	// Cache derived data for the running platform.
 	ITargetPlatform* RunningPlatform = GetTargetPlatformManagerRef().GetRunningTargetPlatform();
 	check(RunningPlatform);
 
-	return BuildSkeletalMeshDerivedDataKey(RunningPlatform, this);
+	return BuildDerivedDataKey(RunningPlatform);
 }
 
 int32 USkeletalMesh::ValidatePreviewAttachedObjects()
@@ -5306,6 +5397,37 @@ int32 USkeletalMesh::GetMaxNumOptionalLODs(const ITargetPlatform* TargetPlatform
 		return GetDefault<URendererSettings>()->bDiscardSkeletalMeshOptionalLODs.GetValueForPlatform(*TargetPlatform->IniPlatformName()) ? 0 : MAX_MESH_LOD_COUNT;
 	}
 }
+
+void USkeletalMesh::BuildLODModel(const ITargetPlatform* TargetPlatform, int32 LODIndex)
+{
+	FSkeletalMeshModel* SkelMeshModel = GetImportedModel();
+	check(SkelMeshModel);
+
+	FSkeletalMeshLODInfo* LODInfoPtr = GetLODInfo(LODIndex);
+	check(LODInfoPtr);
+
+	//We want to avoid building a LOD if the LOD was generated from a previous LODIndex.
+	const bool bIsGeneratedLodNotInline = (LODInfoPtr->bHasBeenSimplified && IsReductionActive(LODIndex) && GetReductionSettings(LODIndex).BaseLOD < LODIndex);
+
+	//Make sure the LOD have all the data needed to be build
+	const bool bRawDataEmpty = IsLODImportedDataEmpty(LODIndex);
+	const bool bRawBuildDataAvailable = IsLODImportedDataBuildAvailable(LODIndex);
+
+	//Build the source model before the render data, if we are a purely generated LOD we do not need to be build
+	IMeshBuilderModule& MeshBuilderModule = IMeshBuilderModule::GetForPlatform(TargetPlatform);
+	if (!bIsGeneratedLodNotInline && !bRawDataEmpty && bRawBuildDataAvailable)
+	{
+		LODInfoPtr->bHasBeenSimplified = false;
+		const bool bRegenDepLODs = true;
+		FSkeletalMeshBuildParameters BuildParameters(this, TargetPlatform, LODIndex, bRegenDepLODs);
+		MeshBuilderModule.BuildSkeletalMesh(BuildParameters);
+	}
+	else
+	{
+		//We need to synchronize when we are generated mesh or if we have load an old asset that was not re-imported
+		SkelMeshModel->LODModels[LODIndex].SyncronizeUserSectionsDataArray();
+	}
+}
 #endif
 
 void USkeletalMesh::SetLODSettings(USkeletalMeshLODSettings* InLODSettings)
diff --git a/Engine/Source/Runtime/Engine/Private/SkeletalMeshModel.cpp b/Engine/Source/Runtime/Engine/Private/SkeletalMeshModel.cpp
--- a/Engine/Source/Runtime/Engine/Private/SkeletalMeshModel.cpp
+++ b/Engine/Source/Runtime/Engine/Private/SkeletalMeshModel.cpp
@@ -14,7 +14,7 @@ FSkeletalMeshModel::FSkeletalMeshModel()
 {
 }
 
-void FSkeletalMeshModel::Serialize(FArchive& Ar, USkeletalMesh* Owner)
+void FSkeletalMeshModel::Serialize(FArchive& Ar, USkinnedAsset* Owner)
 {
 	DECLARE_SCOPE_CYCLE_COUNTER(TEXT("FSkeletalMeshModel::Serialize"), STAT_SkeletalMeshModel_Serialize, STATGROUP_LoadTime);
 	bool bIsEditorDataStripped = false;
@@ -113,7 +113,7 @@ void FSkeletalMeshModel::GenerateNewGUID()
 }
 
 
-void FSkeletalMeshModel::GenerateGUIDFromHash(USkeletalMesh* Owner)
+void FSkeletalMeshModel::GenerateGUIDFromHash(USkinnedAsset* Owner)
 {
 	// Build the hash from the path name + the contents of the bulk data.
 	FSHA1 Sha;
diff --git a/Engine/Source/Runtime/Engine/Private/SkeletalMeshRenderData.cpp b/Engine/Source/Runtime/Engine/Private/SkeletalMeshRenderData.cpp
--- a/Engine/Source/Runtime/Engine/Private/SkeletalMeshRenderData.cpp
+++ b/Engine/Source/Runtime/Engine/Private/SkeletalMeshRenderData.cpp
@@ -38,9 +38,10 @@ extern int32 GStripSkeletalMeshLodsDuringCooking;
 
 #endif // WITH_EDITOR
 
-static TAutoConsoleVariable<int32> CVarSkeletalMeshKeepMobileMinLODSettingOnDesktop(
+int32 GSkeletalMeshKeepMobileMinLODSettingOnDesktop = 0;
+static FAutoConsoleVariableRef CVarSkeletalMeshKeepMobileMinLODSettingOnDesktop(
 	TEXT("r.SkeletalMesh.KeepMobileMinLODSettingOnDesktop"),
-	0,
+	GSkeletalMeshKeepMobileMinLODSettingOnDesktop,
 	TEXT("If non-zero, mobile setting for MinLOD will be stored in the cooked data for desktop platforms"));
 
 #if WITH_EDITOR
@@ -63,7 +64,7 @@ namespace DDCUtils64Bit
 	};
 
 	/** The same as calling GetDerivedDataCacheRef().GetSynchronous(...) but with a TArray64 as the output parameter. */
-	bool GetSynchronous(const FString& DerivedDataKey, USkeletalMesh* Owner, TArray64<uint8>& OutDerivedData)
+	bool GetSynchronous(const FString& DerivedDataKey, USkinnedAsset* Owner, TArray64<uint8>& OutDerivedData)
 	{
 		TStringBuilder<512> OwnerPathName;
 		Owner->GetPathName(nullptr, OwnerPathName);
@@ -114,7 +115,7 @@ namespace DDCUtils64Bit
 	}
 
 	/** The same as calling GetDerivedDataCacheRef().Put(...) but with a TArrayView64 as the input data. */
-	void Put(const FString& DerivedDataKey, USkeletalMesh* Owner, TArrayView64<const uint8> DerivedData)
+	void Put(const FString& DerivedDataKey, USkinnedAsset* Owner, TArrayView64<const uint8> DerivedData)
 	{
 		TStringBuilder<512> OwnerPathName;
 		Owner->GetPathName(nullptr, OwnerPathName);
@@ -158,111 +159,8 @@ namespace DDCUtils64Bit
 	}
 } //namespace DDCUtils64Bit
 
-//Serialize the LODInfo and append the result to the KeySuffix to build the LODInfo part of the DDC KEY
-//Note: this serializer is only used to build the mesh DDC key, no versioning is required
-static void SerializeLODInfoForDDC(USkeletalMesh* SkeletalMesh, FString& KeySuffix)
-{
-	TArray<FSkeletalMeshLODInfo>& LODInfos = SkeletalMesh->GetLODInfoArray();
-	const bool bIs16BitfloatBufferSupported = GVertexElementTypeSupport.IsSupported(VET_Half2);
-	for (int32 LODIndex = 0; LODIndex < SkeletalMesh->GetLODNum(); ++LODIndex)
-	{
-		check(LODInfos.IsValidIndex(LODIndex));
-		FSkeletalMeshLODInfo& LODInfo = LODInfos[LODIndex];
-		bool bValidLODSettings = false;
-		if (SkeletalMesh->GetLODSettings() != nullptr)
-		{
-			const int32 NumSettings = FMath::Min(SkeletalMesh->GetLODSettings()->GetNumberOfSettings(), SkeletalMesh->GetLODNum());
-			if (LODIndex < NumSettings)
-			{
-				bValidLODSettings = true;
-			}
-		}
-		const FSkeletalMeshLODGroupSettings* SkeletalMeshLODGroupSettings = bValidLODSettings ? &SkeletalMesh->GetLODSettings()->GetSettingsForLODLevel(LODIndex) : nullptr;
-		LODInfo.BuildGUID = LODInfo.ComputeDeriveDataCacheKey(SkeletalMeshLODGroupSettings);
-		KeySuffix += LODInfo.BuildGUID.ToString(EGuidFormats::Digits);
-	}
-}
-
-const FString& GetSkeletalMeshDerivedDataVersion()
-{
-	static FString CachedVersionString = FDevSystemGuids::GetSystemGuid(FDevSystemGuids::Get().SkeletalMeshDerivedDataVersion).ToString();
-	return CachedVersionString;
-}
-
-FString BuildSkeletalMeshDerivedDataKey(const ITargetPlatform* TargetPlatform, USkeletalMesh* SkelMesh)
-{
-	FString KeySuffix(TEXT(""));
-
-	if (SkelMesh->GetUseLegacyMeshDerivedDataKey() )
-	{
-		//Old asset will have the same LOD settings for bUseFullPrecisionUVs. We can use the LOD 0
-		const FSkeletalMeshLODInfo* BaseLODInfo = SkelMesh->GetLODInfo(0);
-		bool bUseFullPrecisionUVs = BaseLODInfo ? BaseLODInfo->BuildSettings.bUseFullPrecisionUVs : false;
-		KeySuffix += SkelMesh->GetImportedModel()->GetIdString();
-		KeySuffix += (bUseFullPrecisionUVs || !GVertexElementTypeSupport.IsSupported(VET_Half2)) ? "1" : "0";
-		
-		//Dummy call to update the FSkeletalMeshLODModel::BuildStringID members.
-		SkelMesh->GetImportedModel()->GetLODModelIdString();
-		
-		//Dummy call to update the LODInfo::BuildGUID members.
-		{
-			FString TmpString;
-			SerializeLODInfoForDDC(SkelMesh, TmpString);
-		}
-	}
-	else
-	{
-		FString TmpPartialKeySuffix;
-		//Synchronize the user data that are part of the key
-		SkelMesh->GetImportedModel()->SyncronizeLODUserSectionsData();
-		TmpPartialKeySuffix = SkelMesh->GetImportedModel()->GetIdString();
-		KeySuffix += TmpPartialKeySuffix;
-		TmpPartialKeySuffix = SkelMesh->GetImportedModel()->GetLODModelIdString();
-		KeySuffix += TmpPartialKeySuffix;
-		
-		//Add the max gpu bone per section
-		const int32 MaxGPUSkinBones = FGPUBaseSkinVertexFactory::GetMaxGPUSkinBones(TargetPlatform);
-		KeySuffix += FString::FromInt(MaxGPUSkinBones);
-
-		TmpPartialKeySuffix = TEXT("");
-		SerializeLODInfoForDDC(SkelMesh, TmpPartialKeySuffix);
-		KeySuffix += TmpPartialKeySuffix;
-	}
-
-	KeySuffix += SkelMesh->GetHasVertexColors() ? "1" : "0";
-	KeySuffix += SkelMesh->GetVertexColorGuid().ToString(EGuidFormats::Digits);
-	
-	if (SkelMesh->GetEnableLODStreaming(TargetPlatform))
-	{
-		const int32 MaxNumStreamedLODs = SkelMesh->GetMaxNumStreamedLODs(TargetPlatform);
-		const int32 MaxNumOptionalLODs = SkelMesh->GetMaxNumOptionalLODs(TargetPlatform);
-		KeySuffix += *FString::Printf(TEXT("1%08x%08x"), MaxNumStreamedLODs, MaxNumOptionalLODs);
-	}
-	else
-	{
-		KeySuffix += TEXT("0zzzzzzzzzzzzzzzz");
-	}
-
-	if (TargetPlatform->GetPlatformInfo().PlatformGroupName == TEXT("Desktop")
-		&& GStripSkeletalMeshLodsDuringCooking != 0
-		&& CVarSkeletalMeshKeepMobileMinLODSettingOnDesktop.GetValueOnAnyThread() != 0)
-	{
-		KeySuffix += TEXT("_MinMLOD");
-	}
-
-	IMeshBuilderModule::GetForPlatform(TargetPlatform).AppendToDDCKey(KeySuffix, true);
-	const bool bUnlimitedBoneInfluences = FGPUBaseSkinVertexFactory::GetUnlimitedBoneInfluences();
-	KeySuffix += bUnlimitedBoneInfluences ? "1" : "0";
-
-	return FDerivedDataCacheInterface::BuildCacheKey(
-		TEXT("SKELETALMESH"),
-		*GetSkeletalMeshDerivedDataVersion(),
-		*KeySuffix
-	);
-}
-
 /** This code verify that the data is all in sync index buffer versus sections data. It is active only in debug build*/
-void VerifyAllLodSkeletalMeshModelIntegrity(USkeletalMesh* Owner)
+void VerifyAllLodSkeletalMeshModelIntegrity(USkinnedAsset* Owner)
 {
 	if (!Owner || !Owner->GetImportedModel())
 	{
@@ -312,15 +210,16 @@ void VerifyAllLodSkeletalMeshModelIntegrity(USkeletalMesh* Owner)
 	}
 }
 
-FString FSkeletalMeshRenderData::GetDerivedDataKey(const ITargetPlatform* TargetPlatform, USkeletalMesh* Owner)
+FString FSkeletalMeshRenderData::GetDerivedDataKey(const ITargetPlatform* TargetPlatform, USkinnedAsset* Owner)
 {
-	return BuildSkeletalMeshDerivedDataKey(TargetPlatform, Owner);
+	return Owner->BuildDerivedDataKey(TargetPlatform);
 }
 
-void FSkeletalMeshRenderData::Cache(const ITargetPlatform* TargetPlatform, USkeletalMesh* Owner, FSkeletalMeshCompilationContext* ContextPtr)
+void FSkeletalMeshRenderData::Cache(const ITargetPlatform* TargetPlatform, USkinnedAsset* Owner, FSkeletalMeshCompilationContext* ContextPtr)
 {
 	check(Owner);
-	check(ContextPtr);
+	// Disable ContextPtr check because only USkeletalMesh supports it.
+	//check(ContextPtr);
 
 	check(LODRenderData.Num() == 0); // Should only be called on new, empty RenderData
 	check(TargetPlatform);
@@ -351,7 +250,7 @@ void FSkeletalMeshRenderData::Cache(const ITargetPlatform* TargetPlatform, USkel
 		const bool bAllowDdcFetch = Owner->IsInitialBuildDone();
 		if (bAllowDdcFetch)
 		{
-			DerivedDataKey = BuildSkeletalMeshDerivedDataKey(TargetPlatform, Owner);
+			DerivedDataKey = Owner->BuildDerivedDataKey(TargetPlatform);
 		}
 		
 		//If we have an initial build, the ddc key will be computed only after the build. Some structure are missing until we first build the asset to get the drived data key
@@ -416,7 +315,11 @@ void FSkeletalMeshRenderData::Cache(const ITargetPlatform* TargetPlatform, USkel
 			//Apply the morphtargets change if any
 			if (FinishBuildMorphTargetData.IsValid())
 			{
-				FinishBuildMorphTargetData->ApplyEditorData(Owner, ContextPtr->bIsSerializeSaving);
+				// Morph target is only supported on USkeletalMesh
+				if (USkeletalMesh* SkMesh = Cast<USkeletalMesh>(Owner))
+				{
+					FinishBuildMorphTargetData->ApplyEditorData(SkMesh, ContextPtr ? ContextPtr->bIsSerializeSaving : false);
+				}
 			}
 
 			int32 T1 = FPlatformTime::Cycles();
@@ -432,31 +335,11 @@ void FSkeletalMeshRenderData::Cache(const ITargetPlatform* TargetPlatform, USkel
 
 			for (int32 LODIndex = 0; LODIndex < SkelMeshModel->LODModels.Num(); LODIndex++)
 			{
+				Owner->BuildLODModel(TargetPlatform, LODIndex);
+
 				FSkeletalMeshLODModel* LODModel = &(SkelMeshModel->LODModels[LODIndex]);
 				FSkeletalMeshLODInfo* LODInfo = Owner->GetLODInfo(LODIndex);
 				check(LODInfo);
-				//We want to avoid building a LOD if the LOD was generated from a previous LODIndex.
-				const bool bIsGeneratedLodNotInline = (LODInfo->bHasBeenSimplified && Owner->IsReductionActive(LODIndex) && Owner->GetReductionSettings(LODIndex).BaseLOD < LODIndex);
-				
-				//Make sure the LOD have all the data needed to be build
-				const bool bRawDataEmpty = Owner->IsLODImportedDataEmpty(LODIndex);
-				const bool bRawBuildDataAvailable = Owner->IsLODImportedDataBuildAvailable(LODIndex);
-				
-				//Build the source model before the render data, if we are a purely generated LOD we do not need to be build
-				IMeshBuilderModule& MeshBuilderModule = IMeshBuilderModule::GetForPlatform(TargetPlatform);
-				if (!bIsGeneratedLodNotInline && !bRawDataEmpty && bRawBuildDataAvailable)
-				{
-					LODInfo->bHasBeenSimplified = false;
-					const bool bRegenDepLODs = true;
-					FSkeletalMeshBuildParameters BuildParameters(Owner, TargetPlatform, LODIndex, bRegenDepLODs);
-					MeshBuilderModule.BuildSkeletalMesh(BuildParameters);
-					LODModel = &(SkelMeshModel->LODModels[LODIndex]);
-				}
-				else
-				{
-					//We need to synchronize when we are generated mesh or if we have load an old asset that was not re-imported
-					LODModel->SyncronizeUserSectionsDataArray();
-				}
 
 				FSkeletalMeshLODRenderData* LODData = new FSkeletalMeshLODRenderData();
 				LODRenderData.Add(LODData);
@@ -525,7 +408,7 @@ void FSkeletalMeshRenderData::Cache(const ITargetPlatform* TargetPlatform, USkel
 
 			//Recompute the derived data key in case there was some data correction during the build process, this make sure the DDC key is always representing the correct build result.
 			//There should never be correction of the data during the build, the data has to be corrected in the post load before calling this function.
-			FString BuiltDerivedDataKey = BuildSkeletalMeshDerivedDataKey(TargetPlatform, Owner);
+			FString BuiltDerivedDataKey = Owner->BuildDerivedDataKey(TargetPlatform);
 			//Only compare keys if the ddc fetch was allowed
 			if (bAllowDdcFetch)
 			{
@@ -660,7 +543,7 @@ int32 FSkeletalMeshRenderData::GetNumNonOptionalLODs() const
 	}
 }
 
-void FSkeletalMeshRenderData::Serialize(FArchive& Ar, USkeletalMesh* Owner)
+void FSkeletalMeshRenderData::Serialize(FArchive& Ar, USkinnedAsset* Owner)
 {
 	DECLARE_SCOPE_CYCLE_COUNTER(TEXT("FSkeletalMeshRenderData::Serialize"), STAT_SkeletalMeshRenderData_Serialize, STATGROUP_LoadTime);
 
@@ -669,13 +552,13 @@ void FSkeletalMeshRenderData::Serialize(FArchive& Ar, USkeletalMesh* Owner)
 	if (Ar.IsFilterEditorOnly())
 	{
 		int32 MinMobileLODIdx = 0;
-		bool bShouldSerialize = CVarSkeletalMeshKeepMobileMinLODSettingOnDesktop.GetValueOnAnyThread() != 0;
+		bool bShouldSerialize = GSkeletalMeshKeepMobileMinLODSettingOnDesktop != 0;
 #if WITH_EDITOR
 		if (Ar.IsSaving())
 		{
 			if (Ar.CookingTarget()->GetPlatformInfo().PlatformGroupName == TEXT("Desktop")
 				&& GStripSkeletalMeshLodsDuringCooking != 0
-				&& CVarSkeletalMeshKeepMobileMinLODSettingOnDesktop.GetValueOnAnyThread() != 0)
+				&& GSkeletalMeshKeepMobileMinLODSettingOnDesktop != 0)
 			{
 				// Serialize 0 value when per quality level properties are used
 				if (!Owner->IsMinLodQualityLevelEnable())
@@ -728,7 +611,7 @@ void FSkeletalMeshRenderData::Serialize(FArchive& Ar, USkeletalMesh* Owner)
 	PendingFirstLODIdx = CurrentFirstLODIdx;
 
 PRAGMA_DISABLE_DEPRECATION_WARNINGS
-	bSupportRayTracing = Owner->bSupportRayTracing;
+	bSupportRayTracing = Owner->GetSupportRayTracing();
 PRAGMA_ENABLE_DEPRECATION_WARNINGS
 }
 
diff --git a/Engine/Source/Runtime/Engine/Public/ReferenceSkeleton.h b/Engine/Source/Runtime/Engine/Public/ReferenceSkeleton.h
--- a/Engine/Source/Runtime/Engine/Public/ReferenceSkeleton.h
+++ b/Engine/Source/Runtime/Engine/Public/ReferenceSkeleton.h
@@ -408,6 +408,6 @@ public:
 
 	// very slow search function for all children
 	int32 GetDirectChildBones(int32 ParentBoneIndex, TArray<int32> & Children) const;
-	friend FArchive & operator<<(FArchive & Ar, FReferenceSkeleton & F);
+	ENGINE_API friend FArchive & operator<<(FArchive & Ar, FReferenceSkeleton & F);
 	friend FReferenceSkeletonModifier;
 };
diff --git a/Engine/Source/Runtime/Engine/Public/Rendering/SkeletalMeshModel.h b/Engine/Source/Runtime/Engine/Public/Rendering/SkeletalMeshModel.h
--- a/Engine/Source/Runtime/Engine/Public/Rendering/SkeletalMeshModel.h
+++ b/Engine/Source/Runtime/Engine/Public/Rendering/SkeletalMeshModel.h
@@ -10,6 +10,7 @@
 
 class FSkeletalMeshLODModel;
 class USkeletalMesh;
+class USkinnedAsset;
 class FArchive;
 
 /**
@@ -27,24 +28,24 @@ public:
 	TIndirectArray<FSkeletalMeshLODModel> LODModels;
 
 	/** Default constructor. */
-	FSkeletalMeshModel();
+	ENGINE_API FSkeletalMeshModel();
 
 #if WITH_EDITOR
 	/** Creates a new GUID for this Model */
-	void GenerateNewGUID();
+	ENGINE_API void GenerateNewGUID();
 
 	/** 
 	 *	Util to regenerate a GUID for this Model based on hashing its data 
 	 *	Used by old content, rather than a random new GUID.
 	 */
-	void GenerateGUIDFromHash(USkeletalMesh* Owner);
+	void GenerateGUIDFromHash(USkinnedAsset* Owner);
 
 	/** Get current GUID Id as a string, for DDC key */
-	FString GetIdString() const;
+	ENGINE_API FString GetIdString() const;
 
-	void SyncronizeLODUserSectionsData();
+	ENGINE_API void SyncronizeLODUserSectionsData();
 
-	FString GetLODModelIdString() const;
+	ENGINE_API FString GetLODModelIdString() const;
 
 	ENGINE_API void EmptyOriginalReductionSourceMeshData();
 	
@@ -58,7 +59,7 @@ public:
 #endif
 
 	/** Serialize to/from the specified archive.. */
-	void Serialize(FArchive& Ar, USkeletalMesh* Owner);
+	ENGINE_API void Serialize(FArchive& Ar, USkinnedAsset* Owner);
 
 	void GetResourceSizeEx(FResourceSizeEx& CumulativeResourceSize);
 };
diff --git a/Engine/Source/Runtime/Engine/Public/Rendering/SkeletalMeshRenderData.h b/Engine/Source/Runtime/Engine/Public/Rendering/SkeletalMeshRenderData.h
--- a/Engine/Source/Runtime/Engine/Public/Rendering/SkeletalMeshRenderData.h
+++ b/Engine/Source/Runtime/Engine/Public/Rendering/SkeletalMeshRenderData.h
@@ -8,6 +8,7 @@
 
 struct FMeshUVChannelInfo;
 class USkeletalMesh;
+class USkinnedAsset;
 struct FSkeletalMaterial;
 class UMorphTarget;
 struct FResourceSizeEx;
@@ -36,7 +37,7 @@ public:
 	/** Runtime LOD bias modifier for this skeletal mesh */
 	uint8 LODBiasModifier;
 
-	/** Whether ray tracing acceleration structures should be created for this mesh. Derived from owner USkeletalMesh. */
+	/** Whether ray tracing acceleration structures should be created for this mesh. Derived from owner USkinnedAsset. */
 	bool bSupportRayTracing;
 
 #if WITH_EDITORONLY_DATA
@@ -54,14 +55,14 @@ public:
 	ENGINE_API ~FSkeletalMeshRenderData();
 
 #if WITH_EDITOR
-	void Cache(const ITargetPlatform* TargetPlatform, USkeletalMesh* Owner, class FSkeletalMeshCompilationContext* ContextPtr);
-	FString GetDerivedDataKey(const ITargetPlatform* TargetPlatform, USkeletalMesh* Owner);
+	ENGINE_API void Cache(const ITargetPlatform* TargetPlatform, USkinnedAsset* Owner, class FSkeletalMeshCompilationContext* ContextPtr);
+	FString GetDerivedDataKey(const ITargetPlatform* TargetPlatform, USkinnedAsset* Owner);
 
-	void SyncUVChannelData(const TArray<FSkeletalMaterial>& ObjectData);
+	ENGINE_API void SyncUVChannelData(const TArray<FSkeletalMaterial>& ObjectData);
 #endif
 
 	/** Serialize to/from the specified archive.. */
-	void Serialize(FArchive& Ar, USkeletalMesh* Owner);
+	void Serialize(FArchive& Ar, USkinnedAsset* Owner);
 
 	/** Initializes rendering resources. */
 	ENGINE_API void InitResources(bool bNeedsVertexColors, TArray<UMorphTarget*>& InMorphTargets, USkinnedAsset* Owner);
