Avoid using asset registry while inside ULevel::Postload to avoid triggering asset registry tag update that triggers BP event during postload.

diff --git a/Engine/Source/Runtime/Engine/Public/ExternalPackageHelper.h b/Engine/Source/Runtime/Engine/Public/ExternalPackageHelper.h
--- a/Engine/Source/Runtime/Engine/Public/ExternalPackageHelper.h
+++ b/Engine/Source/Runtime/Engine/Public/ExternalPackageHelper.h
@@ -5,10 +5,12 @@
 #if WITH_EDITOR
 
 #include "UObject/Package.h"
+#include "UObject/UObjectThreadContext.h"
 #include "AssetRegistry/AssetRegistryModule.h"
 #include "Misc/ArchiveMD5.h"
 #include "Misc/PackageName.h"
 #include "Misc/PackagePath.h"
+#include "HAL/FileManager.h"
 #include "Delegates/DelegateCombinations.h"
 
 class ENGINE_API FExternalPackageHelper
@@ -71,23 +73,63 @@ private:
 	static FString GetExternalObjectPackageInstanceName(const FString& OuterPackageName, const FString& ObjectPackageName);
 };
 
+static TArray<FString> GetOnDiskExternalObjectPackages(const FString& ExternalObjectsPath)
+{
+	TArray<FString> ObjectPackageNames;
+	if (!ExternalObjectsPath.IsEmpty())
+	{
+		IFileManager::Get().IterateDirectoryRecursively(*FPackageName::LongPackageNameToFilename(ExternalObjectsPath), [&ObjectPackageNames](const TCHAR* FilenameOrDirectory, bool bIsDirectory)
+		{
+			if (!bIsDirectory)
+			{
+				FString Filename(FilenameOrDirectory);
+				if (Filename.EndsWith(FPackageName::GetAssetPackageExtension()))
+				{
+					ObjectPackageNames.Add(FPackageName::FilenameToLongPackageName(Filename));
+				}
+			}
+			return true;
+		});
+	}
+	return ObjectPackageNames;
+}
+
 template<typename T>
 void FExternalPackageHelper::LoadObjectsFromExternalPackages(UObject* InOuter, TFunctionRef<void(T*)> Operation)
 {
 	const FString ExternalObjectsPath = FExternalPackageHelper::GetExternalObjectsPath(InOuter->GetPackage(), FString(), /*bTryUsingPackageLoadedPath*/ true);
+	TArray<FString> ObjectPackageNames;
 
-	// Do a synchronous scan of the world external objects path.			
-	IAssetRegistry& AssetRegistry = FModuleManager::LoadModuleChecked<FAssetRegistryModule>(TEXT("AssetRegistry")).Get();
-	AssetRegistry.ScanPathsSynchronous({ ExternalObjectsPath }, /*bForceRescan*/false, /*bIgnoreDenyListScanFilters*/false);
+	// Test if function is called while inside postload
+	if (!FUObjectThreadContext::Get().IsRoutingPostLoad)
+	{
+		// Do a synchronous scan of the world external objects path.			
+		IAssetRegistry& AssetRegistry = FModuleManager::LoadModuleChecked<FAssetRegistryModule>(TEXT("AssetRegistry")).Get();
+		AssetRegistry.ScanPathsSynchronous({ ExternalObjectsPath }, /*bForceRescan*/false, /*bIgnoreDenyListScanFilters*/false);
 
-	FARFilter Filter;
-	Filter.bRecursivePaths = true;
-	Filter.bIncludeOnlyOnDiskAssets = true;
-	Filter.ClassPaths.Add(T::StaticClass()->GetClassPathName());
-	Filter.PackagePaths.Add(*ExternalObjectsPath);
+		FARFilter Filter;
+		Filter.bRecursivePaths = true;
+		Filter.bIncludeOnlyOnDiskAssets = true;
+		Filter.ClassPaths.Add(T::StaticClass()->GetClassPathName());
+		Filter.PackagePaths.Add(*ExternalObjectsPath);
 
-	TArray<FAssetData> Assets;
-	AssetRegistry.GetAssets(Filter, Assets);
+		TArray<FAssetData> Assets;
+		AssetRegistry.GetAssets(Filter, Assets);
+
+		ObjectPackageNames.Reserve(Assets.Num());
+		for (const FAssetData& Asset : Assets)
+		{
+			ObjectPackageNames.Add(Asset.PackageName.ToString());
+		}
+	}
+	// If so, we can't use AssetRegistry's ScanPathsSynchronous as it might update AssetRegistry object tags which can call BP events 
+	// (which is not allowed while inside postload). Instead, get all external object packages on disk.
+	// @todo_ow : This is far from ideal because once we move External Actors to share the same folder as external objects (or any new object we will externalize),
+	//            we will discover and iterate on all external packages regarless of the class type we actually want to load.
+	else
+	{
+		 ObjectPackageNames = GetOnDiskExternalObjectPackages(ExternalObjectsPath);
+	}
 
 	FLinkerInstancingContext InstancingContext;
 	TArray<UPackage*> InstancePackages;
@@ -99,9 +141,8 @@ void FExternalPackageHelper::LoadObjectsFromExternalPackages(UObject* InOuter, T
 	{
 		InstancingContext.AddMapping(PackageResourceName, OuterPackage->GetFName());
 
-		for (const FAssetData& Asset : Assets)
+		for (const FString& ObjectPackageName : ObjectPackageNames)
 		{
-			const FString ObjectPackageName = Asset.PackageName.ToString();
 			const FString InstancedName = GetExternalObjectPackageInstanceName(OuterPackage->GetName(), ObjectPackageName);
 			InstancingContext.AddMapping(FName(*ObjectPackageName), FName(*InstancedName));
 
@@ -117,20 +158,20 @@ void FExternalPackageHelper::LoadObjectsFromExternalPackages(UObject* InOuter, T
 	}
 
 	const ELoadFlags LoadFlags = InOuter->GetPackage()->HasAnyPackageFlags(PKG_PlayInEditor) ? LOAD_PackageForPIE : LOAD_None;
-	for (int32 i = 0; i < Assets.Num(); i++)
+	for (int32 i = 0; i < ObjectPackageNames.Num(); i++)
 	{
-		if (UPackage* Package = LoadPackage(bInstanced ? InstancePackages[i] : nullptr, *Assets[i].PackageName.ToString(), LoadFlags, nullptr, &InstancingContext))
+		if (UPackage* Package = LoadPackage(bInstanced ? InstancePackages[i] : nullptr, *ObjectPackageNames[i], LoadFlags, nullptr, &InstancingContext))
 		{
 			T* LoadedObject = nullptr;
 			ForEachObjectWithPackage(Package, [&LoadedObject](UObject* Object)
+			{
+				if (T* TypedObj = Cast<T>(Object))
 				{
-					if (T* TypedObj = Cast<T>(Object))
-					{
-						LoadedObject = TypedObj;
-						return false;
-					}
-					return true;
-				}, true, RF_NoFlags, EInternalObjectFlags::Unreachable);
+					LoadedObject = TypedObj;
+					return false;
+				}
+				return true;
+			}, true, RF_NoFlags, EInternalObjectFlags::Unreachable);
 
 			if (ensure(LoadedObject))
 			{
