Removed call to ULevelInstanceSubsystem::CanUsePackage inside ALevelInstance::PostLoad and moved it in PostRegisterAllComponents / PreSave which could lead to a BP events being called inside PostLoad.

diff --git a/Engine/Source/Runtime/Engine/Private/LevelInstance/LevelInstanceActor.cpp b/Engine/Source/Runtime/Engine/Private/LevelInstance/LevelInstanceActor.cpp
--- a/Engine/Source/Runtime/Engine/Private/LevelInstance/LevelInstanceActor.cpp
+++ b/Engine/Source/Runtime/Engine/Private/LevelInstance/LevelInstanceActor.cpp
@@ -11,6 +11,7 @@
 #include "Net/UnrealNetwork.h"
 
 #if WITH_EDITOR
+#include "UObject/ObjectSaveContext.h"
 #include "WorldPartition/WorldPartitionActorDesc.h"
 #include "WorldPartition/LevelInstance/LevelInstanceActorDesc.h"
 #include "LevelInstance/LevelInstanceEditorPivotActor.h"
@@ -68,7 +69,11 @@ void ALevelInstance::OnRep_LevelInstanceSpawnGuid()
 void ALevelInstance::PostRegisterAllComponents()
 {
 	Super::PostRegisterAllComponents();
-	
+
+#if WITH_EDITOR
+	ResetUnsupportedWorldAsset();
+#endif
+
 	if (GetLocalRole() == ENetRole::ROLE_Authority && GetWorld()->IsGameWorld())
 	{
 #if !WITH_EDITOR
@@ -147,21 +152,6 @@ void ALevelInstance::PostLoad()
 {
 	Super::PostLoad();
 
-#if WITH_EDITOR
-	if (!ULevelInstanceSubsystem::CanUsePackage(*WorldAsset.GetLongPackageName()))
-	{
-		UE_LOG(LogLevelInstance, Warning, TEXT("LevelInstance doesn't support partitioned world %s, make sure to flag world partition's 'Can be Used by Level Instance'."), *WorldAsset.GetLongPackageName());
-		WorldAsset.Reset();
-	}
-#endif
-		
-#if WITH_EDITORONLY_DATA
-	if (IsRunningCookCommandlet() && IsLoadingEnabled())
-	{
-		CookedWorldAsset = WorldAsset;
-	}
-#endif
-
 	OnLevelInstanceActorPostLoad.Broadcast(this);
 }
 
@@ -287,6 +277,32 @@ void ALevelInstance::PushLevelInstanceEditingStateToProxies(bool bInEditingState
 	LevelInstanceActorImpl.PushLevelInstanceEditingStateToProxies(bInEditingState);
 }
 
+void ALevelInstance::PreSave(FObjectPreSaveContext ObjectSaveContext)
+{
+	Super::PreSave(ObjectSaveContext);
+
+	if (IsRunningCookCommandlet())
+	{
+		ResetUnsupportedWorldAsset();
+
+#if WITH_EDITORONLY_DATA
+		if (IsLoadingEnabled())
+		{
+			CookedWorldAsset = WorldAsset;
+		}
+#endif
+	}
+}
+
+void ALevelInstance::ResetUnsupportedWorldAsset()
+{
+	if (!ULevelInstanceSubsystem::CanUsePackage(*WorldAsset.GetLongPackageName()))
+	{
+		UE_LOG(LogLevelInstance, Warning, TEXT("LevelInstance doesn't support partitioned world %s, make sure to flag world partition's 'Can be Used by Level Instance'."), *WorldAsset.GetLongPackageName());
+		WorldAsset.Reset();
+	}
+}
+
 #endif
 
 #undef LOCTEXT_NAMESPACE
diff --git a/Engine/Source/Runtime/Engine/Public/LevelInstance/LevelInstanceActor.h b/Engine/Source/Runtime/Engine/Public/LevelInstance/LevelInstanceActor.h
--- a/Engine/Source/Runtime/Engine/Public/LevelInstance/LevelInstanceActor.h
+++ b/Engine/Source/Runtime/Engine/Public/LevelInstance/LevelInstanceActor.h
@@ -80,6 +80,7 @@ public:
 	virtual void PreEditChange(FProperty* PropertyThatWillChange) override;
 	virtual void PostEditChangeProperty(FPropertyChangedEvent& PropertyChangedEvent) override;
 	virtual void PostEditImport() override;
+	virtual void PreSave(FObjectPreSaveContext ObjectSaveContext) override;
 
 	// AActor overrides
 	virtual void CheckForErrors() override;
@@ -96,6 +97,9 @@ public:
 
 	DECLARE_MULTICAST_DELEGATE_OneParam(FOnLevelInstanceActorPostLoad, ALevelInstance*);
 	static FOnLevelInstanceActorPostLoad OnLevelInstanceActorPostLoad;
+
+private:
+	void ResetUnsupportedWorldAsset();
 #endif
 };
 
