Add an error message when running the Trace.Screenshot command with the Screenshot channel disabled.

diff --git a/Engine/Source/Runtime/Engine/Private/ProfilingDebugging/TraceScreenshot.cpp b/Engine/Source/Runtime/Engine/Private/ProfilingDebugging/TraceScreenshot.cpp
--- a/Engine/Source/Runtime/Engine/Private/ProfilingDebugging/TraceScreenshot.cpp
+++ b/Engine/Source/Runtime/Engine/Private/ProfilingDebugging/TraceScreenshot.cpp
@@ -37,6 +37,12 @@ bool FTraceScreenshot::bSuppressWritingToFile = false;
 
 void FTraceScreenshot::RequestScreenshot(FString Name)
 {
+	if (!SHOULD_TRACE_SCREENSHOT())
+	{
+		UE_LOG(LogCore, Error, TEXT("Could not trace screenshot because the screenshot trace channel is off. Turn it on using \"Trace.Enable Screenshot\"."));
+		return;
+	}
+
 	bSuppressWritingToFile = true;
 
 	if (Name.IsEmpty())
