Static analysis fix -- initialize ComponentsConsumed, even though all code paths do fill it in.

#rnx

diff --git a/Engine/Source/Runtime/Engine/Private/Shader/Preshader.cpp b/Engine/Source/Runtime/Engine/Private/Shader/Preshader.cpp
--- a/Engine/Source/Runtime/Engine/Private/Shader/Preshader.cpp
+++ b/Engine/Source/Runtime/Engine/Private/Shader/Preshader.cpp
@@ -629,7 +629,8 @@ static inline void EvaluateBinaryOpInPlace(FPreshaderStack& Stack, const Operati
 	}
 	else
 	{
-		int32 ComponentsConsumed;
+		// Initialize to avoid static analysis warning -- all code paths do fill in ComponentsConsumed
+		int32 ComponentsConsumed = 0;
 		EValueType ResultType = Op(Lhs.Type.ValueType, Rhs.Type.ValueType, TArrayView<FValueComponent>(Lhs.Component.GetData(), Lhs.Component.Num() + Rhs.Component.Num()), ComponentsConsumed);
 		Stack.MergeTopTwoValues(ResultType, ComponentsConsumed);
 	}
