UE-151338 - Implement refactored implicit conversion for math types.

Added support for FVector2D and FBox2D. Unfortunately, this is a bit of a hack. We currently don't support UScriptStruct creation for FVector2d and FBox2d, since their FNames collide with FVector2D and FBox2D, respectively.

We can work around this by simply not relying on a UScriptStruct double variant to exist. The logic for the conversion table can simply fall back to use the base UScriptStruct instead. It's a little ugly, but it at least allows implicit conversion between these types.

Additionally, the "varianttype1" and "varianttype2" names seemed a bit too abstract. These were changed into "doublevariant" and "floatvariant", which helps with readability.

diff --git a/Engine/Source/Runtime/Engine/Internal/Kismet/BlueprintTypeConversions.h b/Engine/Source/Runtime/Engine/Internal/Kismet/BlueprintTypeConversions.h
--- a/Engine/Source/Runtime/Engine/Internal/Kismet/BlueprintTypeConversions.h
+++ b/Engine/Source/Runtime/Engine/Internal/Kismet/BlueprintTypeConversions.h
@@ -73,8 +73,6 @@ public:
 	UFUNCTION(CustomThunk, meta = (BlueprintInternalUseOnly = "true"))
 	static int32 ConvertFVector3fToFVector3d(int32 InFromData);
 
-	// LWC_TODO - FVector2D/FVector2d have FName collisions.
-#if 0
 	DECLARE_CONVERSION_FUNCTIONS(FVector2D, FVector2f, FVector2d);
 
 	UFUNCTION(CustomThunk, meta = (BlueprintInternalUseOnly = "true"))
@@ -82,7 +80,6 @@ public:
 
 	UFUNCTION(CustomThunk, meta = (BlueprintInternalUseOnly = "true"))
 	static int32 ConvertFVector2fToFVector2d(int32 InFromData);
-#endif
 
 	DECLARE_CONVERSION_FUNCTIONS(FVector4, FVector4f, FVector4d);
 
@@ -137,8 +134,6 @@ public:
 	UFUNCTION(CustomThunk, meta = (BlueprintInternalUseOnly = "true"))
 	static int32 ConvertFMatrix44fToFMatrix44d(int32 InFromData);
 
-	// LWC_TODO - FBox2D/FBox2d have FName collisions.
-#if 0
 	DECLARE_CONVERSION_FUNCTIONS(FBox2D, FBox2f, FBox2d);
 
 	UFUNCTION(CustomThunk, meta = (BlueprintInternalUseOnly = "true"))
@@ -146,5 +141,4 @@ public:
 
 	UFUNCTION(CustomThunk, meta = (BlueprintInternalUseOnly = "true"))
 	static int32 ConvertFBox2fToFBox2d(int32 InFromData);
-#endif
 };
diff --git a/Engine/Source/Runtime/Engine/Private/Kismet/BlueprintTypeConversions.cpp b/Engine/Source/Runtime/Engine/Private/Kismet/BlueprintTypeConversions.cpp
--- a/Engine/Source/Runtime/Engine/Private/Kismet/BlueprintTypeConversions.cpp
+++ b/Engine/Source/Runtime/Engine/Private/Kismet/BlueprintTypeConversions.cpp
@@ -5,8 +5,8 @@
 #define MAKE_CONVERSION_FUNCTION_NAME(SourceType, DestType)													\
 	Convert##SourceType##To##DestType
 
-#define DEFINE_CONVERSION_FUNCTIONS(BASETYPE, VARIANTTYPE1, VARIANTTYPE2)									\
-	DEFINE_FUNCTION(UBlueprintTypeConversions::MAKE_CONVERSION_EXEC_FUNCTION_NAME(VARIANTTYPE1, VARIANTTYPE2)) \
+#define DEFINE_CONVERSION_FUNCTIONS(BASETYPE, DOUBLEVARIANT, FLOATVARIANT)									\
+	DEFINE_FUNCTION(UBlueprintTypeConversions::MAKE_CONVERSION_EXEC_FUNCTION_NAME(DOUBLEVARIANT, FLOATVARIANT)) \
 	{																										\
 		using namespace UE::Kismet::BlueprintTypeConversions::Internal;										\
 																											\
@@ -15,9 +15,9 @@
 		Stack.StepCompiledIn(nullptr, nullptr);																\
 		const void* SourceAddr = Stack.MostRecentPropertyAddress;											\
 		P_FINISH;																							\
-		ConvertType<VARIANTTYPE1, VARIANTTYPE2>(SourceAddr, DestAddr);										\
+		ConvertType<DOUBLEVARIANT, FLOATVARIANT>(SourceAddr, DestAddr);										\
 	}																										\
-	DEFINE_FUNCTION(UBlueprintTypeConversions::MAKE_CONVERSION_EXEC_FUNCTION_NAME(VARIANTTYPE2, VARIANTTYPE1)) \
+	DEFINE_FUNCTION(UBlueprintTypeConversions::MAKE_CONVERSION_EXEC_FUNCTION_NAME(FLOATVARIANT, DOUBLEVARIANT)) \
 	{																										\
 		using namespace UE::Kismet::BlueprintTypeConversions::Internal;										\
 																											\
@@ -26,17 +26,17 @@
 		Stack.StepCompiledIn(nullptr, nullptr);																\
 		const void* SourceAddr = Stack.MostRecentPropertyAddress;											\
 		P_FINISH;																							\
-		ConvertType<VARIANTTYPE2, VARIANTTYPE1>(SourceAddr, DestAddr);										\
+		ConvertType<FLOATVARIANT, DOUBLEVARIANT>(SourceAddr, DestAddr);										\
 	}																										\
 	UE::Kismet::BlueprintTypeConversions::Internal::FStructConversionEntry BASETYPE##Entry(					\
 		&TBaseStructure<BASETYPE>::Get,																		\
 		&TVariantStructure<BASETYPE>::Get,																	\
-		&TVariantStructure<VARIANTTYPE1>::Get,																\
-		&TVariantStructure<VARIANTTYPE2>::Get,																\
-		TEXT(PREPROCESSOR_TO_STRING(MAKE_CONVERSION_FUNCTION_NAME(VARIANTTYPE1, VARIANTTYPE2))),			\
-		TEXT(PREPROCESSOR_TO_STRING(MAKE_CONVERSION_FUNCTION_NAME(VARIANTTYPE2, VARIANTTYPE1))),			\
-		&UE::Kismet::BlueprintTypeConversions::Internal::ConvertType<VARIANTTYPE1, VARIANTTYPE2>,			\
-		&UE::Kismet::BlueprintTypeConversions::Internal::ConvertType<VARIANTTYPE2, VARIANTTYPE1>			\
+		&TVariantStructure<DOUBLEVARIANT>::Get,																\
+		&TVariantStructure<FLOATVARIANT>::Get,																\
+		TEXT(PREPROCESSOR_TO_STRING(MAKE_CONVERSION_FUNCTION_NAME(DOUBLEVARIANT, FLOATVARIANT))),			\
+		TEXT(PREPROCESSOR_TO_STRING(MAKE_CONVERSION_FUNCTION_NAME(FLOATVARIANT, DOUBLEVARIANT))),			\
+		&UE::Kismet::BlueprintTypeConversions::Internal::ConvertType<DOUBLEVARIANT, FLOATVARIANT>,			\
+		&UE::Kismet::BlueprintTypeConversions::Internal::ConvertType<FLOATVARIANT, DOUBLEVARIANT>			\
 	);
 
 namespace UE::Kismet::BlueprintTypeConversions::Internal
@@ -88,20 +88,20 @@ struct FStructConversionEntry
 
 	FStructConversionEntry(GetUScriptStructFunctionT InGetBaseStruct,
 		GetUScriptStructFunctionT InGetVariantFromBaseStruct,
-		GetUScriptStructFunctionT InGetVariantStruct1,
-		GetUScriptStructFunctionT InGetVariantStruct2,
-		const TCHAR* InConvertVariant1ToVariant2FunctionName,
-		const TCHAR* InConvertVariant2ToVariant1FunctionName,
-		UE::Kismet::BlueprintTypeConversions::ConversionFunctionT InConvertVariant1ToVariant2Impl,
-		UE::Kismet::BlueprintTypeConversions::ConversionFunctionT InConvertVariant2ToVariant1Impl)
+		GetUScriptStructFunctionT InGetDoubleVariantStruct,
+		GetUScriptStructFunctionT InGetFloatVariantStruct,
+		const TCHAR* InConvertDoubleToFloatFunctionName,
+		const TCHAR* InConvertFloatToDoubleFunctionName,
+		UE::Kismet::BlueprintTypeConversions::ConversionFunctionT InConvertDoubleToFloatImpl,
+		UE::Kismet::BlueprintTypeConversions::ConversionFunctionT InConvertFloatToDoubleImpl)
 		: GetBaseStruct(InGetBaseStruct)
 		, GetVariantFromBaseStruct(InGetVariantFromBaseStruct)
-		, GetVariantStruct1(InGetVariantStruct1)
-		, GetVariantStruct2(InGetVariantStruct2)
-		, ConvertVariant1ToVariant2FunctionName(InConvertVariant1ToVariant2FunctionName)
-		, ConvertVariant2ToVariant1FunctionName(InConvertVariant2ToVariant1FunctionName)
-		, ConvertVariant1ToVariant2Impl(InConvertVariant1ToVariant2Impl)
-		, ConvertVariant2ToVariant1Impl(InConvertVariant2ToVariant1Impl)
+		, GetDoubleVariantStruct(InGetDoubleVariantStruct)
+		, GetFloatVariantStruct(InGetFloatVariantStruct)
+		, ConvertDoubleToFloatFunctionName(InConvertDoubleToFloatFunctionName)
+		, ConvertFloatToDoubleFunctionName(InConvertFloatToDoubleFunctionName)
+		, ConvertDoubleToFloatImpl(InConvertDoubleToFloatImpl)
+		, ConvertFloatToDoubleImpl(InConvertFloatToDoubleImpl)
 	{
 		NextEntry = FStructConversionEntry::GetListHead();
 		FStructConversionEntry::GetListHead() = this;
@@ -109,12 +109,12 @@ struct FStructConversionEntry
 
 	GetUScriptStructFunctionT GetBaseStruct;
 	GetUScriptStructFunctionT GetVariantFromBaseStruct;
-	GetUScriptStructFunctionT GetVariantStruct1;
-	GetUScriptStructFunctionT GetVariantStruct2;
-	const TCHAR* ConvertVariant1ToVariant2FunctionName;
-	const TCHAR* ConvertVariant2ToVariant1FunctionName;
-	UE::Kismet::BlueprintTypeConversions::ConversionFunctionT ConvertVariant1ToVariant2Impl;
-	UE::Kismet::BlueprintTypeConversions::ConversionFunctionT ConvertVariant2ToVariant1Impl;
+	GetUScriptStructFunctionT GetDoubleVariantStruct;
+	GetUScriptStructFunctionT GetFloatVariantStruct;
+	const TCHAR* ConvertDoubleToFloatFunctionName;
+	const TCHAR* ConvertFloatToDoubleFunctionName;
+	UE::Kismet::BlueprintTypeConversions::ConversionFunctionT ConvertDoubleToFloatImpl;
+	UE::Kismet::BlueprintTypeConversions::ConversionFunctionT ConvertFloatToDoubleImpl;
 	FStructConversionEntry* NextEntry;
 
 	static FStructConversionEntry*& GetListHead()
@@ -141,17 +141,34 @@ FStructConversionTable::FStructConversionTable()
 	Internal::FStructConversionEntry* Entry = Internal::FStructConversionEntry::GetListHead();
 	while (Entry)
 	{
-		StructVariantLookupTable.Add(Entry->GetBaseStruct(), Entry->GetVariantFromBaseStruct());
+		const UScriptStruct* BaseStruct = Entry->GetBaseStruct();
+		const UScriptStruct* VariantFromBaseStruct = Entry->GetVariantFromBaseStruct();
+		const UScriptStruct* DoubleVariantStruct = Entry->GetDoubleVariantStruct();
+		const UScriptStruct* FloatVariantStruct = Entry->GetFloatVariantStruct();
 
-		StructVariantPairT ConversionKeyPair1 = { Entry->GetVariantStruct1(), Entry->GetVariantStruct2() };
+		// The base and float variant structs *must* exist.
+		// However, we don't always have the double variant struct, as is the case with FBox2d and FVector2d.
+
+		check(BaseStruct);
+		check(FloatVariantStruct);
+
+		if (VariantFromBaseStruct)
+		{
+			StructVariantLookupTable.Add(BaseStruct, VariantFromBaseStruct);
+		}
+
+		const UScriptStruct* FloatKeyEntry = FloatVariantStruct;
+		const UScriptStruct* DoubleKeyEntry = DoubleVariantStruct ? DoubleVariantStruct : BaseStruct;
+
+		StructVariantPairT ConversionKeyPair1 = { DoubleKeyEntry, FloatKeyEntry };
 		ConversionFunctionPairT ConversionValuePair1 = 
-			{ Entry->ConvertVariant1ToVariant2Impl,
-			  BlueprintTypeConversionsClass->FindFunctionByName(Entry->ConvertVariant1ToVariant2FunctionName) };
+			{ Entry->ConvertDoubleToFloatImpl,
+			  BlueprintTypeConversionsClass->FindFunctionByName(Entry->ConvertDoubleToFloatFunctionName) };
 
-		StructVariantPairT ConversionKeyPair2 = { Entry->GetVariantStruct2(), Entry->GetVariantStruct1() };
+		StructVariantPairT ConversionKeyPair2 = { FloatKeyEntry, DoubleKeyEntry };
 		ConversionFunctionPairT ConversionValuePair2 = 
-			{ Entry->ConvertVariant2ToVariant1Impl,
-			  BlueprintTypeConversionsClass->FindFunctionByName(Entry->ConvertVariant2ToVariant1FunctionName) };
+			{ Entry->ConvertFloatToDoubleImpl,
+			  BlueprintTypeConversionsClass->FindFunctionByName(Entry->ConvertFloatToDoubleFunctionName) };
 
 		ImplicitCastLookupTable.Add(ConversionKeyPair1, ConversionValuePair1);
 		ImplicitCastLookupTable.Add(ConversionKeyPair2, ConversionValuePair2);
@@ -346,11 +363,87 @@ DEFINE_FUNCTION(UBlueprintTypeConversions::execConvertMapType)
 // Custom struct conversions
 
 DEFINE_CONVERSION_FUNCTIONS(FVector, FVector3d, FVector3f)
-// DEFINE_CONVERSION_FUNCTIONS(FVector2D, FVector2d, FVector2f)
 DEFINE_CONVERSION_FUNCTIONS(FVector4, FVector4d, FVector4f)
 DEFINE_CONVERSION_FUNCTIONS(FPlane, FPlane4d, FPlane4f)
 DEFINE_CONVERSION_FUNCTIONS(FQuat, FQuat4d, FQuat4f)
 DEFINE_CONVERSION_FUNCTIONS(FRotator, FRotator3d, FRotator3f)
 DEFINE_CONVERSION_FUNCTIONS(FTransform, FTransform3d, FTransform3f)
 DEFINE_CONVERSION_FUNCTIONS(FMatrix, FMatrix44d, FMatrix44f)
-// DEFINE_CONVERSION_FUNCTIONS(FBox2D, FBox2d, FBox2f)
+
+// LWC_TODO - This is a bit of a hack, but FBox2d and FVector2d don't have script structs due to FName collisions with their base type.
+// So we need to add our own expanded version of DEFINE_CONVERSION_FUNCTIONS that returns a nullptr for the variant script structs.
+
+UScriptStruct* ReturnNullVariant()
+{
+	return nullptr;
+}
+
+DEFINE_FUNCTION(UBlueprintTypeConversions::MAKE_CONVERSION_EXEC_FUNCTION_NAME(FBox2d, FBox2f))
+{
+	using namespace UE::Kismet::BlueprintTypeConversions::Internal;
+
+	void* DestAddr = Stack.MostRecentPropertyAddress;
+	Stack.MostRecentProperty = nullptr;
+	Stack.StepCompiledIn(nullptr, nullptr);
+	const void* SourceAddr = Stack.MostRecentPropertyAddress;
+	P_FINISH;
+	ConvertType<FBox2d, FBox2f>(SourceAddr, DestAddr);
+}
+
+DEFINE_FUNCTION(UBlueprintTypeConversions::MAKE_CONVERSION_EXEC_FUNCTION_NAME(FBox2f, FBox2d))
+{
+	using namespace UE::Kismet::BlueprintTypeConversions::Internal;
+
+	void* DestAddr = Stack.MostRecentPropertyAddress;
+	Stack.MostRecentProperty = nullptr;
+	Stack.StepCompiledIn(nullptr, nullptr);
+	const void* SourceAddr = Stack.MostRecentPropertyAddress;
+	P_FINISH;
+	ConvertType<FBox2f, FBox2d>(SourceAddr, DestAddr);
+}
+
+UE::Kismet::BlueprintTypeConversions::Internal::FStructConversionEntry FBox2DEntry(
+	&TBaseStructure<FBox2D>::Get,
+	&ReturnNullVariant,
+	&ReturnNullVariant,
+	&TVariantStructure<FBox2f>::Get,
+	TEXT(PREPROCESSOR_TO_STRING(MAKE_CONVERSION_FUNCTION_NAME(FBox2d, FBox2f))),
+	TEXT(PREPROCESSOR_TO_STRING(MAKE_CONVERSION_FUNCTION_NAME(FBox2f, FBox2d))),
+	&UE::Kismet::BlueprintTypeConversions::Internal::ConvertType<FBox2d, FBox2f>,
+	&UE::Kismet::BlueprintTypeConversions::Internal::ConvertType<FBox2f, FBox2d>
+);
+
+DEFINE_FUNCTION(UBlueprintTypeConversions::MAKE_CONVERSION_EXEC_FUNCTION_NAME(FVector2d, FVector2f))
+{
+	using namespace UE::Kismet::BlueprintTypeConversions::Internal;
+
+	void* DestAddr = Stack.MostRecentPropertyAddress;
+	Stack.MostRecentProperty = nullptr;
+	Stack.StepCompiledIn(nullptr, nullptr);
+	const void* SourceAddr = Stack.MostRecentPropertyAddress;
+	P_FINISH;
+	ConvertType<FVector2d, FVector2f>(SourceAddr, DestAddr);
+}
+
+DEFINE_FUNCTION(UBlueprintTypeConversions::MAKE_CONVERSION_EXEC_FUNCTION_NAME(FVector2f, FVector2d))
+{
+	using namespace UE::Kismet::BlueprintTypeConversions::Internal;
+
+	void* DestAddr = Stack.MostRecentPropertyAddress;
+	Stack.MostRecentProperty = nullptr;
+	Stack.StepCompiledIn(nullptr, nullptr);
+	const void* SourceAddr = Stack.MostRecentPropertyAddress;
+	P_FINISH;
+	ConvertType<FVector2f, FVector2d>(SourceAddr, DestAddr);
+}
+
+UE::Kismet::BlueprintTypeConversions::Internal::FStructConversionEntry FVector2DEntry(
+	&TBaseStructure<FVector2D>::Get,
+	&ReturnNullVariant,
+	&ReturnNullVariant,
+	&TVariantStructure<FVector2f>::Get,
+	TEXT(PREPROCESSOR_TO_STRING(MAKE_CONVERSION_FUNCTION_NAME(FVector2d, FVector2f))),
+	TEXT(PREPROCESSOR_TO_STRING(MAKE_CONVERSION_FUNCTION_NAME(FVector2f, FVector2d))),
+	&UE::Kismet::BlueprintTypeConversions::Internal::ConvertType<FVector2d, FVector2f>,
+	&UE::Kismet::BlueprintTypeConversions::Internal::ConvertType<FVector2f, FVector2d>
+);
