Fix non-unity build compile error "SkinnedAsset.h error: 'ESkeletalMeshVertexFlags': is not a class or namespace name"

diff --git a/Engine/Source/Runtime/Engine/Classes/Engine/SkinnedAsset.h b/Engine/Source/Runtime/Engine/Classes/Engine/SkinnedAsset.h
--- a/Engine/Source/Runtime/Engine/Classes/Engine/SkinnedAsset.h
+++ b/Engine/Source/Runtime/Engine/Classes/Engine/SkinnedAsset.h
@@ -12,6 +12,7 @@
 #include "Interfaces/Interface_AsyncCompilation.h"
 #include "ReferenceSkeleton.h"
 #include "PerPlatformProperties.h"
+#include "SkeletalMeshTypes.h"
 #include "SkinnedAsset.generated.h"
 
 class USkeleton;
