Remove check that was too strict during CL Validation.
DataLayerSubsystem will not be created when submitting actors from a given level when the level is not loaded.
When the DLSubsystem is not created, validation will occur using all the WorldDataLayerActorDescs present in the container of the actors.

diff --git a/Engine/Source/Runtime/Engine/Private/WorldPartition/DataLayer/DataLayerUtils.cpp b/Engine/Source/Runtime/Engine/Private/WorldPartition/DataLayer/DataLayerUtils.cpp
--- a/Engine/Source/Runtime/Engine/Private/WorldPartition/DataLayer/DataLayerUtils.cpp
+++ b/Engine/Source/Runtime/Engine/Private/WorldPartition/DataLayer/DataLayerUtils.cpp
@@ -26,8 +26,8 @@ TArray<FName> FDataLayerUtils::ResolvedDataLayerInstanceNames(const FWorldPartit
 	{
 		World = InActorDesc->GetContainer() ? InActorDesc->GetContainer()->GetWorld() : nullptr;
 	}
+
 	const UDataLayerSubsystem* DataLayerSubsystem = UWorld::GetSubsystem<UDataLayerSubsystem>(World);
-	check(DataLayerSubsystem != nullptr);
 
 	// DataLayers not using DataLayer Assets represent DataLayerInstanceNames
 	if (!InActorDesc->IsUsingDataLayerAsset())
