Allow blueprints to avoid loading the FindInBlueprint (FiB) asset tags during cook commandlets unless we're cooking for the CookedEditor target platform.  This can save a substantial amount of memory at cook time for data that is usually not needed at cook time.

diff --git a/Engine/Source/Runtime/CoreUObject/Private/UObject/Obj.cpp b/Engine/Source/Runtime/CoreUObject/Private/UObject/Obj.cpp
--- a/Engine/Source/Runtime/CoreUObject/Private/UObject/Obj.cpp
+++ b/Engine/Source/Runtime/CoreUObject/Private/UObject/Obj.cpp
@@ -2016,6 +2016,15 @@ void UObject::GetAssetRegistryTags(TArray<FAssetRegistryTag>& OutTags) const
 #endif // WITH_EDITOR
 }
 
+#if WITH_EDITOR
+void UObject::GetExtendedAssetRegistryTagsForSave(const ITargetPlatform* TargetPlatform, TArray<FAssetRegistryTag>& OutTags) const
+{
+	PRAGMA_DISABLE_DEPRECATION_WARNINGS
+	GetExternalActorExtendedAssetRegistryTags(OutTags);
+	PRAGMA_ENABLE_DEPRECATION_WARNINGS
+}
+#endif // WITH_EDITOR
+
 static FAssetDataTagMapSharedView MakeSharedTagMap(TArray<UObject::FAssetRegistryTag>&& Tags)
 {
 	FAssetDataTagMap Out;
diff --git a/Engine/Source/Runtime/CoreUObject/Private/UObject/SavePackage/SavePackageUtilities.cpp b/Engine/Source/Runtime/CoreUObject/Private/UObject/SavePackage/SavePackageUtilities.cpp
--- a/Engine/Source/Runtime/CoreUObject/Private/UObject/SavePackage/SavePackageUtilities.cpp
+++ b/Engine/Source/Runtime/CoreUObject/Private/UObject/SavePackage/SavePackageUtilities.cpp
@@ -2220,7 +2220,7 @@ namespace UE
 namespace AssetRegistry
 {
 	// See the corresponding ReadPackageDataMain and ReadPackageDataDependencies defined in PackageReader.cpp in AssetRegistry module
-	void WritePackageData(FStructuredArchiveRecord& ParentRecord, bool bIsCooking, const UPackage* Package, FLinkerSave* Linker, const TSet<UObject*>& ImportsUsedInGame, const TSet<FName>& SoftPackagesUsedInGame)
+	void WritePackageData(FStructuredArchiveRecord& ParentRecord, bool bIsCooking, const UPackage* Package, FLinkerSave* Linker, const TSet<UObject*>& ImportsUsedInGame, const TSet<FName>& SoftPackagesUsedInGame, const ITargetPlatform* TargetPlatform)
 	{
 		// To avoid large patch sizes, we have frozen cooked package format at the format before VER_UE4_ASSETREGISTRY_DEPENDENCYFLAGS
 		bool bPreDependencyFormat = bIsCooking;
@@ -2270,7 +2270,9 @@ namespace AssetRegistry
 
 			TArray<UObject::FAssetRegistryTag> SourceTags;
 			Object->GetAssetRegistryTags(SourceTags);
-			Object->GetExternalActorExtendedAssetRegistryTags(SourceTags);
+#if WITH_EDITOR
+			Object->GetExtendedAssetRegistryTagsForSave(TargetPlatform, SourceTags);
+#endif // WITH_EDITOR
 
 			TArray<UObject::FAssetRegistryTag> Tags;
 			for (UObject::FAssetRegistryTag& SourceTag : SourceTags)
diff --git a/Engine/Source/Runtime/CoreUObject/Private/UObject/SavePackage2.cpp b/Engine/Source/Runtime/CoreUObject/Private/UObject/SavePackage2.cpp
--- a/Engine/Source/Runtime/CoreUObject/Private/UObject/SavePackage2.cpp
+++ b/Engine/Source/Runtime/CoreUObject/Private/UObject/SavePackage2.cpp
@@ -1694,7 +1694,7 @@ ESavePackageResult WritePackageHeader(FStructuredArchive::FRecord& StructuredArc
 	{
 		// Save asset registry data so the editor can search for information about assets in this package
 		SCOPED_SAVETIMER(UPackage_Save_SaveAssetRegistryData);
-		UE::AssetRegistry::WritePackageData(StructuredArchiveRoot, SaveContext.IsCooking(), SaveContext.GetPackage(), Linker, SaveContext.GetImportsUsedInGame(), SaveContext.GetSoftPackagesUsedInGame());
+		UE::AssetRegistry::WritePackageData(StructuredArchiveRoot, SaveContext.IsCooking(), SaveContext.GetPackage(), Linker, SaveContext.GetImportsUsedInGame(), SaveContext.GetSoftPackagesUsedInGame(), SaveContext.GetTargetPlatform());
 	}
 	// Save level information used by World browser
 	{
diff --git a/Engine/Source/Runtime/CoreUObject/Public/Misc/AssetRegistryInterface.h b/Engine/Source/Runtime/CoreUObject/Public/Misc/AssetRegistryInterface.h
--- a/Engine/Source/Runtime/CoreUObject/Public/Misc/AssetRegistryInterface.h
+++ b/Engine/Source/Runtime/CoreUObject/Public/Misc/AssetRegistryInterface.h
@@ -170,7 +170,7 @@ namespace UE::AssetRegistry
 	};
 
 	// Functions to read and write the data used by the AssetRegistry in each package; the format of this data is separate from the format of the data in the asset registry
-	COREUOBJECT_API void WritePackageData(FStructuredArchiveRecord& ParentRecord, bool bIsCooking, const UPackage* Package, FLinkerSave* Linker, const TSet<UObject*>& ImportsUsedInGame, const TSet<FName>& SoftPackagesUsedInGame);
+	COREUOBJECT_API void WritePackageData(FStructuredArchiveRecord& ParentRecord, bool bIsCooking, const UPackage* Package, FLinkerSave* Linker, const TSet<UObject*>& ImportsUsedInGame, const TSet<FName>& SoftPackagesUsedInGame, const ITargetPlatform* TargetPlatform);
 	// ReadPackageDataMain and ReadPackageDataDependencies are declared in IAssetRegistry.h, in the AssetRegistry module, because they depend upon some structures defined in the AssetRegistry module
 
 	namespace Private
diff --git a/Engine/Source/Runtime/CoreUObject/Public/UObject/Object.h b/Engine/Source/Runtime/CoreUObject/Public/UObject/Object.h
--- a/Engine/Source/Runtime/CoreUObject/Public/UObject/Object.h
+++ b/Engine/Source/Runtime/CoreUObject/Public/UObject/Object.h
@@ -830,12 +830,18 @@ public:
 	 */
 	virtual void GetAssetRegistryTags(TArray<FAssetRegistryTag>& OutTags) const;
 
+	UE_DEPRECATED(5.1, "Use the new GetExtendedAssetRegistryTagsForSave that takes a TargetPlatform")
+	virtual void GetExternalActorExtendedAssetRegistryTags(TArray<FAssetRegistryTag>& OutTags) const {}
+
+#if WITH_EDITOR
 	/**
-	 * Temporary interim solution to gather external actors asset registry data.
+	 * Temporary interim solution to gather asset registry data at save time only.  Can depend on the target platform being saved for.
 	 *
-	 * @param	OutTags		A list of key-value pairs associated with this object and their types
+	 * @param	TargetPlatform	The platform that this object is being saved for.
+	 * @param	OutTags			A list of key-value pairs associated with this object and their types
 	 */
-	virtual void GetExternalActorExtendedAssetRegistryTags(TArray<FAssetRegistryTag>& OutTags) const {}
+	virtual void GetExtendedAssetRegistryTagsForSave(const ITargetPlatform* TargetPlatform, TArray<FAssetRegistryTag>& OutTags) const;
+#endif // WITH_EDITOR
 
 	/** Gathers a list of asset registry tags for an FAssetData  */
 	void GetAssetRegistryTags(FAssetData& Out) const;
diff --git a/Engine/Source/Runtime/Engine/Classes/Engine/Blueprint.h b/Engine/Source/Runtime/Engine/Classes/Engine/Blueprint.h
--- a/Engine/Source/Runtime/Engine/Classes/Engine/Blueprint.h
+++ b/Engine/Source/Runtime/Engine/Classes/Engine/Blueprint.h
@@ -960,6 +960,9 @@ public:
 	virtual void PostLoad() override;
 	virtual bool Modify(bool bAlwaysMarkDirty = true) override;
 	virtual void GetAssetRegistryTags(TArray<FAssetRegistryTag>& OutTags) const override;
+#if WITH_EDITOR
+	virtual void GetExtendedAssetRegistryTagsForSave(const ITargetPlatform* TargetPlatform, TArray<FAssetRegistryTag>& OutTags) const override;
+#endif
 	virtual void PostLoadAssetRegistryTags(const FAssetData& InAssetData, TArray<FAssetRegistryTag>& OutTagsAndValuesToUpdate) const;
 	static void PostLoadBlueprintAssetRegistryTags(const FAssetData& InAssetData, TArray<FAssetRegistryTag>& OutTagsAndValuesToUpdate);
 	virtual FPrimaryAssetId GetPrimaryAssetId() const override;
diff --git a/Engine/Source/Runtime/Engine/Classes/GameFramework/Actor.h b/Engine/Source/Runtime/Engine/Classes/GameFramework/Actor.h
--- a/Engine/Source/Runtime/Engine/Classes/GameFramework/Actor.h
+++ b/Engine/Source/Runtime/Engine/Classes/GameFramework/Actor.h
@@ -2030,7 +2030,7 @@ public:
 #if WITH_EDITOR
 	virtual bool Modify(bool bAlwaysMarkDirty = true) override;
 	virtual void GetAssetRegistryTags(TArray<FAssetRegistryTag>& OutTags) const override;
-	virtual void GetExternalActorExtendedAssetRegistryTags(TArray<FAssetRegistryTag>& OutTags) const override;
+	virtual void GetExtendedAssetRegistryTagsForSave(const ITargetPlatform* TargetPlatform, TArray<FAssetRegistryTag>& OutTags) const override;
 	virtual bool NeedsLoadForTargetPlatform(const ITargetPlatform* TargetPlatform) const;
 	virtual void PreEditChange(FProperty* PropertyThatWillChange) override;
 	virtual bool CanEditChange(const FProperty* InProperty) const;
diff --git a/Engine/Source/Runtime/Engine/Private/Actor.cpp b/Engine/Source/Runtime/Engine/Private/Actor.cpp
--- a/Engine/Source/Runtime/Engine/Private/Actor.cpp
+++ b/Engine/Source/Runtime/Engine/Private/Actor.cpp
@@ -483,9 +483,9 @@ void AActor::GetAssetRegistryTags(TArray<FAssetRegistryTag>& OutTags) const
 	}
 }
 
-void AActor::GetExternalActorExtendedAssetRegistryTags(TArray<FAssetRegistryTag>& OutTags) const
+void AActor::GetExtendedAssetRegistryTagsForSave(const ITargetPlatform* TargetPlatform, TArray<FAssetRegistryTag>& OutTags) const
 {
-	Super::GetExternalActorExtendedAssetRegistryTags(OutTags);
+	Super::GetExtendedAssetRegistryTagsForSave(TargetPlatform, OutTags);
 
 	if (IsPackageExternal() && !IsChildActor())
 	{
diff --git a/Engine/Source/Runtime/Engine/Private/Blueprint.cpp b/Engine/Source/Runtime/Engine/Private/Blueprint.cpp
--- a/Engine/Source/Runtime/Engine/Private/Blueprint.cpp
+++ b/Engine/Source/Runtime/Engine/Private/Blueprint.cpp
@@ -964,7 +964,7 @@ void UBlueprint::GetAssetRegistryTags(TArray<FAssetRegistryTag>& OutTags) const
 			FAssetRegistryTag::TT_Alphabetical ) );
 
 	// Only add the FiB tags in the editor, this now gets run for standalone uncooked games
-	if ( ParentClass && GIsEditor && !GetOutermost()->HasAnyPackageFlags(PKG_ForDiffing))
+	if ( ParentClass && GIsEditor && !GetOutermost()->HasAnyPackageFlags(PKG_ForDiffing) && !IsRunningCookCommandlet())
 	{
 		FString Value;
 		const bool bRebuildSearchData = false;
@@ -1012,6 +1012,29 @@ void UBlueprint::GetAssetRegistryTags(TArray<FAssetRegistryTag>& OutTags) const
 	}
 }
 
+#if WITH_EDITOR
+void UBlueprint::GetExtendedAssetRegistryTagsForSave(const ITargetPlatform* TargetPlatform, TArray<FAssetRegistryTag>& OutTags) const
+{
+	Super::GetExtendedAssetRegistryTagsForSave(TargetPlatform, OutTags);
+
+	if ( ParentClass && GIsEditor && !GetOutermost()->HasAnyPackageFlags(PKG_ForDiffing) && IsRunningCookCommandlet())
+	{
+		if (!TargetPlatform || TargetPlatform->HasEditorOnlyData())
+		{
+			FString Value;
+			const bool bRebuildSearchData = false;
+			FSearchData SearchData = FFindInBlueprintSearchManager::Get().QuerySingleBlueprint((UBlueprint*)this, bRebuildSearchData);
+			if (SearchData.IsValid())
+			{
+				Value = SearchData.Value;
+			}
+			
+			OutTags.Add( FAssetRegistryTag(FBlueprintTags::FindInBlueprintsData, Value, FAssetRegistryTag::TT_Hidden) );
+		}
+	}
+}
+#endif
+
 void UBlueprint::PostLoadAssetRegistryTags(const FAssetData& InAssetData, TArray<FAssetRegistryTag>& OutTagsAndValuesToUpdate) const
 {
 	Super::PostLoadAssetRegistryTags(InAssetData, OutTagsAndValuesToUpdate);
