Optimize code to avoid converting a White vertex color (ie no vertex color). Was making a significant difference in test on the Twinmotion stream.

diff --git a/Engine/Source/Runtime/Engine/Private/StaticMesh.cpp b/Engine/Source/Runtime/Engine/Private/StaticMesh.cpp
--- a/Engine/Source/Runtime/Engine/Private/StaticMesh.cpp
+++ b/Engine/Source/Runtime/Engine/Private/StaticMesh.cpp
@@ -6256,9 +6256,12 @@ void UStaticMesh::BuildFromMeshDescription(const FMeshDescription& MeshDescripti
 			if (Color != FLinearColor::White)
 			{
 				bHasVertexColors = true;
+				StaticMeshVertex.Color = Color.ToFColor(true);
+			}
+			else
+			{
+				StaticMeshVertex.Color = FColor::White;
 			}
-
-			StaticMeshVertex.Color = Color.ToFColor(true);
 		}
 	}
 
