WorldPartition - Level Instance - Retrieve proper Actor GUID if actor is not replicated

diff --git a/Engine/Source/Runtime/Engine/Private/LevelInstance/LevelInstanceActorGuid.cpp b/Engine/Source/Runtime/Engine/Private/LevelInstance/LevelInstanceActorGuid.cpp
--- a/Engine/Source/Runtime/Engine/Private/LevelInstance/LevelInstanceActorGuid.cpp
+++ b/Engine/Source/Runtime/Engine/Private/LevelInstance/LevelInstanceActorGuid.cpp
@@ -31,7 +31,7 @@ const FGuid& FLevelInstanceActorGuid::GetGuid_Internal() const
 	check(Actor);
 #if WITH_EDITOR
 	// It's possible for an actor not to have its world when duplicating unsaved actor when starting PIE (actor is outered to its UActorContainer)
-	const FGuid& Guid = (!Actor->GetWorld() || !Actor->GetWorld()->IsGameWorld() || Actor->GetLocalRole() == ENetRole::ROLE_Authority) ? Actor->GetActorGuid() : ActorGuid;
+	const FGuid& Guid = (!Actor->GetWorld() || !Actor->GetWorld()->IsGameWorld() || !Actor->GetIsReplicated() || Actor->HasAuthority()) ? Actor->GetActorGuid() : ActorGuid;
 #else
 	const FGuid& Guid = ActorGuid;
 #endif
