Optionally allow BPGC property GUIDs to be included in a cooked build

This can be useful to avoid needing to add CoreRedirects for BPGC that are being persisted via tagged properties in SaveGame data, as that data may get stale if properties on the BPGC are renamed

diff --git a/Engine/Source/Developer/DeveloperToolSettings/Classes/CookerSettings.h b/Engine/Source/Developer/DeveloperToolSettings/Classes/CookerSettings.h
--- a/Engine/Source/Developer/DeveloperToolSettings/Classes/CookerSettings.h
+++ b/Engine/Source/Developer/DeveloperToolSettings/Classes/CookerSettings.h
@@ -57,6 +57,19 @@ enum class EBlueprintComponentDataCookingMethod
 	EnabledBlueprintsOnly,
 };
 
+UENUM()
+enum class EBlueprintPropertyGuidsCookingMethod
+{
+	/** Do not include the property GUIDs in a cooked build. No additional memory will be used. */
+	Disabled,
+
+	/** Include the property GUIDs in a cooked build for all Blueprint types. This option will require the most additional memory. */
+	AllBlueprints,
+
+	/** Include the property GUIDs in a cooked build only for Blueprint types that have explicitly enabled this feature in their class settings. */
+	EnabledBlueprintsOnly,
+};
+
 UENUM()
 enum class ETextureFormatASTCCompressor
 {
@@ -132,6 +145,13 @@ public:
 	UPROPERTY(GlobalConfig, EditAnywhere, Category = Cooker, AdvancedDisplay, meta = (DisplayName = "Generate optimized Blueprint component data"))
 	EBlueprintComponentDataCookingMethod BlueprintComponentDataCookingMethod;
 
+	/**
+	 * Should we include the property GUIDs for a Blueprint class in a cooked build, so that SaveGame archives can redirect property names via the GUIDs?
+	 * @note This option can increase the overall Blueprint memory usage in a cooked build, but can avoid needing to add CoreRedirect data for Blueprint classes stored within SaveGame archives.
+	 */
+	UPROPERTY(GlobalConfig, EditAnywhere, Category = Cooker, AdvancedDisplay, meta = (DisplayName = "Cook Blueprint property GUIDs?"))
+	EBlueprintPropertyGuidsCookingMethod BlueprintPropertyGuidsCookingMethod = EBlueprintPropertyGuidsCookingMethod::EnabledBlueprintsOnly;
+
 	/** List of class names to exclude when cooking for dedicated server */
 	UPROPERTY(GlobalConfig, EditAnywhere, Category = Cooker, AdvancedDisplay, meta = (DisplayName = "Classes excluded when cooking for dedicated server"))
 	TArray<FString> ClassesExcludedOnDedicatedServer;
diff --git a/Engine/Source/Runtime/CoreUObject/Private/UObject/Class.cpp b/Engine/Source/Runtime/CoreUObject/Private/UObject/Class.cpp
--- a/Engine/Source/Runtime/CoreUObject/Private/UObject/Class.cpp
+++ b/Engine/Source/Runtime/CoreUObject/Private/UObject/Class.cpp
@@ -1290,7 +1290,7 @@ void UStruct::SerializeVersionedTaggedProperties(FStructuredArchive::FSlot Slot,
 	//SCOPED_LOADTIMER(SerializeTaggedPropertiesTime);
 
 	// Determine if this struct supports optional property guid's (UBlueprintGeneratedClasses Only)
-	const bool bArePropertyGuidsAvailable = (UnderlyingArchive.UEVer() >= VER_UE4_PROPERTY_GUID_IN_PROPERTY_TAG) && !FPlatformProperties::RequiresCookedData() && ArePropertyGuidsAvailable();
+	const bool bArePropertyGuidsAvailable = (UnderlyingArchive.UEVer() >= VER_UE4_PROPERTY_GUID_IN_PROPERTY_TAG) && (!FPlatformProperties::RequiresCookedData() || UnderlyingArchive.IsSaveGame()) && ArePropertyGuidsAvailable();
 	const bool bUseRedirects = (!FPlatformProperties::RequiresCookedData() || UnderlyingArchive.IsSaveGame()) && !UnderlyingArchive.IsLoadingFromCookedPackage();
 
 	if (UnderlyingArchive.IsLoading())
@@ -1339,7 +1339,7 @@ void UStruct::SerializeVersionedTaggedProperties(FStructuredArchive::FSlot Slot,
 				}
 				bAdvanceProperty = false;
 
-				// Optionally resolve properties using Guid Property tags in non cooked builds that support it.
+				// Optionally resolve properties using Guid Property tags if the class supports it.
 				if (bArePropertyGuidsAvailable && Tag.HasPropertyGuid)
 				{
 					// Use property guids from blueprint generated classes to redirect serialised data.
diff --git a/Engine/Source/Runtime/Engine/Classes/Engine/Blueprint.h b/Engine/Source/Runtime/Engine/Classes/Engine/Blueprint.h
--- a/Engine/Source/Runtime/Engine/Classes/Engine/Blueprint.h
+++ b/Engine/Source/Runtime/Engine/Classes/Engine/Blueprint.h
@@ -422,6 +422,17 @@ enum class UE_DEPRECATED(5.0, "Blueprint Nativization has been removed as a supp
 	ExplicitlyEnabled
 };
 
+UENUM()
+enum class EShouldCookBlueprintPropertyGuids
+{
+	/** Don't cook the property GUIDs for this Blueprint */
+	No,
+	/** Cook the property GUIDs for this Blueprint (see UCookerSettings::BlueprintPropertyGuidsCookingMethod) */
+	Yes,
+	/** Inherit whether to cook the property GUIDs for this Blueprint from the parent Blueprint (behaves like 'No' if there is no parent Blueprint) */
+	Inherit,
+};
+
 #if WITH_EDITOR
 /** Control flags for current object/world accessor methods */
 enum class EGetObjectOrWorldBeingDebuggedFlags
@@ -528,6 +539,15 @@ class ENGINE_API UBlueprint : public UBlueprintCore
 	UPROPERTY()
 	mutable uint8 bDuplicatingReadOnly:1;
 
+	/**
+	 * Whether to include the property GUIDs for the generated class in a cooked build.
+	 * @note This option may slightly increase memory usage in a cooked build, but can avoid needing to add CoreRedirect data for Blueprint classes stored within SaveGame archives.
+	 */
+	UPROPERTY(EditAnywhere, Category=ClassOptions, AdvancedDisplay, meta=(DisplayName="Should Cook Property Guids?"))
+	EShouldCookBlueprintPropertyGuids ShouldCookPropertyGuidsValue = EShouldCookBlueprintPropertyGuids::Inherit;
+
+	bool ShouldCookPropertyGuids() const;
+
 public:
 	/** When exclusive nativization is enabled, then this asset will be nativized. All super classes must be also nativized. */
 	UE_DEPRECATED(5.0, "Blueprint Nativization has been removed as a supported feature. This property will eventually be removed.")
diff --git a/Engine/Source/Runtime/Engine/Classes/Engine/BlueprintGeneratedClass.h b/Engine/Source/Runtime/Engine/Classes/Engine/BlueprintGeneratedClass.h
--- a/Engine/Source/Runtime/Engine/Classes/Engine/BlueprintGeneratedClass.h
+++ b/Engine/Source/Runtime/Engine/Classes/Engine/BlueprintGeneratedClass.h
@@ -710,6 +710,10 @@ public:
 	TArray<TObjectPtr<UFunction>> CalledFunctions;
 #endif //WITH_EDITORONLY_DATA
 
+	/** Property guid map (if any), for use only when this BP is cooked */
+	UPROPERTY()
+	TMap<FName, FGuid> CookedPropertyGuids;
+
 	// Mapping of changed properties & data to apply when instancing components in a cooked build (one entry per named AddComponent node template for fast lookup at runtime).
 	// Note: This is not currently utilized by the editor; it is a runtime optimization for cooked builds only. It assumes that the component class structure does not change.
 	UPROPERTY()
@@ -749,9 +753,10 @@ public:
 	virtual bool CanBeClusterRoot() const override;
 #if WITH_EDITOR
 	virtual UClass* RegenerateClass(UClass* ClassToRegenerate, UObject* PreviousCDO) override;
+	virtual void BeginCacheForCookedPlatformData(const ITargetPlatform* TargetPlatform) override;
+	virtual void ClearAllCachedCookedPlatformData() override;
 #endif	//WITH_EDITOR
 	virtual bool IsAsset() const override;
-
 	// End UObject interface
 
 	// UClass interface
diff --git a/Engine/Source/Runtime/Engine/Private/Blueprint.cpp b/Engine/Source/Runtime/Engine/Private/Blueprint.cpp
--- a/Engine/Source/Runtime/Engine/Private/Blueprint.cpp
+++ b/Engine/Source/Runtime/Engine/Private/Blueprint.cpp
@@ -1445,6 +1445,27 @@ void UBlueprint::BeginDestroy()
 #endif // WITH_EDITOR
 
 #if WITH_EDITORONLY_DATA
+bool UBlueprint::ShouldCookPropertyGuids() const
+{
+	switch (ShouldCookPropertyGuidsValue)
+	{
+	case EShouldCookBlueprintPropertyGuids::No:
+		return false;
+
+	case EShouldCookBlueprintPropertyGuids::Yes:
+		return true;
+
+	case EShouldCookBlueprintPropertyGuids::Inherit:
+		if (const UBlueprint* ParentBlueprint = UBlueprint::GetBlueprintFromClass(ParentClass))
+		{
+			return ParentBlueprint->ShouldCookPropertyGuids();
+		}
+		break;
+	}
+
+	return false;
+}
+
 UBlueprint* UBlueprint::GetBlueprintFromClass(const UClass* InClass)
 {
 	UBlueprint* BP = NULL;
diff --git a/Engine/Source/Runtime/Engine/Private/BlueprintGeneratedClass.cpp b/Engine/Source/Runtime/Engine/Private/BlueprintGeneratedClass.cpp
--- a/Engine/Source/Runtime/Engine/Private/BlueprintGeneratedClass.cpp
+++ b/Engine/Source/Runtime/Engine/Private/BlueprintGeneratedClass.cpp
@@ -28,6 +28,7 @@
 #include "GenericPlatform/GenericPlatformCrashContext.h"
 
 #if WITH_EDITOR
+#include "CookerSettings.h"
 #include "Kismet2/BlueprintEditorUtils.h"
 #include "Kismet2/KismetEditorUtilities.h"
 #include "BlueprintCompilationManager.h"
@@ -1702,6 +1703,49 @@ UClass* UBlueprintGeneratedClass::RegenerateClass(UClass* ClassToRegenerate, UOb
 
 	return this;
 }
+
+void UBlueprintGeneratedClass::BeginCacheForCookedPlatformData(const ITargetPlatform* TargetPlatform)
+{
+	Super::BeginCacheForCookedPlatformData(TargetPlatform);
+
+	auto ShouldCookBlueprintPropertyGuids = [this]() -> bool
+	{
+		switch (GetDefault<UCookerSettings>()->BlueprintPropertyGuidsCookingMethod)
+		{
+		case EBlueprintPropertyGuidsCookingMethod::EnabledBlueprintsOnly:
+			if (const UBlueprint* BP = Cast<UBlueprint>(ClassGeneratedBy))
+			{
+				return BP->ShouldCookPropertyGuids();
+			}
+			break;
+
+		case EBlueprintPropertyGuidsCookingMethod::AllBlueprints:
+			return true;
+
+		case EBlueprintPropertyGuidsCookingMethod::Disabled:
+		default:
+			break;
+		}
+
+		return false;
+	};
+
+	if (ShouldCookBlueprintPropertyGuids())
+	{
+		CookedPropertyGuids = PropertyGuids;
+	}
+	else
+	{
+		CookedPropertyGuids.Reset();
+	}
+}
+
+void UBlueprintGeneratedClass::ClearAllCachedCookedPlatformData()
+{
+	Super::ClearAllCachedCookedPlatformData();
+
+	CookedPropertyGuids.Reset();
+}
 #endif
 
 bool UBlueprintGeneratedClass::IsAsset() const
@@ -2078,58 +2122,97 @@ void FBlueprintCookedComponentInstancingData::BuildCachedPropertyDataFromTemplat
 
 bool UBlueprintGeneratedClass::ArePropertyGuidsAvailable() const
 {
-#if WITH_EDITORONLY_DATA
-	// Property guid's are generated during compilation.
-	if (PropertyGuids.Num() > 0)
+	auto ArePropertyGuidsAvailableImpl = [this](const TMap<FName, FGuid>& PropertyGuidsToUse) -> bool
 	{
-		return true;
-	}
-	else if (UBlueprintGeneratedClass* Super = Cast<UBlueprintGeneratedClass>(GetSuperStruct()))
+		// Property guid's are generated during compilation.
+		if (PropertyGuidsToUse.Num() > 0)
+		{
+			return true;
+		}
+		
+		if (UBlueprintGeneratedClass* Super = Cast<UBlueprintGeneratedClass>(GetSuperStruct()))
+		{
+			// Our parent may have guids for inherited variables
+			return Super->ArePropertyGuidsAvailable();
+		}
+		
+		return false;
+	};
+
+	if (bCooked)
 	{
-		// Our parent may have guids for inherited variables
-		return Super->ArePropertyGuidsAvailable();
+		return ArePropertyGuidsAvailableImpl(CookedPropertyGuids);
 	}
-#endif // WITH_EDITORONLY_DATA
+
+#if WITH_EDITORONLY_DATA
+	return ArePropertyGuidsAvailableImpl(PropertyGuids);
+#else	// WITH_EDITORONLY_DATA
 	return false;
+#endif	// WITH_EDITORONLY_DATA
 }
 
 FName UBlueprintGeneratedClass::FindPropertyNameFromGuid(const FGuid& PropertyGuid) const
 {
-	FName RedirectedName = NAME_None;
-#if WITH_EDITORONLY_DATA
-	// Check parent first as it may have renamed a property since this class was last saved
-	if (UBlueprintGeneratedClass* Super = Cast<UBlueprintGeneratedClass>(GetSuperStruct()))
+	auto FindPropertyNameFromGuidImpl = [this, &PropertyGuid](const TMap<FName, FGuid>& PropertyGuidsToUse) -> FName
 	{
-		RedirectedName = Super->FindPropertyNameFromGuid(PropertyGuid);
-		if (!RedirectedName.IsNone())
+		// Check parent first as it may have renamed a property since this class was last saved
+		if (UBlueprintGeneratedClass* Super = Cast<UBlueprintGeneratedClass>(GetSuperStruct()))
 		{
-			return RedirectedName;
+			FName RedirectedName = Super->FindPropertyNameFromGuid(PropertyGuid);
+			if (!RedirectedName.IsNone())
+			{
+				return RedirectedName;
+			}
 		}
-	}
 
-	if (const FName* Result = PropertyGuids.FindKey(PropertyGuid))
+		if (const FName* Result = PropertyGuidsToUse.FindKey(PropertyGuid))
+		{
+			return *Result;
+		}
+
+		return NAME_None;
+	};
+
+	if (bCooked)
 	{
-		RedirectedName = *Result;
+		return FindPropertyNameFromGuidImpl(CookedPropertyGuids);
 	}
-#endif // WITH_EDITORONLY_DATA
-	return RedirectedName;
+
+#if WITH_EDITORONLY_DATA
+	return FindPropertyNameFromGuidImpl(PropertyGuids);
+#else	// WITH_EDITORONLY_DATA
+	return NAME_None;
+#endif	// WITH_EDITORONLY_DATA
 }
 
 FGuid UBlueprintGeneratedClass::FindPropertyGuidFromName(const FName InName) const
 {
-	FGuid PropertyGuid;
-#if WITH_EDITORONLY_DATA
-	if (const FGuid* Result = PropertyGuids.Find(InName))
+	auto FindPropertyGuidFromNameImpl = [this, &InName](const TMap<FName, FGuid>& PropertyGuidsToUse) -> FGuid
 	{
-		PropertyGuid = *Result;
-	}
-	else if (UBlueprintGeneratedClass* Super = Cast<UBlueprintGeneratedClass>(GetSuperStruct()))
+		if (const FGuid* Result = PropertyGuidsToUse.Find(InName))
+		{
+			return *Result;
+		}
+		
+		if (UBlueprintGeneratedClass* Super = Cast<UBlueprintGeneratedClass>(GetSuperStruct()))
+		{
+			// Fall back to parent if this is an inherited variable
+			return Super->FindPropertyGuidFromName(InName);
+		}
+
+		return FGuid();
+	};
+
+	if (bCooked)
 	{
-		// Fall back to parent if this is an inherited variable
-		PropertyGuid = Super->FindPropertyGuidFromName(InName);
+		return FindPropertyGuidFromNameImpl(CookedPropertyGuids);
 	}
-#endif // WITH_EDITORONLY_DATA
-	return PropertyGuid;
+
+#if WITH_EDITORONLY_DATA
+	return FindPropertyGuidFromNameImpl(PropertyGuids);
+#else	// WITH_EDITORONLY_DATA
+	return FGuid();
+#endif	// WITH_EDITORONLY_DATA
 }
 
 #if WITH_EDITORONLY_DATA
