Fix bug, bone rendering hits ensure when skeletal mesh LOD'd

diff --git a/Engine/Source/Runtime/Engine/Private/SkeletalDebugRendering.cpp b/Engine/Source/Runtime/Engine/Private/SkeletalDebugRendering.cpp
--- a/Engine/Source/Runtime/Engine/Private/SkeletalDebugRendering.cpp
+++ b/Engine/Source/Runtime/Engine/Private/SkeletalDebugRendering.cpp
@@ -343,7 +343,7 @@ void DrawBones(
 		// determine color of bone based on selection / affected state
 		const bool bIsSelected = InSelectedBones.Contains(BoneIndex);
 		const bool bIsAffected = AffectedBones[BoneIndex];
-		FLinearColor DefaultBoneColor = BoneColors.IsEmpty() ? DrawConfig.DefaultBoneColor : BoneColors[Index];
+		FLinearColor DefaultBoneColor = BoneColors.IsEmpty() ? DrawConfig.DefaultBoneColor : BoneColors[BoneIndex];
 		FLinearColor BoneColor = bIsAffected ? DrawConfig.AffectedBoneColor : DefaultBoneColor;
 		BoneColor = bIsSelected ? DrawConfig.SelectedBoneColor : BoneColor;
 
@@ -357,7 +357,7 @@ void DrawBones(
 		for (int32 ChildIndex = 0; ChildIndex < RefSkeleton.GetNum(); ++ChildIndex)
 		{
 			const int32 ParentIndex = RefSkeleton.GetParentIndex(ChildIndex);
-			if (ParentIndex == BoneIndex)
+			if (ParentIndex == BoneIndex && RequiredBones.Contains(ChildIndex))
 			{
 				ChildPositions.Add(WorldTransforms[ChildIndex].GetLocation());
 				FLinearColor ChildLineColor = BoneColor;
@@ -369,7 +369,6 @@ void DrawBones(
 			}
 		}
 
-		const FName BoneName = RefSkeleton.GetBoneName(BoneIndex);
 		const FTransform BoneTransform = WorldTransforms[BoneIndex];
 
 		// Always set new hit proxy to prevent unintentionally using last drawn element's proxy
