Added CSV scopes for Niagara / FXSystem / GPUSortManager.

diff --git a/Engine/Plugins/FX/Niagara/Source/Niagara/Private/NiagaraGpuComputeDispatch.cpp b/Engine/Plugins/FX/Niagara/Source/Niagara/Private/NiagaraGpuComputeDispatch.cpp
--- a/Engine/Plugins/FX/Niagara/Source/Niagara/Private/NiagaraGpuComputeDispatch.cpp
+++ b/Engine/Plugins/FX/Niagara/Source/Niagara/Private/NiagaraGpuComputeDispatch.cpp
@@ -1607,6 +1607,7 @@ void FNiagaraGpuComputeDispatch::DispatchStage(FRDGBuilder& GraphBuilder, const
 void FNiagaraGpuComputeDispatch::PreInitViews(FRDGBuilder& GraphBuilder, bool bAllowGPUParticleUpdate)
 {
 	SCOPE_CYCLE_COUNTER(STAT_NiagaraGPUDispatchSetup_RT);
+	RDG_CSV_STAT_EXCLUSIVE_SCOPE(GraphBuilder, Niagara);
 
 	bRequiresReadback = false;
 	GNiagaraViewDataManager.ClearSceneTextureParameters();
@@ -1682,6 +1683,8 @@ void FNiagaraGpuComputeDispatch::PostInitViews(FRDGBuilder& GraphBuilder, TArray
 
 	if (bAllowGPUParticleUpdate && FNiagaraUtilities::AllowGPUParticles(GetShaderPlatform()))
 	{
+		RDG_CSV_STAT_EXCLUSIVE_SCOPE(GraphBuilder, Niagara);
+
 		ExecuteTicks(GraphBuilder, Views, ENiagaraGpuComputeTickStage::PostInitViews);
 
 	#if WITH_MGPU
diff --git a/Engine/Source/Runtime/Engine/Private/GPUSortManager.cpp b/Engine/Source/Runtime/Engine/Private/GPUSortManager.cpp
--- a/Engine/Source/Runtime/Engine/Private/GPUSortManager.cpp
+++ b/Engine/Source/Runtime/Engine/Private/GPUSortManager.cpp
@@ -728,6 +728,8 @@ void FGPUSortManager::UpdateSortBuffersPool()
 
 void FGPUSortManager::OnPreRender(FRDGBuilder& GraphBuilder)
 {
+	RDG_CSV_STAT_EXCLUSIVE_SCOPE(GraphBuilder, GPUSort);
+
 	AddPass(
 		GraphBuilder,
 		RDG_EVENT_NAME("FGPUSortManager::OnPreRender"),
@@ -783,6 +785,7 @@ void FGPUSortManager::OnPreRender(FRDGBuilder& GraphBuilder)
 void FGPUSortManager::OnPostRenderOpaque(FRDGBuilder& GraphBuilder)
 {
 	LLM_SCOPE(ELLMTag::GPUSort);
+	RDG_CSV_STAT_EXCLUSIVE_SCOPE(GraphBuilder, GPUSort);
 
 	AddPass(
 		GraphBuilder,
diff --git a/Engine/Source/Runtime/Engine/Private/Particles/FXSystem.cpp b/Engine/Source/Runtime/Engine/Private/Particles/FXSystem.cpp
--- a/Engine/Source/Runtime/Engine/Private/Particles/FXSystem.cpp
+++ b/Engine/Source/Runtime/Engine/Private/Particles/FXSystem.cpp
@@ -482,6 +482,8 @@ void FFXSystem::PreRender(FRDGBuilder& GraphBuilder, TConstArrayView<FViewInfo>
 
 	if (RHISupportsGPUParticles() && bAllowGPUParticleSceneUpdate)
 	{
+		RDG_CSV_STAT_EXCLUSIVE_SCOPE(GraphBuilder, FXSystem);
+
 		FRHIUniformBuffer* ViewUniformBuffer = GetReferenceViewUniformBuffer(Views);
 		const FGlobalDistanceFieldParameterData* GlobalDistanceFieldParameterData = GetReferenceGlobalDistanceFieldData(Views);
 
@@ -526,6 +528,8 @@ void FFXSystem::PostRenderOpaque(FRDGBuilder& GraphBuilder, TConstArrayView<FVie
 
 	if (RHISupportsGPUParticles() && IsParticleCollisionModeSupported(GetShaderPlatform(), PCM_DepthBuffer) && bAllowGPUParticleUpdate)
 	{
+		RDG_CSV_STAT_EXCLUSIVE_SCOPE(GraphBuilder, FXSystem);
+
 		FRHIUniformBuffer* ViewUniformBuffer = GetReferenceViewUniformBuffer(Views);
 
 		AddPass(GraphBuilder, RDG_EVENT_NAME("FFXSystem::PostRenderOpaque"), 
diff --git a/Engine/Source/Runtime/Renderer/Private/SceneRendering.cpp b/Engine/Source/Runtime/Renderer/Private/SceneRendering.cpp
--- a/Engine/Source/Runtime/Renderer/Private/SceneRendering.cpp
+++ b/Engine/Source/Runtime/Renderer/Private/SceneRendering.cpp
@@ -3964,7 +3964,7 @@ void FSceneRenderer::ViewExtensionPreRender_RenderThread(FRHICommandListImmediat
 
 	{
 		FRDGBuilder GraphBuilder(RHICmdList);
-		CSV_SCOPED_TIMING_STAT_EXCLUSIVE(PreRender);
+		RDG_CSV_STAT_EXCLUSIVE_SCOPE(GraphBuilder, PreRender);
 		SCOPE_CYCLE_COUNTER(STAT_FDeferredShadingSceneRenderer_ViewExtensionPreRenderView);
 
 		for (int ViewExt = 0; ViewExt < SceneRenderer->ViewFamily.ViewExtensions.Num(); ViewExt++)
