Printing references to stale objects will now take into account GarbageCollectorKeep flags that can prevent an object from being GC even though it's not referenced

diff --git a/Engine/Source/Runtime/Engine/Private/UnrealEngine.cpp b/Engine/Source/Runtime/Engine/Private/UnrealEngine.cpp
--- a/Engine/Source/Runtime/Engine/Private/UnrealEngine.cpp
+++ b/Engine/Source/Runtime/Engine/Private/UnrealEngine.cpp
@@ -15471,6 +15471,11 @@ void UEngine::VerifyLoadMapWorldCleanup(FWorldContext* ForWorldContext)
 	}
 }
 
+static FORCEINLINE bool HasGarbageCollectionKeepFlags(FGCObjectInfo* ObjectInfo)
+{
+	return ObjectInfo && (ObjectInfo->HasAnyFlags(GARBAGE_COLLECTION_KEEPFLAGS) || ObjectInfo->HasAnyInternalFlags(EInternalObjectFlags::GarbageCollectionKeepFlags));
+}
+
 static bool PrintStaleReferenceChainsAndFindReferencingObjects(UObject* ObjectToFindReferencesTo, FReferenceChainSearch& RefChainSearch, FGCObjectInfo*& OutGarbageObject, FGCObjectInfo*& OutReferencingObject, ELogVerbosity::Type Verbosity)
 {
 #if !NO_LOGGING
@@ -15482,7 +15487,7 @@ static bool PrintStaleReferenceChainsAndFindReferencingObjects(UObject* ObjectTo
 	return RefChainSearch.PrintResults([&ObjectToFindReferencesTo, &OutGarbageObject, &OutReferencingObject](FReferenceChainSearch::FCallbackParams& Params)
 		{
 			check(Params.Object);
-			check(Params.Referencer);
+			check(Params.Referencer || HasGarbageCollectionKeepFlags(Params.Object));
 			if (!Params.Object->IsValid())
 			{
 				// We may find many chains that lead to the leak but for brevity we only report the first one in the fatal error below
@@ -15598,8 +15603,9 @@ TArray<FString> UEngine::FindAndPrintStaleReferencesToObjects(TConstArrayView<UO
 		{
 			PathToCulprit = GetPathToStaleObjectReferencer(ObjectToFindReferencesTo, RefChainSearch);
 			FString GarbageObjectName = OutGarbageObject ? OutGarbageObject->GetFullName() : ObjectToFindReferencesTo->GetFullName();
-			checkf(OutReferencingObject, TEXT("No object referencing %s found even though we have a valid reference chain"), *ObjectToFindReferencesTo->GetPathName());
-			GarbageErrorMessage = FString::Printf(TEXT("Object %s is being referenced by %s"), *GarbageObjectName, OutReferencingObject ? *OutReferencingObject->GetFullName() : TEXT("NULL"));
+			checkf(OutReferencingObject || HasGarbageCollectionKeepFlags(OutGarbageObject), TEXT("No object referencing %s found even though we have a valid reference chain"), *ObjectToFindReferencesTo->GetPathName());
+			GarbageErrorMessage = FString::Printf(TEXT("Object %s is being referenced by %s"), *GarbageObjectName, 
+				OutReferencingObject ? *OutReferencingObject->GetFullName() : (HasGarbageCollectionKeepFlags(OutGarbageObject) ? TEXT("GarbageCollectionKeepFlags") : TEXT("NULL")));
 		}
 #if ENABLE_GC_HISTORY
 		else
