Fixed exports of UShapeComponent templated methods so that the explicit instantiations can be called from external modules.

diff --git a/Engine/Source/Runtime/Engine/Private/Components/ShapeComponent.cpp b/Engine/Source/Runtime/Engine/Private/Components/ShapeComponent.cpp
--- a/Engine/Source/Runtime/Engine/Private/Components/ShapeComponent.cpp
+++ b/Engine/Source/Runtime/Engine/Private/Components/ShapeComponent.cpp
@@ -171,23 +171,23 @@ void UShapeComponent::SetUseSystemDefaultObstacleAreaClass()
 	bUseSystemDefaultObstacleAreaClass = true;
 }
 
-template <> void UShapeComponent::AddShapeToGeomArray<FKBoxElem>() { ShapeBodySetup->AggGeom.BoxElems.Add(FKBoxElem()); }
-template <> void UShapeComponent::AddShapeToGeomArray<FKSphereElem>() { ShapeBodySetup->AggGeom.SphereElems.Add(FKSphereElem()); }
-template <> void UShapeComponent::AddShapeToGeomArray<FKSphylElem>() { ShapeBodySetup->AggGeom.SphylElems.Add(FKSphylElem()); }
+template <> ENGINE_API void UShapeComponent::AddShapeToGeomArray<FKBoxElem>() { ShapeBodySetup->AggGeom.BoxElems.Add(FKBoxElem()); }
+template <> ENGINE_API void UShapeComponent::AddShapeToGeomArray<FKSphereElem>() { ShapeBodySetup->AggGeom.SphereElems.Add(FKSphereElem()); }
+template <> ENGINE_API void UShapeComponent::AddShapeToGeomArray<FKSphylElem>() { ShapeBodySetup->AggGeom.SphylElems.Add(FKSphylElem()); }
 
-template <>
+template <> ENGINE_API
 void UShapeComponent::SetShapeToNewGeom<FKBoxElem>(const FPhysicsShapeHandle& Shape)
 {
 	FPhysicsInterface::SetUserData(Shape, (void*)ShapeBodySetup->AggGeom.BoxElems[0].GetUserData());
 }
 
-template <>
+template <> ENGINE_API
 void UShapeComponent::SetShapeToNewGeom<FKSphereElem>(const FPhysicsShapeHandle& Shape)
 {
 	FPhysicsInterface::SetUserData(Shape, (void*)ShapeBodySetup->AggGeom.SphereElems[0].GetUserData());
 }
 
-template <>
+template <> ENGINE_API
 void UShapeComponent::SetShapeToNewGeom<FKSphylElem>(const FPhysicsShapeHandle& Shape)
 {
 	FPhysicsInterface::SetUserData(Shape, (void*)ShapeBodySetup->AggGeom.SphylElems[0].GetUserData());
@@ -243,6 +243,6 @@ void UShapeComponent::CreateShapeBodySetupIfNeeded()
 }
 
 //Explicit instantiation of the different shape components
-template void UShapeComponent::CreateShapeBodySetupIfNeeded<FKSphylElem>();
-template void UShapeComponent::CreateShapeBodySetupIfNeeded<FKBoxElem>();
-template void UShapeComponent::CreateShapeBodySetupIfNeeded<FKSphereElem>();
+template ENGINE_API void UShapeComponent::CreateShapeBodySetupIfNeeded<FKSphylElem>();
+template ENGINE_API void UShapeComponent::CreateShapeBodySetupIfNeeded<FKBoxElem>();
+template ENGINE_API void UShapeComponent::CreateShapeBodySetupIfNeeded<FKSphereElem>();
