Adding an exec command `material dumpdebuginfo` which will dump debug info for each shader in the material given as an argument.

Can be run like the following: `material dumpdebuginfo WorldGridMaterial`

Can also be used to see which shaders get invalidated when changing shader code:

1. `material dumpdebuginfo WorldGridMaterial`
2. Change shader usf/ush file.
3. `recompileshaders material WorldGridMaterial`
4. `material dumpdebuginfo WorldGridMaterial`
5. Diff the outputs.

- Make DumpDebugInfo ENGINE_API exposed and virtual through the UMaterialInterface.
- Making DumpDebugInfo accept an FOutputDevice so it can either output to a file or the log.
- Making DumpDebugInfo strictly output csv data.  This makes it much easier to do data processing in a spreadsheet or python.  The tradeoff here is skipping the output of extraneous per-material information.
- Adds BuildSettings to Engine so I can access the current build information
- Output WorldGridMaterial Debug information on the build machine and output to the Automation directory so it will be uploaded as a build artifact.

The filename follows the following format:

Engine\Programs\AutomationTool\{ProjectName}++UE5+Main-CL-20942948-WorldGridMaterial-[REDACTED_PHONE].csv

The file output looks like this:

Frequency, Target, VFType, ShaderType, SourceHash, VFSourceHash, OutputHash, IsShaderPipeline
SF_Pixel, PCD3D_SM5, FGPUSpriteVertexFactory, TShadowDepthPSPixelShadowDepth_PerspectiveCorrect, 2DD4725EA72D0498EB6B1E656F35564D97083FC8, 920506EE58508D22D29046DCE7DF15BF34940059, EC976B76E986CB7E5235B154F9A0B3AF3AB9B1AD, false
SF_Vertex, PCD3D_SM5, FGPUSpriteVertexFactory, TVirtualTextureVSBaseColorNormalRoughness, A3C811C874410080F3138AEF4489CCA932D7395A, 920506EE58508D22D29046DCE7DF15BF34940059, CD25F583872C5ECF0E156EA8F50AF96B6B199BDA, false

diff --git a/Engine/Source/Runtime/Engine/Classes/Materials/Material.h b/Engine/Source/Runtime/Engine/Classes/Materials/Material.h
--- a/Engine/Source/Runtime/Engine/Classes/Materials/Material.h
+++ b/Engine/Source/Runtime/Engine/Classes/Materials/Material.h
@@ -1713,7 +1713,7 @@ protected:
 #endif
 
 public:
-	void DumpDebugInfo() const;
+	ENGINE_API virtual void DumpDebugInfo(FOutputDevice& OutputDevice) const override;
 	void SaveShaderStableKeys(const class ITargetPlatform* TP);
 	ENGINE_API virtual void SaveShaderStableKeysInner(const class ITargetPlatform* TP, const struct FStableShaderKeyAndValue& SaveKeyVal) override;
 
diff --git a/Engine/Source/Runtime/Engine/Classes/Materials/MaterialInstance.h b/Engine/Source/Runtime/Engine/Classes/Materials/MaterialInstance.h
--- a/Engine/Source/Runtime/Engine/Classes/Materials/MaterialInstance.h
+++ b/Engine/Source/Runtime/Engine/Classes/Materials/MaterialInstance.h
@@ -879,7 +879,7 @@ public:
 	 */
 	ENGINE_API virtual void GetLightingGuidChain(bool bIncludeTextures, TArray<FGuid>& OutGuids) const override;
 
-	void DumpDebugInfo() const;
+	ENGINE_API virtual void DumpDebugInfo(FOutputDevice& OutputDevice) const override;
 	void SaveShaderStableKeys(const class ITargetPlatform* TP);
 	ENGINE_API virtual void SaveShaderStableKeysInner(const class ITargetPlatform* TP, const struct FStableShaderKeyAndValue& SaveKeyVal) override;
 
diff --git a/Engine/Source/Runtime/Engine/Classes/Materials/MaterialInterface.h b/Engine/Source/Runtime/Engine/Classes/Materials/MaterialInterface.h
--- a/Engine/Source/Runtime/Engine/Classes/Materials/MaterialInterface.h
+++ b/Engine/Source/Runtime/Engine/Classes/Materials/MaterialInterface.h
@@ -596,6 +596,8 @@ public:
 	ENGINE_API virtual void LogMaterialsAndTextures(FOutputDevice& Ar, int32 Indent) const {}
 #endif
 
+	ENGINE_API virtual void DumpDebugInfo(FOutputDevice& OutputDevice) const {}
+
 private:
 	// might get called from game or render thread
 	FMaterialRelevance GetRelevance_Internal(const UMaterial* Material, ERHIFeatureLevel::Type InFeatureLevel) const;
diff --git a/Engine/Source/Runtime/Engine/Engine.Build.cs b/Engine/Source/Runtime/Engine/Engine.Build.cs
--- a/Engine/Source/Runtime/Engine/Engine.Build.cs
+++ b/Engine/Source/Runtime/Engine/Engine.Build.cs
@@ -122,6 +122,7 @@ public class Engine : ModuleRules
 			new string[] {
 				"AnimationCore",
 				"AppFramework",
+				"BuildSettings",
 				"Networking",
 				"Landscape",
 				"UMG",
diff --git a/Engine/Source/Runtime/Engine/Private/Materials/Material.cpp b/Engine/Source/Runtime/Engine/Private/Materials/Material.cpp
--- a/Engine/Source/Runtime/Engine/Private/Materials/Material.cpp
+++ b/Engine/Source/Runtime/Engine/Private/Materials/Material.cpp
@@ -86,6 +86,10 @@
 #include "Materials/MaterialExpressionConstant.h"
 #include "Materials/MaterialExpressionBreakMaterialAttributes.h"
 #include "MaterialCachedData.h"
+#include "Misc/OutputDeviceArchiveWrapper.h"
+#include "Misc/OutputDeviceFile.h"
+#include "HAL/FileManager.h"
+#include "BuildSettings.h"
 
 #if WITH_EDITOR
 #include "MaterialCachedHLSLTree.h"
@@ -3768,23 +3772,25 @@ void UMaterial::PostLoad()
 		UpdateLightmassTextureTracking();
 	}
 
-	//DumpDebugInfo();
+	//DumpDebugInfo(*GLog);
 }
 
-void UMaterial::DumpDebugInfo() const
+void UMaterial::DumpDebugInfo(FOutputDevice& OutputDevice) const
 {
-	UE_LOG(LogConsoleResponse, Display, TEXT("----------------------------- %s"), *GetFullName());
-
+	for (FMaterialResource* Resource : MaterialResources)
 	{
-		static const UEnum* Enum = StaticEnum<EMaterialDomain>();
-		check(Enum);
-		UE_LOG(LogConsoleResponse, Display, TEXT("  MaterialDomain %s"), *Enum->GetNameStringByValue(int64(MaterialDomain)));
+		Resource->DumpDebugInfo(OutputDevice);
 	}
 
-	for (FMaterialResource* Resource : MaterialResources)
+#if WITH_EDITOR
+	for (auto& It : CachedMaterialResourcesForCooking)
 	{
-		Resource->DumpDebugInfo();
+		for (FMaterialResource* Resource : It.Value)
+		{
+			Resource->DumpDebugInfo(OutputDevice);
+		}
 	}
+#endif
 }
 
 void UMaterial::SaveShaderStableKeys(const class ITargetPlatform* TP)
@@ -3950,6 +3956,30 @@ bool UMaterial::IsCachedCookedPlatformDataLoaded( const ITargetPlatform* TargetP
 
 void UMaterial::ClearCachedCookedPlatformData( const ITargetPlatform *TargetPlatform )
 {
+#if WITH_EDITOR
+	if (GIsBuildMachine)
+	{
+		// Dump debug info for the DefaultMaterial.
+		TRACE_CPUPROFILER_EVENT_SCOPE(DumpDebugShaderInfoDefaultMaterial);
+
+		UMaterial* DefaultMaterial = GDefaultMaterials[EMaterialDomain::MD_Surface];
+		if (this == DefaultMaterial)
+		{
+			// Make the file in the automation directory so it will be uploaded as a build artifact.
+			const FString Filename = FString::Printf(TEXT("%s%s-%s-%s.csv"), FApp::GetProjectName(), BuildSettings::GetBuildVersion(), *GetName(), *FDateTime::Now().ToString());
+			const FString FilenameFull = FPaths::Combine(*FPaths::EngineDir(), TEXT("Programs"), TEXT("AutomationTool"), TEXT("Saved"), TEXT("Logs"), Filename);
+
+			if (FArchive* FileArchive = IFileManager::Get().CreateDebugFileWriter(*FilenameFull))
+			{
+				FOutputDeviceArchiveWrapper Wrapper(FileArchive);
+				DumpDebugInfo(Wrapper);
+				FileArchive->Close();
+				delete FileArchive;
+			}
+		}
+	}
+#endif
+
 	TArray<FMaterialResource*>* CachedMaterialResourcesForPlatform = CachedMaterialResourcesForCooking.Find( TargetPlatform );
 	if ( CachedMaterialResourcesForPlatform != nullptr)
 	{
diff --git a/Engine/Source/Runtime/Engine/Private/Materials/MaterialInstance.cpp b/Engine/Source/Runtime/Engine/Private/Materials/MaterialInstance.cpp
--- a/Engine/Source/Runtime/Engine/Private/Materials/MaterialInstance.cpp
+++ b/Engine/Source/Runtime/Engine/Private/Materials/MaterialInstance.cpp
@@ -2873,7 +2873,7 @@ void UMaterialInstance::PostLoad()
 
 		LightingGuidFixupMap.Add(GetLightingGuid(), this);
 	}
-	//DumpDebugInfo();
+	//DumpDebugInfo(*GLog);
 }
 
 void UMaterialInstance::BeginDestroy()
@@ -4228,29 +4228,26 @@ static FAutoConsoleCommand FindRedundantMICSCmd(
 
 #endif
 
-void UMaterialInstance::DumpDebugInfo() const
+void UMaterialInstance::DumpDebugInfo(FOutputDevice& OutputDevice) const
 {
-	UE_LOG(LogConsoleResponse, Display, TEXT("----------------------------- %s"), *GetFullName());
-
-	UE_LOG(LogConsoleResponse, Display, TEXT("  Parent %s"), Parent ? *Parent->GetFullName() : TEXT("null"));
-
 	if (Parent)
 	{
-		const UMaterial* Base = GetMaterial();
-		UE_LOG(LogConsoleResponse, Display, TEXT("  Base %s"), Base ? *Base->GetFullName() : TEXT("null"));
-
-		if (Base)
-		{
-			static const UEnum* Enum = StaticEnum<EMaterialDomain>();
-			check(Enum);
-			UE_LOG(LogConsoleResponse, Display, TEXT("  MaterialDomain %s"), *Enum->GetNameStringByValue(int64(Base->MaterialDomain)));
-		}
 		if (bHasStaticPermutationResource)
 		{
 			for (FMaterialResource* CurrentResource : StaticPermutationMaterialResources)
 			{
-				CurrentResource->DumpDebugInfo();
+				CurrentResource->DumpDebugInfo(OutputDevice);
+			}
+
+#if WITH_EDITOR
+			for (auto& It : CachedMaterialResourcesForCooking)
+			{
+				for (FMaterialResource* CurrentResource : It.Value)
+				{
+					CurrentResource->DumpDebugInfo(OutputDevice);
+				}
 			}
+#endif // WITH_EDITOR
 		}
 		else
 		{
diff --git a/Engine/Source/Runtime/Engine/Private/Materials/MaterialShader.cpp b/Engine/Source/Runtime/Engine/Private/Materials/MaterialShader.cpp
--- a/Engine/Source/Runtime/Engine/Private/Materials/MaterialShader.cpp
+++ b/Engine/Source/Runtime/Engine/Private/Materials/MaterialShader.cpp
@@ -3262,46 +3262,54 @@ void FMaterialShaderMapContent::RemoveMeshShaderMap(const FHashedName& VertexFac
 	}
 }
 
-void FMaterialShaderMap::DumpDebugInfo() const
+void FMaterialShaderMap::DumpDebugInfo(FOutputDevice& OutputDevice) const
 {
-	const FString& FriendlyNameS = GetFriendlyName();
-	UE_LOG(LogConsoleResponse, Display, TEXT("FMaterialShaderMap:  FriendlyName %s"), *FriendlyNameS);
-	const FString& DebugDescriptionS = GetDebugDescription();
-	UE_LOG(LogConsoleResponse, Display, TEXT("  DebugDescription %s"), *DebugDescriptionS);
+	// Turn off as it makes diffing hard
+	TGuardValue<ELogTimes::Type> DisableLogTimes(GPrintLogTimes, ELogTimes::None);
+
+	OutputDevice.Logf(TEXT("Frequency, Target, VFType, ShaderType, SourceHash, VFSourceHash, OutputHash, IsShaderPipeline"));
 
 	{
 		TMap<FShaderId, TShaderRef<FShader>> ShadersL;
 		GetShaderList(ShadersL);
-		UE_LOG(LogConsoleResponse, Display, TEXT("  --- %d shaders"), ShadersL.Num());
-		int32 Index = 0;
 		for (auto& KeyValue : ShadersL)
 		{
-			UE_LOG(LogConsoleResponse, Display, TEXT("    --- shader %d"), Index);
 			FShader* Shader = KeyValue.Value.GetShader();
-			Shader->DumpDebugInfo(GetPointerTable());
-			Index++;
+			const FVertexFactoryType* VertexFactoryType = Shader->GetVertexFactoryType(GetPointerTable());
+			OutputDevice.Logf(TEXT("%s, %s, %s, %s, %s, %s, %s, %s"),
+				GetShaderFrequencyString(Shader->GetFrequency()),
+				*LegacyShaderPlatformToShaderFormat(GetShaderPlatform()).ToString(),
+				VertexFactoryType ? VertexFactoryType->GetName() : TEXT("null"),
+				Shader->GetType(GetPointerTable())->GetName(),
+				*Shader->GetHash().ToString(),
+				*Shader->GetVertexFactoryHash().ToString(),
+				*Shader->GetOutputHash().ToString(),
+				TEXT("false")
+				);
 		}
 	}
 
 	{
 		TArray<FShaderPipelineRef> ShaderPipelinesL;
 		GetShaderPipelineList(ShaderPipelinesL);
-		UE_LOG(LogConsoleResponse, Display, TEXT("  --- %d shaderpipelines"), ShaderPipelinesL.Num());
-		int32 Index = 0;
 		for (FShaderPipelineRef Value : ShaderPipelinesL)
 		{
 			FShaderPipeline* ShaderPipeline = Value.GetPipeline();
-			UE_LOG(LogConsoleResponse, Display, TEXT("    --- shaderpipeline %d"), Index);
 			TArray<TShaderRef<FShader>> ShadersL = ShaderPipeline->GetShaders(*this);
-			UE_LOG(LogConsoleResponse, Display, TEXT("    --- %d shaders"), ShadersL.Num());
-			int32 ShaderIdx = 0;
 			for (TShaderRef<FShader>& Shader : ShadersL)
 			{
-				UE_LOG(LogConsoleResponse, Display, TEXT("      --- shader %d"), ShaderIdx);
-				Shader->DumpDebugInfo(GetPointerTable());
-				++ShaderIdx;
+				const FVertexFactoryType* VertexFactoryType = Shader->GetVertexFactoryType(GetPointerTable());
+				OutputDevice.Logf(TEXT("%s, %s, %s, %s, %s, %s, %s, %s"),
+					GetShaderFrequencyString(Shader->GetFrequency()),
+					*LegacyShaderPlatformToShaderFormat(GetShaderPlatform()).ToString(),
+					VertexFactoryType ? VertexFactoryType->GetName() : TEXT("null"),
+					Shader->GetType(GetPointerTable())->GetName(),
+					*Shader->GetHash().ToString(),
+					*Shader->GetVertexFactoryHash().ToString(),
+					*Shader->GetOutputHash().ToString(),
+					TEXT("true")
+					);
 			}
-			Index++;
 		}
 	}
 }
diff --git a/Engine/Source/Runtime/Engine/Private/Materials/MaterialShared.cpp b/Engine/Source/Runtime/Engine/Private/Materials/MaterialShared.cpp
--- a/Engine/Source/Runtime/Engine/Private/Materials/MaterialShared.cpp
+++ b/Engine/Source/Runtime/Engine/Private/Materials/MaterialShared.cpp
@@ -129,6 +129,37 @@ static FAutoConsoleCommand GFlushMaterialUniforms(
 })
 );
 
+#if WITH_EDITOR
+class FMaterialDumpDebugInfoExecHelper : public FSelfRegisteringExec
+{
+	virtual bool Exec(class UWorld* InWorld, const TCHAR* Cmd, FOutputDevice& Ar)
+	{
+		if (FParse::Command(&Cmd, TEXT("material dumpdebuginfo")))
+		{
+			FString RequestedMaterialName(FParse::Token(Cmd, 0));
+
+			if (RequestedMaterialName.Len() > 0)
+			{
+				for (TObjectIterator<UMaterialInterface> It; It; ++It)
+				{
+					UMaterialInterface* Material = *It;
+					if (Material && Material->GetName() == RequestedMaterialName)
+					{
+						Material->DumpDebugInfo(Ar);
+						break;
+					}
+				}
+				return true;
+			}
+
+			TODO: PRINT WARNING
+		}
+		return false;
+	}
+};
+static FMaterialDumpDebugInfoExecHelper GMaterialDumpDebugInfoExecHelper;
+#endif
+
 // defined in the same module (Material.cpp)
 bool PoolSpecialMaterialsCompileJobs();
 
@@ -4245,18 +4276,11 @@ void FMaterial::RestoreEditorLoadedMaterialShadersFromMemory(const TMap<FMateria
 }
 #endif // WITH_EDITOR
 
-void FMaterial::DumpDebugInfo()
+void FMaterial::DumpDebugInfo(FOutputDevice& OutputDevice)
 {
 	if (GameThreadShaderMap)
 	{
-		FString FeatureLevelName;
-		GetFeatureLevelName(FeatureLevel, FeatureLevelName);
-		FString QualityLevelString;
-		GetMaterialQualityLevelName(QualityLevel, QualityLevelString);
-
-		UE_LOG(LogConsoleResponse, Display, TEXT("FMaterial:  FeatureLevel %s     Quality Level %s"), *FeatureLevelName, *QualityLevelString);
-
-		GameThreadShaderMap->DumpDebugInfo();
+		GameThreadShaderMap->DumpDebugInfo(OutputDevice);
 	}
 }
 
diff --git a/Engine/Source/Runtime/Engine/Public/MaterialShared.h b/Engine/Source/Runtime/Engine/Public/MaterialShared.h
--- a/Engine/Source/Runtime/Engine/Public/MaterialShared.h
+++ b/Engine/Source/Runtime/Engine/Public/MaterialShared.h
@@ -1400,7 +1400,7 @@ public:
 			}
 		}
 	}
-	void DumpDebugInfo() const;
+	void DumpDebugInfo(FOutputDevice& OutputDevice) const;
 
 #if WITH_EDITOR
 	void InitalizeForODSC(EShaderPlatform TargetShaderPlatform, const FMaterialCompilationOutput& NewCompilationOutput);
@@ -2151,7 +2151,7 @@ public:
 	virtual void SetupExtaCompilationSettings(const EShaderPlatform Platform, FExtraShaderCompilerSettings& Settings) const
 	{}
 
-	void DumpDebugInfo();
+	void DumpDebugInfo(FOutputDevice& OutputDevice);
 	void SaveShaderStableKeys(EShaderPlatform TargetShaderPlatform, struct FStableShaderKeyAndValue& SaveKeyVal); // arg is non-const, we modify it as we go
 
 #if WITH_EDITOR
