Build fix.

diff --git a/Engine/Source/Runtime/Engine/Private/Materials/MaterialShared.cpp b/Engine/Source/Runtime/Engine/Private/Materials/MaterialShared.cpp
--- a/Engine/Source/Runtime/Engine/Private/Materials/MaterialShared.cpp
+++ b/Engine/Source/Runtime/Engine/Private/Materials/MaterialShared.cpp
@@ -151,8 +151,6 @@ class FMaterialDumpDebugInfoExecHelper : public FSelfRegisteringExec
 				}
 				return true;
 			}
-
-			TODO: PRINT WARNING
 		}
 		return false;
 	}
