Fixed linkage error in GLTF exporter whose UGLTFMaterialAnalyzer class derives from UMaterialInterface

diff --git a/Engine/Source/Runtime/Engine/Classes/Materials/MaterialInterface.h b/Engine/Source/Runtime/Engine/Classes/Materials/MaterialInterface.h
--- a/Engine/Source/Runtime/Engine/Classes/Materials/MaterialInterface.h
+++ b/Engine/Source/Runtime/Engine/Classes/Materials/MaterialInterface.h
@@ -271,8 +271,8 @@ private:
 
 public:
 
-	UMaterialInterface();
-	UMaterialInterface(FVTableHelper& Helper);
+	ENGINE_API UMaterialInterface();
+	ENGINE_API UMaterialInterface(FVTableHelper& Helper);
 	ENGINE_API virtual ~UMaterialInterface();
 
 	//~ Begin IInterface_AssetUserData Interface
