Fix Reverb to update correctly through level transitions and some edge cases with bApplyReverb = false

diff --git a/Engine/Source/Runtime/Engine/Private/AudioDevice.cpp b/Engine/Source/Runtime/Engine/Private/AudioDevice.cpp
--- a/Engine/Source/Runtime/Engine/Private/AudioDevice.cpp
+++ b/Engine/Source/Runtime/Engine/Private/AudioDevice.cpp
@@ -338,6 +338,7 @@ FAudioDevice::FAudioDevice()
 	, AudioClock(0.0)
 	, bAllowCenterChannel3DPanning(false)
 	, DeviceDeltaTime(0.0f)
+	, bHasActivatedReverb(false)
 	, bAllowPlayWhenSilent(true)
 	, bUseAttenuationForNonGameWorlds(false)
 	, ConcurrencyManager(this)
@@ -2371,6 +2372,7 @@ void FAudioDevice::UpdateHighestPriorityReverb()
 		const FActivatedReverb& NewActiveReverbRef = ActivatedReverbs.CreateConstIterator().Value();
 		FAudioThread::RunCommandOnAudioThread([AudioDevice, NewActiveReverbRef]()
 		{
+			AudioDevice->bHasActivatedReverb = true;
 			AudioDevice->HighestPriorityActivatedReverb = NewActiveReverbRef;
 		}, GET_STATID(STAT_AudioUpdateHighestPriorityReverb));
 	}
@@ -2378,6 +2380,7 @@ void FAudioDevice::UpdateHighestPriorityReverb()
 	{
 		FAudioThread::RunCommandOnAudioThread([AudioDevice]()
 		{
+			AudioDevice->bHasActivatedReverb = false;
 			AudioDevice->HighestPriorityActivatedReverb = FActivatedReverb();
 		}, GET_STATID(STAT_AudioUpdateHighestPriorityReverb));
 	}
@@ -4592,7 +4595,6 @@ void FAudioDevice::UpdateAudioVolumeEffects()
 	bool bHasVolumeSettings = false;
 	FAudioVolumeSettings PlayerAudioVolumeSettings;
 	bool bUsingDefaultReverb = true;
-	bool bEnteredOrExitedAudioVolume = false;
 
 	FAudioVolumeSettings PreviousPlayerAudioVolumeSettings = CurrentAudioVolumeSettings;
 
@@ -4610,9 +4612,17 @@ void FAudioDevice::UpdateAudioVolumeEffects()
 			bHasVolumeSettings = true;
 			PlayerAudioVolumeSettings = NewPlayerAudioVolumeSettings;
 
-			if (NewPlayerAudioVolumeSettings.AudioVolumeID > 0 && NewPlayerAudioVolumeSettings.ReverbSettings.bApplyReverb)
+			if (NewPlayerAudioVolumeSettings.AudioVolumeID > 0)
 			{
-				bUsingDefaultReverb = false;
+				if (NewPlayerAudioVolumeSettings.ReverbSettings.bApplyReverb)
+				{
+					bUsingDefaultReverb = false;
+				}
+				else if (GetDefaultAudioSettings(Listener.WorldID, NewPlayerAudioVolumeSettings.ReverbSettings, NewPlayerAudioVolumeSettings.InteriorSettings))
+				{
+					// Fall back to world reverb settings
+					PlayerAudioVolumeSettings.ReverbSettings = NewPlayerAudioVolumeSettings.ReverbSettings;
+				}
 			}
 		}
 	}
@@ -4620,27 +4630,31 @@ void FAudioDevice::UpdateAudioVolumeEffects()
 	// Reset audio volume proxies to reset state back to before a proxy was updated
 	ResetAudioVolumeProxyChangedState();
 
-	if (PlayerAudioVolumeSettings.AudioVolumeID != CurrentAudioVolumeSettings.AudioVolumeID)
-	{
-		// We're inside a new AudioVolume
-		CurrentAudioVolumeSettings = PlayerAudioVolumeSettings;
-		bEnteredOrExitedAudioVolume = true;
-	}
-	else if (PlayerAudioVolumeSettings.AudioVolumeID == CurrentAudioVolumeSettings.AudioVolumeID && PlayerAudioVolumeSettings.bChanged)
-	{
-		// Our active AudioVolume was updated
-		CurrentAudioVolumeSettings = PlayerAudioVolumeSettings;
-		bEnteredOrExitedAudioVolume = true;
-	}
+	// We need to update submix effects if we've entered a new volume or if the current volume settings have been updated
+	bool bUpdateSubmixEffects = (PlayerAudioVolumeSettings.AudioVolumeID != CurrentAudioVolumeSettings.AudioVolumeID) ||
+		(PlayerAudioVolumeSettings.AudioVolumeID == CurrentAudioVolumeSettings.AudioVolumeID && PlayerAudioVolumeSettings.bChanged);
+
+	CurrentAudioVolumeSettings = PlayerAudioVolumeSettings;
 
 	if (Effects)
 	{
 		// Check if we should be using activated reverb
-		if (HighestPriorityActivatedReverb.Priority > PlayerAudioVolumeSettings.Priority || bUsingDefaultReverb)
+		if (bHasActivatedReverb && (HighestPriorityActivatedReverb.Priority > PlayerAudioVolumeSettings.Priority || bUsingDefaultReverb))
 		{
 			CurrentAudioVolumeSettings.ReverbSettings = HighestPriorityActivatedReverb.ReverbSettings;
 		}
 
+#if ENABLE_AUDIO_DEBUG
+		// This is for debug visualization only - the audio debugger will show a reverb effect is active, if bApplyReverb is false for the world (default) reverb.
+		// In order to not mislead people using au.debug.reverb, we manually set the reverb effect to none.  This doesn't impact non-debug enabled
+		// builds, as a reverb effect w/ ApplyReverb false is equivalent to a null reverb effect that's been applied.
+		else if (bUsingDefaultReverb && !CurrentAudioVolumeSettings.ReverbSettings.bApplyReverb)
+		{
+			CurrentAudioVolumeSettings.ReverbSettings.bApplyReverb = true;
+			CurrentAudioVolumeSettings.ReverbSettings.ReverbEffect = nullptr;
+		}
+#endif // ENABLE_AUDIO_DEBUG
+
 		// Update the master reverb if it has changed
 		if (CurrentAudioVolumeSettings.bChanged || (CurrentAudioVolumeSettings.ReverbSettings != PreviousPlayerAudioVolumeSettings.ReverbSettings))
 		{
@@ -4651,7 +4665,7 @@ void FAudioDevice::UpdateAudioVolumeEffects()
 		Effects->Update();
 
 		// If we any submix override settings apply those overrides to the indicated submixes
-		if (IsAudioMixerEnabled() && bEnteredOrExitedAudioVolume)
+		if (IsAudioMixerEnabled() && bUpdateSubmixEffects)
 		{
 			// Clear out any previous submix effect chain overrides if the audio volume changed
 			if (PreviousPlayerAudioVolumeSettings.SubmixOverrideSettings.Num() > 0)
@@ -6302,6 +6316,7 @@ void FAudioDevice::Flush(UWorld* WorldToFlush, bool bClearActivatedReverb)
 	if (bClearActivatedReverb)
 	{
 		ActivatedReverbs.Reset();
+		bHasActivatedReverb = false;
 	}
 
 	if (WorldToFlush == nullptr)
diff --git a/Engine/Source/Runtime/Engine/Public/AudioDevice.h b/Engine/Source/Runtime/Engine/Public/AudioDevice.h
--- a/Engine/Source/Runtime/Engine/Public/AudioDevice.h
+++ b/Engine/Source/Runtime/Engine/Public/AudioDevice.h
@@ -2185,6 +2185,9 @@ protected:
 
 private:
 
+	/** Whether the value in HighestPriorityActivatedReverb should be used - Audio Thread owned */
+	uint8 bHasActivatedReverb:1;
+
 	/** Whether or not we're supporting zero volume wave instances */
 	uint8 bAllowPlayWhenSilent:1;
 
