Fix crash when handling replay map load failure

diff --git a/Engine/Source/Runtime/Engine/Private/UnrealEngine.cpp b/Engine/Source/Runtime/Engine/Private/UnrealEngine.cpp
--- a/Engine/Source/Runtime/Engine/Private/UnrealEngine.cpp
+++ b/Engine/Source/Runtime/Engine/Private/UnrealEngine.cpp
@@ -13248,7 +13248,10 @@ void DestroyNamedNetDriver_Local(FWorldContext &Context, FName NetDriverName)
 			NetDriver->LowLevelDestroy();
 			NetDriver->MarkAsGarbage();
 
-			Context.World()->ClearNetDriver(NetDriver);
+			if (UWorld* World = Context.World())
+			{
+				World->ClearNetDriver(NetDriver);
+			}
 
 			break;
 		}
