Fix includes for AnimMontage
Use UE_ENABLE_INCLUDE_ORDER_DEPRECATED_IN_5_1

diff --git a/Engine/Plugins/Developer/AnimationSharing/Source/AnimationSharing/Public/AnimationSharingTypes.h b/Engine/Plugins/Developer/AnimationSharing/Source/AnimationSharing/Public/AnimationSharingTypes.h
--- a/Engine/Plugins/Developer/AnimationSharing/Source/AnimationSharing/Public/AnimationSharingTypes.h
+++ b/Engine/Plugins/Developer/AnimationSharing/Source/AnimationSharing/Public/AnimationSharingTypes.h
@@ -2,17 +2,22 @@
 
 #pragma once
 
-#include "Engine/DeveloperSettings.h"
 #include "PerPlatformProperties.h"
+#include "UObject/Class.h"
+#include "AnimationSharingInstances.h"
+#if UE_ENABLE_INCLUDE_ORDER_DEPRECATED_IN_5_1
+#include "Engine/DeveloperSettings.h"
 #include "Animation/AnimBlueprint.h"
 #include "Animation/AnimSequence.h"
 #include "Animation/Skeleton.h"
 #include "Engine/SkeletalMesh.h"
-#include "UObject/Class.h"
-#include "Animation/AnimInstance.h"
-#include "AnimationSharingInstances.h"
+#endif
 #include "AnimationSharingTypes.generated.h"
 
+class UAnimSequence;
+class USkeleton;
+class USkeletonMesh;
+
 USTRUCT()
 struct FAnimationSetup
 {
diff --git a/Engine/Source/Runtime/Engine/Classes/Animation/AnimMontage.h b/Engine/Source/Runtime/Engine/Classes/Animation/AnimMontage.h
--- a/Engine/Source/Runtime/Engine/Classes/Animation/AnimMontage.h
+++ b/Engine/Source/Runtime/Engine/Classes/Animation/AnimMontage.h
@@ -12,7 +12,9 @@
 #include "UObject/Object.h"
 #include "Animation/AnimLinkableElement.h"
 #include "Animation/AnimTypes.h"
+#if UE_ENABLE_INCLUDE_ORDER_DEPRECATED_IN_5_1
 #include "Animation/BlendProfile.h"
+#endif
 #include "Animation/Skeleton.h"
 #include "Animation/AnimationAsset.h"
 #include "AlphaBlend.h"
@@ -23,8 +25,11 @@
 class UAnimInstance;
 class UAnimMontage;
 class UAnimSequence;
+class UBlendProfile;
 class USkeletalMeshComponent;
 
+enum class EBlendProfileMode : uint8;
+
 
 /**
  * Section data for each track. Reference of data will be stored in the child class for the way they want
diff --git a/Engine/Source/Runtime/Engine/Classes/Animation/AnimationAsset.h b/Engine/Source/Runtime/Engine/Classes/Animation/AnimationAsset.h
--- a/Engine/Source/Runtime/Engine/Classes/Animation/AnimationAsset.h
+++ b/Engine/Source/Runtime/Engine/Classes/Animation/AnimationAsset.h
@@ -8,18 +8,22 @@
 #pragma once
 
 #include "CoreMinimal.h"
+#include "Animation/AnimTypes.h"
 #include "UObject/ObjectMacros.h"
 #include "UObject/Object.h"
 #include "Misc/Guid.h"
 #include "Templates/SubclassOf.h"
 #include "Interfaces/Interface_AssetUserData.h"
+#if UE_ENABLE_INCLUDE_ORDER_DEPRECATED_IN_5_1
 #include "Engine/SkeletalMesh.h"
+#endif
 #include "AnimInterpFilter.h"
 #include "AnimEnums.h"
 #include "Interfaces/Interface_PreviewMeshProvider.h"
 #include "AnimationAsset.generated.h"
 
 class UAnimMetaData;
+class UAnimMontage;
 class UAssetMappingTable;
 class UAssetUserData;
 class USkeleton;
@@ -27,6 +31,7 @@ class UAnimSequenceBase;
 class UBlendSpace;
 class UPoseAsset;
 class UMirrorDataTable;
+class USkeletalMesh;
 struct FAnimationUpdateContext;
 
 namespace UE { namespace Anim {
diff --git a/Engine/Source/Runtime/Engine/Private/Animation/AnimMontage.cpp b/Engine/Source/Runtime/Engine/Private/Animation/AnimMontage.cpp
--- a/Engine/Source/Runtime/Engine/Private/Animation/AnimMontage.cpp
+++ b/Engine/Source/Runtime/Engine/Private/Animation/AnimMontage.cpp
@@ -12,6 +12,7 @@
 #include "Animation/AssetMappingTable.h"
 #include "Animation/AnimSequence.h"
 #include "Animation/AnimInstance.h"
+#include "Animation/BlendProfile.h"
 #include "AnimationUtils.h"
 #include "AnimationRuntime.h"
 #include "Animation/AnimNotifies/AnimNotify.h"
diff --git a/Engine/Source/Runtime/Engine/Private/AnimationAsset.cpp b/Engine/Source/Runtime/Engine/Private/AnimationAsset.cpp
--- a/Engine/Source/Runtime/Engine/Private/AnimationAsset.cpp
+++ b/Engine/Source/Runtime/Engine/Private/AnimationAsset.cpp
@@ -2,6 +2,7 @@
 
 #include "Animation/AnimationAsset.h"
 #include "Engine/AssetUserData.h"
+#include "Engine/SkeletalMesh.h"
 #include "Animation/AssetMappingTable.h"
 #include "Animation/AnimMetaData.h"
 #include "Animation/AnimSequence.h"
