Forward declare EMoveComponentFlags to fix non unity build

diff --git a/Engine/Source/Runtime/Engine/Public/CharacterMovementComponentAsync.h b/Engine/Source/Runtime/Engine/Public/CharacterMovementComponentAsync.h
--- a/Engine/Source/Runtime/Engine/Public/CharacterMovementComponentAsync.h
+++ b/Engine/Source/Runtime/Engine/Public/CharacterMovementComponentAsync.h
@@ -6,7 +6,7 @@
 #include "Engine/OverlapInfo.h"
 #include "CharacterMovementComponentAsync.generated.h"
 
-
+enum EMoveComponentFlags;
 struct FCharacterMovementComponentAsyncOutput;
 struct FCharacterMovementComponentAsyncInput;
 
