Allow FMemberReference to resolve via the MemberGuid when available in a cooked build

Note: Only class property references can potentially have MemberGuid info in a cooked build

#rnx

diff --git a/Engine/Source/Runtime/Engine/Classes/Engine/Blueprint.h b/Engine/Source/Runtime/Engine/Classes/Engine/Blueprint.h
--- a/Engine/Source/Runtime/Engine/Classes/Engine/Blueprint.h
+++ b/Engine/Source/Runtime/Engine/Classes/Engine/Blueprint.h
@@ -453,7 +453,7 @@ ENUM_CLASS_FLAGS(EGetObjectOrWorldBeingDebuggedFlags);
  * within Unreal Editor without ever needing to write a line of code.
  */
 UCLASS(config=Engine)
-class ENGINE_API UBlueprint : public UBlueprintCore
+class ENGINE_API UBlueprint : public UBlueprintCore, public IBlueprintPropertyGuidProvider
 {
 	GENERATED_UCLASS_BODY()
 
@@ -1036,6 +1036,11 @@ public:
 	void ConformNativeComponents();
 #endif	//#if WITH_EDITOR
 
+	//~ Begin IBlueprintPropertyGuidProvider interface
+	virtual FName FindBlueprintPropertyNameFromGuid(const FGuid& PropertyGuid) const override final;
+	virtual FGuid FindBlueprintPropertyGuidFromName(const FName PropertyName) const override final;
+	//~ End IBlueprintPropertyGuidProvider interface
+
 	//~ Begin UObject Interface
 #if WITH_EDITORONLY_DATA
 	PRAGMA_DISABLE_DEPRECATION_WARNINGS // Suppress compiler warning on override of deprecated function
@@ -1075,31 +1080,37 @@ public:
 	 * @return						true if there were no status errors in any of the parent blueprints, otherwise false
 	 */
 	static bool GetBlueprintHierarchyFromClass(const UClass* InClass, TArray<UBlueprintGeneratedClass*>& OutBlueprintParents);
-	
+
+private:
+	/**
+	 * Gets an array of all IBlueprintPropertyGuidProviders for this class and its parents.  0th elements is the IBlueprintPropertyGuidProvider for InClass
+	 *
+	 * @param InClass				The class to get the blueprint lineage for
+	 * @param OutBlueprintParents	Array of IBlueprintPropertyGuidProviders for this class and its parents.  0th = this, Nth = least derived BP-based parent
+	 * @return						true if there were no status errors in any of the parent blueprints, otherwise false
+	 */
+	static bool GetBlueprintHierarchyFromClass(const UClass* InClass, TArray<IBlueprintPropertyGuidProvider*>& OutBlueprintParents);
+
+public:
 #if WITH_EDITOR
 	/** returns true if the class hierarchy is error free */
 	static bool IsBlueprintHierarchyErrorFree(const UClass* InClass);
 #endif
 
-#if WITH_EDITOR
 	template<class TFieldType>
 	static FName GetFieldNameFromClassByGuid(const UClass* InClass, const FGuid VarGuid)
 	{
 		FProperty* AssertPropertyType = (TFieldType*)0;
 
-		TArray<UBlueprint*> Blueprints;
-		UBlueprint::GetBlueprintHierarchyFromClass(InClass, Blueprints);
+		TArray<IBlueprintPropertyGuidProvider*> BlueprintPropertyGuidProviders;
+		UBlueprint::GetBlueprintHierarchyFromClass(InClass, BlueprintPropertyGuidProviders);
 
-		for (int32 BPIndex = 0; BPIndex < Blueprints.Num(); ++BPIndex)
+		for (IBlueprintPropertyGuidProvider* BlueprintPropertyGuidProvider : BlueprintPropertyGuidProviders)
 		{
-			UBlueprint* Blueprint = Blueprints[BPIndex];
-			for (int32 VarIndex = 0; VarIndex < Blueprint->NewVariables.Num(); ++VarIndex)
+			const FName FoundPropertyName = BlueprintPropertyGuidProvider->FindBlueprintPropertyNameFromGuid(VarGuid);
+			if (FoundPropertyName != NAME_None)
 			{
-				const FBPVariableDescription& BPVarDesc = Blueprint->NewVariables[VarIndex];
-				if (BPVarDesc.VarGuid == VarGuid)
-				{
-					return BPVarDesc.VarName;
-				}
+				return FoundPropertyName;
 			}
 		}
 
@@ -1111,26 +1122,23 @@ public:
 	{
 		FProperty* AssertPropertyType = (TFieldType*)0;
 
-		TArray<UBlueprint*> Blueprints;
-		UBlueprint::GetBlueprintHierarchyFromClass(InClass, Blueprints);
+		TArray<IBlueprintPropertyGuidProvider*> BlueprintPropertyGuidProviders;
+		UBlueprint::GetBlueprintHierarchyFromClass(InClass, BlueprintPropertyGuidProviders);
 
-		for (int32 BPIndex = 0; BPIndex < Blueprints.Num(); ++BPIndex)
+		for (IBlueprintPropertyGuidProvider* BlueprintPropertyGuidProvider : BlueprintPropertyGuidProviders)
 		{
-			UBlueprint* Blueprint = Blueprints[BPIndex];
-			for (int32 VarIndex = 0; VarIndex < Blueprint->NewVariables.Num(); ++VarIndex)
+			const FGuid FoundPropertyGuid = BlueprintPropertyGuidProvider->FindBlueprintPropertyGuidFromName(VarName);
+			if (FoundPropertyGuid.IsValid())
 			{
-				const FBPVariableDescription& BPVarDesc = Blueprint->NewVariables[VarIndex];
-				if (BPVarDesc.VarName == VarName)
-				{
-					VarGuid = BPVarDesc.VarGuid;
-					return true;
-				}
+				VarGuid = FoundPropertyGuid;
+				return true;
 			}
 		}
 
 		return false;
 	}
 
+#if WITH_EDITOR
 	static FName GetFunctionNameFromClassByGuid(const UClass* InClass, const FGuid FunctionGuid);
 	static bool GetFunctionGuidFromClassByFieldName(const UClass* InClass, const FName FunctionName, FGuid& FunctionGuid);
 
@@ -1276,17 +1284,22 @@ public:
 };
 
 
-#if WITH_EDITOR
 template<>
 inline FName UBlueprint::GetFieldNameFromClassByGuid<UFunction>(const UClass* InClass, const FGuid FunctionGuid)
 {
+#if WITH_EDITOR
 	return GetFunctionNameFromClassByGuid(InClass, FunctionGuid);
+#else	// WITH_EDITOR
+	return NAME_None;
+#endif	// WITH_EDITOR
 }
 
 template<>
 inline bool UBlueprint::GetGuidFromClassByFieldName<UFunction>(const UClass* InClass, const FName FunctionName, FGuid& FunctionGuid)
 {
+#if WITH_EDITOR
 	return GetFunctionGuidFromClassByFieldName(InClass, FunctionName, FunctionGuid);
+#else	// WITH_EDITOR
+	return false;
+#endif	// WITH_EDITOR
 }
-
-#endif // #if WITH_EDITOR
diff --git a/Engine/Source/Runtime/Engine/Classes/Engine/BlueprintGeneratedClass.h b/Engine/Source/Runtime/Engine/Classes/Engine/BlueprintGeneratedClass.h
--- a/Engine/Source/Runtime/Engine/Classes/Engine/BlueprintGeneratedClass.h
+++ b/Engine/Source/Runtime/Engine/Classes/Engine/BlueprintGeneratedClass.h
@@ -624,8 +624,30 @@ struct FBPComponentClassOverride
 	}
 };
 
+/**
+ * Interface to query the property name<->GUID relationship using either a UBlueprint or a UBlueprintGeneratedClass.
+ * This allows cooked and uncooked Blueprints to be queried via the same API.
+ */
+class IBlueprintPropertyGuidProvider
+{
+public:
+	virtual ~IBlueprintPropertyGuidProvider() = default;
+
+	/**
+	 * Returns the property name for the given GUID, if any.
+	 * @note Does not recurse into parents.
+	 */
+	virtual FName FindBlueprintPropertyNameFromGuid(const FGuid& PropertyGuid) const = 0;
+
+	/**
+	 * Returns the property GUID for the given name, if any.
+	 * @note Does not recurse into parents.
+	 */
+	virtual FGuid FindBlueprintPropertyGuidFromName(const FName PropertyName) const = 0;
+};
+
 UCLASS(NeedsDeferredDependencyLoading)
-class ENGINE_API UBlueprintGeneratedClass : public UClass
+class ENGINE_API UBlueprintGeneratedClass : public UClass, public IBlueprintPropertyGuidProvider
 {
 	GENERATED_UCLASS_BODY()
 
@@ -737,6 +759,11 @@ public:
 	static void CreateComponentsForActor(const UClass* ThisClass, AActor* Actor);
 	static void CreateTimelineComponent(AActor* Actor, const UTimelineTemplate* TimelineTemplate);
 
+	// IBlueprintPropertyGuidProvider interface
+	virtual FName FindBlueprintPropertyNameFromGuid(const FGuid& PropertyGuid) const override final;
+	virtual FGuid FindBlueprintPropertyGuidFromName(const FName PropertyName) const override final;
+	// End IBlueprintPropertyGuidProvider interface
+
 	// UObject interface
 	virtual void Serialize(FArchive& Ar) override;
 	virtual void PostLoad() override;
diff --git a/Engine/Source/Runtime/Engine/Classes/Engine/MemberReference.h b/Engine/Source/Runtime/Engine/Classes/Engine/MemberReference.h
--- a/Engine/Source/Runtime/Engine/Classes/Engine/MemberReference.h
+++ b/Engine/Source/Runtime/Engine/Classes/Engine/MemberReference.h
@@ -132,13 +132,13 @@ public:
 		{
 			MemberParent = ParentAsClass->GetAuthoritativeClass();
 		}
+#endif
 
 		MemberGuid.Invalidate();
 		if (OwnerClass != nullptr)
 		{
 			UBlueprint::GetGuidFromClassByFieldName<TFieldType>(OwnerClass, InField->GetFName(), MemberGuid);
 		}
-#endif
 	}
 
 #if WITH_EDITOR
@@ -407,9 +407,7 @@ public:
 				MemberParent = Cast<UClass>(GetFieldOuter(static_cast<typename TFieldType::BaseFieldClass*>(ReturnField)));
 
 				MemberGuid.Invalidate();
-#if WITH_EDITOR
 				UBlueprint::GetGuidFromClassByFieldName<TFieldType>(TargetScope, MemberName, MemberGuid);
-#endif
 
 				if (UClass* ParentAsClass = GetMemberParentClass())
 				{
@@ -454,25 +452,28 @@ public:
 					}
 				}
 
-#if WITH_EDITOR
 
-				// If the reference variable is valid we need to make sure that our GUID matches
-				if (ReturnField != nullptr)
-				{
-					UBlueprint::GetGuidFromClassByFieldName<TFieldType>(TargetScope, MemberName, MemberGuid);
-				}
-				// If we have a GUID find the reference variable and make sure the name is up to date and find the field again
-				// For now only variable references will have valid GUIDs.  Will have to deal with finding other names subsequently
-				else if (MemberGuid.IsValid())
+#if !WITH_EDITOR
+				if (bAlwaysFollowRedirects)
+#endif
 				{
-					const FName RenamedMemberName = UBlueprint::GetFieldNameFromClassByGuid<TFieldType>(TargetScope, MemberGuid);
-					if (RenamedMemberName != NAME_None)
+					// If the reference variable is valid we need to make sure that our GUID matches
+					if (ReturnField != nullptr)
 					{
-						MemberName = RenamedMemberName;
-						ReturnField = FindUFieldOrFProperty<TFieldType>(TargetScope, MemberName, EFieldIterationFlags::IncludeAll);
+						UBlueprint::GetGuidFromClassByFieldName<TFieldType>(TargetScope, MemberName, MemberGuid);
+					}
+					// If we have a GUID find the reference variable and make sure the name is up to date and find the field again
+					// For now only variable references will have valid GUIDs.  Will have to deal with finding other names subsequently
+					else if (MemberGuid.IsValid())
+					{
+						const FName RenamedMemberName = UBlueprint::GetFieldNameFromClassByGuid<TFieldType>(TargetScope, MemberGuid);
+						if (RenamedMemberName != NAME_None)
+						{
+							MemberName = RenamedMemberName;
+							ReturnField = FindUFieldOrFProperty<TFieldType>(TargetScope, MemberName, EFieldIterationFlags::IncludeAll);
+						}
 					}
 				}
-#endif
 			}
 			else if (UPackage* TargetPackage = GetMemberParentPackage())
 			{
diff --git a/Engine/Source/Runtime/Engine/Private/Blueprint.cpp b/Engine/Source/Runtime/Engine/Private/Blueprint.cpp
--- a/Engine/Source/Runtime/Engine/Private/Blueprint.cpp
+++ b/Engine/Source/Runtime/Engine/Private/Blueprint.cpp
@@ -1538,6 +1538,44 @@ bool UBlueprint::GetBlueprintHierarchyFromClass(const UClass* InClass, TArray<UB
 	return bNoErrors;
 }
 
+bool UBlueprint::GetBlueprintHierarchyFromClass(const UClass* InClass, TArray<IBlueprintPropertyGuidProvider*>& OutBlueprintParents)
+{
+	OutBlueprintParents.Reset();
+
+	bool bNoErrors = true;
+
+	UBlueprintGeneratedClass* CurrentClass = Cast<UBlueprintGeneratedClass>(const_cast<UClass*>(InClass));
+	while (CurrentClass)
+	{
+		IBlueprintPropertyGuidProvider* GuidProviderToAdd = CurrentClass;
+
+#if WITH_EDITORONLY_DATA
+		UBlueprint* BP = UBlueprint::GetBlueprintFromClass(CurrentClass);
+
+		if (BP)
+		{
+			GuidProviderToAdd = BP;
+			bNoErrors &= (BP->Status != BS_Error);
+		}
+
+		// If valid, use stored ParentClass rather than the actual UClass::GetSuperClass(); handles the case when the class has not been recompiled yet after a reparent operation.
+		if (BP && BP->ParentClass)
+		{
+			CurrentClass = Cast<UBlueprintGeneratedClass>(BP->ParentClass);
+		}
+		else
+#endif // #if WITH_EDITORONLY_DATA
+		{
+			check(CurrentClass);
+			CurrentClass = Cast<UBlueprintGeneratedClass>(CurrentClass->GetSuperClass());
+		}
+
+		OutBlueprintParents.Add(GuidProviderToAdd);
+	}
+
+	return bNoErrors;
+}
+
 #if WITH_EDITOR
 bool UBlueprint::IsBlueprintHierarchyErrorFree(const UClass* InClass)
 {
@@ -1565,6 +1603,36 @@ bool UBlueprint::IsBlueprintHierarchyErrorFree(const UClass* InClass)
 }
 #endif
 
+FName UBlueprint::FindBlueprintPropertyNameFromGuid(const FGuid& PropertyGuid) const
+{
+#if WITH_EDITORONLY_DATA
+	for (const FBPVariableDescription& BPVarDesc : NewVariables)
+	{
+		if (BPVarDesc.VarGuid == PropertyGuid)
+		{
+			return BPVarDesc.VarName;
+		}
+	}
+#endif
+
+	return NAME_None;
+}
+
+FGuid UBlueprint::FindBlueprintPropertyGuidFromName(const FName PropertyName) const
+{
+#if WITH_EDITORONLY_DATA
+	for (const FBPVariableDescription& BPVarDesc : NewVariables)
+	{
+		if (BPVarDesc.VarName == PropertyName)
+		{
+			return BPVarDesc.VarGuid;
+		}
+	}
+#endif
+
+	return FGuid();
+}
+
 ETimelineSigType UBlueprint::GetTimelineSignatureForFunctionByName(const FName& FunctionName, const FName& ObjectPropertyName)
 {
 	check(SkeletonGeneratedClass != NULL);
diff --git a/Engine/Source/Runtime/Engine/Private/BlueprintGeneratedClass.cpp b/Engine/Source/Runtime/Engine/Private/BlueprintGeneratedClass.cpp
--- a/Engine/Source/Runtime/Engine/Private/BlueprintGeneratedClass.cpp
+++ b/Engine/Source/Runtime/Engine/Private/BlueprintGeneratedClass.cpp
@@ -2120,99 +2120,114 @@ void FBlueprintCookedComponentInstancingData::BuildCachedPropertyDataFromTemplat
 	INC_MEMORY_STAT_BY(STAT_BPCompInstancingFastPathMemory, CachedPropertyData.GetAllocatedSize());
 }
 
-bool UBlueprintGeneratedClass::ArePropertyGuidsAvailable() const
+FName UBlueprintGeneratedClass::FindBlueprintPropertyNameFromGuid(const FGuid& PropertyGuid) const
 {
-	auto ArePropertyGuidsAvailableImpl = [this](const TMap<FName, FGuid>& PropertyGuidsToUse) -> bool
+	auto FindPropertyNameFromGuidImpl = [&PropertyGuid](const TMap<FName, FGuid>& PropertyGuidsToUse) -> FName
 	{
-		// Property guid's are generated during compilation.
-		if (PropertyGuidsToUse.Num() > 0)
-		{
-			return true;
-		}
-		
-		if (UBlueprintGeneratedClass* Super = Cast<UBlueprintGeneratedClass>(GetSuperStruct()))
+		if (const FName* Result = PropertyGuidsToUse.FindKey(PropertyGuid))
 		{
-			// Our parent may have guids for inherited variables
-			return Super->ArePropertyGuidsAvailable();
+			return *Result;
 		}
-		
-		return false;
+
+		return NAME_None;
 	};
 
 	if (bCooked)
 	{
-		return ArePropertyGuidsAvailableImpl(CookedPropertyGuids);
+		return FindPropertyNameFromGuidImpl(CookedPropertyGuids);
 	}
 
 #if WITH_EDITORONLY_DATA
-	return ArePropertyGuidsAvailableImpl(PropertyGuids);
+	return FindPropertyNameFromGuidImpl(PropertyGuids);
 #else	// WITH_EDITORONLY_DATA
-	return false;
+	return NAME_None;
 #endif	// WITH_EDITORONLY_DATA
 }
 
-FName UBlueprintGeneratedClass::FindPropertyNameFromGuid(const FGuid& PropertyGuid) const
+FGuid UBlueprintGeneratedClass::FindBlueprintPropertyGuidFromName(const FName PropertyName) const
 {
-	auto FindPropertyNameFromGuidImpl = [this, &PropertyGuid](const TMap<FName, FGuid>& PropertyGuidsToUse) -> FName
+	auto FindPropertyGuidFromNameImpl = [&PropertyName](const TMap<FName, FGuid>& PropertyGuidsToUse) -> FGuid
 	{
-		// Check parent first as it may have renamed a property since this class was last saved
-		if (UBlueprintGeneratedClass* Super = Cast<UBlueprintGeneratedClass>(GetSuperStruct()))
-		{
-			FName RedirectedName = Super->FindPropertyNameFromGuid(PropertyGuid);
-			if (!RedirectedName.IsNone())
-			{
-				return RedirectedName;
-			}
-		}
-
-		if (const FName* Result = PropertyGuidsToUse.FindKey(PropertyGuid))
+		if (const FGuid* Result = PropertyGuidsToUse.Find(PropertyName))
 		{
 			return *Result;
 		}
 
-		return NAME_None;
+		return FGuid();
 	};
 
 	if (bCooked)
 	{
-		return FindPropertyNameFromGuidImpl(CookedPropertyGuids);
+		return FindPropertyGuidFromNameImpl(CookedPropertyGuids);
 	}
 
 #if WITH_EDITORONLY_DATA
-	return FindPropertyNameFromGuidImpl(PropertyGuids);
+	return FindPropertyGuidFromNameImpl(PropertyGuids);
 #else	// WITH_EDITORONLY_DATA
-	return NAME_None;
+	return FGuid();
 #endif	// WITH_EDITORONLY_DATA
 }
 
-FGuid UBlueprintGeneratedClass::FindPropertyGuidFromName(const FName InName) const
+bool UBlueprintGeneratedClass::ArePropertyGuidsAvailable() const
 {
-	auto FindPropertyGuidFromNameImpl = [this, &InName](const TMap<FName, FGuid>& PropertyGuidsToUse) -> FGuid
+	auto AreBlueprintPropertyGuidsAvailable = [this]()
 	{
-		if (const FGuid* Result = PropertyGuidsToUse.Find(InName))
+		if (bCooked)
 		{
-			return *Result;
+			return CookedPropertyGuids.Num() > 0;
 		}
-		
-		if (UBlueprintGeneratedClass* Super = Cast<UBlueprintGeneratedClass>(GetSuperStruct()))
+
+#if WITH_EDITORONLY_DATA
+		return PropertyGuids.Num() > 0;
+#else	// WITH_EDITORONLY_DATA
+		return false;
+#endif	// WITH_EDITORONLY_DATA
+	};
+
+	if (AreBlueprintPropertyGuidsAvailable())
+	{
+		return true;
+	}
+
+	if (UBlueprintGeneratedClass* Super = Cast<UBlueprintGeneratedClass>(GetSuperStruct()))
+	{
+		// Our parent may have guids for inherited variables
+		return Super->ArePropertyGuidsAvailable();
+	}
+
+	return false;
+}
+
+FName UBlueprintGeneratedClass::FindPropertyNameFromGuid(const FGuid& PropertyGuid) const
+{
+	// Check parent first as it may have renamed a property since this class was last saved
+	if (UBlueprintGeneratedClass* Super = Cast<UBlueprintGeneratedClass>(GetSuperStruct()))
+	{
+		FName RedirectedName = Super->FindPropertyNameFromGuid(PropertyGuid);
+		if (!RedirectedName.IsNone())
 		{
-			// Fall back to parent if this is an inherited variable
-			return Super->FindPropertyGuidFromName(InName);
+			return RedirectedName;
 		}
+	}
 
-		return FGuid();
-	};
+	return FindBlueprintPropertyNameFromGuid(PropertyGuid);
+}
 
-	if (bCooked)
+FGuid UBlueprintGeneratedClass::FindPropertyGuidFromName(const FName InName) const
+{
+	const FGuid FoundPropertyGuid = FindBlueprintPropertyGuidFromName(InName);
+	if (FoundPropertyGuid.IsValid())
 	{
-		return FindPropertyGuidFromNameImpl(CookedPropertyGuids);
+		return FoundPropertyGuid;
+	}
+
+	if (UBlueprintGeneratedClass* Super = Cast<UBlueprintGeneratedClass>(GetSuperStruct()))
+	{
+		// Fall back to parent if this is an inherited variable
+		return Super->FindPropertyGuidFromName(InName);
 	}
 
-#if WITH_EDITORONLY_DATA
-	return FindPropertyGuidFromNameImpl(PropertyGuids);
-#else	// WITH_EDITORONLY_DATA
 	return FGuid();
-#endif	// WITH_EDITORONLY_DATA
 }
 
 #if WITH_EDITORONLY_DATA
