Remove assertions on referencer in chain dumping lambda.
Add RootSet to keepflags check.
Pick whichever of garbage object or target objects is non-null for checking gc assertions. (Target 'leaked' object may not actually be marked garbage.)

diff --git a/Engine/Source/Runtime/Engine/Private/UnrealEngine.cpp b/Engine/Source/Runtime/Engine/Private/UnrealEngine.cpp
--- a/Engine/Source/Runtime/Engine/Private/UnrealEngine.cpp
+++ b/Engine/Source/Runtime/Engine/Private/UnrealEngine.cpp
@@ -15476,7 +15476,7 @@ void UEngine::VerifyLoadMapWorldCleanup(FWorldContext* ForWorldContext)
 
 static FORCEINLINE bool HasGarbageCollectionKeepFlags(FGCObjectInfo* ObjectInfo)
 {
-	return ObjectInfo && (ObjectInfo->HasAnyFlags(GARBAGE_COLLECTION_KEEPFLAGS) || ObjectInfo->HasAnyInternalFlags(EInternalObjectFlags::GarbageCollectionKeepFlags));
+	return ObjectInfo && (ObjectInfo->HasAnyFlags(GARBAGE_COLLECTION_KEEPFLAGS) || ObjectInfo->HasAnyInternalFlags(EInternalObjectFlags::GarbageCollectionKeepFlags | EInternalObjectFlags::RootSet));
 }
 
 static bool PrintStaleReferenceChainsAndFindReferencingObjects(UObject* ObjectToFindReferencesTo, FReferenceChainSearch& RefChainSearch, FGCObjectInfo*& OutGarbageObject, FGCObjectInfo*& OutReferencingObject, ELogVerbosity::Type Verbosity)
@@ -15490,7 +15490,6 @@ static bool PrintStaleReferenceChainsAndFindReferencingObjects(UObject* ObjectTo
 	return RefChainSearch.PrintResults([&ObjectToFindReferencesTo, &OutGarbageObject, &OutReferencingObject](FReferenceChainSearch::FCallbackParams& Params)
 		{
 			check(Params.Object);
-			check(Params.Referencer || HasGarbageCollectionKeepFlags(Params.Object));
 			if (!Params.Object->IsValid())
 			{
 				// We may find many chains that lead to the leak but for brevity we only report the first one in the fatal error below
@@ -15606,7 +15605,7 @@ TArray<FString> UEngine::FindAndPrintStaleReferencesToObjects(TConstArrayView<UO
 		{
 			PathToCulprit = GetPathToStaleObjectReferencer(ObjectToFindReferencesTo, RefChainSearch);
 			FString GarbageObjectName = OutGarbageObject ? OutGarbageObject->GetFullName() : ObjectToFindReferencesTo->GetFullName();
-			checkf(OutReferencingObject || HasGarbageCollectionKeepFlags(OutGarbageObject), TEXT("No object referencing %s found even though we have a valid reference chain"), *ObjectToFindReferencesTo->GetPathName());
+			checkf(OutReferencingObject || HasGarbageCollectionKeepFlags(OutGarbageObject ? OutGarbageObject : OutReferencingObject), TEXT("No object referencing %s found even though we have a valid reference chain"), *ObjectToFindReferencesTo->GetPathName());
 			GarbageErrorMessage = FString::Printf(TEXT("Object %s is being referenced by %s"), *GarbageObjectName, 
 				OutReferencingObject ? *OutReferencingObject->GetFullName() : (HasGarbageCollectionKeepFlags(OutGarbageObject) ? TEXT("GarbageCollectionKeepFlags") : TEXT("NULL")));
 		}
@@ -15628,7 +15627,7 @@ TArray<FString> UEngine::FindAndPrintStaleReferencesToObjects(TConstArrayView<UO
 					bReferenceChainFound = true;
 					PathToCulprit = GetPathToStaleObjectReferencer(ObjectToFindReferencesTo, HistorySearch);
 					FString GarbageObjectName = OutGarbageObject ? OutGarbageObject->GetFullName() : ObjectToFindReferencesTo->GetFullName();
-					checkf(OutReferencingObject, TEXT("No object referencing %s found even though we have a valid reference chain"), *ObjectToFindReferencesTo->GetPathName());
+					checkf(OutReferencingObject || HasGarbageCollectionKeepFlags(OutGarbageObject ? OutGarbageObject : OutReferencingObject), TEXT("No object referencing %s found even though we have a valid reference chain"), *ObjectToFindReferencesTo->GetPathName());
 					GarbageErrorMessage = FString::Printf(TEXT("Garbage object %s was previously being referenced by %s"), *GarbageObjectName, OutReferencingObject ? *OutReferencingObject->GetFullName() : TEXT("NULL"));
 				}
 			}
