Fixed usage of ERayTracingPrimitiveFlags::UnsupportedProxyType

* Its value is 0, so it must not be used with EnumHasAnyFlags()
* Added few more comments in related code

diff --git a/Engine/Source/Runtime/Engine/Private/StaticMeshRender.cpp b/Engine/Source/Runtime/Engine/Private/StaticMeshRender.cpp
--- a/Engine/Source/Runtime/Engine/Private/StaticMeshRender.cpp
+++ b/Engine/Source/Runtime/Engine/Private/StaticMeshRender.cpp
@@ -1774,6 +1774,13 @@ void FStaticMeshSceneProxy::GetDynamicRayTracingInstances(FRayTracingMaterialGat
 		return;
 	}
 
+	if (!ensureMsgf(IsRayTracingRelevant(),
+		TEXT("GetDynamicRayTracingInstances() is only expected to be called for scene proxies that are compatible with ray tracing. ")
+		TEXT("RT-relevant primitive gathering code in FDeferredShadingSceneRenderer may be wrong.")))
+	{
+		return;
+	}
+
 	ESceneDepthPriorityGroup PrimitiveDPG = GetStaticDepthPriorityGroup();
 	const uint32 LODIndex = FMath::Max(GetLOD(Context.ReferenceView), (int32)GetCurrentFirstLODIdx_RenderThread());
 	const FStaticMeshLODResources& LODModel = RenderData->LODResources[LODIndex];
