Fix GetLevelInstanceBounds trying to compute the level instance bounding box using an unregistered level instance.
This can occur while validating the level instance world asset in ALevelInstance::PostRegisterAllComponent

diff --git a/Engine/Source/Runtime/Engine/Private/LevelInstance/LevelInstanceSubsystem.cpp b/Engine/Source/Runtime/Engine/Private/LevelInstance/LevelInstanceSubsystem.cpp
--- a/Engine/Source/Runtime/Engine/Private/LevelInstance/LevelInstanceSubsystem.cpp
+++ b/Engine/Source/Runtime/Engine/Private/LevelInstance/LevelInstanceSubsystem.cpp
@@ -628,12 +628,17 @@ bool ULevelInstanceSubsystem::GetLevelInstanceBounds(const ILevelInstanceInterfa
 		OutBounds = LevelInstanceEntry.LevelStreaming->GetBounds();
 		return true;
 	}
-	else if (const FLevelInstanceEdit* CurrentEdit = GetLevelInstanceEdit(LevelInstance))
+	
+	if (LevelInstance->HasValidLevelInstanceID()) // Check ID to make sure GetLevelInstancesBounds is called on a registered LevelInstance that can retrieve its Edit.
 	{
-		OutBounds = CurrentEdit->LevelStreaming->GetBounds();
-		return true;
+		if (const FLevelInstanceEdit* CurrentEdit = GetLevelInstanceEdit(LevelInstance))
+		{
+			OutBounds = CurrentEdit->LevelStreaming->GetBounds();
+			return true;
+		}
 	}
-	else if(LevelInstance->IsWorldAssetValid())
+	
+	if(LevelInstance->IsWorldAssetValid())
 	{
 		// @todo_ow: remove this temp fix once it is again safe to call the asset registry while saving.
 		// Currently it can lead to a FindObject which is illegal while saving.
