Reducing Check to error log in Stream Cache shutdown.

diff --git a/Engine/Source/Runtime/Engine/Public/AudioStreamingCache.h b/Engine/Source/Runtime/Engine/Public/AudioStreamingCache.h
--- a/Engine/Source/Runtime/Engine/Public/AudioStreamingCache.h
+++ b/Engine/Source/Runtime/Engine/Public/AudioStreamingCache.h
@@ -361,13 +361,18 @@ private:
 		~FCacheElement()
 		{
 			WaitForAsyncLoadCompletion(true);
-			checkf(NumConsumers.GetValue() == 0, TEXT("Tried to destroy streaming cache while the cached data was in use!"));
-			if (ChunkData)
+
+			if(NumConsumers.GetValue() != 0)
+			{
+				UE_LOG(LogAudioStreamCaching, Error, TEXT("Tried to destroy streaming cache while the cached data was in use! (Sound: %s - Loading Behavior: %s - Chunk Index: %i)")
+					, *Key.SoundWaveName.ToString(), EnumToString(Key.GetLoadingBehavior()), Key.ChunkIndex);
+			}
+			else if (ChunkData)
 			{
 				FMemory::Free(ChunkData);
+				ChunkData = nullptr;
 			}
 
-			ChunkData = nullptr;
 		}
 	};
 
