HLOD - Prevent false positive when testing for HLOD dirtiness
* Don't check HLODs for levels that are not visible, as they haven't got any LODActors yet, which means they'll always get cleaned out and dirtied.
* Sort the HLODProxyDesc subactors arrays to ensure a stable order for comparisons

diff --git a/Engine/Source/Runtime/Engine/Private/HLOD/HLODProxyDesc.cpp b/Engine/Source/Runtime/Engine/Private/HLOD/HLODProxyDesc.cpp
--- a/Engine/Source/Runtime/Engine/Private/HLOD/HLODProxyDesc.cpp
+++ b/Engine/Source/Runtime/Engine/Private/HLOD/HLODProxyDesc.cpp
@@ -128,6 +128,10 @@ bool UHLODProxyDesc::UpdateFromLODActor(const ALODActor* InLODActor)
 		}
 	}
 
+	// Sort the arrays to ensure a stable order for comparisons
+	SubActors.Sort(FNameLexicalLess());
+	SubHLODDescs.Sort(FSoftObjectPtrLexicalLess());
+
 	StaticMesh = InLODActor->StaticMeshComponent ? InLODActor->StaticMeshComponent->GetStaticMesh() : nullptr;
 
 	const TMap<FHLODInstancingKey, TObjectPtr<UInstancedStaticMeshComponent>>& ISMComponents = InLODActor->InstancedStaticMeshComponents;
@@ -184,6 +188,10 @@ bool UHLODProxyDesc::ShouldUpdateDesc(const ALODActor* InLODActor) const
 		}
 	}
 
+	// Sort the arrays to ensure a stable order for the comparisons below
+	LocalSubActors.Sort(FNameLexicalLess());
+	LocalSubHLODDescs.Sort(FSoftObjectPtrLexicalLess());
+
 	if (LocalSubActors != SubActors)
 	{
 		return true;
diff --git a/Engine/Source/Runtime/Engine/Private/HLODProxy.cpp b/Engine/Source/Runtime/Engine/Private/HLODProxy.cpp
--- a/Engine/Source/Runtime/Engine/Private/HLODProxy.cpp
+++ b/Engine/Source/Runtime/Engine/Private/HLODProxy.cpp
@@ -113,7 +113,9 @@ void UHLODProxy::Clean()
 	if (GetDefault<UHierarchicalLODSettings>()->bSaveLODActorsToHLODPackages)
 	{
 		UWorld* World = Cast<UWorld>(OwningMap.ToSoftObjectPath().ResolveObject());
-		if (World)
+
+		// Don't check HLODs for levels that are not visible, as they haven't got any LODActors yet, which means they'll always get cleaned out and dirtied.
+		if (World && World->PersistentLevel->bIsVisible)
 		{
 			UpdateHLODDescs(World->PersistentLevel);
 		}
