Fix potential missed notification when async shader compilation reaches 0 or when other compiler do not broadcast event when they finish all compilations

diff --git a/Engine/Source/Runtime/Engine/Private/AssetCompilingManager.cpp b/Engine/Source/Runtime/Engine/Private/AssetCompilingManager.cpp
--- a/Engine/Source/Runtime/Engine/Private/AssetCompilingManager.cpp
+++ b/Engine/Source/Runtime/Engine/Private/AssetCompilingManager.cpp
@@ -381,6 +381,8 @@ void FAssetCompilingManager::FinishAllCompilation()
 	{
 		AssetCompilingManager->FinishAllCompilation();
 	}
+
+	UpdateNumRemainingAssets();
 }
 
 /**
@@ -416,6 +418,21 @@ void FAssetCompilingManager::ProcessAsyncTasks(bool bLimitExecutionTime)
 	{
 		AssetCompilingManager->ProcessAsyncTasks(bLimitExecutionTime);
 	}
+
+	UpdateNumRemainingAssets();
+}
+
+void FAssetCompilingManager::UpdateNumRemainingAssets()
+{
+	const int32 NumRemainingAssets = GetNumRemainingAssets();
+	if (LastNumRemainingAssets > 0 && NumRemainingAssets == 0)
+	{
+		// This is important to at least broadcast once we reach 0 remaining assets
+		// because some listener are only interested to be notified when the number
+		// of total async compilation reaches 0.
+		FAssetCompilingManager::Get().OnAssetPostCompileEvent().Broadcast({});
+	}
+	LastNumRemainingAssets = NumRemainingAssets;
 }
 
 #undef LOCTEXT_NAMESPACE
\ No newline at end of file
diff --git a/Engine/Source/Runtime/Engine/Private/ShaderCompiler/ShaderCompiler.cpp b/Engine/Source/Runtime/Engine/Private/ShaderCompiler/ShaderCompiler.cpp
--- a/Engine/Source/Runtime/Engine/Private/ShaderCompiler/ShaderCompiler.cpp
+++ b/Engine/Source/Runtime/Engine/Private/ShaderCompiler/ShaderCompiler.cpp
@@ -3885,6 +3885,8 @@ void FShaderCompilingManager::SubmitJobs(TArray<FShaderCommonCompileJobPtr>& New
 	}
 
 	AllJobs.SubmitJobs(NewJobs);
+
+	UpdateNumRemainingAssets();
 }
 
 bool FShaderCompilingManager::IsCompilingShaderMap(uint32 Id)
@@ -4177,6 +4179,8 @@ void FShaderCompilingManager::BlockOnShaderMapCompletion(const TArray<int32>& Sh
 			}
 		}
 	}
+
+	UpdateNumRemainingAssets();
 }
 
 void FShaderCompilingManager::BlockOnAllShaderMapCompletion(TMap<int32, FShaderMapFinalizeResults>& CompiledShaderMaps)
@@ -4289,6 +4293,8 @@ void FShaderCompilingManager::BlockOnAllShaderMapCompletion(TMap<int32, FShaderM
 			It.RemoveCurrent();
 		}
 	}
+
+	UpdateNumRemainingAssets();
 }
 
 namespace
@@ -4603,6 +4609,8 @@ void FShaderCompilingManager::ProcessCompiledShaderMaps(
 	}
 
 	GShaderCompilingManager->PrintStats();
+
+	UpdateNumRemainingAssets();
 #endif // WITH_EDITOR
 }
 
@@ -5173,7 +5181,28 @@ void FShaderCompilingManager::ProcessAsyncResults(bool bLimitExecutionTime, bool
 		check(AllJobs.GetNumPendingJobs() == 0);
 	}
 
-	Notification.Update(GetNumRemainingAssets());
+	UpdateNumRemainingAssets();
+}
+
+void FShaderCompilingManager::UpdateNumRemainingAssets()
+{
+	if (IsInGameThread())
+	{
+		const int32 NumRemainingAssets = GetNumRemainingAssets();
+		if (LastNumRemainingAssets != NumRemainingAssets)
+		{
+			if (NumRemainingAssets == 0)
+			{
+				// This is important to at least broadcast once we reach 0 remaining assets
+				// even if we don't have any UObject to report because some listener are only 
+				// interested to be notified when the number of async compilation reaches 0.
+				FAssetCompilingManager::Get().OnAssetPostCompileEvent().Broadcast({});
+			}
+
+			LastNumRemainingAssets = NumRemainingAssets;
+			Notification.Update(NumRemainingAssets);
+		}
+	}
 }
 
 bool FShaderCompilingManager::IsShaderCompilerWorkerRunning(FProcHandle & WorkerHandle)
diff --git a/Engine/Source/Runtime/Engine/Public/AssetCompilingManager.h b/Engine/Source/Runtime/Engine/Public/AssetCompilingManager.h
--- a/Engine/Source/Runtime/Engine/Public/AssetCompilingManager.h
+++ b/Engine/Source/Runtime/Engine/Public/AssetCompilingManager.h
@@ -127,8 +127,12 @@ private:
 	FAssetCompilingManager();
 	friend struct IAssetCompilingManager;
 
+	/** Take some action whenever the number of remaining asset changes. */
+	void UpdateNumRemainingAssets();
+
 	bool bHasShutdown = false;
-	
+	int32 LastNumRemainingAssets = 0;
+
 	TArray<IAssetCompilingManager*> AssetCompilingManagers;
 	TArray<IAssetCompilingManager*> AssetCompilingManagersWithValidDependencies;
 
diff --git a/Engine/Source/Runtime/Engine/Public/ShaderCompiler.h b/Engine/Source/Runtime/Engine/Public/ShaderCompiler.h
--- a/Engine/Source/Runtime/Engine/Public/ShaderCompiler.h
+++ b/Engine/Source/Runtime/Engine/Public/ShaderCompiler.h
@@ -900,6 +900,8 @@ private:
 	FString AbsoluteShaderDebugInfoDirectory;
 	/** Name of the shader worker application. */
 	FString ShaderCompileWorkerName;
+	/** Last value of GetNumRemainingAssets */
+	int32 LastNumRemainingAssets = 0;
 
 	/** 
 	 * Tracks the total time that shader compile workers have been busy since startup.  
@@ -950,7 +952,10 @@ private:
 	
 	/** Checks if any target platform down't support remote shader compiling */
 	bool AllTargetPlatformSupportsRemoteShaderCompiling();
-	
+
+	/** Take some action whenever the number of remaining asset changes. */
+	void UpdateNumRemainingAssets();
+
 	/** Returns the first remote compiler controller found */
 	IDistributedBuildController* FindRemoteCompilerController() const;
 
