Do not create PrecomputedLightingUniformBuffer for primitive types having their static lighting data in GPUScene (static meshes, instanced static meshes, landscapes)

diff --git a/Engine/Source/Runtime/Engine/Private/SceneManagement.cpp b/Engine/Source/Runtime/Engine/Private/SceneManagement.cpp
--- a/Engine/Source/Runtime/Engine/Private/SceneManagement.cpp
+++ b/Engine/Source/Runtime/Engine/Private/SceneManagement.cpp
@@ -927,7 +927,7 @@ void InitializeSharedSamplerStates()
 
 void FLightCacheInterface::CreatePrecomputedLightingUniformBuffer_RenderingThread(ERHIFeatureLevel::Type FeatureLevel)
 {
-	if (LightMap || ShadowMap)
+	if (bNeedsPrecomputedLightingUniformBuffer && (LightMap || ShadowMap))
 	{
 		FPrecomputedLightingUniformParameters Parameters;
 		GetPrecomputedLightingParameters(FeatureLevel, Parameters, this);
diff --git a/Engine/Source/Runtime/Engine/Private/StaticMeshRender.cpp b/Engine/Source/Runtime/Engine/Private/StaticMeshRender.cpp
--- a/Engine/Source/Runtime/Engine/Private/StaticMeshRender.cpp
+++ b/Engine/Source/Runtime/Engine/Private/StaticMeshRender.cpp
@@ -2080,6 +2080,7 @@ FStaticMeshSceneProxy::FLODInfo::FLODInfo(const UStaticMeshComponent* InComponen
 			SetLightMap(MeshMapBuildData->LightMap);
 			SetShadowMap(MeshMapBuildData->ShadowMap);
 			SetResourceCluster(MeshMapBuildData->ResourceCluster);
+			bNeedsPrecomputedLightingUniformBuffer = false;
 			IrrelevantLights = MeshMapBuildData->IrrelevantLights;
 		}
 	}
@@ -2099,6 +2100,7 @@ FStaticMeshSceneProxy::FLODInfo::FLODInfo(const UStaticMeshComponent* InComponen
 					SetLightMap(MeshMapBuildData->LightMap);
 					SetShadowMap(MeshMapBuildData->ShadowMap);
 					SetResourceCluster(MeshMapBuildData->ResourceCluster);
+					bNeedsPrecomputedLightingUniformBuffer = false;
 					IrrelevantLights = MeshMapBuildData->IrrelevantLights;
 				}
 			}
@@ -2163,6 +2165,7 @@ FStaticMeshSceneProxy::FLODInfo::FLODInfo(const UStaticMeshComponent* InComponen
 				SetLightMap(MeshMapBuildData->LightMap);
 				SetShadowMap(MeshMapBuildData->ShadowMap);
 				SetResourceCluster(MeshMapBuildData->ResourceCluster);
+				bNeedsPrecomputedLightingUniformBuffer = false;
 				IrrelevantLights = MeshMapBuildData->IrrelevantLights;
 			}
 		}
diff --git a/Engine/Source/Runtime/Engine/Public/SceneManagement.h b/Engine/Source/Runtime/Engine/Public/SceneManagement.h
--- a/Engine/Source/Runtime/Engine/Public/SceneManagement.h
+++ b/Engine/Source/Runtime/Engine/Public/SceneManagement.h
@@ -837,17 +837,7 @@ ENGINE_API extern TGlobalResource< FDefaultLightmapResourceClusterUniformBuffer
 class FLightCacheInterface
 {
 public:
-	FLightCacheInterface()
-		: bGlobalVolumeLightmap(false)
-		, LightMap(nullptr)
-		, ShadowMap(nullptr)
-		, ResourceCluster(nullptr)
-	{
-	}
-
-	virtual ~FLightCacheInterface()
-	{
-	}
+	virtual ~FLightCacheInterface() {}
 
 	// @param LightSceneProxy must not be 0
 	virtual FLightInteraction GetInteraction(const class FLightSceneProxy* LightSceneProxy) const = 0;
@@ -914,20 +904,24 @@ public:
 
 	ENGINE_API FShadowMapInteraction GetShadowMapInteraction(ERHIFeatureLevel::Type InFeatureLevel) const;
 
+protected:
+	// If false, precomputed lighting parameters will be drawn from GPUScene
+	bool bNeedsPrecomputedLightingUniformBuffer = true;
+	
 private:
 
-	bool bGlobalVolumeLightmap;
+	bool bGlobalVolumeLightmap = false;
 
 	// The light-map used by the element. may be 0
-	const FLightMap* LightMap;
+	const FLightMap* LightMap = nullptr;
 
 	// The shadowmap used by the element, may be 0
-	const FShadowMap* ShadowMap;
+	const FShadowMap* ShadowMap = nullptr;
 
-	const FLightmapResourceCluster* ResourceCluster;
+	const FLightmapResourceCluster* ResourceCluster = nullptr;
 
 	/** The uniform buffer holding mapping the lightmap policy resources. */
-	FUniformBufferRHIRef PrecomputedLightingUniformBuffer;
+	FUniformBufferRHIRef PrecomputedLightingUniformBuffer = nullptr;
 };
 
 
diff --git a/Engine/Source/Runtime/Landscape/Public/LandscapeRender.h b/Engine/Source/Runtime/Landscape/Public/LandscapeRender.h
--- a/Engine/Source/Runtime/Landscape/Public/LandscapeRender.h
+++ b/Engine/Source/Runtime/Landscape/Public/LandscapeRender.h
@@ -591,6 +591,7 @@ class LANDSCAPE_API FLandscapeComponentSceneProxy : public FPrimitiveSceneProxy,
 				SetLightMap(MapBuildData->LightMap);
 				SetShadowMap(MapBuildData->ShadowMap);
 				SetResourceCluster(MapBuildData->ResourceCluster);
+				bNeedsPrecomputedLightingUniformBuffer = false;
 				IrrelevantLights = MapBuildData->IrrelevantLights;
 			}
 		}
