Ensure that a BP with cooked property GUIDs also has the GUIDs from its parent classes available (if they're not cooking their own GUIDs)

#rnx

diff --git a/Engine/Source/Runtime/Engine/Private/BlueprintGeneratedClass.cpp b/Engine/Source/Runtime/Engine/Private/BlueprintGeneratedClass.cpp
--- a/Engine/Source/Runtime/Engine/Private/BlueprintGeneratedClass.cpp
+++ b/Engine/Source/Runtime/Engine/Private/BlueprintGeneratedClass.cpp
@@ -1733,6 +1733,20 @@ void UBlueprintGeneratedClass::BeginCacheForCookedPlatformData(const ITargetPlat
 	if (ShouldCookBlueprintPropertyGuids())
 	{
 		CookedPropertyGuids = PropertyGuids;
+
+		if (GetDefault<UCookerSettings>()->BlueprintPropertyGuidsCookingMethod == EBlueprintPropertyGuidsCookingMethod::EnabledBlueprintsOnly)
+		{
+			// Ensure that we also have the GUIDs from our parent classes available (if they're not cooking their own GUIDs)
+			for (UBlueprintGeneratedClass* ParentBPGC = Cast<UBlueprintGeneratedClass>(GetSuperClass()); ParentBPGC; ParentBPGC = Cast<UBlueprintGeneratedClass>(ParentBPGC->GetSuperClass()))
+			{
+				const UBlueprint* ParentBP = Cast<UBlueprint>(ParentBPGC->ClassGeneratedBy);
+				if (!ParentBP || ParentBP->ShouldCookPropertyGuids())
+				{
+					break;
+				}
+				CookedPropertyGuids.Append(ParentBPGC->PropertyGuids);
+			}
+		}
 	}
 	else
 	{
