Strata - Fixed StrataSimpleClearCoat tooltip and cleanup

diff --git a/Engine/Source/Runtime/Engine/Classes/Materials/MaterialExpressionStrata.h b/Engine/Source/Runtime/Engine/Classes/Materials/MaterialExpressionStrata.h
--- a/Engine/Source/Runtime/Engine/Classes/Materials/MaterialExpressionStrata.h
+++ b/Engine/Source/Runtime/Engine/Classes/Materials/MaterialExpressionStrata.h
@@ -419,7 +419,6 @@ class UMaterialExpressionStrataSimpleClearCoatBSDF : public UMaterialExpressionS
 	virtual void GatherStrataMaterialInfo(FStrataMaterialInfo& StrataMaterialInfo, int32 OutputIndex) override;
 	virtual FStrataOperator* StrataGenerateMaterialTopologyTree(class FMaterialCompiler* Compiler, class UMaterialExpression* Parent, int32 OutputIndex) override;
 	virtual FName GetInputName(int32 InputIndex) const override;
-	virtual void GetConnectorToolTip(int32 InputIndex, int32 OutputIndex, TArray<FString>& OutToolTip) override;
 	virtual const TArray<FExpressionInput*> GetInputs() override;
 #endif
 	//~ End UMaterialExpression Interface
diff --git a/Engine/Source/Runtime/Engine/Private/Materials/MaterialExpressions.cpp b/Engine/Source/Runtime/Engine/Private/Materials/MaterialExpressions.cpp
--- a/Engine/Source/Runtime/Engine/Private/Materials/MaterialExpressions.cpp
+++ b/Engine/Source/Runtime/Engine/Private/Materials/MaterialExpressions.cpp
@@ -22627,8 +22627,8 @@ const TArray<FExpressionInput*> UMaterialExpressionStrataSimpleClearCoatBSDF::Ge
 	Result.Add(&Metallic);
 	Result.Add(&Specular);
 	Result.Add(&Roughness);
-	Result.Add(&ClearCoatRoughness);
 	Result.Add(&ClearCoatCoverage);
+	Result.Add(&ClearCoatRoughness);
 	Result.Add(&Normal);
 	Result.Add(&EmissiveColor);
 	return Result;
@@ -22664,11 +22664,11 @@ uint32 UMaterialExpressionStrataSimpleClearCoatBSDF::GetInputType(int32 InputInd
 	}
 	else if (InputIndex == 4)
 	{
-		return MCT_Float1; // ClearCoatRoughness
+		return MCT_Float1; // ClearCoatCoverage 
 	}
 	else if (InputIndex == 5)
 	{
-		return MCT_Float1; // ClearCoatCoverage
+		return MCT_Float1; // ClearCoatRoughness
 	}
 	else if (InputIndex == 6)
 	{
@@ -22703,11 +22703,11 @@ FName UMaterialExpressionStrataSimpleClearCoatBSDF::GetInputName(int32 InputInde
 	}
 	else if (InputIndex == 4)
 	{
-		return TEXT("Clear Coat Roughness");
+		return TEXT("Clear Coat Coverage");
 	}
 	else if (InputIndex == 5)
 	{
-		return TEXT("Clear Coat Coverage");
+		return TEXT("Clear Coat Roughness");
 	}
 	else if (InputIndex == 6)
 	{
@@ -22721,48 +22721,6 @@ FName UMaterialExpressionStrataSimpleClearCoatBSDF::GetInputName(int32 InputInde
 	return TEXT("Unknown");
 }
 
-void UMaterialExpressionStrataSimpleClearCoatBSDF::GetConnectorToolTip(int32 InputIndex, int32 OutputIndex, TArray<FString>& OutToolTip)
-{
-	if (OutputIndex == 0)
-	{
-		OutToolTip.Add(TEXT("TT Ouput"));
-		return;
-	}
-
-	if (InputIndex == 0)
-	{
-		Super::GetConnectorToolTip(0, INDEX_NONE, OutToolTip);
-	}
-	else if (InputIndex == 1)
-	{
-		Super::GetConnectorToolTip(1, INDEX_NONE, OutToolTip);
-	}
-	else if (InputIndex == 2)
-	{
-		Super::GetConnectorToolTip(2, INDEX_NONE, OutToolTip);
-	}
-	else if (InputIndex == 3)
-	{
-		Super::GetConnectorToolTip(3, INDEX_NONE, OutToolTip);
-	}
-	else if (InputIndex == 4)
-	{
-		Super::GetConnectorToolTip(4, INDEX_NONE, OutToolTip);
-	}
-	else if (InputIndex == 5)
-	{
-		Super::GetConnectorToolTip(5, INDEX_NONE, OutToolTip);
-	}
-	else if (InputIndex == 6)
-	{
-		Super::GetConnectorToolTip(6, INDEX_NONE, OutToolTip);
-	}
-	else if (InputIndex == 7)
-	{
-		Super::GetConnectorToolTip(7, INDEX_NONE, OutToolTip);
-	}
-}
-
 bool UMaterialExpressionStrataSimpleClearCoatBSDF::IsResultStrataMaterial(int32 OutputIndex)
 {
 	return true;
