Add a few important Insight trace tags

diff --git a/Engine/Source/Editor/UnrealEd/Private/Fbx/FbxMainImport.cpp b/Engine/Source/Editor/UnrealEd/Private/Fbx/FbxMainImport.cpp
--- a/Engine/Source/Editor/UnrealEd/Private/Fbx/FbxMainImport.cpp
+++ b/Engine/Source/Editor/UnrealEd/Private/Fbx/FbxMainImport.cpp
@@ -621,6 +621,8 @@ void FFbxImporter::PartialCleanUp()
 //-------------------------------------------------------------------------
 void FFbxImporter::ReleaseScene()
 {
+	TRACE_CPUPROFILER_EVENT_SCOPE(FFbxImporter::ReleaseScene);
+
 	if (Importer)
 	{
 		Importer->Destroy();
diff --git a/Engine/Source/Runtime/CoreUObject/Private/UObject/SavePackage.cpp b/Engine/Source/Runtime/CoreUObject/Private/UObject/SavePackage.cpp
--- a/Engine/Source/Runtime/CoreUObject/Private/UObject/SavePackage.cpp
+++ b/Engine/Source/Runtime/CoreUObject/Private/UObject/SavePackage.cpp
@@ -34,6 +34,8 @@ FSavePackageResultStruct UPackage::Save(UPackage* InOuter, UObject* Base, EObjec
 FSavePackageResultStruct UPackage::Save(UPackage* InOuter, UObject* InAsset, const TCHAR* Filename,
 	const FSavePackageArgs& SaveArgs)
 {
+	TRACE_CPUPROFILER_EVENT_SCOPE(UPackage::Save);
+
 	return UPackage::Save2(InOuter, InAsset, Filename, SaveArgs);
 }
 
diff --git a/Engine/Source/Runtime/Engine/Private/DistanceFieldAtlas.cpp b/Engine/Source/Runtime/Engine/Private/DistanceFieldAtlas.cpp
--- a/Engine/Source/Runtime/Engine/Private/DistanceFieldAtlas.cpp
+++ b/Engine/Source/Runtime/Engine/Private/DistanceFieldAtlas.cpp
@@ -151,6 +151,8 @@ FString BuildDistanceFieldDerivedDataKey(const FString& InMeshKey)
 
 void FDistanceFieldVolumeData::CacheDerivedData(const FString& InStaticMeshDerivedDataKey, const ITargetPlatform* TargetPlatform, UStaticMesh* Mesh, FStaticMeshRenderData& RenderData, UStaticMesh* GenerateSource, float DistanceFieldResolutionScale, bool bGenerateDistanceFieldAsIfTwoSided)
 {
+	TRACE_CPUPROFILER_EVENT_SCOPE(FDistanceFieldVolumeData::CacheDerivedData);
+
 	const TArray<FStaticMaterial>& StaticMaterials = Mesh->GetStaticMaterials();
 
 	TArray<FSignedDistanceFieldBuildMaterialData> BuildMaterialData;
diff --git a/Engine/Source/Runtime/Engine/Private/Materials/HLSLMaterialTranslator.cpp b/Engine/Source/Runtime/Engine/Private/Materials/HLSLMaterialTranslator.cpp
--- a/Engine/Source/Runtime/Engine/Private/Materials/HLSLMaterialTranslator.cpp
+++ b/Engine/Source/Runtime/Engine/Private/Materials/HLSLMaterialTranslator.cpp
@@ -670,6 +670,8 @@ void FHLSLMaterialTranslator::ValidateVtPropertyLimits()
  
 bool FHLSLMaterialTranslator::Translate()
 {
+	TRACE_CPUPROFILER_EVENT_SCOPE(FHLSLMaterialTranslator::Translate);
+
 	STAT(double HLSLTranslateTime = 0);
 	{
 		SCOPE_SECONDS_COUNTER(HLSLTranslateTime);
