Revert change to EngineUtils.cpp from CL#20967184 as it's causing runtime issues on Switch.

diff --git a/Engine/Source/Runtime/Engine/Private/EngineUtils.cpp b/Engine/Source/Runtime/Engine/Private/EngineUtils.cpp
--- a/Engine/Source/Runtime/Engine/Private/EngineUtils.cpp
+++ b/Engine/Source/Runtime/Engine/Private/EngineUtils.cpp
@@ -556,8 +556,8 @@ FStripDataFlags::FStripDataFlags( class FArchive& Ar, uint8 InClassFlags /*= 0*/
 		{
 			// When cooking GlobalStripFlags are automatically generated based on the current target
 			// platform's properties.
-			GlobalStripFlags |= Ar.IsFilterEditorOnly() ? static_cast<uint8>(FStripDataFlags::EStrippedData::Editor) : static_cast<uint8>(FStripDataFlags::EStrippedData::None);
-			GlobalStripFlags |= !Ar.CookingTarget()->AllowAudioVisualData() ? static_cast<uint8>(FStripDataFlags::EStrippedData::Server) : static_cast<uint8>(FStripDataFlags::EStrippedData::None);
+			GlobalStripFlags |= Ar.CookingTarget()->HasEditorOnlyData() ? FStripDataFlags::None : FStripDataFlags::Editor;
+			GlobalStripFlags |= !Ar.CookingTarget()->AllowAudioVisualData() ? FStripDataFlags::Server : FStripDataFlags::None;
 			ClassStripFlags = InClassFlags;
 		}
 		Ar << GlobalStripFlags;
@@ -600,8 +600,8 @@ FStripDataFlags::FStripDataFlags(FStructuredArchive::FSlot Slot, uint8 InClassFl
 		{
 			// When cooking GlobalStripFlags are automatically generated based on the current target
 			// platform's properties.
-			GlobalStripFlags |= UnderlyingArchive.IsFilterEditorOnly() ? static_cast<uint8>(FStripDataFlags::EStrippedData::Editor) : static_cast<uint8>(FStripDataFlags::EStrippedData::None);
-			GlobalStripFlags |= !UnderlyingArchive.CookingTarget()->AllowAudioVisualData() ? static_cast<uint8>(FStripDataFlags::EStrippedData::Server) : static_cast<uint8>(FStripDataFlags::EStrippedData::None);
+			GlobalStripFlags |= UnderlyingArchive.CookingTarget()->HasEditorOnlyData() ? FStripDataFlags::None : FStripDataFlags::Editor;
+			GlobalStripFlags |= !UnderlyingArchive.CookingTarget()->AllowAudioVisualData() ? FStripDataFlags::Server : FStripDataFlags::None;
 			ClassStripFlags = InClassFlags;
 		}
 		Record << SA_VALUE(TEXT("GlobalStripFlags"), GlobalStripFlags);
