Properly calculate final texture size for longlat cubmaps and POT-padded textures when determining compression format.

diff --git a/Engine/Source/Runtime/Engine/Private/Texture.cpp b/Engine/Source/Runtime/Engine/Private/Texture.cpp
--- a/Engine/Source/Runtime/Engine/Private/Texture.cpp
+++ b/Engine/Source/Runtime/Engine/Private/Texture.cpp
@@ -2961,11 +2961,32 @@ FName GetDefaultTextureFormatName( const ITargetPlatform* TargetPlatform, const
 	//	 but does not set bNoCompression
 	// bNoCompression does different mappings than bIsTCThatMapsToUncompressed
 
-	if (!bNoCompression && Texture->PowerOfTwoMode == ETexturePowerOfTwoSetting::None && ! Texture->Source.IsLongLatCubemap())
+	if (!bNoCompression)
 	{
 		int32 SizeX = Texture->Source.GetSizeX();
 		int32 SizeY = Texture->Source.GetSizeY();
 
+		if (Texture->PowerOfTwoMode == ETexturePowerOfTwoSetting::PadToPowerOfTwo || Texture->PowerOfTwoMode == ETexturePowerOfTwoSetting::PadToSquarePowerOfTwo)
+		{
+			SizeX = FMath::RoundUpToPowerOfTwo(SizeX);
+			SizeY = FMath::RoundUpToPowerOfTwo(SizeY);
+
+			if (Texture->PowerOfTwoMode == ETexturePowerOfTwoSetting::PadToSquarePowerOfTwo)
+			{
+				SizeX = SizeY = FMath::Max(SizeX, SizeY);
+			}
+		}
+		else
+		{
+			checkf(Texture->PowerOfTwoMode == ETexturePowerOfTwoSetting::None, TEXT("Unknown entry in ETexturePowerOfTwoSetting::Type"));
+		}
+
+		if (Texture->Source.IsLongLatCubemap())
+		{
+			// this should be kept in sync with ComputeLongLatCubemapExtents()
+			SizeX = SizeY = FMath::Max(1 << FMath::FloorLog2(SizeX / 2), 32);
+		}
+
 		//we need to really have the actual top mip size of output platformdata
 		//	(hence the LODBias check below)
 		// trying to reproduce here exactly what TextureCompressor + serialization will do = brittle
