Implement asynchronous compilation for VolumeTexture, Texture2DArray and TextureCubeArray.

diff --git a/Engine/Source/Runtime/Engine/Classes/Engine/Texture2DArray.cpp b/Engine/Source/Runtime/Engine/Classes/Engine/Texture2DArray.cpp
--- a/Engine/Source/Runtime/Engine/Classes/Engine/Texture2DArray.cpp
+++ b/Engine/Source/Runtime/Engine/Classes/Engine/Texture2DArray.cpp
@@ -19,6 +19,10 @@
 #include "Streaming/TextureStreamOut.h"
 #include "Engine/TextureMipDataProviderFactory.h"
 #include "Misc/ScopedSlowTask.h"
+#include "TextureCompiler.h"
+#include "Engine/Public/ImageUtils.h"
+#include "Misc/ScopedSlowTask.h"
+#include "UObject/StrongObjectPtr.h"
 
 #define LOCTEXT_NAMESPACE "UTexture2DArray"
 
@@ -32,7 +36,63 @@ static TAutoConsoleVariable<int32> CVarAllowTexture2DArrayAssetCreation(
 	ECVF_Default
 );
 
-UTexture2DArray::UTexture2DArray(const FObjectInitializer& ObjectInitializer) : Super(ObjectInitializer)
+UTexture2DArray* UTexture2DArray::CreateTransient(int32 InSizeX, int32 InSizeY, int32 InArraySize, EPixelFormat InFormat, const FName InName)
+{
+	UTexture2DArray* NewTexture = nullptr;
+	if (InSizeX > 0 && InSizeY > 0 && InArraySize > 0 &&
+		(InSizeX % GPixelFormats[InFormat].BlockSizeX) == 0 &&
+		(InSizeY % GPixelFormats[InFormat].BlockSizeY) == 0)
+	{
+		NewTexture = NewObject<UTexture2DArray>(
+			GetTransientPackage(),
+			InName,
+			RF_Transient
+			);
+
+		NewTexture->SetPlatformData(new FTexturePlatformData());
+		NewTexture->GetPlatformData()->SizeX = InSizeX;
+		NewTexture->GetPlatformData()->SizeY = InSizeY;
+		NewTexture->GetPlatformData()->SetNumSlices(InArraySize);
+		NewTexture->GetPlatformData()->PixelFormat = InFormat;
+
+		// Allocate first mipmap.
+		int32 NumBlocksX = InSizeX / GPixelFormats[InFormat].BlockSizeX;
+		int32 NumBlocksY = InSizeY / GPixelFormats[InFormat].BlockSizeY;
+		FTexture2DMipMap* Mip = new FTexture2DMipMap();
+		NewTexture->GetPlatformData()->Mips.Add(Mip);
+		Mip->SizeX = InSizeX;
+		Mip->SizeY = InSizeY;
+		Mip->SizeZ = InArraySize;
+		Mip->BulkData.Lock(LOCK_READ_WRITE);
+		Mip->BulkData.Realloc((int64)GPixelFormats[InFormat].BlockBytes * NumBlocksX * NumBlocksY * InArraySize);
+		Mip->BulkData.Unlock();
+	}
+	else
+	{
+		UE_LOG(LogTexture, Warning, TEXT("Invalid parameters specified for UTexture2DArray::CreateTransient()"));
+	}
+	return NewTexture;
+}
+
+/**
+ * Get the optimal placeholder to use during texture compilation
+ */
+static UTexture2DArray* GetDefaultTexture2DArray(const UTexture2DArray* Texture)
+{
+	static TStrongObjectPtr<UTexture2DArray> CheckerboardTexture;
+	if (!CheckerboardTexture.IsValid())
+	{
+		CheckerboardTexture.Reset(FImageUtils::CreateCheckerboardTexture2DArray(FColor(200, 200, 200, 128), FColor(128, 128, 128, 128)));
+	}
+	return CheckerboardTexture.Get();
+}
+
+UTexture2DArray::UTexture2DArray(const FObjectInitializer& ObjectInitializer)
+	: Super(ObjectInitializer)
+	, PrivatePlatformData(nullptr)
+	, PlatformData(
+		[this]()-> FTexturePlatformData* { return GetPlatformData(); },
+		[this](FTexturePlatformData* InPlatformData) { SetPlatformData(InPlatformData); })
 {
 #if WITH_EDITORONLY_DATA
 	SRGB = true;
@@ -40,15 +100,59 @@ UTexture2DArray::UTexture2DArray(const FObjectInitializer& ObjectInitializer) :
 #endif
 }
 
+FTexturePlatformData** UTexture2DArray::GetRunningPlatformData()
+{
+	// @todo DC GetRunningPlatformData is fundamentally unsafe but almost unused... should we replace it with Get/SetRunningPlatformData directly in the base class
+	return &PrivatePlatformData;
+}
+
+void UTexture2DArray::SetPlatformData(FTexturePlatformData* InPlatformData)
+{
+	PrivatePlatformData = InPlatformData;
+}
+
+// Any direct access to GetPlatformData() will stall until the structure
+// is safe to use. It is advisable to replace those use case with
+// async aware code to avoid stalls where possible.
+const FTexturePlatformData* UTexture2DArray::GetPlatformData() const
+{
+#if WITH_EDITOR
+	if (PrivatePlatformData && !PrivatePlatformData->IsAsyncWorkComplete())
+	{
+		TRACE_CPUPROFILER_EVENT_SCOPE(UTexture2DArray::GetPlatformDataStall);
+		UE_LOG(LogTexture, Log, TEXT("Call to GetPlatformData() is forcing a wait on data that is not yet ready."));
+
+		FText Msg = FText::Format(LOCTEXT("WaitOnTextureCompilation", "Waiting on texture compilation {0} ..."), FText::FromString(GetName()));
+		FScopedSlowTask Progress(1.f, Msg, true);
+		Progress.MakeDialog(true);
+		uint64 StartTime = FPlatformTime::Cycles64();
+		PrivatePlatformData->FinishCache();
+		AsyncCompilationHelpers::SaveStallStack(FPlatformTime::Cycles64() - StartTime);
+	}
+#endif // #if WITH_EDITOR
+
+	return PrivatePlatformData;
+}
+
+FTexturePlatformData* UTexture2DArray::GetPlatformData()
+{
+	// For now, this is the same implementation as the const version.
+	const UTexture2DArray* ConstThis = this;
+	return const_cast<FTexturePlatformData*>(ConstThis->GetPlatformData());
+}
+
 #if WITH_EDITOR
 bool UTexture2DArray::GetStreamableRenderResourceState(FTexturePlatformData* InPlatformData, FStreamableRenderResourceState& OutState) const
 {
-	TGuardValue<FTexturePlatformData*> Guard(const_cast<UTexture2DArray*>(this)->PlatformData, InPlatformData);
-	const FPixelFormatInfo& FormatInfo = GPixelFormats[GetPixelFormat()];
-	if (GetNumMips() > 0 && FormatInfo.Supported)
+	TGuardValue<FTexturePlatformData*> Guard(const_cast<UTexture2DArray*>(this)->PrivatePlatformData, InPlatformData);
+	if (GetPlatformData())
 	{
-		OutState = GetResourcePostInitState(PlatformData, GSupportsTexture2DArrayStreaming, 0, 0, /*bSkipCanBeLoaded*/ true);
-		return true;
+		const FPixelFormatInfo& FormatInfo = GPixelFormats[GetPixelFormat()];
+		if (GetNumMips() > 0 && FormatInfo.Supported)
+		{
+			OutState = GetResourcePostInitState(GetPlatformData(), GSupportsTexture2DArrayStreaming, 0, 0, /*bSkipCanBeLoaded*/ true);
+			return true;
+		}
 	}
 	return false;
 }
@@ -56,10 +160,28 @@ bool UTexture2DArray::GetStreamableRenderResourceState(FTexturePlatformData* InP
 
 FTextureResource* UTexture2DArray::CreateResource()
 {
+#if WITH_EDITOR
+	if (PrivatePlatformData)
+	{
+		if (PrivatePlatformData->IsAsyncWorkComplete())
+		{
+			// Make sure AsyncData has been destroyed in case it still exists to avoid
+			// IsDefaultTexture thinking platform data is still being computed.
+			PrivatePlatformData->FinishCache();
+		}
+		else
+		{
+			FTextureCompilingManager::Get().AddTextures({ this });
+			UnlinkStreaming();
+			return new FTexture2DArrayResource(this, (const FTexture2DArrayResource*)GetDefaultTexture2DArray(this)->GetResource());
+		}
+	}
+#endif
+
 	const FPixelFormatInfo& FormatInfo = GPixelFormats[GetPixelFormat()];
 	if (GetNumMips() > 0 && FormatInfo.Supported)
 	{
-		return new FTexture2DArrayResource(this, GetResourcePostInitState(PlatformData, GSupportsTexture2DArrayStreaming));
+		return new FTexture2DArrayResource(this, GetResourcePostInitState(GetPlatformData(), GSupportsTexture2DArrayStreaming));
 	}
 	else if (GetNumMips() == 0)
 	{
@@ -79,7 +201,7 @@ FTextureResource* UTexture2DArray::CreateResource()
 FTexturePlatformData* UTexture2DArray::GetPlatformDataOrWait()
 {
 #if WITH_EDITOR
-	if (PlatformData && !PlatformData->IsAsyncWorkComplete())
+	if (PrivatePlatformData && !PrivatePlatformData->IsAsyncWorkComplete())
 	{
 		TRACE_CPUPROFILER_EVENT_SCOPE(UTexture2D::GetPlatformDataStall);
 		UE_LOG(LogTexture, Log, TEXT("Call to GetPlatformDataOrWait() is forcing a wait on data that is not yet ready."));
@@ -88,19 +210,26 @@ FTexturePlatformData* UTexture2DArray::GetPlatformDataOrWait()
 		FScopedSlowTask Progress(1.f, Msg, true);
 		Progress.MakeDialog(true);
 		uint64 StartTime = FPlatformTime::Cycles64();
-		PlatformData->FinishCache();
+		PrivatePlatformData->FinishCache();
 		AsyncCompilationHelpers::SaveStallStack(FPlatformTime::Cycles64() - StartTime);
 	}
 #endif
 
-	return PlatformData;
+	return PrivatePlatformData;
 }
 
 void UTexture2DArray::UpdateResource()
 {
 #if WITH_EDITOR
 	// Re-cache platform data if the source has changed.
-	CachePlatformData();
+	if (FTextureCompilingManager::Get().IsAsyncCompilationAllowed(this))
+	{
+		BeginCachePlatformData();
+	}
+	else
+	{
+		CachePlatformData();
+	}
 #endif // #if WITH_EDITOR
 
 	Super::UpdateResource();
@@ -108,10 +237,17 @@ void UTexture2DArray::UpdateResource()
 
 uint32 UTexture2DArray::CalcTextureMemorySize(int32 MipCount) const
 {
+#if WITH_EDITOR
+	if (IsDefaultTexture())
+	{
+		return GetDefaultTexture2DArray(this)->CalcTextureMemorySize(MipCount);
+	}
+#endif
+
 	const FPixelFormatInfo& FormatInfo = GPixelFormats[GetPixelFormat()];
 
 	uint32 Size = 0;
-	if (FormatInfo.Supported && PlatformData)
+	if (FormatInfo.Supported && GetPlatformData())
 	{
 		const EPixelFormat Format = GetPixelFormat();
 		if (Format != PF_Unknown)
@@ -123,7 +259,7 @@ uint32 UTexture2DArray::CalcTextureMemorySize(int32 MipCount) const
 			// Must be consistent with the logic in FTexture2DResource::InitRHI
 			const FIntPoint MipExtents = CalcMipMapExtent(GetSizeX(), GetSizeY(), Format, FirstMip);
 			uint32 TextureAlign = 0;
-			Size = (uint32)RHICalcTexture2DArrayPlatformSize(MipExtents.X, MipExtents.Y, GetArraySize(), Format, FMath::Max(1, MipCount), 1, Flags, FRHIResourceCreateInfo(PlatformData->GetExtData()), TextureAlign);
+			Size = (uint32)RHICalcTexture2DArrayPlatformSize(MipExtents.X, MipExtents.Y, GetArraySize(), Format, FMath::Max(1, MipCount), 1, Flags, FRHIResourceCreateInfo(GetPlatformData()->GetExtData()), TextureAlign);
 		}
 	}
 	return Size;
@@ -139,7 +275,88 @@ uint32 UTexture2DArray::CalcTextureMemorySizeEnum(ETextureMipCount Enum) const
 	return CalcTextureMemorySize(GetNumMips());
 }
 
+// While compiling the platform data in editor, we will return the 
+// placeholders value to ensure rendering works as expected and that
+// there are no thread-unsafe access to the platform data being built.
+// Any process requiring a fully up-to-date platform data is expected to
+// call FTextureCompiler:Get().FinishCompilation on UTexture first.
+int32 UTexture2DArray::GetSizeX() const
+{
+	if (PrivatePlatformData)
+	{
 #if WITH_EDITOR
+		if (IsDefaultTexture())
+		{
+			return GetDefaultTexture2DArray(this)->GetSizeX();
+		}
+#endif
+		return PrivatePlatformData->SizeX;
+	}
+	return 0;
+}
+
+int32 UTexture2DArray::GetSizeY() const
+{
+	if (PrivatePlatformData)
+	{
+#if WITH_EDITOR
+		if (IsDefaultTexture())
+		{
+			return GetDefaultTexture2DArray(this)->GetSizeY();
+		}
+#endif
+		return PrivatePlatformData->SizeY;
+	}
+	return 0;
+}
+
+int32 UTexture2DArray::GetArraySize() const
+{
+	if (PrivatePlatformData)
+	{
+#if WITH_EDITOR
+		if (IsDefaultTexture())
+		{
+			return GetDefaultTexture2DArray(this)->GetArraySize();
+		}
+#endif
+		return PrivatePlatformData->GetNumSlices();
+	}
+	return 0;
+}
+
+int32 UTexture2DArray::GetNumMips() const
+{
+	if (PrivatePlatformData)
+	{
+#if WITH_EDITOR
+		if (IsDefaultTexture())
+		{
+			return GetDefaultTexture2DArray(this)->GetNumMips();
+		}
+#endif
+		return PrivatePlatformData->Mips.Num();
+	}
+	return 0;
+}
+
+EPixelFormat UTexture2DArray::GetPixelFormat() const
+{
+	if (PrivatePlatformData)
+	{
+#if WITH_EDITOR
+		if (IsDefaultTexture())
+		{
+			return GetDefaultTexture2DArray(this)->GetPixelFormat();
+		}
+#endif
+		return PrivatePlatformData->PixelFormat;
+	}
+	return PF_Unknown;
+}
+
+#if WITH_EDITOR
+
 ENGINE_API bool UTexture2DArray::CheckArrayTexturesCompatibility()
 {
 	bool bError = false;
@@ -261,14 +478,14 @@ ENGINE_API bool UTexture2DArray::UpdateSourceFromSourceTextures(bool bCreatingNe
 	return true;
 }
 
-ENGINE_API void UTexture2DArray::InvadiateTextureSource()
+ENGINE_API void UTexture2DArray::InvalidateTextureSource()
 {
 	Modify();
 
-	if (PlatformData) 
+	if (GetPlatformData())
 	{
-		delete PlatformData;
-		PlatformData = NULL;
+		delete GetPlatformData();
+		SetPlatformData(nullptr);
 	}
 
 	Source = FTextureSource();
@@ -292,21 +509,23 @@ void UTexture2DArray::Serialize(FArchive& Ar)
 	{
 		SerializeCookedPlatformData(Ar);
 	}
-
-#if WITH_EDITOR
-	if (Ar.IsLoading() && !Ar.IsTransacting() && !bCooked)
-	{
-		BeginCachePlatformData();
-	}
-#endif
 }
 
 void UTexture2DArray::PostLoad() 
 {
 #if WITH_EDITOR
-	FinishCachePlatformData();
+	if (FApp::CanEverRender())
+	{
+		if (FTextureCompilingManager::Get().IsAsyncCompilationAllowed(this))
+		{
+			BeginCachePlatformData();
+		}
+		else
+		{
+			FinishCachePlatformData();
+		}
+	}
 #endif // #if WITH_EDITOR
-
 	Super::PostLoad();
 };
 
@@ -345,6 +564,12 @@ void UTexture2DArray::GetResourceSizeEx(FResourceSizeEx& CumulativeResourceSize)
 }
 
 #if WITH_EDITOR
+
+bool UTexture2DArray::IsDefaultTexture() const
+{
+	return (PrivatePlatformData && !PrivatePlatformData->IsAsyncWorkComplete()) || (GetResource() && GetResource()->IsProxy());
+}
+
 uint32 UTexture2DArray::GetMaximumDimension() const
 {
 	return GetMax2DTextureDimension();
@@ -360,7 +585,7 @@ ENGINE_API void UTexture2DArray::PostEditChangeProperty(FPropertyChangedEvent &
 		// Empty SourceTextures, remove any resources if present.
 		if (SourceTextures.Num() == 0) 
 		{
-			InvadiateTextureSource();
+			InvalidateTextureSource();
 		}
 		// First entry into an empty texture array.
 		else if (SourceTextures.Num() == 1)
diff --git a/Engine/Source/Runtime/Engine/Classes/Engine/Texture2DArray.h b/Engine/Source/Runtime/Engine/Classes/Engine/Texture2DArray.h
--- a/Engine/Source/Runtime/Engine/Classes/Engine/Texture2DArray.h
+++ b/Engine/Source/Runtime/Engine/Classes/Engine/Texture2DArray.h
@@ -16,36 +16,29 @@ class UTexture2DArray : public UTexture
 {
 	GENERATED_UCLASS_BODY()
 
+	/** The derived data for this texture on this platform. */
+	FTexturePlatformData* PrivatePlatformData;
+
 public:
-	/** Platform data. */
-	FTexturePlatformData* PlatformData;
+	UE_DEPRECATED(5.1, "Use GetPlatformData() / SetPlatformData() accessors instead.")
+	TFieldPtrAccessor<FTexturePlatformData> PlatformData;
+
+	/** Set the derived data for this texture on this platform. */
+	ENGINE_API void SetPlatformData(FTexturePlatformData* PlatformData);
+	/** Get the derived data for this texture on this platform. */
+	ENGINE_API FTexturePlatformData* GetPlatformData();
+	/** Get the const derived data for this texture on this platform. */
+	ENGINE_API const FTexturePlatformData* GetPlatformData() const;
+
+#if WITH_EDITOR
 	TMap<FString, FTexturePlatformData*> CookedPlatformData;
+#endif
 
-	/*
-	* Initialize texture source from textures in SourceArray.
-	* @param bUpdateSourceSettings Set to false to prevent overriding current texture settings.
-	*/
-	/** Trivial accessors. */
-	FORCEINLINE int32 GetSizeX() const
-	{
-		return PlatformData ? PlatformData->SizeX : 0;
-	}
-	FORCEINLINE int32 GetSizeY() const
-	{
-		return PlatformData ? PlatformData->SizeY : 0;
-	}
-	FORCEINLINE int32 GetArraySize() const
-	{
-		return PlatformData ? PlatformData->GetNumSlices() : 0;
-	}
-	FORCEINLINE int32 GetNumMips() const
-	{
-		return PlatformData ? PlatformData->Mips.Num() : 0;
-	}
-	FORCEINLINE EPixelFormat GetPixelFormat() const
-	{
-		return PlatformData ? PlatformData->PixelFormat : PF_Unknown;
-	}
+	ENGINE_API int32 GetSizeX() const;
+	ENGINE_API int32 GetSizeY() const;
+	ENGINE_API int32 GetArraySize() const;
+	ENGINE_API int32 GetNumMips() const;
+	ENGINE_API EPixelFormat GetPixelFormat() const;
 
 	/**
 	*  Makes a copy of the mip chain for the texture array starting with InFirstMipToLoad.
@@ -80,13 +73,16 @@ public:
 #if WITH_EDITOR
 	ENGINE_API virtual void PostEditChangeProperty(FPropertyChangedEvent& PropertyChangedEvent) override;
 	ENGINE_API bool UpdateSourceFromSourceTextures(bool bCreatingNewTexture = true);
-	ENGINE_API void InvadiateTextureSource();
+	ENGINE_API void InvalidateTextureSource();
 	ENGINE_API bool CheckArrayTexturesCompatibility();
 #endif // WITH_EDITOR
 	virtual void UpdateResource() override;
 	virtual EMaterialValueType GetMaterialType() const override { return MCT_Texture2DArray; }
-	virtual FTexturePlatformData** GetRunningPlatformData() override { return &PlatformData; }
-	virtual TMap<FString, FTexturePlatformData*> *GetCookedPlatformData() override { return &CookedPlatformData; }
+	virtual FTexturePlatformData** GetRunningPlatformData() override;
+#if WITH_EDITOR
+	virtual bool IsDefaultTexture() const override;
+	virtual TMap<FString, FTexturePlatformData*>* GetCookedPlatformData() override { return &CookedPlatformData; }
+#endif // WITH_EDITOR
 	//~ End UTexture Interface
 
 	/** The addressing mode to use for the X axis.*/
@@ -107,6 +103,9 @@ public:
 	TArray<TObjectPtr<UTexture2D>> SourceTextures;
 #endif
 
+	/** Creates and initializes a new Texture2DArray with the requested settings */
+	ENGINE_API static class UTexture2DArray* CreateTransient(int32 InSizeX, int32 InSizeY, int32 InArraySize, EPixelFormat InFormat = PF_B8G8R8A8, const FName InName = NAME_None);
+
 	/**
 	 * Calculates the size of this texture in bytes if it had MipCount mip-levels streamed in.
 	 *
diff --git a/Engine/Source/Runtime/Engine/Classes/Engine/TextureCubeArray.h b/Engine/Source/Runtime/Engine/Classes/Engine/TextureCubeArray.h
--- a/Engine/Source/Runtime/Engine/Classes/Engine/TextureCubeArray.h
+++ b/Engine/Source/Runtime/Engine/Classes/Engine/TextureCubeArray.h
@@ -12,36 +12,29 @@ class UTextureCubeArray : public UTexture
 {
 	GENERATED_UCLASS_BODY()
 
+	/** The derived data for this texture on this platform. */
+	FTexturePlatformData* PrivatePlatformData;
+
 public:
-	/** Platform data. */
-	FTexturePlatformData* PlatformData;
+	UE_DEPRECATED(5.1, "Use GetPlatformData() / SetPlatformData() accessors instead.")
+	TFieldPtrAccessor<FTexturePlatformData> PlatformData;
+
+	/** Set the derived data for this texture on this platform. */
+	ENGINE_API void SetPlatformData(FTexturePlatformData* PlatformData);
+	/** Get the derived data for this texture on this platform. */
+	ENGINE_API FTexturePlatformData* GetPlatformData();
+	/** Get the const derived data for this texture on this platform. */
+	ENGINE_API const FTexturePlatformData* GetPlatformData() const;
+
+#if WITH_EDITOR
 	TMap<FString, FTexturePlatformData*> CookedPlatformData;
+#endif
 
-	/*
-	* Initialize texture source from textures in SourceArray.
-	* @param bUpdateSourceSettings Set to false to prevent overriding current texture settings.
-	*/
-	/** Trivial accessors. */
-	FORCEINLINE int32 GetSizeX() const
-	{
-		return PlatformData ? PlatformData->SizeX : 0;
-	}
-	FORCEINLINE int32 GetSizeY() const
-	{
-		return PlatformData ? PlatformData->SizeY : 0;
-	}
-	FORCEINLINE int32 GetNumSlices() const
-	{
-		return PlatformData ? PlatformData->GetNumSlices() : 0;
-	}
-	FORCEINLINE int32 GetNumMips() const
-	{
-		return PlatformData ? PlatformData->Mips.Num() : 0;
-	}
-	FORCEINLINE EPixelFormat GetPixelFormat() const
-	{
-		return PlatformData ? PlatformData->PixelFormat : PF_Unknown;
-	}
+	ENGINE_API int32 GetSizeX() const;
+	ENGINE_API int32 GetSizeY() const;
+	ENGINE_API int32 GetNumSlices() const;
+	ENGINE_API int32 GetNumMips() const;
+	ENGINE_API EPixelFormat GetPixelFormat() const;
 
 	//~ Begin UTexture Interface
 	virtual ETextureClass GetTextureClass() const override { return ETextureClass::CubeArray; }
@@ -58,13 +51,16 @@ public:
 #if WITH_EDITOR
 	ENGINE_API virtual void PostEditChangeProperty(FPropertyChangedEvent& PropertyChangedEvent) override;
 	ENGINE_API bool UpdateSourceFromSourceTextures(bool bCreatingNewTexture = true);
-	ENGINE_API void InvadiateTextureSource();
+	ENGINE_API void InvalidateTextureSource();
 	ENGINE_API bool CheckArrayTexturesCompatibility();
 #endif // WITH_EDITOR
 	virtual void UpdateResource() override;
 	virtual EMaterialValueType GetMaterialType() const override { return MCT_TextureCubeArray; }
-	virtual FTexturePlatformData** GetRunningPlatformData() override { return &PlatformData; }
+	virtual FTexturePlatformData** GetRunningPlatformData() override;
+#if WITH_EDITOR
+	virtual bool IsDefaultTexture() const override;
 	virtual TMap<FString, FTexturePlatformData*>* GetCookedPlatformData() override { return &CookedPlatformData; }
+#endif // WITH_EDITOR
 	//~ End UTexture Interface
 
 #if WITH_EDITORONLY_DATA
@@ -73,6 +69,9 @@ public:
 	TArray<TObjectPtr<UTextureCube>> SourceTextures;
 #endif
 
+	/** Creates and initializes a new TextureCubeArray with the requested settings */
+	ENGINE_API static class UTextureCubeArray* CreateTransient(int32 InSizeX, int32 InSizeY, int32 InArraySize, EPixelFormat InFormat = PF_B8G8R8A8, const FName InName = NAME_None);
+
 	/**
 	 * Calculates the size of this texture in bytes if it had MipCount mip-levels streamed in.
 	 *
diff --git a/Engine/Source/Runtime/Engine/Classes/Engine/VolumeTexture.h b/Engine/Source/Runtime/Engine/Classes/Engine/VolumeTexture.h
--- a/Engine/Source/Runtime/Engine/Classes/Engine/VolumeTexture.h
+++ b/Engine/Source/Runtime/Engine/Classes/Engine/VolumeTexture.h
@@ -17,10 +17,23 @@ class UVolumeTexture : public UTexture
 {
 	GENERATED_UCLASS_BODY()
 
+	/** The derived data for this texture on this platform. */
+	FTexturePlatformData* PrivatePlatformData;
+
 public:
-	/** Platform data. */
-	FTexturePlatformData* PlatformData;
+	UE_DEPRECATED(5.1, "Use GetPlatformData() / SetPlatformData() accessors instead.")
+	TFieldPtrAccessor<FTexturePlatformData> PlatformData;
+
+	/** Set the derived data for this texture on this platform. */
+	ENGINE_API void SetPlatformData(FTexturePlatformData* PlatformData);
+	/** Get the derived data for this texture on this platform. */
+	ENGINE_API FTexturePlatformData* GetPlatformData();
+	/** Get the const derived data for this texture on this platform. */
+	ENGINE_API const FTexturePlatformData* GetPlatformData() const;
+
+#if WITH_EDITOR
 	TMap<FString, FTexturePlatformData*> CookedPlatformData;
+#endif
 
 #if WITH_EDITORONLY_DATA
 	/** A (optional) reference texture from which the volume texture was built */
@@ -60,27 +73,11 @@ public:
 	virtual void GetResourceSizeEx(FResourceSizeEx& CumulativeResourceSize) override;
 	//~ End UObject Interface.
 
-	/** Trivial accessors. */
-	FORCEINLINE int32 GetSizeX() const
-	{
-		return PlatformData ? PlatformData->SizeX : 0;
-	}
-	FORCEINLINE int32 GetSizeY() const
-	{
-		return PlatformData ? PlatformData->SizeY : 0;
-	}
-	FORCEINLINE int32 GetSizeZ() const
-	{
-		return PlatformData ? PlatformData->GetNumSlices() : 0;
-	}
-	FORCEINLINE int32 GetNumMips() const
-	{
-		return PlatformData ? PlatformData->Mips.Num() : 0;
-	}
-	FORCEINLINE EPixelFormat GetPixelFormat() const
-	{
-		return PlatformData ? PlatformData->PixelFormat : PF_Unknown;
-	}
+	ENGINE_API int32 GetSizeX() const;
+	ENGINE_API int32 GetSizeY() const;
+	ENGINE_API int32 GetSizeZ() const;
+	ENGINE_API int32 GetNumMips() const;
+	ENGINE_API EPixelFormat GetPixelFormat() const;
 
 	//~ Begin UTexture Interface
 	virtual ETextureClass GetTextureClass() const override { return ETextureClass::Volume; }
@@ -95,10 +92,16 @@ public:
 #endif // WITH_EDITOR
 	virtual void UpdateResource() override;
 	virtual EMaterialValueType GetMaterialType() const override { return MCT_VolumeTexture; }
-	virtual FTexturePlatformData** GetRunningPlatformData() override { return &PlatformData; }
-	virtual TMap<FString, FTexturePlatformData*> *GetCookedPlatformData() override { return &CookedPlatformData; }
+	virtual FTexturePlatformData** GetRunningPlatformData() override;
+#if WITH_EDITOR
+	virtual bool IsDefaultTexture() const override;
+	virtual TMap<FString, FTexturePlatformData*>* GetCookedPlatformData() override { return &CookedPlatformData; }
+#endif // WITH_EDITOR
 	//~ End UTexture Interface
-	
+
+	/** Creates and initializes a new VolumeTexture with the requested settings */
+	ENGINE_API static class UVolumeTexture* CreateTransient(int32 InSizeX, int32 InSizeY, int32 InSizeZ, EPixelFormat InFormat = PF_B8G8R8A8, const FName InName = NAME_None);
+
 	/**
 	 * Calculates the size of this texture in bytes if it had MipCount miplevels streamed in.
 	 *
diff --git a/Engine/Source/Runtime/Engine/Private/ImageUtils.cpp b/Engine/Source/Runtime/Engine/Private/ImageUtils.cpp
--- a/Engine/Source/Runtime/Engine/Private/ImageUtils.cpp
+++ b/Engine/Source/Runtime/Engine/Private/ImageUtils.cpp
@@ -8,6 +8,7 @@ ImageUtils.cpp: Image utility functions.
 
 #include "CubemapUnwrapUtils.h"
 #include "Engine/Texture2D.h"
+#include "Engine/Texture2DArray.h"
 #include "Engine/TextureCube.h"
 #include "Engine/TextureCubeArray.h"
 #include "Engine/VolumeTexture.h"
@@ -807,6 +808,36 @@ UTexture2D* FImageUtils::CreateCheckerboardTexture(FColor ColorOne, FColor Color
 	return CheckerboardTexture;
 }
 
+UTexture2DArray* FImageUtils::CreateCheckerboardTexture2DArray(FColor ColorOne, FColor ColorTwo, int32 CheckerSize, int32 ArraySize)
+{
+	CheckerSize = FMath::Min<uint32>(FMath::RoundUpToPowerOfTwo(CheckerSize), 4096);
+	const int32 HalfCheckerSize = CheckerSize >> 1;
+
+	// Create the texture
+	UTexture2DArray* CheckerboardTexture = UTexture2DArray::CreateTransient(CheckerSize, CheckerSize, ArraySize, PF_B8G8R8A8);
+
+	// Lock the checkerboard texture so it can be modified
+	FColor* MipData = reinterpret_cast<FColor*>(CheckerboardTexture->GetPlatformData()->Mips[0].BulkData.Lock(LOCK_READ_WRITE));
+
+	// Fill in the colors in a checkerboard pattern
+	for (int32 ArrayIndex = 0; ArrayIndex < ArraySize; ArrayIndex++)
+	{
+		for (int32 RowIndex = 0; RowIndex < CheckerSize; RowIndex++)
+		{
+			for (int32 ColumnIndex = 0; ColumnIndex < CheckerSize; ColumnIndex++)
+			{
+				*MipData++ = ColumnIndex < HalfCheckerSize == RowIndex < HalfCheckerSize ? ColorOne : ColorTwo;
+			}
+		}
+	}
+
+	// Unlock the texture
+	CheckerboardTexture->GetPlatformData()->Mips[0].BulkData.Unlock();
+	CheckerboardTexture->UpdateResource();
+
+	return CheckerboardTexture;
+}
+
 UTextureCube* FImageUtils::CreateCheckerboardCubeTexture(FColor ColorOne, FColor ColorTwo, int32 CheckerSize)
 {
 	CheckerSize = FMath::Min<uint32>(FMath::RoundUpToPowerOfTwo(CheckerSize), 4096);
@@ -847,6 +878,69 @@ UTextureCube* FImageUtils::CreateCheckerboardCubeTexture(FColor ColorOne, FColor
 	return CheckerboardTexture;
 }
 
+UTextureCubeArray* FImageUtils::CreateCheckerboardTextureCubeArray(FColor ColorOne, FColor ColorTwo, int32 CheckerSize, int32 ArraySize)
+{
+	CheckerSize = FMath::Min<uint32>(FMath::RoundUpToPowerOfTwo(CheckerSize), 4096);
+	const int32 HalfCheckerSize = CheckerSize >> 1;
+
+	// Create the texture
+	UTextureCubeArray* CheckerboardTexture = UTextureCubeArray::CreateTransient(CheckerSize, CheckerSize, ArraySize, PF_B8G8R8A8);
+
+	// Lock the checkerboard texture so it can be modified
+	FColor* MipData = reinterpret_cast<FColor*>(CheckerboardTexture->GetPlatformData()->Mips[0].BulkData.Lock(LOCK_READ_WRITE));
+
+	// Fill in the colors in a checkerboard pattern
+	for (int32 ArrayIndex = 0; ArrayIndex < ArraySize; ArrayIndex++)
+	{
+		for (int32 FaceIndex = 0; FaceIndex < 6; FaceIndex++)
+		{
+			for (int32 RowIndex = 0; RowIndex < CheckerSize; RowIndex++)
+			{
+				for (int32 ColumnIndex = 0; ColumnIndex < CheckerSize; ColumnIndex++)
+				{
+					*MipData++ = ColumnIndex < HalfCheckerSize == RowIndex < HalfCheckerSize ? ColorOne : ColorTwo;
+				}
+			}
+		}
+	}
+
+	// Unlock the texture
+	CheckerboardTexture->GetPlatformData()->Mips[0].BulkData.Unlock();
+	CheckerboardTexture->UpdateResource();
+
+	return CheckerboardTexture;
+}
+
+UVolumeTexture* FImageUtils::CreateCheckerboardVolumeTexture(FColor ColorOne, FColor ColorTwo, int32 CheckerSize)
+{
+	CheckerSize = FMath::Min<uint32>(FMath::RoundUpToPowerOfTwo(CheckerSize), 4096);
+	const int32 HalfCheckerSize = CheckerSize >> 1;
+
+	// Create the texture
+	UVolumeTexture* CheckerboardTexture = UVolumeTexture::CreateTransient(CheckerSize, CheckerSize, CheckerSize, PF_B8G8R8A8);
+
+	// Lock the checkerboard texture so it can be modified
+	FColor* MipData = reinterpret_cast<FColor*>(CheckerboardTexture->GetPlatformData()->Mips[0].BulkData.Lock(LOCK_READ_WRITE));
+
+	// Fill in the colors in a checkerboard pattern
+	for (int32 SliceIndex = 0; SliceIndex < CheckerSize; SliceIndex++)
+	{
+		for (int32 RowIndex = 0; RowIndex < CheckerSize; RowIndex++)
+		{
+			for (int32 ColumnIndex = 0; ColumnIndex < CheckerSize; ColumnIndex++)
+			{
+				*MipData++ = ColumnIndex < HalfCheckerSize == RowIndex < HalfCheckerSize ? ColorOne : ColorTwo;
+			}
+		}
+	}
+
+	// Unlock the texture
+	CheckerboardTexture->GetPlatformData()->Mips[0].BulkData.Unlock();
+	CheckerboardTexture->UpdateResource();
+
+	return CheckerboardTexture;
+}
+
 /*------------------------------------------------------------------------------
 HDR file format helper.
 ------------------------------------------------------------------------------*/
diff --git a/Engine/Source/Runtime/Engine/Private/Rendering/Texture2DArrayResource.cpp b/Engine/Source/Runtime/Engine/Private/Rendering/Texture2DArrayResource.cpp
--- a/Engine/Source/Runtime/Engine/Private/Rendering/Texture2DArrayResource.cpp
+++ b/Engine/Source/Runtime/Engine/Private/Rendering/Texture2DArrayResource.cpp
@@ -15,7 +15,7 @@
 //*****************************************************************************
 
 FTexture2DArrayResource::FTexture2DArrayResource(UTexture2DArray* InOwner, const FStreamableRenderResourceState& InState) 
-: FStreamableTextureResource(InOwner, InOwner->PlatformData, InState, false)
+: FStreamableTextureResource(InOwner, InOwner->GetPlatformData(), InState, false)
 {
 	AddressU = InOwner->AddressX == TA_Wrap ? AM_Wrap : (InOwner->AddressX == TA_Clamp ? AM_Clamp : AM_Mirror);
 	AddressV = InOwner->AddressY == TA_Wrap ? AM_Wrap : (InOwner->AddressY == TA_Clamp ? AM_Clamp : AM_Mirror);
@@ -42,6 +42,12 @@ FTexture2DArrayResource::FTexture2DArrayResource(UTexture2DArray* InOwner, const
 	}
 }
 
+FTexture2DArrayResource::FTexture2DArrayResource(UTexture2DArray* InOwner, const FTexture2DArrayResource* InProxiedResource)
+	: FStreamableTextureResource(InOwner, InProxiedResource->PlatformData, FStreamableRenderResourceState(), false)
+	, ProxiedResource(InProxiedResource)
+{
+}
+
 void FTexture2DArrayResource::CreateTexture()
 {
 	TArrayView<const FTexture2DMipMap*> MipsViewPostLODBias = GetPlatformMipsView();
@@ -97,6 +103,21 @@ uint64 FTexture2DArrayResource::GetPlatformMipsSize(uint32 NumMips) const
 	}
 }
 
+void FTexture2DArrayResource::InitRHI()
+{
+	if (ProxiedResource)
+	{
+		TextureRHI = ProxiedResource->TextureRHI;
+		RHIUpdateTextureReference(TextureReferenceRHI, TextureRHI);
+		SamplerStateRHI = ProxiedResource->SamplerStateRHI;
+		DeferredPassSamplerStateRHI = ProxiedResource->DeferredPassSamplerStateRHI;
+	}
+	else
+	{
+		FStreamableTextureResource::InitRHI();
+	}
+}
+
 void FTexture2DArrayResource::GetData(int32 BaseRHIMipSizeX, int32 BaseRHIMipSizeY, uint32 ArrayIndex, uint32 MipIndex, void* Dest, uint32 DestPitch) const
 {
 	const uint32 ArrayCount = SizeZ;
diff --git a/Engine/Source/Runtime/Engine/Private/Rendering/Texture2DArrayResource.h b/Engine/Source/Runtime/Engine/Private/Rendering/Texture2DArrayResource.h
--- a/Engine/Source/Runtime/Engine/Private/Rendering/Texture2DArrayResource.h
+++ b/Engine/Source/Runtime/Engine/Private/Rendering/Texture2DArrayResource.h
@@ -20,6 +20,14 @@ public:
 
 	FTexture2DArrayResource(UTexture2DArray* InOwner, const FStreamableRenderResourceState& InState);
 
+	/**
+	 * Make this Texture2DArrayResource Proxy another one.
+	 *
+	 * @param InOwner             UTexture2DArray which this FTexture2DArrayResource represents.
+	 * @param InProxiedResource   The resource to proxy.
+	 */
+	FTexture2DArrayResource(UTexture2DArray* InOwner, const FTexture2DArrayResource* InProxiedResource);
+
 	// Dynamic cast methods.
 	ENGINE_API virtual FTexture2DArrayResource* GetTexture2DArrayResource() { return this; }
 	// Dynamic cast methods (const).
@@ -27,6 +35,9 @@ public:
 
 	virtual uint64 GetPlatformMipsSize(uint32 NumMips) const override;
 
+	virtual void InitRHI() override;
+	virtual bool IsProxy() const override { return ProxiedResource != nullptr; }
+
 protected:
 
 	void CreateTexture() final override;
@@ -34,6 +45,9 @@ protected:
 
 	void GetData(int32 BaseRHIMipSizeX, int32 BaseRHIMipSizeY, uint32 ArrayIndex, uint32 MipIndex, void* Dest, uint32 DestPitch) const;
 
+	/** Another resource being proxied by this one. */
+	const FTexture2DArrayResource* const ProxiedResource = nullptr;
+
 	// Each mip has all array slices. This will be [State.NumRequestedLODs] long, less any packed mips.
 	TArray<FUniqueBuffer, TInlineAllocator<MAX_TEXTURE_MIP_COUNT>> AllMipsData;
 };
diff --git a/Engine/Source/Runtime/Engine/Private/Rendering/Texture3DResource.cpp b/Engine/Source/Runtime/Engine/Private/Rendering/Texture3DResource.cpp
--- a/Engine/Source/Runtime/Engine/Private/Rendering/Texture3DResource.cpp
+++ b/Engine/Source/Runtime/Engine/Private/Rendering/Texture3DResource.cpp
@@ -65,7 +65,7 @@ void FVolumeTextureBulkData::MergeMips(int32 NumMips)
 //*****************************************************************************
 
 FTexture3DResource::FTexture3DResource(UVolumeTexture* InOwner, const FStreamableRenderResourceState& InState)
-: FStreamableTextureResource(InOwner, InOwner->PlatformData, InState, false)
+: FStreamableTextureResource(InOwner, InOwner->GetPlatformData(), InState, false)
 , InitialData(InState.RequestedFirstLODIdx())
 {
 	const int32 FirstLODIdx = InState.RequestedFirstLODIdx();
@@ -107,6 +107,13 @@ FTexture3DResource::FTexture3DResource(UVolumeTexture* InOwner, const FStreamabl
 	}
 }
 
+FTexture3DResource::FTexture3DResource(UVolumeTexture* InOwner, const FTexture3DResource* InProxiedResource)
+	: FStreamableTextureResource(InOwner, InProxiedResource->PlatformData, FStreamableRenderResourceState(), false)
+	, ProxiedResource(InProxiedResource)
+	, InitialData(0)
+{
+}
+
 void FTexture3DResource::CreateTexture()
 {
 	TArrayView<const FTexture2DMipMap*> MipsView = GetPlatformMipsView();
@@ -194,3 +201,18 @@ uint64 FTexture3DResource::GetPlatformMipsSize(uint32 NumMips) const
 		return 0;
 	}
 }
+
+void FTexture3DResource::InitRHI()
+{
+	if (ProxiedResource)
+	{
+		TextureRHI = ProxiedResource->TextureRHI;
+		RHIUpdateTextureReference(TextureReferenceRHI, TextureRHI);
+		SamplerStateRHI = ProxiedResource->SamplerStateRHI;
+		DeferredPassSamplerStateRHI = ProxiedResource->DeferredPassSamplerStateRHI;
+	}
+	else
+	{
+		FStreamableTextureResource::InitRHI();
+	}
+}
diff --git a/Engine/Source/Runtime/Engine/Private/Rendering/Texture3DResource.h b/Engine/Source/Runtime/Engine/Private/Rendering/Texture3DResource.h
--- a/Engine/Source/Runtime/Engine/Private/Rendering/Texture3DResource.h
+++ b/Engine/Source/Runtime/Engine/Private/Rendering/Texture3DResource.h
@@ -67,6 +67,14 @@ public:
 
 	FTexture3DResource(UVolumeTexture* InOwner, const FStreamableRenderResourceState& InState);
 
+	/**
+	 * Make this FTexture3DResource Proxy another one.
+	 *
+	 * @param InOwner             UVolumeTexture which this FTexture3DResource represents.
+	 * @param InProxiedResource   The resource to proxy.
+	 */
+	FTexture3DResource(UVolumeTexture* InOwner, const FTexture3DResource* InProxiedResource);
+
 	// Dynamic cast methods.
 	ENGINE_API virtual FTexture3DResource* GetTexture3DResource() override { return this; }
 	// Dynamic cast methods (const).
@@ -75,11 +83,17 @@ public:
 	/** Returns the platform mip size for the given mip count. */
 	virtual uint64 GetPlatformMipsSize(uint32 NumMips) const override;
 
+	virtual void InitRHI() override;
+	virtual bool IsProxy() const override { return ProxiedResource != nullptr; }
+
 private:
 
 	void CreateTexture() final override;
 	void CreatePartiallyResidentTexture() final override;
 
+	/** Another resource being proxied by this one. */
+	const FTexture3DResource* const ProxiedResource = nullptr;
+
 protected:
 	FVolumeTextureBulkData InitialData;
 
diff --git a/Engine/Source/Runtime/Engine/Private/Texture2D.cpp b/Engine/Source/Runtime/Engine/Private/Texture2D.cpp
--- a/Engine/Source/Runtime/Engine/Private/Texture2D.cpp
+++ b/Engine/Source/Runtime/Engine/Private/Texture2D.cpp
@@ -601,11 +601,7 @@ void UTexture2D::PostLoad()
 	{
 		if (FTextureCompilingManager::Get().IsAsyncCompilationAllowed(this))
 		{
-			// Might already have been triggered during serialization
-			if (!PrivatePlatformData)
-			{
-				BeginCachePlatformData();
-			}
+			BeginCachePlatformData();
 		}
 		else
 		{
@@ -1047,6 +1043,7 @@ UTexture2D* UTexture2D::CreateTransient(int32 InSizeX, int32 InSizeY, EPixelForm
 		NewTexture->SetPlatformData(new FTexturePlatformData());
 		NewTexture->GetPlatformData()->SizeX = InSizeX;
 		NewTexture->GetPlatformData()->SizeY = InSizeY;
+		NewTexture->GetPlatformData()->SetNumSlices(1);
 		NewTexture->GetPlatformData()->PixelFormat = InFormat;
 
 		// Allocate first mipmap.
@@ -1056,8 +1053,9 @@ UTexture2D* UTexture2D::CreateTransient(int32 InSizeX, int32 InSizeY, EPixelForm
 		NewTexture->GetPlatformData()->Mips.Add(Mip);
 		Mip->SizeX = InSizeX;
 		Mip->SizeY = InSizeY;
+		Mip->SizeZ = 1;
 		Mip->BulkData.Lock(LOCK_READ_WRITE);
-		Mip->BulkData.Realloc(NumBlocksX * NumBlocksY * GPixelFormats[InFormat].BlockBytes);
+		Mip->BulkData.Realloc((int64)NumBlocksX * NumBlocksY * GPixelFormats[InFormat].BlockBytes);
 		Mip->BulkData.Unlock();
 	}
 	else
diff --git a/Engine/Source/Runtime/Engine/Private/TextureCube.cpp b/Engine/Source/Runtime/Engine/Private/TextureCube.cpp
--- a/Engine/Source/Runtime/Engine/Private/TextureCube.cpp
+++ b/Engine/Source/Runtime/Engine/Private/TextureCube.cpp
@@ -35,6 +35,8 @@ UTextureCube* UTextureCube::CreateTransient(int32 InSizeX, int32 InSizeY, EPixel
 		NewTexture->SetPlatformData(new FTexturePlatformData());
 		NewTexture->GetPlatformData()->SizeX = InSizeX;
 		NewTexture->GetPlatformData()->SizeY = InSizeY;
+		NewTexture->GetPlatformData()->SetIsCubemap(true);
+		NewTexture->GetPlatformData()->SetNumSlices(6);
 		NewTexture->GetPlatformData()->PixelFormat = InFormat;
 
 		// Allocate first mipmap.
@@ -44,8 +46,9 @@ UTextureCube* UTextureCube::CreateTransient(int32 InSizeX, int32 InSizeY, EPixel
 		NewTexture->GetPlatformData()->Mips.Add(Mip);
 		Mip->SizeX = InSizeX;
 		Mip->SizeY = InSizeY;
+		Mip->SizeZ = 1;
 		Mip->BulkData.Lock(LOCK_READ_WRITE);
-		Mip->BulkData.Realloc(6 * NumBlocksX * NumBlocksY * GPixelFormats[InFormat].BlockBytes);
+		Mip->BulkData.Realloc((int64)6 * NumBlocksX * NumBlocksY * GPixelFormats[InFormat].BlockBytes);
 		Mip->BulkData.Unlock();
 	}
 	else
@@ -155,11 +158,7 @@ void UTextureCube::PostLoad()
 	{
 		if (FTextureCompilingManager::Get().IsAsyncCompilationAllowed(this))
 		{
-			// Might already have been triggered during serialization
-			if (!PrivatePlatformData)
-			{
 				BeginCachePlatformData();
-			}
 		}
 		else
 		{
@@ -217,6 +216,13 @@ FString UTextureCube::GetDesc()
 
 uint32 UTextureCube::CalcTextureMemorySize( int32 MipCount ) const
 {
+#if WITH_EDITOR
+	if (IsDefaultTexture())
+	{
+		return GetDefaultTextureCube(this)->CalcTextureMemorySize(MipCount);
+	}
+#endif
+
 	uint32 Size = 0;
 	if (GetPlatformData())
 	{
@@ -602,11 +608,20 @@ private:
 FTextureResource* UTextureCube::CreateResource()
 {
 #if WITH_EDITOR
-	if (!IsAsyncCacheComplete())
+	if (PrivatePlatformData)
 	{
-		FTextureCompilingManager::Get().AddTextures({ this });
-
-		return new FTextureCubeResource(this, (const FTextureCubeResource*)GetDefaultTextureCube(this)->GetResource());
+		if (PrivatePlatformData->IsAsyncWorkComplete())
+		{
+			// Make sure AsyncData has been destroyed in case it still exists to avoid
+			// IsDefaultTexture thinking platform data is still being computed.
+			PrivatePlatformData->FinishCache();
+		}
+		else
+		{
+			FTextureCompilingManager::Get().AddTextures({ this });
+			UnlinkStreaming();
+			return new FTextureCubeResource(this, (const FTextureCubeResource*)GetDefaultTextureCube(this)->GetResource());
+		}
 	}
 #endif
 
diff --git a/Engine/Source/Runtime/Engine/Private/TextureCubeArray.cpp b/Engine/Source/Runtime/Engine/Private/TextureCubeArray.cpp
--- a/Engine/Source/Runtime/Engine/Private/TextureCubeArray.cpp
+++ b/Engine/Source/Runtime/Engine/Private/TextureCubeArray.cpp
@@ -11,6 +11,12 @@
 #include "DeviceProfiles/DeviceProfile.h"
 #include "DeviceProfiles/DeviceProfileManager.h"
 #include "Containers/ResourceArray.h"
+#include "TextureCompiler.h"
+#include "Engine/Public/ImageUtils.h"
+#include "Misc/ScopedSlowTask.h"
+#include "UObject/StrongObjectPtr.h"
+
+#define LOCTEXT_NAMESPACE "UTextureCubeArray"
 
 class FTextureCubeArrayResource : public FTextureResource
 {
@@ -30,8 +36,8 @@ public:
 		check(Owner);
 		check(Owner->GetNumMips() > 0);
 
-		TIndirectArray<FTexture2DMipMap>& Mips = InOwner->PlatformData->Mips;
-		const int32 FirstMipTailIndex = Mips.Num() - FMath::Max(1, InOwner->PlatformData->GetNumMipsInTail());
+		TIndirectArray<FTexture2DMipMap>& Mips = InOwner->GetPlatformData()->Mips;
+		const int32 FirstMipTailIndex = Mips.Num() - FMath::Max(1, InOwner->GetPlatformData()->GetNumMipsInTail());
 
 		NumSlices = Owner->GetNumSlices();
 		MipData.Empty(NumSlices * (FirstMipTailIndex + 1));
@@ -123,7 +129,7 @@ public:
 			.SetFormat(Owner->GetPixelFormat())
 			.SetNumMips(Owner->GetNumMips())
 			.SetFlags(TexCreateFlags)
-			.SetExtData(Owner->PlatformData->GetExtData());
+			.SetExtData(Owner->GetPlatformData()->GetExtData());
 
 		TextureCubeRHI = RHICreateTexture(Desc);
 		TextureRHI = TextureCubeRHI;
@@ -133,7 +139,7 @@ public:
 		RHIUpdateTextureReference(Owner->TextureReference.TextureReferenceRHI, TextureRHI);
 
 		// Read the mip-levels into the RHI texture.
-		const int32 FirstMipTailIndex = Owner->GetNumMips() - FMath::Max(1, Owner->PlatformData->GetNumMipsInTail());
+		const int32 FirstMipTailIndex = Owner->GetNumMips() - FMath::Max(1, Owner->GetPlatformData()->GetNumMipsInTail());
 		for (int32 MipIndex = 0; MipIndex <= FirstMipTailIndex; MipIndex++)
 		{
 			for (uint32 ArrayIndex = 0u; ArrayIndex < ArraySize; ++ArrayIndex)
@@ -248,7 +254,7 @@ private:
 		// runtime block size checking, conversion, or the like
 		if (DestPitch == 0)
 		{
-			FMemory::Memcpy(Dest, MipData[Index], Owner->PlatformData->Mips[MipIndex].BulkData.GetBulkDataSize() / Owner->GetNumSlices());
+			FMemory::Memcpy(Dest, MipData[Index], Owner->GetPlatformData()->Mips[MipIndex].BulkData.GetBulkDataSize() / Owner->GetNumSlices());
 		}
 		else
 		{
@@ -292,7 +298,64 @@ private:
 	}
 };
 
-UTextureCubeArray::UTextureCubeArray(const FObjectInitializer& ObjectInitializer) : Super(ObjectInitializer)
+UTextureCubeArray* UTextureCubeArray::CreateTransient(int32 InSizeX, int32 InSizeY, int32 InArraySize, EPixelFormat InFormat, const FName InName)
+{
+	UTextureCubeArray* NewTexture = nullptr;
+	if (InSizeX > 0 && InSizeY > 0 && InArraySize > 0 &&
+		(InSizeX % GPixelFormats[InFormat].BlockSizeX) == 0 &&
+		(InSizeY % GPixelFormats[InFormat].BlockSizeY) == 0)
+	{
+		NewTexture = NewObject<UTextureCubeArray>(
+			GetTransientPackage(),
+			InName,
+			RF_Transient
+			);
+
+		NewTexture->SetPlatformData(new FTexturePlatformData());
+		NewTexture->GetPlatformData()->SizeX = InSizeX;
+		NewTexture->GetPlatformData()->SizeY = InSizeY;
+		NewTexture->GetPlatformData()->SetIsCubemap(true);
+		NewTexture->GetPlatformData()->SetNumSlices(InArraySize * 6);
+		NewTexture->GetPlatformData()->PixelFormat = InFormat;
+
+		// Allocate first mipmap.
+		int32 NumBlocksX = InSizeX / GPixelFormats[InFormat].BlockSizeX;
+		int32 NumBlocksY = InSizeY / GPixelFormats[InFormat].BlockSizeY;
+		FTexture2DMipMap* Mip = new FTexture2DMipMap();
+		NewTexture->GetPlatformData()->Mips.Add(Mip);
+		Mip->SizeX = InSizeX;
+		Mip->SizeY = InSizeY;
+		Mip->SizeZ = InArraySize * 6;
+		Mip->BulkData.Lock(LOCK_READ_WRITE);
+		Mip->BulkData.Realloc((int64)GPixelFormats[InFormat].BlockBytes * NumBlocksX * NumBlocksY * 6 * InArraySize);
+		Mip->BulkData.Unlock();
+	}
+	else
+	{
+		UE_LOG(LogTexture, Warning, TEXT("Invalid parameters specified for UTextureCubeArray::CreateTransient()"));
+	}
+	return NewTexture;
+}
+
+/**
+ * Get the optimal placeholder to use during texture compilation
+ */
+static UTextureCubeArray* GetDefaultTextureCubeArray(const UTextureCubeArray* Texture)
+{
+	static TStrongObjectPtr<UTextureCubeArray> CheckerboardTexture;
+	if (!CheckerboardTexture.IsValid())
+	{
+		CheckerboardTexture.Reset(FImageUtils::CreateCheckerboardTextureCubeArray(FColor(200, 200, 200, 128), FColor(128, 128, 128, 128)));
+	}
+	return CheckerboardTexture.Get();
+}
+
+UTextureCubeArray::UTextureCubeArray(const FObjectInitializer& ObjectInitializer)
+	: Super(ObjectInitializer)
+	, PrivatePlatformData(nullptr)
+	, PlatformData(
+		[this]()-> FTexturePlatformData* { return GetPlatformData(); },
+		[this](FTexturePlatformData* InPlatformData) { SetPlatformData(InPlatformData); })
 {
 #if WITH_EDITORONLY_DATA
 	SRGB = true;
@@ -300,8 +363,67 @@ UTextureCubeArray::UTextureCubeArray(const FObjectInitializer& ObjectInitializer
 #endif
 }
 
+FTexturePlatformData** UTextureCubeArray::GetRunningPlatformData()
+{
+	// @todo DC GetRunningPlatformData is fundamentally unsafe but almost unused... should we replace it with Get/SetRunningPlatformData directly in the base class
+	return &PrivatePlatformData;
+}
+
+void UTextureCubeArray::SetPlatformData(FTexturePlatformData* InPlatformData)
+{
+	PrivatePlatformData = InPlatformData;
+}
+
+// Any direct access to GetPlatformData() will stall until the structure
+// is safe to use. It is advisable to replace those use case with
+// async aware code to avoid stalls where possible.
+const FTexturePlatformData* UTextureCubeArray::GetPlatformData() const
+{
+#if WITH_EDITOR
+	if (PrivatePlatformData && !PrivatePlatformData->IsAsyncWorkComplete())
+	{
+		TRACE_CPUPROFILER_EVENT_SCOPE(UTextureCubeArray::GetPlatformDataStall);
+		UE_LOG(LogTexture, Log, TEXT("Call to GetPlatformData() is forcing a wait on data that is not yet ready."));
+
+		FText Msg = FText::Format(LOCTEXT("WaitOnTextureCompilation", "Waiting on texture compilation {0} ..."), FText::FromString(GetName()));
+		FScopedSlowTask Progress(1.f, Msg, true);
+		Progress.MakeDialog(true);
+		uint64 StartTime = FPlatformTime::Cycles64();
+		PrivatePlatformData->FinishCache();
+		AsyncCompilationHelpers::SaveStallStack(FPlatformTime::Cycles64() - StartTime);
+	}
+#endif // #if WITH_EDITOR
+
+	return PrivatePlatformData;
+}
+
+FTexturePlatformData* UTextureCubeArray::GetPlatformData()
+{
+	// For now, this is the same implementation as the const version.
+	const UTextureCubeArray* ConstThis = this;
+	return const_cast<FTexturePlatformData*>(ConstThis->GetPlatformData());
+}
+
 FTextureResource* UTextureCubeArray::CreateResource()
 {
+#if WITH_EDITOR
+	if (PrivatePlatformData)
+	{
+		if (PrivatePlatformData->IsAsyncWorkComplete())
+		{
+			// Make sure AsyncData has been destroyed in case it still exists to avoid
+			// IsDefaultTexture thinking platform data is still being computed.
+			PrivatePlatformData->FinishCache();
+		}
+		else
+		{
+			FTextureCompilingManager::Get().AddTextures({ this });
+			UnlinkStreaming();
+			return new FTextureCubeArrayResource(this, (const FTextureCubeArrayResource*)GetDefaultTextureCubeArray(this)->GetResource());
+		}
+	}
+#endif
+
 	const FPixelFormatInfo& FormatInfo = GPixelFormats[GetPixelFormat()];
 	if (GetNumMips() > 0 && FormatInfo.Supported)
 	{
@@ -322,7 +444,14 @@ void UTextureCubeArray::UpdateResource()
 {
 #if WITH_EDITOR
 	// Re-cache platform data if the source has changed.
-	CachePlatformData();
+	if (FTextureCompilingManager::Get().IsAsyncCompilationAllowed(this))
+	{
+		BeginCachePlatformData();
+	}
+	else
+	{
+		CachePlatformData();
+	}
 #endif // #if WITH_EDITOR
 
 	Super::UpdateResource();
@@ -330,8 +459,15 @@ void UTextureCubeArray::UpdateResource()
 
 uint32 UTextureCubeArray::CalcTextureMemorySize(int32 MipCount) const
 {
+#if WITH_EDITOR
+	if (IsDefaultTexture())
+	{
+		return GetDefaultTextureCubeArray(this)->CalcTextureMemorySize(MipCount);
+	}
+#endif
+
 	uint32 Size = 0;
-	if (PlatformData)
+	if (GetPlatformData())
 	{
 		int32 SizeX = GetSizeX();
 		int32 SizeY = GetSizeY();
@@ -347,7 +483,7 @@ uint32 UTextureCubeArray::CalcTextureMemorySize(int32 MipCount) const
 
 		// TODO add RHICalcTextureCubeArrayPlatformSize
 		uint32 TextureAlign = 0;
-		uint64 TextureSize = RHICalcTextureCubePlatformSize(MipExtents.X, Format, FMath::Max(1, MipCount), TexCreate_None, FRHIResourceCreateInfo(PlatformData->GetExtData()), TextureAlign) * ArraySize;
+		uint64 TextureSize = RHICalcTextureCubePlatformSize(MipExtents.X, Format, FMath::Max(1, MipCount), TexCreate_None, FRHIResourceCreateInfo(GetPlatformData()->GetExtData()), TextureAlign) * ArraySize;
 		Size = (uint32)TextureSize;
 	}
 	return Size;
@@ -363,7 +499,88 @@ uint32 UTextureCubeArray::CalcTextureMemorySizeEnum(ETextureMipCount Enum) const
 	return CalcTextureMemorySize(GetNumMips());
 }
 
+// While compiling the platform data in editor, we will return the 
+// placeholders value to ensure rendering works as expected and that
+// there are no thread-unsafe access to the platform data being built.
+// Any process requiring a fully up-to-date platform data is expected to
+// call FTextureCompiler:Get().FinishCompilation on UTexture first.
+int32 UTextureCubeArray::GetSizeX() const
+{
+	if (PrivatePlatformData)
+	{
 #if WITH_EDITOR
+		if (IsDefaultTexture())
+		{
+			return GetDefaultTextureCubeArray(this)->GetSizeX();
+		}
+#endif
+		return PrivatePlatformData->SizeX;
+	}
+	return 0;
+}
+
+int32 UTextureCubeArray::GetSizeY() const
+{
+	if (PrivatePlatformData)
+	{
+#if WITH_EDITOR
+		if (IsDefaultTexture())
+		{
+			return GetDefaultTextureCubeArray(this)->GetSizeY();
+		}
+#endif
+		return PrivatePlatformData->SizeY;
+	}
+	return 0;
+}
+
+int32 UTextureCubeArray::GetNumSlices() const
+{
+	if (PrivatePlatformData)
+	{
+#if WITH_EDITOR
+		if (IsDefaultTexture())
+		{
+			return GetDefaultTextureCubeArray(this)->GetNumSlices();
+		}
+#endif
+		return PrivatePlatformData->GetNumSlices();
+	}
+	return 0;
+}
+
+int32 UTextureCubeArray::GetNumMips() const
+{
+	if (PrivatePlatformData)
+	{
+#if WITH_EDITOR
+		if (IsDefaultTexture())
+		{
+			return GetDefaultTextureCubeArray(this)->GetNumMips();
+		}
+#endif
+		return PrivatePlatformData->Mips.Num();
+	}
+	return 0;
+}
+
+EPixelFormat UTextureCubeArray::GetPixelFormat() const
+{
+	if (PrivatePlatformData)
+	{
+#if WITH_EDITOR
+		if (IsDefaultTexture())
+		{
+			return GetDefaultTextureCubeArray(this)->GetPixelFormat();
+		}
+#endif
+		return PrivatePlatformData->PixelFormat;
+	}
+	return PF_Unknown;
+}
+
+#if WITH_EDITOR
+
 ENGINE_API bool UTextureCubeArray::CheckArrayTexturesCompatibility()
 {
 	if (SourceTextures.Num() == 0)
@@ -503,14 +720,14 @@ ENGINE_API bool UTextureCubeArray::UpdateSourceFromSourceTextures(bool bCreating
 	return true;
 }
 
-ENGINE_API void UTextureCubeArray::InvadiateTextureSource()
+ENGINE_API void UTextureCubeArray::InvalidateTextureSource()
 {
 	Modify();
 
-	if (PlatformData)
+	if (GetPlatformData())
 	{
-		delete PlatformData;
-		PlatformData = NULL;
+		delete GetPlatformData();
+		SetPlatformData(nullptr);
 	}
 
 	Source = FTextureSource();
@@ -534,21 +751,23 @@ void UTextureCubeArray::Serialize(FArchive& Ar)
 	{
 		SerializeCookedPlatformData(Ar);
 	}
-
-#if WITH_EDITOR
-	if (Ar.IsLoading() && !Ar.IsTransacting() && !bCooked)
-	{
-		BeginCachePlatformData();
-	}
-#endif
 }
 
 void UTextureCubeArray::PostLoad()
 {
 #if WITH_EDITOR
-	FinishCachePlatformData();	
+	if (FApp::CanEverRender())
+	{
+		if (FTextureCompilingManager::Get().IsAsyncCompilationAllowed(this))
+		{
+			BeginCachePlatformData();
+		}
+		else
+		{
+			FinishCachePlatformData();
+		}
+	}
 #endif // #if WITH_EDITOR
-
 	Super::PostLoad();
 };
 
@@ -587,6 +806,12 @@ void UTextureCubeArray::GetResourceSizeEx(FResourceSizeEx& CumulativeResourceSiz
 }
 
 #if WITH_EDITOR
+
+bool UTextureCubeArray::IsDefaultTexture() const
+{
+	return (PrivatePlatformData && !PrivatePlatformData->IsAsyncWorkComplete()) || (GetResource() && GetResource()->IsProxy());
+}
+
 uint32 UTextureCubeArray::GetMaximumDimension() const
 {
 	return GetMaxCubeTextureDimension();
@@ -602,7 +827,7 @@ ENGINE_API void UTextureCubeArray::PostEditChangeProperty(FPropertyChangedEvent&
 		// Empty SourceTextures, remove any resources if present.
 		if (SourceTextures.Num() == 0)
 		{
-			InvadiateTextureSource();
+			InvalidateTextureSource();
 		}
 		// First entry into an empty texture array.
 		else if (SourceTextures.Num() == 1)
@@ -627,3 +852,5 @@ ENGINE_API void UTextureCubeArray::PostEditChangeProperty(FPropertyChangedEvent&
 }
 
 #endif // #if WITH_EDITOR
+
+#undef LOCTEXT_NAMESPACE
\ No newline at end of file
diff --git a/Engine/Source/Runtime/Engine/Private/VolumeTexture.cpp b/Engine/Source/Runtime/Engine/Private/VolumeTexture.cpp
--- a/Engine/Source/Runtime/Engine/Private/VolumeTexture.cpp
+++ b/Engine/Source/Runtime/Engine/Private/VolumeTexture.cpp
@@ -16,6 +16,12 @@
 #include "Streaming/TextureStreamIn.h"
 #include "Streaming/TextureStreamOut.h"
 #include "Engine/TextureMipDataProviderFactory.h"
+#include "TextureCompiler.h"
+#include "Engine/Public/ImageUtils.h"
+#include "Misc/ScopedSlowTask.h"
+#include "UObject/StrongObjectPtr.h"
+
+#define LOCTEXT_NAMESPACE "UVolumeTexture"
 
 //*****************************************************************************
 
@@ -29,12 +35,109 @@ const int32 MAX_VOLUME_TEXTURE_DEPTH = 512;
 //***************************** UVolumeTexture ********************************
 //*****************************************************************************
 
+UVolumeTexture* UVolumeTexture::CreateTransient(int32 InSizeX, int32 InSizeY, int32 InSizeZ, EPixelFormat InFormat, const FName InName)
+{
+	UVolumeTexture* NewTexture = nullptr;
+	if (InSizeX > 0 && InSizeY > 0 && InSizeZ > 0 &&
+		(InSizeX % GPixelFormats[InFormat].BlockSizeX) == 0 &&
+		(InSizeY % GPixelFormats[InFormat].BlockSizeY) == 0)
+	{
+		NewTexture = NewObject<UVolumeTexture>(
+			GetTransientPackage(),
+			InName,
+			RF_Transient
+			);
+
+		NewTexture->SetPlatformData(new FTexturePlatformData());
+		NewTexture->GetPlatformData()->SizeX = InSizeX;
+		NewTexture->GetPlatformData()->SizeY = InSizeY;
+		NewTexture->GetPlatformData()->SetNumSlices(InSizeZ);
+		NewTexture->GetPlatformData()->PixelFormat = InFormat;
+
+		// Allocate first mipmap.
+		int32 NumBlocksX = InSizeX / GPixelFormats[InFormat].BlockSizeX;
+		int32 NumBlocksY = InSizeY / GPixelFormats[InFormat].BlockSizeY;
+		int32 NumBlocksZ = InSizeZ / GPixelFormats[InFormat].BlockSizeZ;
+		FTexture2DMipMap* Mip = new FTexture2DMipMap();
+		NewTexture->GetPlatformData()->Mips.Add(Mip);
+		Mip->SizeX = InSizeX;
+		Mip->SizeY = InSizeY;
+		Mip->SizeZ = InSizeZ;
+		Mip->BulkData.Lock(LOCK_READ_WRITE);
+		Mip->BulkData.Realloc((int64)GPixelFormats[InFormat].BlockBytes * NumBlocksX * NumBlocksY * NumBlocksZ);
+		Mip->BulkData.Unlock();
+	}
+	else
+	{
+		UE_LOG(LogTexture, Warning, TEXT("Invalid parameters specified for UVolumeTexture::CreateTransient()"));
+	}
+	return NewTexture;
+}
+
+/**
+ * Get the optimal placeholder to use during texture compilation
+ */
+static UVolumeTexture* GetDefaultVolumeTexture(const UVolumeTexture* Texture)
+{
+	static TStrongObjectPtr<UVolumeTexture> CheckerboardTexture;
+	if (!CheckerboardTexture.IsValid())
+	{
+		CheckerboardTexture.Reset(FImageUtils::CreateCheckerboardVolumeTexture(FColor(200, 200, 200, 128), FColor(128, 128, 128, 128)));
+	}
+	return CheckerboardTexture.Get();
+}
+
 UVolumeTexture::UVolumeTexture(const FObjectInitializer& ObjectInitializer)
 	: Super(ObjectInitializer)
+	, PrivatePlatformData(nullptr)
+	, PlatformData(
+		[this]()-> FTexturePlatformData* { return GetPlatformData(); },
+		[this](FTexturePlatformData* InPlatformData) { SetPlatformData(InPlatformData); })
 {
 	SRGB = true;
 }
 
+FTexturePlatformData** UVolumeTexture::GetRunningPlatformData()
+{
+	// @todo DC GetRunningPlatformData is fundamentally unsafe but almost unused... should we replace it with Get/SetRunningPlatformData directly in the base class
+	return &PrivatePlatformData;
+}
+
+void UVolumeTexture::SetPlatformData(FTexturePlatformData* InPlatformData)
+{
+	PrivatePlatformData = InPlatformData;
+}
+
+// Any direct access to GetPlatformData() will stall until the structure
+// is safe to use. It is advisable to replace those use case with
+// async aware code to avoid stalls where possible.
+const FTexturePlatformData* UVolumeTexture::GetPlatformData() const
+{
+#if WITH_EDITOR
+	if (PrivatePlatformData && !PrivatePlatformData->IsAsyncWorkComplete())
+	{
+		TRACE_CPUPROFILER_EVENT_SCOPE(UVolumeTexture::GetPlatformDataStall);
+		UE_LOG(LogTexture, Log, TEXT("Call to GetPlatformData() is forcing a wait on data that is not yet ready."));
+
+		FText Msg = FText::Format(LOCTEXT("WaitOnTextureCompilation", "Waiting on texture compilation {0} ..."), FText::FromString(GetName()));
+		FScopedSlowTask Progress(1.f, Msg, true);
+		Progress.MakeDialog(true);
+		uint64 StartTime = FPlatformTime::Cycles64();
+		PrivatePlatformData->FinishCache();
+		AsyncCompilationHelpers::SaveStallStack(FPlatformTime::Cycles64() - StartTime);
+	}
+#endif // #if WITH_EDITOR
+	
+	return PrivatePlatformData;
+}
+
+FTexturePlatformData* UVolumeTexture::GetPlatformData()
+{
+	// For now, this is the same implementation as the const version.
+	const UVolumeTexture* ConstThis = this;
+	return const_cast<FTexturePlatformData*>(ConstThis->GetPlatformData());
+}
+
 bool UVolumeTexture::UpdateSourceFromSourceTexture()
 {
 	bool bSourceValid = false;
@@ -98,9 +201,9 @@ bool UVolumeTexture::UpdateSourceFromSourceTexture()
 		Source.SetOwner(this);
 		// Source2DTexture = nullptr; // ??
 
-		if (PlatformData)
+		if (GetPlatformData())
 		{
-			*PlatformData = FTexturePlatformData();
+			SetPlatformData(new FTexturePlatformData());
 		}
 	}
 	
@@ -185,19 +288,23 @@ void UVolumeTexture::Serialize(FArchive& Ar)
 	{
 		SerializeCookedPlatformData(Ar);
 	}
-
-#if WITH_EDITOR
-	if (Ar.IsLoading() && !Ar.IsTransacting() && !bCooked)
-	{
-		BeginCachePlatformData();
-	}
-#endif // #if WITH_EDITOR
 }
 
 void UVolumeTexture::PostLoad()
 {
 #if WITH_EDITOR
-	FinishCachePlatformData();
+
+	if (FApp::CanEverRender())
+	{
+		if (FTextureCompilingManager::Get().IsAsyncCompilationAllowed(this))
+		{
+			BeginCachePlatformData();
+		}
+		else
+		{
+			FinishCachePlatformData();
+		}
+	}
 #endif // #if WITH_EDITOR
 
 	Super::PostLoad();
@@ -226,7 +333,14 @@ void UVolumeTexture::UpdateResource()
 {
 #if WITH_EDITOR
 	// Recache platform data if the source has changed.
-	CachePlatformData();
+	if (FTextureCompilingManager::Get().IsAsyncCompilationAllowed(this))
+	{
+		BeginCachePlatformData();
+	}
+	else
+	{
+		CachePlatformData();
+	}
 #endif // #if WITH_EDITOR
 
 	// Route to super.
@@ -245,10 +359,17 @@ FString UVolumeTexture::GetDesc()
 
 uint32 UVolumeTexture::CalcTextureMemorySize(int32 MipCount) const
 {
+#if WITH_EDITOR
+	if (IsDefaultTexture())
+	{
+		return GetDefaultVolumeTexture(this)->CalcTextureMemorySize(MipCount);
+	}
+#endif
+
 	const FPixelFormatInfo& FormatInfo = GPixelFormats[GetPixelFormat()];
 
 	uint32 Size = 0;
-	if (FormatInfo.Supported && PlatformData)
+	if (FormatInfo.Supported && GetPlatformData())
 	{
 		const EPixelFormat Format = GetPixelFormat();
 		if (Format != PF_Unknown)
@@ -261,7 +382,7 @@ uint32 UVolumeTexture::CalcTextureMemorySize(int32 MipCount) const
 			CalcMipMapExtent3D(GetSizeX(), GetSizeY(), GetSizeZ(), Format, FMath::Max<int32>(0, GetNumMips() - MipCount), SizeX, SizeY, SizeZ);
 
 			uint32 TextureAlign = 0;
-			Size = (uint32)RHICalcTexture3DPlatformSize(SizeX, SizeY, SizeZ, Format, FMath::Max(1, MipCount), Flags, FRHIResourceCreateInfo(PlatformData->GetExtData()), TextureAlign);
+			Size = (uint32)RHICalcTexture3DPlatformSize(SizeX, SizeY, SizeZ, Format, FMath::Max(1, MipCount), Flags, FRHIResourceCreateInfo(GetPlatformData()->GetExtData()), TextureAlign);
 		}
 	}
 	return Size;
@@ -279,16 +400,99 @@ uint32 UVolumeTexture::CalcTextureMemorySizeEnum( ETextureMipCount Enum ) const
 	}
 }
 
+// While compiling the platform data in editor, we will return the 
+// placeholders value to ensure rendering works as expected and that
+// there are no thread-unsafe access to the platform data being built.
+// Any process requiring a fully up-to-date platform data is expected to
+// call FTextureCompiler:Get().FinishCompilation on UTexture first.
+int32 UVolumeTexture::GetSizeX() const
+{
+	if (PrivatePlatformData)
+	{
+#if WITH_EDITOR
+		if (IsDefaultTexture())
+		{
+			return GetDefaultVolumeTexture(this)->GetSizeX();
+		}
+#endif
+		return PrivatePlatformData->SizeX;
+	}
+	return 0;
+}
+
+int32 UVolumeTexture::GetSizeY() const
+{
+	if (PrivatePlatformData)
+	{
+#if WITH_EDITOR
+		if (IsDefaultTexture())
+		{
+			return GetDefaultVolumeTexture(this)->GetSizeY();
+		}
+#endif
+		return PrivatePlatformData->SizeY;
+	}
+	return 0;
+}
+
+int32 UVolumeTexture::GetSizeZ() const
+{
+	if (PrivatePlatformData)
+	{
+#if WITH_EDITOR
+		if (IsDefaultTexture())
+		{
+			return GetDefaultVolumeTexture(this)->GetSizeZ();
+		}
+#endif
+		return PrivatePlatformData->GetNumSlices();
+	}
+	return 0;
+}
+
+int32 UVolumeTexture::GetNumMips() const
+{
+	if (PrivatePlatformData)
+	{
+#if WITH_EDITOR
+		if (IsDefaultTexture())
+		{
+			return GetDefaultVolumeTexture(this)->GetNumMips();
+		}
+#endif
+		return PrivatePlatformData->Mips.Num();
+	}
+	return 0;
+}
+
+EPixelFormat UVolumeTexture::GetPixelFormat() const
+{
+	if (PrivatePlatformData)
+	{
+#if WITH_EDITOR
+		if (IsDefaultTexture())
+		{
+			return GetDefaultVolumeTexture(this)->GetPixelFormat();
+		}
+#endif
+		return PrivatePlatformData->PixelFormat;
+	}
+	return PF_Unknown;
+}
+
 #if WITH_EDITOR
 bool UVolumeTexture::GetStreamableRenderResourceState(FTexturePlatformData* InPlatformData, FStreamableRenderResourceState& OutState) const
 {
-	TGuardValue<FTexturePlatformData*> Guard(const_cast<UVolumeTexture*>(this)->PlatformData, InPlatformData);
-	const FPixelFormatInfo& FormatInfo = GPixelFormats[GetPixelFormat()];
-	const bool bFormatIsSupported = FormatInfo.Supported;
-	if (GetNumMips() > 0 && GSupportsTexture3D && bFormatIsSupported)
+	TGuardValue<FTexturePlatformData*> Guard(const_cast<UVolumeTexture*>(this)->PrivatePlatformData, InPlatformData);
+	if (GetPlatformData())
 	{
-		OutState = GetResourcePostInitState(PlatformData, GSupportsVolumeTextureStreaming, 0, 0, /*bSkipCanBeLoaded*/ true);
-		return true;
+		const FPixelFormatInfo& FormatInfo = GPixelFormats[GetPixelFormat()];
+		const bool bFormatIsSupported = FormatInfo.Supported;
+		if (GetNumMips() > 0 && GSupportsTexture3D && bFormatIsSupported)
+		{
+			OutState = GetResourcePostInitState(GetPlatformData(), GSupportsVolumeTextureStreaming, 0, 0, /*bSkipCanBeLoaded*/ true);
+			return true;
+		}
 	}
 	return false;
 }
@@ -296,12 +500,30 @@ bool UVolumeTexture::GetStreamableRenderResourceState(FTexturePlatformData* InPl
 
 FTextureResource* UVolumeTexture::CreateResource()
 {
+#if WITH_EDITOR
+	if (PrivatePlatformData)
+	{
+		if (PrivatePlatformData->IsAsyncWorkComplete())
+		{
+			// Make sure AsyncData has been destroyed in case it still exists to avoid
+			// IsDefaultTexture thinking platform data is still being computed.
+			PrivatePlatformData->FinishCache();
+		}
+		else
+		{
+			FTextureCompilingManager::Get().AddTextures({ this });
+			UnlinkStreaming();
+			return new FTexture3DResource(this, (const FTexture3DResource*)GetDefaultVolumeTexture(this)->GetResource());
+		}
+	}
+#endif
+
 	const FPixelFormatInfo& FormatInfo = GPixelFormats[GetPixelFormat()];
 	const bool bFormatIsSupported = FormatInfo.Supported;
 
 	if (GetNumMips() > 0 && GSupportsTexture3D && bFormatIsSupported)
 	{
-		return new FTexture3DResource(this, GetResourcePostInitState(PlatformData, GSupportsVolumeTextureStreaming));
+		return new FTexture3DResource(this, GetResourcePostInitState(GetPlatformData(), GSupportsVolumeTextureStreaming));
 	}
 	else if (GetNumMips() == 0)
 	{
@@ -326,6 +548,11 @@ void UVolumeTexture::GetResourceSizeEx(FResourceSizeEx& CumulativeResourceSize)
 
 #if WITH_EDITOR
 
+bool UVolumeTexture::IsDefaultTexture() const
+{
+	return (PrivatePlatformData && !PrivatePlatformData->IsAsyncWorkComplete()) || (GetResource() && GetResource()->IsProxy());
+}
+
 void UVolumeTexture::SetDefaultSource2DTileSize()
 {
 	Source2DTileSizeX = 0;
@@ -439,3 +666,4 @@ bool UVolumeTexture::StreamIn(int32 NewMipCount, bool bHighPrio)
 	return false;
 }
 
+#undef LOCTEXT_NAMESPACE
\ No newline at end of file
diff --git a/Engine/Source/Runtime/Engine/Public/ImageUtils.h b/Engine/Source/Runtime/Engine/Public/ImageUtils.h
--- a/Engine/Source/Runtime/Engine/Public/ImageUtils.h
+++ b/Engine/Source/Runtime/Engine/Public/ImageUtils.h
@@ -12,7 +12,10 @@ ImageUtils.h: Image/Textures utility functions
 #include "Engine/Texture.h"
 
 class UTexture2D;
+class UTexture2DArray;
 class UTextureCube;
+class UTextureCubeArray;
+class UVolumeTexture;
 class UTextureRenderTarget2D;
 class UTextureRenderTargetCube;
 class IImageWrapperModule;
@@ -266,21 +269,53 @@ public:
 	 *
 	 * @param ColorOne		The color of half of the squares.
 	 * @param ColorTwo		The color of the other half of the squares.
-	 * @param CheckerSize	The size in pixels of each checker square.
+	 * @param CheckerSize	The size in pixels of each side of the texture.
 	 *
 	 */
 	ENGINE_API static UTexture2D* CreateCheckerboardTexture(FColor ColorOne = FColor(64, 64, 64), FColor ColorTwo = FColor(128, 128, 128), int32 CheckerSize = 32);
 
+	/**
+	 * Creates a new UTexture2DArray with a checkerboard pattern.
+	 *
+	 * @param ColorOne		The color of half of the squares.
+	 * @param ColorTwo		The color of the other half of the squares.
+	 * @param CheckerSize	The size in pixels of each side of the texture.
+	 * @param ArraySize		The number of elements in the array.
+	 *
+	 */
+	ENGINE_API static UTexture2DArray* CreateCheckerboardTexture2DArray(FColor ColorOne = FColor(64, 64, 64), FColor ColorTwo = FColor(128, 128, 128), int32 CheckerSize = 32, int32 ArraySize = 1);
+
 	/**
 	 * Creates a new UTextureCube with a checkerboard pattern.
 	 *
 	 * @param ColorOne		The color of half of the squares.
 	 * @param ColorTwo		The color of the other half of the squares.
-	 * @param CheckerSize	The size in pixels of each checker square.
+	 * @param CheckerSize	The size in pixels of each side of the texture.
 	 *
 	 */
 	ENGINE_API static UTextureCube* CreateCheckerboardCubeTexture(FColor ColorOne = FColor(64, 64, 64), FColor ColorTwo = FColor(128, 128, 128), int32 CheckerSize = 32);
 
+	/**
+	 * Creates a new UTextureCubeArray with a checkerboard pattern.
+	 *
+	 * @param ColorOne		The color of half of the squares.
+	 * @param ColorTwo		The color of the other half of the squares.
+	 * @param CheckerSize	The size in pixels of each side of the texture.
+	 * @param ArraySize		The number of elements in the array.
+	 *
+	 */
+	ENGINE_API static UTextureCubeArray* CreateCheckerboardTextureCubeArray(FColor ColorOne = FColor(64, 64, 64), FColor ColorTwo = FColor(128, 128, 128), int32 CheckerSize = 32, int32 ArraySize = 1);
+
+	/**
+	 * Creates a new UVolumeTexture with a checkerboard pattern.
+	 *
+	 * @param ColorOne		The color of half of the squares.
+	 * @param ColorTwo		The color of the other half of the squares.
+	 * @param CheckerSize	The size in pixels of each side of the texture.
+	 *
+	 */
+	ENGINE_API static UVolumeTexture* CreateCheckerboardVolumeTexture(FColor ColorOne = FColor(64, 64, 64), FColor ColorTwo = FColor(128, 128, 128), int32 CheckerSize = 16);
+
 	/**
 	 * Exports a UTextureRenderTarget2D as an HDR image on the disk.
 	 * HDR is very low quality, prefer EXR or PNG instead
