Fix conflict in localization keys.

diff --git a/Engine/Source/Runtime/Engine/Classes/Engine/Texture2DArray.cpp b/Engine/Source/Runtime/Engine/Classes/Engine/Texture2DArray.cpp
--- a/Engine/Source/Runtime/Engine/Classes/Engine/Texture2DArray.cpp
+++ b/Engine/Source/Runtime/Engine/Classes/Engine/Texture2DArray.cpp
@@ -122,7 +122,7 @@ const FTexturePlatformData* UTexture2DArray::GetPlatformData() const
 		TRACE_CPUPROFILER_EVENT_SCOPE(UTexture2DArray::GetPlatformDataStall);
 		UE_LOG(LogTexture, Log, TEXT("Call to GetPlatformData() is forcing a wait on data that is not yet ready."));
 
-		FText Msg = FText::Format(LOCTEXT("WaitOnTextureCompilation", "Waiting on texture compilation {0} ..."), FText::FromString(GetName()));
+		FText Msg = FText::Format(LOCTEXT("WaitOnTextureCompilation2", "Waiting on texture compilation {0} ..."), FText::FromString(GetName()));
 		FScopedSlowTask Progress(1.f, Msg, true);
 		Progress.MakeDialog(true);
 		uint64 StartTime = FPlatformTime::Cycles64();
