Added external morph targets to skinned meshes, which are GPU compressed morph targets that are managed externally by other systems and do not get serialized as part of the SkinnedMesh itself.

diff --git a/Engine/Source/Runtime/Engine/Classes/Components/SkinnedMeshComponent.h b/Engine/Source/Runtime/Engine/Classes/Components/SkinnedMeshComponent.h
--- a/Engine/Source/Runtime/Engine/Classes/Components/SkinnedMeshComponent.h
+++ b/Engine/Source/Runtime/Engine/Classes/Components/SkinnedMeshComponent.h
@@ -139,6 +139,44 @@ struct FActiveMorphTarget
 	}
 };
 
+/** The weight data for a specific external morph set. */
+struct ENGINE_API FExternalMorphSetWeights
+{
+	/** Update the number of active morph targets. */
+	void UpdateNumActiveMorphTargets();
+
+	/** The debug name. */
+	FName Name = FName(TEXT("Unknown"));
+
+	/** The weights, which can also be negative and go beyond 1.0 or -1.0. */
+	TArray<float> Weights;
+
+	/** The number of active morph targets. */
+	int32 NumActiveMorphTargets = 0;
+
+	/** The treshold used to determine if a morph target is active or not. Any weight equal to or above this value is seen as active morph target. */
+	float ActiveWeightThreshold = 0.001f;
+};
+
+/** The morph target weight data for all external morph target sets. */
+struct ENGINE_API FExternalMorphWeightData
+{
+	/** Update the number of active morph targets for all sets. */
+	void UpdateNumActiveMorphTargets();
+
+	/** Reset the morph target sets. */
+	void Reset() { MorphSets.Reset(); NumActiveMorphTargets = 0; }
+
+	/** Check if we have active morph targets or not. */
+	bool HasActiveMorphs() const { return (NumActiveMorphTargets > 0); }
+
+	/** The map with a collection of morph sets. Each set can contains multiple morph targets. */
+	TMap<int32, FExternalMorphSetWeights> MorphSets;
+
+	/** The number of active morph targets. */
+	int32 NumActiveMorphTargets = 0;
+};
+
 /** Vertex skin weight info supplied for a component override. */
 USTRUCT(BlueprintType, meta = (HasNativeMake = "/Script/Engine.KismetRenderingLibrary.MakeSkinWeightInfo", HasNativeBreak = "/Script/Engine.KismetRenderingLibrary.BreakSkinWeightInfo"))
 struct FSkelMeshSkinWeightInfo
@@ -373,6 +411,9 @@ public:
 
 	const TArray<int32>& GetMasterBoneMap() const { return MasterBoneMap; }
 
+	const FExternalMorphWeightData& GetExternalMorphWeights(int32 LOD) const { return ExternalMorphWeightData[LOD]; }
+	FExternalMorphWeightData& GetExternalMorphWeights(int32 LOD) { return ExternalMorphWeightData[LOD]; }
+
 	/** 
 	 * Get CPU skinned vertices for the specified LOD level. Includes morph targets if they are enabled.
 	 * Note: This function is very SLOW as it needs to flush the render thread.
@@ -403,6 +444,9 @@ public:
 	/** Array of weights for all morph targets. This array is updated inside RefreshBoneTransforms based on the Anim Blueprint. */
 	TArray<float> MorphTargetWeights;
 
+	/** The external morph target set weight data, for each LOD. This data is (re)initialized by RefreshExternalMorphTargetWeights(). */
+	TArray<FExternalMorphWeightData> ExternalMorphWeightData;
+
 #if WITH_EDITORONLY_DATA
 private:
 	/** Index of the section to preview... If set to -1, all section will be rendered */
@@ -917,6 +961,9 @@ public:
 	/** Get the pre-skinning local space bounds for this component. */
 	void GetPreSkinnedLocalBounds(FBoxSphereBounds& OutBounds) const;
 
+	/** Refresh the external morph target weight buffers. This makes sure the amount of morph sets and number of weights are valid. */
+	void RefreshExternalMorphTargetWeights();
+
 	/**
 	 *	Sets the value of the bForceWireframe flag and reattaches the component as necessary.
 	 *
diff --git a/Engine/Source/Runtime/Engine/Private/Components/SkinnedMeshComponent.cpp b/Engine/Source/Runtime/Engine/Private/Components/SkinnedMeshComponent.cpp
--- a/Engine/Source/Runtime/Engine/Private/Components/SkinnedMeshComponent.cpp
+++ b/Engine/Source/Runtime/Engine/Private/Components/SkinnedMeshComponent.cpp
@@ -330,6 +330,32 @@ namespace FAnimUpdateRateManager
 
 //////////////////////////////////////////////////////////////////////////
 
+void FExternalMorphSetWeights::UpdateNumActiveMorphTargets()
+{
+	NumActiveMorphTargets = 0;
+	for (int32 Index = 0; Index < Weights.Num(); ++Index)
+	{
+		if (FMath::Abs<float>(Weights[Index]) >= ActiveWeightThreshold)
+		{
+			NumActiveMorphTargets++;
+		}
+	}
+}
+
+//////////////////////////////////////////////////////////////////////////
+
+void FExternalMorphWeightData::UpdateNumActiveMorphTargets()
+{
+	NumActiveMorphTargets = 0;
+	for (auto& MapItem : MorphSets)
+	{
+		MapItem.Value.UpdateNumActiveMorphTargets();
+		NumActiveMorphTargets += MapItem.Value.NumActiveMorphTargets;
+	}
+}
+
+//////////////////////////////////////////////////////////////////////////
+
 USkinnedMeshComponent::USkinnedMeshComponent(const FObjectInitializer& ObjectInitializer)
 	: Super(ObjectInitializer)
 	, MeshObjectFactory(nullptr)
@@ -524,6 +550,8 @@ void USkinnedMeshComponent::OnRegister()
 	InvalidateCachedBounds();
 
 	MeshDeformerInstance = (MeshDeformer != nullptr) ? MeshDeformer->CreateInstance(this) : nullptr;
+
+	RefreshExternalMorphTargetWeights();
 }
 
 void USkinnedMeshComponent::OnUnregister()
@@ -540,6 +568,49 @@ void USkinnedMeshComponent::OnUnregister()
 	MeshDeformerInstance = nullptr;
 }
 
+void USkinnedMeshComponent::RefreshExternalMorphTargetWeights()
+{
+	// Clear the external weights if there is no skeletal mesh.
+	if (SkeletalMesh.Get() == nullptr || SkeletalMesh->GetResourceForRendering() == nullptr)
+	{
+		ExternalMorphWeightData.Empty();
+		ExternalMorphWeightData.AddDefaulted(GetNumLODs());
+		return;
+	}
+
+	// Init based on the highest LOD.
+	const int32 NumLODs = GetNumLODs();
+	ExternalMorphWeightData.Reset();
+	ExternalMorphWeightData.AddDefaulted(NumLODs);
+	check(SkeletalMesh->GetResourceForRendering()->LODRenderData.Num() == NumLODs);
+
+	for (int32 LOD = 0; LOD < NumLODs; ++LOD)
+	{
+		FExternalMorphWeightData& ExternalWeights = GetExternalMorphWeights(LOD);
+
+		auto& MorphSets = SkeletalMesh->GetResourceForRendering()->LODRenderData[LOD].ExternalMorphSets;
+		for (const auto& Item : MorphSets)
+		{
+			check(Item.Value.IsValid());
+			const int32 MorphSetID = Item.Key;
+			const FMorphTargetVertexInfoBuffers& MorphBuffer = Item.Value->MorphBuffers;
+
+			FExternalMorphSetWeights* MorphSetWeights = ExternalWeights.MorphSets.Find(MorphSetID);
+			if (MorphSetWeights == nullptr)
+			{
+				MorphSetWeights = &ExternalWeights.MorphSets.Add(MorphSetID);
+			}
+			check(MorphSetWeights != nullptr);
+
+			MorphSetWeights->Name = Item.Value->Name;
+			MorphSetWeights->Weights.Reset();
+			MorphSetWeights->Weights.AddZeroed(MorphBuffer.GetNumMorphs());
+		}
+
+		ExternalWeights.UpdateNumActiveMorphTargets();
+	}
+}
+
 void USkinnedMeshComponent::CreateRenderState_Concurrent(FRegisterComponentContext* Context)
 {
 	LLM_SCOPE(ELLMTag::SkeletalMesh);
@@ -657,7 +728,7 @@ void USkinnedMeshComponent::CreateRenderState_Concurrent(FRegisterComponentConte
 					ActiveMorphTargets.Empty();
 				}
 
-				MeshObject->Update(ModifiedLODLevel, this, ActiveMorphTargets, MorphTargetWeights, EPreviousBoneTransformUpdateMode::UpdatePrevious);  // send to rendering thread
+				MeshObject->Update(ModifiedLODLevel, this, ActiveMorphTargets, MorphTargetWeights, EPreviousBoneTransformUpdateMode::UpdatePrevious, GetExternalMorphWeights(ModifiedLODLevel));  // send to rendering thread
 			}
 		}
 
@@ -749,7 +820,15 @@ void USkinnedMeshComponent::SendRenderDynamicData_Concurrent()
 				ActiveMorphTargets.Empty();
 			}
 
-			MeshObject->Update(UseLOD, this, ActiveMorphTargets, MorphTargetWeights, bExternalEvaluationRateLimited && !bExternalInterpolate ? EPreviousBoneTransformUpdateMode::DuplicateCurrentToPrevious : EPreviousBoneTransformUpdateMode::None);  // send to rendering thread
+			MeshObject->Update(
+				UseLOD,
+				this,
+				ActiveMorphTargets,
+				MorphTargetWeights,
+				bExternalEvaluationRateLimited && !bExternalInterpolate ? EPreviousBoneTransformUpdateMode::DuplicateCurrentToPrevious : EPreviousBoneTransformUpdateMode::None,
+				GetExternalMorphWeights(UseLOD)
+			);  // send to rendering thread
+
 			MeshObject->bHasBeenUpdatedAtLeastOnce = true;
 			bForceMeshObjectUpdate = false;
 
@@ -789,10 +868,10 @@ void USkinnedMeshComponent::ClearMotionVector()
 		// use this function to clear it
 		// this function updates renderer twice using increasing of revision number, so that renderer updates previous/new transform correctly
 		++CurrentBoneTransformRevisionNumber;
-		MeshObject->Update(UseLOD, this, ActiveMorphTargets, MorphTargetWeights, EPreviousBoneTransformUpdateMode::None);  // send to rendering thread
+		MeshObject->Update(UseLOD, this, ActiveMorphTargets, MorphTargetWeights, EPreviousBoneTransformUpdateMode::None, GetExternalMorphWeights(UseLOD));  // send to rendering thread
 
 		++CurrentBoneTransformRevisionNumber;
-		MeshObject->Update(UseLOD, this, ActiveMorphTargets, MorphTargetWeights, EPreviousBoneTransformUpdateMode::None);  // send to rendering thread
+		MeshObject->Update(UseLOD, this, ActiveMorphTargets, MorphTargetWeights, EPreviousBoneTransformUpdateMode::None, GetExternalMorphWeights(UseLOD));  // send to rendering thread
 
 		// Skin cache may need updating
 		MarkForNeededEndOfFrameUpdate();
@@ -811,7 +890,7 @@ void USkinnedMeshComponent::ForceMotionVector()
 		}
 
 		++CurrentBoneTransformRevisionNumber;
-		MeshObject->Update(UseLOD, this, ActiveMorphTargets, MorphTargetWeights, EPreviousBoneTransformUpdateMode::None);
+		MeshObject->Update(UseLOD, this, ActiveMorphTargets, MorphTargetWeights, EPreviousBoneTransformUpdateMode::None, GetExternalMorphWeights(UseLOD));
 
 		// Skin cache may need updating
 		MarkForNeededEndOfFrameUpdate();
@@ -1704,8 +1783,12 @@ void USkinnedMeshComponent::SetSkeletalMesh(USkeletalMesh* InSkelMesh, bool bRei
 			// clear morphtarget cache
 			ActiveMorphTargets.Empty();			
 			MorphTargetWeights.Empty();
+			ExternalMorphWeightData.Empty();
 		}
 	}
+
+	// Update external weight array sizes.
+	RefreshExternalMorphTargetWeights();
 	
 	if (IsRegistered())
 	{
diff --git a/Engine/Source/Runtime/Engine/Private/Rendering/MorphTargetVertexInfoBuffers.cpp b/Engine/Source/Runtime/Engine/Private/Rendering/MorphTargetVertexInfoBuffers.cpp
--- a/Engine/Source/Runtime/Engine/Private/Rendering/MorphTargetVertexInfoBuffers.cpp
+++ b/Engine/Source/Runtime/Engine/Private/Rendering/MorphTargetVertexInfoBuffers.cpp
@@ -23,8 +23,12 @@ void FMorphTargetVertexInfoBuffers::InitRHI()
 	RHIUnlockBuffer(MorphDataBuffer);
 	MorphDataSRV = RHICreateShaderResourceView(MorphDataBuffer);
 
-	MorphData.Empty();
-	bIsMorphCPUDataValid = false;
+	if (bEmptyMorphCPUDataOnInitRHI)
+	{
+		MorphData.Empty();
+		bIsMorphCPUDataValid = false;
+	}
+
 	bRHIIntialized = true;
 }
 
diff --git a/Engine/Source/Runtime/Engine/Private/SkeletalMeshLODRenderData.cpp b/Engine/Source/Runtime/Engine/Private/SkeletalMeshLODRenderData.cpp
--- a/Engine/Source/Runtime/Engine/Private/SkeletalMeshLODRenderData.cpp
+++ b/Engine/Source/Runtime/Engine/Private/SkeletalMeshLODRenderData.cpp
@@ -205,7 +205,7 @@ void FSkeletalMeshLODRenderData::InitResources(bool bNeedsVertexColors, int32 LO
 		}
 	}
 
-	// Always make sure the morph target resources are in an initialized state. This could should not get hit since the data should come with the load.
+	// Always make sure the morph target resources are in an initialized state. This should not get hit since the data should come with the load.
 	bool bNeedsMorphTargetRenderData = Owner && Owner->GetMorphTargets().Num() > 0 && !MorphTargetVertexInfoBuffers.IsMorphResourcesInitialized();
 	if (bNeedsMorphTargetRenderData)
 	{
@@ -255,7 +255,7 @@ void FSkeletalMeshLODRenderData::ReleaseResources()
 		}
 	}
 	BeginReleaseResource(&MorphTargetVertexInfoBuffers);
-	
+
 	DEC_DWORD_STAT_BY(STAT_SkeletalMeshVertexMemory, SkinWeightProfilesData.GetResourcesSize());
 	SkinWeightProfilesData.ReleaseResources();
 
@@ -961,3 +961,23 @@ void FSkeletalMeshLODRenderData::GetSectionFromVertexIndex(int32 InVertIndex, in
 	// InVertIndex should always be in some chunk!
 	//check(false);
 }
+
+void FSkeletalMeshLODRenderData::AddExternalMorphSet(int32 ID, TSharedPtr<FExternalMorphTargetSet> MorphSet)
+{
+	ExternalMorphSets.Add(ID, MorphSet);
+}
+
+void FSkeletalMeshLODRenderData::RemoveExternalMorphSet(int32 ID)
+{
+	ExternalMorphSets.Remove(ID);
+}
+
+bool FSkeletalMeshLODRenderData::HasExternalMorphSet(int32 ID) const
+{
+	return (ExternalMorphSets.Find(ID) != nullptr);
+}
+
+void FSkeletalMeshLODRenderData::ClearExternalMorphSets()
+{
+	ExternalMorphSets.Empty();
+}
diff --git a/Engine/Source/Runtime/Engine/Private/SkeletalRenderCPUSkin.cpp b/Engine/Source/Runtime/Engine/Private/SkeletalRenderCPUSkin.cpp
--- a/Engine/Source/Runtime/Engine/Private/SkeletalRenderCPUSkin.cpp
+++ b/Engine/Source/Runtime/Engine/Private/SkeletalRenderCPUSkin.cpp
@@ -128,7 +128,13 @@ void FSkeletalMeshObjectCPUSkin::EnableOverlayRendering(bool bEnabled, const TAr
 	}
 }
 
-void FSkeletalMeshObjectCPUSkin::Update(int32 LODIndex,USkinnedMeshComponent* InMeshComponent,const TArray<FActiveMorphTarget>& ActiveMorphTargets, const TArray<float>& MorphTargetWeights, EPreviousBoneTransformUpdateMode PreviousBoneTransformUpdateMode)
+void FSkeletalMeshObjectCPUSkin::Update(
+	int32 LODIndex,
+	USkinnedMeshComponent* InMeshComponent,
+	const TArray<FActiveMorphTarget>& ActiveMorphTargets,
+	const TArray<float>& MorphTargetWeights,
+	EPreviousBoneTransformUpdateMode PreviousBoneTransformUpdateMode,
+	const FExternalMorphWeightData& InExternalMorphWeightData)
 {
 	if (InMeshComponent)
 	{
diff --git a/Engine/Source/Runtime/Engine/Private/SkeletalRenderCPUSkin.h b/Engine/Source/Runtime/Engine/Private/SkeletalRenderCPUSkin.h
--- a/Engine/Source/Runtime/Engine/Private/SkeletalRenderCPUSkin.h
+++ b/Engine/Source/Runtime/Engine/Private/SkeletalRenderCPUSkin.h
@@ -98,7 +98,7 @@ public:
 	//~ Begin FSkeletalMeshObject Interface
 	virtual void InitResources(USkinnedMeshComponent* InMeshComponent) override;
 	virtual void ReleaseResources() override;
-	virtual void Update(int32 LODIndex,USkinnedMeshComponent* InMeshComponent,const TArray<FActiveMorphTarget>& ActiveMorphTargets, const TArray<float>& MorphTargetsWeights, EPreviousBoneTransformUpdateMode PreviousBoneTransformUpdateMode) override;
+	virtual void Update(int32 LODIndex,USkinnedMeshComponent* InMeshComponent,const TArray<FActiveMorphTarget>& ActiveMorphTargets, const TArray<float>& MorphTargetsWeights, EPreviousBoneTransformUpdateMode PreviousBoneTransformUpdateMode, const FExternalMorphWeightData& InExternalMorphWeightData) override;
 	void UpdateDynamicData_RenderThread(FRHICommandListImmediate& RHICmdList, FDynamicSkelMeshObjectDataCPUSkin* InDynamicData, uint32 FrameNumberToPrepare, uint32 RevisionNumber);
 	virtual void EnableOverlayRendering(bool bEnabled, const TArray<int32>* InBonesOfInterest, const TArray<UMorphTarget*>* InMorphTargetOfInterest) override;
 	virtual void CacheVertices(int32 LODIndex, bool bForce) const override;
diff --git a/Engine/Source/Runtime/Engine/Private/SkeletalRenderGPUSkin.cpp b/Engine/Source/Runtime/Engine/Private/SkeletalRenderGPUSkin.cpp
--- a/Engine/Source/Runtime/Engine/Private/SkeletalRenderGPUSkin.cpp
+++ b/Engine/Source/Runtime/Engine/Private/SkeletalRenderGPUSkin.cpp
@@ -403,12 +403,18 @@
 	bMorphResourcesInitialized = false;
 }
 
-void FSkeletalMeshObjectGPUSkin::Update(int32 LODIndex,USkinnedMeshComponent* InMeshComponent,const TArray<FActiveMorphTarget>& ActiveMorphTargets, const TArray<float>& MorphTargetWeights, EPreviousBoneTransformUpdateMode PreviousBoneTransformUpdateMode)
+void FSkeletalMeshObjectGPUSkin::Update(
+	int32 LODIndex,
+	USkinnedMeshComponent* InMeshComponent,
+	const TArray<FActiveMorphTarget>& ActiveMorphTargets,
+	const TArray<float>& MorphTargetWeights,
+	EPreviousBoneTransformUpdateMode PreviousBoneTransformUpdateMode,
+	const FExternalMorphWeightData& InExternalMorphWeightData)
 {
 	LLM_SCOPE(ELLMTag::SkeletalMesh);
 
 	// make sure morph data has been initialized for each LOD
-	if(InMeshComponent && !bMorphResourcesInitialized && ActiveMorphTargets.Num() > 0 )
+	if(InMeshComponent && !bMorphResourcesInitialized && (ActiveMorphTargets.Num() > 0 || !InExternalMorphWeightData.MorphSets.IsEmpty()))
 	{
 		// initialized on-the-fly in order to avoid creating extra vertex streams for each skel mesh instance
 		InitMorphResources(InMeshComponent->bPerBoneMotionBlur, MorphTargetWeights);		
@@ -417,7 +423,15 @@
 	// create the new dynamic data for use by the rendering thread
 	// this data is only deleted when another update is sent
 	FDynamicSkelMeshObjectDataGPUSkin* NewDynamicData = FDynamicSkelMeshObjectDataGPUSkin::AllocDynamicSkelMeshObjectDataGPUSkin();		
-	NewDynamicData->InitDynamicSkelMeshObjectDataGPUSkin(InMeshComponent,SkeletalMeshRenderData,this,LODIndex,ActiveMorphTargets, MorphTargetWeights, PreviousBoneTransformUpdateMode);
+	NewDynamicData->InitDynamicSkelMeshObjectDataGPUSkin(
+		InMeshComponent,
+		SkeletalMeshRenderData,
+		this,
+		LODIndex,
+		ActiveMorphTargets,
+		MorphTargetWeights,
+		PreviousBoneTransformUpdateMode,
+		InExternalMorphWeightData);
 
 	// We prepare the next frame but still have the value from the last one
 	uint32 FrameNumberToPrepare = GFrameNumber + 1;
@@ -513,6 +527,7 @@
 	bMorphNeedsUpdate = 
 		bAlwaysUpdateMorphVertexBuffer ||
 		(bMorphNeedsUpdateDeferred && bNeedsUpdateDeferred) || // the need for an update sticks
+		InDynamicData->ExternalMorphWeightData.HasActiveMorphs() ||
 		(DynamicData ? (DynamicData->LODIndex != InDynamicData->LODIndex ||
 		!DynamicData->ActiveMorphTargetsEqual(InDynamicData->ActiveMorphTargets, InDynamicData->MorphTargetWeights))
 		: true);
@@ -596,7 +611,10 @@
 	const TArray<FSkelMeshRenderSection>& Sections = GetRenderSections(LODIndex);
 
 	// Only consider morphs with active curves and data to deform.
-	const bool bMorph = bAlwaysUpdateMorphVertexBuffer || (DynamicData->NumWeightedActiveMorphTargets > 0 && LODData.GetNumVertices() > 0);
+	const bool bMorph = 
+		bAlwaysUpdateMorphVertexBuffer || 
+		(DynamicData->NumWeightedActiveMorphTargets > 0 && LODData.GetNumVertices() > 0) || 
+		(DynamicData->ExternalMorphWeightData.HasActiveMorphs() && LODData.GetNumVertices() > 0);
 
 	// use correct vertex factories based on alternate weights usage
 	FVertexFactoryData& VertexFactoryData = LOD.GPUSkinVertexFactories;
@@ -612,7 +630,6 @@
 		FGPUSkinCache::Release(SkinCacheEntry);
 	}
 
-
 	if (MorphVertexBuffer.bNeedsInitialClear && !(bMorph && bMorphNeedsUpdate))
 	{
 		QUICK_SCOPE_CYCLE_COUNTER(STAT_FSkeletalMeshObjectGPUSkin_ProcessUpdatedDynamicData_ClearMorphBuffer);
@@ -661,7 +678,9 @@
 				}
 				else
 				{
-					if(bAlwaysUpdateMorphVertexBuffer || (DynamicData->NumWeightedActiveMorphTargets > 0 && DynamicData->SectionIdsUseByActiveMorphTargets.Contains(SectionIdx)))
+					if (bAlwaysUpdateMorphVertexBuffer || 
+						(DynamicData->NumWeightedActiveMorphTargets > 0 && DynamicData->SectionIdsUseByActiveMorphTargets.Contains(SectionIdx)) || 
+						!DynamicData->ExternalMorphWeightData.MorphSets.IsEmpty())
 					{
 						VertexFactory = VertexFactoryData.MorphVertexFactories[SectionIdx].Get();
 					}
@@ -899,15 +918,180 @@
 
 #endif // RHI_RAYTRACING
 
+
+int32 FSkeletalMeshObjectGPUSkin::CalcNumActiveGPUMorphSets(FMorphVertexBuffer& MorphVertexBuffer, const FSkeletalMeshLODRenderData& LODData) const
+{
+	if (!UseGPUMorphTargets(FeatureLevel) || !IsValidRef(MorphVertexBuffer.VertexBufferRHI))
+	{
+		return 0;
+	}
+
+	// Count all active external morph sets.
+	int32 NumMorphSets = 1; // Start at one, as we have our standard morph targets as well.
+	for (const auto& MorphSet : LODData.ExternalMorphSets)
+	{
+		if (IsExternalMorphSetActive(MorphSet.Key, *MorphSet.Value))
+		{
+			NumMorphSets++;
+		}
+	}
+	return NumMorphSets;
+}
+
+bool FSkeletalMeshObjectGPUSkin::IsExternalMorphSetActive(int32 MorphSetID, FExternalMorphTargetSet& MorphSet) const
+{
+	const FMorphTargetVertexInfoBuffers& CompressedBuffers = MorphSet.MorphBuffers;
+	FExternalMorphSetWeights* WeightData = DynamicData->ExternalMorphWeightData.MorphSets.Find(MorphSetID);
+	return (WeightData &&
+			WeightData->Weights.Num() == CompressedBuffers.GetNumMorphs() &&
+			WeightData->NumActiveMorphTargets > 0);
+}
+
+static void CalculateMorphDeltaBoundsAccum(
+	const TArray<float>& MorphTargetWeights,
+	const FMorphTargetVertexInfoBuffers& MorphTargetVertexInfoBuffers,
+	FVector4& MinAccumScale,
+	FVector4& MaxAccumScale,
+	FVector4& MaxScale)
+{
+	for (uint32 i = 0; i < MorphTargetVertexInfoBuffers.GetNumMorphs(); i++)
+	{
+		FVector4f MinMorphScale = MorphTargetVertexInfoBuffers.GetMinimumMorphScale(i);
+		FVector4f MaxMorphScale = MorphTargetVertexInfoBuffers.GetMaximumMorphScale(i);
+
+		for (uint32 j = 0; j < 4; j++)
+		{
+			if (MorphTargetWeights.IsValidIndex(i))
+			{
+				MinAccumScale[j] += MorphTargetWeights[i] * MinMorphScale[j];
+				MaxAccumScale[j] += MorphTargetWeights[i] * MaxMorphScale[j];
+			}
+
+			double AbsMorphScale = FMath::Max<double>(FMath::Abs(MinMorphScale[j]), FMath::Abs(MaxMorphScale[j]));
+			double AbsAccumScale = FMath::Max<double>(FMath::Abs(MinAccumScale[j]), FMath::Abs(MaxAccumScale[j]));
+
+			// The maximum accumulated and the maximum local value have to fit into out int24.
+			MaxScale[j] = FMath::Max(MaxScale[j], FMath::Max(AbsMorphScale, AbsAccumScale));
+		}
+	}
+}
+
+static void CalculateMorphDeltaBoundsIncludingExternalMorphs(
+	const TArray<float>& MorphTargetWeights,
+	const FMorphTargetVertexInfoBuffers& MorphTargetVertexInfoBuffers,
+	const TMap<int32, TSharedPtr<FExternalMorphTargetSet>>& ExternalMorphSets,
+	const TMap<int32, FExternalMorphSetWeights>& ExternalWeights,
+	FVector4& MorphScale,
+	FVector4& InvMorphScale)
+{
+	FVector4 MinAccumScale(0, 0, 0, 0);
+	FVector4 MaxAccumScale(0, 0, 0, 0);
+	FVector4 MaxScale(0, 0, 0, 0);
+
+	// Include the standard morph targets.
+	CalculateMorphDeltaBoundsAccum(MorphTargetWeights, MorphTargetVertexInfoBuffers, MinAccumScale, MaxAccumScale, MaxScale);
+
+	// Include all external morph targets.
+	for (const auto& MorphSet : ExternalMorphSets)
+	{
+		const int32 MorphSetID = MorphSet.Key;
+		const FMorphTargetVertexInfoBuffers& CompressedBuffers = MorphSet.Value->MorphBuffers;
+		const FExternalMorphSetWeights* WeightData = ExternalWeights.Find(MorphSetID);
+		check(WeightData);
+		CalculateMorphDeltaBoundsAccum(WeightData->Weights, CompressedBuffers, MinAccumScale, MaxAccumScale, MaxScale);
+	}
+
+	MaxScale[0] = FMath::Max<double>(MaxScale[0], 1.0);
+	MaxScale[1] = FMath::Max<double>(MaxScale[1], 1.0);
+	MaxScale[2] = FMath::Max<double>(MaxScale[2], 1.0);
+	MaxScale[3] = FMath::Max<double>(MaxScale[3], 1.0);
+
+	const double ScaleToInt24 = 16777216.0;
+
+	MorphScale = FVector4
+	(
+		ScaleToInt24 / (MaxScale[0]),
+		ScaleToInt24 / (MaxScale[1]),
+		ScaleToInt24 / (MaxScale[2]),
+		ScaleToInt24 / (MaxScale[3])
+	);
+
+	InvMorphScale = FVector4
+	(
+		MaxScale[0] / ScaleToInt24,
+		MaxScale[1] / ScaleToInt24,
+		MaxScale[2] / ScaleToInt24,
+		MaxScale[3] / ScaleToInt24
+	);
+}
+
 void FSkeletalMeshObjectGPUSkin::UpdateMorphVertexBuffer(FRHICommandListImmediate& RHICmdList, EGPUSkinCacheEntryMode Mode, FSkeletalMeshObjectLOD& LOD, const FSkeletalMeshLODRenderData& LODData, 
 															bool bGPUSkinCacheEnabled, FMorphVertexBuffer& MorphVertexBuffer)
 {
 	QUICK_SCOPE_CYCLE_COUNTER(STAT_FSkeletalMeshObjectGPUSkin_ProcessUpdatedDynamicData_UpdateMorphBuffer);
 	if (UseGPUMorphTargets(FeatureLevel))
 	{
-		// sometimes this goes out of bound, we'll ensure here
+		const int32 NumMorphSets = CalcNumActiveGPUMorphSets(MorphVertexBuffer, LODData);
+		int32 MorphSetIndex = 0;
+
+		// Calculate the delta bounds.
+		FVector4 MorphScale;
+		FVector4 InvMorphScale;
+		{
+			SCOPE_CYCLE_COUNTER(STAT_MorphVertexBuffer_ApplyDelta);
+			CalculateMorphDeltaBoundsIncludingExternalMorphs(
+				DynamicData->MorphTargetWeights,
+				LODData.MorphTargetVertexInfoBuffers,
+				LODData.ExternalMorphSets,
+				DynamicData->ExternalMorphWeightData.MorphSets,
+				MorphScale,
+				InvMorphScale);
+		}
+
+		// Sometimes this goes out of bound, we'll ensure here.
 		ensureAlways(DynamicData->MorphTargetWeights.Num() == LODData.MorphTargetVertexInfoBuffers.GetNumMorphs());
-		LOD.UpdateMorphVertexBufferGPU(RHICmdList, DynamicData->MorphTargetWeights, LODData.MorphTargetVertexInfoBuffers, DynamicData->SectionIdsUseByActiveMorphTargets, GetDebugName(), Mode, MorphVertexBuffer);
+		LOD.UpdateMorphVertexBufferGPU(
+			RHICmdList, 
+			DynamicData->MorphTargetWeights, 
+			LODData.MorphTargetVertexInfoBuffers, 
+			DynamicData->SectionIdsUseByActiveMorphTargets, 
+			GetDebugName(), 
+			Mode, 
+			MorphVertexBuffer, 
+			true,	// Only clear the morph vertex buffer at the first morph set.
+			(MorphSetIndex == NumMorphSets-1),
+			MorphScale,
+			InvMorphScale);	// Normalize only after the last morph set.
+
+		MorphSetIndex++;
+
+		// Process all external morph targets.
+		for (const auto& MorphSet : LODData.ExternalMorphSets)
+		{
+			const int32 MorphSetID = MorphSet.Key;
+			const FMorphTargetVertexInfoBuffers& CompressedBuffers = MorphSet.Value->MorphBuffers;
+			FExternalMorphSetWeights* WeightData = DynamicData->ExternalMorphWeightData.MorphSets.Find(MorphSetID);
+			check(WeightData);
+			if (IsExternalMorphSetActive(MorphSetID, *MorphSet.Value))
+			{
+				LOD.UpdateMorphVertexBufferGPU(
+					RHICmdList, WeightData->Weights,
+					CompressedBuffers,
+					TArray<int32>(),
+					GetDebugName(),
+					Mode,
+					MorphVertexBuffer,
+					false,	// Don't clear the vertex buffer as we already did with the standard morph targets above.
+					(MorphSetIndex == NumMorphSets - 1),
+					MorphScale,
+					InvMorphScale);	// Normalize only after the last morph set.
+
+				MorphSetIndex++;
+			}
+		}
+
+		// If this hits, the CalcNumActiveGPUMorphSets most likely returns the wrong number.
+		check(NumMorphSets == MorphSetIndex);
 	}
 	else
 	{
@@ -951,7 +1135,6 @@
 	SetShaderValue(RHICmdList, CS, MorphTargetWeightsParameter,		*(float(*)[MorphTargetDispatchBatchSize]) Weights);
 }
 
-
 void FGPUMorphUpdateCS::Dispatch(FRHICommandList& RHICmdList, uint32 Size)
 {
 	const FIntVector DispatchSize = FComputeShaderUtils::GetGroupCountWrapped(Size);
@@ -987,60 +1170,18 @@
 
 IMPLEMENT_SHADER_TYPE(, FGPUMorphNormalizeCS, TEXT("/Engine/Private/MorphTargets.usf"), TEXT("GPUMorphNormalizeCS"), SF_Compute);
 
-static void CalculateMorphDeltaBounds(const TArray<float>& MorphTargetWeights, const FMorphTargetVertexInfoBuffers& MorphTargetVertexInfoBuffers, FVector4& MorphScale, FVector4& InvMorphScale)
+void FSkeletalMeshObjectGPUSkin::FSkeletalMeshObjectLOD::UpdateMorphVertexBufferGPU(
+	FRHICommandListImmediate& RHICmdList, 
+	const TArray<float>& MorphTargetWeights,
+	const FMorphTargetVertexInfoBuffers& MorphTargetVertexInfoBuffers,
+	const TArray<int32>& SectionIdsUseByActiveMorphTargets,
+	const FName& OwnerName,
+	EGPUSkinCacheEntryMode Mode,
+	FMorphVertexBuffer& MorphVertexBuffer,
+	bool bClearMorphVertexBuffer,
+	bool bNormalizePass,
+	const FVector4& MorphScale,
+	const FVector4& InvMorphScale)
 {
-	double MinAccumScale[4] = { 0, 0, 0, 0 };
-	double MaxAccumScale[4] = { 0, 0, 0, 0 };
-	double MaxScale[4] = { 0, 0, 0, 0 };
-
-	for (uint32 i = 0; i < MorphTargetVertexInfoBuffers.GetNumMorphs(); i++)
-	{
-		FVector4f MinMorphScale = MorphTargetVertexInfoBuffers.GetMinimumMorphScale(i);
-		FVector4f MaxMorphScale = MorphTargetVertexInfoBuffers.GetMaximumMorphScale(i);
-
-		for (uint32 j = 0; j < 4; j++)
-		{
-			if (MorphTargetWeights.IsValidIndex(i))
-			{
-				MinAccumScale[j] += MorphTargetWeights[i] * MinMorphScale[j];
-				MaxAccumScale[j] += MorphTargetWeights[i] * MaxMorphScale[j];
-			}
-
-			double AbsMorphScale = FMath::Max<double>(FMath::Abs(MinMorphScale[j]), FMath::Abs(MaxMorphScale[j]));
-			double AbsAccumScale = FMath::Max<double>(FMath::Abs(MinAccumScale[j]), FMath::Abs(MaxAccumScale[j]));
-			//the maximum accumulated and the maximum local value have to fit into out int24
-			MaxScale[j] = FMath::Max(MaxScale[j], FMath::Max(AbsMorphScale, AbsAccumScale));
-		}
-	}
-
-	MaxScale[0] = FMath::Max<double>(MaxScale[0], 1.0);
-	MaxScale[1] = FMath::Max<double>(MaxScale[1], 1.0);
-	MaxScale[2] = FMath::Max<double>(MaxScale[2], 1.0);
-	MaxScale[3] = FMath::Max<double>(MaxScale[3], 1.0);
-
-	const double ScaleToInt24 = 16777216.0;
-
-	MorphScale = FVector4
-	(
-		ScaleToInt24 / (MaxScale[0]),
-		ScaleToInt24 / (MaxScale[1]),
-		ScaleToInt24 / (MaxScale[2]),
-		ScaleToInt24 / (MaxScale[3])
-	);
-
-	InvMorphScale = FVector4
-	(
-		MaxScale[0] / ScaleToInt24,
-		MaxScale[1] / ScaleToInt24,
-		MaxScale[2] / ScaleToInt24,
-		MaxScale[3] / ScaleToInt24
-	);
-}
-
-void FSkeletalMeshObjectGPUSkin::FSkeletalMeshObjectLOD::UpdateMorphVertexBufferGPU(FRHICommandListImmediate& RHICmdList, const TArray<float>& MorphTargetWeights, 
-																					const FMorphTargetVertexInfoBuffers& MorphTargetVertexInfoBuffers, 
-																					const TArray<int32>& SectionIdsUseByActiveMorphTargets, const FName& OwnerName, EGPUSkinCacheEntryMode Mode,
-																					FMorphVertexBuffer& MorphVertexBuffer)
-{
 	if (IsValidRef(MorphVertexBuffer.VertexBufferRHI))
 	{
@@ -1062,16 +1203,12 @@
 			MorphTargetVertexInfoBuffers.GetNumBatches());
 
 		RHICmdList.Transition(FRHITransitionInfo(MorphVertexBuffer.GetUAV(), ERHIAccess::Unknown, ERHIAccess::UAVCompute));
-		RHICmdList.ClearUAVUint(MorphVertexBuffer.GetUAV(), FUintVector4(0, 0, 0, 0));
+		if (bClearMorphVertexBuffer)
+		{
+			RHICmdList.ClearUAVUint(MorphVertexBuffer.GetUAV(), FUintVector4(0, 0, 0, 0));
+		}
 
 		{
-			FVector4 MorphScale; 
-			FVector4 InvMorphScale; 
 			{
-				SCOPE_CYCLE_COUNTER(STAT_MorphVertexBuffer_ApplyDelta);
-				CalculateMorphDeltaBounds(MorphTargetWeights, MorphTargetVertexInfoBuffers, MorphScale, InvMorphScale);
-			}
-
-			{
 				SCOPED_DRAW_EVENTF(RHICmdList, MorphUpdateScatter, TEXT("Scatter"));
 
@@ -1124,6 +1261,7 @@
 				RHICmdList.EndUAVOverlap(MorphVertexBuffer.GetUAV());
 			}
 
+			if (bNormalizePass)
 			{
 				SCOPED_DRAW_EVENTF(RHICmdList, MorphUpdateNormalize, TEXT("Normalize"));
 
@@ -1405,7 +1543,8 @@
 			}
 		}
 	}
-	else if (bAlwaysUpdateMorphVertexBuffer)
+	
+	if (bAlwaysUpdateMorphVertexBuffer || !DynamicData->ExternalMorphWeightData.MorphSets.IsEmpty())
 	{
 		return LOD.GPUSkinVertexFactories.MorphVertexFactories[ChunkIdx].Get();
 	}
@@ -2073,6 +2212,7 @@
 	LODIndex = 0;
 	ActiveMorphTargets.Reset();
 	MorphTargetWeights.Reset();
+	ExternalMorphWeightData.Reset();
 	NumWeightedActiveMorphTargets = 0;
 	ClothingSimData.Reset();
 	ClothBlendWeight = 0.0f;
@@ -2174,8 +2314,8 @@
 	int32 InLODIndex,
 	const TArray<FActiveMorphTarget>& InActiveMorphTargets,
 	const TArray<float>& InMorphTargetWeights, 
-	EPreviousBoneTransformUpdateMode PreviousBoneTransformUpdateMode
-	)
+	EPreviousBoneTransformUpdateMode PreviousBoneTransformUpdateMode,
+	const FExternalMorphWeightData& InExternalMorphWeightData)
 {
 	LODIndex = InLODIndex;
 	check(!ActiveMorphTargets.Num() && !ReferenceToLocal.Num() && !ClothingSimData.Num() && !MorphTargetWeights.Num());
@@ -2184,6 +2324,8 @@
 	ActiveMorphTargets.Append(InActiveMorphTargets);
 	MorphTargetWeights.Append(InMorphTargetWeights);
 	NumWeightedActiveMorphTargets = 0;
+	ExternalMorphWeightData = InExternalMorphWeightData;	// TODO: Make this more efficient?
+	ExternalMorphWeightData.UpdateNumActiveMorphTargets();
 
 	// Gather any bones referenced by shadow shapes
 	FSkeletalMeshSceneProxy* SkeletalMeshProxy = InMeshComponent ? (FSkeletalMeshSceneProxy*)InMeshComponent->SceneProxy : nullptr;
diff --git a/Engine/Source/Runtime/Engine/Private/SkeletalRenderGPUSkin.h b/Engine/Source/Runtime/Engine/Private/SkeletalRenderGPUSkin.h
--- a/Engine/Source/Runtime/Engine/Private/SkeletalRenderGPUSkin.h
+++ b/Engine/Source/Runtime/Engine/Private/SkeletalRenderGPUSkin.h
@@ -62,8 +62,8 @@ public:
 		int32 InLODIndex,
 		const TArray<FActiveMorphTarget>& InActiveMorphTargets,
 		const TArray<float>& InMorphTargetWeights, 
-		EPreviousBoneTransformUpdateMode PreviousBoneTransformUpdateMode
-		);
+		EPreviousBoneTransformUpdateMode PreviousBoneTransformUpdateMode,
+		const FExternalMorphWeightData& InExternalMorphWeightData);
 
 	/** ref pose to local space transforms */
 	TArray<FMatrix44f> ReferenceToLocal;
@@ -93,6 +93,12 @@ public:
 	/** number of active morph targets with weights > 0 */
 	int32 NumWeightedActiveMorphTargets;
 
+	/** 
+	 * The dynamic data for each external morph target set.
+	 * This dynamic data contains things such as the weights for each set of external morph targets.
+	 */
+	FExternalMorphWeightData ExternalMorphWeightData;
+
 	/** data for updating cloth section */
 	TMap<int32, FClothSimulData> ClothingSimData;
 
@@ -323,7 +329,7 @@ public:
 	//~ Begin FSkeletalMeshObject Interface
 	virtual void InitResources(USkinnedMeshComponent* InMeshComponent) override;
 	virtual void ReleaseResources() override;
-	virtual void Update(int32 LODIndex,USkinnedMeshComponent* InMeshComponent,const TArray<FActiveMorphTarget>& ActiveMorphTargets, const TArray<float>& MorphTargetWeights, EPreviousBoneTransformUpdateMode PreviousBoneTransformUpdateMode) override;
+	virtual void Update(int32 LODIndex,USkinnedMeshComponent* InMeshComponent,const TArray<FActiveMorphTarget>& ActiveMorphTargets, const TArray<float>& MorphTargetWeights, EPreviousBoneTransformUpdateMode PreviousBoneTransformUpdateMode, const FExternalMorphWeightData& InExternalMorphWeightData) override;
 	void UpdateDynamicData_RenderThread(FGPUSkinCache* GPUSkinCache, FRHICommandListImmediate& RHICmdList, FDynamicSkelMeshObjectDataGPUSkin* InDynamicData, FSceneInterface* Scene, uint32 FrameNumberToPrepare, uint32 RevisionNumber);
 	virtual void PreGDMECallback(FGPUSkinCache* GPUSkinCache, uint32 FrameNumber) override;
 	virtual const FVertexFactory* GetSkinVertexFactory(const FSceneView* View, int32 LODIndex,int32 ChunkIdx, ESkinVertexFactoryMode VFMode = ESkinVertexFactoryMode::Default) const override;
@@ -403,6 +409,15 @@ public:
 	}
 	//~ End FSkeletalMeshObject Interface
 
+	/**
+	 * Calculate how many GPU compressed morph target sets are active.
+	 * This includes regular morph targets as well as external morph targets.
+	 */
+	int32 CalcNumActiveGPUMorphSets(FMorphVertexBuffer& MorphVertexBuffer, const FSkeletalMeshLODRenderData& LODData) const;
+
+	/** Check if a given morph set is active or not. If so, we will process it. */
+	bool IsExternalMorphSetActive(int32 MorphSetID, FExternalMorphTargetSet& MorphSet) const;
+
 	FSkinWeightVertexBuffer* GetSkinWeightVertexBuffer(int32 LODIndex) const;
 
 	/** 
@@ -592,8 +607,10 @@ protected:
 		 */
 		void UpdateMorphVertexBufferCPU(const TArray<FActiveMorphTarget>& ActiveMorphTargets, const TArray<float>& MorphTargetWeights, const TArray<int32>& SectionIdsUseByActiveMorphTargets, 
 										bool bGPUSkinCacheEnabled, FMorphVertexBuffer& MorphVertexBuffer);
+
 		void UpdateMorphVertexBufferGPU(FRHICommandListImmediate& RHICmdList, const TArray<float>& MorphTargetWeights, const FMorphTargetVertexInfoBuffers& MorphTargetVertexInfoBuffers, 
-										const TArray<int32>& SectionIdsUseByActiveMorphTargets, const FName& OwnerName, EGPUSkinCacheEntryMode Mode, FMorphVertexBuffer& MorphVertexBuffer);
+										const TArray<int32>& SectionIdsUseByActiveMorphTargets, const FName& OwnerName, EGPUSkinCacheEntryMode Mode, FMorphVertexBuffer& MorphVertexBuffer,
+										bool bClearMorphVertexBuffer, bool bNormalizePass, const FVector4& MorphScale, const FVector4& InvMorphScale);
 
 		void UpdateSkinWeights(FSkelMeshComponentLODInfo* CompLODInfo);
 
diff --git a/Engine/Source/Runtime/Engine/Private/SkeletalRenderStatic.h b/Engine/Source/Runtime/Engine/Private/SkeletalRenderStatic.h
--- a/Engine/Source/Runtime/Engine/Private/SkeletalRenderStatic.h
+++ b/Engine/Source/Runtime/Engine/Private/SkeletalRenderStatic.h
@@ -26,7 +26,7 @@ public:
 	//~ Begin FSkeletalMeshObject Interface
 	virtual void InitResources(USkinnedMeshComponent* InMeshComponent) override;
 	virtual void ReleaseResources() override;
-	virtual void Update(int32 LODIndex,USkinnedMeshComponent* InMeshComponent,const TArray<FActiveMorphTarget>& ActiveMorphTargets, const TArray<float>& MorphTargetWeights, EPreviousBoneTransformUpdateMode PreviousBoneTransformUpdateMode) override {};
+	virtual void Update(int32 LODIndex,USkinnedMeshComponent* InMeshComponent,const TArray<FActiveMorphTarget>& ActiveMorphTargets, const TArray<float>& MorphTargetWeights, EPreviousBoneTransformUpdateMode PreviousBoneTransformUpdateMode, const FExternalMorphWeightData& InExternalMorphWeightData) override {};
 	//virtual void UpdateRecomputeTangent(int32 MaterialIndex, int32 LODIndex, bool bRecomputeTangent) override {};
 	virtual void EnableOverlayRendering(bool bEnabled, const TArray<int32>* InBonesOfInterest, const TArray<UMorphTarget*>* InMorphTargetOfInterest) override {};
 	virtual void CacheVertices(int32 LODIndex, bool bForce) const override {};
diff --git a/Engine/Source/Runtime/Engine/Public/Rendering/MorphTargetVertexInfoBuffers.h b/Engine/Source/Runtime/Engine/Public/Rendering/MorphTargetVertexInfoBuffers.h
--- a/Engine/Source/Runtime/Engine/Public/Rendering/MorphTargetVertexInfoBuffers.h
+++ b/Engine/Source/Runtime/Engine/Public/Rendering/MorphTargetVertexInfoBuffers.h
@@ -10,15 +10,18 @@ class UMorphTarget;
 class FMorphTargetVertexInfoBuffers : public FRenderResource
 {
 public:
-	FMorphTargetVertexInfoBuffers() : NumTotalBatches(0)
+	ENGINE_API FMorphTargetVertexInfoBuffers() : NumTotalBatches(0)
 	{
 	}
 
-	void InitMorphResources(EShaderPlatform ShaderPlatform, const TArray<FSkelMeshRenderSection>& RenderSections, const TArray<UMorphTarget*>& MorphTargets, int NumVertices, int32 LODIndex, float TargetPositionErrorTolerance);
+	ENGINE_API void InitMorphResources(EShaderPlatform ShaderPlatform, const TArray<FSkelMeshRenderSection>& RenderSections, const TArray<UMorphTarget*>& MorphTargets, int NumVertices, int32 LODIndex, float TargetPositionErrorTolerance);
 
 	inline bool IsMorphResourcesInitialized() const { return bResourcesInitialized; }
 	inline bool IsRHIIntialized() const { return bRHIIntialized; }
 	inline bool IsMorphCPUDataValid() const{ return bIsMorphCPUDataValid; }
+	
+	ENGINE_API bool GetEmptyMorphCPUDataOnInitRHI() const { return bEmptyMorphCPUDataOnInitRHI; }
+	ENGINE_API void SetEmptyMorphCPUDataOnInitRHI(bool bEmpty) { bEmptyMorphCPUDataOnInitRHI = bEmpty; }
 
 	ENGINE_API virtual void InitRHI() override;
 	ENGINE_API virtual void ReleaseRHI() override;
@@ -30,13 +33,13 @@ public:
 		return uint32(FMath::Min<uint64>(MaximumThreadGroupSize, UINT32_MAX));
 	}
 
-	uint32 GetNumBatches(uint32 index = UINT_MAX) const
+	ENGINE_API uint32 GetNumBatches(uint32 index = UINT_MAX) const
 	{
 		check(index == UINT_MAX || index < (uint32)BatchesPerMorph.Num());
 		return index != UINT_MAX ? BatchesPerMorph[index] : NumTotalBatches;
 	}
 
-	uint32 GetNumMorphs() const
+	ENGINE_API uint32 GetNumMorphs() const
 	{
 		return BatchesPerMorph.Num();
 	}
@@ -94,7 +97,7 @@ protected:
 	void ValidateVertexBuffers(bool bMorphTargetsShouldBeValid);
 	void Serialize(FArchive& Ar);
 
-	// Transient data. Gets deleted as soon as the GPU resource has been initialized.
+	// Transient data. Gets deleted as soon as the GPU resource has been initialized (unless excplicitly disabled by bEmptyMorphCPUDataOnInitRHI).
 	TArray<uint32> MorphData;
 
 	//x,y,y separate for position and shared w for tangent
@@ -110,9 +113,10 @@ protected:
 	bool bIsMorphCPUDataValid = false;
 	bool bResourcesInitialized = false;
 	bool bRHIIntialized = false;
+	bool bEmptyMorphCPUDataOnInitRHI = true;
 
 	friend class FSkeletalMeshLODRenderData;
-	friend FArchive& operator<<(FArchive& Ar, FMorphTargetVertexInfoBuffers& MorphTargetVertexInfoBuffers);
+	ENGINE_API friend FArchive& operator<<(FArchive& Ar, FMorphTargetVertexInfoBuffers& MorphTargetVertexInfoBuffers);
 };
 
-FArchive& operator<<(FArchive& Ar, FMorphTargetVertexInfoBuffers& MorphTargetVertexInfoBuffers);
+ENGINE_API FArchive& operator<<(FArchive& Ar, FMorphTargetVertexInfoBuffers& MorphTargetVertexInfoBuffers);
diff --git a/Engine/Source/Runtime/Engine/Public/Rendering/SkeletalMeshLODRenderData.h b/Engine/Source/Runtime/Engine/Public/Rendering/SkeletalMeshLODRenderData.h
--- a/Engine/Source/Runtime/Engine/Public/Rendering/SkeletalMeshLODRenderData.h
+++ b/Engine/Source/Runtime/Engine/Public/Rendering/SkeletalMeshLODRenderData.h
@@ -21,6 +21,17 @@ class USkinnedAsset;
 class FSkeletalMeshLODModel;
 #endif // WITH_EDITOR
 
+
+/** An external morph target set. External morph targets are managed by systems outside of the skinned meshes. */
+struct ENGINE_API FExternalMorphTargetSet
+{
+	/** A name for this set, useful for debugging. */
+	FName Name = FName(TEXT("Unknown"));
+
+	/** The GPU compressed morph buffers. */
+	FMorphTargetVertexInfoBuffers MorphBuffers;
+};
+
 struct FSkelMeshRenderSection
 {
 	/** Material (texture) used for this section. */
@@ -138,6 +149,15 @@ public:
 	/** GPU friendly access data for MorphTargets for an LOD */
 	FMorphTargetVertexInfoBuffers	MorphTargetVertexInfoBuffers;
 
+	/** 
+	 * External GPU based morph target buffers.
+	 * This contains an additional set of GPU only morphs that come from external other systems that generate morph targets.
+	 * These morph targets can only be updated on the GPU, will not be serialized as part of the Skeletal Mesh, and will not show up in the editors morph target list.
+	 * We only register them inside the render data so that these morph targets can get applied when rendering.
+	 * Every set of external morph targets has some given ID.
+	 */
+	TMap<int32, TSharedPtr<FExternalMorphTargetSet>> ExternalMorphSets;
+
 	/** Skin weight profile data structures, can contain multiple profiles and their runtime FSkinWeightVertexBuffer */
 	FSkinWeightProfilesData SkinWeightProfilesData;
 
@@ -266,6 +286,24 @@ public:
 		return OverrideBuffer != nullptr ? OverrideBuffer : &SkinWeightVertexBuffer;
 	}
 
+	/** 
+	 * Register an external set of GPU compressed morph targets. 
+	 * These compressed morph targets are GPU only morph targets that will not appear inside the UI and are owned by external systems.
+	 * Every set of these morph targets has some unique ID.
+	 * @param ID The unique ID for this set of morph targets.
+	 * @param MorphSet A shared pointer to a morph set, which basically contains a GPU friendly morph buffer. This buffer should be owned by an external system that calls this method.
+	 */
+	ENGINE_API void AddExternalMorphSet(int32 ID, TSharedPtr<FExternalMorphTargetSet> MorphSet);
+
+	/** Remove a given set of external GPU based morph targets. */
+	ENGINE_API void RemoveExternalMorphSet(int32 ID);
+
+	/** Clear all externally registered morph target buffers. */
+	ENGINE_API void ClearExternalMorphSets();
+
+	/** Do we have a given set of external morph targets? */
+	ENGINE_API bool HasExternalMorphSet(int32 ID) const;
+
 	/** Utility function for returning total number of faces in this LOD. */
 	ENGINE_API int32 GetTotalFaces() const;
 
diff --git a/Engine/Source/Runtime/Engine/Public/SkeletalRenderPublic.h b/Engine/Source/Runtime/Engine/Public/SkeletalRenderPublic.h
--- a/Engine/Source/Runtime/Engine/Public/SkeletalRenderPublic.h
+++ b/Engine/Source/Runtime/Engine/Public/SkeletalRenderPublic.h
@@ -70,7 +70,7 @@ public:
 	 * @param	InSkeletalMeshComponen - parent prim component doing the updating
 	 * @param	ActiveMorphs - morph targets to blend with during skinning
 	 */
-	virtual void Update(int32 LODIndex,USkinnedMeshComponent* InMeshComponent,const TArray<FActiveMorphTarget>& ActiveMorphTargets, const TArray<float>& MorphTargetWeights, EPreviousBoneTransformUpdateMode PreviousBoneTransformUpdateMode) = 0;
+	virtual void Update(int32 LODIndex,USkinnedMeshComponent* InMeshComponent,const TArray<FActiveMorphTarget>& ActiveMorphTargets, const TArray<float>& MorphTargetWeights, EPreviousBoneTransformUpdateMode PreviousBoneTransformUpdateMode, const FExternalMorphWeightData& InExternalMorphWeightData) = 0;
 
 	/**
 	 * Called by FSkeletalMeshObject prior to GDME. This allows the GPU skin version to update bones etc now that we know we are going to render
