Fix crash after rebuilding lighting if some materials don't have the "used with static lighting" flag set.

FShaderCompilingManager::ProcessCompiledShaderMaps was storing smart pointers to material resources in a local map, which meant that they had a refcount of 2 when UMaterialInstance::CheckMaterialUsage found the missing flag and tried to recreate the resources. This in turn made the FMaterial destructor trigger an assertion failure, because the refcount was still 1 after removing the resource from the shader map. Emptying the local map before recreating the render state fixes the issue.

#rnx

diff --git a/Engine/Source/Runtime/Engine/Private/ShaderCompiler/ShaderCompiler.cpp b/Engine/Source/Runtime/Engine/Private/ShaderCompiler/ShaderCompiler.cpp
--- a/Engine/Source/Runtime/Engine/Private/ShaderCompiler/ShaderCompiler.cpp
+++ b/Engine/Source/Runtime/Engine/Private/ShaderCompiler/ShaderCompiler.cpp
@@ -4602,6 +4602,7 @@ void FShaderCompilingManager::ProcessCompiledShaderMaps(
 
 		if (FApp::CanEverRender())
 		{
+			// This empties MaterialsToUpdate, see the comment inside the function for the reason.
 			PropagateMaterialChangesToPrimitives(MaterialsToUpdate);
 
 			FEditorSupportDelegates::RedrawAllViewports.Broadcast();
@@ -4614,7 +4615,7 @@ void FShaderCompilingManager::ProcessCompiledShaderMaps(
 #endif // WITH_EDITOR
 }
 
-void FShaderCompilingManager::PropagateMaterialChangesToPrimitives(const TMap<TRefCountPtr<FMaterial>, TRefCountPtr<FMaterialShaderMap>>& MaterialsToUpdate)
+void FShaderCompilingManager::PropagateMaterialChangesToPrimitives(TMap<TRefCountPtr<FMaterial>, TRefCountPtr<FMaterialShaderMap>>& MaterialsToUpdate)
 {
 	// don't perform any work if no compiling
 	if (!AllowShaderCompiling())
@@ -4672,6 +4673,13 @@ void FShaderCompilingManager::PropagateMaterialChangesToPrimitives(const TMap<TR
 		}
 	}
 
+	// Recreating the render state for the primitives may end up recreating the material resources if some materials are missing some usage flags.
+	// For example, if some materials are not marked as used with static lighting and we build lightmaps, UMaterialInstance::CheckMaterialUsage
+	// will catch the problem and try to set the flag. However, since MaterialsToUpdate stores smart pointers, the material resources will have
+	// a refcount of 2, so the FMaterial destructor will trigger a check failure because the refcount doesn't reach 0. Empty this map before
+	// recreating the render state to allow resources to be deleted cleanly.
+	MaterialsToUpdate.Empty();
+
 	UpdateAllPrimitiveSceneInfosForScenes(ScenesToUpdate);
 	ComponentContexts.Empty();
 	UpdateAllPrimitiveSceneInfosForScenes(ScenesToUpdate);
diff --git a/Engine/Source/Runtime/Engine/Public/ShaderCompiler.h b/Engine/Source/Runtime/Engine/Public/ShaderCompiler.h
--- a/Engine/Source/Runtime/Engine/Public/ShaderCompiler.h
+++ b/Engine/Source/Runtime/Engine/Public/ShaderCompiler.h
@@ -945,7 +945,7 @@ private:
 	void ProcessCompiledNiagaraShaderMaps(TMap<int32, FShaderMapFinalizeResults>& CompiledShaderMaps, float TimeBudget);
 
 	/** Propagate the completed compile to primitives that might be using the materials compiled. */
-	void PropagateMaterialChangesToPrimitives(const TMap<TRefCountPtr<FMaterial>, TRefCountPtr<FMaterialShaderMap>>& MaterialsToUpdate);
+	void PropagateMaterialChangesToPrimitives(TMap<TRefCountPtr<FMaterial>, TRefCountPtr<FMaterialShaderMap>>& MaterialsToUpdate);
 
 	/** Recompiles shader jobs with errors if requested, and returns true if a retry was needed. */
 	bool HandlePotentialRetryOnError(TMap<int32, FShaderMapFinalizeResults>& CompletedShaderMaps);
