FMessageLog::CriticalError has been deprecated.

FMessageLog::CriticalError has a delay before throwing an assert to allow compositing of the message. As a result the assert isn't thrown at the location of the problem, but somewhere in the surrounding area, which makes the problem harder to debug and generates misleading log entries. Due to these problems FMessageLog::CriticalError and EMessageSeverity::CriticalError have been deprecated. It's recommended to use checkf instead.

When FMessageLog goes out of scope or is destroyed explicitly it will flush its messages to the log, which in turn means the log may immediately use the messages. In a few places a reference to the message was kept that outlived the lifetime of the owning FMessageLog in order to add arguments later on. The found cases have been updated to make sure FMessageLog doesn't outlive the addition of the arguments so the arguments are correctly passed to the log.

diff --git a/Engine/Plugins/Animation/ControlRig/Source/ControlRig/Private/ControlRig.cpp b/Engine/Plugins/Animation/ControlRig/Source/ControlRig/Private/ControlRig.cpp
--- a/Engine/Plugins/Animation/ControlRig/Source/ControlRig/Private/ControlRig.cpp
+++ b/Engine/Plugins/Animation/ControlRig/Source/ControlRig/Private/ControlRig.cpp
@@ -3034,7 +3034,6 @@ void UControlRig::LogOnce(EMessageSeverity::Type InSeverity, int32 InInstruction
 
 	switch (InSeverity)
 	{
-		case EMessageSeverity::CriticalError:
 		case EMessageSeverity::Error:
 		{
 			UE_LOG(LogControlRig, Error, TEXT("%s"), *InMessage);
diff --git a/Engine/Plugins/Animation/ControlRig/Source/ControlRigEditor/Private/Editor/SControlRigStackView.cpp b/Engine/Plugins/Animation/ControlRig/Source/ControlRigEditor/Private/Editor/SControlRigStackView.cpp
--- a/Engine/Plugins/Animation/ControlRig/Source/ControlRigEditor/Private/Editor/SControlRigStackView.cpp
+++ b/Engine/Plugins/Animation/ControlRig/Source/ControlRigEditor/Private/Editor/SControlRigStackView.cpp
@@ -796,7 +796,6 @@ void SControlRigStackView::RefreshTreeView(URigVM* InVM)
 							break;
 						}
 						case EMessageSeverity::Error:
-						case EMessageSeverity::CriticalError:
 						{
 							Operators[LogEntry.InstructionIndex]->Children.Add(MakeShared<FRigStackEntry>(ChildIndex, ERigStackEntry::Error, LogEntry.InstructionIndex, ERigVMOpCode::Invalid, LogEntry.Message, FRigVMASTProxy()));
 							break;
diff --git a/Engine/Plugins/Animation/LiveLink/Source/LiveLink/Private/LiveLinkLogInstance.cpp b/Engine/Plugins/Animation/LiveLink/Source/LiveLink/Private/LiveLinkLogInstance.cpp
--- a/Engine/Plugins/Animation/LiveLink/Source/LiveLink/Private/LiveLinkLogInstance.cpp
+++ b/Engine/Plugins/Animation/LiveLink/Source/LiveLink/Private/LiveLinkLogInstance.cpp
@@ -438,7 +438,6 @@ void FLiveLinkLogInstance::LogMessage(EMessageSeverity::Type Severity, const FSt
 
 	switch (Severity)
 	{
-	case EMessageSeverity::CriticalError:
 	case EMessageSeverity::Error:
 		UE_LOG(LogLiveLink, Error, TEXT("%s"), *Message);
 		break;
@@ -457,7 +456,6 @@ void FLiveLinkLogInstance::IncrementLogCount(EMessageSeverity::Type Severity)
 {
 	switch (Severity)
 	{
-	case EMessageSeverity::CriticalError:
 	case EMessageSeverity::Error:
 		++ErrorCount;
 		break;
diff --git a/Engine/Plugins/FX/Niagara/Source/NiagaraEditor/Private/NiagaraMessageUtilities.cpp b/Engine/Plugins/FX/Niagara/Source/NiagaraEditor/Private/NiagaraMessageUtilities.cpp
--- a/Engine/Plugins/FX/Niagara/Source/NiagaraEditor/Private/NiagaraMessageUtilities.cpp
+++ b/Engine/Plugins/FX/Niagara/Source/NiagaraEditor/Private/NiagaraMessageUtilities.cpp
@@ -92,7 +92,6 @@ UNiagaraStackEntry::FStackIssue FNiagaraMessageUtilities::MessageToStackIssue(TS
 	EStackIssueSeverity StackIssueSeverity;
 	switch (TokenizedMessage->GetSeverity())
 	{
-	case EMessageSeverity::CriticalError:
 	case EMessageSeverity::Error:
 		StackIssueSeverity = EStackIssueSeverity::Error;
 		break;
diff --git a/Engine/Plugins/Importers/USDImporter/Source/USDUtilities/Private/USDErrorUtils.cpp b/Engine/Plugins/Importers/USDImporter/Source/USDUtilities/Private/USDErrorUtils.cpp
--- a/Engine/Plugins/Importers/USDImporter/Source/USDUtilities/Private/USDErrorUtils.cpp
+++ b/Engine/Plugins/Importers/USDImporter/Source/USDUtilities/Private/USDErrorUtils.cpp
@@ -200,7 +200,7 @@ void FUsdLogManager::LogMessage( const TSharedRef< FTokenizedMessage >& Message
 
 	if ( !bMessageProcessed )
 	{
-		if ( Message->GetSeverity() == EMessageSeverity::CriticalError || Message->GetSeverity() == EMessageSeverity::Error )
+		if ( Message->GetSeverity() == EMessageSeverity::Error )
 		{
 			UE_LOG( LogUsd, Error, TEXT("%s"), *(Message->ToText().ToString()) );
 		}
diff --git a/Engine/Plugins/Runtime/StateTree/Source/StateTreeEditorModule/Private/StateTreeCompilerLog.cpp b/Engine/Plugins/Runtime/StateTree/Source/StateTreeEditorModule/Private/StateTreeCompilerLog.cpp
--- a/Engine/Plugins/Runtime/StateTree/Source/StateTreeEditorModule/Private/StateTreeCompilerLog.cpp
+++ b/Engine/Plugins/Runtime/StateTree/Source/StateTreeEditorModule/Private/StateTreeCompilerLog.cpp
@@ -1,4 +1,4 @@
-﻿// Copyright Epic Games, Inc. All Rights Reserved.
+// Copyright Epic Games, Inc. All Rights Reserved.
 
 #include "StateTreeCompilerLog.h"
 #include "StateTreeState.h"
@@ -53,7 +53,6 @@ void FStateTreeCompilerLog::DumpToLog(const FLogCategoryBase& Category) const
 
 		switch (StateTreeMessage.Severity)
 		{
-		case EMessageSeverity::CriticalError:
 		case EMessageSeverity::Error:
 			UE_LOG_REF(Category, Error, TEXT("%s"), *Message);
 			break;
diff --git a/Engine/Source/Developer/RigVMDeveloper/Private/RigVMModel/RigVMController.cpp b/Engine/Source/Developer/RigVMDeveloper/Private/RigVMModel/RigVMController.cpp
--- a/Engine/Source/Developer/RigVMDeveloper/Private/RigVMModel/RigVMController.cpp
+++ b/Engine/Source/Developer/RigVMDeveloper/Private/RigVMModel/RigVMController.cpp
@@ -11352,8 +11352,7 @@ URigVMUserWorkflowOptions* URigVMController::MakeOptionsForWorkflow(UObject* InS
 		{
 			if(URigVMController* StrongThis = WeakThis.Get())
 			{
-				if(InSeverity == EMessageSeverity::Error ||
-					InSeverity == EMessageSeverity::CriticalError)
+				if(InSeverity == EMessageSeverity::Error)
 				{
 					StrongThis->ReportAndNotifyError(InMessage);
 				}
diff --git a/Engine/Source/Editor/Persona/Private/SAnimationEditorViewport.cpp b/Engine/Source/Editor/Persona/Private/SAnimationEditorViewport.cpp
--- a/Engine/Source/Editor/Persona/Private/SAnimationEditorViewport.cpp
+++ b/Engine/Source/Editor/Persona/Private/SAnimationEditorViewport.cpp
@@ -364,7 +364,6 @@ TWeakPtr<SWidget> SAnimationEditorViewportTabBody::AddNotification(TAttribute<EM
 	{
 		switch(InSeverity.Get())
 		{
-		case EMessageSeverity::CriticalError:
 		case EMessageSeverity::Error:
 			return FAppStyle::GetBrush("AnimViewport.Notification.Error");
 		case EMessageSeverity::PerformanceWarning:
diff --git a/Engine/Source/Editor/SwarmInterface/Public/SwarmDefines.h b/Engine/Source/Editor/SwarmInterface/Public/SwarmDefines.h
--- a/Engine/Source/Editor/SwarmInterface/Public/SwarmDefines.h
+++ b/Engine/Source/Editor/SwarmInterface/Public/SwarmDefines.h
@@ -217,7 +217,7 @@ enum TAlertLevel
 	ALERT_LEVEL_INFO						= 0x00000001,
 	ALERT_LEVEL_WARNING						= 0x00000002,
 	ALERT_LEVEL_ERROR						= 0x00000003,
-	ALERT_LEVEL_CRITICAL_ERROR				= 0x00000004,
+	ALERT_LEVEL_CRITICAL_ERROR	UE_DEPRECATED(5.1, "CriticalError was removed because it can't trigger an assert at the callsite. Use 'checkf' instead.") = 0x00000004,
 };
 
 }
diff --git a/Engine/Source/Editor/UATHelper/UATHelperModule.cpp b/Engine/Source/Editor/UATHelper/UATHelperModule.cpp
--- a/Engine/Source/Editor/UATHelper/UATHelperModule.cpp
+++ b/Engine/Source/Editor/UATHelper/UATHelperModule.cpp
@@ -223,8 +223,8 @@ private:
 					bHasAssetErrors = true;
 				}
 
-				TSharedRef<FTokenizedMessage> PackagingMsg = MessageType == EMessageSeverity::Error ? FMessageLog("PackagingResults").Error()
-					: FMessageLog("PackagingResults").Warning();
+				FMessageLog Message("PackagingResults");
+				TSharedRef<FTokenizedMessage> PackagingMsg = MessageType == EMessageSeverity::Error ? Message.Error() : Message.Warning();
 
 				PackagingMsg->AddToken(FTextToken::Create(FText::FromString(MessageArray.Num() > 1 ? MessageArray[1] : MessageString)));
 
diff --git a/Engine/Source/Editor/UnrealEd/Private/Fbx/FbxSkeletalMeshImport.cpp b/Engine/Source/Editor/UnrealEd/Private/Fbx/FbxSkeletalMeshImport.cpp
--- a/Engine/Source/Editor/UnrealEd/Private/Fbx/FbxSkeletalMeshImport.cpp
+++ b/Engine/Source/Editor/UnrealEd/Private/Fbx/FbxSkeletalMeshImport.cpp
@@ -4551,7 +4551,7 @@ FFbxLogger::~FFbxLogger()
 	{
 		for (TSharedRef<FTokenizedMessage> TokenMessage : TokenizedErrorMessages)
 		{
-			if (TokenMessage->GetSeverity() == EMessageSeverity::CriticalError || TokenMessage->GetSeverity() == EMessageSeverity::Error)
+			if (TokenMessage->GetSeverity() == EMessageSeverity::Error)
 			{
 				ShowLogMessage = true;
 				break;
diff --git a/Engine/Source/Editor/UnrealEd/Private/Kismet2/BlueprintEditorUtils.cpp b/Engine/Source/Editor/UnrealEd/Private/Kismet2/BlueprintEditorUtils.cpp
--- a/Engine/Source/Editor/UnrealEd/Private/Kismet2/BlueprintEditorUtils.cpp
+++ b/Engine/Source/Editor/UnrealEd/Private/Kismet2/BlueprintEditorUtils.cpp
@@ -6384,7 +6384,6 @@ static void ShowNotification(const FText& Message, EMessageSeverity::Type Severi
 	{
 		switch(Severity)
 		{
-		case EMessageSeverity::CriticalError:
 		case EMessageSeverity::Error:
 			UE_LOG(LogBlueprint, Error, TEXT("%s"), *Message.ToString());
 			break;
@@ -6404,7 +6403,6 @@ static void ShowNotification(const FText& Message, EMessageSeverity::Type Severi
 		Warning.bFireAndForget = true;
 		switch(Severity)
 		{
-		case EMessageSeverity::CriticalError:
 		case EMessageSeverity::Error:
 			Warning.Image = FCoreStyle::Get().GetBrush(TEXT("MessageLog.Error"));
 			break;
diff --git a/Engine/Source/Editor/UnrealEd/Private/Kismet2/CompilerResultsLog.cpp b/Engine/Source/Editor/UnrealEd/Private/Kismet2/CompilerResultsLog.cpp
--- a/Engine/Source/Editor/UnrealEd/Private/Kismet2/CompilerResultsLog.cpp
+++ b/Engine/Source/Editor/UnrealEd/Private/Kismet2/CompilerResultsLog.cpp
@@ -439,7 +439,7 @@ void FCompilerResultsLog::InternalLogMessage(FName MessageID, const TSharedRef<F
 
 	if (!bSilentMode && (!bLogInfoOnly || (Severity == EMessageSeverity::Info)))
 	{
-		if (Severity == EMessageSeverity::CriticalError || Severity == EMessageSeverity::Error)
+		if (Severity == EMessageSeverity::Error)
 		{
 			if (IsRunningCommandlet())
 			{
diff --git a/Engine/Source/Editor/UnrealEd/Private/Lightmass/Lightmass.cpp b/Engine/Source/Editor/UnrealEd/Private/Lightmass/Lightmass.cpp
--- a/Engine/Source/Editor/UnrealEd/Private/Lightmass/Lightmass.cpp
+++ b/Engine/Source/Editor/UnrealEd/Private/Lightmass/Lightmass.cpp
@@ -449,7 +449,6 @@ void FLightmassProcessor::SwarmCallback( NSwarm::FMessage* CallbackMessage, void
 			case NSwarm::ALERT_LEVEL_INFO:	break;
 			case NSwarm::ALERT_LEVEL_WARNING:			CheckType = EMessageSeverity::Warning;			break;
 			case NSwarm::ALERT_LEVEL_ERROR:				CheckType = EMessageSeverity::Error;			break;
-			case NSwarm::ALERT_LEVEL_CRITICAL_ERROR:	CheckType = EMessageSeverity::CriticalError;	break;
 			}
 
 			FGuid ObjectGuid = FGuid(	AlertMessage->ObjectGuid.A, 
diff --git a/Engine/Source/Editor/UnrealEd/Private/PlayLevel.cpp b/Engine/Source/Editor/UnrealEd/Private/PlayLevel.cpp
--- a/Engine/Source/Editor/UnrealEd/Private/PlayLevel.cpp
+++ b/Engine/Source/Editor/UnrealEd/Private/PlayLevel.cpp
@@ -160,7 +160,7 @@ public:
 				FMessageLog(NAME_CategoryPIE).SuppressLoggingToOutputLog(true).Error(Message);
 				break;
 			case ELogVerbosity::Fatal:
-				FMessageLog(NAME_CategoryPIE).SuppressLoggingToOutputLog(true).CriticalError(Message);
+				checkf(false, *Message.ToString());
 				break;
 			}
 		}
@@ -461,9 +461,9 @@ void UEditorEngine::EndPlayMap()
 		// We cannot safely recover from this.
 		if (UObjectBaseUtility::IsPendingKillEnabled())
 		{
-			FMessageLog(NAME_CategoryPIE).CriticalError()
-				->AddToken(FUObjectToken::Create(Object, FText::FromString(Object->GetFullName())))
-				->AddToken(FTextToken::Create(FText::Format(LOCTEXT("PIEObjectStillReferenced", "Object from PIE level still referenced. Shortest path from root: {Path}"), Arguments)));
+			Arguments.Add(TEXT("Object"), FText::FromString(Object->GetFullName()));
+			checkf(false, *FText::Format(
+				LOCTEXT("PIEObjectStillReferenced", "Object '{Object}' from PIE level still referenced. Shortest path from root: {Path}"), Arguments).ToString());
 		}
 		else
 		{
diff --git a/Engine/Source/Editor/UnrealEd/Public/Kismet2/CompilerResultsLog.h b/Engine/Source/Editor/UnrealEd/Public/Kismet2/CompilerResultsLog.h
--- a/Engine/Source/Editor/UnrealEd/Public/Kismet2/CompilerResultsLog.h
+++ b/Engine/Source/Editor/UnrealEd/Public/Kismet2/CompilerResultsLog.h
@@ -245,7 +245,6 @@ public:
 		switch (InMessage->GetSeverity())
 		{
 		case EMessageSeverity::Error:
-		case EMessageSeverity::CriticalError:
 			++NumErrors;
 			break;
 		case EMessageSeverity::Warning:
diff --git a/Engine/Source/Runtime/Core/Private/Logging/MessageLog.cpp b/Engine/Source/Runtime/Core/Private/Logging/MessageLog.cpp
--- a/Engine/Source/Runtime/Core/Private/Logging/MessageLog.cpp
+++ b/Engine/Source/Runtime/Core/Private/Logging/MessageLog.cpp
@@ -130,7 +130,13 @@ TSharedRef<FTokenizedMessage> FMessageLog::Message( EMessageSeverity::Type InSev
 
 TSharedRef<FTokenizedMessage> FMessageLog::CriticalError( const FText& InMessage )
 {
-	TSharedRef<FTokenizedMessage> Message = FTokenizedMessage::Create(EMessageSeverity::CriticalError, InMessage);
+	// Deprecated: This call has been deprecated. As EMessageSeverity::CriticalError has been deprecated as well, this call
+	//		will no longer function as before. To avoid this going unnoticed an explicit call to checkf will be made
+	checkf(false, 
+		TEXT("[Deprecated] FMessageLog::CriticalError has been deprecated. Please update callsites. The original error message: \n%s"), 
+		*InMessage.ToString());
+	
+	TSharedRef<FTokenizedMessage> Message = FTokenizedMessage::Create(EMessageSeverity::Error, InMessage);
 	Messages.Add(Message);
 	return Message;
 }
@@ -223,8 +229,6 @@ ELogVerbosity::Type FMessageLog::GetLogVerbosity( EMessageSeverity::Type InSever
 {
 	switch(InSeverity)
 	{
-	case EMessageSeverity::CriticalError:
-		return ELogVerbosity::Fatal;
 	case EMessageSeverity::Error:
 		return ELogVerbosity::Error;
 	case EMessageSeverity::PerformanceWarning:
@@ -243,14 +247,14 @@ const TCHAR* const FMessageLog::GetLogColor( EMessageSeverity::Type InSeverity )
 #else
 	switch(InSeverity)
 	{
-	case EMessageSeverity::CriticalError:
+	case EMessageSeverity::Error:
 		return OutputDeviceColor::COLOR_RED;
 	case EMessageSeverity::PerformanceWarning:
 	case EMessageSeverity::Warning:
 		return OutputDeviceColor::COLOR_YELLOW;
 	case EMessageSeverity::Info:
 	default:
-		return NULL;
+		return nullptr;
 	}
 #endif
 }
diff --git a/Engine/Source/Runtime/Core/Private/Logging/TokenizedMessage.cpp b/Engine/Source/Runtime/Core/Private/Logging/TokenizedMessage.cpp
--- a/Engine/Source/Runtime/Core/Private/Logging/TokenizedMessage.cpp
+++ b/Engine/Source/Runtime/Core/Private/Logging/TokenizedMessage.cpp
@@ -64,8 +64,6 @@ FText FTokenizedMessage::GetSeverityText( EMessageSeverity::Type InSeverity )
 {
 	switch (InSeverity)
 	{
-	case EMessageSeverity::CriticalError:		
-		return LOCTEXT("CritError", "Critical Error");			
 	case EMessageSeverity::Error:				
 		return LOCTEXT("Error", "Error");						
 	case EMessageSeverity::PerformanceWarning:	
@@ -84,9 +82,6 @@ FName FTokenizedMessage::GetSeverityIconName(EMessageSeverity::Type InSeverity)
 	FName SeverityIconName;
 	switch (InSeverity)
 	{
-	case EMessageSeverity::CriticalError:		
-		SeverityIconName = "Icons.ErrorWithColor";		
-		break;
 	case EMessageSeverity::Error:				
 		SeverityIconName = "Icons.ErrorWithColor";		
 		break;
diff --git a/Engine/Source/Runtime/Core/Public/Logging/MessageLog.h b/Engine/Source/Runtime/Core/Public/Logging/MessageLog.h
--- a/Engine/Source/Runtime/Core/Public/Logging/MessageLog.h
+++ b/Engine/Source/Runtime/Core/Public/Logging/MessageLog.h
@@ -50,7 +50,8 @@ public:
 	 * @return the message for chaining calls.
 	 */
 	CORE_API TSharedRef<FTokenizedMessage> Message( EMessageSeverity::Type InSeverity, const FText& InMessage = FText() );
-	CORE_API TSharedRef<FTokenizedMessage> CriticalError( const FText& InMessage = FText() );
+	UE_DEPRECATED(5.1, "CriticalError was removed because it can't trigger an assert at the callsite. Use 'checkf' instead.")
+	CORE_API TSharedRef<FTokenizedMessage> CriticalError(const FText& InMessage = FText());
 	CORE_API TSharedRef<FTokenizedMessage> Error( const FText& InMessage = FText() );
 	CORE_API TSharedRef<FTokenizedMessage> PerformanceWarning( const FText& InMessage = FText() );
 	CORE_API TSharedRef<FTokenizedMessage> Warning( const FText& InMessage = FText() );
diff --git a/Engine/Source/Runtime/Core/Public/Logging/TokenizedMessage.h b/Engine/Source/Runtime/Core/Public/Logging/TokenizedMessage.h
--- a/Engine/Source/Runtime/Core/Public/Logging/TokenizedMessage.h
+++ b/Engine/Source/Runtime/Core/Public/Logging/TokenizedMessage.h
@@ -20,11 +20,11 @@ namespace EMessageSeverity
 	/** Ordered according to their severity */
 	enum Type
 	{
-		CriticalError		= 0,
-		Error				= 1,
-		PerformanceWarning	= 2,
-		Warning				= 3,
-		Info				= 4,	// Should be last
+		CriticalError UE_DEPRECATED(5.1, "CriticalError was removed because it can't trigger an assert at the callsite. Use 'checkf' instead.") = 0,
+		Error = 1,
+		PerformanceWarning = 2,
+		Warning = 3,
+		Info = 4,	// Should be last
 	};
 }
 
diff --git a/Engine/Source/Runtime/Engine/Private/LevelInstance/LevelInstanceActorImpl.cpp b/Engine/Source/Runtime/Engine/Private/LevelInstance/LevelInstanceActorImpl.cpp
--- a/Engine/Source/Runtime/Engine/Private/LevelInstance/LevelInstanceActorImpl.cpp
+++ b/Engine/Source/Runtime/Engine/Private/LevelInstance/LevelInstanceActorImpl.cpp
@@ -361,7 +361,8 @@ void FLevelInstanceActorImpl::CheckForErrors()
 	const ILevelInstanceInterface* LoopStart = nullptr;
 	if (!ULevelInstanceSubsystem::CheckForLoop(LevelInstance, LevelInstance->GetWorldAsset(), & LoopInfo, &LoopStart))
 	{
-		TSharedRef<FTokenizedMessage> Error = FMessageLog("MapCheck").Error()->AddToken(FTextToken::Create(LOCTEXT("LevelInstanceActor_Loop_CheckForErrors", "LevelInstance level loop found!")));
+		FMessageLog Message("MapCheck");
+		TSharedRef<FTokenizedMessage> Error = Message.Error()->AddToken(FTextToken::Create(LOCTEXT("LevelInstanceActor_Loop_CheckForErrors", "LevelInstance level loop found!")));
 		const AActor* LoopStartActor = CastChecked<AActor>(LoopStart);
 		TSoftObjectPtr<UWorld> LoopStartAsset(LoopStartActor->GetLevel()->GetTypedOuter<UWorld>());
 		Error->AddToken(FAssetNameToken::Create(LoopStartAsset.GetLongPackageName(), FText::FromString(LoopStartAsset.GetAssetName())));
@@ -380,7 +381,7 @@ void FLevelInstanceActorImpl::CheckForErrors()
 	FPackagePath WorldAssetPath;
 	if (!FPackagePath::TryFromPackageName(LevelInstance->GetWorldAssetPackage(), WorldAssetPath) || !FPackageName::DoesPackageExist(WorldAssetPath))
 	{
-		TSharedRef<FTokenizedMessage> Error = FMessageLog("MapCheck").Error()
+		FMessageLog("MapCheck").Error()
 			->AddToken(FTextToken::Create(LOCTEXT("LevelInstanceActor_InvalidPackage", "LevelInstance actor")))
 			->AddToken(FUObjectToken::Create(CastChecked<AActor>(LevelInstance)))
 			->AddToken(FTextToken::Create(FText::FromString(TEXT("refers to an invalid asset:"))))
diff --git a/Engine/Source/Runtime/Engine/Private/Slate/SlateGameResources.cpp b/Engine/Source/Runtime/Engine/Private/Slate/SlateGameResources.cpp
--- a/Engine/Source/Runtime/Engine/Private/Slate/SlateGameResources.cpp
+++ b/Engine/Source/Runtime/Engine/Private/Slate/SlateGameResources.cpp
@@ -152,18 +152,20 @@ const FSlateWidgetStyle* FSlateGameResources::GetWidgetStyleInternal(const FName
 
 void FSlateGameResources::Log( ISlateStyle::EStyleMessageSeverity Severity, const FText& Message ) const
 {
-	EMessageSeverity::Type EngineMessageSeverity = EMessageSeverity::CriticalError;
-	switch( Severity )
+	if (GIsEditor)
 	{
-	case ISlateStyle::EStyleMessageSeverity::CriticalError: EngineMessageSeverity = EMessageSeverity::CriticalError; break;
-	case ISlateStyle::EStyleMessageSeverity::Error: EngineMessageSeverity = EMessageSeverity::Error; break;
-	case ISlateStyle::EStyleMessageSeverity::PerformanceWarning: EngineMessageSeverity = EMessageSeverity::PerformanceWarning; break;
-	case ISlateStyle::EStyleMessageSeverity::Warning: EngineMessageSeverity = EMessageSeverity::Warning; break;
-	case ISlateStyle::EStyleMessageSeverity::Info: EngineMessageSeverity = EMessageSeverity::Info; break;
-	}
+		const EMessageSeverity::Type EngineMessageSeverity = [Severity]()
+		{
+			switch (Severity)
+			{
+			default:
+			case ISlateStyle::EStyleMessageSeverity::Error:					return EMessageSeverity::Error;
+			case ISlateStyle::EStyleMessageSeverity::PerformanceWarning:	return EMessageSeverity::PerformanceWarning;
+			case ISlateStyle::EStyleMessageSeverity::Warning:				return EMessageSeverity::Warning;
+			case ISlateStyle::EStyleMessageSeverity::Info:					return EMessageSeverity::Info;
+			}
+		}();
 
-	if( GIsEditor )
-	{
 		FMessageLog SlateStyleLog("SlateStyleLog");
 		SlateStyleLog.AddMessage(FTokenizedMessage::Create(EngineMessageSeverity, Message));
 
diff --git a/Engine/Source/Runtime/SlateCore/Public/Styling/ISlateStyle.h b/Engine/Source/Runtime/SlateCore/Public/Styling/ISlateStyle.h
--- a/Engine/Source/Runtime/SlateCore/Public/Styling/ISlateStyle.h
+++ b/Engine/Source/Runtime/SlateCore/Public/Styling/ISlateStyle.h
@@ -173,11 +173,11 @@ protected:
 	 */
 	enum class EStyleMessageSeverity : uint8
 	{
-		CriticalError		= 0,
-		Error				= 1,
-		PerformanceWarning	= 2,
-		Warning				= 3,
-		Info				= 4,	// Should be last
+		CriticalError UE_DEPRECATED(5.1, "CriticalError was removed because it can't trigger an assert at the callsite. Use 'checkf' instead.") = 0,
+		Error = 1,
+		PerformanceWarning = 2,
+		Warning = 3,
+		Info = 4,	// Should be last
 	};
 
 
