Optimize Raycast on Convex,
Moving a lock that slow down raycast

diff --git a/Engine/Source/Runtime/Engine/Classes/GameFramework/LightWeightInstanceSubsystem.h b/Engine/Source/Runtime/Engine/Classes/GameFramework/LightWeightInstanceSubsystem.h
--- a/Engine/Source/Runtime/Engine/Classes/GameFramework/LightWeightInstanceSubsystem.h
+++ b/Engine/Source/Runtime/Engine/Classes/GameFramework/LightWeightInstanceSubsystem.h
@@ -28,11 +28,13 @@ struct ENGINE_API FLightWeightInstanceSubsystem
 
 	static FLightWeightInstanceSubsystem& Get()
 	{
-		FScopeLock Lock(&GetFunctionCS);
-
 		if (!LWISubsystem)
 		{
-			LWISubsystem = MakeShareable(new FLightWeightInstanceSubsystem());
+			FScopeLock Lock(&GetFunctionCS);
+			if (!LWISubsystem)
+			{
+				LWISubsystem = MakeShareable(new FLightWeightInstanceSubsystem());
+			}
 		}
 		return *LWISubsystem;
 	}
