Remove UTexture2DArray::GetPlatformDataOrWait() as it was superseded by UTexture2DArray::GetPlatformData().

diff --git a/Engine/Source/Runtime/Engine/Classes/Engine/Texture2DArray.cpp b/Engine/Source/Runtime/Engine/Classes/Engine/Texture2DArray.cpp
--- a/Engine/Source/Runtime/Engine/Classes/Engine/Texture2DArray.cpp
+++ b/Engine/Source/Runtime/Engine/Classes/Engine/Texture2DArray.cpp
@@ -122,7 +122,7 @@ const FTexturePlatformData* UTexture2DArray::GetPlatformData() const
 		TRACE_CPUPROFILER_EVENT_SCOPE(UTexture2DArray::GetPlatformDataStall);
 		UE_LOG(LogTexture, Log, TEXT("Call to GetPlatformData() is forcing a wait on data that is not yet ready."));
 
-		FText Msg = FText::Format(LOCTEXT("WaitOnTextureCompilation2", "Waiting on texture compilation {0} ..."), FText::FromString(GetName()));
+		FText Msg = FText::Format(LOCTEXT("WaitOnTextureCompilation", "Waiting on texture compilation {0} ..."), FText::FromString(GetName()));
 		FScopedSlowTask Progress(1.f, Msg, true);
 		Progress.MakeDialog(true);
 		uint64 StartTime = FPlatformTime::Cycles64();
@@ -194,30 +194,6 @@ FTextureResource* UTexture2DArray::CreateResource()
 	return nullptr;
 }
 
-
-// Any direct access to GetPlatformDataOrWait() will stall until the structure
-// is safe to use. It is advisable to replace those use case with
-// async aware code to avoid stalls where possible.
-FTexturePlatformData* UTexture2DArray::GetPlatformDataOrWait()
-{
-#if WITH_EDITOR
-	if (PrivatePlatformData && !PrivatePlatformData->IsAsyncWorkComplete())
-	{
-		TRACE_CPUPROFILER_EVENT_SCOPE(UTexture2D::GetPlatformDataStall);
-		UE_LOG(LogTexture, Log, TEXT("Call to GetPlatformDataOrWait() is forcing a wait on data that is not yet ready."));
-
-		FText Msg = FText::Format(LOCTEXT("WaitOnTextureCompilation", "Waiting on texture array compilation {0} ..."), FText::FromString(GetName()));
-		FScopedSlowTask Progress(1.f, Msg, true);
-		Progress.MakeDialog(true);
-		uint64 StartTime = FPlatformTime::Cycles64();
-		PrivatePlatformData->FinishCache();
-		AsyncCompilationHelpers::SaveStallStack(FPlatformTime::Cycles64() - StartTime);
-	}
-#endif
-
-	return PrivatePlatformData;
-}
-
 void UTexture2DArray::UpdateResource()
 {
 #if WITH_EDITOR
diff --git a/Engine/Source/Runtime/Engine/Classes/Engine/Texture2DArray.h b/Engine/Source/Runtime/Engine/Classes/Engine/Texture2DArray.h
--- a/Engine/Source/Runtime/Engine/Classes/Engine/Texture2DArray.h
+++ b/Engine/Source/Runtime/Engine/Classes/Engine/Texture2DArray.h
@@ -47,14 +47,6 @@ public:
 	*/
 	bool GetMipData(int32 InFirstMipToLoad, TArray<FUniqueBuffer, TInlineAllocator<MAX_TEXTURE_MIP_COUNT>>& OutMipData);
 
-	/** 
-	*  Returns the platform data, potentially waiting on texture build if not ready.
-	*  Generally waits should be avoided by structuring callers to expect async behavior.
-	* 
-	*  Waits only occur in the editor (as builds only occur in the editor).
-	*/
-	FTexturePlatformData* GetPlatformDataOrWait();
-
 	//~ Begin UTexture Interface
 	virtual ETextureClass GetTextureClass() const override { return ETextureClass::Array; }
 	virtual void Serialize(FArchive& Ar) override;
diff --git a/Engine/Source/Runtime/Engine/Private/TextureDerivedData.cpp b/Engine/Source/Runtime/Engine/Private/TextureDerivedData.cpp
--- a/Engine/Source/Runtime/Engine/Private/TextureDerivedData.cpp
+++ b/Engine/Source/Runtime/Engine/Private/TextureDerivedData.cpp
@@ -2617,7 +2617,7 @@ void FTexturePlatformData::SerializeCooked(FArchive& Ar, UTexture* Owner, bool b
 
 bool UTexture2DArray::GetMipData(int32 InFirstMipToLoad, TArray<FUniqueBuffer, TInlineAllocator<MAX_TEXTURE_MIP_COUNT>>& OutMipData)
 {
-	FTexturePlatformData* LocalPlatformData = GetPlatformDataOrWait();
+	FTexturePlatformData* LocalPlatformData = GetPlatformData();
 	const int32 ReadableMipCount = LocalPlatformData->Mips.Num() - (LocalPlatformData->GetNumMipsInTail() > 0 ? LocalPlatformData->GetNumMipsInTail() - 1 : 0);
 
 	int32 OutputMipCount = ReadableMipCount - InFirstMipToLoad;
