[VisualLogger] fixed automation tests to fix potential heap-use-after-free issues by always retrieving a fresh list of redirections after redirecting objects since containers might have been reallocated.
#rnx

diff --git a/Engine/Source/Runtime/Engine/Private/VisualLogger/VisualLoggerAutomationTests.cpp b/Engine/Source/Runtime/Engine/Private/VisualLogger/VisualLoggerAutomationTests.cpp
--- a/Engine/Source/Runtime/Engine/Private/VisualLogger/VisualLoggerAutomationTests.cpp
+++ b/Engine/Source/Runtime/Engine/Private/VisualLogger/VisualLoggerAutomationTests.cpp
@@ -381,28 +381,36 @@ bool FVisualLogRedirectionsMultipleChildrenTest::RunTest(const FString& Paramete
 	const UObject* ObjB = NewObject<AActor>(World->GetCurrentLevel(), TEXT("VLogTestObjB"), RF_Transient);
 	const UObject* ObjC = NewObject<AActor>(World->GetCurrentLevel(), TEXT("VLogTestObjC"), RF_Transient);
 	const UObject* ObjD = NewObject<AActor>(World->GetCurrentLevel(), TEXT("VLogTestObjD"), RF_Transient);
-	const UObject* ObjE = NewObject<AActor>(World->GetCurrentLevel(), TEXT("VLogTestObjE"), RF_Transient);
-	const UObject* ObjF = NewObject<AActor>(World->GetCurrentLevel(), TEXT("VLogTestObjF"), RF_Transient);
 
 	REDIRECT_OBJECT_TO_VLOG(ObjB, ObjA);
 	// B -> A
-	TArray<TWeakObjectPtr<const UObject>>* ChildrenOfA = RedirectionMap.Find(ObjA);
-	CHECK_NOT_NULL(ChildrenOfA);
-	CHECK_SUCCESS(ChildrenOfA->Num() == 1);
-	CHECK_SUCCESS(Logger.FindRedirection(ObjB) == ObjA);
+	{
+		TArray<TWeakObjectPtr<const UObject>>* ChildrenOfA = RedirectionMap.Find(ObjA);
+		CHECK_NOT_NULL(ChildrenOfA);
+		CHECK_SUCCESS(ChildrenOfA->Num() == 1);
+		CHECK_SUCCESS(Logger.FindRedirection(ObjB) == ObjA);
+	}
 
 	REDIRECT_OBJECT_TO_VLOG(ObjC, ObjA);
 	// B -> A
 	// C -> A
-	CHECK_SUCCESS(ChildrenOfA->Num() == 2);
-	CHECK_SUCCESS(Logger.FindRedirection(ObjC) == ObjA);
+	{
+		TArray<TWeakObjectPtr<const UObject>>* ChildrenOfA = RedirectionMap.Find(ObjA);
+		CHECK_NOT_NULL(ChildrenOfA);
+		CHECK_SUCCESS(ChildrenOfA->Num() == 2);
+		CHECK_SUCCESS(Logger.FindRedirection(ObjC) == ObjA);
+	}
 
 	REDIRECT_OBJECT_TO_VLOG(ObjD, ObjA);
 	// B -> A
 	// C -> A
 	// D -> A
-	CHECK_SUCCESS(ChildrenOfA->Num() == 3);
-	CHECK_SUCCESS(Logger.FindRedirection(ObjD) == ObjA);
+	{
+		TArray<TWeakObjectPtr<const UObject>>* ChildrenOfA = RedirectionMap.Find(ObjA);
+		CHECK_NOT_NULL(ChildrenOfA);
+		CHECK_SUCCESS(ChildrenOfA->Num() == 3);
+		CHECK_SUCCESS(Logger.FindRedirection(ObjD) == ObjA);
+	}
 
 	return true;
 }
@@ -430,18 +438,20 @@ bool FVisualLogRedirectionsCreationOrderTest::RunTest(const FString& Parameters)
 	// C -> B -> A
 	REDIRECT_OBJECT_TO_VLOG(ObjD, ObjC);
 	// D -> C -> B -> A
-	TArray<TWeakObjectPtr<const UObject>>* ChildrenOfA = RedirectionMap.Find(ObjA);
-	TArray<TWeakObjectPtr<const UObject>>* ChildrenOfB = RedirectionMap.Find(ObjB);
-	TArray<TWeakObjectPtr<const UObject>>* ChildrenOfC = RedirectionMap.Find(ObjC);
-	CHECK_NOT_NULL(ChildrenOfA);
-	CHECK_NOT_NULL(ChildrenOfB);
-	CHECK_NOT_NULL(ChildrenOfC);
-	CHECK_SUCCESS(ChildrenOfA->Num() == 3);
-	CHECK_SUCCESS(ChildrenOfB->Num() == 2);
-	CHECK_SUCCESS(ChildrenOfC->Num() == 1);
-	CHECK_SUCCESS(Logger.FindRedirection(ObjB) == ObjA);
-	CHECK_SUCCESS(Logger.FindRedirection(ObjC) == ObjA);
-	CHECK_SUCCESS(Logger.FindRedirection(ObjD) == ObjA);
+	{
+		TArray<TWeakObjectPtr<const UObject>>* ChildrenOfA = RedirectionMap.Find(ObjA);
+		TArray<TWeakObjectPtr<const UObject>>* ChildrenOfB = RedirectionMap.Find(ObjB);
+		TArray<TWeakObjectPtr<const UObject>>* ChildrenOfC = RedirectionMap.Find(ObjC);
+		CHECK_NOT_NULL(ChildrenOfA);
+		CHECK_NOT_NULL(ChildrenOfB);
+		CHECK_NOT_NULL(ChildrenOfC);
+		CHECK_SUCCESS(ChildrenOfA->Num() == 3);
+		CHECK_SUCCESS(ChildrenOfB->Num() == 2);
+		CHECK_SUCCESS(ChildrenOfC->Num() == 1);
+		CHECK_SUCCESS(Logger.FindRedirection(ObjB) == ObjA);
+		CHECK_SUCCESS(Logger.FindRedirection(ObjC) == ObjA);
+		CHECK_SUCCESS(Logger.FindRedirection(ObjD) == ObjA);
+	}
 
 	Logger.Cleanup(World);
 
@@ -451,18 +461,20 @@ bool FVisualLogRedirectionsCreationOrderTest::RunTest(const FString& Parameters)
 	// D -> C -> B
 	REDIRECT_OBJECT_TO_VLOG(ObjB, ObjA);
 	// D -> C -> B -> A
-	ChildrenOfA = RedirectionMap.Find(ObjA);
-	ChildrenOfB = RedirectionMap.Find(ObjB);
-	ChildrenOfC = RedirectionMap.Find(ObjC);
-	CHECK_NOT_NULL(ChildrenOfA);
-	CHECK_NOT_NULL(ChildrenOfB);
-	CHECK_NOT_NULL(ChildrenOfC);
-	CHECK_SUCCESS(ChildrenOfA->Num() == 3);
-	CHECK_SUCCESS(ChildrenOfB->Num() == 2);
-	CHECK_SUCCESS(ChildrenOfC->Num() == 1);
-	CHECK_SUCCESS(Logger.FindRedirection(ObjB) == ObjA);
-	CHECK_SUCCESS(Logger.FindRedirection(ObjC) == ObjA);
-	CHECK_SUCCESS(Logger.FindRedirection(ObjD) == ObjA);
+	{
+		TArray<TWeakObjectPtr<const UObject>>* ChildrenOfA = RedirectionMap.Find(ObjA);
+		TArray<TWeakObjectPtr<const UObject>>* ChildrenOfB = RedirectionMap.Find(ObjB);
+		TArray<TWeakObjectPtr<const UObject>>* ChildrenOfC = RedirectionMap.Find(ObjC);
+		CHECK_NOT_NULL(ChildrenOfA);
+		CHECK_NOT_NULL(ChildrenOfB);
+		CHECK_NOT_NULL(ChildrenOfC);
+		CHECK_SUCCESS(ChildrenOfA->Num() == 3);
+		CHECK_SUCCESS(ChildrenOfB->Num() == 2);
+		CHECK_SUCCESS(ChildrenOfC->Num() == 1);
+		CHECK_SUCCESS(Logger.FindRedirection(ObjB) == ObjA);
+		CHECK_SUCCESS(Logger.FindRedirection(ObjC) == ObjA);
+		CHECK_SUCCESS(Logger.FindRedirection(ObjD) == ObjA);
+	}
 
 	return true;
 }
@@ -490,33 +502,43 @@ bool FVisualLogRedirectionsWithinHierachyTest::RunTest(const FString& Parameters
 	REDIRECT_OBJECT_TO_VLOG(ObjC, ObjA);
 	// E -> D -> B -> A
 	//           C -> A
-	TArray<TWeakObjectPtr<const UObject>>* ChildrenOfA = RedirectionMap.Find(ObjA);
-	TArray<TWeakObjectPtr<const UObject>>* ChildrenOfB = RedirectionMap.Find(ObjB);
-	TArray<TWeakObjectPtr<const UObject>>* ChildrenOfD = RedirectionMap.Find(ObjD);
-	CHECK_NOT_NULL(ChildrenOfA);
-	CHECK_NOT_NULL(ChildrenOfB);
-	CHECK_NOT_NULL(ChildrenOfD);
-	CHECK_SUCCESS(ChildrenOfA->Num() == 4);
-	CHECK_SUCCESS(ChildrenOfB->Num() == 2);
-	CHECK_SUCCESS(ChildrenOfD->Num() == 1);
-	CHECK_SUCCESS(Logger.FindRedirection(ObjB) == ObjA);
-	CHECK_SUCCESS(Logger.FindRedirection(ObjC) == ObjA);
-	CHECK_SUCCESS(Logger.FindRedirection(ObjD) == ObjA);
-	CHECK_SUCCESS(Logger.FindRedirection(ObjE) == ObjA);
+	{
+		TArray<TWeakObjectPtr<const UObject>>* ChildrenOfA = RedirectionMap.Find(ObjA);
+		TArray<TWeakObjectPtr<const UObject>>* ChildrenOfB = RedirectionMap.Find(ObjB);
+		TArray<TWeakObjectPtr<const UObject>>* ChildrenOfD = RedirectionMap.Find(ObjD);
+		CHECK_NOT_NULL(ChildrenOfA);
+		CHECK_NOT_NULL(ChildrenOfB);
+		CHECK_NOT_NULL(ChildrenOfD);
+		CHECK_SUCCESS(ChildrenOfA->Num() == 4);
+		CHECK_SUCCESS(ChildrenOfB->Num() == 2);
+		CHECK_SUCCESS(ChildrenOfD->Num() == 1);
+		CHECK_SUCCESS(Logger.FindRedirection(ObjB) == ObjA);
+		CHECK_SUCCESS(Logger.FindRedirection(ObjC) == ObjA);
+		CHECK_SUCCESS(Logger.FindRedirection(ObjD) == ObjA);
+		CHECK_SUCCESS(Logger.FindRedirection(ObjE) == ObjA);
+	}
 
 	REDIRECT_OBJECT_TO_VLOG(ObjD, ObjC);
 	//           B -> A
 	// E -> D -> C -> A
-	TArray<TWeakObjectPtr<const UObject>>* ChildrenOfC = RedirectionMap.Find(ObjC);
-	CHECK_NOT_NULL(ChildrenOfC);
-	CHECK_SUCCESS(ChildrenOfA->Num() == 4);
-	CHECK_SUCCESS(ChildrenOfB->Num() == 0);
-	CHECK_SUCCESS(ChildrenOfC->Num() == 2);
-	CHECK_SUCCESS(ChildrenOfD->Num() == 1);
-	CHECK_SUCCESS(Logger.FindRedirection(ObjB) == ObjA);
-	CHECK_SUCCESS(Logger.FindRedirection(ObjC) == ObjA);
-	CHECK_SUCCESS(Logger.FindRedirection(ObjD) == ObjA);
-	CHECK_SUCCESS(Logger.FindRedirection(ObjE) == ObjA);
+	{
+		TArray<TWeakObjectPtr<const UObject>>* ChildrenOfA = RedirectionMap.Find(ObjA);
+		TArray<TWeakObjectPtr<const UObject>>* ChildrenOfB = RedirectionMap.Find(ObjB);
+		TArray<TWeakObjectPtr<const UObject>>* ChildrenOfC = RedirectionMap.Find(ObjC);
+		TArray<TWeakObjectPtr<const UObject>>* ChildrenOfD = RedirectionMap.Find(ObjD);
+		CHECK_NOT_NULL(ChildrenOfA);
+		CHECK_NOT_NULL(ChildrenOfB);
+		CHECK_NOT_NULL(ChildrenOfC);
+		CHECK_NOT_NULL(ChildrenOfD);
+		CHECK_SUCCESS(ChildrenOfA->Num() == 4);
+		CHECK_SUCCESS(ChildrenOfB->Num() == 0);
+		CHECK_SUCCESS(ChildrenOfC->Num() == 2);
+		CHECK_SUCCESS(ChildrenOfD->Num() == 1);
+		CHECK_SUCCESS(Logger.FindRedirection(ObjB) == ObjA);
+		CHECK_SUCCESS(Logger.FindRedirection(ObjC) == ObjA);
+		CHECK_SUCCESS(Logger.FindRedirection(ObjD) == ObjA);
+		CHECK_SUCCESS(Logger.FindRedirection(ObjE) == ObjA);
+	}
 
 	return true;
 }
@@ -543,28 +565,40 @@ bool FVisualLogRedirectionsDeepHierarchyToNewParentTest::RunTest(const FString&
 	REDIRECT_OBJECT_TO_VLOG(ObjD, ObjC);
 	REDIRECT_OBJECT_TO_VLOG(ObjE, ObjD);
 	// E -> D -> C -> B -> A
-	TArray<TWeakObjectPtr<const UObject>>* ChildrenOfA = RedirectionMap.Find(ObjA);
-	TArray<TWeakObjectPtr<const UObject>>* ChildrenOfB = RedirectionMap.Find(ObjB);
-	TArray<TWeakObjectPtr<const UObject>>* ChildrenOfC = RedirectionMap.Find(ObjC);
-	TArray<TWeakObjectPtr<const UObject>>* ChildrenOfD = RedirectionMap.Find(ObjD);
-	CHECK_NOT_NULL(ChildrenOfA);
-	CHECK_NOT_NULL(ChildrenOfB);
-	CHECK_NOT_NULL(ChildrenOfC);
-	CHECK_NOT_NULL(ChildrenOfD);
-	CHECK_SUCCESS(ChildrenOfA->Num() == 4);
-	CHECK_SUCCESS(ChildrenOfB->Num() == 3);
-	CHECK_SUCCESS(ChildrenOfC->Num() == 2);
-	CHECK_SUCCESS(ChildrenOfD->Num() == 1);
+	{
+		TArray<TWeakObjectPtr<const UObject>>* ChildrenOfA = RedirectionMap.Find(ObjA);
+		TArray<TWeakObjectPtr<const UObject>>* ChildrenOfB = RedirectionMap.Find(ObjB);
+		TArray<TWeakObjectPtr<const UObject>>* ChildrenOfC = RedirectionMap.Find(ObjC);
+		TArray<TWeakObjectPtr<const UObject>>* ChildrenOfD = RedirectionMap.Find(ObjD);
+		CHECK_NOT_NULL(ChildrenOfA);
+		CHECK_NOT_NULL(ChildrenOfB);
+		CHECK_NOT_NULL(ChildrenOfC);
+		CHECK_NOT_NULL(ChildrenOfD);
+		CHECK_SUCCESS(ChildrenOfA->Num() == 4);
+		CHECK_SUCCESS(ChildrenOfB->Num() == 3);
+		CHECK_SUCCESS(ChildrenOfC->Num() == 2);
+		CHECK_SUCCESS(ChildrenOfD->Num() == 1);
+	}
 	
 	REDIRECT_OBJECT_TO_VLOG(ObjB, ObjF);
 	// E -> D -> C -> B -> F
-	TArray<TWeakObjectPtr<const UObject>>* ChildrenOfF = RedirectionMap.Find(ObjF);
-	CHECK_NOT_NULL(ChildrenOfF);
-	CHECK_SUCCESS(ChildrenOfA->Num() == 0);
-	CHECK_SUCCESS(ChildrenOfF->Num() == 4);
-	CHECK_SUCCESS(ChildrenOfB->Num() == 3);
-	CHECK_SUCCESS(ChildrenOfC->Num() == 2);
-	CHECK_SUCCESS(ChildrenOfD->Num() == 1);
+	{
+		TArray<TWeakObjectPtr<const UObject>>* ChildrenOfA = RedirectionMap.Find(ObjA);
+		TArray<TWeakObjectPtr<const UObject>>* ChildrenOfB = RedirectionMap.Find(ObjB);
+		TArray<TWeakObjectPtr<const UObject>>* ChildrenOfC = RedirectionMap.Find(ObjC);
+		TArray<TWeakObjectPtr<const UObject>>* ChildrenOfD = RedirectionMap.Find(ObjD);
+		TArray<TWeakObjectPtr<const UObject>>* ChildrenOfF = RedirectionMap.Find(ObjF);
+		CHECK_NOT_NULL(ChildrenOfA);
+		CHECK_NOT_NULL(ChildrenOfB);
+		CHECK_NOT_NULL(ChildrenOfC);
+		CHECK_NOT_NULL(ChildrenOfD);
+		CHECK_NOT_NULL(ChildrenOfF);
+		CHECK_SUCCESS(ChildrenOfA->Num() == 0);
+		CHECK_SUCCESS(ChildrenOfF->Num() == 4);
+		CHECK_SUCCESS(ChildrenOfB->Num() == 3);
+		CHECK_SUCCESS(ChildrenOfC->Num() == 2);
+		CHECK_SUCCESS(ChildrenOfD->Num() == 1);
+	}
 
 	return true;
 }
