Remove the deprecated SetCaptureMouseOnClick function on the GameViewportClient, which was deprecated in 4.26. Being that 5.1 is the next release, we can safely delete this function.

diff --git a/Engine/Source/Runtime/Engine/Classes/Engine/GameViewportClient.h b/Engine/Source/Runtime/Engine/Classes/Engine/GameViewportClient.h
--- a/Engine/Source/Runtime/Engine/Classes/Engine/GameViewportClient.h
+++ b/Engine/Source/Runtime/Engine/Classes/Engine/GameViewportClient.h
@@ -748,9 +748,6 @@ public:
 	 */
 	void SetMouseCaptureMode(EMouseCaptureMode Mode);
 
-	UE_DEPRECATED(4.26, "Please call UGameViewportClient::SetMouseCaptureMode(EMouseCaptureMode) instead.")
-	void SetCaptureMouseOnClick(EMouseCaptureMode Mode) { SetMouseCaptureMode(Mode); }
-
 	/**
 	 * Gets the mouse capture behavior when the viewport is clicked
 	 */
