Cache ray tracing FMeshBatches, instance mask and flags in FStaticMeshSceneProxy::GetDynamicRayTracingInstances, similar to Nanite::FSceneProxy.

diff --git a/Engine/Source/Runtime/Engine/Private/StaticMeshRender.cpp b/Engine/Source/Runtime/Engine/Private/StaticMeshRender.cpp
--- a/Engine/Source/Runtime/Engine/Private/StaticMeshRender.cpp
+++ b/Engine/Source/Runtime/Engine/Private/StaticMeshRender.cpp
@@ -1782,7 +1782,7 @@ void FStaticMeshSceneProxy::GetDynamicRayTracingInstances(FRayTracingMaterialGat
 	}
 
 	ESceneDepthPriorityGroup PrimitiveDPG = GetStaticDepthPriorityGroup();
-	const uint32 LODIndex = FMath::Max(GetLOD(Context.ReferenceView), (int32)GetCurrentFirstLODIdx_RenderThread());
+	const int32 LODIndex = FMath::Max(GetLOD(Context.ReferenceView), (int32)GetCurrentFirstLODIdx_RenderThread());
 	const FStaticMeshLODResources& LODModel = RenderData->LODResources[LODIndex];
 
 	bool bEvaluateWPO = bDynamicRayTracingGeometry;
@@ -1820,29 +1820,43 @@ void FStaticMeshSceneProxy::GetDynamicRayTracingInstances(FRayTracingMaterialGat
 		FRayTracingInstance &RayTracingInstance = OutRayTracingInstances.AddDefaulted_GetRef();
 	
 		const int32 NumBatches = GetNumMeshBatches();
+		const int32 NumRayTracingMaterialEntries = LODModel.Sections.Num() * NumBatches;
 
-		RayTracingInstance.Materials.Reserve(LODModel.Sections.Num() * NumBatches);
-		for (int32 BatchIndex = 0; BatchIndex < NumBatches; BatchIndex++)
+		if (NumRayTracingMaterialEntries != CachedRayTracingMaterials.Num() || CachedRayTracingMaterialsLODIndex != LODIndex)
 		{
-			for (int32 SectionIndex = 0; SectionIndex < LODModel.Sections.Num(); SectionIndex++)
+			CachedRayTracingMaterials.Reset();
+			CachedRayTracingMaterials.Reserve(NumRayTracingMaterialEntries);
+
+			for (int32 BatchIndex = 0; BatchIndex < NumBatches; BatchIndex++)
 			{
-				FMeshBatch &Mesh = RayTracingInstance.Materials.AddDefaulted_GetRef();
-	
-				bool bResult = GetMeshElement(LODIndex, BatchIndex, SectionIndex, PrimitiveDPG, false, false, Mesh);
-				if (!bResult)
+				for (int32 SectionIndex = 0; SectionIndex < LODModel.Sections.Num(); SectionIndex++)
 				{
-					// Hidden material
-					Mesh.MaterialRenderProxy = UMaterial::GetDefaultMaterial(MD_Surface)->GetRenderProxy();
-					Mesh.VertexFactory = &RenderData->LODVertexFactories[LODIndex].VertexFactory;
-				}
+					FMeshBatch &MeshBatch = CachedRayTracingMaterials.AddDefaulted_GetRef();
+
+					bool bResult = GetMeshElement(LODIndex, BatchIndex, SectionIndex, PrimitiveDPG, false, false, MeshBatch);
+					if (!bResult)
+					{
+						// Hidden material
+						MeshBatch.MaterialRenderProxy = UMaterial::GetDefaultMaterial(MD_Surface)->GetRenderProxy();
+						MeshBatch.VertexFactory = &RenderData->LODVertexFactories[LODIndex].VertexFactory;
+					}
 
-				Mesh.SegmentIndex = SectionIndex;
-				Mesh.MeshIdInPrimitive = SectionIndex;
+					MeshBatch.SegmentIndex = SectionIndex;
+					MeshBatch.MeshIdInPrimitive = SectionIndex;
+				}
 			}
+
+			CachedRayTracingInstanceMaskAndFlags = BuildRayTracingInstanceMaskAndFlags(CachedRayTracingMaterials, GetScene().GetFeatureLevel());
+			CachedRayTracingMaterialsLODIndex = LODIndex;
 		}
 
 		RayTracingInstance.Geometry = &Geometry;
-		RayTracingInstance.InstanceTransforms.Add(GetLocalToWorld());
+
+		// scene proxies live for the duration of Render(), making array views below safe
+		const FMatrix& ThisLocalToWorld = GetLocalToWorld();
+		RayTracingInstance.InstanceTransformsView = MakeArrayView(&ThisLocalToWorld, 1);
+		RayTracingInstance.MaterialsView = MakeArrayView(CachedRayTracingMaterials);
+
 		if (bEvaluateWPO && RenderData->LODVertexFactories[LODIndex].VertexFactory.GetType()->SupportsRayTracingDynamicGeometry())
 		{
 			// Use the internal vertex buffer only when initialized otherwise used the shared vertex buffer - needs to be updated every frame
@@ -1855,7 +1869,7 @@ void FStaticMeshSceneProxy::GetDynamicRayTracingInstances(FRayTracingMaterialGat
 			Context.DynamicRayTracingGeometriesToUpdate.Add(
 				FRayTracingDynamicGeometryUpdateParams
 				{
-					RayTracingInstance.Materials,
+					CachedRayTracingMaterials, // TODO: this copy can be avoided if FRayTracingDynamicGeometryUpdateParams supported array views
 					false,
 					(uint32)LODModel.GetNumVertices(),
 					uint32((SIZE_T)LODModel.GetNumVertices() * sizeof(FVector3f)),
@@ -1867,11 +1881,14 @@ void FStaticMeshSceneProxy::GetDynamicRayTracingInstances(FRayTracingMaterialGat
 			);
 		}
 		
-		RayTracingInstance.BuildInstanceMaskAndFlags(GetScene().GetFeatureLevel());
+		RayTracingInstance.Mask = CachedRayTracingInstanceMaskAndFlags.Mask;
+		RayTracingInstance.bForceOpaque = CachedRayTracingInstanceMaskAndFlags.bForceOpaque;
+		RayTracingInstance.bDoubleSided = CachedRayTracingInstanceMaskAndFlags.bDoubleSided;
 
-		checkf(RayTracingInstance.Geometry->Initializer.Segments.Num() == RayTracingInstance.Materials.Num(), TEXT("Segments/Materials mismatch. Number of segments: %d. Number of Materials: %d. LOD Index: %d"), 
+		check(CachedRayTracingMaterials.Num() == RayTracingInstance.GetMaterials().Num());
+		checkf(RayTracingInstance.Geometry->Initializer.Segments.Num() == CachedRayTracingMaterials.Num(), TEXT("Segments/Materials mismatch. Number of segments: %d. Number of Materials: %d. LOD Index: %d"), 
 			RayTracingInstance.Geometry->Initializer.Segments.Num(), 
-			RayTracingInstance.Materials.Num(), 
+			CachedRayTracingMaterials.Num(), 
 			LODIndex);
 	}
 }
diff --git a/Engine/Source/Runtime/Engine/Public/StaticMeshResources.h b/Engine/Source/Runtime/Engine/Public/StaticMeshResources.h
--- a/Engine/Source/Runtime/Engine/Public/StaticMeshResources.h
+++ b/Engine/Source/Runtime/Engine/Public/StaticMeshResources.h
@@ -36,6 +36,7 @@
 #include "Templates/UniquePtr.h"
 #include "WeightedRandomSampler.h"
 #include "PerPlatformProperties.h"
+#include "RayTracingInstance.h"
 
 class FDistanceFieldVolumeData;
 class UBodySetup;
@@ -999,6 +1000,9 @@ protected:
 	bool bDynamicRayTracingGeometry;
 	TArray<FRayTracingGeometry, TInlineAllocator<MAX_MESH_LOD_COUNT>> DynamicRayTracingGeometries;
 	TArray<FRWBuffer, TInlineAllocator<MAX_MESH_LOD_COUNT>> DynamicRayTracingGeometryVertexBuffers;
+	TArray<FMeshBatch> CachedRayTracingMaterials;
+	int16 CachedRayTracingMaterialsLODIndex = INDEX_NONE;
+	FRayTracingMaskAndFlags CachedRayTracingInstanceMaskAndFlags;
 #endif
 	/**
 	 * The forcedLOD set in the static mesh editor, copied from the mesh component
