Hide SourceTextures array from UI in cases when Texture2DArray or TextureCubeArray have their own sources.

diff --git a/Engine/Source/Runtime/Engine/Classes/Engine/Texture2DArray.cpp b/Engine/Source/Runtime/Engine/Classes/Engine/Texture2DArray.cpp
--- a/Engine/Source/Runtime/Engine/Classes/Engine/Texture2DArray.cpp
+++ b/Engine/Source/Runtime/Engine/Classes/Engine/Texture2DArray.cpp
@@ -207,6 +207,10 @@ void UTexture2DArray::UpdateResource()
 		CachePlatformData();
 	}
 #endif // #if WITH_EDITOR
+	
+#if WITH_EDITORONLY_DATA
+	bSourceGeneratedFromSourceTexturesArray = !Source.GetNumSlices() || SourceTextures.Num();
+#endif
 
 	Super::UpdateResource();
 }
@@ -448,7 +452,7 @@ ENGINE_API bool UTexture2DArray::UpdateSourceFromSourceTextures(bool bCreatingNe
 
 		ValidateSettingsAfterImportOrEdit();
 		SetLightingGuid();
-		UpdateResource();		
+		UpdateResource();
 	}
 
 	return true;
@@ -502,6 +506,10 @@ void UTexture2DArray::PostLoad()
 		}
 	}
 #endif // #if WITH_EDITOR
+#if WITH_EDITORONLY_DATA
+	bSourceGeneratedFromSourceTexturesArray = !Source.GetNumSlices() || SourceTextures.Num();
+#endif
+
 	Super::PostLoad();
 };
 
diff --git a/Engine/Source/Runtime/Engine/Classes/Engine/Texture2DArray.h b/Engine/Source/Runtime/Engine/Classes/Engine/Texture2DArray.h
--- a/Engine/Source/Runtime/Engine/Classes/Engine/Texture2DArray.h
+++ b/Engine/Source/Runtime/Engine/Classes/Engine/Texture2DArray.h
@@ -91,8 +91,17 @@ public:
 
 #if WITH_EDITORONLY_DATA
 	/** Add Textures*/
-	UPROPERTY(EditAnywhere, Category = Source2D, meta = (DisplayName = "Source Textures"))
+	UPROPERTY(EditAnywhere, Category = Source2D, meta = (DisplayName = "Source Textures", EditCondition = bSourceGeneratedFromSourceTexturesArray, EditConditionHides, HideEditConditionToggle))
 	TArray<TObjectPtr<UTexture2D>> SourceTextures;
+
+	/** 
+	* Is set to true if the source texture was generated from the SourceTextures array
+	* (which is not always the case, i.e. the source texture could be imported from a DDS file containing multiple slices).
+	* This transient property is used to control access to the SourceTextures array from UI using EditCondition mechanism
+	* (as any operation with the SourceTextures array would invalidate the originally imported source texture).
+	*/
+	UPROPERTY(Transient, SkipSerialization)
+	bool bSourceGeneratedFromSourceTexturesArray = true;
 #endif
 
 	/** Creates and initializes a new Texture2DArray with the requested settings */
diff --git a/Engine/Source/Runtime/Engine/Classes/Engine/TextureCubeArray.h b/Engine/Source/Runtime/Engine/Classes/Engine/TextureCubeArray.h
--- a/Engine/Source/Runtime/Engine/Classes/Engine/TextureCubeArray.h
+++ b/Engine/Source/Runtime/Engine/Classes/Engine/TextureCubeArray.h
@@ -65,8 +65,17 @@ public:
 
 #if WITH_EDITORONLY_DATA
 	/** Add Textures*/
-	UPROPERTY(EditAnywhere, Category = SourceCube, meta = (DisplayName = "Source Textures"))
+	UPROPERTY(EditAnywhere, Category = SourceCube, meta = (DisplayName = "Source Textures", EditCondition = bSourceGeneratedFromSourceTexturesArray, EditConditionHides, HideEditConditionToggle))
 	TArray<TObjectPtr<UTextureCube>> SourceTextures;
+
+	/**
+	* Is set to true if the source texture was generated from the SourceTextures array
+	* (which is not always the case, i.e. the source texture could be imported from a DDS file containing multiple cubemaps).
+	* This transient property is used to control access to the SourceTextures array from UI using EditCondition mechanism
+	* (as any operation with the SourceTextures array would invalidate the originally imported source texture).
+	*/
+	UPROPERTY(Transient, SkipSerialization)
+	bool bSourceGeneratedFromSourceTexturesArray = true;
 #endif
 
 	/** Creates and initializes a new TextureCubeArray with the requested settings */
diff --git a/Engine/Source/Runtime/Engine/Private/TextureCubeArray.cpp b/Engine/Source/Runtime/Engine/Private/TextureCubeArray.cpp
--- a/Engine/Source/Runtime/Engine/Private/TextureCubeArray.cpp
+++ b/Engine/Source/Runtime/Engine/Private/TextureCubeArray.cpp
@@ -454,6 +454,10 @@ void UTextureCubeArray::UpdateResource()
 	}
 #endif // #if WITH_EDITOR
 
+#if WITH_EDITORONLY_DATA
+	bSourceGeneratedFromSourceTexturesArray = !Source.GetNumSlices() || SourceTextures.Num();
+#endif
+
 	Super::UpdateResource();
 }
 
@@ -768,6 +772,10 @@ void UTextureCubeArray::PostLoad()
 		}
 	}
 #endif // #if WITH_EDITOR
+#if WITH_EDITORONLY_DATA
+	bSourceGeneratedFromSourceTexturesArray = !Source.GetNumSlices() || SourceTextures.Num();
+#endif
+
 	Super::PostLoad();
 };
 
