Make sure PrecomputedLightingBuffer is created when GPUScene isn't available

diff --git a/Engine/Source/Runtime/Engine/Private/SceneManagement.cpp b/Engine/Source/Runtime/Engine/Private/SceneManagement.cpp
--- a/Engine/Source/Runtime/Engine/Private/SceneManagement.cpp
+++ b/Engine/Source/Runtime/Engine/Private/SceneManagement.cpp
@@ -927,7 +927,10 @@ void InitializeSharedSamplerStates()
 
 void FLightCacheInterface::CreatePrecomputedLightingUniformBuffer_RenderingThread(ERHIFeatureLevel::Type FeatureLevel)
 {
-	if (bNeedsPrecomputedLightingUniformBuffer && (LightMap || ShadowMap))
+	const bool bPrecomputedLightingParametersFromGPUScene = UseGPUScene(GMaxRHIShaderPlatform, FeatureLevel) && bCanUsePrecomputedLightingParametersFromGPUScene;
+
+	// Only create UB when GPUScene isn't available
+	if (!bPrecomputedLightingParametersFromGPUScene && (LightMap || ShadowMap))
 	{
 		FPrecomputedLightingUniformParameters Parameters;
 		GetPrecomputedLightingParameters(FeatureLevel, Parameters, this);
diff --git a/Engine/Source/Runtime/Engine/Private/StaticMeshRender.cpp b/Engine/Source/Runtime/Engine/Private/StaticMeshRender.cpp
--- a/Engine/Source/Runtime/Engine/Private/StaticMeshRender.cpp
+++ b/Engine/Source/Runtime/Engine/Private/StaticMeshRender.cpp
@@ -2097,7 +2097,7 @@ FStaticMeshSceneProxy::FLODInfo::FLODInfo(const UStaticMeshComponent* InComponen
 			SetLightMap(MeshMapBuildData->LightMap);
 			SetShadowMap(MeshMapBuildData->ShadowMap);
 			SetResourceCluster(MeshMapBuildData->ResourceCluster);
-			bNeedsPrecomputedLightingUniformBuffer = false;
+			bCanUsePrecomputedLightingParametersFromGPUScene = true;
 			IrrelevantLights = MeshMapBuildData->IrrelevantLights;
 		}
 	}
@@ -2117,7 +2117,7 @@ FStaticMeshSceneProxy::FLODInfo::FLODInfo(const UStaticMeshComponent* InComponen
 					SetLightMap(MeshMapBuildData->LightMap);
 					SetShadowMap(MeshMapBuildData->ShadowMap);
 					SetResourceCluster(MeshMapBuildData->ResourceCluster);
-					bNeedsPrecomputedLightingUniformBuffer = false;
+					bCanUsePrecomputedLightingParametersFromGPUScene = true;
 					IrrelevantLights = MeshMapBuildData->IrrelevantLights;
 				}
 			}
@@ -2182,7 +2182,7 @@ FStaticMeshSceneProxy::FLODInfo::FLODInfo(const UStaticMeshComponent* InComponen
 				SetLightMap(MeshMapBuildData->LightMap);
 				SetShadowMap(MeshMapBuildData->ShadowMap);
 				SetResourceCluster(MeshMapBuildData->ResourceCluster);
-				bNeedsPrecomputedLightingUniformBuffer = false;
+				bCanUsePrecomputedLightingParametersFromGPUScene = true;
 				IrrelevantLights = MeshMapBuildData->IrrelevantLights;
 			}
 		}
diff --git a/Engine/Source/Runtime/Engine/Public/SceneManagement.h b/Engine/Source/Runtime/Engine/Public/SceneManagement.h
--- a/Engine/Source/Runtime/Engine/Public/SceneManagement.h
+++ b/Engine/Source/Runtime/Engine/Public/SceneManagement.h
@@ -905,8 +905,10 @@ public:
 	ENGINE_API FShadowMapInteraction GetShadowMapInteraction(ERHIFeatureLevel::Type InFeatureLevel) const;
 
 protected:
-	// If false, precomputed lighting parameters will be drawn from GPUScene
-	bool bNeedsPrecomputedLightingUniformBuffer = true;
+	// Load parameters from GPUScene when possible
+	// Basically this is the same as VF_SUPPORTS_PRIMITIVE_SCENE_DATA on the vertex factory, but we can't deduce automatically
+	// because we don't know about VF type until we see the actual mesh batch
+	bool bCanUsePrecomputedLightingParametersFromGPUScene = false;
 	
 private:
 
diff --git a/Engine/Source/Runtime/Landscape/Public/LandscapeRender.h b/Engine/Source/Runtime/Landscape/Public/LandscapeRender.h
--- a/Engine/Source/Runtime/Landscape/Public/LandscapeRender.h
+++ b/Engine/Source/Runtime/Landscape/Public/LandscapeRender.h
@@ -591,7 +591,7 @@ class LANDSCAPE_API FLandscapeComponentSceneProxy : public FPrimitiveSceneProxy,
 				SetLightMap(MapBuildData->LightMap);
 				SetShadowMap(MapBuildData->ShadowMap);
 				SetResourceCluster(MapBuildData->ResourceCluster);
-				bNeedsPrecomputedLightingUniformBuffer = false;
+				bCanUsePrecomputedLightingParametersFromGPUScene = true;
 				IrrelevantLights = MapBuildData->IrrelevantLights;
 			}
 		}
