Fix Strata shader compilation when normal is not a float or float3 vector.

diff --git a/Engine/Source/Runtime/Engine/Private/Materials/HLSLMaterialTranslator.cpp b/Engine/Source/Runtime/Engine/Private/Materials/HLSLMaterialTranslator.cpp
--- a/Engine/Source/Runtime/Engine/Private/Materials/HLSLMaterialTranslator.cpp
+++ b/Engine/Source/Runtime/Engine/Private/Materials/HLSLMaterialTranslator.cpp
@@ -8290,7 +8290,7 @@ int32 FHLSLMaterialTranslator::TransformBase(EMaterialCommonBasis SourceCoordBas
 		}
 		else if (NumInputComponents < 3u)
 		{
-			return Errorf(TEXT("input must be a vector (%s: %s) or a scalar (if source is Tangent)"), *GetParameterCode(A), DescribeType(SourceType));
+			return Errorf(TEXT("input must be a 3-component vector (current: %s: %s) or a scalar (if source is Tangent)"), *GetParameterCode(A), DescribeType(SourceType));
 		}
 	}
 		
diff --git a/Engine/Source/Runtime/Engine/Private/Materials/MaterialExpressions.cpp b/Engine/Source/Runtime/Engine/Private/Materials/MaterialExpressions.cpp
--- a/Engine/Source/Runtime/Engine/Private/Materials/MaterialExpressions.cpp
+++ b/Engine/Source/Runtime/Engine/Private/Materials/MaterialExpressions.cpp
@@ -21629,6 +21629,13 @@ int32 UMaterialExpressionStrataLegacyConversion::Compile(class FMaterialCompiler
 
 	// Regular normal basis
 	int32 NormalCodeChunk = CompileWithDefaultNormalWS(Compiler, Normal);
+
+	// When computing NormalCodeChunk, we invoke TransformNormalFromRequestedBasisToWorld which requires input to be float or float3.
+	// Certain material do not respect this requirement. We handle here a simple recovery when source material doesn't have a valid 
+	// normal (e.g., vec2 normal), and avoid crashing the material compilation. The error will still be reported by the compiler up 
+	// to the user, but the compilation will succeed.
+	if (NormalCodeChunk == INDEX_NONE) { NormalCodeChunk = Compiler->VertexNormal(); } 
+
 	int32 TangentCodeChunk = bHasAnisotropy ? CompileWithDefaultTangentWS(Compiler, Tangent) : INDEX_NONE;
 	const FStrataRegisteredSharedLocalBasis NewRegisteredSharedLocalBasis = StrataCompilationInfoCreateSharedLocalBasis(Compiler, NormalCodeChunk, TangentCodeChunk);
 	const FString BasisIndexMacro = Compiler->GetStrataSharedLocalBasisIndexMacro(NewRegisteredSharedLocalBasis);
