Fix UI input material when strata is enabled.

diff --git a/Engine/Source/Runtime/Engine/Private/Materials/Material.cpp b/Engine/Source/Runtime/Engine/Private/Materials/Material.cpp
--- a/Engine/Source/Runtime/Engine/Private/Materials/Material.cpp
+++ b/Engine/Source/Runtime/Engine/Private/Materials/Material.cpp
@@ -6106,6 +6106,15 @@ bool UMaterial::IsPropertySupported(EMaterialProperty InProperty) const
 		case MP_FrontMaterial:
 			bSupported = true;
 			break;
+		// UI domain uses the following remapping:
+		// * MP_EmissiveColor       -> Final Color
+		// * MP_WorldPositionOffset -> Screen Position
+		// * MP_OpacityMask         -> Opacity Mask
+		// * MP_Opacity             -> Opacity
+		case MP_EmissiveColor:
+		case MP_Opacity:
+			bSupported = MaterialDomain == EMaterialDomain::MD_UI;
+			break;
 		}
 	}
 	return bSupported;
