Remove some non-inclusive language.
Deprecates AddViewSlaveLocation() for AddViewLocation() and AddTextureStreamingSlaveLoc() for AddTextureStreamingLoc().

diff --git a/Engine/Plugins/Experimental/VirtualHeightfieldMesh/Source/VirtualHeightfieldMeshEditor/Private/HeightfieldMinMaxTextureBuild.cpp b/Engine/Plugins/Experimental/VirtualHeightfieldMesh/Source/VirtualHeightfieldMeshEditor/Private/HeightfieldMinMaxTextureBuild.cpp
--- a/Engine/Plugins/Experimental/VirtualHeightfieldMesh/Source/VirtualHeightfieldMeshEditor/Private/HeightfieldMinMaxTextureBuild.cpp
+++ b/Engine/Plugins/Experimental/VirtualHeightfieldMesh/Source/VirtualHeightfieldMeshEditor/Private/HeightfieldMinMaxTextureBuild.cpp
@@ -183,7 +183,7 @@ namespace VirtualHeightfieldMesh
 				//todo[vt]: Batch groups of streaming locations and render commands to reduce number of flushes.
 				const FVector StreamingWorldPos = Transform.TransformPosition(FVector(UVRange.GetCenter(), 0.5f));
 				IStreamingManager::Get().Tick(0.f);
-				IStreamingManager::Get().AddViewSlaveLocation(StreamingWorldPos);
+				IStreamingManager::Get().AddViewLocation(StreamingWorldPos);
 				IStreamingManager::Get().StreamAllResources(0);
 
 				ENQUEUE_RENDER_COMMAND(MinMaxTextureTileCommand)([
diff --git a/Engine/Source/Editor/VirtualTexturingEditor/Private/RuntimeVirtualTextureBuildStreamingMips.cpp b/Engine/Source/Editor/VirtualTexturingEditor/Private/RuntimeVirtualTextureBuildStreamingMips.cpp
--- a/Engine/Source/Editor/VirtualTexturingEditor/Private/RuntimeVirtualTextureBuildStreamingMips.cpp
+++ b/Engine/Source/Editor/VirtualTexturingEditor/Private/RuntimeVirtualTextureBuildStreamingMips.cpp
@@ -212,7 +212,7 @@ namespace RuntimeVirtualTexture
 				//todo[vt]: Batch groups of streaming locations and render commands to reduce number of flushes.
 				const FVector StreamingWorldPos = Transform.TransformPosition(FVector(UVRange.GetCenter(), 0.5f));
 				IStreamingManager::Get().Tick(0.f);
-				IStreamingManager::Get().AddViewSlaveLocation(StreamingWorldPos);
+				IStreamingManager::Get().AddViewLocation(StreamingWorldPos);
 				IStreamingManager::Get().StreamAllResources(0);
 
 				ENQUEUE_RENDER_COMMAND(BakeStreamingTextureTileCommand)([
diff --git a/Engine/Source/Runtime/Engine/Classes/Engine/Engine.h b/Engine/Source/Runtime/Engine/Classes/Engine/Engine.h
--- a/Engine/Source/Runtime/Engine/Classes/Engine/Engine.h
+++ b/Engine/Source/Runtime/Engine/Classes/Engine/Engine.h
@@ -2629,7 +2629,13 @@ public:
 	 * @param bOverrideLocation		Whether this is an override location, which forces the streaming system to ignore all other locations
 	 * @param OverrideDuration		How long the streaming system should keep checking this location if bOverrideLocation is true, in seconds. 0 means just for the next Tick.
 	 */
-	void AddTextureStreamingSlaveLoc(FVector InLoc, float BoostFactor, bool bOverrideLocation, float OverrideDuration);
+	void AddTextureStreamingLoc(FVector InLoc, float BoostFactor, bool bOverrideLocation, float OverrideDuration);
+
+	UE_DEPRECATED(5.1, "This is deprecated to follow inclusive naming rules. Use AddTextureStreamingLoc() instead.")
+	void AddTextureStreamingSlaveLoc(FVector InLoc, float BoostFactor, bool bOverrideLocation, float OverrideDuration)
+	{
+		AddTextureStreamingLoc(InLoc, BoostFactor, bOverrideLocation, OverrideDuration);
+	}
 
 	/** 
 	 * Obtain a world object pointer from an object with has a world context.
diff --git a/Engine/Source/Runtime/Engine/Private/ContentStreaming.cpp b/Engine/Source/Runtime/Engine/Private/ContentStreaming.cpp
--- a/Engine/Source/Runtime/Engine/Private/ContentStreaming.cpp
+++ b/Engine/Source/Runtime/Engine/Private/ContentStreaming.cpp
@@ -56,7 +56,7 @@ TArray<FStreamingViewInfo> IStreamingManager::PendingViewInfos;
 TArray<FStreamingViewInfo> IStreamingManager::LastingViewInfos;
 
 /** Collection of view locations that will be added at the next call to AddViewInformation. */
-TArray<IStreamingManager::FSlaveLocation> IStreamingManager::SlaveLocations;
+TArray<IStreamingManager::FSecondaryLocation> IStreamingManager::SecondaryLocations;
 
 /** Set when Tick() has been called. The first time a new view is added, it will clear out all old views. */
 bool IStreamingManager::bPendingRemoveViews = false;
@@ -464,7 +464,7 @@ TArray<FStreamingViewInfo> GPrevViewLocations;
 #endif
 
 /**
- * Sets up the CurrentViewInfos array based on PendingViewInfos, LastingViewInfos and SlaveLocations.
+ * Sets up the CurrentViewInfos array based on PendingViewInfos, LastingViewInfos and SecondaryLocations.
  * Removes out-dated LastingViewInfos.
  *
  * @param DeltaTime		Time since last call in seconds
@@ -472,7 +472,7 @@ TArray<FStreamingViewInfo> GPrevViewLocations;
 void IStreamingManager::SetupViewInfos( float DeltaTime )
 {
 	// Reset CurrentViewInfos
-	CurrentViewInfos.Empty( PendingViewInfos.Num() + LastingViewInfos.Num() + SlaveLocations.Num() );
+	CurrentViewInfos.Empty( PendingViewInfos.Num() + LastingViewInfos.Num() + SecondaryLocations.Num() );
 
 	bool bHaveMultiplePlayerViews = (PendingViewInfos.Num() > 1) ? true : false;
 
@@ -492,14 +492,14 @@ void IStreamingManager::SetupViewInfos( float DeltaTime )
 
 	// Add them to the appropriate array (pending views or lasting views).
 	{
-		// Disable this flag as it could be used in AddViewInformation to empty SlaveLocation.
+		// Disable this flag as it could be used in AddViewInformation to empty SecondaryLocation.
 		const bool bPendingRemoveViewsBackup = bPendingRemoveViews;
 		bPendingRemoveViews = false;
 
-		for ( int32 SlaveLocationIndex=0; SlaveLocationIndex < SlaveLocations.Num(); SlaveLocationIndex++ )
+		for ( int32 SecondaryLocationIndex=0; SecondaryLocationIndex < SecondaryLocations.Num(); SecondaryLocationIndex++ )
 		{
-			const FSlaveLocation& SlaveLocation = SlaveLocations[ SlaveLocationIndex ];
-			AddViewInformation( SlaveLocation.Location, ScreenSize, FOVScreenSize, SlaveLocation.BoostFactor, SlaveLocation.bOverrideLocation, SlaveLocation.Duration );
+			const FSecondaryLocation& SecondaryLocation = SecondaryLocations[ SecondaryLocationIndex ];
+			AddViewInformation( SecondaryLocation.Location, ScreenSize, FOVScreenSize, SecondaryLocation.BoostFactor, SecondaryLocation.bOverrideLocation, SecondaryLocation.Duration );
 		}
 
 		bPendingRemoveViews = bPendingRemoveViewsBackup;
@@ -646,7 +646,7 @@ void IStreamingManager::AddViewInformation( const FVector& ViewOrigin, float Scr
 		{
 			bPendingRemoveViews = false;
 
-			// Remove out-dated override views and empty the PendingViewInfos/SlaveLocation arrays to be populated again during next frame.
+			// Remove out-dated override views and empty the PendingViewInfos/SecondaryLocation arrays to be populated again during next frame.
 			RemoveStreamingViews( RemoveStreamingViews_Normal );
 		}
 
@@ -673,15 +673,15 @@ void IStreamingManager::AddViewInformation( const FVector& ViewOrigin, float Scr
 }
 
 /**
- * Queue up view "slave" locations to the streaming system. These locations will be added properly at the next call to AddViewInformation,
+ * Queue up view locations to the streaming system. These locations will be added properly at the next call to AddViewInformation,
  * re-using the screensize and FOV settings.
  *
- * @param SlaveLocation			World-space view origin
+ * @param Location				World-space view origin
  * @param BoostFactor			A factor that affects all streaming distances for this location. 1.0f is default. Higher means higher-resolution textures and vice versa.
  * @param bOverrideLocation		Whether this is an override location, which forces the streaming system to ignore all other locations
  * @param Duration				How long the streaming system should keep checking this location (in seconds). 0 means just for the next Tick.
  */
-void IStreamingManager::AddViewSlaveLocation( const FVector& SlaveLocation, float BoostFactor/*=1.0f*/, bool bOverrideLocation/*=false*/, float Duration/*=0.0f*/ )
+void IStreamingManager::AddViewLocation( const FVector& Location, float BoostFactor/*=1.0f*/, bool bOverrideLocation/*=false*/, float Duration/*=0.0f*/ )
 {
 	const float MinBoost = CVarStreamingMinBoost.GetValueOnGameThread();
 	const float BoostScale = FMath::Max(CVarStreamingBoost.GetValueOnGameThread(),MinBoost);
@@ -691,11 +691,11 @@ void IStreamingManager::AddViewSlaveLocation( const FVector& SlaveLocation, floa
 	{
 		bPendingRemoveViews = false;
 
-		// Remove out-dated override views and empty the PendingViewInfos/SlaveLocation arrays to be populated again during next frame.
+		// Remove out-dated override views and empty the PendingViewInfos/SecondaryLocation arrays to be populated again during next frame.
 		RemoveStreamingViews( RemoveStreamingViews_Normal );
 	}
 
-	new (SlaveLocations) FSlaveLocation(SlaveLocation, BoostFactor, bOverrideLocation, Duration );
+	new (SecondaryLocations) FSecondaryLocation(Location, BoostFactor, bOverrideLocation, Duration );
 }
 
 /**
@@ -706,7 +706,7 @@ void IStreamingManager::AddViewSlaveLocation( const FVector& SlaveLocation, floa
 void IStreamingManager::RemoveStreamingViews( ERemoveStreamingViews RemovalType )
 {
 	PendingViewInfos.Empty();
-	SlaveLocations.Empty();
+	SecondaryLocations.Empty();
 	if ( RemovalType == RemoveStreamingViews_All )
 	{
 		LastingViewInfos.Empty();
diff --git a/Engine/Source/Runtime/Engine/Private/Interpolation.cpp b/Engine/Source/Runtime/Engine/Private/Interpolation.cpp
--- a/Engine/Source/Runtime/Engine/Private/Interpolation.cpp
+++ b/Engine/Source/Runtime/Engine/Private/Interpolation.cpp
@@ -770,7 +770,7 @@ void AMatineeActor::UpdateStreamingForCameraCuts(float CurrentTime, bool bPrevie
 			float TimeDifference = CutInfo.TimeStamp - CurrentTime;
 			if ( TimeDifference > 0.0f && TimeDifference < GCameraCutLookAhead )
 			{
-				IStreamingManager::Get().AddViewSlaveLocation( CutInfo.Location );
+				IStreamingManager::Get().AddViewLocation( CutInfo.Location );
 			}
 			else if ( TimeDifference >= GCameraCutLookAhead )
 			{
diff --git a/Engine/Source/Runtime/Engine/Private/PlayerController.cpp b/Engine/Source/Runtime/Engine/Private/PlayerController.cpp
--- a/Engine/Source/Runtime/Engine/Private/PlayerController.cpp
+++ b/Engine/Source/Runtime/Engine/Private/PlayerController.cpp
@@ -417,7 +417,7 @@ void APlayerController::ClientAddTextureStreamingLoc_Implementation(FVector InLo
 {
 	if (!IStreamingManager::HasShutdown())
 	{
-		IStreamingManager::Get().AddViewSlaveLocation(InLoc, 1.0f, bOverrideLocation, Duration);
+		IStreamingManager::Get().AddViewLocation(InLoc, 1.0f, bOverrideLocation, Duration);
 	}
 }
 
diff --git a/Engine/Source/Runtime/Engine/Private/UnrealEngine.cpp b/Engine/Source/Runtime/Engine/Private/UnrealEngine.cpp
--- a/Engine/Source/Runtime/Engine/Private/UnrealEngine.cpp
+++ b/Engine/Source/Runtime/Engine/Private/UnrealEngine.cpp
@@ -10690,17 +10690,17 @@ void UEngine::EnableScreenSaver( bool bEnable )
 }
 
 /**
-* Queue up view "slave" locations to the streaming system. These locations will be added properly at the next call to AddViewInformation,
-* re-using the screensize and FOV settings.
+* Adds a world location as a secondary view location for purposes of texture streaming.
+* Lasts one frame, or a specified number of seconds (for overriding locations only).
 *
-* @param SlaveLocation			World-space view origin
+* @param InLoc					World-space view origin
 * @param BoostFactor			A factor that affects all streaming distances for this location. 1.0f is default. Higher means higher-resolution textures and vice versa.
 * @param bOverrideLocation		Whether this is an override location, which forces the streaming system to ignore all other locations
 * @param OverrideDuration		How long the streaming system should keep checking this location if bOverrideLocation is true, in seconds. 0 means just for the next Tick.
 */
-void UEngine::AddTextureStreamingSlaveLoc(FVector InLoc, float BoostFactor, bool bOverrideLocation, float OverrideDuration)
+void UEngine::AddTextureStreamingLoc(FVector InLoc, float BoostFactor, bool bOverrideLocation, float OverrideDuration)
 {
-	IStreamingManager::Get().AddViewSlaveLocation(InLoc, BoostFactor, bOverrideLocation, OverrideDuration);
+	IStreamingManager::Get().AddViewLocation(InLoc, BoostFactor, bOverrideLocation, OverrideDuration);
 }
 
 /** Looks up the GUID of a package on disk. The package must NOT be in the autodownload cache.
diff --git a/Engine/Source/Runtime/Engine/Public/ContentStreaming.h b/Engine/Source/Runtime/Engine/Public/ContentStreaming.h
--- a/Engine/Source/Runtime/Engine/Public/ContentStreaming.h
+++ b/Engine/Source/Runtime/Engine/Public/ContentStreaming.h
@@ -234,15 +234,21 @@ struct IStreamingManager
 	ENGINE_API void AddViewInformation(const FVector& ViewOrigin, float ScreenSize, float FOVScreenSize, float BoostFactor = 1.0f, bool bOverrideLocation = false, float Duration = 0.0f, TWeakObjectPtr<AActor> InActorToBoost = NULL);
 
 	/**
-	 * Queue up view "slave" locations to the streaming system. These locations will be added properly at the next call to AddViewInformation,
+	 * Queue up view locations to the streaming system. These locations will be added properly at the next call to AddViewInformation,
 	 * re-using the screensize and FOV settings.
 	 *
-	 * @param SlaveLocation			World-space view origin
+	 * @param Location				World-space view origin
 	 * @param BoostFactor			A factor that affects all streaming distances for this location. 1.0f is default. Higher means higher-resolution textures and vice versa.
 	 * @param bOverrideLocation		Whether this is an override location, which forces the streaming system to ignore all other locations
 	 * @param Duration				How long the streaming system should keep checking this location, in seconds. 0 means just for the next Tick.
 	 */
-	ENGINE_API void AddViewSlaveLocation(const FVector& SlaveLocation, float BoostFactor = 1.0f, bool bOverrideLocation = false, float Duration = 0.0f);
+	ENGINE_API void AddViewLocation(const FVector& Location, float BoostFactor = 1.0f, bool bOverrideLocation = false, float Duration = 0.0f);
+
+	UE_DEPRECATED(5.1, "This is deprecated to follow inclusive naming rules. Use AddViewLocation() instead.")
+	ENGINE_API void AddViewSlaveLocation(const FVector& Location, float BoostFactor = 1.0f, bool bOverrideLocation = false, float Duration = 0.0f)
+	{
+		AddViewLocation(Location, BoostFactor, bOverrideLocation, Duration);
+	}
 
 	/** Don't stream world resources for the next NumFrames. */
 	virtual void SetDisregardWorldResourcesForFrames(int32 NumFrames) = 0;
@@ -332,7 +338,7 @@ struct IStreamingManager
 protected:
 
 	/**
-	 * Sets up the CurrentViewInfos array based on PendingViewInfos, LastingViewInfos and SlaveLocations.
+	 * Sets up the CurrentViewInfos array based on PendingViewInfos, LastingViewInfos and SecondaryLocations.
 	 * Removes out-dated LastingViewInfos.
 	 *
 	 * @param DeltaTime		Time since last call in seconds
@@ -361,9 +367,9 @@ protected:
 	 */
 	static void RemoveViewInfoFromArray( TArray<FStreamingViewInfo> &ViewInfos, const FVector& ViewOrigin );
 
-	struct FSlaveLocation
+	struct FSecondaryLocation
 	{
-		FSlaveLocation( const FVector& InLocation, float InBoostFactor, bool bInOverrideLocation, float InDuration )
+		FSecondaryLocation( const FVector& InLocation, float InBoostFactor, bool bInOverrideLocation, float InDuration )
 		:	Location( InLocation )
 		,	BoostFactor( InBoostFactor )
 		,	Duration( InDuration )
@@ -390,7 +396,7 @@ protected:
 	static TArray<FStreamingViewInfo> LastingViewInfos;
 
 	/** Collection of view locations that will be added at the next call to AddViewInformation. */
-	static TArray<FSlaveLocation> SlaveLocations;
+	static TArray<FSecondaryLocation> SecondaryLocations;
 
 	/** Set when Tick() has been called. The first time a new view is added, it will clear out all old views. */
 	static bool bPendingRemoveViews;
diff --git a/Engine/Source/Runtime/MovieSceneTracks/Private/TrackInstances/MovieSceneCameraCutTrackInstance.cpp b/Engine/Source/Runtime/MovieSceneTracks/Private/TrackInstances/MovieSceneCameraCutTrackInstance.cpp
--- a/Engine/Source/Runtime/MovieSceneTracks/Private/TrackInstances/MovieSceneCameraCutTrackInstance.cpp
+++ b/Engine/Source/Runtime/MovieSceneTracks/Private/TrackInstances/MovieSceneCameraCutTrackInstance.cpp
@@ -146,7 +146,7 @@ namespace MovieScene
 			if (Params.bHasCutTransform)
 			{
 				FVector Location = Params.CutTransform.GetLocation();
-				IStreamingManager::Get().AddViewSlaveLocation(Location);
+				IStreamingManager::Get().AddViewLocation(Location);
 			}
 			else
 			{
@@ -155,7 +155,7 @@ namespace MovieScene
 				if (AActor* Actor = Cast<AActor>(CameraObject))
 				{
 					FVector Location = Actor->GetActorLocation();
-					IStreamingManager::Get().AddViewSlaveLocation(Location);
+					IStreamingManager::Get().AddViewLocation(Location);
 				}
 			}
 		}
