Refactor skeletal mesh DDC async compile logic to be generic for skinned asset type. Added async render data compilation for cloth asset.

diff --git a/Engine/Source/Developer/SkeletalMeshUtilitiesCommon/Private/LODUtilities.cpp b/Engine/Source/Developer/SkeletalMeshUtilitiesCommon/Private/LODUtilities.cpp
--- a/Engine/Source/Developer/SkeletalMeshUtilitiesCommon/Private/LODUtilities.cpp
+++ b/Engine/Source/Developer/SkeletalMeshUtilitiesCommon/Private/LODUtilities.cpp
@@ -2775,10 +2775,10 @@ void FLODUtilities::RegenerateDependentLODs(USkeletalMesh* SkeletalMesh, int32 L
 			}
 
 			// Reduce LODs in parallel (reduction is multithread safe)
-			const bool bHasAccessToLockedProperties = !FSkeletalMeshAsyncBuildScope::ShouldWaitOnLockedProperties(SkeletalMesh);
+			const bool bHasAccessToLockedProperties = !FSkinnedAssetAsyncBuildScope::ShouldWaitOnLockedProperties(SkeletalMesh);
 			ParallelFor(DependentLODs.Num(), [&DependentLODs, &SkeletalMesh, &bNeedsPackageDirtied, bHasAccessToLockedProperties, &TargetPlatform](int32 IterationIndex)
 			{
-				TUniquePtr<FSkeletalMeshAsyncBuildScope> AsyncBuildScope(bHasAccessToLockedProperties ? MakeUnique<FSkeletalMeshAsyncBuildScope>(SkeletalMesh) : nullptr);
+				TUniquePtr<FSkinnedAssetAsyncBuildScope> AsyncBuildScope(bHasAccessToLockedProperties ? MakeUnique<FSkinnedAssetAsyncBuildScope>(SkeletalMesh) : nullptr);
 
 				check(DependentLODs.IsValidIndex(IterationIndex));
 				int32 DependentLODIndex = DependentLODs[IterationIndex];
diff --git a/Engine/Source/Runtime/Engine/Classes/Engine/SkeletalMesh.h b/Engine/Source/Runtime/Engine/Classes/Engine/SkeletalMesh.h
--- a/Engine/Source/Runtime/Engine/Classes/Engine/SkeletalMesh.h
+++ b/Engine/Source/Runtime/Engine/Classes/Engine/SkeletalMesh.h
@@ -19,7 +19,6 @@
 #include "Rendering/SkeletalMeshRenderData.h"
 #include "Animation/PreviewAssetAttachComponent.h"
 #include "BoneContainer.h"
-#include "Interfaces/Interface_AsyncCompilation.h"
 #include "Interfaces/Interface_CollisionDataProvider.h"
 #include "EngineTypes.h"
 #include "Engine/Engine.h"
@@ -126,130 +125,15 @@ enum class ESkeletalMeshAsyncProperties : uint64
 
 ENUM_CLASS_FLAGS(ESkeletalMeshAsyncProperties);
 
-class FSkeletalMeshCompilationContext
-{
-public:
-	FSkeletalMeshCompilationContext() = default;
-	// Non-copyable
-	FSkeletalMeshCompilationContext(const FSkeletalMeshCompilationContext&) = delete;
-	FSkeletalMeshCompilationContext& operator=(const FSkeletalMeshCompilationContext&) = delete;
-	// Movable
-	FSkeletalMeshCompilationContext(FSkeletalMeshCompilationContext&&) = default;
-	FSkeletalMeshCompilationContext& operator=(FSkeletalMeshCompilationContext&&) = default;
-	
-	//True if this compilation context is start from a serialize save
-	bool bIsSerializeSaving = false;
-};
-
-class FSkeletalMeshPostLoadContext : public FSkeletalMeshCompilationContext
-{
-public:
-	bool bHasCachedDerivedData = false;
-};
-
-class FSkeletalMeshBuildContext : public FSkeletalMeshCompilationContext
-{
-public:
-	TUniquePtr<FSkinnedMeshComponentRecreateRenderStateContext> RecreateRenderStateContext;
-};
 
+using FSkeletalMeshCompilationContext UE_DEPRECATED(5.1, "Use FSkinnedAssetCompilationContext instead.") = FSkinnedAssetCompilationContext;
+using FSkeletalMeshPostLoadContext UE_DEPRECATED(5.1, "Use FSkinnedAssetPostLoadContext instead.") = FSkinnedAssetPostLoadContext;
+using FSkeletalMeshBuildContext UE_DEPRECATED(5.1, "Use FSkinnedAssetBuildContext instead.") = FSkinnedAssetBuildContext;
 #if WITH_EDITOR
-
-// Any thread implicated in the build must have a valid scope to be granted access to protected properties without causing any stalls.
-class FSkeletalMeshAsyncBuildScope
-{
-public:
-	ENGINE_API FSkeletalMeshAsyncBuildScope(const USkeletalMesh* SkeletalMesh);
-	ENGINE_API ~FSkeletalMeshAsyncBuildScope();
-	ENGINE_API static bool ShouldWaitOnLockedProperties(const USkeletalMesh* SkeletalMesh);
-
-private:
-	const USkeletalMesh* PreviousScope = nullptr;
-	// Only the thread(s) compiling this Skeletal mesh will have full access to protected properties without causing any stalls.
-	static thread_local const USkeletalMesh* SkeletalMeshBeingAsyncCompiled;
-};
-
-/**
- * Worker used to perform async static mesh compilation.
- */
-class FSkeletalMeshAsyncBuildWorker : public FNonAbandonableTask
-{
-public:
-	USkeletalMesh* SkeletalMesh;
-	TOptional<FSkeletalMeshPostLoadContext> PostLoadContext;
-	TOptional<FSkeletalMeshBuildContext> BuildContext;
-
-	/** Initialization constructor. */
-	FSkeletalMeshAsyncBuildWorker(
-		USkeletalMesh* InSkeletalMesh,
-		FSkeletalMeshBuildContext&& InBuildContext)
-		: SkeletalMesh(InSkeletalMesh)
-		, BuildContext(MoveTemp(InBuildContext))
-	{
-	}
-
-	/** Initialization constructor. */
-	FSkeletalMeshAsyncBuildWorker(
-		USkeletalMesh* InSkeletalMesh,
-		FSkeletalMeshPostLoadContext&& InPostLoadContext)
-		: SkeletalMesh(InSkeletalMesh)
-		, PostLoadContext(MoveTemp(InPostLoadContext))
-	{
-	}
-
-	FORCEINLINE TStatId GetStatId() const
-	{
-		RETURN_QUICK_DECLARE_CYCLE_STAT(FSkeletalMeshAsyncBuildWorker, STATGROUP_ThreadPoolAsyncTasks);
-	}
-
-	void DoWork();
-};
-
-struct FSkeletalMeshAsyncBuildTask : public FAsyncTask<FSkeletalMeshAsyncBuildWorker>
-{
-	FSkeletalMeshAsyncBuildTask(
-		USkeletalMesh* InSkeletalMesh,
-		FSkeletalMeshPostLoadContext&& InPostLoadContext)
-		: FAsyncTask<FSkeletalMeshAsyncBuildWorker>(InSkeletalMesh, MoveTemp(InPostLoadContext))
-		, SkeletalMesh(InSkeletalMesh)
-	{
-	}
-
-	FSkeletalMeshAsyncBuildTask(
-		USkeletalMesh* InSkeletalMesh,
-		FSkeletalMeshBuildContext&& InBuildContext)
-		: FAsyncTask<FSkeletalMeshAsyncBuildWorker>(InSkeletalMesh, MoveTemp(InBuildContext))
-		, SkeletalMesh(InSkeletalMesh)
-	{
-	}
-
-	const USkeletalMesh* SkeletalMesh;
-};
-
-#endif // #if WITH_EDITOR
-
-UENUM()
-enum class ESkinCacheUsage : uint8
-{
-	// Auto will defer to child or global behavior based on context. If Support Ray Tracing is enabled on the mesh, will imply Enabled
-	Auto		= 0,
-
-	// Mesh will not use the skin cache. If Support Ray Tracing is enabled on the mesh, will imply Enabled
-	Disabled	= uint8(-1),
-
-	// Mesh will use the skin cache
-	Enabled		= 1,
-};
-
-UENUM()
-enum class ESkinCacheDefaultBehavior : uint8
-{
-	// All skeletal meshes are excluded from the skin cache. Each must opt in individually. If Support Ray Tracing is enabled on a mesh, will force inclusive behavior on that mesh
-	Exclusive = 0,
-
-	// All skeletal meshes are included into the skin cache. Each must opt out individually
-	Inclusive = 1,
-};
+using FSkeletalMeshAsyncBuildScope UE_DEPRECATED(5.1, "Use FSkinnedAssetAsyncBuildScope instead.") = FSkinnedAssetAsyncBuildScope;
+using FSkeletalMeshAsyncBuildWorker UE_DEPRECATED(5.1, "Use FSkinnedAssetAsyncBuildWorker instead.") = FSkinnedAssetAsyncBuildWorker;
+using FSkeletalMeshAsyncBuildTask UE_DEPRECATED(5.1, "Use FSkinnedAssetAsyncBuildTask instead.") = FSkinnedAssetAsyncBuildTask;
+#endif
 
 struct UE_DEPRECATED(5.0, "FBoneMirrorInfo is deprecated. Please use UMirrorDataTable for mirroring support.") FBoneMirrorInfo;
 USTRUCT()
@@ -338,187 +222,6 @@ struct ENGINE_API FSkeletalMeshClothBuildParams
 	TSoftObjectPtr<UPhysicsAsset> PhysicsAsset;
 };
 
-USTRUCT()
-struct FSectionReference
-{
-	GENERATED_USTRUCT_BODY()
-
-	/** Index of the section we reference. **/
-	UPROPERTY(EditAnywhere, Category = SectionReference)
-	int32 SectionIndex;
-
-	FSectionReference()
-		: SectionIndex(INDEX_NONE)
-	{
-	}
-
-	FSectionReference(const int32& InSectionIndex)
-		: SectionIndex(InSectionIndex)
-	{
-	}
-
-	bool operator==(const FSectionReference& Other) const
-	{
-		return SectionIndex == Other.SectionIndex;
-	}
-
-#if WITH_EDITOR
-	/** return true if it has a valid section index for LodModel parameter **/
-	ENGINE_API bool IsValidToEvaluate(const FSkeletalMeshLODModel& LodModel) const;
-
-	const struct FSkelMeshSection* GetMeshLodSection(const FSkeletalMeshLODModel& LodModel) const;
-	int32 GetMeshLodSectionIndex(const FSkeletalMeshLODModel& LodModel) const;
-#endif
-
-	bool Serialize(FArchive& Ar)
-	{
-		Ar << SectionIndex;
-		return true;
-	}
-
-	friend FArchive& operator<<(FArchive& Ar, FSectionReference& B)
-	{
-		B.Serialize(Ar);
-		return Ar;
-	}
-};
-
-/** Struct containing information for a particular LOD level, such as materials and info for when to use it. */
-USTRUCT()
-struct FSkeletalMeshLODInfo
-{
-	GENERATED_USTRUCT_BODY()
-
-	/** 
-	 * ScreenSize to display this LOD.
-	 * The screen size is based around the projected diameter of the bounding
-	 * sphere of the model. i.e. 0.5 means half the screen's maximum dimension.
-	 */
-	UPROPERTY(EditAnywhere, Category=SkeletalMeshLODInfo)
-	FPerPlatformFloat ScreenSize;
-
-	/**	Used to avoid 'flickering' when on LOD boundary. Only taken into account when moving from complex->simple. */
-	UPROPERTY(EditAnywhere, Category=SkeletalMeshLODInfo, meta=(DisplayName="LOD Hysteresis"))
-	float LODHysteresis;
-
-	/** Mapping table from this LOD's materials to the USkeletalMesh materials array.
-	 * section index is the key
-	 * remapped material index is the value, can be INDEX_NONE for no remapping
-	 */
-	UPROPERTY()
-	TArray<int32> LODMaterialMap;
-
-#if WITH_EDITORONLY_DATA
-	/** Per-section control over whether to enable shadow casting. */
-	UPROPERTY()
-	TArray<bool> bEnableShadowCasting_DEPRECATED;
-
-	/** This has been removed in editor. We could re-apply this in import time or by mesh reduction utilities*/
-	UPROPERTY()
-	TArray<FName> RemovedBones_DEPRECATED;
-#endif
-
-	/** build settings to apply when building render data. */
-	UPROPERTY(EditAnywhere, Category = BuildSettings)
-	FSkeletalMeshBuildSettings BuildSettings;
-
-	/** Reduction settings to apply when building render data. */
-	UPROPERTY(EditAnywhere, Category = ReductionSettings)
-	FSkeletalMeshOptimizationSettings ReductionSettings;
-
-	/** Bones which should be removed from the skeleton for the LOD level */
-	UPROPERTY(EditAnywhere, Category = ReductionSettings)
-	TArray<FBoneReference> BonesToRemove;
-
-	/** Bones which should be prioritized for the quality, this will be weighted toward keeping source data. Use WeightOfPrioritization to control the value. */
-	UPROPERTY(EditAnywhere, Category = ReductionSettings)
-	TArray<FBoneReference> BonesToPrioritize;
-
-	/** Sections which should be prioritized for the quality, this will be weighted toward keeping source data. Use WeightOfPrioritization to control the value. */
-	UPROPERTY(EditAnywhere, Category = ReductionSettings)
-	TArray<FSectionReference> SectionsToPrioritize;
-
-	/** How much to consideration to give BonesToPrioritize and SectionsToPrioritize.  The weight is an additional vertex simplification penalty where 0 means nothing. */
-	UPROPERTY(EditAnywhere, Category = ReductionSettings, meta = (UIMin = "0.0", ClampMin = "0.0"))
-	float WeightOfPrioritization;
-
-	/** Pose which should be used to reskin vertex influences for which the bones will be removed in this LOD level, uses ref-pose by default */
-	UPROPERTY(EditAnywhere, Category = ReductionSettings)
-	TObjectPtr<UAnimSequence> BakePose;
-
-	/** This is used when you are sharing the LOD settings, but you'd like to override the BasePose. This precedes prior to BakePose*/
-	UPROPERTY(EditAnywhere, Category = ReductionSettings)
-	TObjectPtr<UAnimSequence> BakePoseOverride;
-
-	/** The filename of the file tha was used to import this LOD if it was not auto generated. */
-	UPROPERTY(VisibleAnywhere, Category= SkeletalMeshLODInfo, AdvancedDisplay)
-	FString SourceImportFilename;
-
-	/**
-	 * How this LOD uses the skin cache feature. Auto will defer to the default project global option. If Support Ray Tracing is enabled on the mesh, will imply Enabled
-	 */
-	UPROPERTY(EditAnywhere, Category = SkeletalMeshLODInfo)
-	ESkinCacheUsage SkinCacheUsage = ESkinCacheUsage::Auto;
-
-	/** The Morph target position error tolerance in microns. Larger values result in better compression and lower memory footprint, but also lower quality. */
-	UPROPERTY(EditAnywhere, Category = SkeletalMeshLODInfo, meta = (UIMin = "0.01", ClampMin = "0.01", UIMax = "10000.0", ClampMax = "10000.0"))
-	float MorphTargetPositionErrorTolerance = 20.0f;
-
-	/** Whether to disable morph targets for this LOD. */
-	UPROPERTY()
-	uint8 bHasBeenSimplified:1;
-
-	UPROPERTY()
-	uint8 bHasPerLODVertexColors : 1;
-
-	/** Keeps this LODs data on the CPU so it can be used for things such as sampling in FX. */
-	UPROPERTY(EditAnywhere, Category = SkeletalMeshLODInfo)
-	uint8 bAllowCPUAccess : 1;
-
-	/**
-	Mesh supports uniformly distributed sampling in constant time.
-	Memory cost is 8 bytes per triangle.
-	Example usage is uniform spawning of particles.
-	*/
-	UPROPERTY(EditAnywhere, AdvancedDisplay, Category = SkeletalMeshLODInfo, meta=(EditCondition="bAllowCPUAccess"))
-	uint8 bSupportUniformlyDistributedSampling : 1;
-
-#if WITH_EDITORONLY_DATA
-	/*
-	 * This boolean specify if the LOD was imported with the base mesh or not.
-	 */
-	UPROPERTY()
-	uint8 bImportWithBaseMesh:1;
-
-	//Temporary build GUID data
-	//We use this GUID to store the LOD Key so we can know if the LOD needs to be rebuilt
-	//This GUID is set when we Cache the render data (build function)
-	FGuid BuildGUID;
-
-	ENGINE_API FGuid ComputeDeriveDataCacheKey(const FSkeletalMeshLODGroupSettings* SkeletalMeshLODGroupSettings);
-#endif
-
-	FSkeletalMeshLODInfo()
-		: ScreenSize(1.0)
-		, LODHysteresis(0.0f)
-		, WeightOfPrioritization(1.f)
-		, BakePose(nullptr)
-		, BakePoseOverride(nullptr)
-		, bHasBeenSimplified(false)
-		, bHasPerLODVertexColors(false)
-		, bAllowCPUAccess(false)
-		, bSupportUniformlyDistributedSampling(false)
-#if WITH_EDITORONLY_DATA
-		, bImportWithBaseMesh(false)
-#endif
-	{
-#if WITH_EDITORONLY_DATA
-		BuildGUID.Invalidate();
-#endif
-	}
-
-};
-
 /**
  * Legacy object for back-compat loading, no longer used by clothing system
  */
@@ -638,70 +341,6 @@ enum class EClothLODBiasMode : uint8
 	MappingsToAnyLOD,
 };
 
-//~ Begin Material Interface for USkeletalMesh - contains a material and a shadow casting flag
-USTRUCT(BlueprintType)
-struct FSkeletalMaterial
-{
-	GENERATED_USTRUCT_BODY()
-
-	FSkeletalMaterial()
-		: MaterialInterface( NULL )
-		, MaterialSlotName( NAME_None )
-#if WITH_EDITORONLY_DATA
-		, bEnableShadowCasting_DEPRECATED(true)
-		, bRecomputeTangent_DEPRECATED(false)
-		, ImportedMaterialSlotName( NAME_None )
-#endif
-	{
-
-	}
-
-	FSkeletalMaterial( class UMaterialInterface* InMaterialInterface
-						, bool bInEnableShadowCasting = true
-						, bool bInRecomputeTangent = false
-						, FName InMaterialSlotName = NAME_None
-						, FName InImportedMaterialSlotName = NAME_None)
-		: MaterialInterface( InMaterialInterface )
-		, MaterialSlotName(InMaterialSlotName)
-#if WITH_EDITORONLY_DATA
-		, bEnableShadowCasting_DEPRECATED(bInEnableShadowCasting)
-		, bRecomputeTangent_DEPRECATED(bInRecomputeTangent)
-		, ImportedMaterialSlotName(InImportedMaterialSlotName)
-#endif //WITH_EDITORONLY_DATA
-	{
-
-	}
-
-	friend FArchive& operator<<( FArchive& Ar, FSkeletalMaterial& Elem );
-#if WITH_EDITORONLY_DATA
-	static void DeclareCustomVersions(FArchive& Ar);
-#endif
-
-	ENGINE_API friend bool operator==( const FSkeletalMaterial& LHS, const FSkeletalMaterial& RHS );
-	ENGINE_API friend bool operator==( const FSkeletalMaterial& LHS, const UMaterialInterface& RHS );
-	ENGINE_API friend bool operator==( const UMaterialInterface& LHS, const FSkeletalMaterial& RHS );
-
-	UPROPERTY(EditAnywhere, BlueprintReadWrite, Category=SkeletalMesh)
-	TObjectPtr<class UMaterialInterface> 	MaterialInterface;
-	
-	/*This name should be use by the gameplay to avoid error if the skeletal mesh Materials array topology change*/
-	UPROPERTY(EditAnywhere, BlueprintReadWrite, Category = SkeletalMesh)
-	FName						MaterialSlotName;
-#if WITH_EDITORONLY_DATA
-	UPROPERTY()
-	bool						bEnableShadowCasting_DEPRECATED;
-	UPROPERTY()
-	bool						bRecomputeTangent_DEPRECATED;
-	/*This name should be use when we re-import a skeletal mesh so we can order the Materials array like it should be*/
-	UPROPERTY(VisibleAnywhere, Category = SkeletalMesh)
-	FName						ImportedMaterialSlotName;
-#endif //WITH_EDITORONLY_DATA
-
-	/** Data used for texture streaming relative to each UV channels. */
-	UPROPERTY(VisibleAnywhere, BlueprintReadOnly, Category = SkeletalMesh)
-	FMeshUVChannelInfo			UVChannelData;
-};
-
 #if WITH_EDITOR
 /** delegate type for pre skeletal mesh build events */
 DECLARE_MULTICAST_DELEGATE_OneParam(FOnPostMeshCache, class USkeletalMesh*);
@@ -962,7 +601,7 @@ public:
 	UFUNCTION(BlueprintGetter)
 	virtual const USkeleton* GetSkeleton() const override
 	{
-		WaitUntilAsyncPropertyReleased(ESkeletalMeshAsyncProperties::Skeleton, EAsyncPropertyLockType::ReadOnly);
+		WaitUntilAsyncPropertyReleased(ESkeletalMeshAsyncProperties::Skeleton, ESkinnedAssetAsyncPropertyLockType::ReadOnly);
 		PRAGMA_DISABLE_DEPRECATION_WARNINGS
 		return Skeleton;
 		PRAGMA_ENABLE_DEPRECATION_WARNINGS
@@ -990,7 +629,7 @@ private:
 
 	const FBoxSphereBounds& GetExtendedBounds() const
 	{
-		WaitUntilAsyncPropertyReleased(ESkeletalMeshAsyncProperties::ExtendedBounds, EAsyncPropertyLockType::ReadOnly);
+		WaitUntilAsyncPropertyReleased(ESkeletalMeshAsyncProperties::ExtendedBounds, ESkinnedAssetAsyncPropertyLockType::ReadOnly);
 		PRAGMA_DISABLE_DEPRECATION_WARNINGS
 		return ExtendedBounds;
 		PRAGMA_ENABLE_DEPRECATION_WARNINGS
@@ -1055,7 +694,7 @@ public:
 	/** Get bound extension values in the positive direction of XYZ **/
 	const FVector& GetPositiveBoundsExtension() const
 	{
-		WaitUntilAsyncPropertyReleased(ESkeletalMeshAsyncProperties::ExtendedBounds, EAsyncPropertyLockType::ReadOnly);
+		WaitUntilAsyncPropertyReleased(ESkeletalMeshAsyncProperties::ExtendedBounds, ESkinnedAssetAsyncPropertyLockType::ReadOnly);
 		PRAGMA_DISABLE_DEPRECATION_WARNINGS
 		return PositiveBoundsExtension;
 		PRAGMA_ENABLE_DEPRECATION_WARNINGS
@@ -1067,7 +706,7 @@ public:
 	/** Get bound extension values in the negative direction of XYZ **/
 	const FVector& GetNegativeBoundsExtension() const
 	{
-		WaitUntilAsyncPropertyReleased(ESkeletalMeshAsyncProperties::ExtendedBounds, EAsyncPropertyLockType::ReadOnly);
+		WaitUntilAsyncPropertyReleased(ESkeletalMeshAsyncProperties::ExtendedBounds, ESkinnedAssetAsyncPropertyLockType::ReadOnly);
 		PRAGMA_DISABLE_DEPRECATION_WARNINGS
 		return NegativeBoundsExtension;
 		PRAGMA_ENABLE_DEPRECATION_WARNINGS
@@ -1128,7 +767,7 @@ public:
 	UFUNCTION(BlueprintGetter)
 	virtual const TArray<FSkeletalMaterial>& GetMaterials() const override
 	{
-		WaitUntilAsyncPropertyReleased(ESkeletalMeshAsyncProperties::Materials, EAsyncPropertyLockType::ReadOnly);
+		WaitUntilAsyncPropertyReleased(ESkeletalMeshAsyncProperties::Materials, ESkinnedAssetAsyncPropertyLockType::ReadOnly);
 		PRAGMA_DISABLE_DEPRECATION_WARNINGS
 		return Materials;
 		PRAGMA_ENABLE_DEPRECATION_WARNINGS
@@ -1168,7 +807,7 @@ public:
     UE_DEPRECATED(5.0, "Please use UMirrorDataTable for mirroring support")
 	const TArray<struct FBoneMirrorInfo>& GetSkelMirrorTable() const
 	{
-		WaitUntilAsyncPropertyReleased(ESkeletalMeshAsyncProperties::SkelMirrorTable, EAsyncPropertyLockType::ReadOnly);
+		WaitUntilAsyncPropertyReleased(ESkeletalMeshAsyncProperties::SkelMirrorTable, ESkinnedAssetAsyncPropertyLockType::ReadOnly);
 		PRAGMA_DISABLE_DEPRECATION_WARNINGS
 		return SkelMirrorTable;
 		PRAGMA_ENABLE_DEPRECATION_WARNINGS
@@ -1259,7 +898,7 @@ public:
 	/** USkinnedAsset interface. */
 	virtual const FPerPlatformInt& GetMinLod() const override
 	{
-		WaitUntilAsyncPropertyReleased(ESkeletalMeshAsyncProperties::MinLod, EAsyncPropertyLockType::ReadOnly);
+		WaitUntilAsyncPropertyReleased(ESkeletalMeshAsyncProperties::MinLod, ESkinnedAssetAsyncPropertyLockType::ReadOnly);
 		PRAGMA_DISABLE_DEPRECATION_WARNINGS
 		return MinLod;
 		PRAGMA_ENABLE_DEPRECATION_WARNINGS
@@ -1288,7 +927,7 @@ public:
 	/** USkinnedAsset interface. */
 	virtual const FPerPlatformBool& GetDisableBelowMinLodStripping() const override
 	{
-		WaitUntilAsyncPropertyReleased(ESkeletalMeshAsyncProperties::DisableBelowMinLodStripping, EAsyncPropertyLockType::ReadOnly);
+		WaitUntilAsyncPropertyReleased(ESkeletalMeshAsyncProperties::DisableBelowMinLodStripping, ESkinnedAssetAsyncPropertyLockType::ReadOnly);
 		PRAGMA_DISABLE_DEPRECATION_WARNINGS
 		return DisableBelowMinLodStripping;
 		PRAGMA_ENABLE_DEPRECATION_WARNINGS
@@ -1316,7 +955,7 @@ public:
 
 	bool GetOverrideLODStreamingSettings() const
 	{
-		WaitUntilAsyncPropertyReleased(ESkeletalMeshAsyncProperties::OverrideLODStreamingSettings, EAsyncPropertyLockType::ReadOnly);
+		WaitUntilAsyncPropertyReleased(ESkeletalMeshAsyncProperties::OverrideLODStreamingSettings, ESkinnedAssetAsyncPropertyLockType::ReadOnly);
 		PRAGMA_DISABLE_DEPRECATION_WARNINGS
 		return bOverrideLODStreamingSettings;
 		PRAGMA_ENABLE_DEPRECATION_WARNINGS
@@ -1343,7 +982,7 @@ public:
 
 	const FPerPlatformBool& GetSupportLODStreaming() const
 	{
-		WaitUntilAsyncPropertyReleased(ESkeletalMeshAsyncProperties::SupportLODStreaming, EAsyncPropertyLockType::ReadOnly);
+		WaitUntilAsyncPropertyReleased(ESkeletalMeshAsyncProperties::SupportLODStreaming, ESkinnedAssetAsyncPropertyLockType::ReadOnly);
 		PRAGMA_DISABLE_DEPRECATION_WARNINGS
 		return bSupportLODStreaming;
 		PRAGMA_ENABLE_DEPRECATION_WARNINGS
@@ -1370,7 +1009,7 @@ public:
 
 	const FPerPlatformInt& GetMaxNumStreamedLODs() const
 	{
-		WaitUntilAsyncPropertyReleased(ESkeletalMeshAsyncProperties::MaxNumStreamedLODs, EAsyncPropertyLockType::ReadOnly);
+		WaitUntilAsyncPropertyReleased(ESkeletalMeshAsyncProperties::MaxNumStreamedLODs, ESkinnedAssetAsyncPropertyLockType::ReadOnly);
 		PRAGMA_DISABLE_DEPRECATION_WARNINGS
 		return MaxNumStreamedLODs;
 		PRAGMA_ENABLE_DEPRECATION_WARNINGS
@@ -1397,7 +1036,7 @@ public:
 
 	const FPerPlatformInt& GetMaxNumOptionalLODs() const
 	{
-		WaitUntilAsyncPropertyReleased(ESkeletalMeshAsyncProperties::MaxNumOptionalLODs, EAsyncPropertyLockType::ReadOnly);
+		WaitUntilAsyncPropertyReleased(ESkeletalMeshAsyncProperties::MaxNumOptionalLODs, ESkinnedAssetAsyncPropertyLockType::ReadOnly);
 		PRAGMA_DISABLE_DEPRECATION_WARNINGS
 		return MaxNumOptionalLODs;
 		PRAGMA_ENABLE_DEPRECATION_WARNINGS
@@ -1452,7 +1091,7 @@ public:
 	UFUNCTION(BlueprintGetter)
 	const USkeletalMeshLODSettings* GetLODSettings() const
 	{
-		WaitUntilAsyncPropertyReleased(ESkeletalMeshAsyncProperties::LODSettings, EAsyncPropertyLockType::ReadOnly);
+		WaitUntilAsyncPropertyReleased(ESkeletalMeshAsyncProperties::LODSettings, ESkinnedAssetAsyncPropertyLockType::ReadOnly);
 		PRAGMA_DISABLE_DEPRECATION_WARNINGS
 #if WITH_EDITORONLY_DATA
 		return LODSettings;
@@ -1500,7 +1139,7 @@ public:
 	UE_DEPRECATED(5.0, "Please use UMirrorDataTable for mirroring support")
 	TEnumAsByte<EAxis::Type> GetSkelMirrorAxis() const
 	{
-		WaitUntilAsyncPropertyReleased(ESkeletalMeshAsyncProperties::SkelMirrorAxis, EAsyncPropertyLockType::ReadOnly);
+		WaitUntilAsyncPropertyReleased(ESkeletalMeshAsyncProperties::SkelMirrorAxis, ESkinnedAssetAsyncPropertyLockType::ReadOnly);
 		PRAGMA_DISABLE_DEPRECATION_WARNINGS
 		return SkelMirrorAxis;
 		PRAGMA_ENABLE_DEPRECATION_WARNINGS
@@ -1528,7 +1167,7 @@ public:
     UE_DEPRECATED(5.0, "Please use UMirrorDataTable for mirroring support")
 	TEnumAsByte<EAxis::Type> GetSkelMirrorFlipAxis() const
 	{
-		WaitUntilAsyncPropertyReleased(ESkeletalMeshAsyncProperties::SkelMirrorFlipAxis, EAsyncPropertyLockType::ReadOnly);
+		WaitUntilAsyncPropertyReleased(ESkeletalMeshAsyncProperties::SkelMirrorFlipAxis, ESkinnedAssetAsyncPropertyLockType::ReadOnly);
 		PRAGMA_DISABLE_DEPRECATION_WARNINGS
 		return SkelMirrorFlipAxis;
 		PRAGMA_ENABLE_DEPRECATION_WARNINGS
@@ -1564,7 +1203,7 @@ public:
 
 	bool GetHasBeenSimplified() const
 	{
-		WaitUntilAsyncPropertyReleased(ESkeletalMeshAsyncProperties::HasBeenSimplified, EAsyncPropertyLockType::ReadOnly);
+		WaitUntilAsyncPropertyReleased(ESkeletalMeshAsyncProperties::HasBeenSimplified, ESkinnedAssetAsyncPropertyLockType::ReadOnly);
 		PRAGMA_DISABLE_DEPRECATION_WARNINGS
 		return bHasBeenSimplified;
 		PRAGMA_ENABLE_DEPRECATION_WARNINGS
@@ -1593,7 +1232,7 @@ public:
 	/** Return whether or not the mesh has vertex colors. USkinnedAsset interface. */
 	virtual bool GetHasVertexColors() const override
 	{
-		WaitUntilAsyncPropertyReleased(ESkeletalMeshAsyncProperties::HasVertexColors, EAsyncPropertyLockType::ReadOnly);
+		WaitUntilAsyncPropertyReleased(ESkeletalMeshAsyncProperties::HasVertexColors, ESkinnedAssetAsyncPropertyLockType::ReadOnly);
 		PRAGMA_DISABLE_DEPRECATION_WARNINGS
 		return bHasVertexColors != 0;
 		PRAGMA_ENABLE_DEPRECATION_WARNINGS
@@ -1640,7 +1279,7 @@ public:
 
 	bool GetEnablePerPolyCollision() const
 	{
-		WaitUntilAsyncPropertyReleased(ESkeletalMeshAsyncProperties::EnablePerPolyCollision, EAsyncPropertyLockType::ReadOnly);
+		WaitUntilAsyncPropertyReleased(ESkeletalMeshAsyncProperties::EnablePerPolyCollision, ESkinnedAssetAsyncPropertyLockType::ReadOnly);
 		PRAGMA_DISABLE_DEPRECATION_WARNINGS
 		return bEnablePerPolyCollision != 0;
 		PRAGMA_ENABLE_DEPRECATION_WARNINGS
@@ -1669,7 +1308,7 @@ public:
 	
 	FGuid GetVertexColorGuid() const
 	{
-		WaitUntilAsyncPropertyReleased(ESkeletalMeshAsyncProperties::VertexColorGuid, EAsyncPropertyLockType::ReadOnly);
+		WaitUntilAsyncPropertyReleased(ESkeletalMeshAsyncProperties::VertexColorGuid, ESkinnedAssetAsyncPropertyLockType::ReadOnly);
 		PRAGMA_DISABLE_DEPRECATION_WARNINGS
 		return VertexColorGuid;
 		PRAGMA_ENABLE_DEPRECATION_WARNINGS
@@ -1810,7 +1449,7 @@ public:
 	UFUNCTION(BlueprintGetter)
 	const TArray<class UNodeMappingContainer*>& GetNodeMappingData() const
 	{
-		WaitUntilAsyncPropertyReleased(ESkeletalMeshAsyncProperties::NodeMappingData, EAsyncPropertyLockType::ReadOnly);
+		WaitUntilAsyncPropertyReleased(ESkeletalMeshAsyncProperties::NodeMappingData, ESkinnedAssetAsyncPropertyLockType::ReadOnly);
 		PRAGMA_DISABLE_DEPRECATION_WARNINGS
 		return NodeMappingData;
 		PRAGMA_ENABLE_DEPRECATION_WARNINGS
@@ -1909,7 +1548,7 @@ public:
 
 	bool GetHasCustomDefaultEditorCamera() const
 	{
-		WaitUntilAsyncPropertyReleased(ESkeletalMeshAsyncProperties::HasCustomDefaultEditorCamera, EAsyncPropertyLockType::ReadOnly);
+		WaitUntilAsyncPropertyReleased(ESkeletalMeshAsyncProperties::HasCustomDefaultEditorCamera, ESkinnedAssetAsyncPropertyLockType::ReadOnly);
 		PRAGMA_DISABLE_DEPRECATION_WARNINGS
 		return bHasCustomDefaultEditorCamera;
 		PRAGMA_ENABLE_DEPRECATION_WARNINGS
@@ -1937,7 +1576,7 @@ public:
 
 	const FVector& GetDefaultEditorCameraLocation() const
 	{
-		WaitUntilAsyncPropertyReleased(ESkeletalMeshAsyncProperties::DefaultEditorCameraLocation, EAsyncPropertyLockType::ReadOnly);
+		WaitUntilAsyncPropertyReleased(ESkeletalMeshAsyncProperties::DefaultEditorCameraLocation, ESkinnedAssetAsyncPropertyLockType::ReadOnly);
 		PRAGMA_DISABLE_DEPRECATION_WARNINGS
 		return DefaultEditorCameraLocation;
 		PRAGMA_ENABLE_DEPRECATION_WARNINGS
@@ -1965,7 +1604,7 @@ public:
 
 	const FRotator& GetDefaultEditorCameraRotation() const
 	{
-		WaitUntilAsyncPropertyReleased(ESkeletalMeshAsyncProperties::DefaultEditorCameraRotation, EAsyncPropertyLockType::ReadOnly);
+		WaitUntilAsyncPropertyReleased(ESkeletalMeshAsyncProperties::DefaultEditorCameraRotation, ESkinnedAssetAsyncPropertyLockType::ReadOnly);
 		PRAGMA_DISABLE_DEPRECATION_WARNINGS
 		return DefaultEditorCameraRotation;
 		PRAGMA_ENABLE_DEPRECATION_WARNINGS
@@ -1993,7 +1632,7 @@ public:
 
 	const FVector& GetDefaultEditorCameraLookAt() const
 	{
-		WaitUntilAsyncPropertyReleased(ESkeletalMeshAsyncProperties::DefaultEditorCameraLookAt, EAsyncPropertyLockType::ReadOnly);
+		WaitUntilAsyncPropertyReleased(ESkeletalMeshAsyncProperties::DefaultEditorCameraLookAt, ESkinnedAssetAsyncPropertyLockType::ReadOnly);
 		PRAGMA_DISABLE_DEPRECATION_WARNINGS
 		return DefaultEditorCameraLookAt;
 		PRAGMA_ENABLE_DEPRECATION_WARNINGS
@@ -2021,7 +1660,7 @@ public:
 
 	float GetDefaultEditorCameraOrthoZoom() const
 	{
-		WaitUntilAsyncPropertyReleased(ESkeletalMeshAsyncProperties::DefaultEditorCameraOrthoZoom, EAsyncPropertyLockType::ReadOnly);
+		WaitUntilAsyncPropertyReleased(ESkeletalMeshAsyncProperties::DefaultEditorCameraOrthoZoom, ESkinnedAssetAsyncPropertyLockType::ReadOnly);
 		PRAGMA_DISABLE_DEPRECATION_WARNINGS
 		return DefaultEditorCameraOrthoZoom;
 		PRAGMA_ENABLE_DEPRECATION_WARNINGS
@@ -2051,7 +1690,7 @@ public:
 
 	const FPreviewAssetAttachContainer& GetPreviewAttachedAssetContainer() const
 	{
-		WaitUntilAsyncPropertyReleased(ESkeletalMeshAsyncProperties::PreviewAttachedAssetContainer, EAsyncPropertyLockType::ReadOnly);
+		WaitUntilAsyncPropertyReleased(ESkeletalMeshAsyncProperties::PreviewAttachedAssetContainer, ESkinnedAssetAsyncPropertyLockType::ReadOnly);
 		PRAGMA_DISABLE_DEPRECATION_WARNINGS
 		return PreviewAttachedAssetContainer;
 		PRAGMA_ENABLE_DEPRECATION_WARNINGS
@@ -2074,7 +1713,7 @@ public:
 
 	bool GetRequiresLODScreenSizeConversion() const
 	{
-		WaitUntilAsyncPropertyReleased(ESkeletalMeshAsyncProperties::RequiresLODScreenSizeConversion, EAsyncPropertyLockType::ReadOnly);
+		WaitUntilAsyncPropertyReleased(ESkeletalMeshAsyncProperties::RequiresLODScreenSizeConversion, ESkinnedAssetAsyncPropertyLockType::ReadOnly);
 		PRAGMA_DISABLE_DEPRECATION_WARNINGS
 		return bRequiresLODScreenSizeConversion;
 		PRAGMA_ENABLE_DEPRECATION_WARNINGS
@@ -2097,7 +1736,7 @@ public:
 
 	bool GetRequiresLODHysteresisConversion() const
 	{
-		WaitUntilAsyncPropertyReleased(ESkeletalMeshAsyncProperties::RequiresLODHysteresisConversion, EAsyncPropertyLockType::ReadOnly);
+		WaitUntilAsyncPropertyReleased(ESkeletalMeshAsyncProperties::RequiresLODHysteresisConversion, ESkinnedAssetAsyncPropertyLockType::ReadOnly);
 		PRAGMA_DISABLE_DEPRECATION_WARNINGS
 		return bRequiresLODHysteresisConversion;
 		PRAGMA_ENABLE_DEPRECATION_WARNINGS
@@ -2123,7 +1762,7 @@ public:
 	/** USkinnedAsset interface. */
 	virtual bool GetSupportRayTracing() const override
 	{
-		WaitUntilAsyncPropertyReleased(ESkeletalMeshAsyncProperties::bSupportRayTracing, EAsyncPropertyLockType::ReadOnly);
+		WaitUntilAsyncPropertyReleased(ESkeletalMeshAsyncProperties::bSupportRayTracing, ESkinnedAssetAsyncPropertyLockType::ReadOnly);
 		PRAGMA_DISABLE_DEPRECATION_WARNINGS
 		return bSupportRayTracing;
 		PRAGMA_ENABLE_DEPRECATION_WARNINGS
@@ -2147,7 +1786,7 @@ public:
 	/** USkinnedAsset interface. */
 	virtual int32 GetRayTracingMinLOD() const override
 	{
-		WaitUntilAsyncPropertyReleased(ESkeletalMeshAsyncProperties::RayTracingMinLOD, EAsyncPropertyLockType::ReadOnly);
+		WaitUntilAsyncPropertyReleased(ESkeletalMeshAsyncProperties::RayTracingMinLOD, ESkinnedAssetAsyncPropertyLockType::ReadOnly);
 		PRAGMA_DISABLE_DEPRECATION_WARNINGS
 		return RayTracingMinLOD;
 		PRAGMA_ENABLE_DEPRECATION_WARNINGS
@@ -2171,7 +1810,7 @@ public:
 
 	EClothLODBiasMode GetClothLODBiasMode() const
 	{
-		WaitUntilAsyncPropertyReleased(ESkeletalMeshAsyncProperties::ClothLODBiasMode, EAsyncPropertyLockType::ReadOnly);
+		WaitUntilAsyncPropertyReleased(ESkeletalMeshAsyncProperties::ClothLODBiasMode, ESkinnedAssetAsyncPropertyLockType::ReadOnly);
 		PRAGMA_DISABLE_DEPRECATION_WARNINGS
 		return ClothLODBiasMode;
 		PRAGMA_ENABLE_DEPRECATION_WARNINGS
@@ -2209,7 +1848,7 @@ public:
 	UFUNCTION(BlueprintGetter)
 	virtual const TArray<UMorphTarget*>& GetMorphTargets() const override
 	{
-		WaitUntilAsyncPropertyReleased(ESkeletalMeshAsyncProperties::MorphTargets, EAsyncPropertyLockType::ReadOnly);
+		WaitUntilAsyncPropertyReleased(ESkeletalMeshAsyncProperties::MorphTargets, ESkinnedAssetAsyncPropertyLockType::ReadOnly);
 		PRAGMA_DISABLE_DEPRECATION_WARNINGS
 		return MorphTargets;
 		PRAGMA_ENABLE_DEPRECATION_WARNINGS
@@ -2257,7 +1896,7 @@ public:
 	/** USkinnedAsset interface. */
 	virtual const FReferenceSkeleton& GetRefSkeleton() const override
 	{
-		WaitUntilAsyncPropertyReleased(ESkeletalMeshAsyncProperties::RefSkeleton, EAsyncPropertyLockType::ReadOnly);
+		WaitUntilAsyncPropertyReleased(ESkeletalMeshAsyncProperties::RefSkeleton, ESkinnedAssetAsyncPropertyLockType::ReadOnly);
 		PRAGMA_DISABLE_DEPRECATION_WARNINGS
 		return RefSkeleton;
 		PRAGMA_ENABLE_DEPRECATION_WARNINGS
@@ -2292,7 +1931,7 @@ public:
 
 	const TMap<FName, int32>& GetMorphTargetIndexMap() const
 	{
-		WaitUntilAsyncPropertyReleased(ESkeletalMeshAsyncProperties::MorphTargetIndexMap, EAsyncPropertyLockType::ReadOnly);
+		WaitUntilAsyncPropertyReleased(ESkeletalMeshAsyncProperties::MorphTargetIndexMap, ESkinnedAssetAsyncPropertyLockType::ReadOnly);
 		PRAGMA_DISABLE_DEPRECATION_WARNINGS
 		return MorphTargetIndexMap;
 		PRAGMA_ENABLE_DEPRECATION_WARNINGS
@@ -2329,7 +1968,7 @@ public:
 	/** USkinnedAsset interface. */
 	virtual const TArray<FMatrix44f>& GetRefBasesInvMatrix() const override
 	{
-		WaitUntilAsyncPropertyReleased(ESkeletalMeshAsyncProperties::RefBasesInvMatrix, EAsyncPropertyLockType::ReadOnly);
+		WaitUntilAsyncPropertyReleased(ESkeletalMeshAsyncProperties::RefBasesInvMatrix, ESkinnedAssetAsyncPropertyLockType::ReadOnly);
 		PRAGMA_DISABLE_DEPRECATION_WARNINGS
 		return RefBasesInvMatrix;
 		PRAGMA_ENABLE_DEPRECATION_WARNINGS
@@ -2359,7 +1998,7 @@ public:
 
 	float GetFloorOffset() const
 	{
-		WaitUntilAsyncPropertyReleased(ESkeletalMeshAsyncProperties::FloorOffset, EAsyncPropertyLockType::ReadOnly);
+		WaitUntilAsyncPropertyReleased(ESkeletalMeshAsyncProperties::FloorOffset, ESkinnedAssetAsyncPropertyLockType::ReadOnly);
 		PRAGMA_DISABLE_DEPRECATION_WARNINGS
 		return FloorOffset;
 		PRAGMA_ENABLE_DEPRECATION_WARNINGS
@@ -2395,7 +2034,7 @@ public:
 
 	const TArray<FTransform>& GetRetargetBasePose() const
 	{
-		WaitUntilAsyncPropertyReleased(ESkeletalMeshAsyncProperties::RetargetBasePose, EAsyncPropertyLockType::ReadOnly);
+		WaitUntilAsyncPropertyReleased(ESkeletalMeshAsyncProperties::RetargetBasePose, ESkinnedAssetAsyncPropertyLockType::ReadOnly);
 		PRAGMA_DISABLE_DEPRECATION_WARNINGS
 		return RetargetBasePose;
 		PRAGMA_ENABLE_DEPRECATION_WARNINGS
@@ -2431,7 +2070,7 @@ public:
 
 	TSubclassOf<UAnimInstance> GetPostProcessAnimBlueprint() const
 	{
-		WaitUntilAsyncPropertyReleased(ESkeletalMeshAsyncProperties::PostProcessAnimBlueprint, EAsyncPropertyLockType::ReadOnly);
+		WaitUntilAsyncPropertyReleased(ESkeletalMeshAsyncProperties::PostProcessAnimBlueprint, ESkinnedAssetAsyncPropertyLockType::ReadOnly);
 		PRAGMA_DISABLE_DEPRECATION_WARNINGS
 		return PostProcessAnimBlueprint;
 		PRAGMA_ENABLE_DEPRECATION_WARNINGS
@@ -2458,20 +2097,17 @@ public:
 	void RemoveLegacyClothingSections();
 
 	/*
-	* Handle some common preparation steps between async post load and async build
+	* Handle some common preparation steps between async post load and async build. USkinnedAsset interface.
 	*/
-	void PrepareForAsyncCompilation();
+	virtual void PrepareForAsyncCompilation() override;
 
-	/** Returns false if there is currently an async task running */
-	bool IsAsyncTaskComplete() const;
+	/** Returns false if there is currently an async task running. USkinnedAsset interface. */
+	virtual bool IsAsyncTaskComplete() const override;
 
 	/** Try to cancel any pending async tasks.
 	 *  Returns true if there is no more async tasks pending, false otherwise.
 	 */
 	bool TryCancelAsyncTasks();
-
-	/** Holds the pointer to an async task if one exists. */
-	TUniquePtr<FSkeletalMeshAsyncBuildTask> AsyncTask;
 #endif // WITH_EDITOR
 
 	/**
@@ -2508,7 +2144,7 @@ public:
 	UFUNCTION(BlueprintGetter)
 	const TArray<UClothingAssetBase*>& GetMeshClothingAssets() const
 	{
-		WaitUntilAsyncPropertyReleased(ESkeletalMeshAsyncProperties::MeshClothingAssets, EAsyncPropertyLockType::ReadOnly);
+		WaitUntilAsyncPropertyReleased(ESkeletalMeshAsyncProperties::MeshClothingAssets, ESkinnedAssetAsyncPropertyLockType::ReadOnly);
 		PRAGMA_DISABLE_DEPRECATION_WARNINGS
 		return MeshClothingAssets;
 		PRAGMA_ENABLE_DEPRECATION_WARNINGS
@@ -2558,7 +2194,7 @@ public:
 
 	const FSkeletalMeshSamplingInfo& GetSamplingInfo() const 
 	{
-		WaitUntilAsyncPropertyReleased(ESkeletalMeshAsyncProperties::SamplingInfo, EAsyncPropertyLockType::ReadOnly);
+		WaitUntilAsyncPropertyReleased(ESkeletalMeshAsyncProperties::SamplingInfo, ESkinnedAssetAsyncPropertyLockType::ReadOnly);
 		PRAGMA_DISABLE_DEPRECATION_WARNINGS
 		return SamplingInfo;
 		PRAGMA_ENABLE_DEPRECATION_WARNINGS
@@ -2632,7 +2268,7 @@ private:
 
 	const TArray<FMatrix>& GetCachedComposedRefPoseMatrices() const
 	{
-		WaitUntilAsyncPropertyReleased(ESkeletalMeshAsyncProperties::CachedComposedRefPoseMatrices, EAsyncPropertyLockType::ReadOnly);
+		WaitUntilAsyncPropertyReleased(ESkeletalMeshAsyncProperties::CachedComposedRefPoseMatrices, ESkinnedAssetAsyncPropertyLockType::ReadOnly);
 		PRAGMA_DISABLE_DEPRECATION_WARNINGS
 		return CachedComposedRefPoseMatrices;
 		PRAGMA_ENABLE_DEPRECATION_WARNINGS
@@ -3002,27 +2638,27 @@ private:
 
 #if WITH_EDITOR
 	/** Generate SkeletalMeshRenderData from ImportedModel */
-	void CacheDerivedData(FSkeletalMeshCompilationContext* ContextPtr);
+	void CacheDerivedData(FSkinnedAssetCompilationContext* ContextPtr);
 
 	/**
-	 * Initial step for the building process - Can't be done in parallel.
+	 * Initial step for the building process - Can't be done in parallel. USkinnedAsset Interface.
 	 */
-	void BeginBuildInternal(FSkeletalMeshBuildContext& Context);
+	virtual void BeginBuildInternal(FSkinnedAssetBuildContext& Context) override;
 
 	/**
-	 * Thread-safe part.
+	 * Thread-safe part. USkinnedAsset Interface.
 	 */
-	void ExecuteBuildInternal(FSkeletalMeshBuildContext& Context);
+	virtual void ExecuteBuildInternal(FSkinnedAssetBuildContext& Context) override;
 
 	/**
-	 * Complete the building process - Can't be done in parallel.
+	 * Complete the building process - Can't be done in parallel. USkinnedAsset Interface.
 	 */
-	void FinishBuildInternal(FSkeletalMeshBuildContext& Context);
+	virtual void FinishBuildInternal(FSkinnedAssetBuildContext& Context) override;
 
 	/**
 	 * Copy build/load context result data to the skeletalmesh member on the game thread - Can't be done in parallel.
 	 */
-	void ApplyFinishBuildInternalData(FSkeletalMeshCompilationContext* ContextPtr);
+	void ApplyFinishBuildInternalData(FSkinnedAssetCompilationContext* ContextPtr);
 	
 #endif
 
@@ -3063,10 +2699,6 @@ private:
 
 #if WITH_EDITOR
 public:
-	/** IInterface_AsyncCompilation begin*/
-	virtual bool IsCompiling() const override;
-	/** IInterface_AsyncCompilation end*/
-
 	/** Delegates for asset editor events */
 
 	FDelegateHandle RegisterOnClothingChange(const FSimpleMulticastDelegate::FDelegate& InDelegate);
@@ -3080,31 +2712,14 @@ private:
 	// INodeMappingProviderInterface
 	virtual void GetMappableNodeData(TArray<FName>& OutNames, TArray<FNodeItem>& OutTransforms) const override;
 
-	enum class EAsyncPropertyLockType
-	{
-		None = 0,
-		ReadOnly = 1,
-		WriteOnly = 2,
-		ReadWrite = 3
-	};
-
-	FRIEND_ENUM_CLASS_FLAGS(EAsyncPropertyLockType);
-
-#if WITH_EDITOR
-	/** Used as a bit-field indicating which properties are read by async compilation. */
-	std::atomic<uint64> AccessedProperties;
-	/** Used as a bit-field indicating which properties are written to by async compilation. */
-	std::atomic<uint64> ModifiedProperties;
-	
-	void AcquireAsyncProperty(ESkeletalMeshAsyncProperties AsyncProperties = ESkeletalMeshAsyncProperties::All, EAsyncPropertyLockType LockType = EAsyncPropertyLockType::ReadWrite);
-	void ReleaseAsyncProperty(ESkeletalMeshAsyncProperties AsyncProperties = ESkeletalMeshAsyncProperties::All, EAsyncPropertyLockType LockType = EAsyncPropertyLockType::ReadWrite);
-	void WaitUntilAsyncPropertyReleased(ESkeletalMeshAsyncProperties AsyncProperties, EAsyncPropertyLockType LockType = EAsyncPropertyLockType::ReadWrite) const;
-#else
-	FORCEINLINE void AcquireAsyncProperty(ESkeletalMeshAsyncProperties AsyncProperties = ESkeletalMeshAsyncProperties::All, EAsyncPropertyLockType LockType = EAsyncPropertyLockType::ReadWrite) {};
-	FORCEINLINE void ReleaseAsyncProperty(ESkeletalMeshAsyncProperties AsyncProperties = ESkeletalMeshAsyncProperties::All, EAsyncPropertyLockType LockType = EAsyncPropertyLockType::ReadWrite) {};
-	FORCEINLINE void WaitUntilAsyncPropertyReleased(ESkeletalMeshAsyncProperties AsyncProperties, EAsyncPropertyLockType LockType = EAsyncPropertyLockType::ReadWrite) const {}
-#endif
+	/**
+	 * Wait for the asset to finish compilation to protect internal skinned asset data from race conditions during async build.
+	 * This should be called before accessing all async accessible properties.
+	 */
+	void WaitUntilAsyncPropertyReleased(ESkeletalMeshAsyncProperties AsyncProperties, ESkinnedAssetAsyncPropertyLockType LockType = ESkinnedAssetAsyncPropertyLockType::ReadWrite) const;
 
+	/** Convert async property from enum value to string. USkinnedAsset interface. */
+	virtual FString GetAsyncPropertyName(uint64 Property) const override;
 
 public:
 	/*
@@ -3164,7 +2779,7 @@ public:
 	 */
 	virtual const TArray<FSkeletalMeshLODInfo>& GetLODInfoArray() const override
 	{
-		WaitUntilAsyncPropertyReleased(ESkeletalMeshAsyncProperties::LODInfo, EAsyncPropertyLockType::ReadOnly);
+		WaitUntilAsyncPropertyReleased(ESkeletalMeshAsyncProperties::LODInfo, ESkinnedAssetAsyncPropertyLockType::ReadOnly);
 		PRAGMA_DISABLE_DEPRECATION_WARNINGS
 		return LODInfo;
 		PRAGMA_ENABLE_DEPRECATION_WARNINGS
@@ -3186,7 +2801,7 @@ public:
 	 */	
 	virtual const FSkeletalMeshLODInfo* GetLODInfo(int32 Index) const override
 	{
-		WaitUntilAsyncPropertyReleased(ESkeletalMeshAsyncProperties::LODInfo, EAsyncPropertyLockType::ReadOnly);
+		WaitUntilAsyncPropertyReleased(ESkeletalMeshAsyncProperties::LODInfo, ESkinnedAssetAsyncPropertyLockType::ReadOnly);
 		PRAGMA_DISABLE_DEPRECATION_WARNINGS
 		return LODInfo.IsValidIndex(Index) ? &LODInfo[Index] : nullptr;
 		PRAGMA_ENABLE_DEPRECATION_WARNINGS
@@ -3207,7 +2822,7 @@ public:
 	 */
 	virtual bool IsValidLODIndex(int32 Index) const override
 	{
-		WaitUntilAsyncPropertyReleased(ESkeletalMeshAsyncProperties::LODInfo, EAsyncPropertyLockType::ReadOnly);
+		WaitUntilAsyncPropertyReleased(ESkeletalMeshAsyncProperties::LODInfo, ESkinnedAssetAsyncPropertyLockType::ReadOnly);
 		PRAGMA_DISABLE_DEPRECATION_WARNINGS
 		return LODInfo.IsValidIndex(Index);
 		PRAGMA_ENABLE_DEPRECATION_WARNINGS
@@ -3217,7 +2832,7 @@ public:
 	 */
 	virtual int32 GetLODNum() const override
 	{
-		WaitUntilAsyncPropertyReleased(ESkeletalMeshAsyncProperties::LODInfo, EAsyncPropertyLockType::ReadOnly);
+		WaitUntilAsyncPropertyReleased(ESkeletalMeshAsyncProperties::LODInfo, ESkinnedAssetAsyncPropertyLockType::ReadOnly);
 		PRAGMA_DISABLE_DEPRECATION_WARNINGS
 		return LODInfo.Num();
 		PRAGMA_ENABLE_DEPRECATION_WARNINGS
@@ -3229,7 +2844,7 @@ public:
 public:
 	const TArray<FSkinWeightProfileInfo>& GetSkinWeightProfiles() const 
 	{
-		WaitUntilAsyncPropertyReleased(ESkeletalMeshAsyncProperties::SkinWeightProfiles, EAsyncPropertyLockType::ReadOnly);
+		WaitUntilAsyncPropertyReleased(ESkeletalMeshAsyncProperties::SkinWeightProfiles, ESkinnedAssetAsyncPropertyLockType::ReadOnly);
 		PRAGMA_DISABLE_DEPRECATION_WARNINGS
 		return SkinWeightProfiles; 
 		PRAGMA_ENABLE_DEPRECATION_WARNINGS
@@ -3252,7 +2867,7 @@ public:
 	}
 	int32 GetNumSkinWeightProfiles() const 
 	{ 
-		WaitUntilAsyncPropertyReleased(ESkeletalMeshAsyncProperties::SkinWeightProfiles, EAsyncPropertyLockType::ReadOnly);
+		WaitUntilAsyncPropertyReleased(ESkeletalMeshAsyncProperties::SkinWeightProfiles, ESkinnedAssetAsyncPropertyLockType::ReadOnly);
 		PRAGMA_DISABLE_DEPRECATION_WARNINGS
 		return SkinWeightProfiles.Num(); 
 		PRAGMA_ENABLE_DEPRECATION_WARNINGS
@@ -3277,7 +2892,7 @@ public:
 	// USkinnedAsset interface
 	virtual bool GetUseLegacyMeshDerivedDataKey() const override
 	{
-		WaitUntilAsyncPropertyReleased(ESkeletalMeshAsyncProperties::UseLegacyMeshDerivedDataKey, EAsyncPropertyLockType::ReadOnly);
+		WaitUntilAsyncPropertyReleased(ESkeletalMeshAsyncProperties::UseLegacyMeshDerivedDataKey, ESkinnedAssetAsyncPropertyLockType::ReadOnly);
 		PRAGMA_DISABLE_DEPRECATION_WARNINGS
 		return UseLegacyMeshDerivedDataKey;
 		PRAGMA_ENABLE_DEPRECATION_WARNINGS
@@ -3300,23 +2915,20 @@ protected:
 	TArray<FSkinWeightProfileInfo> SkinWeightProfiles;
 
 private:
-	friend class FSkeletalMeshCompilingManager;
-	friend class FSkeletalMeshAsyncBuildWorker;
-
 	/**
-	 * Initial step for the Post Load process - Can't be done in parallel.
+	 * Initial step for the Post Load process - Can't be done in parallel. USkinnedAsset Interface.
 	 */
-	void BeginPostLoadInternal(FSkeletalMeshPostLoadContext& Context);
+	virtual void BeginPostLoadInternal(FSkinnedAssetPostLoadContext& Context) override;
 
 	/**
-	 * Thread-safe part of the Post Load
+	 * Thread-safe part of the Post Load. USkinnedAsset interface.
 	 */
-	void ExecutePostLoadInternal(FSkeletalMeshPostLoadContext& Context);
+	virtual void ExecutePostLoadInternal(FSkinnedAssetPostLoadContext& Context) override;
 
 	/**
-	 * Complete the postload process - Can't be done in parallel.
+	 * Complete the postload process - Can't be done in parallel. USkinnedAsset interface.
 	 */
-	void FinishPostLoadInternal(FSkeletalMeshPostLoadContext& Context);
+	virtual void FinishPostLoadInternal(FSkinnedAssetPostLoadContext& Context) override;
 };
 
 struct FSkeletalMeshBuildParameters
diff --git a/Engine/Source/Runtime/Engine/Classes/Engine/SkinnedAsset.h b/Engine/Source/Runtime/Engine/Classes/Engine/SkinnedAsset.h
--- a/Engine/Source/Runtime/Engine/Classes/Engine/SkinnedAsset.h
+++ b/Engine/Source/Runtime/Engine/Classes/Engine/SkinnedAsset.h
@@ -13,6 +13,8 @@
 #include "ReferenceSkeleton.h"
 #include "PerPlatformProperties.h"
 #include "SkeletalMeshTypes.h"
+#include "SkinnedAssetAsyncCompileUtils.h"
+#include "SkinnedAssetCommon.h"
 #include "SkinnedAsset.generated.h"
 
 class USkeleton;
@@ -22,6 +24,16 @@ class ITargetPlatform;
 struct FSkeletalMaterial;
 struct FSkeletalMeshLODInfo;
 
+enum class ESkinnedAssetAsyncPropertyLockType
+{
+	None = 0,
+	ReadOnly = 1,
+	WriteOnly = 2,
+	ReadWrite = 3
+};
+ENUM_CLASS_FLAGS(ESkinnedAssetAsyncPropertyLockType);
+
+
 UCLASS(hidecategories = Object, config = Engine, editinlinenew, abstract)
 class ENGINE_API USkinnedAsset : public UStreamableRenderAsset, public IInterface_AsyncCompilation
 {
@@ -93,7 +105,7 @@ public:
 	PURE_VIRTUAL(USkinnedAsset::GetLODInfoArray, static const TArray<FSkeletalMeshLODInfo> Dummy; return Dummy;);
 
 	/** Get the data to use for rendering. */
-	virtual FORCEINLINE class FSkeletalMeshRenderData* GetResourceForRendering() const
+	virtual class FSkeletalMeshRenderData* GetResourceForRendering() const
 	PURE_VIRTUAL(USkinnedAsset::GetResourceForRendering, return nullptr;);
 
 	virtual int32 GetDefaultMinLod() const
@@ -173,8 +185,8 @@ public:
 	virtual uint32 GetVertexBufferFlags() const
 	{ return GetHasVertexColors() ? ESkeletalMeshVertexFlags::HasVertexColors : ESkeletalMeshVertexFlags::None; }
 
+	//~ Begin UObject Interface
 	/**
-	* UObject Interface
 	* This will return detail info about this specific object. (e.g. AudioComponent will return the name of the cue,
 	* ParticleSystemComponent will return the name of the ParticleSystem)  The idea here is that in many places
 	* you have a component of interest but what you really want is some characteristic that you can use to track
@@ -183,9 +195,12 @@ public:
 	virtual FString GetDetailedInfoInternal() const override
 	{ return GetPathName(nullptr); }
 
+	virtual void PostLoad() override;
+	//~ End UObject Interface
+
 #if WITH_EDITOR
 	/** IInterface_AsyncCompilation begin*/
-	virtual bool IsCompiling() const override { return false; }
+	virtual bool IsCompiling() const override;
 	/** IInterface_AsyncCompilation end*/
 
 	virtual FString BuildDerivedDataKey(const ITargetPlatform* TargetPlatform)
@@ -217,5 +232,55 @@ public:
 	virtual class FSkeletalMeshModel* GetImportedModel() const
 	PURE_VIRTUAL(USkinnedAsset::GetImportedModel, return nullptr;);
 #endif // WITH_EDITORONLY_DATA
+
+protected:	
+	/** Lock properties that should not be modified/accessed during async build. */
+	void AcquireAsyncProperty(uint64 AsyncProperties = MAX_uint64, ESkinnedAssetAsyncPropertyLockType LockType = ESkinnedAssetAsyncPropertyLockType::ReadWrite);
+	/** Release properties that should not be modified/accessed during async build. */
+	void ReleaseAsyncProperty(uint64 AsyncProperties = MAX_uint64, ESkinnedAssetAsyncPropertyLockType LockType = ESkinnedAssetAsyncPropertyLockType::ReadWrite);
+
+	/**
+	 * Wait for the asset to finish compilation to protect internal skinned asset data from race conditions during async build.
+	 * A derived class may want to define its async properties as enums and cast them to uint64 when calling this function.
+	 */
+	void WaitUntilAsyncPropertyReleasedInternal(uint64 AsyncProperties, ESkinnedAssetAsyncPropertyLockType LockType = ESkinnedAssetAsyncPropertyLockType::ReadWrite) const;
+
+#if WITH_EDITOR
+	/** Initial step for the building process - Can't be done in parallel. */
+	virtual void BeginBuildInternal(FSkinnedAssetBuildContext& Context) {}
+	/** Thread-safe part. */
+	virtual void ExecuteBuildInternal(FSkinnedAssetBuildContext& Context) {}
+	/** Complete the building process - Can't be done in parallel. */
+	virtual void FinishBuildInternal(FSkinnedAssetBuildContext& Context) {}
+
+	/** Holds the pointer to an async task if one exists. */
+	TUniquePtr<FSkinnedAssetAsyncBuildTask> AsyncTask;
+#endif // WITH_EDITOR
+
+private:
+	/** Initial step for the Post Load process - Can't be done in parallel. */
+	virtual void BeginPostLoadInternal(FSkinnedAssetPostLoadContext& Context) {}
+	/** Thread-safe part of the Post Load */
+	virtual void ExecutePostLoadInternal(FSkinnedAssetPostLoadContext& Context) {}
+	/** Complete the postload process - Can't be done in parallel. */
+	virtual void FinishPostLoadInternal(FSkinnedAssetPostLoadContext& Context) {}
+
+	/** Convert async property from enum value to string. */
+	virtual FString GetAsyncPropertyName(uint64 Property) const { return TEXT(""); }
+
+#if WITH_EDITOR
+	/** Handle some common preparation steps between async post load and async build */
+	virtual void PrepareForAsyncCompilation() {}
+	/** Returns false if there is currently an async task running */
+	virtual bool IsAsyncTaskComplete() const { return true; }
+
+	/** Used as a bit-field indicating which properties are read by async compilation. */
+	std::atomic<uint64> AccessedProperties;
+	/** Used as a bit-field indicating which properties are written to by async compilation. */
+	std::atomic<uint64> ModifiedProperties;
+#endif // WITH_EDITOR
+
+	friend class FSkinnedAssetCompilingManager;
+	friend class FSkinnedAssetAsyncBuildWorker;
 };
 
diff --git a/Engine/Source/Runtime/Engine/Classes/Engine/SkinnedAssetAsyncCompileUtils.h b/Engine/Source/Runtime/Engine/Classes/Engine/SkinnedAssetAsyncCompileUtils.h
new file mode 100644
--- /dev/null
+++ b/Engine/Source/Runtime/Engine/Classes/Engine/SkinnedAssetAsyncCompileUtils.h
@@ -0,0 +1,113 @@
+// Copyright Epic Games, Inc. All Rights Reserved.
+
+#pragma once
+
+/**
+ * Contains the helper functions to compile skinned asset asynchronously.
+ */
+
+#include "Async/AsyncWork.h"
+
+class USkinnedAsset;
+
+class FSkinnedAssetCompilationContext
+{
+public:
+	FSkinnedAssetCompilationContext() = default;
+	// Non-copyable
+	FSkinnedAssetCompilationContext(const FSkinnedAssetCompilationContext&) = delete;
+	FSkinnedAssetCompilationContext& operator=(const FSkinnedAssetCompilationContext&) = delete;
+	// Movable
+	FSkinnedAssetCompilationContext(FSkinnedAssetCompilationContext&&) = default;
+	FSkinnedAssetCompilationContext& operator=(FSkinnedAssetCompilationContext&&) = default;
+	
+	//True if this compilation context is start from a serialize save
+	bool bIsSerializeSaving = false;
+};
+
+class FSkinnedAssetPostLoadContext : public FSkinnedAssetCompilationContext
+{
+public:
+	bool bHasCachedDerivedData = false;
+};
+
+class FSkinnedAssetBuildContext : public FSkinnedAssetCompilationContext
+{
+public:
+	TUniquePtr<class FSkinnedMeshComponentRecreateRenderStateContext> RecreateRenderStateContext;
+};
+
+#if WITH_EDITOR
+
+// Any thread implicated in the build must have a valid scope to be granted access to protected properties without causing any stalls.
+class FSkinnedAssetAsyncBuildScope
+{
+public:
+	ENGINE_API FSkinnedAssetAsyncBuildScope(const USkinnedAsset* SkinnedAsset);
+	ENGINE_API ~FSkinnedAssetAsyncBuildScope();
+	ENGINE_API static bool ShouldWaitOnLockedProperties(const USkinnedAsset* SkinnedAsset);
+
+private:
+	const USkinnedAsset* PreviousScope = nullptr;
+	// Only the thread(s) compiling this Skeletal mesh will have full access to protected properties without causing any stalls.
+	static thread_local const USkinnedAsset* SkinnedAssetBeingAsyncCompiled;
+};
+
+/**
+ * Worker used to perform async static mesh compilation.
+ */
+class FSkinnedAssetAsyncBuildWorker : public FNonAbandonableTask
+{
+public:
+	USkinnedAsset* SkinnedAsset;
+	TOptional<FSkinnedAssetPostLoadContext> PostLoadContext;
+	TOptional<FSkinnedAssetBuildContext> BuildContext;
+
+	/** Initialization constructor. */
+	FSkinnedAssetAsyncBuildWorker(
+		USkinnedAsset* InSkinnedAsset,
+		FSkinnedAssetBuildContext&& InBuildContext)
+		: SkinnedAsset(InSkinnedAsset)
+		, BuildContext(MoveTemp(InBuildContext))
+	{
+	}
+
+	/** Initialization constructor. */
+	FSkinnedAssetAsyncBuildWorker(
+		USkinnedAsset* InSkinnedAsset,
+		FSkinnedAssetPostLoadContext&& InPostLoadContext)
+		: SkinnedAsset(InSkinnedAsset)
+		, PostLoadContext(MoveTemp(InPostLoadContext))
+	{
+	}
+
+	FORCEINLINE TStatId GetStatId() const
+	{
+		RETURN_QUICK_DECLARE_CYCLE_STAT(FSkinnedAssetAsyncBuildWorker, STATGROUP_ThreadPoolAsyncTasks);
+	}
+
+	void DoWork();
+};
+
+struct FSkinnedAssetAsyncBuildTask : public FAsyncTask<FSkinnedAssetAsyncBuildWorker>
+{
+	FSkinnedAssetAsyncBuildTask(
+		USkinnedAsset* InSkinnedAsset,
+		FSkinnedAssetPostLoadContext&& InPostLoadContext)
+		: FAsyncTask<FSkinnedAssetAsyncBuildWorker>(InSkinnedAsset, MoveTemp(InPostLoadContext))
+		, SkinnedAsset(InSkinnedAsset)
+	{
+	}
+
+	FSkinnedAssetAsyncBuildTask(
+		USkinnedAsset* InSkinnedAsset,
+		FSkinnedAssetBuildContext&& InBuildContext)
+		: FAsyncTask<FSkinnedAssetAsyncBuildWorker>(InSkinnedAsset, MoveTemp(InBuildContext))
+		, SkinnedAsset(InSkinnedAsset)
+	{
+	}
+
+	const USkinnedAsset* SkinnedAsset;
+};
+
+#endif // #if WITH_EDITOR
diff --git a/Engine/Source/Runtime/Engine/Classes/Engine/SkinnedAssetCommon.h b/Engine/Source/Runtime/Engine/Classes/Engine/SkinnedAssetCommon.h
new file mode 100644
--- /dev/null
+++ b/Engine/Source/Runtime/Engine/Classes/Engine/SkinnedAssetCommon.h
@@ -0,0 +1,288 @@
+// Copyright Epic Games, Inc. All Rights Reserved.
+
+#pragma once
+
+/**
+ * Contains the shared data that is used by skinned assets
+ */
+
+#include "CoreMinimal.h"
+#include "Components.h"
+#include "BoneContainer.h"
+#include "PerPlatformProperties.h"
+#include "PerQualityLevelProperties.h"
+#include "SkeletalMeshReductionSettings.h"
+#include "SkinnedAssetCommon.generated.h"
+
+struct FSkeletalMeshLODGroupSettings;
+class UAnimSequence;
+class FSkeletalMeshLODModel;
+
+
+UENUM()
+enum class ESkinCacheUsage : uint8
+{
+	// Auto will defer to child or global behavior based on context. If Support Ray Tracing is enabled on the mesh, will imply Enabled
+	Auto		= 0,
+
+	// Mesh will not use the skin cache. If Support Ray Tracing is enabled on the mesh, will imply Enabled
+	Disabled	= uint8(-1),
+
+	// Mesh will use the skin cache
+	Enabled		= 1,
+};
+
+UENUM()
+enum class ESkinCacheDefaultBehavior : uint8
+{
+	// All skeletal meshes are excluded from the skin cache. Each must opt in individually. If Support Ray Tracing is enabled on a mesh, will force inclusive behavior on that mesh
+	Exclusive = 0,
+
+	// All skeletal meshes are included into the skin cache. Each must opt out individually
+	Inclusive = 1,
+};
+
+USTRUCT()
+struct FSectionReference
+{
+	GENERATED_USTRUCT_BODY()
+
+	/** Index of the section we reference. **/
+	UPROPERTY(EditAnywhere, Category = SectionReference)
+	int32 SectionIndex;
+
+	FSectionReference()
+		: SectionIndex(INDEX_NONE)
+	{
+	}
+
+	FSectionReference(const int32& InSectionIndex)
+		: SectionIndex(InSectionIndex)
+	{
+	}
+
+	bool operator==(const FSectionReference& Other) const
+	{
+		return SectionIndex == Other.SectionIndex;
+	}
+
+#if WITH_EDITOR
+	/** return true if it has a valid section index for LodModel parameter **/
+	ENGINE_API bool IsValidToEvaluate(const FSkeletalMeshLODModel& LodModel) const;
+
+	const struct FSkelMeshSection* GetMeshLodSection(const FSkeletalMeshLODModel& LodModel) const;
+	int32 GetMeshLodSectionIndex(const FSkeletalMeshLODModel& LodModel) const;
+#endif
+
+	bool Serialize(FArchive& Ar)
+	{
+		Ar << SectionIndex;
+		return true;
+	}
+
+	friend FArchive& operator<<(FArchive& Ar, FSectionReference& B)
+	{
+		B.Serialize(Ar);
+		return Ar;
+	}
+};
+
+/** Struct containing information for a particular LOD level, such as materials and info for when to use it. */
+USTRUCT()
+struct FSkeletalMeshLODInfo
+{
+	GENERATED_USTRUCT_BODY()
+
+	/** 
+	 * ScreenSize to display this LOD.
+	 * The screen size is based around the projected diameter of the bounding
+	 * sphere of the model. i.e. 0.5 means half the screen's maximum dimension.
+	 */
+	UPROPERTY(EditAnywhere, Category=SkeletalMeshLODInfo)
+	FPerPlatformFloat ScreenSize;
+
+	/**	Used to avoid 'flickering' when on LOD boundary. Only taken into account when moving from complex->simple. */
+	UPROPERTY(EditAnywhere, Category=SkeletalMeshLODInfo, meta=(DisplayName="LOD Hysteresis"))
+	float LODHysteresis;
+
+	/** Mapping table from this LOD's materials to the USkeletalMesh materials array.
+	 * section index is the key
+	 * remapped material index is the value, can be INDEX_NONE for no remapping
+	 */
+	UPROPERTY()
+	TArray<int32> LODMaterialMap;
+
+#if WITH_EDITORONLY_DATA
+	/** Per-section control over whether to enable shadow casting. */
+	UPROPERTY()
+	TArray<bool> bEnableShadowCasting_DEPRECATED;
+
+	/** This has been removed in editor. We could re-apply this in import time or by mesh reduction utilities*/
+	UPROPERTY()
+	TArray<FName> RemovedBones_DEPRECATED;
+#endif
+
+	/** build settings to apply when building render data. */
+	UPROPERTY(EditAnywhere, Category = BuildSettings)
+	FSkeletalMeshBuildSettings BuildSettings;
+
+	/** Reduction settings to apply when building render data. */
+	UPROPERTY(EditAnywhere, Category = ReductionSettings)
+	FSkeletalMeshOptimizationSettings ReductionSettings;
+
+	/** Bones which should be removed from the skeleton for the LOD level */
+	UPROPERTY(EditAnywhere, Category = ReductionSettings)
+	TArray<FBoneReference> BonesToRemove;
+
+	/** Bones which should be prioritized for the quality, this will be weighted toward keeping source data. Use WeightOfPrioritization to control the value. */
+	UPROPERTY(EditAnywhere, Category = ReductionSettings)
+	TArray<FBoneReference> BonesToPrioritize;
+
+	/** Sections which should be prioritized for the quality, this will be weighted toward keeping source data. Use WeightOfPrioritization to control the value. */
+	UPROPERTY(EditAnywhere, Category = ReductionSettings)
+	TArray<FSectionReference> SectionsToPrioritize;
+
+	/** How much to consideration to give BonesToPrioritize and SectionsToPrioritize.  The weight is an additional vertex simplification penalty where 0 means nothing. */
+	UPROPERTY(EditAnywhere, Category = ReductionSettings, meta = (UIMin = "0.0", ClampMin = "0.0"))
+	float WeightOfPrioritization;
+
+	/** Pose which should be used to reskin vertex influences for which the bones will be removed in this LOD level, uses ref-pose by default */
+	UPROPERTY(EditAnywhere, Category = ReductionSettings)
+	TObjectPtr<UAnimSequence> BakePose;
+
+	/** This is used when you are sharing the LOD settings, but you'd like to override the BasePose. This precedes prior to BakePose*/
+	UPROPERTY(EditAnywhere, Category = ReductionSettings)
+	TObjectPtr<UAnimSequence> BakePoseOverride;
+
+	/** The filename of the file tha was used to import this LOD if it was not auto generated. */
+	UPROPERTY(VisibleAnywhere, Category= SkeletalMeshLODInfo, AdvancedDisplay)
+	FString SourceImportFilename;
+
+	/**
+	 * How this LOD uses the skin cache feature. Auto will defer to the default project global option. If Support Ray Tracing is enabled on the mesh, will imply Enabled
+	 */
+	UPROPERTY(EditAnywhere, Category = SkeletalMeshLODInfo)
+	ESkinCacheUsage SkinCacheUsage = ESkinCacheUsage::Auto;
+
+	/** The Morph target position error tolerance in microns. Larger values result in better compression and lower memory footprint, but also lower quality. */
+	UPROPERTY(EditAnywhere, Category = SkeletalMeshLODInfo, meta = (UIMin = "0.01", ClampMin = "0.01", UIMax = "10000.0", ClampMax = "10000.0"))
+	float MorphTargetPositionErrorTolerance = 20.0f;
+
+	/** Whether to disable morph targets for this LOD. */
+	UPROPERTY()
+	uint8 bHasBeenSimplified:1;
+
+	UPROPERTY()
+	uint8 bHasPerLODVertexColors : 1;
+
+	/** Keeps this LODs data on the CPU so it can be used for things such as sampling in FX. */
+	UPROPERTY(EditAnywhere, Category = SkeletalMeshLODInfo)
+	uint8 bAllowCPUAccess : 1;
+
+	/**
+	Mesh supports uniformly distributed sampling in constant time.
+	Memory cost is 8 bytes per triangle.
+	Example usage is uniform spawning of particles.
+	*/
+	UPROPERTY(EditAnywhere, AdvancedDisplay, Category = SkeletalMeshLODInfo, meta=(EditCondition="bAllowCPUAccess"))
+	uint8 bSupportUniformlyDistributedSampling : 1;
+
+#if WITH_EDITORONLY_DATA
+	/*
+	 * This boolean specify if the LOD was imported with the base mesh or not.
+	 */
+	UPROPERTY()
+	uint8 bImportWithBaseMesh:1;
+
+	//Temporary build GUID data
+	//We use this GUID to store the LOD Key so we can know if the LOD needs to be rebuilt
+	//This GUID is set when we Cache the render data (build function)
+	FGuid BuildGUID;
+
+	ENGINE_API FGuid ComputeDeriveDataCacheKey(const FSkeletalMeshLODGroupSettings* SkeletalMeshLODGroupSettings);
+#endif
+
+	FSkeletalMeshLODInfo()
+		: ScreenSize(1.0)
+		, LODHysteresis(0.0f)
+		, WeightOfPrioritization(1.f)
+		, BakePose(nullptr)
+		, BakePoseOverride(nullptr)
+		, bHasBeenSimplified(false)
+		, bHasPerLODVertexColors(false)
+		, bAllowCPUAccess(false)
+		, bSupportUniformlyDistributedSampling(false)
+#if WITH_EDITORONLY_DATA
+		, bImportWithBaseMesh(false)
+#endif
+	{
+#if WITH_EDITORONLY_DATA
+		BuildGUID.Invalidate();
+#endif
+	}
+
+};
+
+//~ Begin Material Interface for USkeletalMesh - contains a material and a shadow casting flag
+USTRUCT(BlueprintType)
+struct FSkeletalMaterial
+{
+	GENERATED_USTRUCT_BODY()
+
+	FSkeletalMaterial()
+		: MaterialInterface( NULL )
+		, MaterialSlotName( NAME_None )
+#if WITH_EDITORONLY_DATA
+		, bEnableShadowCasting_DEPRECATED(true)
+		, bRecomputeTangent_DEPRECATED(false)
+		, ImportedMaterialSlotName( NAME_None )
+#endif
+	{
+
+	}
+
+	FSkeletalMaterial( class UMaterialInterface* InMaterialInterface
+						, bool bInEnableShadowCasting = true
+						, bool bInRecomputeTangent = false
+						, FName InMaterialSlotName = NAME_None
+						, FName InImportedMaterialSlotName = NAME_None)
+		: MaterialInterface( InMaterialInterface )
+		, MaterialSlotName(InMaterialSlotName)
+#if WITH_EDITORONLY_DATA
+		, bEnableShadowCasting_DEPRECATED(bInEnableShadowCasting)
+		, bRecomputeTangent_DEPRECATED(bInRecomputeTangent)
+		, ImportedMaterialSlotName(InImportedMaterialSlotName)
+#endif //WITH_EDITORONLY_DATA
+	{
+
+	}
+
+	friend FArchive& operator<<( FArchive& Ar, FSkeletalMaterial& Elem );
+#if WITH_EDITORONLY_DATA
+	static void DeclareCustomVersions(FArchive& Ar);
+#endif
+
+	ENGINE_API friend bool operator==( const FSkeletalMaterial& LHS, const FSkeletalMaterial& RHS );
+	ENGINE_API friend bool operator==( const FSkeletalMaterial& LHS, const UMaterialInterface& RHS );
+	ENGINE_API friend bool operator==( const UMaterialInterface& LHS, const FSkeletalMaterial& RHS );
+
+	UPROPERTY(EditAnywhere, BlueprintReadWrite, Category=SkeletalMesh)
+	TObjectPtr<class UMaterialInterface> 	MaterialInterface;
+	
+	/*This name should be use by the gameplay to avoid error if the skeletal mesh Materials array topology change*/
+	UPROPERTY(EditAnywhere, BlueprintReadWrite, Category = SkeletalMesh)
+	FName						MaterialSlotName;
+#if WITH_EDITORONLY_DATA
+	UPROPERTY()
+	bool						bEnableShadowCasting_DEPRECATED;
+	UPROPERTY()
+	bool						bRecomputeTangent_DEPRECATED;
+	/*This name should be use when we re-import a skeletal mesh so we can order the Materials array like it should be*/
+	UPROPERTY(VisibleAnywhere, Category = SkeletalMesh)
+	FName						ImportedMaterialSlotName;
+#endif //WITH_EDITORONLY_DATA
+
+	/** Data used for texture streaming relative to each UV channels. */
+	UPROPERTY(VisibleAnywhere, BlueprintReadOnly, Category = SkeletalMesh)
+	FMeshUVChannelInfo			UVChannelData;
+};
\ No newline at end of file
diff --git a/Engine/Source/Runtime/Engine/Private/AssetCompilingManager.cpp b/Engine/Source/Runtime/Engine/Private/AssetCompilingManager.cpp
--- a/Engine/Source/Runtime/Engine/Private/AssetCompilingManager.cpp
+++ b/Engine/Source/Runtime/Engine/Private/AssetCompilingManager.cpp
@@ -22,7 +22,7 @@ LLM_DEFINE_TAG(AssetCompilation, NAME_None, NAME_None, GET_STATFNAME(STAT_AssetC
 #include "ObjectCacheContext.h"
 #include "AsyncCompilationHelpers.h"
 #include "Experimental/Misc/ExecutionResource.h"
-#include "SkeletalMeshCompiler.h"
+#include "SkinnedAssetCompiler.h"
 #include "Algo/TopologicalSort.h"
 #include "Algo/Find.h"
 #include "ProfilingDebugging/CountersTrace.h"
@@ -267,7 +267,7 @@ FAssetCompilingManager::FAssetCompilingManager()
 	AssetCompilingManagerImpl::EnsureInitializedCVars();
 
 	RegisterManager(&FStaticMeshCompilingManager::Get());
-	RegisterManager(&FSkeletalMeshCompilingManager::Get());
+	RegisterManager(&FSkinnedAssetCompilingManager::Get());
 	RegisterManager(&FTextureCompilingManager::Get());
 	RegisterManager(&FActorDeferredScriptManager::Get());
 	RegisterManager(&FSoundWaveCompilingManager::Get());
diff --git a/Engine/Source/Runtime/Engine/Private/Components/SkeletalMeshComponent.cpp b/Engine/Source/Runtime/Engine/Private/Components/SkeletalMeshComponent.cpp
--- a/Engine/Source/Runtime/Engine/Private/Components/SkeletalMeshComponent.cpp
+++ b/Engine/Source/Runtime/Engine/Private/Components/SkeletalMeshComponent.cpp
@@ -25,7 +25,7 @@
 #include "Animation/AnimBlueprint.h"
 #include "SkeletalRender.h"
 #include "HAL/LowLevelMemTracker.h"
-#include "SkeletalMeshCompiler.h"
+#include "SkinnedAssetCompiler.h"
 #include "Rendering/SkeletalMeshRenderData.h"
 #include "SkeletalDebugRendering.h"
 
@@ -2522,7 +2522,7 @@ void USkeletalMeshComponent::DispatchParallelEvaluationTasks(FActorComponentTick
 	// We can only finish compilation on the game-thread, so wait here before spawning eval tasks.
 	if (GetSkeletalMesh() && GetSkeletalMesh()->IsCompiling())
 	{
-		FSkeletalMeshCompilingManager::Get().FinishCompilation({ GetSkeletalMesh() });
+		FSkinnedAssetCompilingManager::Get().FinishCompilation({ GetSkeletalMesh() });
 	}
 #endif
 
diff --git a/Engine/Source/Runtime/Engine/Private/Components/SkinnedMeshComponent.cpp b/Engine/Source/Runtime/Engine/Private/Components/SkinnedMeshComponent.cpp
--- a/Engine/Source/Runtime/Engine/Private/Components/SkinnedMeshComponent.cpp
+++ b/Engine/Source/Runtime/Engine/Private/Components/SkinnedMeshComponent.cpp
@@ -16,7 +16,6 @@
 #include "SkeletalRenderGPUSkin.h"
 #include "SkeletalRenderStatic.h"
 #include "Animation/AnimStats.h"
-#include "SkeletalMeshCompiler.h"
 #include "SkeletalMeshDeformerHelpers.h"
 #include "Engine/SkeletalMeshSocket.h"
 #include "PhysicsEngine/PhysicsAsset.h"
diff --git a/Engine/Source/Runtime/Engine/Private/SkeletalMesh.cpp b/Engine/Source/Runtime/Engine/Private/SkeletalMesh.cpp
--- a/Engine/Source/Runtime/Engine/Private/SkeletalMesh.cpp
+++ b/Engine/Source/Runtime/Engine/Private/SkeletalMesh.cpp
@@ -68,7 +68,7 @@
 #include "IMeshBuilderModule.h"
 #include "IMeshReductionManagerModule.h"
 #include "IMeshReductionInterfaces.h"
-#include "SkeletalMeshCompiler.h"
+#include "SkinnedAssetCompiler.h"
 #include "MeshUtilities.h"
 #include "Engine/SkeletalMeshEditorData.h"
 #include "DerivedDataCacheInterface.h"
@@ -360,61 +360,6 @@ void FScopedSkeletalMeshPostEditChange::SetSkeletalMesh(USkeletalMesh* InSkeleta
 	}
 }
 
-
-/*-----------------------------------------------------------------------------
-	FSectionReference
------------------------------------------------------------------------------*/
-
-bool FSectionReference::IsValidToEvaluate(const FSkeletalMeshLODModel& LodModel) const
-{
-	return GetMeshLodSection(LodModel) != nullptr;
-}
-
-const FSkelMeshSection* FSectionReference::GetMeshLodSection(const FSkeletalMeshLODModel& LodModel) const
-{
-	for (int32 LodModelSectionIndex = 0; LodModelSectionIndex < LodModel.Sections.Num(); ++LodModelSectionIndex)
-	{
-		const FSkelMeshSection& Section = LodModel.Sections[LodModelSectionIndex];
-		if (Section.ChunkedParentSectionIndex == INDEX_NONE && SectionIndex == Section.OriginalDataSectionIndex)
-		{
-			return &Section;
-		}
-	}
-	return nullptr;
-}
-
-int32 FSectionReference::GetMeshLodSectionIndex(const FSkeletalMeshLODModel& LodModel) const
-{
-	for (int32 LodModelSectionIndex = 0; LodModelSectionIndex < LodModel.Sections.Num(); ++LodModelSectionIndex)
-	{
-		const FSkelMeshSection& Section = LodModel.Sections[LodModelSectionIndex];
-		if (Section.ChunkedParentSectionIndex == INDEX_NONE && SectionIndex == Section.OriginalDataSectionIndex)
-		{
-			return LodModelSectionIndex;
-		}
-	}
-	return INDEX_NONE;
-}
-
-
-/*-----------------------------------------------------------------------------
-	FSkeletalMeshAsyncBuildWorker
------------------------------------------------------------------------------*/
-
-
-void FSkeletalMeshAsyncBuildWorker::DoWork()
-{
-	if (PostLoadContext.IsSet())
-	{
-		SkeletalMesh->ExecutePostLoadInternal(*PostLoadContext);
-	}
-
-	if (BuildContext.IsSet())
-	{
-		SkeletalMesh->ExecuteBuildInternal(*BuildContext);
-	}
-}
-
 const FString& GetSkeletalMeshDerivedDataVersion()
 {
 	static FString CachedVersionString = FDevSystemGuids::GetSystemGuid(FDevSystemGuids::Get().SkeletalMeshDerivedDataVersion).ToString();
@@ -451,7 +396,6 @@ namespace SkeletalMeshImpl
 /*-----------------------------------------------------------------------------
 	USkeletalMesh
 -----------------------------------------------------------------------------*/
-ENUM_CLASS_FLAGS(USkeletalMesh::EAsyncPropertyLockType);
 
 USkeletalMesh::USkeletalMesh(const FObjectInitializer& ObjectInitializer)
 	: Super(ObjectInitializer)
@@ -501,7 +445,7 @@ void USkeletalMesh::PostInitProperties()
 
 FBoxSphereBounds USkeletalMesh::GetBounds() const
 {
-	WaitUntilAsyncPropertyReleased(ESkeletalMeshAsyncProperties::ExtendedBounds, EAsyncPropertyLockType::ReadOnly);
+	WaitUntilAsyncPropertyReleased(ESkeletalMeshAsyncProperties::ExtendedBounds, ESkinnedAssetAsyncPropertyLockType::ReadOnly);
 	PRAGMA_DISABLE_DEPRECATION_WARNINGS
 	return ExtendedBounds;
 	PRAGMA_ENABLE_DEPRECATION_WARNINGS
@@ -509,7 +453,7 @@ FBoxSphereBounds USkeletalMesh::GetBounds() const
 
 FBoxSphereBounds USkeletalMesh::GetImportedBounds() const
 {
-	WaitUntilAsyncPropertyReleased(ESkeletalMeshAsyncProperties::ImportedBounds, EAsyncPropertyLockType::ReadOnly);
+	WaitUntilAsyncPropertyReleased(ESkeletalMeshAsyncProperties::ImportedBounds, ESkinnedAssetAsyncPropertyLockType::ReadOnly);
 	PRAGMA_DISABLE_DEPRECATION_WARNINGS
 	return ImportedBounds;
 	PRAGMA_ENABLE_DEPRECATION_WARNINGS
@@ -1539,7 +1483,7 @@ void CachePlatform(USkeletalMesh* Mesh, const ITargetPlatform* TargetPlatform, F
 {
 	//Cache the platform, dcc should be valid so it will be fast
 	check(TargetPlatform);
-	FSkeletalMeshBuildContext Context;
+	FSkinnedAssetBuildContext Context;
 	Context.bIsSerializeSaving = bIsSerializeSaving;
 	PlatformRenderData->Cache(TargetPlatform, Mesh, &Context);
 }
@@ -1595,7 +1539,7 @@ void USkeletalMesh::Serialize( FArchive& Ar )
 		// the protection put in place to automatically finish compilation if a locked property is accessed will not work. 
 		// We have no choice but to force finish the compilation here to avoid potential race conditions between 
 		// async compilation and the serialization.
-		FSkeletalMeshCompilingManager::Get().FinishCompilation({ this });
+		FSkinnedAssetCompilingManager::Get().FinishCompilation({ this });
 	}
 #endif
 
@@ -1885,89 +1829,6 @@ void USkeletalMesh::FlushRenderState()
 
 #if WITH_EDITOR
 
-thread_local const USkeletalMesh* FSkeletalMeshAsyncBuildScope::SkeletalMeshBeingAsyncCompiled = nullptr;
-
-FSkeletalMeshAsyncBuildScope::FSkeletalMeshAsyncBuildScope(const USkeletalMesh* SkeletalMesh)
-{
-	PreviousScope = SkeletalMeshBeingAsyncCompiled;
-	SkeletalMeshBeingAsyncCompiled = SkeletalMesh;
-}
-
-FSkeletalMeshAsyncBuildScope::~FSkeletalMeshAsyncBuildScope()
-{
-	check(SkeletalMeshBeingAsyncCompiled);
-	SkeletalMeshBeingAsyncCompiled = PreviousScope;
-}
-
-bool FSkeletalMeshAsyncBuildScope::ShouldWaitOnLockedProperties(const USkeletalMesh* SkeletalMesh)
-{
-	return SkeletalMeshBeingAsyncCompiled != SkeletalMesh;
-}
-
-void USkeletalMesh::WaitUntilAsyncPropertyReleased(ESkeletalMeshAsyncProperties AsyncProperties, EAsyncPropertyLockType LockType) const
-{
-	// We need to protect internal skeletal mesh data from race conditions during async build
-	if (IsCompiling())
-	{
-		if (FSkeletalMeshAsyncBuildScope::ShouldWaitOnLockedProperties(this))
-		{
-			bool bIsLocked = true;
-			// We can remove the lock if we're accessing in read-only and there is no write-lock
-			if ((LockType & EAsyncPropertyLockType::ReadOnly) == EAsyncPropertyLockType::ReadOnly)
-			{
-				// Maintain the lock if the write-lock bit is non-zero
-				bIsLocked &= (ModifiedProperties & (uint64)AsyncProperties) != 0;
-			}
-
-			if (bIsLocked)
-			{
-				FString PropertyName = StaticEnum<ESkeletalMeshAsyncProperties>()->GetNameByValue((int64)AsyncProperties).ToString();
-				TRACE_CPUPROFILER_EVENT_SCOPE_TEXT(*FString::Printf(TEXT("SkeletalMeshCompilationStall %s"), *PropertyName));
-
-				if (IsInGameThread())
-				{
-					UE_LOG(
-						LogSkeletalMesh,
-						Verbose,
-						TEXT("Accessing property %s of the SkeletalMesh while it is still being built asynchronously will force it to be compiled before continuing. "
-							"For better performance, consider making the caller async aware so it can wait until the static mesh is ready to access this property."
-							"To better understand where those calls are coming from, you can use Editor.AsyncAssetDumpStallStacks on the console."),
-						*PropertyName
-					);
-
-					FSkeletalMeshCompilingManager::Get().FinishCompilation({ const_cast<USkeletalMesh*>(this) });
-				}
-				else
-				{
-					// Trying to access a property from another thread that cannot force finish the compilation is invalid
-					ensureMsgf(
-						false,
-						TEXT("Accessing property %s of the SkeletalMesh while it is still being built asynchronously is only supported on the game-thread. "
-							"To avoid any race-condition, consider finishing the compilation before pushing tasks to other threads or making higher-level game-thread code async aware so it "
-							"schedules the task only when the static mesh's compilation is finished. If this is a blocker, you can disable async static mesh from the editor experimental settings."),
-						*PropertyName
-					);
-				}
-			}
-		}
-		// If we're accessing this property from the async build thread, make sure the property is still protected from access from other threads.
-		else
-		{
-			bool bIsLocked = true;
-			if ((LockType & EAsyncPropertyLockType::ReadOnly) == EAsyncPropertyLockType::ReadOnly)
-			{
-				bIsLocked &= (AccessedProperties & (uint64)AsyncProperties) != 0;
-			}
-
-			if ((LockType & EAsyncPropertyLockType::WriteOnly) == EAsyncPropertyLockType::WriteOnly)
-			{
-				bIsLocked &= (ModifiedProperties & (uint64)AsyncProperties) != 0;
-			}
-			ensureMsgf(bIsLocked, TEXT("Property %s has not been locked properly for use by async build"), *StaticEnum<ESkeletalMeshAsyncProperties>()->GetNameByValue((int64)AsyncProperties).ToString());
-		}
-	}
-}
-
 void USkeletalMesh::PrepareForAsyncCompilation()
 {
 	// Make sure statics are initialized before calling from multiple threads
@@ -1996,7 +1857,7 @@ void USkeletalMesh::PrepareForAsyncCompilation()
 	// from the game thread.
 
 	// Not touched during async build and can cause stalls when the content browser refresh its tiles.
-	ReleaseAsyncProperty(ESkeletalMeshAsyncProperties::ThumbnailInfo);
+	ReleaseAsyncProperty((uint64)ESkeletalMeshAsyncProperties::ThumbnailInfo);
 }
 
 void USkeletalMesh::Build()
@@ -2005,25 +1866,25 @@ void USkeletalMesh::Build()
 
 	if (IsCompiling())
 	{
-		FSkeletalMeshCompilingManager::Get().FinishCompilation({this});
+		FSkinnedAssetCompilingManager::Get().FinishCompilation({this});
 	}
 
-	FSkeletalMeshAsyncBuildScope AsyncBuildScope(this);
+	FSkinnedAssetAsyncBuildScope AsyncBuildScope(this);
 
-	FSkeletalMeshBuildContext Context;
+	FSkinnedAssetBuildContext Context;
 	BeginBuildInternal(Context);
 	
 	// Inline reduction is not thread-safe, prevent async build until this is fixed.
-	if (FSkeletalMeshCompilingManager::Get().IsAsyncCompilationAllowed(this))
+	if (FSkinnedAssetCompilingManager::Get().IsAsyncCompilationAllowed(this))
 	{
 		PrepareForAsyncCompilation();
 
-		FQueuedThreadPool* SkeletalMeshThreadPool = FSkeletalMeshCompilingManager::Get().GetThreadPool();
-		EQueuedWorkPriority BasePriority = FSkeletalMeshCompilingManager::Get().GetBasePriority(this);
+		FQueuedThreadPool* SkeletalMeshThreadPool = FSkinnedAssetCompilingManager::Get().GetThreadPool();
+		EQueuedWorkPriority BasePriority = FSkinnedAssetCompilingManager::Get().GetBasePriority(this);
 		check(AsyncTask == nullptr);
-		AsyncTask = MakeUnique<FSkeletalMeshAsyncBuildTask>(this, MoveTemp(Context));
+		AsyncTask = MakeUnique<FSkinnedAssetAsyncBuildTask>(this, MoveTemp(Context));
 		AsyncTask->StartBackgroundTask(SkeletalMeshThreadPool, BasePriority, EQueuedWorkFlags::DoNotRunInsideBusyWait);
-		FSkeletalMeshCompilingManager::Get().AddSkeletalMeshes({ this });
+		FSkinnedAssetCompilingManager::Get().AddSkinnedAssets({ this });
 	}
 	else
 	{
@@ -2032,7 +1893,7 @@ void USkeletalMesh::Build()
 	}
 }
 
-void USkeletalMesh::BeginBuildInternal(FSkeletalMeshBuildContext& Context)
+void USkeletalMesh::BeginBuildInternal(FSkinnedAssetBuildContext& Context)
 {
 	TRACE_CPUPROFILER_EVENT_SCOPE(USkeletalMesh::BeginBuildInternal);
 
@@ -2060,12 +1921,12 @@ void USkeletalMesh::BeginBuildInternal(FSkeletalMeshBuildContext& Context)
 	AcquireAsyncProperty();
 }
 
-void USkeletalMesh::ExecuteBuildInternal(FSkeletalMeshBuildContext& Context)
+void USkeletalMesh::ExecuteBuildInternal(FSkinnedAssetBuildContext& Context)
 {
 	TRACE_CPUPROFILER_EVENT_SCOPE(USkeletalMesh::ExecuteBuildInternal);
 
 	// This scope allows us to use any locked properties without causing stalls
-	FSkeletalMeshAsyncBuildScope AsyncBuildScope(this);
+	FSkinnedAssetAsyncBuildScope AsyncBuildScope(this);
 
 	// rebuild render data from imported model
 	CacheDerivedData(&Context);
@@ -2080,14 +1941,14 @@ void USkeletalMesh::ExecuteBuildInternal(FSkeletalMeshBuildContext& Context)
 	UpdateGenerateUpToData();
 }
 
-void USkeletalMesh::ApplyFinishBuildInternalData(FSkeletalMeshCompilationContext* ContextPtr)
+void USkeletalMesh::ApplyFinishBuildInternalData(FSkinnedAssetCompilationContext* ContextPtr)
 {
 	//We cannot execute this code outside of the game thread
 	checkf(IsInGameThread(), TEXT("Cannot execute function USkeletalMesh::ApplyFinishBuildInternalData asynchronously. Asset: %s"), *this->GetFullName());
 	check(ContextPtr);
 }
 
-void USkeletalMesh::FinishBuildInternal(FSkeletalMeshBuildContext& Context)
+void USkeletalMesh::FinishBuildInternal(FSkinnedAssetBuildContext& Context)
 {
 	TRACE_CPUPROFILER_EVENT_SCOPE(USkeletalMesh::FinishBuildInternal);
 
@@ -2740,32 +2601,6 @@ bool USkeletalMesh::IsAsyncTaskComplete() const
 	return AsyncTask == nullptr || AsyncTask->IsWorkDone();
 }
 
-void USkeletalMesh::AcquireAsyncProperty(ESkeletalMeshAsyncProperties AsyncProperties, EAsyncPropertyLockType LockType)
-{
-	if ((LockType & EAsyncPropertyLockType::ReadOnly) == EAsyncPropertyLockType::ReadOnly)
-	{
-		AccessedProperties |= (uint64)AsyncProperties;
-	}
-	
-	if ((LockType & EAsyncPropertyLockType::WriteOnly) == EAsyncPropertyLockType::WriteOnly)
-	{
-		ModifiedProperties |= (uint64)AsyncProperties;
-	}
-}
-
-void USkeletalMesh::ReleaseAsyncProperty(ESkeletalMeshAsyncProperties AsyncProperties, EAsyncPropertyLockType LockType)
-{
-	if ((LockType & EAsyncPropertyLockType::ReadOnly) == EAsyncPropertyLockType::ReadOnly)
-	{
-		AccessedProperties &= ~(uint64)AsyncProperties;
-	}
-
-	if ((LockType & EAsyncPropertyLockType::WriteOnly) == EAsyncPropertyLockType::WriteOnly)
-	{
-		ModifiedProperties &= ~(uint64)AsyncProperties;
-	}
-}
-
 bool USkeletalMesh::TryCancelAsyncTasks()
 {
 	if (AsyncTask)
@@ -3057,7 +2892,7 @@ bool USkeletalMesh::IsPostLoadThreadSafe() const
 	return false;	// PostLoad is not thread safe because of the call to InitMorphTargets, which can call VerifySmartName() that can mutate a shared map in the skeleton.
 }
 
-void USkeletalMesh::BeginPostLoadInternal(FSkeletalMeshPostLoadContext& Context)
+void USkeletalMesh::BeginPostLoadInternal(FSkinnedAssetPostLoadContext& Context)
 {
 #if WITH_EDITOR
 	TRACE_CPUPROFILER_EVENT_SCOPE(UStaticMesh::BeginPostLoadInternal);
@@ -3068,7 +2903,7 @@ void USkeletalMesh::BeginPostLoadInternal(FSkeletalMeshPostLoadContext& Context)
 	AcquireAsyncProperty();
 
 	// This scope allows us to use any locked properties without causing stalls
-	FSkeletalMeshAsyncBuildScope AsyncBuildScope(this);
+	FSkinnedAssetAsyncBuildScope AsyncBuildScope(this);
 
 	// Make sure the cloth assets have finished loading
 	// TODO: Remove all UObject PostLoad dependencies.
@@ -3240,41 +3075,15 @@ void USkeletalMesh::PostLoad()
 {
 	LLM_SCOPE(ELLMTag::SkeletalMesh);
 	Super::PostLoad();
-
-#if WITH_EDITOR
-	FSkeletalMeshAsyncBuildScope AsyncBuildScope(this);
-#endif
-
-	FSkeletalMeshPostLoadContext Context;
-	BeginPostLoadInternal(Context);
-
-#if WITH_EDITOR
-	if (FSkeletalMeshCompilingManager::Get().IsAsyncCompilationAllowed(this))
-	{
-		PrepareForAsyncCompilation();
-
-		FQueuedThreadPool* SkeletalMeshThreadPool = FSkeletalMeshCompilingManager::Get().GetThreadPool();
-		EQueuedWorkPriority BasePriority = FSkeletalMeshCompilingManager::Get().GetBasePriority(this);
-
-		AsyncTask = MakeUnique<FSkeletalMeshAsyncBuildTask>(this, MoveTemp(Context));
-		AsyncTask->StartBackgroundTask(SkeletalMeshThreadPool, BasePriority, EQueuedWorkFlags::DoNotRunInsideBusyWait);
-		FSkeletalMeshCompilingManager::Get().AddSkeletalMeshes({ this });
-	}
-	else
-#endif
-	{
-		ExecutePostLoadInternal(Context);
-		FinishPostLoadInternal(Context);
-	}
 }
 
-void USkeletalMesh::ExecutePostLoadInternal(FSkeletalMeshPostLoadContext& Context)
+void USkeletalMesh::ExecutePostLoadInternal(FSkinnedAssetPostLoadContext& Context)
 {
 #if WITH_EDITOR
 	TRACE_CPUPROFILER_EVENT_SCOPE(USkeletalMesh::ExecutePostLoadInternal);
 
 	// This scope allows us to use any locked properties without causing stalls
-	FSkeletalMeshAsyncBuildScope AsyncBuildScope(this);
+	FSkinnedAssetAsyncBuildScope AsyncBuildScope(this);
 
 	if (!GetOutermost()->bIsCookedForEditor)
 	{
@@ -3303,7 +3112,7 @@ void USkeletalMesh::ExecutePostLoadInternal(FSkeletalMeshPostLoadContext& Contex
 #endif // WITH_EDITOR
 }
 
-void USkeletalMesh::FinishPostLoadInternal(FSkeletalMeshPostLoadContext& Context)
+void USkeletalMesh::FinishPostLoadInternal(FSkinnedAssetPostLoadContext& Context)
 {
 	TRACE_CPUPROFILER_EVENT_SCOPE(USkeletalMesh::FinishPostLoadInternal);
 	
@@ -3312,7 +3121,7 @@ void USkeletalMesh::FinishPostLoadInternal(FSkeletalMeshPostLoadContext& Context
 	ClearInternalFlags(EInternalObjectFlags::Async);
 
 	// This scope allows us to use any locked properties without causing stalls
-	FSkeletalMeshAsyncBuildScope AsyncBuildScope(this);
+	FSkinnedAssetAsyncBuildScope AsyncBuildScope(this);
 
 	if (Context.bHasCachedDerivedData)
 	{
@@ -4156,16 +3965,6 @@ void USkeletalMesh::MoveMaterialFlagsToSections()
 
 #if WITH_EDITOR
 
-/** IInterface_AsyncCompilation implementation begin*/
-
-bool USkeletalMesh::IsCompiling() const
-{
-	return AsyncTask != nullptr || AccessedProperties.load(std::memory_order_relaxed) != 0;
-}
-
-/** IInterface_AsyncCompilation implementation end*/
-
-
 FDelegateHandle USkeletalMesh::RegisterOnClothingChange(const FSimpleMulticastDelegate::FDelegate& InDelegate)
 {
 	return OnClothingChange.Add(InDelegate);
@@ -4196,97 +3995,6 @@ bool USkeletalMesh::AreAllFlagsIdentical( const TArray<bool>& BoolArray ) const
 	return true;
 }
 
-bool operator== ( const FSkeletalMaterial& LHS, const FSkeletalMaterial& RHS )
-{
-	return ( LHS.MaterialInterface == RHS.MaterialInterface );
-}
-
-bool operator== ( const FSkeletalMaterial& LHS, const UMaterialInterface& RHS )
-{
-	return ( LHS.MaterialInterface == &RHS );
-}
-
-bool operator== ( const UMaterialInterface& LHS, const FSkeletalMaterial& RHS )
-{
-	return ( RHS.MaterialInterface == &LHS );
-}
-
-FArchive& operator<<(FArchive& Ar, FMeshUVChannelInfo& ChannelData)
-{
-	Ar << ChannelData.bInitialized;
-	Ar << ChannelData.bOverrideDensities;
-
-	for (int32 CoordIndex = 0; CoordIndex < TEXSTREAM_MAX_NUM_UVCHANNELS; ++CoordIndex)
-	{
-		Ar << ChannelData.LocalUVDensities[CoordIndex];
-	}
-
-	return Ar;
-}
-
-#if WITH_EDITORONLY_DATA
-void FSkeletalMaterial::DeclareCustomVersions(FArchive& Ar)
-{
-	Ar.UsingCustomVersion(FEditorObjectVersion::GUID);
-	Ar.UsingCustomVersion(FCoreObjectVersion::GUID);
-}
-#endif
-
-FArchive& operator<<(FArchive& Ar, FSkeletalMaterial& Elem)
-{
-	Ar.UsingCustomVersion(FEditorObjectVersion::GUID);
-	Ar.UsingCustomVersion(FCoreObjectVersion::GUID);
-
-	Ar << Elem.MaterialInterface;
-
-	//Use the automatic serialization instead of this custom operator
-	if (Ar.CustomVer(FEditorObjectVersion::GUID) >= FEditorObjectVersion::RefactorMeshEditorMaterials)
-	{
-		Ar << Elem.MaterialSlotName;
-
-		bool bSerializeImportedMaterialSlotName = !Ar.IsCooking() || Ar.CookingTarget()->HasEditorOnlyData();
-		if (Ar.CustomVer(FCoreObjectVersion::GUID) >= FCoreObjectVersion::SkeletalMaterialEditorDataStripping)
-		{
-			Ar << bSerializeImportedMaterialSlotName;
-		}
-		else if (!FPlatformProperties::HasEditorOnlyData())
-		{
-			bSerializeImportedMaterialSlotName = false;
-		}
-		if (bSerializeImportedMaterialSlotName)
-		{
-#if WITH_EDITORONLY_DATA
-			Ar << Elem.ImportedMaterialSlotName;
-#else
-			FName UnusedImportedMaterialSlotName;
-			Ar << UnusedImportedMaterialSlotName;
-#endif
-		}
-	}
-#if WITH_EDITORONLY_DATA
-	else
-	{
-		if (Ar.UEVer() >= VER_UE4_MOVE_SKELETALMESH_SHADOWCASTING)
-		{
-			Ar << Elem.bEnableShadowCasting_DEPRECATED;
-		}
-
-		Ar.UsingCustomVersion(FRecomputeTangentCustomVersion::GUID);
-		if (Ar.CustomVer(FRecomputeTangentCustomVersion::GUID) >= FRecomputeTangentCustomVersion::RuntimeRecomputeTangent)
-		{
-			Ar << Elem.bRecomputeTangent_DEPRECATED;
-		}
-	}
-#endif
-	
-	if (!Ar.IsLoading() || Ar.CustomVer(FRenderingObjectVersion::GUID) >= FRenderingObjectVersion::TextureStreamingMeshUVChannelData)
-	{
-		Ar << Elem.UVChannelData;
-	}
-
-	return Ar;
-}
-
 TArray<USkeletalMeshSocket*> USkeletalMesh::GetActiveSocketList() const
 {
 	TArray<USkeletalMeshSocket*> ActiveSockets = Sockets;
@@ -4423,7 +4131,7 @@ namespace InternalSkeletalMeshHelper
 	}
 } //namespace InternalSkeletalMeshHelper
 
-void USkeletalMesh::CacheDerivedData(FSkeletalMeshCompilationContext* ContextPtr)
+void USkeletalMesh::CacheDerivedData(FSkinnedAssetCompilationContext* ContextPtr)
 {
 	TRACE_CPUPROFILER_EVENT_SCOPE(USkeletalMesh::CacheDerivedData);
 	check(ContextPtr);
@@ -5510,120 +5218,7 @@ FText USkeletalMesh::GetSourceFileLabelFromIndex(int32 SourceFileIndex)
 	return RealSourceFileIndex == 0 ? NSSkeletalMeshSourceFileLabels::GeoAndSkinningText() : RealSourceFileIndex == 1 ? NSSkeletalMeshSourceFileLabels::GeometryText() : NSSkeletalMeshSourceFileLabels::SkinningText();
 }
 
-static void SerializeReductionSettingsForDDC(FArchive& Ar, FSkeletalMeshOptimizationSettings& ReductionSettings)
-{
-	check(Ar.IsSaving());
-	// Note: this serializer is only used to build the mesh DDC key, no versioning is required
-	FArchive_Serialize_BitfieldBool(Ar, ReductionSettings.bRemapMorphTargets);
-	FArchive_Serialize_BitfieldBool(Ar, ReductionSettings.bRecalcNormals);
-	FArchive_Serialize_BitfieldBool(Ar, ReductionSettings.bEnforceBoneBoundaries);
-	FArchive_Serialize_BitfieldBool(Ar, ReductionSettings.bMergeCoincidentVertBones);
-	FArchive_Serialize_BitfieldBool(Ar, ReductionSettings.bLockEdges);
-	FArchive_Serialize_BitfieldBool(Ar, ReductionSettings.bLockColorBounaries);
-	Ar << ReductionSettings.TerminationCriterion;
-	Ar << ReductionSettings.NumOfTrianglesPercentage;
-	Ar << ReductionSettings.NumOfVertPercentage;
-	Ar << ReductionSettings.MaxNumOfTriangles;
-	Ar << ReductionSettings.MaxNumOfVerts;
-
-	// Keep old DDC keys if these are not set
-	if (ReductionSettings.MaxNumOfTrianglesPercentage != MAX_uint32 || 
-		ReductionSettings.MaxNumOfVertsPercentage != MAX_uint32)
-	{
-		uint32 AvoidCachePoisoningFromOldBug = 0;
-		Ar << AvoidCachePoisoningFromOldBug;
-		Ar << ReductionSettings.MaxNumOfTrianglesPercentage;
-		Ar << ReductionSettings.MaxNumOfVertsPercentage;
-	}
-
-	Ar << ReductionSettings.MaxDeviationPercentage;
-	Ar << ReductionSettings.ReductionMethod;
-	Ar << ReductionSettings.SilhouetteImportance;
-	Ar << ReductionSettings.TextureImportance;
-	Ar << ReductionSettings.ShadingImportance;
-	Ar << ReductionSettings.SkinningImportance;
-	Ar << ReductionSettings.WeldingThreshold;
-	Ar << ReductionSettings.NormalsThreshold;
-	Ar << ReductionSettings.MaxBonesPerVertex;
-	Ar << ReductionSettings.VolumeImportance;
-	Ar << ReductionSettings.BaseLOD;
-}
-
-static void SerializeBuildSettingsForDDC(FArchive& Ar, FSkeletalMeshBuildSettings& BuildSettings)
-{
-	check(Ar.IsSaving());
-	// Note: this serializer is only used to build the mesh DDC key, no versioning is required
-	FArchive_Serialize_BitfieldBool(Ar, BuildSettings.bRecomputeNormals);
-	FArchive_Serialize_BitfieldBool(Ar, BuildSettings.bRecomputeTangents);
-	FArchive_Serialize_BitfieldBool(Ar, BuildSettings.bUseMikkTSpace);
-	FArchive_Serialize_BitfieldBool(Ar, BuildSettings.bComputeWeightedNormals);
-	FArchive_Serialize_BitfieldBool(Ar, BuildSettings.bRemoveDegenerates);
-	FArchive_Serialize_BitfieldBool(Ar, BuildSettings.bUseFullPrecisionUVs);
-	FArchive_Serialize_BitfieldBool(Ar, BuildSettings.bUseBackwardsCompatibleF16TruncUVs);
-	FArchive_Serialize_BitfieldBool(Ar, BuildSettings.bUseHighPrecisionTangentBasis);
-	Ar << BuildSettings.ThresholdPosition;
-	Ar << BuildSettings.ThresholdTangentNormal;
-	Ar << BuildSettings.ThresholdUV;
-	Ar << BuildSettings.MorphThresholdPosition;
-}
-
-
-FGuid FSkeletalMeshLODInfo::ComputeDeriveDataCacheKey(const FSkeletalMeshLODGroupSettings* SkeletalMeshLODGroupSettings)
-{
-	const bool bIs16BitfloatBufferSupported = GVertexElementTypeSupport.IsSupported(VET_Half2);
-
-	// Serialize the LOD info members, the BuildSettings and the ReductionSettings into a temporary array.
-	TArray<uint8> TempBytes;
-	TempBytes.Reserve(64);
-	//The archive is flagged as persistent so that machines of different endianness produce identical binary results.
-	FMemoryWriter Ar(TempBytes, /*bIsPersistent=*/ true);
-
-	Ar << BonesToRemove;
-	Ar << BonesToPrioritize;
-	Ar << SectionsToPrioritize;
-	Ar << WeightOfPrioritization;
-	Ar << MorphTargetPositionErrorTolerance;
-
-	//TODO: Ask the derivedata key of the UObject reference by FSoftObjectPath. So if someone change the UObject, this LODs will get dirty
-	//and will be rebuild.
-	if (BakePose != nullptr)
-	{
-		FString BakePosePath = BakePose->GetFullName();
-		Ar << BakePosePath;
-	}
-	if (BakePoseOverride != nullptr)
-	{
-		FString BakePoseOverridePath = BakePoseOverride->GetFullName();
-		Ar << BakePoseOverridePath;
-	}
-	FArchive_Serialize_BitfieldBool(Ar, bAllowCPUAccess);
-	FArchive_Serialize_BitfieldBool(Ar, bSupportUniformlyDistributedSampling);
-
-	//Use the LOD settings asset if there is one
-	FSkeletalMeshOptimizationSettings RealReductionSettings = ReductionSettings;
-	if (SkeletalMeshLODGroupSettings != nullptr)
-	{
-		RealReductionSettings = SkeletalMeshLODGroupSettings->GetReductionSettings();
-	}
-
-	if (!BuildSettings.bUseFullPrecisionUVs && !bIs16BitfloatBufferSupported)
-	{
-		BuildSettings.bUseFullPrecisionUVs = true;
-	}
-	SerializeBuildSettingsForDDC(Ar, BuildSettings);
-	SerializeReductionSettingsForDDC(Ar, RealReductionSettings);
-
-	FSHA1 Sha;
-	Sha.Update(TempBytes.GetData(), TempBytes.Num() * TempBytes.GetTypeSize());
-	Sha.Final();
-	// Retrieve the hash and use it to construct a pseudo-GUID.
-	uint32 Hash[5];
-	Sha.GetHash((uint8*)Hash);
-	FGuid Guid = FGuid(Hash[0] ^ Hash[4], Hash[1], Hash[2], Hash[3]);
-	return Guid;
-}
-
-#endif //WITH_EDITOR
+#endif //WITH_EDITORONLY_DATA
 
 
 TArray<FString> USkeletalMesh::K2_GetAllMorphTargetNames() const
@@ -5728,6 +5323,17 @@ void USkeletalMesh::OnLodStrippingQualityLevelChanged(IConsoleVariable* Variable
 #endif
 }
 
+void USkeletalMesh::WaitUntilAsyncPropertyReleased(ESkeletalMeshAsyncProperties AsyncProperties, ESkinnedAssetAsyncPropertyLockType LockType) const
+{
+	// Cast strongly-typed enum to uint64
+	WaitUntilAsyncPropertyReleasedInternal((uint64)AsyncProperties, LockType);
+}
+
+FString USkeletalMesh::GetAsyncPropertyName(uint64 Property) const
+{
+	return StaticEnum<ESkeletalMeshAsyncProperties>()->GetNameByValue(Property).ToString();
+}
+
 /*-----------------------------------------------------------------------------
 USkeletalMeshSocket
 -----------------------------------------------------------------------------*/
diff --git a/Engine/Source/Runtime/Engine/Private/SkeletalMeshCompiler.cpp b/Engine/Source/Runtime/Engine/Private/SkeletalMeshCompiler.cpp
deleted file mode 100644
diff --git a/Engine/Source/Runtime/Engine/Private/SkeletalMeshRenderData.cpp b/Engine/Source/Runtime/Engine/Private/SkeletalMeshRenderData.cpp
--- a/Engine/Source/Runtime/Engine/Private/SkeletalMeshRenderData.cpp
+++ b/Engine/Source/Runtime/Engine/Private/SkeletalMeshRenderData.cpp
@@ -215,7 +215,7 @@ FString FSkeletalMeshRenderData::GetDerivedDataKey(const ITargetPlatform* Target
 	return Owner->BuildDerivedDataKey(TargetPlatform);
 }
 
-void FSkeletalMeshRenderData::Cache(const ITargetPlatform* TargetPlatform, USkinnedAsset* Owner, FSkeletalMeshCompilationContext* ContextPtr)
+void FSkeletalMeshRenderData::Cache(const ITargetPlatform* TargetPlatform, USkinnedAsset* Owner, FSkinnedAssetCompilationContext* ContextPtr)
 {
 	check(Owner);
 	// Disable ContextPtr check because only USkeletalMesh supports it.
diff --git a/Engine/Source/Runtime/Engine/Private/SkinnedAsset.cpp b/Engine/Source/Runtime/Engine/Private/SkinnedAsset.cpp
new file mode 100644
--- /dev/null
+++ b/Engine/Source/Runtime/Engine/Private/SkinnedAsset.cpp
@@ -0,0 +1,140 @@
+// Copyright Epic Games, Inc. All Rights Reserved.
+
+#include "Engine/SkinnedAsset.h"
+#include "SkinnedAssetCompiler.h"
+
+DEFINE_LOG_CATEGORY_STATIC(LogSkinnedAsset, Log, All);
+
+void USkinnedAsset::PostLoad()
+{
+	Super::PostLoad();
+
+#if WITH_EDITOR
+	FSkinnedAssetAsyncBuildScope AsyncBuildScope(this);
+#endif
+
+	FSkinnedAssetPostLoadContext Context;
+	BeginPostLoadInternal(Context);
+
+#if WITH_EDITOR
+	if (FSkinnedAssetCompilingManager::Get().IsAsyncCompilationAllowed(this))
+	{
+		PrepareForAsyncCompilation();
+
+		FQueuedThreadPool* ThreadPool = FSkinnedAssetCompilingManager::Get().GetThreadPool();
+		EQueuedWorkPriority BasePriority = FSkinnedAssetCompilingManager::Get().GetBasePriority(this);
+
+		AsyncTask = MakeUnique<FSkinnedAssetAsyncBuildTask>(this, MoveTemp(Context));
+		AsyncTask->StartBackgroundTask(ThreadPool, BasePriority, EQueuedWorkFlags::DoNotRunInsideBusyWait);
+		FSkinnedAssetCompilingManager::Get().AddSkinnedAssets({ this });
+	}
+	else
+#endif
+	{
+		ExecutePostLoadInternal(Context);
+		FinishPostLoadInternal(Context);
+	}
+}
+
+#if WITH_EDITOR
+bool USkinnedAsset::IsCompiling() const
+{
+	return AsyncTask != nullptr || AccessedProperties.load(std::memory_order_relaxed) != 0;
+}
+#endif // WITH_EDITOR
+
+void USkinnedAsset::AcquireAsyncProperty(uint64 AsyncProperties, ESkinnedAssetAsyncPropertyLockType LockType)
+{
+#if WITH_EDITOR
+	if ((LockType & ESkinnedAssetAsyncPropertyLockType::ReadOnly) == ESkinnedAssetAsyncPropertyLockType::ReadOnly)
+	{
+		AccessedProperties |= AsyncProperties;
+	}
+
+	if ((LockType & ESkinnedAssetAsyncPropertyLockType::WriteOnly) == ESkinnedAssetAsyncPropertyLockType::WriteOnly)
+	{
+		ModifiedProperties |= AsyncProperties;
+	}
+#endif
+}
+
+void USkinnedAsset::ReleaseAsyncProperty(uint64 AsyncProperties, ESkinnedAssetAsyncPropertyLockType LockType)
+{
+#if WITH_EDITOR
+	if ((LockType & ESkinnedAssetAsyncPropertyLockType::ReadOnly) == ESkinnedAssetAsyncPropertyLockType::ReadOnly)
+	{
+		AccessedProperties &= ~AsyncProperties;
+	}
+
+	if ((LockType & ESkinnedAssetAsyncPropertyLockType::WriteOnly) == ESkinnedAssetAsyncPropertyLockType::WriteOnly)
+	{
+		ModifiedProperties &= ~AsyncProperties;
+	}
+#endif
+}
+
+void USkinnedAsset::WaitUntilAsyncPropertyReleasedInternal(uint64 AsyncProperties, ESkinnedAssetAsyncPropertyLockType LockType) const
+{
+#if WITH_EDITOR
+	// We need to protect internal skinned asset data from race conditions during async build
+	if (IsCompiling())
+	{
+		if (FSkinnedAssetAsyncBuildScope::ShouldWaitOnLockedProperties(this))
+		{
+			bool bIsLocked = true;
+			// We can remove the lock if we're accessing in read-only and there is no write-lock
+			if ((LockType & ESkinnedAssetAsyncPropertyLockType::ReadOnly) == ESkinnedAssetAsyncPropertyLockType::ReadOnly)
+			{
+				// Maintain the lock if the write-lock bit is non-zero
+				bIsLocked &= (ModifiedProperties & AsyncProperties) != 0;
+			}
+
+			if (bIsLocked)
+			{
+				FString PropertyName = GetAsyncPropertyName(AsyncProperties);
+				TRACE_CPUPROFILER_EVENT_SCOPE_TEXT(*FString::Printf(TEXT("SkinnedAssetCompilationStall %s"), *PropertyName));
+
+				if (IsInGameThread())
+				{
+					UE_LOG(
+						LogSkinnedAsset,
+						Verbose,
+						TEXT("Accessing property %s of the SkinnedAsset while it is still being built asynchronously will force it to be compiled before continuing. "
+							"For better performance, consider making the caller async aware so it can wait until the static mesh is ready to access this property."
+							"To better understand where those calls are coming from, you can use Editor.AsyncAssetDumpStallStacks on the console."),
+						*PropertyName
+					);
+
+					FSkinnedAssetCompilingManager::Get().FinishCompilation({ const_cast<USkinnedAsset*>(this) });
+				}
+				else
+				{
+					// Trying to access a property from another thread that cannot force finish the compilation is invalid
+					ensureMsgf(
+						false,
+						TEXT("Accessing property %s of the SkinnedAsset while it is still being built asynchronously is only supported on the game-thread. "
+							"To avoid any race-condition, consider finishing the compilation before pushing tasks to other threads or making higher-level game-thread code async aware so it "
+							"schedules the task only when the static mesh's compilation is finished. If this is a blocker, you can disable async static mesh from the editor experimental settings."),
+						*PropertyName
+					);
+				}
+			}
+		}
+		// If we're accessing this property from the async build thread, make sure the property is still protected from access from other threads.
+		else
+		{
+			bool bIsLocked = true;
+			if ((LockType & ESkinnedAssetAsyncPropertyLockType::ReadOnly) == ESkinnedAssetAsyncPropertyLockType::ReadOnly)
+			{
+				bIsLocked &= (AccessedProperties & AsyncProperties) != 0;
+			}
+
+			if ((LockType & ESkinnedAssetAsyncPropertyLockType::WriteOnly) == ESkinnedAssetAsyncPropertyLockType::WriteOnly)
+			{
+				bIsLocked &= (ModifiedProperties & AsyncProperties) != 0;
+			}
+			ensureMsgf(bIsLocked, TEXT("Property %s has not been locked properly for use by async build"), *GetAsyncPropertyName(AsyncProperties));
+		}
+	}
+#endif
+}
\ No newline at end of file
diff --git a/Engine/Source/Runtime/Engine/Private/SkinnedAssetAsyncCompileUtils.cpp b/Engine/Source/Runtime/Engine/Private/SkinnedAssetAsyncCompileUtils.cpp
new file mode 100644
--- /dev/null
+++ b/Engine/Source/Runtime/Engine/Private/SkinnedAssetAsyncCompileUtils.cpp
@@ -0,0 +1,48 @@
+// Copyright Epic Games, Inc. All Rights Reserved.
+
+#include "Engine/SkinnedAssetAsyncCompileUtils.h"
+#include "Engine/SkinnedAsset.h"
+#include "Engine/SkeletalMesh.h"
+
+#if WITH_EDITOR
+
+/*-----------------------------------------------------------------------------
+	FSkinnedAssetAsyncBuildScope
+-----------------------------------------------------------------------------*/
+thread_local const USkinnedAsset* FSkinnedAssetAsyncBuildScope::SkinnedAssetBeingAsyncCompiled = nullptr;
+
+FSkinnedAssetAsyncBuildScope::FSkinnedAssetAsyncBuildScope(const USkinnedAsset* SkinnedAsset)
+{
+	PreviousScope = SkinnedAssetBeingAsyncCompiled;
+	SkinnedAssetBeingAsyncCompiled = SkinnedAsset;
+}
+
+FSkinnedAssetAsyncBuildScope::~FSkinnedAssetAsyncBuildScope()
+{
+	check(SkinnedAssetBeingAsyncCompiled);
+	SkinnedAssetBeingAsyncCompiled = PreviousScope;
+}
+
+bool FSkinnedAssetAsyncBuildScope::ShouldWaitOnLockedProperties(const USkinnedAsset* SkinnedAsset)
+{
+	return SkinnedAssetBeingAsyncCompiled != SkinnedAsset;
+}
+
+
+/*-----------------------------------------------------------------------------
+	FSkinnedAssetAsyncBuildWorker
+-----------------------------------------------------------------------------*/
+void FSkinnedAssetAsyncBuildWorker::DoWork()
+{
+	if (PostLoadContext.IsSet())
+	{
+		SkinnedAsset->ExecutePostLoadInternal(*PostLoadContext);
+	}
+
+	if (BuildContext.IsSet())
+	{
+		SkinnedAsset->ExecuteBuildInternal(*BuildContext);
+	}
+}
+
+#endif //WITH_EDITOR
\ No newline at end of file
diff --git a/Engine/Source/Runtime/Engine/Private/SkinnedAssetCommon.cpp b/Engine/Source/Runtime/Engine/Private/SkinnedAssetCommon.cpp
new file mode 100644
--- /dev/null
+++ b/Engine/Source/Runtime/Engine/Private/SkinnedAssetCommon.cpp
@@ -0,0 +1,259 @@
+// Copyright Epic Games, Inc. All Rights Reserved.
+
+#include "Engine/SkinnedAssetCommon.h"
+#include "Animation/AnimSequence.h"
+#include "Rendering/SkeletalMeshModel.h"
+#include "Engine/SkeletalMeshLODSettings.h"
+#include "UObject/CoreObjectVersion.h"
+#include "UObject/EditorObjectVersion.h"
+#include "Interfaces/ITargetPlatform.h"
+
+/*-----------------------------------------------------------------------------
+	FSkeletalMeshLODInfo
+-----------------------------------------------------------------------------*/
+
+#if WITH_EDITORONLY_DATA
+static void SerializeReductionSettingsForDDC(FArchive& Ar, FSkeletalMeshOptimizationSettings& ReductionSettings)
+{
+	check(Ar.IsSaving());
+	// Note: this serializer is only used to build the mesh DDC key, no versioning is required
+	FArchive_Serialize_BitfieldBool(Ar, ReductionSettings.bRemapMorphTargets);
+	FArchive_Serialize_BitfieldBool(Ar, ReductionSettings.bRecalcNormals);
+	FArchive_Serialize_BitfieldBool(Ar, ReductionSettings.bEnforceBoneBoundaries);
+	FArchive_Serialize_BitfieldBool(Ar, ReductionSettings.bMergeCoincidentVertBones);
+	FArchive_Serialize_BitfieldBool(Ar, ReductionSettings.bLockEdges);
+	FArchive_Serialize_BitfieldBool(Ar, ReductionSettings.bLockColorBounaries);
+	Ar << ReductionSettings.TerminationCriterion;
+	Ar << ReductionSettings.NumOfTrianglesPercentage;
+	Ar << ReductionSettings.NumOfVertPercentage;
+	Ar << ReductionSettings.MaxNumOfTriangles;
+	Ar << ReductionSettings.MaxNumOfVerts;
+
+	// Keep old DDC keys if these are not set
+	if (ReductionSettings.MaxNumOfTrianglesPercentage != MAX_uint32 ||
+		ReductionSettings.MaxNumOfVertsPercentage != MAX_uint32)
+	{
+		uint32 AvoidCachePoisoningFromOldBug = 0;
+		Ar << AvoidCachePoisoningFromOldBug;
+		Ar << ReductionSettings.MaxNumOfTrianglesPercentage;
+		Ar << ReductionSettings.MaxNumOfVertsPercentage;
+	}
+
+	Ar << ReductionSettings.MaxDeviationPercentage;
+	Ar << ReductionSettings.ReductionMethod;
+	Ar << ReductionSettings.SilhouetteImportance;
+	Ar << ReductionSettings.TextureImportance;
+	Ar << ReductionSettings.ShadingImportance;
+	Ar << ReductionSettings.SkinningImportance;
+	Ar << ReductionSettings.WeldingThreshold;
+	Ar << ReductionSettings.NormalsThreshold;
+	Ar << ReductionSettings.MaxBonesPerVertex;
+	Ar << ReductionSettings.VolumeImportance;
+	Ar << ReductionSettings.BaseLOD;
+}
+
+static void SerializeBuildSettingsForDDC(FArchive& Ar, FSkeletalMeshBuildSettings& BuildSettings)
+{
+	check(Ar.IsSaving());
+	// Note: this serializer is only used to build the mesh DDC key, no versioning is required
+	FArchive_Serialize_BitfieldBool(Ar, BuildSettings.bRecomputeNormals);
+	FArchive_Serialize_BitfieldBool(Ar, BuildSettings.bRecomputeTangents);
+	FArchive_Serialize_BitfieldBool(Ar, BuildSettings.bUseMikkTSpace);
+	FArchive_Serialize_BitfieldBool(Ar, BuildSettings.bComputeWeightedNormals);
+	FArchive_Serialize_BitfieldBool(Ar, BuildSettings.bRemoveDegenerates);
+	FArchive_Serialize_BitfieldBool(Ar, BuildSettings.bUseFullPrecisionUVs);
+	FArchive_Serialize_BitfieldBool(Ar, BuildSettings.bUseBackwardsCompatibleF16TruncUVs);
+	FArchive_Serialize_BitfieldBool(Ar, BuildSettings.bUseHighPrecisionTangentBasis);
+	Ar << BuildSettings.ThresholdPosition;
+	Ar << BuildSettings.ThresholdTangentNormal;
+	Ar << BuildSettings.ThresholdUV;
+	Ar << BuildSettings.MorphThresholdPosition;
+}
+
+FGuid FSkeletalMeshLODInfo::ComputeDeriveDataCacheKey(const FSkeletalMeshLODGroupSettings* SkeletalMeshLODGroupSettings)
+{
+	const bool bIs16BitfloatBufferSupported = GVertexElementTypeSupport.IsSupported(VET_Half2);
+
+	// Serialize the LOD info members, the BuildSettings and the ReductionSettings into a temporary array.
+	TArray<uint8> TempBytes;
+	TempBytes.Reserve(64);
+	//The archive is flagged as persistent so that machines of different endianness produce identical binary results.
+	FMemoryWriter Ar(TempBytes, /*bIsPersistent=*/ true);
+
+	Ar << BonesToRemove;
+	Ar << BonesToPrioritize;
+	Ar << SectionsToPrioritize;
+	Ar << WeightOfPrioritization;
+	Ar << MorphTargetPositionErrorTolerance;
+
+	//TODO: Ask the derivedata key of the UObject reference by FSoftObjectPath. So if someone change the UObject, this LODs will get dirty
+	//and will be rebuild.
+	if (BakePose != nullptr)
+	{
+		FString BakePosePath = BakePose->GetFullName();
+		Ar << BakePosePath;
+	}
+	if (BakePoseOverride != nullptr)
+	{
+		FString BakePoseOverridePath = BakePoseOverride->GetFullName();
+		Ar << BakePoseOverridePath;
+	}
+	FArchive_Serialize_BitfieldBool(Ar, bAllowCPUAccess);
+	FArchive_Serialize_BitfieldBool(Ar, bSupportUniformlyDistributedSampling);
+
+	//Use the LOD settings asset if there is one
+	FSkeletalMeshOptimizationSettings RealReductionSettings = ReductionSettings;
+	if (SkeletalMeshLODGroupSettings != nullptr)
+	{
+		RealReductionSettings = SkeletalMeshLODGroupSettings->GetReductionSettings();
+	}
+
+	if (!BuildSettings.bUseFullPrecisionUVs && !bIs16BitfloatBufferSupported)
+	{
+		BuildSettings.bUseFullPrecisionUVs = true;
+	}
+	SerializeBuildSettingsForDDC(Ar, BuildSettings);
+	SerializeReductionSettingsForDDC(Ar, RealReductionSettings);
+
+	FSHA1 Sha;
+	Sha.Update(TempBytes.GetData(), TempBytes.Num() * TempBytes.GetTypeSize());
+	Sha.Final();
+	// Retrieve the hash and use it to construct a pseudo-GUID.
+	uint32 Hash[5];
+	Sha.GetHash((uint8*)Hash);
+	FGuid Guid = FGuid(Hash[0] ^ Hash[4], Hash[1], Hash[2], Hash[3]);
+	return Guid;
+}
+#endif //WITH_EDITORONLY_DATA
+
+
+/*-----------------------------------------------------------------------------
+	FSectionReference
+-----------------------------------------------------------------------------*/
+#if WITH_EDITOR
+bool FSectionReference::IsValidToEvaluate(const FSkeletalMeshLODModel& LodModel) const
+{
+	return GetMeshLodSection(LodModel) != nullptr;
+}
+
+const FSkelMeshSection* FSectionReference::GetMeshLodSection(const FSkeletalMeshLODModel& LodModel) const
+{
+	for (int32 LodModelSectionIndex = 0; LodModelSectionIndex < LodModel.Sections.Num(); ++LodModelSectionIndex)
+	{
+		const FSkelMeshSection& Section = LodModel.Sections[LodModelSectionIndex];
+		if (Section.ChunkedParentSectionIndex == INDEX_NONE && SectionIndex == Section.OriginalDataSectionIndex)
+		{
+			return &Section;
+		}
+	}
+	return nullptr;
+}
+
+int32 FSectionReference::GetMeshLodSectionIndex(const FSkeletalMeshLODModel& LodModel) const
+{
+	for (int32 LodModelSectionIndex = 0; LodModelSectionIndex < LodModel.Sections.Num(); ++LodModelSectionIndex)
+	{
+		const FSkelMeshSection& Section = LodModel.Sections[LodModelSectionIndex];
+		if (Section.ChunkedParentSectionIndex == INDEX_NONE && SectionIndex == Section.OriginalDataSectionIndex)
+		{
+			return LodModelSectionIndex;
+		}
+	}
+	return INDEX_NONE;
+}
+#endif // #if WITH_EDITOR
+
+/*-----------------------------------------------------------------------------
+	FSkeletalMaterial
+-----------------------------------------------------------------------------*/
+
+bool operator== (const FSkeletalMaterial& LHS, const FSkeletalMaterial& RHS)
+{
+	return (LHS.MaterialInterface == RHS.MaterialInterface);
+}
+
+bool operator== (const FSkeletalMaterial& LHS, const UMaterialInterface& RHS)
+{
+	return (LHS.MaterialInterface == &RHS);
+}
+
+bool operator== (const UMaterialInterface& LHS, const FSkeletalMaterial& RHS)
+{
+	return (RHS.MaterialInterface == &LHS);
+}
+
+FArchive& operator<<(FArchive& Ar, FMeshUVChannelInfo& ChannelData)
+{
+	Ar << ChannelData.bInitialized;
+	Ar << ChannelData.bOverrideDensities;
+
+	for (int32 CoordIndex = 0; CoordIndex < TEXSTREAM_MAX_NUM_UVCHANNELS; ++CoordIndex)
+	{
+		Ar << ChannelData.LocalUVDensities[CoordIndex];
+	}
+
+	return Ar;
+}
+
+#if WITH_EDITORONLY_DATA
+void FSkeletalMaterial::DeclareCustomVersions(FArchive& Ar)
+{
+	Ar.UsingCustomVersion(FEditorObjectVersion::GUID);
+	Ar.UsingCustomVersion(FCoreObjectVersion::GUID);
+}
+#endif
+
+FArchive& operator<<(FArchive& Ar, FSkeletalMaterial& Elem)
+{
+	Ar.UsingCustomVersion(FEditorObjectVersion::GUID);
+	Ar.UsingCustomVersion(FCoreObjectVersion::GUID);
+
+	Ar << Elem.MaterialInterface;
+
+	//Use the automatic serialization instead of this custom operator
+	if (Ar.CustomVer(FEditorObjectVersion::GUID) >= FEditorObjectVersion::RefactorMeshEditorMaterials)
+	{
+		Ar << Elem.MaterialSlotName;
+
+		bool bSerializeImportedMaterialSlotName = !Ar.IsCooking() || Ar.CookingTarget()->HasEditorOnlyData();
+		if (Ar.CustomVer(FCoreObjectVersion::GUID) >= FCoreObjectVersion::SkeletalMaterialEditorDataStripping)
+		{
+			Ar << bSerializeImportedMaterialSlotName;
+		}
+		else if (!FPlatformProperties::HasEditorOnlyData())
+		{
+			bSerializeImportedMaterialSlotName = false;
+		}
+		if (bSerializeImportedMaterialSlotName)
+		{
+#if WITH_EDITORONLY_DATA
+			Ar << Elem.ImportedMaterialSlotName;
+#else
+			FName UnusedImportedMaterialSlotName;
+			Ar << UnusedImportedMaterialSlotName;
+#endif
+		}
+	}
+#if WITH_EDITORONLY_DATA
+	else
+	{
+		if (Ar.UEVer() >= VER_UE4_MOVE_SKELETALMESH_SHADOWCASTING)
+		{
+			Ar << Elem.bEnableShadowCasting_DEPRECATED;
+		}
+
+		Ar.UsingCustomVersion(FRecomputeTangentCustomVersion::GUID);
+		if (Ar.CustomVer(FRecomputeTangentCustomVersion::GUID) >= FRecomputeTangentCustomVersion::RuntimeRecomputeTangent)
+		{
+			Ar << Elem.bRecomputeTangent_DEPRECATED;
+		}
+	}
+#endif
+
+	if (!Ar.IsLoading() || Ar.CustomVer(FRenderingObjectVersion::GUID) >= FRenderingObjectVersion::TextureStreamingMeshUVChannelData)
+	{
+		Ar << Elem.UVChannelData;
+	}
+
+	return Ar;
+}
diff --git a/Engine/Source/Runtime/Engine/Private/SkinnedAssetCompiler.cpp b/Engine/Source/Runtime/Engine/Private/SkinnedAssetCompiler.cpp
new file mode 100644
--- /dev/null
+++ b/Engine/Source/Runtime/Engine/Private/SkinnedAssetCompiler.cpp
@@ -0,0 +1,449 @@
+// Copyright Epic Games, Inc. All Rights Reserved.
+
+#include "SkinnedAssetCompiler.h"
+#include "AssetCompilingManager.h"
+#include "Engine/SkinnedAsset.h"
+#include "Engine/SkeletalMesh.h" // For AddSkeletalMeshes
+
+#if WITH_EDITOR
+
+#include "ObjectCacheContext.h"
+#include "Settings/EditorExperimentalSettings.h"
+#include "Misc/QueuedThreadPoolWrapper.h"
+#include "EngineModule.h"
+#include "Misc/ScopedSlowTask.h"
+#include "UObject/StrongObjectPtr.h"
+#include "ShaderCompiler.h"
+#include "TextureCompiler.h"
+#include "Misc/IQueuedWork.h"
+#include "LevelEditor.h"
+#include "SLevelViewport.h"
+#include "Components/PrimitiveComponent.h"
+#include "Components/SkinnedMeshComponent.h"
+#include "ContentStreaming.h"
+#include "DrawDebugHelpers.h"
+#include "GameFramework/PlayerController.h"
+#include "GameFramework/Character.h"
+#include "Components/CapsuleComponent.h"
+#include "AI/NavigationSystemBase.h"
+#include "EngineUtils.h"
+#include "ProfilingDebugging/CountersTrace.h"
+
+#define LOCTEXT_NAMESPACE "SkinnedAssetCompiler"
+
+static AsyncCompilationHelpers::FAsyncCompilationStandardCVars CVarAsyncSkinnedAssetStandard(
+	TEXT("SkinnedAsset"),
+	TEXT("skinned assets"),
+	FConsoleCommandDelegate::CreateLambda(
+		[]()
+		{
+			FSkinnedAssetCompilingManager::Get().FinishAllCompilation();
+		}
+	));
+
+namespace SkinnedAssetCompilingManagerImpl
+{
+	static void EnsureInitializedCVars()
+	{
+		static bool bIsInitialized = false;
+
+		if (!bIsInitialized)
+		{
+			bIsInitialized = true;
+			
+			AsyncCompilationHelpers::EnsureInitializedCVars(
+				TEXT("skinnedasset"),
+				CVarAsyncSkinnedAssetStandard.AsyncCompilation,
+				CVarAsyncSkinnedAssetStandard.AsyncCompilationMaxConcurrency,
+				GET_MEMBER_NAME_CHECKED(UEditorExperimentalSettings, bEnableAsyncSkeletalMeshCompilation)); // TODO: rename to SkinnedAsset
+		}
+	}
+}
+
+FSkinnedAssetCompilingManager::FSkinnedAssetCompilingManager()
+	: Notification(GetAssetNameFormat())
+{
+	SkinnedAssetCompilingManagerImpl::EnsureInitializedCVars();
+}
+
+FName FSkinnedAssetCompilingManager::GetAssetTypeName() const
+{
+	return TEXT("UE-SkinnedAsset");
+}
+
+FTextFormat FSkinnedAssetCompilingManager::GetAssetNameFormat() const
+{
+	return LOCTEXT("SkinnedAssetNameFormat", "{0}|plural(one=Skinned Asset,other=Skinned Assets)");
+}
+
+TArrayView<FName> FSkinnedAssetCompilingManager::GetDependentTypeNames() const
+{
+	// Texture and shaders can affect materials which can affect Skinned Assets once they are visible.
+	// Adding these dependencies can reduces the actual number of render state update we need to do in a frame
+	static FName DependentTypeNames[] = 
+	{
+		FTextureCompilingManager::GetStaticAssetTypeName(), 
+		FShaderCompilingManager::GetStaticAssetTypeName() 
+	};
+	return TArrayView<FName>(DependentTypeNames);
+}
+
+int32 FSkinnedAssetCompilingManager::GetNumRemainingAssets() const
+{
+	return GetNumRemainingJobs();
+}
+
+EQueuedWorkPriority FSkinnedAssetCompilingManager::GetBasePriority(USkinnedAsset* InSkinnedAsset) const
+{
+	return EQueuedWorkPriority::Low;
+}
+
+FQueuedThreadPool* FSkinnedAssetCompilingManager::GetThreadPool() const
+{
+	static FQueuedThreadPoolDynamicWrapper* GSkinnedAssetThreadPool = nullptr;
+	if (GSkinnedAssetThreadPool == nullptr && FAssetCompilingManager::Get().GetThreadPool() != nullptr)
+	{
+		// For now, skinned assets have almost no high-level awareness of their async behavior.
+		// Let them build first to avoid game-thread stalls as much as possible.
+		TFunction<EQueuedWorkPriority(EQueuedWorkPriority)> PriorityMapper = [](EQueuedWorkPriority) { return EQueuedWorkPriority::Highest; };
+
+		// Skinned assets will be scheduled on the asset thread pool, where concurrency limits might by dynamically adjusted depending on memory constraints.
+		GSkinnedAssetThreadPool = new FQueuedThreadPoolDynamicWrapper(FAssetCompilingManager::Get().GetThreadPool(), -1, PriorityMapper);
+
+		AsyncCompilationHelpers::BindThreadPoolToCVar(
+			GSkinnedAssetThreadPool,
+			CVarAsyncSkinnedAssetStandard.AsyncCompilation,
+			CVarAsyncSkinnedAssetStandard.AsyncCompilationResume,
+			CVarAsyncSkinnedAssetStandard.AsyncCompilationMaxConcurrency
+		);
+	}
+
+	return GSkinnedAssetThreadPool;
+}
+
+void FSkinnedAssetCompilingManager::Shutdown()
+{
+	bHasShutdown = true;
+	if (GetNumRemainingJobs())
+	{
+		check(IsInGameThread());
+		TRACE_CPUPROFILER_EVENT_SCOPE(FSkinnedAssetCompilingManager::Shutdown)
+
+		if (GetNumRemainingJobs())
+		{
+			TArray<USkinnedAsset*> PendingSkinnedAssets;
+			PendingSkinnedAssets.Reserve(GetNumRemainingJobs());
+
+			for (TWeakObjectPtr<USkinnedAsset>& WeakSkinnedAsset : RegisteredSkinnedAsset)
+			{
+				if (WeakSkinnedAsset.IsValid())
+				{
+					USkinnedAsset* SkinnedAsset = WeakSkinnedAsset.Get();
+					if (!SkinnedAsset->IsAsyncTaskComplete())
+					{
+						if (SkinnedAsset->AsyncTask->Cancel())
+						{
+							SkinnedAsset->AsyncTask.Reset();
+						}
+					}
+
+					if (SkinnedAsset->AsyncTask)
+					{
+						PendingSkinnedAssets.Add(SkinnedAsset);
+					}
+				}
+			}
+
+			FinishCompilation(PendingSkinnedAssets);
+		}
+	}
+}
+
+bool FSkinnedAssetCompilingManager::IsAsyncCompilationEnabled() const
+{
+	if (bHasShutdown)
+	{
+		return false;
+	}
+
+	return CVarAsyncSkinnedAssetStandard.AsyncCompilation.GetValueOnAnyThread() != 0;
+}
+
+TRACE_DECLARE_INT_COUNTER(QueuedSkinnedAssetCompilation, TEXT("AsyncCompilation/QueuedSkinnedAsset"));
+void FSkinnedAssetCompilingManager::UpdateCompilationNotification()
+{
+	TRACE_COUNTER_SET(QueuedSkinnedAssetCompilation, GetNumRemainingJobs());
+	Notification.Update(GetNumRemainingJobs());
+}
+
+void FSkinnedAssetCompilingManager::PostCompilation(TArrayView<USkinnedAsset* const> InSkinnedAssets)
+{
+	if (InSkinnedAssets.Num())
+	{
+		TRACE_CPUPROFILER_EVENT_SCOPE(OnAssetPostCompileEvent);
+
+		TArray<FAssetCompileData> AssetsData;
+		AssetsData.Reserve(InSkinnedAssets.Num());
+
+		for (USkinnedAsset* SkinnedAsset : InSkinnedAssets)
+		{
+			AssetsData.Emplace(SkinnedAsset);
+		}
+
+		FAssetCompilingManager::Get().OnAssetPostCompileEvent().Broadcast(AssetsData);
+	}
+}
+
+void FSkinnedAssetCompilingManager::PostCompilation(USkinnedAsset* SkinnedAsset)
+{
+	using namespace SkinnedAssetCompilingManagerImpl;
+	
+	// If AsyncTask is null here, the task got canceled so we don't need to do anything
+	if (SkinnedAsset->AsyncTask)
+	{
+		check(IsInGameThread());
+		TRACE_CPUPROFILER_EVENT_SCOPE(FSkinnedAssetCompilingManager::PostCompilation);
+
+		UE_LOG(LogSkeletalMesh, Verbose, TEXT("Refreshing skinned asset %s because it is ready"), *SkinnedAsset->GetName()); // TODO: change to LogSkinnedAsset
+
+		FObjectCacheContextScope ObjectCacheScope;
+
+		// The scope is important here to destroy the FSkinnedAssetAsyncBuildScope before broadcasting events
+		{
+			// Acquire the async task locally to protect against re-entrance
+			TUniquePtr<FSkinnedAssetAsyncBuildTask> LocalAsyncTask = MoveTemp(SkinnedAsset->AsyncTask);
+			LocalAsyncTask->EnsureCompletion();
+
+			FSkinnedAssetAsyncBuildScope AsyncBuildScope(SkinnedAsset);
+
+			if (LocalAsyncTask->GetTask().PostLoadContext.IsSet())
+			{
+				SkinnedAsset->FinishPostLoadInternal(*LocalAsyncTask->GetTask().PostLoadContext);
+
+				LocalAsyncTask->GetTask().PostLoadContext.Reset();
+			}
+
+			if (LocalAsyncTask->GetTask().BuildContext.IsSet())
+			{
+				SkinnedAsset->FinishBuildInternal(*LocalAsyncTask->GetTask().BuildContext);
+
+				LocalAsyncTask->GetTask().BuildContext.Reset();
+			}
+		}
+
+		// Calling this delegate during app exit might be quite dangerous and lead to crash
+		// if the content browser wants to refresh a thumbnail it might try to load a package
+		// which will then fail due to various reasons related to the editor shutting down.
+		// Triggering this callback while garbage collecting can also result in listeners trying to look up objects
+		if (!GExitPurge && !IsGarbageCollecting())
+		{
+			// Generate an empty property changed event, to force the asset registry tag
+			// to be refreshed now that RenderData is available.
+			FPropertyChangedEvent EmptyPropertyChangedEvent(nullptr);
+			FCoreUObjectDelegates::OnObjectPropertyChanged.Broadcast(SkinnedAsset, EmptyPropertyChangedEvent);
+		}
+	}
+}
+
+bool FSkinnedAssetCompilingManager::IsAsyncCompilationAllowed(USkinnedAsset* SkinnedAsset) const
+{
+	return IsAsyncCompilationEnabled();
+}
+
+FSkinnedAssetCompilingManager& FSkinnedAssetCompilingManager::Get()
+{
+	static FSkinnedAssetCompilingManager Singleton;
+	return Singleton;
+}
+
+int32 FSkinnedAssetCompilingManager::GetNumRemainingJobs() const
+{
+	return RegisteredSkinnedAsset.Num();
+}
+
+void FSkinnedAssetCompilingManager::AddSkeletalMeshes(TArrayView<USkeletalMesh* const> InSkeletalMeshes)
+{
+	TArray<USkinnedAsset*> SkinnedAssets;
+	SkinnedAssets.Reserve(InSkeletalMeshes.Num());
+	for (USkeletalMesh* SkMesh : InSkeletalMeshes)
+	{
+		SkinnedAssets.Add(SkMesh);
+	}
+
+	AddSkinnedAssets(SkinnedAssets);
+}
+
+void FSkinnedAssetCompilingManager::AddSkinnedAssets(TArrayView<USkinnedAsset* const> InSkinnedAssets)
+{
+	TRACE_CPUPROFILER_EVENT_SCOPE(FSkinnedAssetCompilingManager::AddSkinnedAssets)
+	check(IsInGameThread());
+
+	// Wait until we gather enough mesh to process
+	// to amortize the cost of scanning components
+	//ProcessSkinnedAssets(32 /* MinBatchSize */);
+
+	for (USkinnedAsset* SkinnedAsset : InSkinnedAssets)
+	{
+		check(SkinnedAsset->AsyncTask != nullptr);
+		RegisteredSkinnedAsset.Emplace(SkinnedAsset);
+	}
+
+	UpdateCompilationNotification();
+}
+
+void FSkinnedAssetCompilingManager::FinishCompilation(TArrayView<USkinnedAsset* const> InSkinnedAssets)
+{
+	TRACE_CPUPROFILER_EVENT_SCOPE(FSkinnedAssetCompilingManager::FinishCompilation);
+
+	check(IsInGameThread());
+
+	TArray<USkinnedAsset*> PendingSkinnedAssets;
+	PendingSkinnedAssets.Reserve(InSkinnedAssets.Num());
+
+	for (USkinnedAsset* SkinnedAsset : InSkinnedAssets)
+	{
+		if (RegisteredSkinnedAsset.Contains(SkinnedAsset))
+		{
+			PendingSkinnedAssets.Emplace(SkinnedAsset);
+		}
+	}
+
+	if (PendingSkinnedAssets.Num())
+	{
+		class FCompilableSkinnedAsset : public AsyncCompilationHelpers::TCompilableAsyncTask<FSkinnedAssetAsyncBuildTask>
+		{
+		public:
+			FCompilableSkinnedAsset(USkinnedAsset* InSkinnedAsset)
+				: SkinnedAsset(InSkinnedAsset)
+			{
+			}
+
+			FSkinnedAssetAsyncBuildTask* GetAsyncTask() override
+			{
+				return SkinnedAsset->AsyncTask.Get();
+			}
+
+			TStrongObjectPtr<USkinnedAsset> SkinnedAsset;
+			FName GetName() override { return SkinnedAsset->GetFName(); }
+		};
+
+		TArray<FCompilableSkinnedAsset> CompilableSkinnedAsset(PendingSkinnedAssets);
+
+		FObjectCacheContextScope ObjectCacheScope;
+		AsyncCompilationHelpers::FinishCompilation(
+			[&CompilableSkinnedAsset](int32 Index) -> AsyncCompilationHelpers::ICompilable& { return CompilableSkinnedAsset[Index]; },
+			CompilableSkinnedAsset.Num(),
+			LOCTEXT("SkinnedAssets", "Skinned Assets"),
+			LogSkeletalMesh, // TODO: change to LogSkinnedAsset
+			[this](AsyncCompilationHelpers::ICompilable* Object)
+			{
+				USkinnedAsset* SkinnedAsset = static_cast<FCompilableSkinnedAsset*>(Object)->SkinnedAsset.Get();
+				PostCompilation(SkinnedAsset);
+				RegisteredSkinnedAsset.Remove(SkinnedAsset);
+			}
+		);
+
+		PostCompilation(PendingSkinnedAssets);
+
+		UpdateCompilationNotification();
+	}
+}
+
+void FSkinnedAssetCompilingManager::FinishCompilationsForGame()
+{
+	
+}
+
+void FSkinnedAssetCompilingManager::FinishAllCompilation()
+{
+	check(IsInGameThread());
+	TRACE_CPUPROFILER_EVENT_SCOPE(FSkinnedAssetCompilingManager::FinishAllCompilation)
+
+	if (GetNumRemainingJobs())
+	{
+		TArray<USkinnedAsset*> PendingSkinnedAssets;
+		PendingSkinnedAssets.Reserve(GetNumRemainingJobs());
+
+		for (TWeakObjectPtr<USkinnedAsset>& SkinnedAsset : RegisteredSkinnedAsset)
+		{
+			if (SkinnedAsset.IsValid())
+			{
+				PendingSkinnedAssets.Add(SkinnedAsset.Get());
+			}
+		}
+
+		FinishCompilation(PendingSkinnedAssets);
+	}
+}
+
+void FSkinnedAssetCompilingManager::Reschedule()
+{
+	
+}
+
+void FSkinnedAssetCompilingManager::ProcessSkinnedAssets(bool bLimitExecutionTime, int32 MinBatchSize)
+{
+	using namespace SkinnedAssetCompilingManagerImpl;
+	TRACE_CPUPROFILER_EVENT_SCOPE(FSkinnedAssetCompilingManager::ProcessSkinnedAssets);
+	const int32 NumRemainingMeshes = GetNumRemainingJobs();
+	// Spread out the load over multiple frames but if too many meshes, convergence is more important than frame time
+	const int32 MaxMeshUpdatesPerFrame = bLimitExecutionTime ? FMath::Max(64, NumRemainingMeshes / 10) : INT32_MAX;
+
+	FObjectCacheContextScope ObjectCacheScope;
+	if (NumRemainingMeshes && NumRemainingMeshes >= MinBatchSize)
+	{
+		TSet<USkinnedAsset*> SkinnedAssetsToProcess;
+		for (TWeakObjectPtr<USkinnedAsset>& SkinnedAsset : RegisteredSkinnedAsset)
+		{
+			if (SkinnedAsset.IsValid())
+			{
+				SkinnedAssetsToProcess.Add(SkinnedAsset.Get());
+			}
+		}
+
+		{
+			TRACE_CPUPROFILER_EVENT_SCOPE(ProcessFinishedSkinnedAssets);
+
+			const double TickStartTime = FPlatformTime::Seconds();
+
+			TSet<TWeakObjectPtr<USkinnedAsset>> SkinnedAssetsToPostpone;
+			TArray<USkinnedAsset*> ProcessedSkinnedAssets;
+			if (SkinnedAssetsToProcess.Num())
+			{
+				for (USkinnedAsset* SkinnedAsset : SkinnedAssetsToProcess)
+				{
+					const bool bHasMeshUpdateLeft = ProcessedSkinnedAssets.Num() <= MaxMeshUpdatesPerFrame;
+					if (bHasMeshUpdateLeft && SkinnedAsset->IsAsyncTaskComplete())
+					{
+						PostCompilation(SkinnedAsset);
+						ProcessedSkinnedAssets.Add(SkinnedAsset);
+					}
+					else
+					{
+						SkinnedAssetsToPostpone.Emplace(SkinnedAsset);
+					}
+				}
+			}
+
+			RegisteredSkinnedAsset = MoveTemp(SkinnedAssetsToPostpone);
+
+			PostCompilation(ProcessedSkinnedAssets);
+		}
+	}
+}
+
+void FSkinnedAssetCompilingManager::ProcessAsyncTasks(bool bLimitExecutionTime)
+{
+	FObjectCacheContextScope ObjectCacheScope;
+	FinishCompilationsForGame();
+
+	Reschedule();
+
+	ProcessSkinnedAssets(bLimitExecutionTime);
+
+	UpdateCompilationNotification();
+}
+
+#endif // #if WITH_EDITOR
+
+#undef LOCTEXT_NAMESPACE
\ No newline at end of file
diff --git a/Engine/Source/Runtime/Engine/Public/Rendering/SkeletalMeshRenderData.h b/Engine/Source/Runtime/Engine/Public/Rendering/SkeletalMeshRenderData.h
--- a/Engine/Source/Runtime/Engine/Public/Rendering/SkeletalMeshRenderData.h
+++ b/Engine/Source/Runtime/Engine/Public/Rendering/SkeletalMeshRenderData.h
@@ -55,7 +55,7 @@ public:
 	ENGINE_API ~FSkeletalMeshRenderData();
 
 #if WITH_EDITOR
-	ENGINE_API void Cache(const ITargetPlatform* TargetPlatform, USkinnedAsset* Owner, class FSkeletalMeshCompilationContext* ContextPtr);
+	ENGINE_API void Cache(const ITargetPlatform* TargetPlatform, USkinnedAsset* Owner, class FSkinnedAssetCompilationContext* ContextPtr);
 	FString GetDerivedDataKey(const ITargetPlatform* TargetPlatform, USkinnedAsset* Owner);
 
 	ENGINE_API void SyncUVChannelData(const TArray<FSkeletalMaterial>& ObjectData);
diff --git a/Engine/Source/Runtime/Engine/Public/SkeletalMeshCompiler.h b/Engine/Source/Runtime/Engine/Public/SkeletalMeshCompiler.h
--- a/Engine/Source/Runtime/Engine/Public/SkeletalMeshCompiler.h
+++ b/Engine/Source/Runtime/Engine/Public/SkeletalMeshCompiler.h
@@ -1,94 +1,11 @@
 // Copyright Epic Games, Inc. All Rights Reserved.
 #pragma once
 
-#include "CoreMinimal.h"
-#include "UObject/WeakObjectPtr.h"
-#include "Containers/Set.h"
-#include "Tickable.h"
-#include "Stats/Stats.h"
-#include "AssetCompilingManager.h"
-#include "AsyncCompilationHelpers.h"
-
-#if WITH_EDITOR
+#include "SkinnedAssetCompiler.h"
 
 class USkeletalMesh;
-class UPrimitiveComponent;
 class USkeletalMeshComponent;
-class FQueuedThreadPool;
-struct FAssetCompileContext;
-enum class EQueuedWorkPriority : uint8;
-
-class FSkeletalMeshCompilingManager : IAssetCompilingManager
-{
-public:
-	ENGINE_API static FSkeletalMeshCompilingManager& Get();
-
-	/**
-	 * Returns true if the feature is currently activated.
-	 */
-	ENGINE_API bool IsAsyncCompilationEnabled() const;
-
-	/** 
-	 * Returns the number of outstanding compilations.
-	 */
-	ENGINE_API int32 GetNumRemainingJobs() const;
-
-	/** 
-	 * Adds skeletal meshes compiled asynchronously so they are monitored. 
-	 */
-	ENGINE_API void AddSkeletalMeshes(TArrayView<USkeletalMesh* const> InSkeletalMeshes);
-
-	/** 
-	 * Blocks until completion of the requested skeletal meshes.
-	 */
-	ENGINE_API void FinishCompilation(TArrayView<USkeletalMesh* const> InSkeletalMeshes);
-
-	/** 
-	 * Blocks until completion of all async skeletal mesh compilation.
-	 */
-	ENGINE_API void FinishAllCompilation() override;
-
-	/**
-	 * Returns if asynchronous compilation is allowed for this skeletal mesh.
-	 */
-	ENGINE_API bool IsAsyncCompilationAllowed(USkeletalMesh* InSkeletalMesh) const;
-
-	/**
-	 * Returns the priority at which the given skeletal mesh should be scheduled.
-	 */
-	ENGINE_API EQueuedWorkPriority GetBasePriority(USkeletalMesh* InSkeletalMesh) const;
-
-	/**
-	 * Returns the threadpool where skeletal mesh compilation should be scheduled.
-	 */
-	ENGINE_API FQueuedThreadPool* GetThreadPool() const;
-
-	/**
-	 * Cancel any pending work and blocks until it is safe to shut down.
-	 */
-	ENGINE_API void Shutdown() override;
-
-private:
-	FSkeletalMeshCompilingManager();
-
-	FName GetAssetTypeName() const override;
-	FTextFormat GetAssetNameFormat() const override;
-	TArrayView<FName> GetDependentTypeNames() const override;
-	int32 GetNumRemainingAssets() const override;
-	void ProcessAsyncTasks(bool bLimitExecutionTime = false) override;
-
-	friend class FAssetCompilingManager;
-	
-	bool bHasShutdown = false;
-	TSet<TWeakObjectPtr<USkeletalMesh>> RegisteredSkeletalMesh;
-	FAsyncCompilationNotification Notification;
-	void FinishCompilationsForGame();
-	void Reschedule();
-	void ProcessSkeletalMeshes(bool bLimitExecutionTime, int32 MinBatchSize = 1);
-	void UpdateCompilationNotification();
-
-	void PostCompilation(USkeletalMesh* SkeletalMesh);
-	void PostCompilation(TArrayView<USkeletalMesh* const> InSkeletalMeshes);
-};
 
+#if WITH_EDITOR
+using FSkeletalMeshCompilingManager UE_DEPRECATED(5.1, "Use FSkinnedAssetCompilingManager instead.") = FSkinnedAssetCompilingManager;
 #endif // #if WITH_EDITOR
\ No newline at end of file
diff --git a/Engine/Source/Runtime/Engine/Public/SkinnedAssetCompiler.h b/Engine/Source/Runtime/Engine/Public/SkinnedAssetCompiler.h
new file mode 100644
--- /dev/null
+++ b/Engine/Source/Runtime/Engine/Public/SkinnedAssetCompiler.h
@@ -0,0 +1,99 @@
+// Copyright Epic Games, Inc. All Rights Reserved.
+#pragma once
+
+#include "CoreMinimal.h"
+#include "UObject/WeakObjectPtr.h"
+#include "Containers/Set.h"
+#include "Tickable.h"
+#include "Stats/Stats.h"
+#include "AssetCompilingManager.h"
+#include "AsyncCompilationHelpers.h"
+
+#if WITH_EDITOR
+
+class USkinnedAsset;
+class UPrimitiveComponent;
+class FQueuedThreadPool;
+struct FAssetCompileContext;
+enum class EQueuedWorkPriority : uint8;
+
+class FSkinnedAssetCompilingManager : IAssetCompilingManager
+{
+public:
+	ENGINE_API static FSkinnedAssetCompilingManager& Get();
+
+	/**
+	 * Returns true if the feature is currently activated.
+	 */
+	ENGINE_API bool IsAsyncCompilationEnabled() const;
+
+	/** 
+	 * Returns the number of outstanding compilations.
+	 */
+	ENGINE_API int32 GetNumRemainingJobs() const;
+
+	/**
+	 * Adds skeletal meshes compiled asynchronously so they are monitored.
+	 */
+	UE_DEPRECATED(5.1, "Use AddSkinnedAssets instead.")
+	ENGINE_API void AddSkeletalMeshes(TArrayView<class USkeletalMesh* const> InSkeletalMeshes);
+
+	/** 
+	 * Adds skinned assets compiled asynchronously so they are monitored. 
+	 */
+	ENGINE_API void AddSkinnedAssets(TArrayView<USkinnedAsset* const> InSkinnedAssets);
+
+	/** 
+	 * Blocks until completion of the requested skinned assets.
+	 */
+	ENGINE_API void FinishCompilation(TArrayView<USkinnedAsset* const> InSkinnedAssets);
+
+	/** 
+	 * Blocks until completion of all async skinned asset compilation.
+	 */
+	ENGINE_API void FinishAllCompilation() override;
+
+	/**
+	 * Returns if asynchronous compilation is allowed for this skinned asset.
+	 */
+	ENGINE_API bool IsAsyncCompilationAllowed(USkinnedAsset* InSkinnedAsset) const;
+
+	/**
+	 * Returns the priority at which the given skinned asset should be scheduled.
+	 */
+	ENGINE_API EQueuedWorkPriority GetBasePriority(USkinnedAsset* InSkinnedAsset) const;
+
+	/**
+	 * Returns the threadpool where skinned asset compilation should be scheduled.
+	 */
+	ENGINE_API FQueuedThreadPool* GetThreadPool() const;
+
+	/**
+	 * Cancel any pending work and blocks until it is safe to shut down.
+	 */
+	ENGINE_API void Shutdown() override;
+
+private:
+	FSkinnedAssetCompilingManager();
+
+	FName GetAssetTypeName() const override;
+	FTextFormat GetAssetNameFormat() const override;
+	TArrayView<FName> GetDependentTypeNames() const override;
+	int32 GetNumRemainingAssets() const override;
+	void ProcessAsyncTasks(bool bLimitExecutionTime = false) override;
+
+	friend class FAssetCompilingManager;
+	
+	bool bHasShutdown = false;
+	TSet<TWeakObjectPtr<USkinnedAsset>> RegisteredSkinnedAsset;
+	FAsyncCompilationNotification Notification;
+	void FinishCompilationsForGame();
+	void Reschedule();
+	void ProcessSkinnedAssets(bool bLimitExecutionTime, int32 MinBatchSize = 1);
+	void UpdateCompilationNotification();
+
+	void PostCompilation(USkinnedAsset* SkinnedAsset);
+	void PostCompilation(TArrayView<USkinnedAsset* const> InSkinnedAssets);
+};
+
+#endif // #if WITH_EDITOR
\ No newline at end of file
