World Partition: added wp.Runtime.RuntimeSpatialHashPlacePartitionActorsUsingLocation cvar to allow placing partition actors in the grid using their bounds pivot instead of using their bounds, to avoid grid promotion.

#rnx

diff --git a/Engine/Source/Runtime/Engine/Private/WorldPartition/RuntimeSpatialHash/RuntimeSpatialHashGridHelper.cpp b/Engine/Source/Runtime/Engine/Private/WorldPartition/RuntimeSpatialHash/RuntimeSpatialHashGridHelper.cpp
--- a/Engine/Source/Runtime/Engine/Private/WorldPartition/RuntimeSpatialHash/RuntimeSpatialHashGridHelper.cpp
+++ b/Engine/Source/Runtime/Engine/Private/WorldPartition/RuntimeSpatialHash/RuntimeSpatialHashGridHelper.cpp
@@ -1,6 +1,7 @@
 // Copyright Epic Games, Inc. All Rights Reserved.
 
 #include "WorldPartition/RuntimeSpatialHash/RuntimeSpatialHashGridHelper.h"
+#include "ActorPartition/PartitionActor.h"
 #include "Algo/Transform.h"
 #include "ProfilingDebugging/ScopedTimers.h"
 #include "WorldPartition/DataLayer/DataLayerInstance.h"
@@ -23,6 +24,12 @@ static FAutoConsoleVariableRef CVarRuntimeSpatialHashPlaceSmallActorsUsingLocati
 	GRuntimeSpatialHashPlaceSmallActorsUsingLocation,
 	TEXT("Set RuntimeSpatialHashPlaceSmallActorsUsingLocation to true to place actors smaller than a cell size into their corresponding cell using their location instead of their bounding box."));
 
+bool GRuntimeSpatialHashPlacePartitionActorsUsingLocation = true;
+static FAutoConsoleVariableRef CVarRuntimeSpatialHashPlacePartitionActorsUsingLocation(
+	TEXT("wp.Runtime.RuntimeSpatialHashPlacePartitionActorsUsingLocation"),
+	GRuntimeSpatialHashPlacePartitionActorsUsingLocation,
+	TEXT("Set RuntimeSpatialHashPlacePartitionActorsUsingLocation to true to place partitioned actors into their corresponding cell using their location instead of their bounding box."));
+
 FSquare2DGridHelper::FSquare2DGridHelper(const FBox& InWorldBounds, const FVector& InOrigin, int64 InCellSize)
 	: WorldBounds(InWorldBounds)
 	, Origin(InOrigin)
@@ -182,6 +189,34 @@ FSquare2DGridHelper GetPartitionedActors(const UWorldPartition* WorldPartition,
 	}
 
 	const float CellArea = PartitionedActors.GetLowestLevel().CellSize * PartitionedActors.GetLowestLevel().CellSize;
+
+	auto ShouldActorUseLocationPlacement = [&CellArea](const FActorClusterInstance* ClusterInstance)
+	{
+		if (GRuntimeSpatialHashPlaceSmallActorsUsingLocation)
+		{
+			const FBox2D ClusterBounds(FVector2D(ClusterInstance->Bounds.Min), FVector2D(ClusterInstance->Bounds.Max));
+			return ClusterBounds.GetArea() <= CellArea;
+		}
+
+		if (GRuntimeSpatialHashPlacePartitionActorsUsingLocation)
+		{
+			bool bUseLocation = true;
+			for (const FGuid& ActorGuid : ClusterInstance->Cluster->Actors)
+			{
+				const FWorldPartitionActorDescView& ActorDescView = ClusterInstance->ContainerInstance->GetActorDescView(ActorGuid);
+
+				if (!ActorDescView.GetActorNativeClass()->IsChildOf<APartitionActor>())
+				{
+					bUseLocation = false;
+					break;
+				}
+			}
+			return bUseLocation;
+		}
+
+		return false;
+	};
+
 	for (const FActorClusterInstance* ClusterInstance : GridActors)
 	{
 		const FActorCluster* ActorCluster = ClusterInstance->Cluster;
@@ -191,7 +226,7 @@ FSquare2DGridHelper GetPartitionedActors(const UWorldPartition* WorldPartition,
 
 		if (ActorCluster->bIsSpatiallyLoaded)
 		{
-			if (GRuntimeSpatialHashPlaceSmallActorsUsingLocation && (ClusterBounds.GetArea() <= CellArea))
+			if (ShouldActorUseLocationPlacement(ClusterInstance))
 			{
 				// Find grid level cell that contains the actor cluster pivot and put actors in it.
 				FGridCellCoord2 CellCoords;
