OpenXR + VP: Restore Editor Viewport correctly after exiting VR Editor Mode (SIMPLE)

OpenXR relies on the logic in FSceneViewport::OnDrawViewport to handle Viewport resizes when stereo is active. Previously, I had a change that only went thru
for PIE + Stereo (VRPIE). I need the same logic for VREM, so I can change the PIE check to just Editor. The viewport handles the check for active stereo rendering already.

diff --git a/Engine/Source/Runtime/Engine/Private/Slate/SceneViewport.cpp b/Engine/Source/Runtime/Engine/Private/Slate/SceneViewport.cpp
--- a/Engine/Source/Runtime/Engine/Private/Slate/SceneViewport.cpp
+++ b/Engine/Source/Runtime/Engine/Private/Slate/SceneViewport.cpp
@@ -417,7 +417,7 @@ void FSceneViewport::OnDrawViewport( const FGeometry& AllottedGeometry, const FS
 
 				// In the stereo case, the HMD display size drives the base RT size, separate from the PIE mirror window
 				const bool bResizeTargetValid = Window->IsViewportSizeDrivenByWindow() || 
-					(IsPlayInEditorViewport() && IsStereoRenderingAllowed());
+					(GIsEditor && IsStereoRenderingAllowed());
 
 				if (bResizeTargetValid)
 				{
