Remove redundant calculation of MeshCenter in FStaticMeshSceneProxy::GetDynamicRayTracingInstances()

diff --git a/Engine/Source/Runtime/Engine/Private/StaticMeshRender.cpp b/Engine/Source/Runtime/Engine/Private/StaticMeshRender.cpp
--- a/Engine/Source/Runtime/Engine/Private/StaticMeshRender.cpp
+++ b/Engine/Source/Runtime/Engine/Private/StaticMeshRender.cpp
@@ -1789,8 +1789,8 @@ void FStaticMeshSceneProxy::GetDynamicRayTracingInstances(FRayTracingMaterialGat
 
 	if (bEvaluateWPO && CVarRayTracingStaticMeshesWPOCulling.GetValueOnRenderThread() > 0)
 	{
-		FVector ViewCenter = Context.ReferenceView->ViewMatrices.GetViewOrigin();		
-		FVector MeshCenter = GetLocalToWorld().TransformPosition({ 0.0f, 0.0f, 0.0f });
+		FVector ViewCenter = Context.ReferenceView->ViewMatrices.GetViewOrigin();
+		FVector MeshCenter = GetLocalToWorld().GetOrigin();
 		const float CullingRadius = CVarRayTracingStaticMeshesWPOCullingRadius.GetValueOnRenderThread();
 		const float BoundingRadius = GetBounds().SphereRadius;
 
