World Partition: stores native class full path name to avoid TryFindType warnings.

#rnx

diff --git a/Engine/Source/Editor/SceneOutliner/Private/ActorDescTreeItem.cpp b/Engine/Source/Editor/SceneOutliner/Private/ActorDescTreeItem.cpp
--- a/Engine/Source/Editor/SceneOutliner/Private/ActorDescTreeItem.cpp
+++ b/Engine/Source/Editor/SceneOutliner/Private/ActorDescTreeItem.cpp
@@ -166,7 +166,7 @@ private:
 		{
 			if (const FWorldPartitionActorDesc* ActorDesc = TreeItem->ActorDescHandle.Get())
 			{
-				const FName IconName = ActorDesc->GetNativeClass();
+				const FName IconName = ActorDesc->GetDisplayClassName();
 
 				const FSlateBrush* CachedBrush = WeakSceneOutliner.Pin()->GetCachedIconForClass(IconName);
 				if (CachedBrush != nullptr)
diff --git a/Engine/Source/Runtime/Engine/Private/WorldPartition/WorldPartitionActorDesc.cpp b/Engine/Source/Runtime/Engine/Private/WorldPartition/WorldPartitionActorDesc.cpp
--- a/Engine/Source/Runtime/Engine/Private/WorldPartition/WorldPartitionActorDesc.cpp
+++ b/Engine/Source/Runtime/Engine/Private/WorldPartition/WorldPartitionActorDesc.cpp
@@ -45,7 +45,7 @@ void FWorldPartitionActorDesc::Init(const AActor* InActor)
 
 	// Get the first native class in the hierarchy
 	ActorNativeClass = GetParentNativeClass(ActorClass);
-	NativeClass = ActorNativeClass->GetFName();
+	NativeClass = *ActorNativeClass->GetPathName();
 	
 	// For native class, don't set this
 	if (!ActorClass->IsNative())
@@ -135,7 +135,7 @@ void FWorldPartitionActorDesc::Init(const FWorldPartitionActorDescInitData& Desc
 	ActorPackage = DescData.PackageName;
 	ActorPath = DescData.ActorPath;
 	ActorNativeClass = DescData.NativeClass;
-	NativeClass = DescData.NativeClass->GetFName();
+	NativeClass = *DescData.NativeClass->GetPathName();
 
 	// Serialize actor metadata
 	FMemoryReader MetadataAr(DescData.SerializedData, true);
@@ -436,21 +436,20 @@ FName FWorldPartitionActorDesc::GetActorLabelOrName() const
 
 FName FWorldPartitionActorDesc::GetDisplayClassName() const
 {
-	if (BaseClass.IsNone())
+	auto GetCleanClassName = [](FName ClassName) -> FName
 	{
-		return NativeClass;
-	}
-
-	int32 Index;
-	const FString BaseClassStr(BaseClass.ToString());
-	if (BaseClassStr.FindLastChar(TCHAR('.'), Index))
-	{
-		FString CleanClassName = BaseClassStr.Mid(Index + 1);
-		CleanClassName.RemoveFromEnd(TEXT("_C"));
-		return *CleanClassName;
-	}
+		int32 Index;
+		FString ClassNameStr(ClassName.ToString());
+		if (ClassNameStr.FindLastChar(TCHAR('.'), Index))
+		{
+			FString CleanClassName = ClassNameStr.Mid(Index + 1);
+			CleanClassName.RemoveFromEnd(TEXT("_C"));
+			return *CleanClassName;
+		}
+		return ClassName;
+	};
 
-	return BaseClass;
+	return BaseClass.IsNone() ? GetCleanClassName(NativeClass) : GetCleanClassName(BaseClass);
 }
 
 bool FWorldPartitionActorDesc::IsLoaded(bool bEvenIfPendingKill) const
