Iris experimental replication system.

diff --git a/Engine/Source/Programs/ReplicationSystemTest/Private/ReplicationSystemTest.cpp b/Engine/Source/Programs/ReplicationSystemTest/Private/ReplicationSystemTest.cpp
new file mode 100644
--- /dev/null
+++ b/Engine/Source/Programs/ReplicationSystemTest/Private/ReplicationSystemTest.cpp
@@ -0,0 +1,178 @@
+// Copyright Epic Games, Inc. All Rights Reserved.
+
+#include "ReplicationSystemTest.h"
+#include "ReplicationSystemTestPlugin/ReplicationSystemTestPlugin.h"
+#include "RequiredProgramMainCPPInclude.h"
+#include "Styling/UMGCoreStyle.h"
+#include "AutomationTestRunner.h"
+#include "NullTestRunner.h"
+
+#if WITH_AUTOMATION_WORKER
+namespace UE::Net
+{
+	class FAutomationTestRunner;
+}
+
+typedef UE::Net::FAutomationTestRunner TestRunner;
+#else
+typedef FNullTestRunner TestRunner;
+#endif
+
+DEFINE_LOG_CATEGORY_STATIC(LogReplicationSystemTest, Log, All);
+
+IMPLEMENT_APPLICATION(ReplicationSystemTest, "");
+
+bool GIsConsoleExecutable = true;
+
+static void PreInit();
+static void LoadModules();
+static void PostInit();
+static void TearDown();
+
+#if defined(PLATFORM_XBOXONE) && PLATFORM_XBOXONE
+int TestMain()
+#else
+INT32_MAIN_INT32_ARGC_TCHAR_ARGV()
+#endif
+{
+	// Due to some code not respecting nullrhi etc.
+	PRIVATE_GIsRunningCommandlet = true;
+	PRIVATE_GAllowCommandletRendering = false;
+	PRIVATE_GAllowCommandletAudio = false;
+
+	FCommandLine::Set(
+		TEXT("-nullrhi -log -NoAsyncLoadingThread -NoAsyncPostLoad -noedl -unattended -ENGINEINI=Engine.ini")
+		TEXT(" -ini:Engine:[/Script/Engine.Engine]:AIControllerClassName=/Script/AIModule.AIController")
+	);
+
+	PreInit();
+	LoadModules();
+	PostInit();
+
+	UE_LOG(LogReplicationSystemTest, Display, TEXT("ReplicationSystemTest"));
+
+	{
+		TUniquePtr<TestRunner> Runner(new TestRunner());
+		Runner->RunTests();
+	}
+
+	TearDown();
+
+	return 0;
+}
+
+static void TearDown()
+{
+	RequestEngineExit(TEXT("Shutting down ReplicationSystemTest"));
+
+	FPlatformApplicationMisc::TearDown();
+	FPlatformMisc::PlatformTearDown();
+
+	if (GLog)
+	{
+		GLog->TearDown();
+	}
+
+	FCoreDelegates::OnExit.Broadcast();
+	FModuleManager::Get().UnloadModulesAtShutdown();
+
+#if STATS
+	FThreadStats::StopThread();
+#endif
+
+	FTaskGraphInterface::Shutdown();
+
+	if (GConfig)
+	{
+		GConfig->Exit();
+		delete GConfig;
+		GConfig = nullptr;
+	}
+}
+
+static void PreInit()
+{
+	FGenericPlatformOutputDevices::SetupOutputDevices();
+
+	GError = FPlatformApplicationMisc::GetErrorOutputDevice();
+	GWarn = FPlatformApplicationMisc::GetFeedbackContext();
+
+	FPlatformMisc::PlatformInit();
+#if WITH_APPLICATION_CORE
+	FPlatformApplicationMisc::Init();
+#endif
+	FPlatformMemory::Init();
+
+	if (IPlatformFile* WrapperFile = FPlatformFileManager::Get().GetPlatformFile(TEXT("ReplicationSystemTestFile")))
+	{
+		IPlatformFile* CurrentPlatformFile = &FPlatformFileManager::Get().GetPlatformFile();
+		WrapperFile->Initialize(CurrentPlatformFile, TEXT(""));
+		FPlatformFileManager::Get().SetPlatformFile(*WrapperFile);
+	}
+
+#if WITH_COREUOBJECT
+	// Initialize the PackageResourceManager, which is needed to load any (non-script) Packages. It is first used in ProcessNewlyLoadedObjects (due to the loading of asset references in Class Default Objects)
+	// It has to be intialized after the AssetRegistryModule; the editor implementations of PackageResourceManager relies on it
+	IPackageResourceManager::Initialize();
+#endif
+
+	FDelayedAutoRegisterHelper::RunAndClearDelayedAutoRegisterDelegates(EDelayedRegisterRunPhase::FileSystemReady);
+
+	FConfigCacheIni::InitializeConfigSystem();
+
+	// 
+	GConfig->SetInt(TEXT("/Script/Engine.GarbageCollectionSettings"), TEXT("gc.MaxObjectsNotConsideredByGC"), 0, GEngineIni);
+
+	// Console commands
+	IConsoleManager::Get().ProcessUserConsoleInput(TEXT("Net.IsPushModelEnabled 1"), *GLog, nullptr);
+
+	GGameThreadId = FPlatformTLS::GetCurrentThreadId();
+	FTaskGraphInterface::Startup(FPlatformMisc::NumberOfCores());
+	FTaskGraphInterface::Get().AttachToThread(ENamedThreads::GameThread);
+
+	FDelayedAutoRegisterHelper::RunAndClearDelayedAutoRegisterDelegates(EDelayedRegisterRunPhase::TaskGraphSystemReady);
+
+#if STATS
+	FThreadStats::StartThread();
+#endif
+
+	FDelayedAutoRegisterHelper::RunAndClearDelayedAutoRegisterDelegates(EDelayedRegisterRunPhase::StatSystemReady);
+
+	FCoreStyle::ResetToDefault();
+	FUMGCoreStyle::ResetToDefault();
+}
+
+static void LoadModules()
+{
+	// Always attempt to load CoreUObject. It requires additional pre-init which is called from its module's StartupModule method.
+#if WITH_COREUOBJECT
+#if USE_PER_MODULE_UOBJECT_BOOTSTRAP // otherwise do it later
+	FModuleManager::Get().OnProcessLoadedObjectsCallback().AddStatic(ProcessNewlyLoadedUObjects);
+#endif
+	FModuleManager::Get().LoadModule(TEXT("CoreUObject"));
+
+	FCoreDelegates::OnInit.Broadcast();
+#endif
+
+	// ChaosEngineSolvers requires ChaosSolvers. We're not able to call ProcessNewlyLoadedObjects before this module is loaded.
+	FModuleManager::Get().LoadModule(TEXT("ChaosSolvers"));
+	ProcessNewlyLoadedUObjects();
+
+	FModuleManager::Get().LoadModule(TEXT("IrisCore"));
+
+	FModuleManager::Get().LoadModule(TEXT("ReplicationSystemTestPlugin"));
+}
+
+static void PostInit()
+{
+#if WITH_COREUOBJECT
+	// Required for GC to be allowed.
+	if (GUObjectArray.IsOpenForDisregardForGC())
+	{
+		GUObjectArray.CloseDisregardForGC();
+	}
+#endif
+
+	// Disable ini file operations
+	GConfig->DisableFileOperations();
+}
diff --git a/Engine/Source/Programs/ReplicationSystemTest/Private/ReplicationSystemTest.h b/Engine/Source/Programs/ReplicationSystemTest/Private/ReplicationSystemTest.h
new file mode 100644
--- /dev/null
+++ b/Engine/Source/Programs/ReplicationSystemTest/Private/ReplicationSystemTest.h
@@ -0,0 +1,6 @@
+// Copyright Epic Games, Inc. All Rights Reserved.
+
+#pragma once
+
+#include "CoreMinimal.h"
+
diff --git a/Engine/Source/Programs/ReplicationSystemTest/Private/ReplicationSystemTestFileModule.cpp b/Engine/Source/Programs/ReplicationSystemTest/Private/ReplicationSystemTestFileModule.cpp
new file mode 100644
--- /dev/null
+++ b/Engine/Source/Programs/ReplicationSystemTest/Private/ReplicationSystemTestFileModule.cpp
@@ -0,0 +1,59 @@
+// Copyright Epic Games, Inc. All Rights Reserved.
+
+#include "HAL/PlatformFile.h"
+#include "HAL/IPlatformFileModule.h"
+#include "Modules/ModuleInterface.h"
+#include "Modules/ModuleManager.h"
+
+class FReplicationSystemTestFile : public IPlatformFile
+{
+public:
+	FReplicationSystemTestFile() : LowerLevel(nullptr) {}
+
+	virtual bool ShouldBeUsed(IPlatformFile* Inner, const TCHAR* CmdLine) const override { return true; }
+	virtual bool Initialize(IPlatformFile* Inner, const TCHAR* CmdLine) override { LowerLevel = Inner; return true; }
+
+	virtual IPlatformFile* GetLowerLevel() { return LowerLevel; }
+	virtual void SetLowerLevel(IPlatformFile* NewLowerLevel) { LowerLevel = NewLowerLevel; }
+	virtual const TCHAR* GetName() const { return TEXT("ReplicationSystemTest"); }
+
+	virtual bool FileExists(const TCHAR* Filename) override { return false; }
+	virtual int64 FileSize(const TCHAR* Filename) override {return -1LL; }
+
+	virtual bool DeleteFile(const TCHAR* Filename) override { return false; }
+	virtual bool IsReadOnly(const TCHAR* Filename) override { return false; }
+
+	virtual bool MoveFile(const TCHAR* To, const TCHAR* From) override { return true; }
+	virtual bool SetReadOnly(const TCHAR* Filename, bool bNewReadOnlyValue) override { return false; }
+	virtual FDateTime GetTimeStamp(const TCHAR* Filename) override { return FDateTime::MinValue(); }
+	virtual void SetTimeStamp(const TCHAR* Filename, FDateTime DateTime) override {}
+	virtual FDateTime GetAccessTimeStamp(const TCHAR* Filename) override { return FDateTime::MinValue(); }
+	virtual FString GetFilenameOnDisk(const TCHAR* Filename) override { return TEXT("I.ris"); }
+	virtual IFileHandle* OpenRead(const TCHAR* Filename, bool bAllowWrite = false) override { return nullptr; }
+	virtual IFileHandle* OpenWrite(const TCHAR* Filename, bool bAppend = false, bool bAllowRead = false) override { return nullptr; }
+
+	virtual bool DirectoryExists(const TCHAR* Directory) override { return false; }
+
+	virtual bool CreateDirectory(const TCHAR* Directory) override { return false; }
+	virtual bool DeleteDirectory(const TCHAR* Directory) override { return false; }
+	virtual FFileStatData GetStatData(const TCHAR* FilenameOrDirectory) override { FFileStatData Data; Data.bIsValid = false; return Data; }
+
+	virtual bool IterateDirectory(const TCHAR* Directory, FDirectoryVisitor& Visitor) override { return true; }
+	virtual bool IterateDirectoryStat(const TCHAR* Directory, FDirectoryStatVisitor& Visitor) override { return true; }
+
+private:
+	IPlatformFile* LowerLevel;
+};
+
+
+class FReplicationSystemTestFileModule : public IPlatformFileModule
+{
+public:
+	virtual IPlatformFile* GetPlatformFile() override
+	{
+		static TUniquePtr<IPlatformFile> AutoDestroySingleton = MakeUnique<FReplicationSystemTestFile>();
+		return AutoDestroySingleton.Get();
+	}
+};
+
+IMPLEMENT_MODULE(FReplicationSystemTestFileModule, ReplicationSystemTestFile);
diff --git a/Engine/Source/Programs/ReplicationSystemTest/ReplicationSystemTest.Build.cs b/Engine/Source/Programs/ReplicationSystemTest/ReplicationSystemTest.Build.cs
new file mode 100644
--- /dev/null
+++ b/Engine/Source/Programs/ReplicationSystemTest/ReplicationSystemTest.Build.cs
@@ -0,0 +1,40 @@
+// Copyright Epic Games, Inc. All Rights Reserved.
+
+using UnrealBuildTool;
+
+public class ReplicationSystemTest : ModuleRules
+{
+	public ReplicationSystemTest(ReadOnlyTargetRules Target) : base(Target)
+	{
+		PublicIncludePaths.Add("Runtime/Launch/Public");
+
+		// For LaunchEngineLoop.cpp include
+		PrivateIncludePaths.Add("Runtime/Launch/Private");
+
+		PrivateDependencyModuleNames.AddRange(
+			new string[]
+			{
+ 				"ApplicationCore",
+ 				"AutomationController",
+ 				"AutomationWorker",
+                "Core",
+				"Projects",
+				"DerivedDataCache",
+				"Engine",
+				"HeadMountedDisplay",
+				"InstallBundleManager",
+                "MediaUtils",
+				"MRMesh",
+				"MoviePlayer",
+				"MoviePlayerProxy",
+				"PreLoadScreen",
+				"ProfilerService",
+				"ReplicationSystemTestPlugin",
+				"SessionServices",
+				"SlateNullRenderer",
+				"SlateRHIRenderer",
+				"TaskGraph",
+ 			}
+		);
+	}
+}
diff --git a/Engine/Source/Programs/ReplicationSystemTest/ReplicationSystemTest.Target.cs b/Engine/Source/Programs/ReplicationSystemTest/ReplicationSystemTest.Target.cs
new file mode 100644
--- /dev/null
+++ b/Engine/Source/Programs/ReplicationSystemTest/ReplicationSystemTest.Target.cs
@@ -0,0 +1,45 @@
+// Copyright Epic Games, Inc. All Rights Reserved.
+
+using UnrealBuildTool;
+using System.Collections.Generic;
+
+[SupportedPlatforms(UnrealPlatformClass.All)]
+public class ReplicationSystemTestTarget : TargetRules
+{
+	public ReplicationSystemTestTarget(TargetInfo Target) : base(Target)
+	{
+		Type = TargetType.Program;
+		LinkType = TargetLinkType.Monolithic;
+		LaunchModuleName = "ReplicationSystemTest";
+
+		bUseMallocProfiler = false;
+
+		// No editor-only data is needed
+		bBuildWithEditorOnlyData = false;
+
+		bCompileAgainstEngine = true;
+		bCompileAgainstCoreUObject = false;
+		bCompileAgainstApplicationCore = true;
+        bCompileWithPluginSupport = true;
+        bBuildDeveloperTools = false;
+        bBuildRequiresCookedData = true; // this program requires no data
+
+		// No ICU internationalization as it causes shutdown errors.
+		bCompileICU = false;
+
+		// Don't need slate
+		bUsesSlate = false;
+
+		IncludeOrderVersion = EngineIncludeOrderVersion.Latest;
+		bIsBuildingConsoleApplication = true;
+		bLegalToDistributeBinary = true;
+
+		// Always want logging so we can see the test results
+        bUseLoggingInShipping = true;
+
+		// Network config
+		bWithPushModel = true;		
+		bUseIris = true;
+	}
+}
+
diff --git a/Engine/Source/Programs/UnrealBuildTool/Configuration/ModuleRules.cs b/Engine/Source/Programs/UnrealBuildTool/Configuration/ModuleRules.cs
--- a/Engine/Source/Programs/UnrealBuildTool/Configuration/ModuleRules.cs
+++ b/Engine/Source/Programs/UnrealBuildTool/Configuration/ModuleRules.cs
@@ -1238,6 +1238,40 @@ namespace UnrealBuildTool
 			}
 		}
 
+		/// <summary>
+		/// Setup this module for Iris support (based on the settings in UEBuildConfiguration)
+		/// </summary>
+		public void SetupIrisSupport(ReadOnlyTargetRules Target, bool bAddAsPublicDependency = false)
+		{
+			if (Target.bUseIris == true)
+			{
+				PublicDefinitions.Add("UE_WITH_IRIS=1");
+
+				if (bAddAsPublicDependency)
+				{
+					PublicDependencyModuleNames.Add("IrisCore");
+				}
+				else
+				{
+					PrivateDependencyModuleNames.Add("IrisCore");
+				}
+			}
+			else
+			{
+				PublicDefinitions.Add("UE_WITH_IRIS=0");
+
+				// If we compile without Iris we have a stub Iris module for UHT dependencies
+				if (bAddAsPublicDependency)
+				{
+					PublicDependencyModuleNames.Add("IrisStub");
+				}
+				else
+				{
+					PrivateDependencyModuleNames.Add("IrisStub");
+				}
+			}
+		}
+
 		/// <summary>
 		/// Setup this module for physics support (based on the settings in UEBuildConfiguration)
 		/// </summary>
diff --git a/Engine/Source/Programs/UnrealBuildTool/Configuration/TargetRules.cs b/Engine/Source/Programs/UnrealBuildTool/Configuration/TargetRules.cs
--- a/Engine/Source/Programs/UnrealBuildTool/Configuration/TargetRules.cs
+++ b/Engine/Source/Programs/UnrealBuildTool/Configuration/TargetRules.cs
@@ -728,6 +728,14 @@ namespace UnrealBuildTool
 		[RequiresUniqueBuildEnvironment]
 		public bool bCompilePython = true;
 
+		/// <summary>
+		/// Whether to use Iris.
+		/// </summary>
+		[RequiresUniqueBuildEnvironment]
+		[CommandLine("-NoUseIris", Value = "false")]
+		[CommandLine("-UseIris", Value = "true")]
+		public bool bUseIris = false;
+
 		/// <summary>
 		/// Whether we are compiling editor code or not. Prefer the more explicit bCompileAgainstEditor instead.
 		/// </summary>
@@ -2609,6 +2617,11 @@ namespace UnrealBuildTool
 			get { return Inner.bCompileISPC; }
 		}
 
+		public bool bUseIris
+		{
+			get { return Inner.bUseIris; }
+		}
+
 		public bool bCompileIntelMetricsDiscovery
 		{
 			get { return Inner.bCompileIntelMetricsDiscovery; }
diff --git a/Engine/Source/Programs/UnrealHeaderTool/Private/CodeGenerator.cpp b/Engine/Source/Programs/UnrealHeaderTool/Private/CodeGenerator.cpp
--- a/Engine/Source/Programs/UnrealHeaderTool/Private/CodeGenerator.cpp
+++ b/Engine/Source/Programs/UnrealHeaderTool/Private/CodeGenerator.cpp
@@ -3375,7 +3375,7 @@ void FNativeClassHeaderGenerator::ExportEnum(FOutputDevice& Out, FUnrealEnumDefi
 }
 
 // Exports the header text for the list of structs specified (GENERATED_BODY impls)
-void FNativeClassHeaderGenerator::ExportGeneratedStructBodyMacros(FOutputDevice& OutGeneratedHeaderText, FOutputDevice& Out, FReferenceGatherers& OutReferenceGatherers, const FUnrealSourceFile& SourceFile, FUnrealScriptStructDefinitionInfo& ScriptStructDef) const
+void FNativeClassHeaderGenerator::ExportGeneratedStructBodyMacros(FOutputDevice& OutGeneratedHeaderText, FOutputDevice& Out, FReferenceGatherers& OutReferenceGatherers, const FUnrealSourceFile& SourceFile, FUnrealScriptStructDefinitionInfo& ScriptStructDef, EExportClassOutFlags& OutFlags) const
 {
 	ScriptStructDef.AddCrossModuleReference(OutReferenceGatherers.UniqueCrossModuleReferences, true);
 
@@ -3518,9 +3518,38 @@ void FNativeClassHeaderGenerator::ExportGeneratedStructBodyMacros(FOutputDevice&
 			}
 		}
 
+		// Check to see if we are a FastArraySerializer and should try to deduce the FastArraySerializerItemType
+		// To fulfill that requirement the struct should be derived from FFastArraySerializer and have a single replicated TArrayProperty
+		bool bGenerateFastArraySerializerTypeDefinition = false;
+		FString FastArraySerializerTypeDefinition;
+		{
+			const FUnrealStructDefinitionInfo* FastArraySerializerStructDef = GTypeDefinitionInfoMap.FindByName<FUnrealStructDefinitionInfo>(TEXT("FastArraySerializer"));
+			if (FastArraySerializerStructDef && BaseStructDef && BaseStructDef->IsChildOf(*FastArraySerializerStructDef))
+			{
+				uint32 ReplicatedPropertyCount = 0U;
+				const FUnrealPropertyDefinitionInfo* PotentialItemArrayPropertyDefinition = nullptr;
+
+				// Only output bindings for FastArraySerializers that have a single replicated dynamic array
+				for (const TSharedRef<FUnrealPropertyDefinitionInfo>& Property : ScriptStructDef.GetProperties())
+				{
+					if (!EnumHasAnyFlags(Property->GetPropertyFlags(), EPropertyFlags::CPF_RepSkip) && Property->IsDynamicArray())
+					{
+						PotentialItemArrayPropertyDefinition = &Property.Get();
+						++ReplicatedPropertyCount;
+					}
+				}
+
+				if (ReplicatedPropertyCount == 1)
+				{
+					FastArraySerializerTypeDefinition = FString::Printf(TEXT("\tUE_NET_DECLARE_FASTARRAY(%s, %s, %s);\r\n"), *StructNameCPP, *(PotentialItemArrayPropertyDefinition->GetName()), (bRequiredAPI ? *GetAPIString() : TEXT("")));
+					bGenerateFastArraySerializerTypeDefinition = true;
+				}
+			}
+		}
+
 		const FString SuperTypedef = BaseStructDef ? FString::Printf(TEXT("\ttypedef %s Super;\r\n"), *BaseStructDef->GetAlternateNameCPP()) : FString();
 
-		FString CombinedLine = FString::Printf(TEXT("%s%s%s%s"), *FriendLine, *StaticClassLine, *RigVMMethodsDeclarations, *SuperTypedef);
+		FString CombinedLine = FString::Printf(TEXT("%s%s%s%s%s"), *FriendLine, *StaticClassLine, *RigVMMethodsDeclarations, *SuperTypedef, *FastArraySerializerTypeDefinition);
 		const FString MacroName = SourceFile.GetGeneratedBodyMacroName(ScriptStructDef.GetMacroDeclaredLineNumber());
 
 		const FString Macroized = Macroize(*MacroName, MoveTemp(CombinedLine));
@@ -3586,6 +3615,13 @@ void FNativeClassHeaderGenerator::ExportGeneratedStructBodyMacros(FOutputDevice&
 		Out.Logf(TEXT("{\r\n"));
 		Out.Logf(TEXT("\treturn %s::StaticStruct();\r\n"), *StructNameCPP);
 		Out.Logf(TEXT("}\r\n"));
+
+		// Inject implementation needed to support auto bindings of fast arrays
+		if (bGenerateFastArraySerializerTypeDefinition)
+		{
+			OutFlags |= EExportClassOutFlags::NeedsFastArrayHeaders;
+			Out.Logf(TEXT("UE_NET_IMPLEMENT_FASTARRAY(%s);\r\n"), *StructNameCPP);
+		}
 	}
 
 	FString StaticsStructName = ChoppedSingletonName + TEXT("_Statics");
@@ -5423,7 +5459,7 @@ void FNativeClassHeaderGenerator::ApplyAlternatePropertyExportText(FUnrealProper
 	}
 }
 
-bool FNativeClassHeaderGenerator::WriteSource(const FManifestModule& Module, FGeneratedFileInfo& FileInfo, const FString& InBodyText, FUnrealSourceFile* InSourceFile, const TSet<FString>& InCrossModuleReferences)
+bool FNativeClassHeaderGenerator::WriteSource(const FManifestModule& Module, FGeneratedFileInfo& FileInfo, const FString& InBodyText, FUnrealSourceFile* InSourceFile, const TSet<FString>& InCrossModuleReferences, const EExportClassOutFlags& ExportFlags)
 {
 	// Collect the includes if this is from a source file
 	TArray<FString> RelativeIncludes;
@@ -5433,6 +5469,11 @@ bool FNativeClassHeaderGenerator::WriteSource(const FManifestModule& Module, FGe
 		ConvertToBuildIncludePath(Module, ModuleRelativeFilename);
 		RelativeIncludes.Add(MoveTemp(ModuleRelativeFilename));
 
+		if (EnumHasAnyFlags(ExportFlags, EExportClassOutFlags::NeedsFastArrayHeaders))
+		{
+			RelativeIncludes.Add(FString(TEXT("Net/Serialization/FastArraySerializerImplementation.h")));
+		}
+
 		bool bAddedStructuredArchiveFromArchiveHeader = false;
 		bool bAddedArchiveUObjectFromStructuredArchiveHeader = false;
 		for (const TSharedRef<FUnrealTypeDefinitionInfo>& TypeDef : InSourceFile->GetDefinedClasses())
@@ -5607,7 +5648,7 @@ void FNativeClassHeaderGenerator::GenerateSourceFile(FGeneratedCPP& GeneratedCPP
 						Singletons.Add(ScriptStructDef);
 					}
 					GeneratedFunctionDeclarations.Log(ScriptStructDef->GetExternDecl(true));
-					Generator.ExportGeneratedStructBodyMacros(GeneratedHeaderText, GeneratedCPPText, ReferenceGatherers, SourceFile, *ScriptStructDef);
+					Generator.ExportGeneratedStructBodyMacros(GeneratedHeaderText, GeneratedCPPText, ReferenceGatherers, SourceFile, *ScriptStructDef, GeneratedCPP.ExportFlags);
 					if (ScriptStructDef->HasAnyStructFlags(STRUCT_Native))
 					{
 						ScriptStructRegs.Add(ScriptStructDef);
@@ -5835,7 +5876,7 @@ void FNativeClassHeaderGenerator::WriteSourceFile(FGeneratedCPP& GeneratedCPP)
 			}
 
 			bool bHasChanged = WriteHeader(GeneratedCPP.Header, GeneratedHeaderText, AdditionalHeaders, GeneratedCPP.ForwardDeclarations);
-			WriteSource(Module, GeneratedCPP.Source, GeneratedCPPText, &SourceFile, GeneratedCPP.CrossModuleReferences);
+			WriteSource(Module, GeneratedCPP.Source, GeneratedCPPText, &SourceFile, GeneratedCPP.CrossModuleReferences, GeneratedCPP.ExportFlags);
 
 			SourceFile.SetGeneratedFilename(MoveTemp(GeneratedCPP.Header.GetFilename()));
 			SourceFile.SetHasChanged(bHasChanged);
@@ -6073,7 +6114,7 @@ void FNativeClassHeaderGenerator::Generate(
 					}
 
 					Generator.ExportGeneratedPackageInitCode(GeneratedFileInfo.GetGeneratedBody(), *GeneratedFunctionDeclarations, ExportedSorted, CombinedHash);
-					WriteSource(Module, GeneratedFileInfo, GeneratedFileInfo.GetGeneratedBody(), nullptr, TSet<FString>());
+					WriteSource(Module, GeneratedFileInfo, GeneratedFileInfo.GetGeneratedBody(), nullptr, TSet<FString>(), EExportClassOutFlags::None);
 				}
 			});
 		};
diff --git a/Engine/Source/Programs/UnrealHeaderTool/Private/NativeClassExporter.h b/Engine/Source/Programs/UnrealHeaderTool/Private/NativeClassExporter.h
--- a/Engine/Source/Programs/UnrealHeaderTool/Private/NativeClassExporter.h
+++ b/Engine/Source/Programs/UnrealHeaderTool/Private/NativeClassExporter.h
@@ -58,6 +58,7 @@ enum class EExportClassOutFlags
 {
 	None = 0x0,
 	NeedsPushModelHeaders = 0x1 << 0,
+	NeedsFastArrayHeaders = NeedsPushModelHeaders << 1,
 };
 ENUM_CLASS_FLAGS(EExportClassOutFlags);
 
@@ -385,7 +386,7 @@ private:
 	/**
 	 * Write the body of a source file using a standard format
 	 */
-	static bool WriteSource(const FManifestModule& Module, FGeneratedFileInfo& FileInfo, const FString& InBodyText, FUnrealSourceFile* InSourceFile, const TSet<FString>& InCrossModuleReferences);
+	static bool WriteSource(const FManifestModule& Module, FGeneratedFileInfo& FileInfo, const FString& InBodyText, FUnrealSourceFile* InSourceFile, const TSet<FString>& InCrossModuleReferences, const EExportClassOutFlags& ExportFlags);
 
 	/**
 	 * Returns a string in the format CLASS_Something|CLASS_Something which represents all class flags that are set for the specified
@@ -415,7 +416,7 @@ private:
 	 * @param	Out				output device
 	 * @param	Struct			The struct to export
 	 */
-	void ExportGeneratedStructBodyMacros(FOutputDevice& OutGeneratedHeaderText, FOutputDevice& Out, FReferenceGatherers& OutReferenceGatherers, const FUnrealSourceFile& SourceFile, FUnrealScriptStructDefinitionInfo& ScriptStructDef) const;
+	void ExportGeneratedStructBodyMacros(FOutputDevice& OutGeneratedHeaderText, FOutputDevice& Out, FReferenceGatherers& OutReferenceGatherers, const FUnrealSourceFile& SourceFile, FUnrealScriptStructDefinitionInfo& ScriptStructDef, EExportClassOutFlags& OutFlags) const;
 
 	/**
 	 * Exports a local mirror of the specified struct; used to get offsets
diff --git a/Engine/Source/Runtime/Core/Core.Build.cs b/Engine/Source/Runtime/Core/Core.Build.cs
--- a/Engine/Source/Runtime/Core/Core.Build.cs
+++ b/Engine/Source/Runtime/Core/Core.Build.cs
@@ -278,6 +278,16 @@ public class Core : ModuleRules
 			PrivateDefinitions.Add("IS_CLIENT_TARGET=0");
 		}
 
+		// Setup definitions to include / exclude Iris modifications to UObject Note: Only the definition is required as we do not depend on Iris in any way.
+		if (Target.bUseIris == true)
+		{
+			PublicDefinitions.Add("UE_WITH_IRIS=1");
+		}
+		else
+		{
+			PublicDefinitions.Add("UE_WITH_IRIS=0");
+		}
+		
 		if (Target.Platform == UnrealTargetPlatform.Win64
 			&& Target.Configuration != UnrealTargetConfiguration.Shipping)
 		{
diff --git a/Engine/Source/Runtime/Core/Private/HAL/LowLevelMemTracker.cpp b/Engine/Source/Runtime/Core/Private/HAL/LowLevelMemTracker.cpp
--- a/Engine/Source/Runtime/Core/Private/HAL/LowLevelMemTracker.cpp
+++ b/Engine/Source/Runtime/Core/Private/HAL/LowLevelMemTracker.cpp
@@ -210,6 +210,9 @@ DECLARE_LLM_MEMORY_STAT(TEXT("AudioSynthesis"), STAT_AudioSynthesisLLM, STATGROU
 DECLARE_LLM_MEMORY_STAT(TEXT("RealTimeCommunications"), STAT_RealTimeCommunicationsLLM, STATGROUP_LLMFULL);
 DECLARE_LLM_MEMORY_STAT(TEXT("FName"), STAT_FNameLLM, STATGROUP_LLMFULL);
 DECLARE_LLM_MEMORY_STAT(TEXT("Networking"), STAT_NetworkingLLM, STATGROUP_LLMFULL);
+DECLARE_LLM_MEMORY_STAT(TEXT("Iris"), STAT_IrisLLM, STATGROUP_LLMFULL);
+DECLARE_LLM_MEMORY_STAT(TEXT("IrisState"), STAT_IrisStateLLM, STATGROUP_LLMFULL);
+DECLARE_LLM_MEMORY_STAT(TEXT("IrisConnection"), STAT_IrisConnectionLLM, STATGROUP_LLMFULL);
 DECLARE_LLM_MEMORY_STAT(TEXT("Meshes"), STAT_MeshesLLM, STATGROUP_LLMFULL);
 DECLARE_LLM_MEMORY_STAT(TEXT("Stats"), STAT_StatsLLM, STATGROUP_LLMFULL);
 DECLARE_LLM_MEMORY_STAT(TEXT("Shaders"), STAT_ShadersLLM, STATGROUP_LLMFULL);
diff --git a/Engine/Source/Runtime/Core/Private/Misc/NetworkVersion.cpp b/Engine/Source/Runtime/Core/Private/Misc/NetworkVersion.cpp
--- a/Engine/Source/Runtime/Core/Private/Misc/NetworkVersion.cpp
+++ b/Engine/Source/Runtime/Core/Private/Misc/NetworkVersion.cpp
@@ -14,6 +14,7 @@
 #include "Misc/NetworkGuid.h" // IWYU pragma: keep
 #include "Misc/Parse.h"
 #include "Trace/Detail/Channel.h"
+#include "Misc/StringBuilder.h"
 
 DEFINE_LOG_CATEGORY( LogNetVersion );
 
@@ -175,3 +176,21 @@ FNetworkReplayVersion FNetworkVersion::GetReplayVersion()
 
 	return FNetworkReplayVersion( FApp::GetProjectName(), ReplayVersion, GetReplayCompatibleChangelist() );
 }
+
+bool FNetworkVersion::AreNetworkRuntimeFeaturesCompatible(EEngineNetworkRuntimeFeatures LocalFeatures, EEngineNetworkRuntimeFeatures RemoteFeatures)
+{
+	return LocalFeatures == RemoteFeatures;
+}
+
+void FNetworkVersion::DescribeNetworkRuntimeFeaturesBitset(EEngineNetworkRuntimeFeatures FeaturesBitflag, FStringBuilderBase& OutVerboseDescription)
+{
+	if (EnumHasAnyFlags(FeaturesBitflag, EEngineNetworkRuntimeFeatures::IrisEnabled))
+	{
+		OutVerboseDescription.Append(TEXT("IrisReplication"));
+	}
+	else
+	{
+		OutVerboseDescription.Append(TEXT("LegacyReplication"));
+	}
+}
+
diff --git a/Engine/Source/Runtime/Core/Public/HAL/LowLevelMemTracker.h b/Engine/Source/Runtime/Core/Public/HAL/LowLevelMemTracker.h
--- a/Engine/Source/Runtime/Core/Public/HAL/LowLevelMemTracker.h
+++ b/Engine/Source/Runtime/Core/Public/HAL/LowLevelMemTracker.h
@@ -166,6 +166,9 @@ enum class ELLMTagSet : uint8
 	macro(RealTimeCommunications,				"RealTimeCommunications",		GET_STATFNAME(STAT_RealTimeCommunicationsLLM),				NAME_None,										-1)\
 	macro(FName,								"FName",						GET_STATFNAME(STAT_FNameLLM),								GET_STATFNAME(STAT_EngineSummaryLLM),			-1)\
 	macro(Networking,							"Networking",					GET_STATFNAME(STAT_NetworkingLLM),							GET_STATFNAME(STAT_NetworkingSummaryLLM),		-1)\
+	macro(Iris,									"Iris",							GET_STATFNAME(STAT_IrisLLM),								GET_STATFNAME(STAT_NetworkingSummaryLLM),			ELLMTag::Networking)\
+	macro(IrisState,							"IrisState",					GET_STATFNAME(STAT_IrisStateLLM),							GET_STATFNAME(STAT_NetworkingSummaryLLM),			ELLMTag::Networking)\
+	macro(IrisConnection,						"IrisConnection",				GET_STATFNAME(STAT_IrisStateLLM),							GET_STATFNAME(STAT_NetworkingSummaryLLM),			ELLMTag::Networking)\
 	macro(Meshes,								"Meshes",						GET_STATFNAME(STAT_MeshesLLM),								GET_STATFNAME(STAT_MeshesSummaryLLM),			-1)\
 	macro(Stats,								"Stats",						GET_STATFNAME(STAT_StatsLLM),								GET_STATFNAME(STAT_EngineSummaryLLM),			-1)\
 	macro(Shaders,								"Shaders",						GET_STATFNAME(STAT_ShadersLLM),								GET_STATFNAME(STAT_EngineSummaryLLM),			-1)\
@@ -250,7 +253,7 @@ enum class ELLMTag : LLM_TAG_TYPE
 
 	//------------------------------
 	// Platform tags
-	PlatformTagStart = 108,
+	PlatformTagStart = 111,
 	PlatformTagEnd = 149,
 
 	//------------------------------
diff --git a/Engine/Source/Runtime/Core/Public/Misc/Build.h b/Engine/Source/Runtime/Core/Public/Misc/Build.h
--- a/Engine/Source/Runtime/Core/Public/Misc/Build.h
+++ b/Engine/Source/Runtime/Core/Public/Misc/Build.h
@@ -452,6 +452,11 @@
 #define WITH_IOSTORE_IN_EDITOR UE_IS_COOKED_EDITOR
 #endif
 
+// Controls whether Iris networking code is compiled in or not; should normally be defined by UBT
+#ifndef UE_WITH_IRIS
+	#define UE_WITH_IRIS 0
+#endif
+
 // Controls whether or not to make a global object to load COnfig.bin as soon as possible
 #ifndef PRELOAD_BINARY_CONFIG
 	#define PRELOAD_BINARY_CONFIG 1
diff --git a/Engine/Source/Runtime/Core/Public/Misc/NetworkVersion.h b/Engine/Source/Runtime/Core/Public/Misc/NetworkVersion.h
--- a/Engine/Source/Runtime/Core/Public/Misc/NetworkVersion.h
+++ b/Engine/Source/Runtime/Core/Public/Misc/NetworkVersion.h
@@ -31,6 +31,16 @@ public:
 	uint32		Changelist;
 };
 
+/**
+ * List of runtime features that can affect network compatibility between two connections
+ */
+enum class EEngineNetworkRuntimeFeatures : uint16
+{
+	None = 0,
+	IrisEnabled = 1 << None, // Are we running the Iris or Legacy network system
+};
+ENUM_CLASS_FLAGS(EEngineNetworkRuntimeFeatures);
+
 enum EEngineNetworkVersionHistory
 {
 	HISTORY_INITIAL = 1,
@@ -60,6 +70,7 @@ enum EEngineNetworkVersionHistory
 	HISTORY_REPMOVE_SERVERFRAME_AND_HANDLE = 25,	// Bump version to support serialization changes to RepMove so we can get the serverframe and physics handle associated with the object
 	HISTORY_21_AND_VIEWPITCH_ONLY_DO_NOT_USE = 26,  // Bump version to support up to history 21 + HISTORY_PAWN_REMOTEVIEWPITCH.  DO NOT USE!!!
 	HISTORY_PLACEHOLDER = 27,                       // Bump version to a placeholder.  This version is the same as HISTORY_REPMOVE_SERVERFRAME_AND_HANDLE
+	HISTORY_RUNTIME_FEATURES_COMPATIBILITY = 28,	// Bump version to add network runtime feature compatibility test to handshake (hello/upgrade) control messages
 	// New history items go above here.
 
 	HISTORY_ENGINENETVERSION_PLUS_ONE,
@@ -129,6 +140,16 @@ struct CORE_API FNetworkVersion
 	*/
 	static void SetGameCompatibleNetworkProtocolVersion(uint32 GameCompatibleNetworkProtocolVersion);
 
+	/**
+	 * Compares if the connection's runtime features are compatible with each other
+	 */
+	static bool AreNetworkRuntimeFeaturesCompatible(EEngineNetworkRuntimeFeatures LocalFeatures, EEngineNetworkRuntimeFeatures RemoteFeatures);
+
+	/**
+	 * Build and return a string describing the status of the the network runtime features bitflag
+	 */
+	static void DescribeNetworkRuntimeFeaturesBitset(EEngineNetworkRuntimeFeatures FeaturesBitflag, FStringBuilderBase& OutVerboseDescription);
+	
 	/**
 	* Returns the project version used by networking
 	* 
diff --git a/Engine/Source/Runtime/CoreUObject/Private/UObject/Obj.cpp b/Engine/Source/Runtime/CoreUObject/Private/UObject/Obj.cpp
--- a/Engine/Source/Runtime/CoreUObject/Private/UObject/Obj.cpp
+++ b/Engine/Source/Runtime/CoreUObject/Private/UObject/Obj.cpp
@@ -4950,6 +4950,12 @@ void UObject::GetReplicatedCustomConditionState(FCustomPropertyConditionState& O
 
 }
 
+#if UE_WITH_IRIS
+void UObject::RegisterReplicationFragments(UE::Net::FFragmentRegistrationContext& Context, UE::Net::EFragmentRegistrationFlags RegistrationFlags)
+{
+}
+#endif // UE_WITH_IRIS
+
 /** Called right before receiving a bunch */
 void UObject::PreNetReceive()
 {
diff --git a/Engine/Source/Runtime/CoreUObject/Public/UObject/CoreNet.h b/Engine/Source/Runtime/CoreUObject/Public/UObject/CoreNet.h
--- a/Engine/Source/Runtime/CoreUObject/Public/UObject/CoreNet.h
+++ b/Engine/Source/Runtime/CoreUObject/Public/UObject/CoreNet.h
@@ -23,12 +23,15 @@ class INetDeltaBaseState;
 class FNetTraceCollector;
 class UPackageMap;
 
-namespace UE
+namespace UE::Net
 {
-	namespace Net
-	{
-		struct FNetResult;
-	}
+	struct FNetResult;
+
+#if UE_WITH_IRIS
+	class FReplicationFragment;
+	struct FReplicationStateDescriptor;
+	typedef FReplicationFragment* (*CreateAndRegisterReplicationFragmentFunc)(UObject* Owner, const FReplicationStateDescriptor* Descriptor, FFragmentRegistrationContext& Context);
+#endif
 }
 
 
@@ -259,6 +262,9 @@ public:
 	ELifetimeCondition Condition;
 	ELifetimeRepNotifyCondition RepNotifyCondition;
 	bool bIsPushBased;
+#if UE_WITH_IRIS
+	UE::Net::CreateAndRegisterReplicationFragmentFunc CreateAndRegisterReplicationFragmentFunction = nullptr;
+#endif
 
 	FLifetimeProperty()
 		: RepIndex(0)
@@ -294,6 +300,9 @@ public:
 			check(Condition == Other.Condition);
 			check(RepNotifyCondition == Other.RepNotifyCondition);
 			check(bIsPushBased == Other.bIsPushBased);
+#if UE_WITH_IRIS
+			check(CreateAndRegisterReplicationFragmentFunction == Other.CreateAndRegisterReplicationFragmentFunction);
+#endif
 			return true;
 		}
 
@@ -605,6 +614,9 @@ struct FNetDeltaSerializeInfo
 	//~ TODO: This feels hacky, and a better alternative might be something like connection specific
 	//~ capabilities.
 
+	/** Whether we are currently initializing base from defaults in which case we should not modify the source **/
+	bool bIsInitializingBaseFromDefault = false;
+
 	/** Whether or not we support FFastArraySerializer::FastArrayDeltaSerialize_DeltaSerializeStructs */
 	bool bSupportsFastArrayDeltaStructSerialization = false;
 
diff --git a/Engine/Source/Runtime/CoreUObject/Public/UObject/Object.h b/Engine/Source/Runtime/CoreUObject/Public/UObject/Object.h
--- a/Engine/Source/Runtime/CoreUObject/Public/UObject/Object.h
+++ b/Engine/Source/Runtime/CoreUObject/Public/UObject/Object.h
@@ -30,6 +30,17 @@ struct FFrame;
 struct FObjectInstancingGraph;
 struct FPropertyChangedChainEvent;
 class UClass;
+#if UE_WITH_IRIS
+namespace UE::Net
+{
+	class FFragmentRegistrationContext;
+	enum class EFragmentRegistrationFlags : uint32;
+	namespace Private
+	{
+		struct FNetHandleLegacyPushModelHelper;
+	}
+}
+#endif // UE_WITH_IRIS
 
 DECLARE_LOG_CATEGORY_EXTERN(LogObj, Log, All);
 
@@ -936,6 +947,16 @@ public:
 
 	/** Called when this object begins replicating to initialize the state of custom property conditions */
 	virtual void GetReplicatedCustomConditionState(FCustomPropertyConditionState& OutActiveState) const;
+#if UE_WITH_IRIS
+	/**
+	 * RegisterReplicationFragments is called when we an object is added to the ReplicationSystem, it allows an object to register new or existing ReplicationFragments describing data to be replicated and how it should be accessed
+	 * For more information about ReplicationFragments see ReplicationFragment.h
+	 * 
+	 * @param Context Context FFragmentRegistrationContext in which FReplicationFragments could be registered
+	 * @param RegistrationFlags Flags specifying what should be registered in the call
+	 */
+	virtual void RegisterReplicationFragments(UE::Net::FFragmentRegistrationContext& Context, UE::Net::EFragmentRegistrationFlags RegistrationFlags);
+#endif // UE_WITH_IRIS
 
 	/** IsNameStableForNetworking means an object can be referred to its path name (relative to outer) over the network */
 	virtual bool IsNameStableForNetworking() const;
@@ -1659,7 +1680,7 @@ public:
 private:
 	
 	friend struct FObjectNetPushIdHelper;
-	virtual void SetNetPushIdDynamic(const int32 NewNetPushId)
+	virtual void SetNetPushIdDynamic(const uint64 NewNetPushId)
 	{
 		// This method should only be called on Objects that are networked, and those should
 		// always have this implemented (by UHT).
@@ -1669,9 +1690,9 @@ private:
 public:
 
 	/** Should only ever be used by internal systems. */
-	virtual int32 GetNetPushIdDynamic() const
+	virtual uint64 GetNetPushIdDynamic() const
 	{
-		return INDEX_NONE;
+		return uint64(int64(INDEX_NONE));
 	}
 };
 
@@ -1680,8 +1701,11 @@ struct FObjectNetPushIdHelper
 private:
 	friend struct FNetPrivatePushIdHelper;
 	friend struct FNetObjectManagerPushIdHelper;
+#if UE_WITH_IRIS
+	friend struct UE::Net::Private::FNetHandleLegacyPushModelHelper;
+#endif // UE_WITH_IRIS
 
-	static void SetNetPushIdDynamic(UObject* Object, const int32 NewNetPushId)
+	static void SetNetPushIdDynamic(UObject* Object, const uint64 NewNetPushId)
 	{
 		Object->SetNetPushIdDynamic(NewNetPushId);
 	}
diff --git a/Engine/Source/Runtime/Engine/Classes/Components/ActorComponent.h b/Engine/Source/Runtime/Engine/Classes/Components/ActorComponent.h
--- a/Engine/Source/Runtime/Engine/Classes/Components/ActorComponent.h
+++ b/Engine/Source/Runtime/Engine/Classes/Components/ActorComponent.h
@@ -575,6 +575,17 @@ public:
 	/** Allows components to wait on outstanding tasks prior to sending EOF update data. Executed on Game Thread and may await tasks. */
 	virtual void OnPreEndOfFrameSync() {}
 
+#if UE_WITH_IRIS
+	/** Register all replication fragments */
+	virtual void RegisterReplicationFragments(UE::Net::FFragmentRegistrationContext& Context, UE::Net::EFragmentRegistrationFlags RegistrationFlags) override;
+	
+	/** Called when we want to start replicating this component */
+	virtual void BeginReplication();
+
+	/** Tell component to end replication */
+	virtual void EndReplication();
+#endif // UE_WITH_IRIS
+
 private:
 	/** Cached pointer to owning actor */
 	mutable AActor* OwnerPrivate;
diff --git a/Engine/Source/Runtime/Engine/Classes/Components/ChildActorComponent.h b/Engine/Source/Runtime/Engine/Classes/Components/ChildActorComponent.h
--- a/Engine/Source/Runtime/Engine/Classes/Components/ChildActorComponent.h
+++ b/Engine/Source/Runtime/Engine/Classes/Components/ChildActorComponent.h
@@ -196,6 +196,12 @@ public:
 	void SetEditorTreeViewVisualizationMode(EChildActorComponentTreeViewVisualizationMode InMode);
 #endif
 
+#if UE_WITH_IRIS
+	/** Register all replication fragments */
+	virtual void RegisterReplicationFragments(UE::Net::FFragmentRegistrationContext& Context, UE::Net::EFragmentRegistrationFlags RegistrationFlags) override;
+#endif // UE_WITH_IRIS
+
+
 private:
 	bool IsChildActorReplicated() const;
 };
diff --git a/Engine/Source/Runtime/Engine/Classes/Engine/Engine.h b/Engine/Source/Runtime/Engine/Classes/Engine/Engine.h
--- a/Engine/Source/Runtime/Engine/Classes/Engine/Engine.h
+++ b/Engine/Source/Runtime/Engine/Classes/Engine/Engine.h
@@ -237,6 +237,43 @@ struct FNetDriverDefinition
 	}
 };
 
+/**
+* Struct used to configure which NetDriver is started with Iris enabled or not
+* Only one attribute out of the NetDriverDefinition, NetDriverName or NetDriverWildcardName should be set along with the bEnableIris property
+*/
+USTRUCT()
+struct FIrisNetDriverConfig 
+{
+	GENERATED_BODY()
+
+	/**
+	 * Name of the net driver definition to configure
+	 * e.g. GameNetDriver, BeaconNetDriver, etc.
+	 */
+	UPROPERTY()
+	FName NetDriverDefinition;
+
+	/**
+	 * Name of the named driver to configure.
+	 * e.g. GameNetDriver, DemoNetDriver, etc.
+	 */
+	UPROPERTY()
+	FName NetDriverName;
+
+	/**
+	 * Wildcard match the netdriver name to configure
+	 * e.g. NetDriverWildcardName="UnitTestNetDriver*" matches with UnitTestNetDriver_1, UnitTestNetDriver_2, etc.
+	 */
+	UPROPERTY()
+	FString NetDriverWildcardName;
+
+	/**
+	 * Configurable property that decides if the NetDriver will use the Iris replication system or not
+	 */
+	UPROPERTY()
+	bool bEnableIris = false;
+};
+
 
 /**
  * Active and named net drivers instantiated from an FNetDriverDefinition
@@ -2937,6 +2974,10 @@ public:
 	UPROPERTY(Config, transient)
 	TArray<FNetDriverDefinition> NetDriverDefinitions;
 
+	/** A list of Iris NetDriverConfigs */
+	UPROPERTY(Config, transient)
+	TArray<FIrisNetDriverConfig> IrisNetDriverConfigs;
+	
 	/** A configurable list of actors that are automatically spawned upon server startup (just prior to InitGame) */
 	UPROPERTY(config)
 	TArray<FString> ServerActors;
diff --git a/Engine/Source/Runtime/Engine/Classes/Engine/NetConnection.h b/Engine/Source/Runtime/Engine/Classes/Engine/NetConnection.h
--- a/Engine/Source/Runtime/Engine/Classes/Engine/NetConnection.h
+++ b/Engine/Source/Runtime/Engine/Classes/Engine/NetConnection.h
@@ -712,6 +712,12 @@ public:
 
 	void AddDestructionInfo(FActorDestructionInfo* DestructionInfo)
 	{
+#if UE_WITH_IRIS
+		if (Driver && Driver->GetReplicationSystem())
+		{
+			return;
+		}
+#endif // UE_WITH_IRIS
 		if (ReplicationConnectionDriver)
 		{
 			ReplicationConnectionDriver->NotifyAddDestructionInfo(DestructionInfo);
@@ -724,6 +730,12 @@ public:
 
 	void RemoveDestructionInfo(FActorDestructionInfo* DestructionInfo)
 	{
+#if UE_WITH_IRIS
+		if (Driver && Driver->GetReplicationSystem())
+		{
+			return;
+		}
+#endif // UE_WITH_IRIS
 		if (ReplicationConnectionDriver)
 		{
 			ReplicationConnectionDriver->NotifyRemoveDestructionInfo(DestructionInfo);
@@ -736,6 +748,12 @@ public:
 	
 	void ResetDestructionInfos()
 	{
+#if UE_WITH_IRIS
+		if (Driver && Driver->GetReplicationSystem())
+		{
+			return;
+		}
+#endif // UE_WITH_IRIS
 		if (ReplicationConnectionDriver)
 		{
 			ReplicationConnectionDriver->NotifyResetDestructionInfo();
diff --git a/Engine/Source/Runtime/Engine/Classes/Engine/NetDriver.h b/Engine/Source/Runtime/Engine/Classes/Engine/NetDriver.h
--- a/Engine/Source/Runtime/Engine/Classes/Engine/NetDriver.h
+++ b/Engine/Source/Runtime/Engine/Classes/Engine/NetDriver.h
@@ -19,6 +19,9 @@
 #include "Net/NetAnalyticsTypes.h"
 #include "Net/NetConnectionIdHandler.h"
 #include "HAL/LowLevelMemTracker.h"
+#if UE_WITH_IRIS
+#include "Templates/PimplPtr.h"
+#endif
 #include "NetDriver.generated.h"
 
 /**
@@ -335,7 +338,17 @@
 class UNetDriver;
 
 enum class ECreateReplicationChangelistMgrFlags;
+enum class EEngineNetworkRuntimeFeatures : uint16;
+#if UE_WITH_IRIS
+class UReplicationSystem;
+class UReplicationBridge;
+namespace UE::Net
+{
+	typedef uint16 FNetObjectGroupHandle;
+}
+#endif // UE_WITH_IRIS
 
+
 using FConnectionMap = TMap<TSharedRef<const FInternetAddr>, UNetConnection*, FDefaultSetAllocator, FInternetAddrConstKeyMapFuncs<UNetConnection*>>;
 
 extern ENGINE_API TAutoConsoleVariable<int32> CVarNetAllowEncryption;
@@ -712,6 +725,10 @@
 	UPROPERTY(Config)
 	FString ReplicationDriverClassName;
 
+	/** Used to specify the class to use for ReplicationBridge */
+	UPROPERTY(Config)
+	FString ReplicationBridgeClassName;
+
 	/** @todo document */
 	UPROPERTY(Config)
 	int32 MaxDownloadSize;
@@ -825,6 +842,9 @@
 	UPROPERTY()
 	TObjectPtr<UClass> ReplicationDriverClass;
 
+	UPROPERTY(transient)
+	TObjectPtr<UClass> ReplicationBridgeClass;
+
 	/** @todo document */
 	FProperty* RoleProperty;
 	
@@ -883,6 +903,9 @@
 	/** Set the NetDriver's NetDriverDefintion. */
 	void SetNetDriverDefinition(FName NewNetDriverDefinition);
 
+	/** Callback after the engine created the NetDriver and set our name for the first time */
+	void PostCreation(bool bInitializeWithIris);
+
 	/** Get the NetDriver's NetDriverDefintion. */
 	FName GetNetDriverDefintion() const { return NetDriverDefinition; }
 
@@ -1281,13 +1304,16 @@
 	 */
 	ENGINE_API virtual void FlushHandler();
 
-
 	/** Initializes the net connection class to use for new connections */
 	ENGINE_API virtual bool InitConnectionClass(void);
 
 	/** Initialized the replication driver class to use for this driver */
 	ENGINE_API virtual bool InitReplicationDriverClass();
 
+#if UE_WITH_IRIS
+	/** Initialized the replication bridge class to use for this driver if using iris replication*/
+	ENGINE_API virtual bool InitReplicationBridgeClass();
+#endif
 	/** Shutdown all connections managed by this net driver */
 	ENGINE_API virtual void Shutdown();
 
@@ -1636,9 +1662,37 @@
 
 	ENGINE_API UReplicationDriver* GetReplicationDriver() const { return ReplicationDriver; }
 
+	/** Returns if this netdriver is initialized to replicate using the Iris replication system or the Legacy replication system. */
+	FORCEINLINE bool IsUsingIrisReplication() const
+	{
+#if UE_WITH_IRIS
+		return bIsUsingIris;
+#else
+		return false;
+#endif //UE_WITH_IRIS
+	}
+
+	/** Returns the bitflag telling which network features are activated for this NetDriver. */
+	ENGINE_API EEngineNetworkRuntimeFeatures GetNetworkRuntimeFeatures() const;
+	
+#if UE_WITH_IRIS
+	/** Remove references to the Iris bridge and system without deleting it */
+	ENGINE_API void ClearIrisSystem();
+
+	/** Set a previously initialized IrisSystem into this NetDriver */
+	ENGINE_API void RestoreIrisSystem(UReplicationSystem* InReplicationSystem);
+#endif // UE_WITH_IRIS
+
 	template<class T>
 	T* GetReplicationDriver() const { return Cast<T>(ReplicationDriver); }
 
+#if UE_WITH_IRIS
+	inline UReplicationSystem* GetReplicationSystem() { return ReplicationSystem; }
+	inline UReplicationSystem* GetReplicationSystem() const { return ReplicationSystem; }
+
+	void UpdateGroupFilterStatusForLevel(const ULevel* Level, UE::Net::FNetObjectGroupHandle LevelGroupHandle);
+#endif // UE_WITH_IRIS
+
 	void RemoveClientConnection(UNetConnection* ClientConnectionToRemove);
 
 	/** Adds (fully initialized, ready to go) client connection to the ClientConnections list + any other game related setup */
@@ -1877,9 +1931,28 @@
 	/** Handle to FNetDelegates::OnSyncLoadDetected delegate */
 	FDelegateHandle ReportSyncLoadDelegateHandle;
 
+#if UE_WITH_IRIS
+	void InitIrisSettings(FName NewDriverName);
+	void SetReplicationSystem(UReplicationSystem* ReplicationSystem);
+	void UpdateReplicationViews() const;
+#endif
+
 	UPROPERTY(transient)
 	TObjectPtr<UReplicationDriver> ReplicationDriver;
 
+#if UE_WITH_IRIS
+	UReplicationSystem* ReplicationSystem = nullptr;
+
+	/** When set this will skip registering all the network relevant actors when setting the World */
+	bool bSkipBeginReplicationForWorld = false;
+
+	/** True when the NetDriver has been configured to run with the Iris replication system.*/
+	bool bIsUsingIris = false;
+
+	// For FindOrAddNetworkObjectInfo
+	TPimplPtr<FNetworkObjectInfo> DummyNetworkObjectInfo;
+#endif
+
 	/** Stores the list of objects to replicate into the replay stream. This should be a TUniquePtr, but it appears the generated.cpp file needs the full definition of the pointed-to type. */
 	TSharedPtr<FNetworkObjectList> NetworkObjects;
 
diff --git a/Engine/Source/Runtime/Engine/Classes/Engine/World.h b/Engine/Source/Runtime/Engine/Classes/Engine/World.h
--- a/Engine/Source/Runtime/Engine/Classes/Engine/World.h
+++ b/Engine/Source/Runtime/Engine/Classes/Engine/World.h
@@ -313,6 +313,23 @@ public:
 	static FOnWorldUnregisteredWithAudioDevice OnWorldUnregisteredWithAudioDevice;
 };
 
+#if UE_WITH_IRIS
+/**
+* Struct that temporarily holds the Iris replication system and bridge.
+*/
+struct FIrisSystemHolder
+{
+	bool IsHolding() const { return ReplicationSystem != nullptr; }
+
+	void Clear()
+	{
+		ReplicationSystem = nullptr;
+	}
+
+	class UReplicationSystem* ReplicationSystem = nullptr;
+};
+#endif // UE_WITH_IRIS
+
 /** class that encapsulates seamless world traveling */
 class FSeamlessTravelHandler
 {
@@ -636,6 +653,10 @@ private:
 
 	friend class UPackageMapClient;
 
+#if UE_WITH_IRIS
+	friend class UActorReplicationBridge;
+#endif // UE_WITH_IRIS
+
 	/* Is the actor remotely owned. This should only be set true by the package map when it is creating an actor on a client that was replicated from the server. */
 	uint8	bRemoteOwned:1;
 	
@@ -1171,6 +1192,10 @@ public:
 
 	/** Returns BlockTillLevelStreamingCompletedEpoch. */
 	int32 GetBlockTillLevelStreamingCompletedEpoch() const { return BlockTillLevelStreamingCompletedEpoch; }
+#if UE_WITH_IRIS
+	/** Store the Iris managers from the GameNetDriver to they can be restored into a different NetDriver later */
+	void StoreIrisAndClearReferences();
+#endif // UE_WITH_IRIS
 
 	/** Prefix we used to rename streaming levels, non empty in PIE and standalone preview */
 	UPROPERTY()
@@ -3713,6 +3738,11 @@ public:
 
 private:
 
+#if UE_WITH_IRIS
+	/** Holds the Iris systems during the NetDriver transition that occurs when Forking */
+	FIrisSystemHolder IrisSystemHolder;
+#endif // UE_WITH_IRIS
+
 	/** Private version without inlining that does *not* check Dedicated server build flags (which should already have been done). */
 	ENetMode InternalGetNetMode() const;
 
diff --git a/Engine/Source/Runtime/Engine/Classes/GameFramework/Actor.h b/Engine/Source/Runtime/Engine/Classes/GameFramework/Actor.h
--- a/Engine/Source/Runtime/Engine/Classes/GameFramework/Actor.h
+++ b/Engine/Source/Runtime/Engine/Classes/GameFramework/Actor.h
@@ -46,6 +46,9 @@ class UDEPRECATED_DataLayer;
 class UDataLayerAsset;
 class UDataLayerInstance;
 class AWorldDataLayers;
+#if UE_WITH_IRIS
+struct FActorBeginReplicationParams;
+#endif // UE_WITH_IRIS
 class UActorFolder;
 
 // By default, debug and development builds (even cooked) will keep actor labels. Manually define this if you want to make a local build
@@ -3083,6 +3086,37 @@ public:
 	/** Destroys the constructed components. */
 	void DestroyConstructedComponents();
 
+#if UE_WITH_IRIS
+	virtual void RegisterReplicationFragments(UE::Net::FFragmentRegistrationContext& Context, UE::Net::EFragmentRegistrationFlags RegistrationFlags) override;
+
+	/**
+	 * Called for all Actors set to replicate during BeginPlay, It will also be called if SetReplicates(true) is called and the object is not already replicating
+	 */
+	virtual void BeginReplication();
+
+	/**
+	 * Called when we want to end replication for this actor, typically called from EndPlay() for actors that should be replicated during their lifespan
+	 */
+	virtual void EndReplication(EEndPlayReason::Type EndPlayReason);
+protected:
+	/**
+	 * Pushes the owning NetConnection for the actor and all of its children to the replication system.
+	 * This information decides whether properties with owner conditionals are replicated or not.
+	 */
+	void UpdateOwningNetConnection() const;
+
+	/**
+	 * Helper to BeginReplication passing on additional parameters to the ReplicationSystem, typically called from code overriding normal BeginReplication()
+	 * @param Params Additional parameters we want to pass on
+	 */
+	void BeginReplication(const FActorBeginReplicationParams& Params);
+#endif // UE_WITH_IRIS
+
+	/**
+	 * Updates the ReplicatePhysics condition. That information needs to be pushed to the ReplicationSystem.
+	 */
+	void UpdateReplicatePhysicsCondition();
+	
 protected:
 	/**
 	 * Virtual call chain to register all tick functions for the actor class hierarchy
diff --git a/Engine/Source/Runtime/Engine/Classes/GameFramework/CharacterMovementReplication.h b/Engine/Source/Runtime/Engine/Classes/GameFramework/CharacterMovementReplication.h
--- a/Engine/Source/Runtime/Engine/Classes/GameFramework/CharacterMovementReplication.h
+++ b/Engine/Source/Runtime/Engine/Classes/GameFramework/CharacterMovementReplication.h
@@ -7,6 +7,10 @@
 #include "Engine/NetSerialization.h"
 #include "Serialization/BitWriter.h"
 #include "Containers/BitArray.h"
+#if UE_WITH_IRIS
+#include "Iris/Serialization/IrisObjectReferencePackageMap.h"
+#endif
+
 #include "CharacterMovementReplication.generated.h"
 
 class UPackageMap;
@@ -54,6 +58,11 @@ struct ENGINE_API FCharacterNetworkSerializationPackedBits
 	// TInlineAllocator used with TBitArray takes the number of 32-bit dwords, but the define is in number of bits, so convert here by dividing by 32.
 	TBitArray<TInlineAllocator<CHARACTER_SERIALIZATION_PACKEDBITS_RESERVED_SIZE / NumBitsPerDWORD>> DataBits;
 
+#if UE_WITH_IRIS
+	// Since this struct uses custom serialization path we need to explicitly capture object references, this is managed by the use of a custom packagemap.
+	UIrisObjectReferencePackageMap::FObjectReferenceArray ObjectReferences;
+#endif
+
 private:
 	UPackageMap* SavedPackageMap;
 };
diff --git a/Engine/Source/Runtime/Engine/Classes/GameFramework/PlayerController.h b/Engine/Source/Runtime/Engine/Classes/GameFramework/PlayerController.h
--- a/Engine/Source/Runtime/Engine/Classes/GameFramework/PlayerController.h
+++ b/Engine/Source/Runtime/Engine/Classes/GameFramework/PlayerController.h
@@ -342,6 +342,10 @@ public:
 	TArray<FForceFeedbackEffectHistoryEntry> ForceFeedbackEffectHistoryEntries;
 #endif
 
+protected:
+#if UE_WITH_IRIS
+	virtual void BeginReplication() override;
+#endif // UE_WITH_IRIS
 	/** The type of async physics data object to use*/
 	UPROPERTY(EditDefaultsOnly, Category=PlayerController)
 	TSubclassOf<UAsyncPhysicsData> AsyncPhysicsDataClass;
@@ -1607,10 +1611,7 @@ public:
 	void IncludeInNetConditionGroup(FName NetGroup);
 
 	/** Remove this player from a netcondition group. */
-	void RemoveFromNetConditionGroup(FName NetGroup)
-	{
-		NetConditionGroups.RemoveSingleSwap(NetGroup);
-	}
+	void RemoveFromNetConditionGroup(FName NetGroup);
 
 	/** Returns true if the player controller is a member of the netcondition group */
 	bool IsMemberOfNetConditionGroup(FName NetGroup) const
diff --git a/Engine/Source/Runtime/Engine/Classes/GameFramework/PlayerState.h b/Engine/Source/Runtime/Engine/Classes/GameFramework/PlayerState.h
--- a/Engine/Source/Runtime/Engine/Classes/GameFramework/PlayerState.h
+++ b/Engine/Source/Runtime/Engine/Classes/GameFramework/PlayerState.h
@@ -338,6 +338,10 @@ protected:
 	/** called after receiving player name */
 	virtual void HandleWelcomeMessage();
 
+#if UE_WITH_IRIS
+	virtual void BeginReplication() override;
+#endif
+
 private:
 	// Hidden functions that don't make sense to use on this class.
 	HIDE_ACTOR_TRANSFORM_FUNCTIONS();
diff --git a/Engine/Source/Runtime/Engine/Engine.Build.cs b/Engine/Source/Runtime/Engine/Engine.Build.cs
--- a/Engine/Source/Runtime/Engine/Engine.Build.cs
+++ b/Engine/Source/Runtime/Engine/Engine.Build.cs
@@ -451,6 +451,9 @@ public class Engine : ModuleRules
 			PublicDefinitions.Add("WITH_ODSC=0");
 		}
 
+		const bool bIrisAddAsPublicDepedency = true;
+		SetupIrisSupport(Target, bIrisAddAsPublicDepedency);
+
 		PrivateDefinitions.Add("UE_DEPRECATE_LEGACY_MATH_CONSTANT_MACRO_NAMES=1");
 	}
 }
diff --git a/Engine/Source/Runtime/Engine/Private/Actor.cpp b/Engine/Source/Runtime/Engine/Private/Actor.cpp
--- a/Engine/Source/Runtime/Engine/Private/Actor.cpp
+++ b/Engine/Source/Runtime/Engine/Private/Actor.cpp
@@ -46,6 +46,11 @@
 #include "ObjectTrace.h"
 #include "Net/Core/PushModel/PushModel.h"
 #include "Engine/AutoDestroySubsystem.h"
+#if UE_WITH_IRIS
+#include "Iris/ReplicationSystem/ReplicationSystem.h"
+#include "Iris/ReplicationSystem/Conditionals/ReplicationCondition.h"
+#include "Net/Iris/ReplicationSystem/ReplicationSystemUtil.h"
+#endif // UE_WITH_IRIS
 #include "LevelUtils.h"
 #include "Components/DecalComponent.h"
 #include "GameFramework/InputSettings.h"
@@ -2055,6 +2060,10 @@ void AActor::SetOwner(AActor* NewOwner)
 		{
 			MarkOwnerRelevantComponentsDirty(this);
 		}
+
+#if UE_WITH_IRIS
+		UpdateOwningNetConnection();
+#endif // UE_WITH_IRIS
 	}
 }
 
@@ -2634,6 +2643,9 @@ void AActor::RouteEndPlay(const EEndPlayReason::Type EndPlayReason)
 			if (UWorld* World = GetWorld())
 			{
 				World->RemoveNetworkActor(this);
+#if UE_WITH_IRIS
+				EndReplication(EndPlayReason);
+#endif // UE_WITH_IRIS
 			}
 		}
 
@@ -2655,6 +2667,10 @@ void AActor::EndPlay(const EEndPlayReason::Type EndPlayReason)
 
 		ActorHasBegunPlay = EActorBeginPlayState::HasNotBegunPlay;
 
+#if UE_WITH_IRIS
+		EndReplication(EndPlayReason);
+#endif // UE_WITH_IRIS
+
 		// Dispatch the blueprint events
 		ReceiveEndPlay(EndPlayReason);
 		OnEndPlay.Broadcast(this, EndPlayReason);
@@ -3242,6 +3258,8 @@ bool AActor::OwnsComponent(UActorComponent* Component) const
 void AActor::UpdateReplicatedComponent(UActorComponent* Component)
 {
 	using namespace UE::Net;
+	
+	bool bWasRemovedFromReplicatedComponents = false;
 
 	if (Component->GetIsReplicated())
 	{
@@ -3249,9 +3267,10 @@ void AActor::UpdateReplicatedComponent(UActorComponent* Component)
 	}
 	else
 	{
-		ReplicatedComponents.RemoveSingleSwap(Component);
+		bWasRemovedFromReplicatedComponents = ReplicatedComponents.RemoveSingleSwap(Component) > 0;
 	}
 
+
 	const ELifetimeCondition NetCondition = Component->GetIsReplicated() ? AllowActorComponentToReplicate(Component) : COND_Never;
 
 	if (FReplicatedComponentInfo* RepComponentInfo = ReplicatedComponentsInfo.FindByKey(Component))
@@ -3263,6 +3282,18 @@ void AActor::UpdateReplicatedComponent(UActorComponent* Component)
 	{
 		ReplicatedComponentsInfo.Emplace(FReplicatedComponentInfo(Component, NetCondition));
 	}
+
+#if UE_WITH_IRIS
+	if (NetCondition != COND_Never && HasActorBegunPlay())
+	{
+		// Begin replication and set NetCondition, if component already is replicated the NetCondition will be updated
+		Component->BeginReplication();
+	}
+	else if (bWasRemovedFromReplicatedComponents || NetCondition == COND_Never)
+	{
+		Component->EndReplication();
+	}
+#endif // UE_WITH_IRIS
 }
 
 void AActor::UpdateAllReplicatedComponents()
@@ -3275,6 +3306,13 @@ void AActor::UpdateAllReplicatedComponents()
 	for (FReplicatedComponentInfo& RepComponentInfo : ReplicatedComponentsInfo)
 	{
 		RepComponentInfo.NetCondition = COND_Never;
+#if UE_WITH_IRIS
+		// Need to end replication for all components that should no longer replicate
+		if (RepComponentInfo.Component && !RepComponentInfo.Component->GetIsReplicated())
+		{
+			RepComponentInfo.Component->EndReplication();
+		}
+#endif
 	}
 
 	for (UActorComponent* Component : OwnedComponents)
@@ -3291,6 +3329,14 @@ void AActor::UpdateAllReplicatedComponents()
 				{
 					const int32 Index = ReplicatedComponentsInfo.AddUnique(FReplicatedComponentInfo(Component));
 					ReplicatedComponentsInfo[Index].NetCondition = NetCondition;
+#if UE_WITH_IRIS
+					// Begin replication and set NetCondition, if component already is replicated the NetCondition will be updated
+					Component->BeginReplication();
+				}
+				else
+				{
+					Component->EndReplication();
+#endif
 				}
 			}
 		}
@@ -3835,6 +3881,20 @@ void AActor::SetReplicates(bool bInReplicates)
 						MyWorld->AddNetworkActor(this);
 					}
 
+#if UE_WITH_IRIS
+					if (HasActorBegunPlay())
+					{
+						if (bNewlyReplicates)
+						{
+							BeginReplication();
+						}
+						else
+						{
+							UpdateOwningNetConnection();
+						}
+					}
+#endif // UE_WITH_IRIS
+
 					ForcePropertyCompare();
 				}
 			}
@@ -3873,6 +3933,20 @@ void AActor::SetAutonomousProxy(const bool bInAutonomousProxy, const bool bAllow
 			// We have to do this so the role change above will replicate (turn off shadow state sharing for a frame)
 			// This is because RemoteRole is special since it will change between connections, so we have to special case
 			ForcePropertyCompare();
+
+#if UE_WITH_IRIS
+			// Update autonomous role condition
+			if (UReplicationSystem* ReplicationSystem = UE::Net::FReplicationSystemUtil::GetReplicationSystem(GetNetOwner()))
+			{
+				const UE::Net::FNetHandle NetHandle = UE::Net::FReplicationSystemUtil::GetNetHandle(this);
+				if (NetHandle.IsValid())
+				{
+					const uint32 OwningNetConnectionId = ReplicationSystem->GetOwningNetConnection(NetHandle);
+					const bool bEnableAutonomousCondition = RemoteRole == ROLE_AutonomousProxy;
+					ReplicationSystem->SetReplicationConditionConnectionFilter(NetHandle, UE::Net::EReplicationCondition::RoleAutonomous, OwningNetConnectionId, bEnableAutonomousCondition);
+				}
+			}
+#endif // UE_WITH_IRIS
 		}
 	}
 	else
@@ -3969,6 +4043,10 @@ void AActor::DispatchBeginPlay(bool bFromLevelStreaming)
 			}
 		}
 
+#if UE_WITH_IRIS
+		BeginReplication();
+#endif // UE_WITH_IRIS
+
 		BeginPlay();
 
 		ensure(BeginPlayCallDepth - 1 == CurrentCallDepth);
@@ -5828,7 +5906,14 @@ FRepMovement& AActor::GetReplicatedMovement_Mutable()
 
 void AActor::SetReplicatedMovement(const FRepMovement& InReplicatedMovement)
 {
+	const bool bRepPhysicsDiffers = (GetReplicatedMovement().bRepPhysics != InReplicatedMovement.bRepPhysics);
+	
 	GetReplicatedMovement_Mutable() = InReplicatedMovement;
+
+	if (bRepPhysicsDiffers)
+	{
+		UpdateReplicatePhysicsCondition();
+	}
 }
 
 void AActor::SetInstigator(APawn* InInstigator)
diff --git a/Engine/Source/Runtime/Engine/Private/ActorReplication.cpp b/Engine/Source/Runtime/Engine/Private/ActorReplication.cpp
--- a/Engine/Source/Runtime/Engine/Private/ActorReplication.cpp
+++ b/Engine/Source/Runtime/Engine/Private/ActorReplication.cpp
@@ -20,6 +20,15 @@
 #include "DrawDebugHelpers.h"
 #include "Net/Core/PushModel/PushModel.h"
 #include "Interfaces/Interface_ActorSubobject.h"
+#if UE_WITH_IRIS
+#include "Engine/NetConnection.h"
+#include "Iris/IrisConfig.h"
+#include "Iris/ReplicationSystem/ReplicationFragmentUtil.h"
+#include "Iris/ReplicationSystem/ReplicationSystem.h"
+#include "Iris/ReplicationSystem/Conditionals/ReplicationCondition.h"
+#include "Net/Iris/ReplicationSystem/ReplicationSystemUtil.h"
+#include "Net/Iris/ReplicationSystem/ActorReplicationBridge.h"
+#endif // UE_WITH_IRIS
 #include "Physics/Experimental/PhysScene_Chaos.h"
 #include "PBDRigidsSolver.h"
 #include "ChaosSolversModule.h"
@@ -360,6 +369,10 @@ void AActor::GatherCurrentMovement()
 		UPrimitiveComponent* RootPrimComp = Cast<UPrimitiveComponent>(GetRootComponent());
 		if (RootPrimComp && RootPrimComp->IsSimulatingPhysics())
 		{
+#if UE_WITH_IRIS
+			const bool bPrevRepPhysics = ReplicatedMovement.bRepPhysics;
+#endif // UE_WITH_IRIS
+
 			bool bFoundInCache = false;
 
 			UWorld* World = GetWorld();
@@ -386,6 +399,14 @@ void AActor::GatherCurrentMovement()
 			
 			// Technically, the values might have stayed the same, but we'll just assume they've changed.
 			bWasRepMovementModified = true;
+
+#if UE_WITH_IRIS
+			// If RepPhysics has changed value then notify the ReplicationSystem
+			if (bPrevRepPhysics != ReplicatedMovement.bRepPhysics)
+			{
+				UpdateReplicatePhysicsCondition();
+			}
+#endif // UE_WITH_IRIS
 		}
 		else if (RootComponent != nullptr)
 		{
@@ -505,7 +526,13 @@ void AActor::SetReplicatedComponentNetCondition(const UActorComponent* Replicate
 
 	if (FReplicatedComponentInfo* ComponentInfo = ReplicatedComponentsInfo.FindByKey(ReplicatedComponent))
 	{
-		ComponentInfo->NetCondition = NetCondition;
+		if (NetCondition != ComponentInfo->NetCondition)
+		{
+#if UE_WITH_IRIS
+			UE::Net::FReplicationSystemUtil::SetActorComponentNetCondition(ReplicatedComponent, NetCondition);
+#endif
+			ComponentInfo->NetCondition = NetCondition;
+		}
 	}
 	else
 	{
@@ -535,6 +562,11 @@ void AActor::AddComponentForReplication(UActorComponent* Component)
 			// Always set the condition because this component could be registering SubObjects ahead of time and already in the list with COND_Never
 			ComponentInfo->NetCondition = NetCondition;
 		}
+
+#if UE_WITH_IRIS
+		Component->BeginReplication();
+#endif
+
 	}
 }
 
@@ -544,6 +576,9 @@ void AActor::RemoveReplicatedComponent(UActorComponent* Component)
 	if (Index != INDEX_NONE)
 	{
 		ReplicatedComponentsInfo.RemoveAtSwap(Index);
+#if UE_WITH_IRIS
+		Component->EndReplication();
+#endif
 	}
 }
 
@@ -561,14 +596,28 @@ void AActor::AddReplicatedSubObject(UObject* SubObject, ELifetimeCondition NetCo
 	// Warn if the subobject was registered with a different net condition.
 	ensureMsgf(Result != FSubObjectRegistry::EResult::NetConditionConflict, TEXT("%s(0x%p) Registered subobject %s (0x%p) again with a different net condition. Active [%s] New [%s]. New condition will be ignored"),
 		*GetName(), this, *SubObject->GetName(), SubObject, *StaticEnum<ELifetimeCondition>()->GetValueAsString(ReplicatedSubObjects.GetNetCondition(SubObject)), *StaticEnum<ELifetimeCondition>()->GetValueAsString(NetCondition));
+
+#if UE_WITH_IRIS
+	if (GetIsReplicated() && HasActorBegunPlay())
+	{
+		UE::Net::FReplicationSystemUtil::BeginReplicationForActorSubObject(this, SubObject, NetCondition);
+	}
+#endif
 }
 
 void AActor::RemoveReplicatedSubObject(UObject* SubObject)
 {
 	check(SubObject);
-	bool bWasRemoved = ReplicatedSubObjects.RemoveSubObject(SubObject);
+	const bool bWasRemoved = ReplicatedSubObjects.RemoveSubObject(SubObject);
 
 	UE_CLOG(bWasRemoved, LogNetSubObject, Verbose, TEXT("%s (0x%p) removed replicated subobject %s (0x%p)"), *GetName(), this, *SubObject->GetName(), SubObject);
+
+#if UE_WITH_IRIS
+	if (bWasRemoved)
+	{
+		UE::Net::FReplicationSystemUtil::EndReplicationForActorSubObject(this, SubObject);
+	}
+#endif // UE_WITH_IRIS
 }
 
 void AActor::AddActorComponentReplicatedSubObject(UActorComponent* OwnerComponent, UObject* SubObject, ELifetimeCondition NetCondition)
@@ -589,6 +638,13 @@ void AActor::AddActorComponentReplicatedSubObject(UActorComponent* OwnerComponen
 		// Warn if the subobject was registered with a different net condition.
 		ensureMsgf(Result != FSubObjectRegistry::EResult::NetConditionConflict, TEXT("%s::%s (0x%p) Registered subobject %s (0x%p) again with a different net condition. Active [%s] New [%s]. New condition will be ignored"),
 			*GetName(), *OwnerComponent->GetName(), OwnerComponent, *SubObject->GetName(), SubObject, *StaticEnum<ELifetimeCondition>()->GetValueAsString(ComponentInfo->SubObjects.GetNetCondition(SubObject)), *StaticEnum<ELifetimeCondition>()->GetValueAsString(NetCondition));
+
+#if UE_WITH_IRIS
+		if (OwnerComponent->GetIsReplicated() && OwnerComponent->HasBegunPlay())
+		{
+			UE::Net::FReplicationSystemUtil::BeginReplicationForActorComponentSubObject(OwnerComponent, SubObject, NetCondition);
+		}
+#endif
 	}
 	else
 	{
@@ -614,6 +670,10 @@ void AActor::RemoveActorComponentReplicatedSubObject(UActorComponent* OwnerCompo
 		bool bWasRemoved = ComponentInfo->SubObjects.RemoveSubObject(SubObject);
 		
 		UE_CLOG(bWasRemoved, LogNetSubObject, Verbose, TEXT("%s::%s (0x%p) removed replicated subobject %s (0x%p)"), *GetName(), *OwnerComponent->GetName(), OwnerComponent, *SubObject->GetName(), SubObject);
+
+#if UE_WITH_IRIS
+		UE::Net::FReplicationSystemUtil::EndReplicationForActorComponentSubObject(OwnerComponent, SubObject);
+#endif // UE_WITH_IRIS
 	}
 }
 
@@ -689,3 +749,141 @@ void AActor::OnRep_Owner()
 {
 
 }
+
+#if UE_WITH_IRIS
+void AActor::RegisterReplicationFragments(UE::Net::FFragmentRegistrationContext& Context, UE::Net::EFragmentRegistrationFlags RegistrationFlags)
+{
+	using namespace UE::Net;
+
+	Super::RegisterReplicationFragments(Context, RegistrationFlags);
+
+	// Build descriptors and allocate PropertyReplicationFragments for this object
+	FReplicationFragmentUtil::CreateAndRegisterFragmentsForObject(this, Context, RegistrationFlags);
+
+	// $IRIS: $TODO: Consider splitting instance protocol / protocol creation, potentially protocol could be created from archetype/CDO or built as part of cook process.
+	for (UActorComponent* ActorComp : OwnedComponents)
+	{
+		if (!IsValid(ActorComp))
+		{
+			continue;
+		}
+
+		if (ActorComp->HasAnyFlags(RF_DefaultSubObject) || ActorComp->IsDefaultSubobject())
+		{
+			if (!ActorComp->GetIsReplicated())
+			{
+				// Register RPC functions for not replicated default subobjects
+				ActorComp->RegisterReplicationFragments(Context, EFragmentRegistrationFlags::RegisterRPCsOnly);
+			}
+		}
+	}
+}
+
+void AActor::BeginReplication(const FActorBeginReplicationParams& Params)
+{
+	// If we do not have a handle we create one
+	UE::Net::FNetHandle ActorHandle = UE::Net::FReplicationSystemUtil::BeginReplication(this, Params);
+	if (ActorHandle.IsValid())
+	{
+		UpdateOwningNetConnection();
+	}
+}
+
+void AActor::BeginReplication()
+{
+	const FActorBeginReplicationParams BeginReplicationParams;
+
+	BeginReplication(BeginReplicationParams);
+}
+
+void AActor::EndReplication(EEndPlayReason::Type EndPlayReason)
+{
+	if (GetIsReplicated())
+	{
+		UE::Net::FReplicationSystemUtil::EndReplication(this, EndPlayReason);
+	}
+}
+
+void AActor::UpdateOwningNetConnection() const
+{
+	using namespace UE::Net;
+
+	UReplicationSystem* ReplicationSystem = FReplicationSystemUtil::GetReplicationSystem(GetNetOwner());
+	if (ReplicationSystem == nullptr)
+	{
+		return;
+	}
+
+	uint32 NewOwningNetConnectionId = 0U;
+	if (const UNetConnection* NetConnection = GetNetConnection())
+	{
+		NewOwningNetConnectionId = NetConnection->GetParentConnectionId();
+	}
+
+	// If this actor isn't replicated there's no way for us to tell whether we need to update our children.
+	bool bUpdateChildren = !GetIsReplicated();
+	if (!bUpdateChildren)
+	{
+#if DO_CHECK
+		// Sanity check
+		{
+			const UReplicationSystem* MyReplicationSystem = FReplicationSystemUtil::GetReplicationSystem(this);
+			check(MyReplicationSystem == ReplicationSystem || MyReplicationSystem == nullptr);
+		}
+#endif
+
+		FNetHandle NetHandle = FReplicationSystemUtil::GetNetHandle(this);
+		const uint32 CurrentOwningNetConnectionId = NetHandle.IsValid() ? ReplicationSystem->GetOwningNetConnection(NetHandle) : 0U;
+		bUpdateChildren = (NewOwningNetConnectionId != CurrentOwningNetConnectionId);
+	}
+
+	if (!bUpdateChildren)
+	{
+		return;
+	}
+
+	constexpr SIZE_T MaxActorCount = 512;
+	const AActor* Actors[MaxActorCount];
+	SIZE_T ActorCount = 1;
+	Actors[0] = this;
+	for (; ActorCount > 0; )
+	{
+		const AActor* Actor = Actors[--ActorCount];
+		check(ActorCount + Actor->Children.Num() <= UE_ARRAY_COUNT(Actors));
+		for (const AActor* Child : MakeArrayView(Actor->Children))
+		{
+			Actors[ActorCount++] = Child;
+		}
+
+		if (Actor->GetIsReplicated())
+		{
+			const FNetHandle NetHandle = FReplicationSystemUtil::GetNetHandle(Actor);
+			if (NetHandle.IsValid())
+			{
+				ReplicationSystem->SetOwningNetConnection(NetHandle, NewOwningNetConnectionId);
+				// Update autonomous proxy condition
+				if (Actor->GetRemoteRole() == ROLE_AutonomousProxy)
+				{
+					const bool bEnableAutonomousCondition = true;
+					ReplicationSystem->SetReplicationConditionConnectionFilter(NetHandle, EReplicationCondition::RoleAutonomous, NewOwningNetConnectionId, bEnableAutonomousCondition);
+				}
+			}
+		}
+	}
+}
+#endif // UE_WITH_IRIS
+
+void AActor::UpdateReplicatePhysicsCondition()
+{
+#if UE_WITH_IRIS
+	using namespace UE::Net;
+
+	const FNetHandle NetHandle = FReplicationSystemUtil::GetNetHandle(this);
+	if (NetHandle.IsValid())
+	{
+		UReplicationSystem* ReplicationSystem = FReplicationSystemUtil::GetReplicationSystem(this);
+		ReplicationSystem->SetReplicationCondition(NetHandle, EReplicationCondition::ReplicatePhysics, GetReplicatedMovement().bRepPhysics);
+	}
+#endif // UE_WITH_IRIS
+}
+
diff --git a/Engine/Source/Runtime/Engine/Private/Components/ActorComponent.cpp b/Engine/Source/Runtime/Engine/Private/Components/ActorComponent.cpp
--- a/Engine/Source/Runtime/Engine/Private/Components/ActorComponent.cpp
+++ b/Engine/Source/Runtime/Engine/Private/Components/ActorComponent.cpp
@@ -44,6 +44,11 @@
 #endif
 #include "ObjectTrace.h"
 
+#if UE_WITH_IRIS
+#include "Iris/ReplicationSystem/ReplicationFragmentUtil.h"
+#include "Net/Iris/ReplicationSystem/ReplicationSystemUtil.h"
+#endif // UE_WITH_IRIS
+
 #define LOCTEXT_NAMESPACE "ActorComponent"
 
 DEFINE_LOG_CATEGORY(LogActorComponent);
@@ -1361,6 +1366,13 @@ void UActorComponent::RegisterComponentWithWorld(UWorld* InWorld, FRegisterCompo
 			RegisterAllComponentTickFunctions(true);
 			if (!bHasBegunPlay)
 			{
+#if UE_WITH_IRIS
+				if (bReplicates)
+				{
+					MyOwner->UpdateReplicatedComponent(this);
+				}
+#endif // UE_WITH_IRIS
+
 				BeginPlay();
 				ensureMsgf(bHasBegunPlay, TEXT("Failed to route BeginPlay (%s)"), *GetFullName());
 			}
@@ -2191,6 +2203,26 @@ bool UActorComponent::CanEditChange(const FProperty* InProperty) const
 }
 #endif
 
+#if UE_WITH_IRIS
+void UActorComponent::RegisterReplicationFragments(UE::Net::FFragmentRegistrationContext& Context, UE::Net::EFragmentRegistrationFlags RegistrationFlags)
+{
+	Super::RegisterReplicationFragments(Context, RegistrationFlags);
+
+	// Build descriptors and allocate PropertyReplicationFragments for this object
+	UE::Net::FReplicationFragmentUtil::CreateAndRegisterFragmentsForObject(this, Context, RegistrationFlags);
+}
+
+void UActorComponent::BeginReplication()
+{
+	UE::Net::FReplicationSystemUtil::BeginReplicationForActorComponent(GetOwner(), this);
+}
+
+void UActorComponent::EndReplication()
+{
+	UE::Net::FReplicationSystemUtil::EndReplicationForActorComponent(this);
+}
+#endif // UE_WITH_IRIS
+
 bool UActorComponent::IsEditableWhenInherited() const
 {
 	bool bCanEdit = bEditableWhenInherited;
diff --git a/Engine/Source/Runtime/Engine/Private/Components/CharacterMovementComponent.cpp b/Engine/Source/Runtime/Engine/Private/Components/CharacterMovementComponent.cpp
--- a/Engine/Source/Runtime/Engine/Private/Components/CharacterMovementComponent.cpp
+++ b/Engine/Source/Runtime/Engine/Private/Components/CharacterMovementComponent.cpp
@@ -33,6 +33,10 @@
 
 #include "Net/PerfCountersHelpers.h"
 #include "ProfilingDebugging/CsvProfiler.h"
+#if UE_WITH_IRIS
+#include "Net/Iris/ReplicationSystem/ReplicationSystemUtil.h"
+#include "Net/Iris/ReplicationSystem/ActorReplicationBridge.h"
+#endif
 #include "Engine/ScopedMovementUpdate.h"
 
 CSV_DEFINE_CATEGORY(CharacterMovement, true);
@@ -8230,6 +8234,43 @@ void UCharacterMovementComponent::ReplicateMoveToServer(float DeltaTime, const F
 	ClientData->PendingMove = NULL;
 }
 
+#if UE_WITH_IRIS
+namespace UE::Private
+{
+
+static UIrisObjectReferencePackageMap* GetIrisPackageMapToCaptureReferences(UNetConnection* NetConnection, UIrisObjectReferencePackageMap::FObjectReferenceArray* InObjectReferences)
+{
+	using namespace UE::Net;
+
+	if (const UActorReplicationBridge* Bridge = FReplicationSystemUtil::GetActorReplicationBridge(NetConnection))
+	{
+		if (UIrisObjectReferencePackageMap* ObjectReferencePackageMap = Bridge->GetObjectReferencePackageMap())
+		{
+			ObjectReferencePackageMap->InitForWrite(InObjectReferences);
+			return ObjectReferencePackageMap;
+		}
+	}
+
+	return nullptr;
+}
+
+static UIrisObjectReferencePackageMap* GetIrisPackageMapToReadReferences(const UNetConnection* NetConnection, const UIrisObjectReferencePackageMap::FObjectReferenceArray* InObjectReferences)
+{
+	using namespace UE::Net;
+	if (const UActorReplicationBridge* Bridge = FReplicationSystemUtil::GetActorReplicationBridge(NetConnection))
+	{
+		if (UIrisObjectReferencePackageMap* ObjectReferencePackageMap = Bridge->GetObjectReferencePackageMap())
+		{
+			ObjectReferencePackageMap->InitForRead(InObjectReferences);
+			return ObjectReferencePackageMap;
+		}
+	}
+
+	return nullptr;
+}
+
+}
+#endif
 
 void UCharacterMovementComponent::CallServerMovePacked(const FSavedMove_Character* NewMove, const FSavedMove_Character* PendingMove, const FSavedMove_Character* OldMove)
 {
@@ -8241,9 +8282,22 @@ void UCharacterMovementComponent::CallServerMovePacked(const FSavedMove_Characte
 	FBitWriterMark BitWriterReset;
 	BitWriterReset.Pop(ServerMoveBitWriter);
 
-	// Extract the net package map used for serializing object references.
-	UNetConnection* NetConnection = CharacterOwner->GetNetConnection();
-	ServerMoveBitWriter.PackageMap = NetConnection ? ToRawPtr(NetConnection->PackageMap) : nullptr;
+	// 'static' to avoid reallocation each invocation
+	static FCharacterServerMovePackedBits PackedBits;
+	UNetConnection* NetConnection = CharacterOwner->GetNetConnection();	
+
+#if UE_WITH_IRIS
+	if (UPackageMap* PackageMap = UE::Private::GetIrisPackageMapToCaptureReferences(NetConnection, &PackedBits.ObjectReferences))
+	{
+		ServerMoveBitWriter.PackageMap = PackageMap;
+	}
+	else
+#endif
+	{
+		// Extract the net package map used for serializing object references.
+		ServerMoveBitWriter.PackageMap = NetConnection ? ToRawPtr(NetConnection->PackageMap) : nullptr;
+	}
+
 	if (ServerMoveBitWriter.PackageMap == nullptr)
 	{
 		UE_LOG(LogNetPlayerMovement, Error, TEXT("CallServerMovePacked: Failed to find a NetConnection/PackageMap for data serialization!"));
@@ -8258,8 +8312,6 @@ void UCharacterMovementComponent::CallServerMovePacked(const FSavedMove_Characte
 	}
 
 	// Copy bits to our struct that we can NetSerialize to the server.
-	// 'static' to avoid reallocation each invocation
-	static FCharacterServerMovePackedBits PackedBits;
 	PackedBits.DataBits.SetNumUninitialized(ServerMoveBitWriter.GetNumBits());
 	
 	check(PackedBits.DataBits.Num() >= ServerMoveBitWriter.GetNumBits());
@@ -8947,7 +8999,17 @@ void UCharacterMovementComponent::ServerMovePacked_ServerReceive(const FCharacte
 
 	// Reuse bit reader to avoid allocating memory each time.
 	ServerMoveBitReader.SetData((uint8*)PackedBits.DataBits.GetData(), NumBits);
-	ServerMoveBitReader.PackageMap = PackedBits.GetPackageMap();
+
+#if UE_WITH_IRIS
+	if (UPackageMap* PackageMap = UE::Private::GetIrisPackageMapToReadReferences(CharacterOwner->GetNetConnection(), &PackedBits.ObjectReferences))
+	{
+		ServerMoveBitReader.PackageMap = PackageMap;
+	}
+	else
+#endif
+	{
+		ServerMoveBitReader.PackageMap = PackedBits.GetPackageMap();
+	}
 
 	// Deserialize bits to move data struct.
 	// We had to wait until now and use the temp bit stream because the RPC doesn't know about the virtual overrides on the possibly custom struct that is our data container.
@@ -9735,7 +9797,17 @@ void UCharacterMovementComponent::MoveResponsePacked_ClientReceive(const FCharac
 
 	// Reuse bit reader to avoid allocating memory each time.
 	MoveResponseBitReader.SetData((uint8*)PackedBits.DataBits.GetData(), NumBits);
-	MoveResponseBitReader.PackageMap = PackedBits.GetPackageMap();
+
+#if UE_WITH_IRIS
+	if (UPackageMap* PackageMap = UE::Private::GetIrisPackageMapToReadReferences(CharacterOwner->GetNetConnection(), &PackedBits.ObjectReferences))
+	{
+		MoveResponseBitReader.PackageMap = PackageMap;
+	}
+	else
+#endif
+	{
+		MoveResponseBitReader.PackageMap = PackedBits.GetPackageMap();
+	}
 
 	// Deserialize bits to response data struct.
 	// We had to wait until now and use the temp bit stream because the RPC doesn't know about the virtual overrides on the possibly custom struct that is our data container.
@@ -9760,9 +9832,22 @@ void UCharacterMovementComponent::ServerSendMoveResponse(const FClientAdjustment
 	FBitWriterMark BitWriterReset;
 	BitWriterReset.Pop(MoveResponseBitWriter);
 
+	// 'static' to avoid reallocation each invocation
+	static FCharacterMoveResponsePackedBits PackedBits;
+	UNetConnection* NetConnection = CharacterOwner->GetNetConnection();	
+
 	// Extract the net package map used for serializing object references.
-	UNetConnection* NetConnection = CharacterOwner->GetNetConnection();
-	MoveResponseBitWriter.PackageMap = NetConnection ? ToRawPtr(NetConnection->PackageMap) : nullptr;
+#if UE_WITH_IRIS
+	if (UPackageMap* PackageMap = UE::Private::GetIrisPackageMapToCaptureReferences(NetConnection, &PackedBits.ObjectReferences))	
+	{
+		MoveResponseBitWriter.PackageMap = PackageMap;
+	}
+	else
+#endif
+	{
+		MoveResponseBitWriter.PackageMap = NetConnection ? ToRawPtr(NetConnection->PackageMap) : nullptr;
+	}
+
 	if (MoveResponseBitWriter.PackageMap == nullptr)
 	{
 		UE_LOG(LogNetPlayerMovement, Error, TEXT("ServerSendMoveResponse: Failed to find a NetConnection/PackageMap for data serialization!"));
@@ -9777,8 +9862,6 @@ void UCharacterMovementComponent::ServerSendMoveResponse(const FClientAdjustment
 	}
 
 	// Copy bits to our struct that we can NetSerialize to the client.
-	// 'static' to avoid reallocation each invocation
-	static FCharacterMoveResponsePackedBits PackedBits;
 	PackedBits.DataBits.SetNumUninitialized(MoveResponseBitWriter.GetNumBits());
 
 	check(PackedBits.DataBits.Num() >= MoveResponseBitWriter.GetNumBits());
diff --git a/Engine/Source/Runtime/Engine/Private/Components/ChildActorComponent.cpp b/Engine/Source/Runtime/Engine/Private/Components/ChildActorComponent.cpp
--- a/Engine/Source/Runtime/Engine/Private/Components/ChildActorComponent.cpp
+++ b/Engine/Source/Runtime/Engine/Private/Components/ChildActorComponent.cpp
@@ -9,6 +9,10 @@
 #include "Engine/Engine.h"
 #include "Engine/DemoNetDriver.h"
 
+#if UE_WITH_IRIS
+#include "Iris/ReplicationSystem/PropertyReplicationFragment.h"
+#endif // UE_WITH_IRIS
+
 DEFINE_LOG_CATEGORY_STATIC(LogChildActorComponent, Warning, All);
 
 ENGINE_API int32 GExperimentalAllowPerInstanceChildActorProperties = 0;
@@ -842,6 +846,17 @@ void UChildActorComponent::BeginPlay()
 	}
 }
 
+#if UE_WITH_IRIS
+void UChildActorComponent::RegisterReplicationFragments(UE::Net::FFragmentRegistrationContext& Context, UE::Net::EFragmentRegistrationFlags RegistrationFlags)
+{
+	using namespace UE::Net;
+
+	// Fallback on default since archetype used on client and server differs giving different default values
+	Super::RegisterReplicationFragments(Context, RegistrationFlags | UE::Net::EFragmentRegistrationFlags::InitializeDefaultStateFromClassDefaults);
+}
+#endif // UE_WITH_IRIS
+
+
 #if WITH_EDITOR
 void UChildActorComponent::SetEditorTreeViewVisualizationMode(EChildActorComponentTreeViewVisualizationMode InMode)
 {
diff --git a/Engine/Source/Runtime/Engine/Private/DataChannel.cpp b/Engine/Source/Runtime/Engine/Private/DataChannel.cpp
--- a/Engine/Source/Runtime/Engine/Private/DataChannel.cpp
+++ b/Engine/Source/Runtime/Engine/Private/DataChannel.cpp
@@ -1646,6 +1646,8 @@ void UControlChannel::ReceivedBunch( FInBunch& Bunch )
 {
 	check(!Closing);
 
+	UE_NET_TRACE_SCOPE(ControlChannel, Bunch, Connection->GetInTraceCollector(), ENetTraceVerbosity::Trace);
+
 	// If this is a new client connection inspect the raw packet for endianess
 	if (Connection && bNeedsEndianInspection && !CheckEndianess(Bunch))
 	{
@@ -3449,6 +3451,8 @@ int64 UActorChannel::ReplicateActor()
 #if !UE_BUILD_SHIPPING
 						CSV_CUSTOM_STAT(Replication, DeletedSubobjects, 1, ECsvCustomStatOp::Accumulate);
 #endif
+						UE_NET_TRACE_SCOPE(ContentBlockForSubObjectDelete, Bunch, GetTraceCollector(Bunch), ENetTraceVerbosity::Trace);
+						UE_NET_TRACE_OBJECT_SCOPE(LocalReplicator->ObjectNetGUID, Bunch, GetTraceCollector(Bunch), ENetTraceVerbosity::Trace);
 
 						// Write a deletion content header:
 						WriteContentBlockForSubObjectDelete(Bunch, LocalReplicator->ObjectNetGUID);
@@ -3485,6 +3489,7 @@ int64 UActorChannel::ReplicateActor()
 	int64 NumBitsWrote = 0;
 	if (bWroteSomethingImportant)
 	{
+		// We must exit the collection scope to report data correctly
 		FPacketIdRange PacketRange = SendBunch( &Bunch, 1 );
 
 		if (!bIsNewlyReplicationPaused)
@@ -4113,20 +4118,22 @@ void UActorChannel::WriteContentBlockForSubObjectDelete( FOutBunch & Bunch, FNet
 int32 UActorChannel::WriteContentBlockPayload( UObject* Obj, FNetBitWriter &Bunch, const bool bHasRepLayout, FNetBitWriter& Payload )
 {
 	const int32 StartHeaderBits = Bunch.GetNumBits();
+	
+	// Trace header
+	{
+		UE_NET_TRACE_SCOPE(ContentBlockHeader, Bunch, GetTraceCollector(Bunch), ENetTraceVerbosity::Trace);
 
-	WriteContentBlockHeader( Obj, Bunch, bHasRepLayout );
+		WriteContentBlockHeader( Obj, Bunch, bHasRepLayout );
 
-	uint32 NumPayloadBits = Payload.GetNumBits();
+		uint32 NumPayloadBits = Payload.GetNumBits();
 
-	Bunch.SerializeIntPacked( NumPayloadBits );
+		Bunch.SerializeIntPacked( NumPayloadBits );
+	}
 
 	const int32 HeaderNumBits = Bunch.GetNumBits() - StartHeaderBits;
 
-	// Trace header
-	UE_NET_TRACE(ContentBlockHeader, GetTraceCollector(Bunch), StartHeaderBits, Bunch.GetNumBits(), ENetTraceVerbosity::Trace);
-
 	// Inject payload events right after header
-	UE_NET_TRACE_EVENTS(GetTraceCollector(Bunch), GetTraceCollector(Payload));
+	UE_NET_TRACE_EVENTS(GetTraceCollector(Bunch), GetTraceCollector(Payload), Bunch);
 
 	Bunch.SerializeBits( Payload.GetData(), Payload.GetNumBits() );
 
@@ -4501,10 +4508,10 @@ int32 UActorChannel::WriteFieldHeaderAndPayload( FNetBitWriter& Bunch, const FCl
 
 	UE_NET_TRACE(FieldHeader, GetTraceCollector(Bunch), NumOriginalBits, Bunch.GetNumBits(), ENetTraceVerbosity::Trace);
 
-	Bunch.SerializeBits( Payload.GetData(), NumPayloadBits );
-
 	// Inject trace data from payload stream	
-	UE_NET_TRACE_EVENTS(GetTraceCollector(Bunch), GetTraceCollector(Payload));
+	UE_NET_TRACE_EVENTS(GetTraceCollector(Bunch), GetTraceCollector(Payload), Bunch);
+
+	Bunch.SerializeBits( Payload.GetData(), NumPayloadBits );
 
 	return Bunch.GetNumBits() - NumOriginalBits;
 }
diff --git a/Engine/Source/Runtime/Engine/Private/DataReplication.cpp b/Engine/Source/Runtime/Engine/Private/DataReplication.cpp
--- a/Engine/Source/Runtime/Engine/Private/DataReplication.cpp
+++ b/Engine/Source/Runtime/Engine/Private/DataReplication.cpp
@@ -364,6 +364,8 @@ bool FObjectReplicator::SendCustomDeltaProperty(UObject* InObject, uint16 Custom
 	Parms.bSupportsFastArrayDeltaStructSerialization = bSupportsFastArrayDelta;
 	Parms.Connection = Connection;
 	Parms.bInternalAck = Connection->IsInternalAck();
+	// When initializing baselines we should not modify the source data if it originates from the CDO or archetype
+	Parms.bIsInitializingBaseFromDefault = Parms.Object && Parms.Object->HasAnyFlags(RF_ClassDefaultObject | RF_ArchetypeObject);
 
 	return FNetSerializeCB::SendCustomDeltaProperty(*RepLayout, Parms, CustomDeltaIndex);
 }
@@ -1582,6 +1584,11 @@ void FObjectReplicator::ReplicateCustomDeltaProperties( FNetBitWriter & Bunch, F
 
 	FNetBitWriter TempBitWriter( Connection->PackageMap, 1024 );
 
+#if UE_NET_TRACE_ENABLED
+	SetTraceCollector(TempBitWriter, GetTraceCollector(Bunch) ? UE_NET_TRACE_CREATE_COLLECTOR(ENetTraceVerbosity::Trace) : nullptr);
+    ON_SCOPE_EXIT { UE_NET_TRACE_DESTROY_COLLECTOR(GetTraceCollector(TempBitWriter)); };
+#endif
+	
 	const bool bIsConnectionInternalAck = Connection->IsInternalAck();
 
 	// Replicate those properties.
@@ -1603,6 +1610,12 @@ void FObjectReplicator::ReplicateCustomDeltaProperties( FNetBitWriter & Bunch, F
 		TSharedPtr<INetDeltaBaseState> NewState;
 
 		TempBitWriter.Reset();
+#if UE_NET_TRACE_ENABLED
+		if (FNetTraceCollector* TraceCollection = GetTraceCollector(TempBitWriter))
+		{
+			TraceCollection->Reset();
+		}
+#endif
 
 		TSharedPtr<INetDeltaBaseState>& OldState = UsingCustomDeltaStates[CustomDeltaProperty];
 
diff --git a/Engine/Source/Runtime/Engine/Private/LevelStreaming.cpp b/Engine/Source/Runtime/Engine/Private/LevelStreaming.cpp
--- a/Engine/Source/Runtime/Engine/Private/LevelStreaming.cpp
+++ b/Engine/Source/Runtime/Engine/Private/LevelStreaming.cpp
@@ -44,7 +44,11 @@ namespace LevelStreamingCVars
 	// There are cases where we might have multiple visibility requests (and data) in flight leading to the server 
 	// starting to replicate data based on an older visibility/streamingstatus update which can lead to broken channels
 	// to mitigate this problem we assign a TransactionId to each request/update to make sure that we are acting on the correct data
+#if UE_WITH_IRIS
+	static bool bShouldClientUseMakingInvisibleTransactionRequest = true;
+#else
 	static bool bShouldClientUseMakingInvisibleTransactionRequest = false;
+#endif
 	FAutoConsoleVariableRef CVarShouldClientUseMakingInvisibleTransactionRequest(
 		TEXT("LevelStreaming.ShouldClientUseMakingInvisibleTransactionRequest"),
 		bShouldClientUseMakingInvisibleTransactionRequest,
@@ -52,7 +56,11 @@ namespace LevelStreamingCVars
 		TEXT("0: Disable, 1: Enable"),
 		ECVF_Default);
 
+#if UE_WITH_IRIS
+	static bool bShouldServerUseMakingVisibleTransactionRequest = true;
+#else
 	static bool bShouldServerUseMakingVisibleTransactionRequest = false;
+#endif
 	FAutoConsoleVariableRef CVarShouldServerUseMakingVisibleTransactionRequest(
 		TEXT("LevelStreaming.ShouldServerUseMakingVisibleTransactionRequest"),
 		bShouldServerUseMakingVisibleTransactionRequest,
diff --git a/Engine/Source/Runtime/Engine/Private/Net/Experimental/Iris/DataStreamChannel.h b/Engine/Source/Runtime/Engine/Private/Net/Experimental/Iris/DataStreamChannel.h
new file mode 100644
--- /dev/null
+++ b/Engine/Source/Runtime/Engine/Private/Net/Experimental/Iris/DataStreamChannel.h
@@ -0,0 +1,74 @@
+// Copyright Epic Games, Inc. All Rights Reserved.
+
+// Channel for sending and receiving all Iris data
+
+#pragma once
+
+#include "CoreTypes.h"
+#include "UObject/ObjectMacros.h"
+#include "Engine/Channel.h"
+#include "Net/Core/Misc/ResizableCircularQueue.h"
+#include "DataStreamChannel.generated.h"
+
+class UDataStreamManager;
+
+UCLASS(transient, customConstructor)
+class ENGINE_API UDataStreamChannel final : public UChannel
+{
+	GENERATED_BODY()
+
+	UDataStreamChannel(const FObjectInitializer& ObjectInitializer = FObjectInitializer::Get());
+
+private:
+	static void AddReferencedObjects(UObject* Object, FReferenceCollector& Collector);
+
+	// UChannel interface
+	virtual void Init(UNetConnection* InConnection, int32 InChIndex, EChannelCreateFlags CreateFlags) override;
+
+	virtual bool CleanUp(const bool bForDestroy, EChannelCloseReason CloseReason) override;
+
+	/**
+	 * Processes the in bound bunch to extract the data streams
+	 */
+	virtual void ReceivedBunch(FInBunch& Bunch) override;
+
+	virtual void Tick() override;
+
+	/** Always ticks */
+	virtual bool CanStopTicking() const override;
+
+	/** Human readable information about the channel */
+	virtual FString Describe() override;
+
+	/** Packet delivery status handling */
+	virtual void ReceivedAck(int32 PacketId) override;
+	virtual void ReceivedNak(int32 PacketId) override;
+	
+private:
+	enum : uint32
+	{
+		MaxPacketsInFlightCount = 256U,
+		BitStreamBufferByteCount = 2048U,
+	};
+
+	struct FDataStreamChannelRecord
+	{
+		const void* Record = nullptr;
+		uint32 PacketId = 0U;
+	};
+
+	bool IsPacketWindowFull() const;
+	void DiscardAllRecords();
+
+	void SendOpenBunch();
+
+	TObjectPtr<UDataStreamManager> DataStreamManager = nullptr;
+
+	TResizableCircularQueue<FDataStreamChannelRecord> WriteRecords;
+
+	uint32 BitStreamBuffer[BitStreamBufferByteCount];
+
+	uint32 bIsReadyToHandshake : 1U;
+	uint32 bHandshakeSent : 1U;
+	uint32 bHandshakeComplete : 1U;
+};
