Iris experimental replication system.

diff --git a/Engine/Source/Runtime/Engine/Private/Net/Iris/ReplicationSystem/ActorReplicationBridge.cpp b/Engine/Source/Runtime/Engine/Private/Net/Iris/ReplicationSystem/ActorReplicationBridge.cpp
new file mode 100644
--- /dev/null
+++ b/Engine/Source/Runtime/Engine/Private/Net/Iris/ReplicationSystem/ActorReplicationBridge.cpp
@@ -0,0 +1,989 @@
+// Copyright Epic Games, Inc. All Rights Reserved.
+#include "Net/Iris/ReplicationSystem/ActorReplicationBridge.h"
+
+#if UE_WITH_IRIS
+
+#include "Net/Iris/ReplicationSystem/ActorReplicationBridgeInternal.h"
+#include "Iris/IrisConfig.h"
+#include "Iris/Core/IrisLog.h"
+#include "Iris/Core/IrisProfiler.h"
+#include "Iris/Core/NetObjectReference.h"
+#include "Iris/ReplicationSystem/NetToken.h"
+#include "Iris/ReplicationSystem/StringTokenStore.h"
+#include "Iris/ReplicationSystem/ReplicationOperations.h"
+#include "Iris/ReplicationSystem/ReplicationSystem.h"
+#include "Iris/ReplicationSystem/Filtering/NetObjectFilter.h"
+#include "Iris/Serialization/NetBitStreamReader.h"
+#include "Iris/Serialization/NetBitStreamUtil.h"
+#include "Iris/Serialization/NetBitStreamWriter.h"
+#include "Iris/Serialization/ObjectNetSerializer.h"
+#include "Iris/Serialization/IrisObjectReferencePackageMap.h"
+#include "Engine/Engine.h"
+#include "Engine/EngineTypes.h"
+#include "Engine/NetDriver.h"
+#include "Engine/World.h"
+#include "GameFramework/Actor.h"
+#include "Net/DataBunch.h"
+#include "Net/Core/Misc/NetSubObjectRegistry.h"
+#include "Net/NetSubObjectRegistryGetter.h"
+#include "Templates/Casts.h"
+
+
+#if UE_NET_ALLOW_MULTIPLE_REPLICATION_SYSTEMS
+#define UE_LOG_ACTORREPLICATIONBRIDGE(Category, Format, ...)	UE_LOG(LogIrisBridge, Category, TEXT("ActorReplicationBridge(%u)::") Format, GetReplicationSystem()->GetId(), ##__VA_ARGS__)
+#else
+#define UE_LOG_ACTORREPLICATIONBRIDGE(Category, Format, ...)	UE_LOG(LogIrisBridge, Category, TEXT("ActorReplicationBridge::") Format, ##__VA_ARGS__)
+#endif
+
+extern bool GDefaultUseSubObjectReplicationList;
+
+namespace UE::Net
+{
+
+bool ShouldUseIrisReplication(const UObject* Object)
+{
+	return ShouldUseIrisReplication();
+}
+
+}
+
+namespace UE::Net::Private
+{
+// This is not pretty but the alternative is exposing the PollFrequencyMultiplier in a public header file.
+// It's really only useful for this particular class and in order to ensure
+// NetUpdateFrequencies are ok.
+extern IRISCORE_API float PollFrequencyMultiplier;
+extern IRISCORE_API const int MaxPollFramePeriod;
+
+static bool bIrisEnableLowNetUpdateFrequencyEnsure = true;
+static FAutoConsoleVariableRef CVarEnableLowNetUpdateFrequencyEnsure(
+		TEXT("net.Iris.EnableLowNetUpdateFrequencyEnsure"),
+		bIrisEnableLowNetUpdateFrequencyEnsure,
+		TEXT("Whether to ensure when NetUpdateFrequency is so low that the poll frame period is clamped. Default is true.")
+		);
+
+bool IsActorValidForIrisReplication(const AActor* Actor)
+{
+	return IsValid(Actor) && !Actor->IsActorBeingDestroyed() && !Actor->IsUnreachable();
+}
+
+void ActorReplicationBridgePreUpdateFunction(FNetHandle Handle, UObject* Instance, const UReplicationBridge* Bridge)
+{
+	AActor* Actor = Cast<AActor>(Instance);
+	if (IsActorValidForIrisReplication(Actor))
+	{
+		/** $IRIS TODO:
+		 * Here we need to call something that either fakes CallPrereplication or we need to mimic the interfaces of NetDriver/PropertyChangeTracker
+		 * however we do think any state changes should be pushed to the network system rather than polled
+		 * One example is the gatherCurrent() movement, movementdata should have its own ReplicationState that is updated when changed. (or updated using a helper method)
+		 */
+		Actor->CallPreReplication(CastChecked<const UActorReplicationBridge>(Bridge)->GetNetDriver());
+	}
+}
+
+FVector ActorReplicationBridgeGetActorWorldLocation(FNetHandle Handle, const UObject* Instance)
+{
+	if (const AActor* Actor = Cast<AActor>(Instance))
+	{
+		return Actor->GetActorLocation();
+	}
+
+	return FVector::Zero();
+}
+
+}
+
+UActorReplicationBridge::UActorReplicationBridge()
+: UObjectReplicationBridge()
+, NetDriver(nullptr)
+, MaxPollFrequency(0.0f)
+, ObjectReferencePackageMap(nullptr)
+{
+	SetInstancePreUpdateFunction(UE::Net::Private::ActorReplicationBridgePreUpdateFunction);
+	SetInstanceGetWorldLocationFunction(UE::Net::Private::ActorReplicationBridgeGetActorWorldLocation);
+}
+
+void UActorReplicationBridge::Initialize(UReplicationSystem* InReplicationSystem)
+{
+	Super::Initialize(InReplicationSystem);
+
+	{
+		auto&& ShouldSpatialize = [](const UClass* Class)
+		{
+			if (AActor* CDO = Cast<AActor>(Class->GetDefaultObject()))
+			{
+				return !(CDO->bAlwaysRelevant || CDO->bOnlyRelevantToOwner || CDO->bNetUseOwnerRelevancy);
+			}
+
+			return false;
+		};
+
+		SetShouldUseDefaultSpatialFilterFunction(ShouldSpatialize);
+	}
+
+	{
+		auto&& ClassesAreRelevantEqual = [](const UClass* Class, const UClass* Subclass)
+		{
+			const AActor* CDO = Cast<AActor>(Class->GetDefaultObject());
+			const AActor* SubCDO = Cast<AActor>(Subclass->GetDefaultObject());
+			// Same CDO (nullptr)?
+			if (CDO == SubCDO)
+			{
+				return true;
+			}
+
+			if (CDO == nullptr || SubCDO == nullptr)
+			{
+				return false;
+			}
+
+			{
+				return CDO->bAlwaysRelevant == SubCDO->bAlwaysRelevant
+					&& CDO->bOnlyRelevantToOwner == SubCDO->bOnlyRelevantToOwner
+					&& CDO->bNetUseOwnerRelevancy == SubCDO->bNetUseOwnerRelevancy;
+			}
+		};
+
+		SetShouldSubclassUseSameFilterFunction(ClassesAreRelevantEqual);
+	}
+
+	ObjectReferencePackageMap = NewObject<UIrisObjectReferencePackageMap>();
+	ObjectReferencePackageMap->AddToRoot();
+}
+
+UActorReplicationBridge::~UActorReplicationBridge()
+{
+	if (ObjectReferencePackageMap)
+	{
+		ObjectReferencePackageMap->RemoveFromRoot();
+		ObjectReferencePackageMap->MarkAsGarbage();
+		ObjectReferencePackageMap = nullptr;
+	}
+}
+
+UE::Net::FNetHandle UActorReplicationBridge::BeginReplication(AActor* Actor, const FActorBeginReplicationParams& Params)
+{
+	using namespace UE::Net;
+
+	if (!ShouldUseIrisReplication(Actor))
+	{
+		return FNetHandle();
+	}
+
+	if (!ensureMsgf(Actor == nullptr || !(Actor->HasAnyFlags(RF_ClassDefaultObject | RF_ArchetypeObject)), TEXT("Actor %s is a CDO or Archetype and should not be replicated."), ToCStr(GetFullNameSafe(Actor))))
+	{
+		return FNetHandle();
+	}
+
+	const bool bIsNetActor = ULevel::IsNetActor(Actor);
+	if (!bIsNetActor)
+	{
+		UE_LOG_ACTORREPLICATIONBRIDGE(VeryVerbose, TEXT("Actor %s doesn't have a NetRole."), ToCStr(GetFullNameSafe(Actor)));
+		return FNetHandle();
+	}
+
+	if (Actor->GetLocalRole() != ROLE_Authority)
+	{
+		UE_LOG_ACTORREPLICATIONBRIDGE(VeryVerbose, TEXT("Actor %s NetRole isn't Authority."), ToCStr(GetFullNameSafe(Actor)));
+		return FNetHandle();
+	}
+
+	if (Actor->IsActorBeingDestroyed() || !IsValid(Actor) || Actor->IsUnreachable())
+	{
+		UE_LOG_ACTORREPLICATIONBRIDGE(Verbose, TEXT("Actor %s is being destroyed or unreachable and can't be replicated."), ToCStr(GetFullNameSafe(Actor)));
+		return FNetHandle();
+	}
+
+	if (!Actor->GetIsReplicated())
+	{
+		UE_LOG_ACTORREPLICATIONBRIDGE(Verbose, TEXT("Actor %s is not supposed to be replicated."), ToCStr(GetFullNameSafe(Actor)));
+		return FNetHandle();
+	}
+
+	if (Actor->GetTearOff())
+	{
+		UE_LOG_ACTORREPLICATIONBRIDGE(Verbose, TEXT("Actor %s is torn off and should not be replicated."), ToCStr(GetFullNameSafe(Actor)));
+		return FNetHandle();
+	}
+
+	if (!Actor->IsActorInitialized())
+	{
+		UE_LOG_ACTORREPLICATIONBRIDGE(Warning, TEXT("Actor %s is not initialized and won't be replicated."), ToCStr(GetFullNameSafe(Actor)));
+		return FNetHandle();
+	}
+
+	if (!NetDriver)
+	{
+		UE_LOG_ACTORREPLICATIONBRIDGE(VeryVerbose, TEXT("There's no NetDriver so nothing can be replicated."));
+		return FNetHandle();
+	}
+
+
+	if (!NetDriver->ShouldReplicateActor(Actor))
+	{
+		UE_LOG_ACTORREPLICATIONBRIDGE(Verbose, TEXT("Actor %s doesn't want to replicate with NetDriver %s."), ToCStr(GetFullNameSafe(Actor)), ToCStr(NetDriver->GetName()));
+		return FNetHandle();
+	}
+
+	// Initially dormant actors begin replication when their dormancy is flushed
+	const ENetDormancy Dormancy = Actor->NetDormancy;	
+	if (Actor->IsNetStartupActor() && (Dormancy == DORM_Initial))
+	{
+		return FNetHandle();
+	}
+
+	FNetHandle ExistingHandle = GetReplicatedHandle(Actor);
+	if (ExistingHandle.IsValid())
+	{
+		return ExistingHandle;
+	}
+
+	if (!Actor->IsUsingRegisteredSubObjectList())
+	{
+		// Ensure the first time to get attention!
+		ensureMsgf(GDefaultUseSubObjectReplicationList, TEXT("Iris requires replicated actors to use registered subobjectslists. Add \n[SystemSettings]\nnet.SubObjects.DefaultUseSubObjectReplicationList=1\n to your DefaultEngine.ini"));
+		UE_LOG_ACTORREPLICATIONBRIDGE(Warning, TEXT("Actor %s does not replicate subobjects using the registered SubObjectsLists, SubObjects will not replicate properly"), ToCStr(GetFullNameSafe(Actor)));
+	}
+
+	// Create NetHandle for the registered fragments
+	Super::FCreateNetHandleParams CreateNetHandleParams = DefaultCreateNetHandleParams;
+	CreateNetHandleParams.bNeedsPreUpdate = 1U;
+	CreateNetHandleParams.bNeedsWorldLocationUpdate = 1U;
+	CreateNetHandleParams.StaticPriority = (Actor->bAlwaysRelevant || Actor->bOnlyRelevantToOwner) ? Actor->NetPriority : 0.0f;
+	const uint32 UnclampedPollPeriod = GetPollFramePeriod(Actor->NetUpdateFrequency);
+	const uint32 ClampedPollPeriod = FMath::Clamp<uint32>(UnclampedPollPeriod, 1U, static_cast<uint32>(std::numeric_limits<uint8>::max()) + 1U);
+	CreateNetHandleParams.PollFramePeriod = (ClampedPollPeriod - 1U) & 255U;
+#if !UE_BUILD_SHIPPING
+	ensureAlwaysMsgf(!Private::bIrisEnableLowNetUpdateFrequencyEnsure || (ClampedPollPeriod >= UnclampedPollPeriod), TEXT("Very low NetUpdateFrequency %f for Actor %s. Suggest setting it to %f or higher."), Actor->NetUpdateFrequency, ToCStr(Actor->GetName()), GetMinSupportedNetUpdateFrequency());
+	ensureAlwaysMsgf(!(Actor->bAlwaysRelevant || Actor->bOnlyRelevantToOwner) || CreateNetHandleParams.StaticPriority >= 1.0f, TEXT("Very low NetPriority %.02f for always relevant or owner relevant Actor %s. Set it to 1.0f or higher."), Actor->NetPriority, ToCStr(Actor->GetName()));
+#endif
+
+	FNetHandle ActorHandle = Super::BeginReplication(Actor, CreateNetHandleParams);
+
+	if (!ActorHandle.IsValid())
+	{
+		ensureMsgf(false, TEXT("Failed to create NetHandle for Actor Named %s"), ToCStr(Actor->GetName()));
+		return FNetHandle();
+	}
+
+	// Set owning connection filtering if actor is only relevant to owner
+	{
+		if (Actor->bOnlyRelevantToOwner & !Actor->bAlwaysRelevant)
+		{
+			GetReplicationSystem()->SetFilter(ActorHandle, ToOwnerFilterHandle);
+		}
+	}
+
+	// Set if this is a NetTemporary
+	{
+		if (Actor->bNetTemporary)
+		{
+			GetReplicationSystem()->SetIsNetTemporary(ActorHandle);
+		}
+	}
+
+	// Dormancy, we track all actors that does want to be dormant
+	{
+		if (Dormancy > DORM_Awake)
+		{
+			SetObjectWantsToBeDormant(ActorHandle, true);
+		}
+	}
+
+	// Setup Level filtering unless this actor belongs to the persistent level
+	{
+		ULevel* Level = Actor->GetLevel();
+		if (Params.bIncludeInLevelGroupFilter && (!Level->IsPersistentLevel() || (Level != NetDriver->GetWorld()->PersistentLevel)))
+		{
+			const UPackage* const LevelPackage = Level->GetOutermost();
+			const FName PackageName = LevelPackage->GetFName();
+
+			FNetObjectGroupHandle LevelGroup = GetLevelGroup(Level);
+			if (!LevelGroup)
+			{
+				LevelGroup = CreateLevelGroup(Level);
+
+				UE_LOG_ACTORREPLICATIONBRIDGE(Log, TEXT("Created new GroupIndex: %u for Level: %s"), LevelGroup, ToCStr(PackageName.ToString()));
+
+				// Update the filtering status of the group based on current level visibility for all connections
+				NetDriver->UpdateGroupFilterStatusForLevel(Level, LevelGroup);
+			}
+
+			// Add object to group
+			UE_LOG_ACTORREPLICATIONBRIDGE(Verbose, TEXT("Added %s to GroupIndex: %u Level: %s"), *ActorHandle.ToString(), LevelGroup, ToCStr(PackageName.ToString()));
+			GetReplicationSystem()->AddToGroup(LevelGroup, ActorHandle);	
+		}
+	}
+
+	// If we have registered sub objects we replicate them as well
+	const FSubObjectRegistry& ActorSubObjects = FSubObjectRegistryGetter::GetSubObjects(Actor);
+	const TArray<FReplicatedComponentInfo>& ReplicatedComponents = FSubObjectRegistryGetter::GetReplicatedComponents(Actor);
+
+	if (ActorSubObjects.GetRegistryList().Num() != 0 || ReplicatedComponents.Num() != 0)
+	{
+		// Start with the Actor's SubObjects (that is SubObjects that are not ActorComponents)
+		for (const FSubObjectRegistry::FEntry& SubObjectInfo : ActorSubObjects.GetRegistryList())
+		{
+			if (IsValid(SubObjectInfo.SubObject) && SubObjectInfo.NetCondition != ELifetimeCondition::COND_Never)
+			{
+				FNetHandle SubObjectHandle = UObjectReplicationBridge::BeginReplication(ActorHandle, SubObjectInfo.SubObject);
+				if (SubObjectHandle.IsValid() && SubObjectInfo.NetCondition != ELifetimeCondition::COND_None)
+				{
+					UObjectReplicationBridge::SetSubObjectNetCondition(SubObjectHandle, SubObjectInfo.NetCondition);
+				}
+			}
+		}
+
+		// Now the replicated ActorComponents and their SubObjects
+		for (const FReplicatedComponentInfo& RepComponentInfo : ReplicatedComponents)
+		{
+			if (IsValid(RepComponentInfo.Component) && RepComponentInfo.NetCondition != ELifetimeCondition::COND_Never)
+			{
+				UActorComponent* ReplicatedComponent = RepComponentInfo.Component;
+				ReplicatedComponent->BeginReplication();
+				// NetCondition is set by replicated component
+			}
+		}
+	}
+
+	return ActorHandle;
+}
+
+UE::Net::FNetHandle UActorReplicationBridge::BeginReplication(FNetHandle OwnerHandle, UActorComponent* SubObject)
+{
+	using namespace UE::Net;
+
+	if (!OwnerHandle.IsValid() || !ShouldUseIrisReplication(SubObject))
+	{
+		return FNetHandle();
+	}
+
+	AActor* Owner = SubObject->GetOwner();
+
+	FNetHandle ReplicatedComponentHandle = GetReplicatedHandle(SubObject);
+	const FReplicatedComponentInfo* RepComponentInfo = FSubObjectRegistryGetter::GetReplicatedComponentInfoForComponent(Owner, SubObject);
+
+	if (!ReplicatedComponentHandle.IsValid())
+	{
+		if (!IsValid(SubObject)
+			|| SubObject->IsUnreachable()
+			|| !SubObject->GetIsReplicated()
+			|| SubObject->HasAnyFlags(RF_ArchetypeObject | RF_ClassDefaultObject)
+		)
+		{
+			return FNetHandle();
+		}
+
+		if (!SubObject->IsUsingRegisteredSubObjectList())
+		{
+			UE_LOG_ACTORREPLICATIONBRIDGE(Warning, TEXT("ActorComponent %s does not replicate subobjects using the registered SubObjectsLists, SubObjects might not replicate properly."), ToCStr(GetFullNameSafe(SubObject)));
+		}
+
+		if (RepComponentInfo == nullptr || RepComponentInfo->NetCondition == ELifetimeCondition::COND_Never)
+		{
+			return FNetHandle();
+		}
+
+		// Create NetHandle and attach for the sub object
+		ReplicatedComponentHandle = Super::BeginReplication(OwnerHandle, SubObject);
+	}
+
+	if (!ReplicatedComponentHandle.IsValid())
+	{
+		ensureMsgf(false, TEXT("Failed to create or find NetHandle for ActorComponent Named %s"), ToCStr(SubObject->GetName()));
+		return FNetHandle();
+	}
+
+	// Update or set any conditionals
+	if (RepComponentInfo->NetCondition != ELifetimeCondition::COND_None)
+	{
+		SetSubObjectNetCondition(ReplicatedComponentHandle, RepComponentInfo->NetCondition);
+	}
+
+	// Begin replication for any SubObjects registered by the component
+	for (const FSubObjectRegistry::FEntry& SubObjectInfo : RepComponentInfo->SubObjects.GetRegistryList())
+	{
+		if (IsValid(SubObjectInfo.SubObject) && SubObjectInfo.NetCondition != ELifetimeCondition::COND_Never)
+		{
+			FNetHandle SubObjectHandle = UObjectReplicationBridge::BeginReplication(OwnerHandle, SubObjectInfo.SubObject, ReplicatedComponentHandle, UReplicationBridge::ESubObjectInsertionOrder::ReplicateWith);
+			if (SubObjectHandle.IsValid() && SubObjectInfo.NetCondition != ELifetimeCondition::COND_None)
+			{
+				SetSubObjectNetCondition(SubObjectHandle, SubObjectInfo.NetCondition);
+			}
+		}
+	}
+
+	return ReplicatedComponentHandle;
+}
+
+void UActorReplicationBridge::EndReplication(AActor* Actor, EEndPlayReason::Type EndPlayReason)
+{
+	using namespace UE::Net;
+
+	FNetHandle NetHandle = GetReplicatedHandle(Actor);
+	if (NetHandle.IsValid())
+	{
+		UE_LOG(LogIrisBridge, Verbose, TEXT("EndReplication for %s %s. Reason %s "), *Actor->GetName(), *NetHandle.ToString(), *UEnum::GetValueAsString(TEXT("Engine.EEndPlayReason"), EndPlayReason));
+	
+		EEndReplicationFlags Flags = EEndReplicationFlags::None;
+		const bool bShouldDestroyObject = EndPlayReason == EEndPlayReason::Destroyed;		
+		if (bShouldDestroyObject)
+		{ 
+			Flags |= EEndReplicationFlags::Destroy;
+		}
+					
+		const bool bShouldCreateDestructionInfo = NetHandle.IsStatic() && bShouldDestroyObject && GetReplicationSystem()->IsServer();
+		if (bShouldCreateDestructionInfo)
+		{
+			UObjectReplicationBridge::FEndReplicationParameters EndReplicationParameters;
+
+			EndReplicationParameters.Location = Actor->GetActorLocation();
+			EndReplicationParameters.Level = Actor->GetLevel();
+			EndReplicationParameters.bUseDistanceBasedPrioritization = Actor->bAlwaysRelevant == false;
+
+			UObjectReplicationBridge::EndReplication(NetHandle, Flags, &EndReplicationParameters);		
+		}
+		else
+		{
+			UObjectReplicationBridge::EndReplication(NetHandle, Flags, nullptr);
+		}
+	}
+}
+
+void UActorReplicationBridge::EndReplicationForActorComponent(UActorComponent* ActorComponent)
+{
+	using namespace UE::Net;
+
+	FNetHandle ComponentHandle = GetReplicatedHandle(ActorComponent);
+	AActor* Actor = ActorComponent->GetOwner();
+	if (ComponentHandle.IsValid() && Actor)
+	{
+		UE_LOG(LogIrisBridge, Verbose, TEXT("EndReplicationForActorComponent for %s %s."), *ActorComponent->GetName(), *ComponentHandle.ToString());
+		
+		UObjectReplicationBridge::EndReplication(ComponentHandle, EEndReplicationFlags::None, nullptr);
+	}
+}
+
+bool UActorReplicationBridge::WriteCreationHeader(UE::Net::FNetSerializationContext& Context, FNetHandle Handle)
+{
+	using namespace UE::Net;
+	using namespace UE::Net::Private;
+
+	FNetBitStreamWriter* Writer = Context.GetBitStreamWriter();
+
+	const UObject* Object = GetReplicatedObject(Handle);
+	if (const AActor* Actor = Cast<AActor>(Object))
+	{
+		// Get Header
+		FActorCreationHeader Header;
+		GetActorCreationHeader(Actor, Header);
+	
+		// Serialize the data
+		// Indicate that this is an actor
+		Writer->WriteBool(true);
+		WriteActorCreationHeader(Context, Header);
+
+		return !Writer->IsOverflown();
+	}
+	else if (Object)
+	{
+		// Get Header
+		FSubObjectCreationHeader Header;
+		GetSubObjectCreationHeader(Object, Header);
+
+		// Serialize the data
+		// Indicate that this is a SubObject
+		Writer->WriteBool(false);
+		WriteSubObjectCreationHeader(Context, Header);
+
+		return !Writer->IsOverflown();
+	}
+
+	ensureAlwaysMsgf(false, TEXT("UActorReplicationBridge::WriteCreationHeader Failed to write creationHeader for NetHandleIndex: %u"), Handle.GetId());
+
+	return false;
+}
+
+UObjectReplicationBridge::FCreationHeader* UActorReplicationBridge::ReadCreationHeader(UE::Net::FNetSerializationContext& Context)
+{
+	using namespace UE::Net;
+	using namespace UE::Net::Private;
+
+	FNetBitStreamReader* Reader = Context.GetBitStreamReader();
+
+	if (const bool bIsActor = Reader->ReadBool())
+	{
+		TUniquePtr<FActorCreationHeader> Header(new FActorCreationHeader);
+		ReadActorCreationHeader(Context, *(Header.Get()));
+
+		if (!Reader->IsOverflown())
+		{
+			return Header.Release();
+		}
+	}
+	else
+	{
+		TUniquePtr<FSubObjectCreationHeader> Header(new FSubObjectCreationHeader);
+		ReadSubObjectCreationHeader(Context, *(Header.Get()));
+
+		if (!Reader->IsOverflown())
+		{
+			return Header.Release();
+		}
+	}
+
+	return nullptr;
+}
+
+UObject* UActorReplicationBridge::BeginInstantiateFromRemote(FNetHandle SubObjectOwnerNetHandle, const UE::Net::FNetObjectResolveContext& ResolveContext, const UObjectReplicationBridge::FCreationHeader* InHeader)
+{
+	using namespace UE::Net::Private;
+
+	IRIS_PROFILER_SCOPE(ActorReplicationBridge_OnBeginInstantiateFromRemote);
+
+	const FActorReplicationBridgeCreationHeader* BridgeHeader = static_cast<const FActorReplicationBridgeCreationHeader*>(InHeader);
+
+	if (BridgeHeader->bIsActor)
+	{
+		const FActorCreationHeader* Header = static_cast<const FActorCreationHeader*>(InHeader);
+
+		// Spawn actor
+		// Try to instantiate the actor (or find it if it already exists)
+		if (Header->bIsDynamic)
+		{
+			// Find archetype
+			UObject* Archetype = ResolveObjectReference(Header->ArchetypeReference, ResolveContext);
+
+			// Find level
+			UObject* Level = nullptr;
+			if (!Header->bUsePersistentLevel)
+			{
+				Level = ResolveObjectReference(Header->LevelReference, ResolveContext);
+			}
+
+			if (Archetype)
+			{
+				// For streaming levels, it's possible that the owning level has been made not-visible but is still loaded.
+				// In that case, the level will still be found but the owning world will be invalid.
+				// If that happens, wait to spawn the Actor until the next time the level is streamed in.
+				// At that point, the Server should resend any dynamic Actors.
+				ULevel* SpawnLevel = Cast<ULevel>(Level);
+				check(SpawnLevel == nullptr || SpawnLevel->GetWorld() != nullptr);
+
+				FActorSpawnParameters SpawnInfo;
+				SpawnInfo.Template = Cast<AActor>(Archetype);
+				SpawnInfo.OverrideLevel = SpawnLevel;
+				SpawnInfo.SpawnCollisionHandlingOverride = ESpawnActorCollisionHandlingMethod::AlwaysSpawn;
+				SpawnInfo.bRemoteOwned = true;
+				SpawnInfo.bNoFail = true;
+
+				UWorld* World = NetDriver->GetWorld();
+				FVector SpawnLocation = FRepMovement::RebaseOntoLocalOrigin(Header->SpawnInfo.Location, World->OriginLocation);
+		
+				AActor* Actor = World->SpawnActorAbsolute(Archetype->GetClass(), FTransform(Header->SpawnInfo.Rotation, SpawnLocation), SpawnInfo);
+
+				// For Iris we expect that we will be able to spawn the actor as streaming always is controlled from server
+				if (ensure(Actor))
+				{
+					const FActorReplicationBridgeSpawnInfo& DefaultSpawnInfo = GetDefaultActorReplicationBridgeSpawnInfo();
+					static constexpr float Epsilon = UE_KINDA_SMALL_NUMBER;
+
+					// Set Velocity if it differs from Default
+					if (!Header->SpawnInfo.Velocity.Equals(DefaultSpawnInfo.Velocity, Epsilon))
+					{
+						Actor->PostNetReceiveVelocity(Header->SpawnInfo.Velocity);
+					}
+					
+					// Set Scale if it differs from Default
+					if (!Header->SpawnInfo.Scale.Equals(DefaultSpawnInfo.Scale, Epsilon))
+					{
+						Actor->SetActorScale3D(Header->SpawnInfo.Scale);
+					}
+
+					UE_LOG_ACTORREPLICATIONBRIDGE(Verbose, TEXT("OnBeginInstantiateFromRemote Spawned Actor %s"), *Actor->GetPathName());
+				}
+
+				return Actor;
+			}
+			else
+			{
+				UE_LOG_ACTORREPLICATIONBRIDGE(Warning, TEXT("OnBeginInstantiateFromRemote Unable to spawn object, failed to resolve archetype with %s"), *(Header->ArchetypeReference.GetRefHandle().ToString()));
+				check(false);
+			}
+		}
+		else
+		{
+			if (UObject* Object = ResolveObjectReference(Header->ObjectReference, ResolveContext))
+			{
+				UE_LOG_ACTORREPLICATIONBRIDGE(Verbose, TEXT("OnBeginInstantiateFromRemote Found static Actor using path %s"), ToCStr(Object->GetPathName()));
+				return Object;
+			}
+			else
+			{
+				UE_LOG_ACTORREPLICATIONBRIDGE(Warning, TEXT("OnBeginInstantiateFromRemote Failed to find Resolve ObjectReference for static Actor %s"), *Header->ObjectReference.ToString());
+			}
+		}
+	}
+	else
+	{
+		const FSubObjectCreationHeader* Header = static_cast<const FSubObjectCreationHeader*>(InHeader);
+
+		// Spawn sub object
+		if (Header->bIsDynamic)
+		{
+			// Resolve owner
+			UObject* ResolvedOwnerObject = GetReplicatedObject(SubObjectOwnerNetHandle);
+			AActor* Owner = Cast<AActor>(ResolvedOwnerObject);
+
+			check(Owner);
+			
+			UObject* SubObj = nullptr;
+
+			if (Header->bIsNameStableForNetworking)
+			{
+				// Resolve by finding object relative to owner
+				SubObj = ResolveObjectReference(Header->ObjectReference, ResolveContext);
+			}
+			else
+			{
+				// We need to spawn the subobject
+				UObject* SubObjClassObj = ResolveObjectReference(Header->ObjectClassReference, ResolveContext);
+				UClass * SubObjClass = Cast<UClass>(SubObjClassObj);
+
+				// Try to spawn SubObject
+ 				SubObj = NewObject< UObject >(Owner, SubObjClass);
+
+				// Notify actor that we created a component from replication
+				Owner->OnSubobjectCreatedFromReplication(SubObj);
+			}
+
+			// Sanity check some things
+			check(SubObj != nullptr);
+			check(SubObj->IsIn( Owner ));
+			check(Cast< AActor >( SubObj ) == nullptr);
+			
+			return SubObj;
+		}
+		else
+		{
+			// try to find the object by path
+			if (UObject* Object = ResolveObjectReference(Header->ObjectReference, ResolveContext))
+			{
+				UE_LOG_ACTORREPLICATIONBRIDGE(Verbose, TEXT("BeginInstantiateFromRemote Found static SubObject using path %s"), ToCStr(Object->GetPathName()));
+				return Object;
+			}
+			else
+			{
+				UE_LOG_ACTORREPLICATIONBRIDGE(Warning, TEXT("BeginInstantiateFromRemote Failed to find Resolve SubObjectReference for static SubObject %s, Owner %s (%s)"), *Header->ObjectReference.ToString(), *SubObjectOwnerNetHandle.ToString(), *GetPathNameSafe(GetReplicatedObject(SubObjectOwnerNetHandle)));
+			}			
+		}
+	}
+
+	return nullptr;
+}
+
+bool UActorReplicationBridge::OnInstantiatedFromRemote(UObject* Instance, const UObjectReplicationBridge::FCreationHeader* InHeader, uint32 ConnectionId) const
+{
+	using namespace UE::Net::Private;
+
+	const FActorReplicationBridgeCreationHeader* BridgeHeader = static_cast<const FActorReplicationBridgeCreationHeader*>(InHeader);
+
+	if (BridgeHeader->bIsActor)
+	{
+		const FActorCreationHeader* Header = static_cast<const FActorCreationHeader*>(InHeader);
+
+		if (Header->bIsDynamic)
+		{
+			// OnActorChannelOpen
+			AActor* Actor = CastChecked<AActor>(Instance);
+			{
+				UNetConnection* Connection = NetDriver->GetConnectionById(ConnectionId);
+				FInBunch Bunch(Connection, const_cast<uint8*>(Header->CustomCreationData.GetData()), Header->CustomCreationDataBitCount);
+				Actor->OnActorChannelOpen(Bunch, Connection);
+
+				if (Bunch.IsError() || Bunch.GetBitsLeft() != 0)
+				{
+					// $IRIS TODO Report an intelligent error
+					check(false);
+					return false;
+				}
+			}
+		}
+	}
+	
+	return true;
+}
+
+void UActorReplicationBridge::EndInstantiateFromRemote(FNetHandle Handle)
+{
+	if (AActor* Actor = Cast<AActor>(GetReplicatedObject(Handle)))
+	{
+		// Optional
+		//Actor->NetHandle = Handle;
+
+		Actor->PostNetInit();
+	}
+}
+
+void UActorReplicationBridge::DestroyInstanceFromRemote(UObject* Instance, bool bTearOff)
+{
+	if (AActor* Actor = Cast<AActor>(Instance))
+	{
+		if (bTearOff && !NetDriver->ShouldClientDestroyTearOffActors())
+		{
+			Actor->SetRole(ROLE_Authority);
+			Actor->SetReplicates(false);
+
+			if (Actor->GetWorld() != nullptr && !IsEngineExitRequested())
+			{
+				Actor->TornOff();
+			}
+
+			NetDriver->NotifyActorTornOff(Actor);
+		}
+		else
+		{
+			// Any subobjects have already been detached from ReplicationBridge
+			Actor->PreDestroyFromReplication();
+			Actor->Destroy(true);
+		}
+	}
+	else
+	{
+		// If the SubObject is being torn-off, it is up to owning actor to clean it up properly
+		if (bTearOff && !NetDriver->ShouldClientDestroyTearOffActors())
+		{
+			return;
+		}
+
+		AActor* Owner = Cast<AActor>(Instance->GetOuter());
+		if (ensureAlwaysMsgf(IsValid(Owner) && !Owner->IsUnreachable(), TEXT("UActorReplicationBridge::DestroyInstanceFromRemote Destroyed subobject after owner %s"), *Instance->GetPathName()))
+		{
+			Owner->OnSubobjectDestroyFromReplication(Instance);
+		}
+
+		Instance->PreDestroyFromReplication();
+		Instance->MarkAsGarbage();
+	}
+}
+
+void UActorReplicationBridge::GetInitialDependencies(FNetHandle Handle, FNetDependencyInfoArray& OutDependencies) const
+{
+	using namespace UE::Net;
+
+	uint32 DependencyCount = 0U;
+
+	// $TODO: Cache create dependencies and do the lookup based on index
+
+	// Handles with static names does not have have any initial dependencies (other than the reference itself)
+	if (Handle.IsDynamic())
+	{
+		UObject* Object = GetReplicatedObject(Handle);
+		if (!Object)
+		{
+			return;
+		}
+
+		if (AActor* Actor = Cast<AActor>(Object))
+		{
+			UObject* Archetype = nullptr;
+			UObject* ActorLevel = nullptr;
+
+			// ChildActor's need to be spawned from the ChildActorTemplate otherwise any non-replicated 
+			// customized properties will be incorrect on the Client.
+			if (UChildActorComponent* CAC = Actor->GetParentComponent())
+			{
+				Archetype = CAC->GetChildActorTemplate();
+			}
+			if (Archetype == nullptr)
+			{
+				Archetype = Actor->GetArchetype();
+			}
+			ActorLevel = Actor->GetLevel();
+
+			check( Archetype != nullptr );
+			check( Actor->NeedsLoadForClient() );			// We have no business sending this unless the client can load
+			check( Archetype->NeedsLoadForClient() );		// We have no business sending this unless the client can load
+
+			// Add initial dependencies
+	
+			// Archetype
+			OutDependencies.Emplace(FNetDependencyInfo(GetOrCreateObjectReference(Archetype)));
+			
+			// Level if it differs from the persistent level
+			const bool bUsePersistentLevel = NetDriver->GetWorld()->PersistentLevel == ActorLevel;
+			if (!bUsePersistentLevel)
+			{
+				OutDependencies.Emplace(FNetDependencyInfo(GetOrCreateObjectReference(ActorLevel)));
+			}
+		}
+		else
+		{
+			// Add initial dependencies
+
+			// SubObjects either have a dependency on their path relative the owner or a reference to their class
+			if (const bool bIsNameStableForNetworking = Object->IsNameStableForNetworking())
+			{
+				OutDependencies.Emplace(FNetDependencyInfo(GetOrCreateObjectReference(Object)));
+			}
+			else
+			{
+				OutDependencies.Emplace(FNetDependencyInfo(GetOrCreateObjectReference(Object->GetClass())));
+			}
+		}
+	}
+}
+
+void UActorReplicationBridge::SetNetDriver(UNetDriver* const InNetDriver)
+{
+	Super::SetNetDriver(InNetDriver);
+
+	NetDriver = InNetDriver;
+	if (InNetDriver != nullptr)
+	{
+		MaxPollFrequency = static_cast<float>(FPlatformMath::Max(InNetDriver->NetServerMaxTickRate, 0));
+	}
+}
+
+void UActorReplicationBridge::GetActorCreationHeader(const AActor* Actor, UE::Net::Private::FActorCreationHeader& Header) const
+{
+	Header.bIsActor = true;
+	
+	UE::Net::FNetObjectReference ActorReference = GetOrCreateObjectReference(Actor);
+
+	Header.bIsDynamic = ActorReference.GetRefHandle().IsDynamic();
+
+	// Dynamic actor?
+	if (Header.bIsDynamic)
+	{
+		// Get creation data
+
+		// This is more or less a straight copy from ClientPackageMap and needs to be updated accordingly
+		UObject* Archetype = nullptr;
+		UObject* ActorLevel = nullptr;
+
+		// ChildActor's need to be spawned from the ChildActorTemplate otherwise any non-replicated 
+		// customized properties will be incorrect on the Client.
+		if (UChildActorComponent* CAC = Actor->GetParentComponent())
+		{
+			Archetype = CAC->GetChildActorTemplate();
+		}
+		if (Archetype == nullptr)
+		{
+			Archetype = Actor->GetArchetype();
+		}
+		ActorLevel = Actor->GetLevel();
+
+		check(Archetype != nullptr);
+		check(Actor->NeedsLoadForClient());			// We have no business sending this unless the client can load
+		check(Archetype->NeedsLoadForClient());		// We have no business sending this unless the client can load
+
+		// Fill in Header
+		Header.ArchetypeReference = GetOrCreateObjectReference(Archetype);
+		Header.bUsePersistentLevel = NetDriver->GetWorld()->PersistentLevel == ActorLevel;
+		if (!Header.bUsePersistentLevel)
+		{
+			Header.LevelReference = GetOrCreateObjectReference(ActorLevel);
+		}
+
+		if (USceneComponent* RootComponent = Actor ? Actor->GetRootComponent() : nullptr)
+		{
+			Header.SpawnInfo.Location = FRepMovement::RebaseOntoZeroOrigin(Actor->GetActorLocation(), Actor);
+			Header.SpawnInfo.Rotation = Actor->GetActorRotation();
+			Header.SpawnInfo.Scale = Actor->GetActorScale();
+			Header.SpawnInfo.Velocity = Actor->GetVelocity();
+		}
+		else
+		{
+			Header.SpawnInfo = UE::Net::Private::GetDefaultActorReplicationBridgeSpawnInfo();
+		}
+	}
+	else // Refer by path
+	{
+		Header.ObjectReference = ActorReference;
+	}
+
+	// Custom actor creation data
+	{
+		FOutBunch Bunch(static_cast<UPackageMap*>(nullptr), 1024);
+		const_cast<AActor*>(Actor)->OnSerializeNewActor(Bunch);
+		Header.CustomCreationDataBitCount = Bunch.GetNumBits();
+		if (Header.CustomCreationDataBitCount > 0)
+		{
+			check(Header.CustomCreationDataBitCount == Bunch.GetNumBits());
+			Header.CustomCreationData.SetNumZeroed(Align(Bunch.GetNumBytes(), 4));
+			FMemory::Memcpy(Header.CustomCreationData.GetData(), Bunch.GetData(), Bunch.GetNumBytes());
+		}
+	}
+}
+
+void UActorReplicationBridge::GetSubObjectCreationHeader(const UObject* Object, UE::Net::Private::FSubObjectCreationHeader& Header) const
+{
+	// SubObject
+	UE::Net::FNetObjectReference ObjectRef = GetOrCreateObjectReference(Object);
+
+	Header.bIsActor = false;
+	Header.bIsDynamic = ObjectRef.GetRefHandle().IsDynamic();
+	Header.bUsePersistentLevel = false;
+	Header.bIsNameStableForNetworking = Object->IsNameStableForNetworking();
+
+	if (Header.bIsDynamic)
+	{
+		check(Object->NeedsLoadForClient() );				// We have no business sending this unless the client can load
+		check(Object->GetClass()->NeedsLoadForClient());	// We have no business sending this unless the client can load
+
+		// We need to spawn this object on the receiving end, so we include the class
+		if (!Header.bIsNameStableForNetworking)
+		{
+			Header.ObjectClassReference = GetOrCreateObjectReference(Object->GetClass());
+		}
+	}
+
+	Header.ObjectReference = ObjectRef;
+}
+
+bool UActorReplicationBridge::RemapPathForPIE(FString &Str, bool bReading) const
+{
+	PRAGMA_DISABLE_DEPRECATION_WARNINGS
+	return GEngine->NetworkRemapPath(NetDriver, Str, bReading);
+	PRAGMA_ENABLE_DEPRECATION_WARNINGS
+}
+
+bool UActorReplicationBridge::ObjectLevelHasFinishedLoading(UObject* Object) const
+{
+	const UWorld* DriverWorld = Object && NetDriver ? NetDriver->GetWorld() : nullptr;
+	if (DriverWorld)
+	{
+		// get the level for the object
+		const ULevel* Level = Object->GetTypedOuter<ULevel>();
+
+		if (Level != nullptr && Level != DriverWorld->PersistentLevel)
+		{
+			return Level->bIsVisible;
+		}
+	}
+
+	return true;
+}
+
+UActorReplicationBridge* UActorReplicationBridge::Create(UNetDriver* NetDriver)
+{
+	UActorReplicationBridge* Bridge = NewObject<UActorReplicationBridge>(GetTransientPackage(), UActorReplicationBridge::StaticClass());
+	Bridge->SetNetDriver(NetDriver);
+
+	return Bridge;
+}
+
+uint32 UActorReplicationBridge::GetPollFramePeriod(float PollFrequency) const
+{
+	const uint32 FramesBetweenUpdatesForObject = static_cast<uint32>(MaxPollFrequency/FPlatformMath::Max(0.001f, UE::Net::Private::PollFrequencyMultiplier*PollFrequency));
+	return FramesBetweenUpdatesForObject;
+}
+
+float UActorReplicationBridge::GetMinSupportedNetUpdateFrequency() const
+{
+	return MaxPollFrequency/(UE::Net::Private::MaxPollFramePeriod*UE::Net::Private::PollFrequencyMultiplier);
+}
+#else
+UActorReplicationBridge::UActorReplicationBridge() = default;
+#endif
diff --git a/Engine/Source/Runtime/Engine/Private/Net/Iris/ReplicationSystem/ActorReplicationBridgeInternal.cpp b/Engine/Source/Runtime/Engine/Private/Net/Iris/ReplicationSystem/ActorReplicationBridgeInternal.cpp
new file mode 100644
--- /dev/null
+++ b/Engine/Source/Runtime/Engine/Private/Net/Iris/ReplicationSystem/ActorReplicationBridgeInternal.cpp
@@ -0,0 +1,161 @@
+// Copyright Epic Games, Inc. All Rights Reserved.
+
+#include "ActorReplicationBridgeInternal.h"
+
+#if UE_WITH_IRIS
+
+#include "Iris/Serialization/NetBitStreamReader.h"
+#include "Iris/Serialization/NetBitStreamWriter.h"
+#include "Iris/Serialization/NetBitStreamUtil.h"
+#include "Iris/Serialization/ObjectNetSerializer.h"
+
+namespace UE::Net::Private
+{
+
+static const FActorReplicationBridgeSpawnInfo DefaultActorReplicationBridgeSpawnInfo = {FVector::ZeroVector, FRotator::ZeroRotator, FVector::OneVector, FVector::ZeroVector};
+
+const FActorReplicationBridgeSpawnInfo& GetDefaultActorReplicationBridgeSpawnInfo()
+{
+	return DefaultActorReplicationBridgeSpawnInfo;
+}
+
+/*
+ * IRIS: TODO: As soon as we have added code generation for Serialization code this will no longer be required
+ */
+void WriteActorReplicationBridgeSpawnInfo(FNetSerializationContext& Context, const FActorReplicationBridgeSpawnInfo& SpawnInfo)
+{
+	FNetBitStreamWriter* Writer = Context.GetBitStreamWriter();
+
+	constexpr float epsilon = 0.001f;
+
+	WriteVector(Writer, SpawnInfo.Location, DefaultActorReplicationBridgeSpawnInfo.Location, epsilon);
+	WriteRotator(Writer, SpawnInfo.Rotation, DefaultActorReplicationBridgeSpawnInfo.Rotation, epsilon);
+	WriteVector(Writer, SpawnInfo.Scale, DefaultActorReplicationBridgeSpawnInfo.Scale, epsilon);
+	WriteVector(Writer, SpawnInfo.Velocity, DefaultActorReplicationBridgeSpawnInfo.Velocity, epsilon);
+}
+
+void ReadActorReplicationBridgeSpawnInfo(FNetSerializationContext& Context, FActorReplicationBridgeSpawnInfo& SpawnInfo)
+{
+	FNetBitStreamReader* Reader = Context.GetBitStreamReader();
+
+	ReadVector(Reader, SpawnInfo.Location, DefaultActorReplicationBridgeSpawnInfo.Location);
+	ReadRotator(Reader, SpawnInfo.Rotation, DefaultActorReplicationBridgeSpawnInfo.Rotation);
+	ReadVector(Reader, SpawnInfo.Scale, DefaultActorReplicationBridgeSpawnInfo.Scale);
+	ReadVector(Reader, SpawnInfo.Velocity, DefaultActorReplicationBridgeSpawnInfo.Velocity);
+}
+
+void WriteActorCreationHeader(FNetSerializationContext& Context, const FActorCreationHeader& Header)
+{
+	FNetBitStreamWriter* Writer = Context.GetBitStreamWriter();
+
+	if (Writer->WriteBool(Header.bIsDynamic))
+	{
+		// Write Archetype and LevelPath
+		WriteFullNetObjectReference(Context, Header.ArchetypeReference);
+		
+		// Only write the LevelPath if it differs from the persistent level
+		if (!Writer->WriteBool(Header.bUsePersistentLevel))
+		{
+			WriteFullNetObjectReference(Context, Header.LevelReference);
+		}
+		
+		// Write actor spawn info
+		WriteActorReplicationBridgeSpawnInfo(Context, Header.SpawnInfo);
+	}
+	else // should be possible to refer by path??
+	{
+		WriteFullNetObjectReference(Context, Header.ObjectReference);
+	}
+
+	if (Writer->WriteBool(Header.CustomCreationDataBitCount > 0))
+	{
+		Writer->WriteBits(Header.CustomCreationDataBitCount - 1U, 16U);
+		Writer->WriteBitStream(reinterpret_cast<const uint32*>(Header.CustomCreationData.GetData()), 0U, Header.CustomCreationDataBitCount);
+	}
+}
+
+void ReadActorCreationHeader(FNetSerializationContext& Context, FActorCreationHeader& Header)
+{
+	FNetBitStreamReader* Reader = Context.GetBitStreamReader();
+
+	// Dynamic actor?
+	Header.bIsActor = true;
+	Header.bIsDynamic = Reader->ReadBool();
+
+	if (Header.bIsDynamic)
+	{
+		// Read Archetype
+		ReadFullNetObjectReference(Context, Header.ArchetypeReference);
+		
+		Header.bUsePersistentLevel = Reader->ReadBool();
+		if (!Header.bUsePersistentLevel)
+		{
+			ReadFullNetObjectReference(Context, Header.LevelReference);
+		}
+
+	 	// Read actor spawn info
+		ReadActorReplicationBridgeSpawnInfo(Context, Header.SpawnInfo);
+	}
+	else
+	{
+		ReadFullNetObjectReference(Context, Header.ObjectReference);
+	}
+
+	Header.CustomCreationDataBitCount = 0;
+	if (Reader->ReadBool())
+	{
+		Header.CustomCreationDataBitCount = 1U + Reader->ReadBits(16U);
+		Header.CustomCreationData.SetNumZeroed(((Header.CustomCreationDataBitCount + 31U) & ~31U) >> 3U);
+		Reader->ReadBitStream(reinterpret_cast<uint32*>(Header.CustomCreationData.GetData()), Header.CustomCreationDataBitCount);
+	}
+}
+
+void WriteSubObjectCreationHeader(FNetSerializationContext& Context, const FSubObjectCreationHeader& Header)
+{
+	FNetBitStreamWriter* Writer = Context.GetBitStreamWriter();
+
+	// Write required references to either find or instantiate the subobject
+	if (Writer->WriteBool(Header.bIsDynamic))
+	{
+		if (Writer->WriteBool(Header.bIsNameStableForNetworking))
+		{
+			WriteFullNetObjectReference(Context, Header.ObjectReference);
+		}
+		else
+		{
+			WriteFullNetObjectReference(Context, Header.ObjectClassReference);
+		}
+	}
+	else
+	{
+		WriteFullNetObjectReference(Context, Header.ObjectReference);
+	}
+}
+
+void ReadSubObjectCreationHeader(FNetSerializationContext& Context, FSubObjectCreationHeader& Header)
+{
+	FNetBitStreamReader* Reader = Context.GetBitStreamReader();
+
+	Header.bIsActor = false;
+	Header.bIsDynamic = Reader->ReadBool();
+	if (Header.bIsDynamic)
+	{
+		Header.bIsNameStableForNetworking = Reader->ReadBool();
+		if (Header.bIsNameStableForNetworking)
+		{
+			ReadFullNetObjectReference(Context, Header.ObjectReference);
+		}
+		else
+		{
+			ReadFullNetObjectReference(Context, Header.ObjectClassReference);
+		}
+	}
+	else
+	{
+		ReadFullNetObjectReference(Context, Header.ObjectReference);
+	}
+}
+
+}
+
+#endif // UE_WITH_IRIS
diff --git a/Engine/Source/Runtime/Engine/Private/Net/Iris/ReplicationSystem/ActorReplicationBridgeInternal.h b/Engine/Source/Runtime/Engine/Private/Net/Iris/ReplicationSystem/ActorReplicationBridgeInternal.h
new file mode 100644
--- /dev/null
+++ b/Engine/Source/Runtime/Engine/Private/Net/Iris/ReplicationSystem/ActorReplicationBridgeInternal.h
@@ -0,0 +1,61 @@
+// Copyright Epic Games, Inc. All Rights Reserved.
+
+#pragma once
+
+#if UE_WITH_IRIS
+
+#include "Iris/ReplicationSystem/ObjectReplicationBridge.h"
+
+namespace UE::Net
+{
+	class FNetSerializationContext;
+}
+
+namespace UE::Net::Private
+{
+
+struct FActorReplicationBridgeSpawnInfo
+{
+	FVector Location;
+	FRotator Rotation;
+	FVector Scale;
+	FVector Velocity;
+};
+
+struct FActorReplicationBridgeCreationHeader : public UObjectReplicationBridge::FCreationHeader
+{
+	FNetObjectReference ObjectReference;
+	uint8 bIsActor : 1;
+	uint8 bIsDynamic : 1;
+	uint8 bIsNameStableForNetworking : 1;
+	uint8 bUsePersistentLevel : 1;
+};
+
+struct FActorCreationHeader : public FActorReplicationBridgeCreationHeader
+{
+	FActorReplicationBridgeSpawnInfo SpawnInfo;
+	FNetObjectReference ArchetypeReference;
+	FNetObjectReference LevelReference;
+	TArray<uint8> CustomCreationData;
+	uint16 CustomCreationDataBitCount;
+};
+
+struct FSubObjectCreationHeader : public FActorReplicationBridgeCreationHeader
+{
+	FNetObjectReference ObjectClassReference;
+};
+
+const FActorReplicationBridgeSpawnInfo& GetDefaultActorReplicationBridgeSpawnInfo();
+
+void WriteActorReplicationBridgeSpawnInfo(FNetSerializationContext& Context, const FActorReplicationBridgeSpawnInfo& SpawnInfo);
+void ReadActorReplicationBridgeSpawnInfo(FNetSerializationContext& Context, FActorReplicationBridgeSpawnInfo& SpawnInfo);
+
+void WriteActorCreationHeader(FNetSerializationContext& Context, const FActorCreationHeader& Header);
+void ReadActorCreationHeader(FNetSerializationContext& Context, FActorCreationHeader& Header);
+
+void WriteSubObjectCreationHeader(FNetSerializationContext& Context, const FSubObjectCreationHeader& Header);
+void ReadSubObjectCreationHeader(FNetSerializationContext& Context, FSubObjectCreationHeader& Header); 
+
+}
+
+#endif // UE_WITH_IRIS
diff --git a/Engine/Source/Runtime/Engine/Private/Net/Iris/ReplicationSystem/ReplicationSystemUtil.cpp b/Engine/Source/Runtime/Engine/Private/Net/Iris/ReplicationSystem/ReplicationSystemUtil.cpp
new file mode 100644
--- /dev/null
+++ b/Engine/Source/Runtime/Engine/Private/Net/Iris/ReplicationSystem/ReplicationSystemUtil.cpp
@@ -0,0 +1,398 @@
+// Copyright Epic Games, Inc. All Rights Reserved.
+
+#include "Net/Iris/ReplicationSystem/ReplicationSystemUtil.h"
+
+#if UE_WITH_IRIS
+
+#include "EngineUtils.h"
+#include "Engine/NetConnection.h"
+#include "Engine/NetDriver.h"
+#include "Engine/World.h"
+#include "Iris/Core/IrisProfiler.h"
+#include "Iris/ReplicationSystem/Filtering/NetObjectFilter.h"
+#include "Iris/ReplicationSystem/ReplicationSystem.h"
+#include "Iris/Core/IrisLog.h"
+#include "Net/Iris/ReplicationSystem/ActorReplicationBridge.h"
+#include "Net/NetSubObjectRegistryGetter.h"
+#include "HAL/IConsoleManager.h"
+#include "Templates/Casts.h"
+#include "Net/Subsystems/NetworkSubsystem.h"
+#include "Net/Core/Misc/NetConditionGroupManager.h"
+#include "GameFramework/PlayerController.h"
+
+namespace UE::Net
+{
+
+UReplicationSystem* FReplicationSystemUtil::GetReplicationSystem(const AActor* Actor)
+{
+	UNetDriver* NetDriver = Actor ? Actor->GetNetDriver() : nullptr;
+	return NetDriver ? NetDriver->GetReplicationSystem() : nullptr;
+}
+
+UActorReplicationBridge* FReplicationSystemUtil::GetActorReplicationBridge(const AActor* Actor)
+{
+	if (UReplicationSystem* ReplicationSystem = GetReplicationSystem(Actor))
+	{
+		return Cast<UActorReplicationBridge>(ReplicationSystem->GetReplicationBridge());
+	}
+	else
+	{
+		return nullptr;
+	}
+}
+
+UActorReplicationBridge* FReplicationSystemUtil::GetActorReplicationBridge(const UNetConnection* NetConnection)
+{
+	const UNetDriver* Driver = NetConnection ? NetConnection->GetDriver() : nullptr;
+	if (const UReplicationSystem* ReplicationSystem = Driver ? Driver->GetReplicationSystem() : nullptr)
+	{
+		return Cast<UActorReplicationBridge>(ReplicationSystem->GetReplicationBridge());
+	}
+	else
+	{
+		return nullptr;
+	}
+}
+
+FNetHandle FReplicationSystemUtil::GetNetHandle(const AActor* Actor)
+{
+	// We really should cache the handle in the actor (or UObject if we can afford it)
+	if (!Actor || !Actor->GetIsReplicated())
+	{
+		return FNetHandle();
+	}
+
+	UActorReplicationBridge* Bridge = GetActorReplicationBridge(Actor);
+	return Bridge ? Bridge->GetReplicatedHandle(Actor) : FNetHandle();
+}
+
+FNetHandle FReplicationSystemUtil::GetNetHandle(const UActorComponent* SubObject)
+{
+	UActorReplicationBridge* Bridge = SubObject ? GetActorReplicationBridge(SubObject->GetOwner()) : nullptr;
+	return Bridge ? Bridge->GetReplicatedHandle(SubObject) : FNetHandle();
+}
+
+FNetHandle FReplicationSystemUtil::BeginReplicationForActorComponent(FNetHandle ActorHandle, UActorComponent* ActorComp)
+{
+	if (ActorHandle.IsValid())
+	{
+		UReplicationSystem* ReplicationSystem = Net::GetReplicationSystem(ActorHandle.GetReplicationSystemId());
+		if (ReplicationSystem && ReplicationSystem->IsServer())
+		{
+			if (UActorReplicationBridge* Bridge = Cast<UActorReplicationBridge>(ReplicationSystem->GetReplicationBridge()))
+			{
+				return Bridge->BeginReplication(ActorHandle, ActorComp);
+			}
+		}
+	}
+	return FNetHandle();
+}
+
+FNetHandle FReplicationSystemUtil::BeginReplicationForActorComponent(const AActor* Actor, UActorComponent* ActorComp)
+{
+	return FReplicationSystemUtil::BeginReplicationForActorComponent(FReplicationSystemUtil::GetNetHandle(Actor), ActorComp);
+}
+
+FNetHandle FReplicationSystemUtil::BeginReplicationForActorSubObject(const AActor* Actor, UObject* ActorSubObject, ELifetimeCondition NetCondition)
+{
+	if (NetCondition != ELifetimeCondition::COND_Never)
+	{
+		const FNetHandle OwnerHandle = GetNetHandle(Actor);
+		if (OwnerHandle.IsValid())
+		{
+			UReplicationSystem* ReplicationSystem = Net::GetReplicationSystem(OwnerHandle.GetReplicationSystemId());
+			if (ReplicationSystem && ReplicationSystem->IsServer())
+			{
+				if (UObjectReplicationBridge* Bridge = ReplicationSystem->GetReplicationBridgeAs<UObjectReplicationBridge>())
+				{
+					const FNetHandle SubObjectHandle = Bridge->BeginReplication(OwnerHandle, ActorSubObject);
+					if (SubObjectHandle.IsValid() && NetCondition != ELifetimeCondition::COND_None)
+					{
+						Bridge->SetSubObjectNetCondition(SubObjectHandle, NetCondition);
+					}
+					return SubObjectHandle;
+				}
+			}
+		}
+	}
+
+	return FNetHandle();
+}
+
+FNetHandle FReplicationSystemUtil::BeginReplicationForActorComponentSubObject(UActorComponent* ActorComponent, UObject* SubObject, ELifetimeCondition NetCondition)
+{
+	const AActor* Actor = ActorComponent->GetOwner();
+	if (Actor && NetCondition != ELifetimeCondition::COND_Never)
+	{		
+		const FNetHandle OwnerHandle = GetNetHandle(ActorComponent->GetOwner());
+		if (OwnerHandle.IsValid())
+		{
+			UReplicationSystem* ReplicationSystem = Net::GetReplicationSystem(OwnerHandle.GetReplicationSystemId());
+			if (ReplicationSystem && ReplicationSystem->IsServer())
+			{
+				if (UObjectReplicationBridge* Bridge = ReplicationSystem->GetReplicationBridgeAs<UObjectReplicationBridge>())
+				{
+					const FNetHandle ActorComponentHandle = Bridge->GetReplicatedHandle(ActorComponent);
+					if (ActorComponentHandle.IsValid())
+					{
+						const FNetHandle SubObjectHandle = Bridge->BeginReplication(OwnerHandle, SubObject, ActorComponentHandle, UReplicationBridge::ESubObjectInsertionOrder::ReplicateWith);
+						if (NetCondition != ELifetimeCondition::COND_None)
+						{
+							Bridge->SetSubObjectNetCondition(SubObjectHandle, NetCondition);
+						}
+						return SubObjectHandle;
+					}
+				}
+			}
+		}
+	}
+
+	return FNetHandle();
+}
+
+void FReplicationSystemUtil::EndReplicationForActorSubObject(const AActor* Actor, UObject* SubObject)
+{
+	if (UObjectReplicationBridge* Bridge = Cast<UObjectReplicationBridge>(GetActorReplicationBridge(Actor)))
+	{
+		Bridge->EndReplication(SubObject);
+	}
+}
+
+void FReplicationSystemUtil::EndReplicationForActorComponentSubObject(UActorComponent* ActorComponent, UObject* SubObject)
+{
+	EndReplicationForActorSubObject(ActorComponent->GetOwner(), SubObject);
+}
+
+FNetHandle FReplicationSystemUtil::BeginReplication(AActor* Actor, const FActorBeginReplicationParams& Params)
+{
+	UReplicationSystem* ReplicationSystem = GetReplicationSystem(Actor);
+	if (ReplicationSystem && ReplicationSystem->IsServer())
+	{
+		if (UActorReplicationBridge* Bridge = ReplicationSystem->GetReplicationBridgeAs<UActorReplicationBridge>())
+		{
+			return Bridge->BeginReplication(Actor, Params);
+		}
+	}
+
+	return FNetHandle();
+}
+
+FNetHandle FReplicationSystemUtil::BeginReplication(AActor* Actor)
+{
+	const FActorBeginReplicationParams BeginReplicationParams;
+	return BeginReplication(Actor, BeginReplicationParams);
+}
+
+void FReplicationSystemUtil::AddDependentActor(const AActor* Parent, AActor* Child)
+{
+	// Can only add dependent actors on already replicating actors
+	FNetHandle ParentHandle = GetNetHandle(Parent);
+	if (ensureMsgf(ParentHandle.IsValid(), TEXT("FReplicationSystemUtil::AddDependentActor Parent %s must be replicated"), *GetPathNameSafe(Parent)))
+	{
+		FNetHandle ChildHandle = BeginReplication(Child);
+		if (ensureMsgf(ChildHandle.IsValid(), TEXT("FReplicationSystemUtil::AddDependentActor Child %s must be replicated"), *GetPathNameSafe(Child)))
+		{
+			UReplicationSystem* ReplicationSystem = Net::GetReplicationSystem(ParentHandle.GetReplicationSystemId());
+			if (UObjectReplicationBridge* Bridge = ReplicationSystem ? ReplicationSystem->GetReplicationBridgeAs<UObjectReplicationBridge>() : nullptr)
+			{
+				Bridge->AddDependentObject(ParentHandle, ChildHandle);
+			}
+		}
+	}
+}
+
+void FReplicationSystemUtil::RemoveDependentActor(const AActor* Parent, AActor* Child)
+{
+	// Need to find replication system
+	FNetHandle ParentHandle = GetNetHandle(Parent);
+	FNetHandle ChildHandle = GetNetHandle(Child);
+	if (ParentHandle.IsValid() && ChildHandle.IsValid())
+	{
+		UReplicationSystem* ReplicationSystem = Net::GetReplicationSystem(ChildHandle.GetReplicationSystemId());
+		if (UObjectReplicationBridge* Bridge = ReplicationSystem ? ReplicationSystem->GetReplicationBridgeAs<UObjectReplicationBridge>() : nullptr)
+		{
+			Bridge->RemoveDependentObject(ParentHandle, ChildHandle);
+		}
+	}
+}
+
+void FReplicationSystemUtil::EndReplication(AActor* Actor, EEndPlayReason::Type EndPlayReason)
+{	
+	if (UActorReplicationBridge* Bridge = GetActorReplicationBridge(Actor))
+	{
+		Bridge->EndReplication(Actor, EndPlayReason);
+	}
+}
+
+void FReplicationSystemUtil::EndReplicationForActorComponent(UActorComponent* SubObject)
+{
+	if (UActorReplicationBridge* Bridge = GetActorReplicationBridge(SubObject->GetOwner()))
+	{
+		Bridge->EndReplicationForActorComponent(SubObject);
+	}
+}
+
+void FReplicationSystemUtil::SetActorComponentNetCondition(const UActorComponent* ActorComponent, ELifetimeCondition Condition)
+{
+	if (UActorReplicationBridge* Bridge = IsValid(ActorComponent) ? GetActorReplicationBridge(ActorComponent->GetOwner()) : nullptr)
+	{
+		Bridge->SetSubObjectNetCondition(Bridge->GetReplicatedHandle(ActorComponent), Condition);
+	}
+}
+
+void FReplicationSystemUtil::SetPropertyCustomCondition(const UObject* Object, uint16 RepIndex, bool bIsActive)
+{
+	// If this is an actor component we can get the proper bridge through its owner, regardless of whether it's a default replicated component or not.
+	if (const UActorComponent* ActorComponent = Cast<UActorComponent>(Object))
+	{
+		const AActor* Owner = ActorComponent->GetOwner();
+		const UActorReplicationBridge* Bridge = GetActorReplicationBridge(Owner);
+		if (Bridge == nullptr)
+		{
+			return;
+		}
+
+		// Try to get a handle for the component. This can fail if we incorporate default components in the protocol for the actor.
+		FNetHandle NetHandle = Bridge->GetReplicatedHandle(Object);
+		if (!NetHandle.IsValid())
+		{
+			NetHandle = Bridge->GetReplicatedHandle(Owner);
+			if (!NetHandle.IsValid())
+			{
+				return;
+			}
+		}
+
+		Net::GetReplicationSystem(NetHandle.GetReplicationSystemId())->SetPropertyCustomCondition(NetHandle, Object, RepIndex, bIsActive);
+	}
+	else if (const AActor* Actor = Cast<AActor>(Object))
+	{
+		FNetHandle NetHandle = GetNetHandle(Actor);
+		if (!NetHandle.IsValid())
+		{
+			return;
+		}
+
+		Net::GetReplicationSystem(NetHandle.GetReplicationSystemId())->SetPropertyCustomCondition(NetHandle, Object, RepIndex, bIsActive);
+	}
+	else
+	{
+		checkf(false, TEXT("Unexpected class %s trying to set custom condition."), ToCStr(Object->GetFName().GetPlainNameString()));
+	}
+}
+
+void FReplicationSystemUtil::BeginReplicationForActorsInWorld(UWorld* World)
+{
+	// We only do this if the world already is initialized
+	// Normally we rely on Begin/EndPlay to control if an actor is replicated or not.
+	if (World && World->bIsWorldInitialized)
+	{
+		for (FActorIterator Iter(World); Iter; ++Iter)
+		{
+			AActor* Actor = *Iter;
+			if (IsValid(Actor) && Actor->HasActorBegunPlay() && ULevel::IsNetActor(Actor))
+			{
+				Actor->BeginReplication();
+			}
+		}
+	}
+}
+
+void FReplicationSystemUtil::NotifyActorDormancyChange(AActor* Actor, ENetDormancy OldDormancyState)
+{
+	const ENetDormancy Dormancy = Actor->NetDormancy;
+	const bool bIsPendingDormancy = (Dormancy > DORM_Awake);
+
+	FNetHandle NetHandle = GetNetHandle(Actor);
+	if (NetHandle.IsValid())
+	{
+		UReplicationSystem* ReplicationSystem = Net::GetReplicationSystem(NetHandle.GetReplicationSystemId());
+		if (UObjectReplicationBridge* Bridge = ReplicationSystem ? ReplicationSystem->GetReplicationBridgeAs<UObjectReplicationBridge>() : nullptr)
+		{
+			Bridge->SetObjectWantsToBeDormant(NetHandle, bIsPendingDormancy);
+		}
+	}
+}
+
+void FReplicationSystemUtil::FlushNetDormancy(AActor* Actor, bool bWasDormInitial)
+{
+	FNetHandle NetHandle = GetNetHandle(Actor);
+	if (NetHandle.IsValid())
+	{
+		UReplicationSystem* ReplicationSystem = Net::GetReplicationSystem(NetHandle.GetReplicationSystemId());
+		if (UObjectReplicationBridge* Bridge = ReplicationSystem ? ReplicationSystem->GetReplicationBridgeAs<UObjectReplicationBridge>() : nullptr)
+		{
+			Bridge->ForceUpdateWantsToBeDormantObject(NetHandle);
+		}
+	}
+	else
+	{
+		UE_CLOG(!Actor->GetIsReplicated(), LogIris, Warning, TEXT("FReplicationSystemUtil::FlushNetDormancy Actor %s that is not replicated"), ToCStr(Actor->GetName()));
+		UE_CLOG(!bWasDormInitial, LogIris, Display, TEXT("FReplicationSystemUtil::FlushNetDormancy For not replicated Actor %s is not initially dormant"), ToCStr(Actor->GetName()));
+
+		if (Actor->GetIsReplicated())
+		{
+ 			Actor->BeginReplication();
+		}
+	}
+}
+
+void FReplicationSystemUtil::UpdateSubObjectGroupMemberships(const UObject* SubObject, const UWorld* World)
+{
+	if (const UNetworkSubsystem* NetSubsystem = World ? World->GetSubsystem<UNetworkSubsystem>() : nullptr)
+	{
+		const UNetDriver* NetDriver = World->GetNetDriver();
+		if (UReplicationSystem* ReplicationSystem = NetDriver ? NetDriver->GetReplicationSystem() : nullptr)
+		{
+			if (UActorReplicationBridge* Bridge = ReplicationSystem->GetReplicationBridgeAs<UActorReplicationBridge>())
+			{
+				FNetHandle Handle = Bridge->GetReplicatedHandle(SubObject);
+				if (Handle.IsValid())
+				{
+					for (const FName NetGroup : NetSubsystem->GetNetConditionGroupManager().GetSubObjectNetConditionGroups(FObjectKey(SubObject)))
+					{
+						FNetObjectGroupHandle SubObjectGroupHandle = ReplicationSystem->GetOrCreateSubObjectFilter(NetGroup);
+						ReplicationSystem->AddToGroup(SubObjectGroupHandle, Handle);
+					}
+				}
+			}
+		}
+	}
+}
+
+void FReplicationSystemUtil::RemoveSubObjectGroupMembership(const APlayerController* PC, const FName NetGroup)
+{
+	UNetConnection* Conn = PC->GetNetConnection();
+	if (ensureAlways(Conn))
+	{
+		if (UReplicationSystem* ReplicationSystem = Conn->GetDriver() ? Conn->GetDriver()->GetReplicationSystem() : nullptr)
+		{
+			if (!IsSpecialNetConditionGroup(NetGroup))
+			{
+				ReplicationSystem->SetSubObjectFilterStatus(NetGroup, Conn->GetConnectionId(), ENetFilterStatus::Disallow);
+			}
+		}
+	}
+}
+
+void FReplicationSystemUtil::UpdateSubObjectGroupMemberships(const APlayerController* PC)
+{
+	if (UNetConnection* Conn = PC->GetNetConnection())
+	{
+		if (UReplicationSystem* ReplicationSystem = Conn->GetDriver() ? Conn->GetDriver()->GetReplicationSystem() : nullptr)
+		{
+			for (const FName NetGroup : PC->GetNetConditionGroups())
+			{
+				if (!IsSpecialNetConditionGroup(NetGroup))
+				{
+					FNetObjectGroupHandle SubObjectGroupHandle = ReplicationSystem->GetOrCreateSubObjectFilter(NetGroup);
+					ReplicationSystem->SetSubObjectFilterStatus(NetGroup, Conn->GetConnectionId(), ENetFilterStatus::Allow);
+				}
+			}
+		}
+	}
+}
+
+}
+
+#endif
diff --git a/Engine/Source/Runtime/Engine/Private/NetConnection.cpp b/Engine/Source/Runtime/Engine/Private/NetConnection.cpp
--- a/Engine/Source/Runtime/Engine/Private/NetConnection.cpp
+++ b/Engine/Source/Runtime/Engine/Private/NetConnection.cpp
@@ -42,11 +42,16 @@
 #include "HAL/LowLevelMemStats.h"
 #include "Net/NetPing.h"
 #include "LevelUtils.h"
+#if UE_WITH_IRIS
+#include "Iris/IrisConfig.h"
+#include "Iris/ReplicationSystem/ReplicationSystem.h"
+#include "Iris/ReplicationSystem/Filtering/NetObjectFilter.h"
+#include "Net/Iris/ReplicationSystem/ActorReplicationBridge.h"
+#endif // UE_WITH_IRIS
 
 DECLARE_LLM_MEMORY_STAT(TEXT("NetConnection"), STAT_NetConnectionLLM, STATGROUP_LLMFULL);
 LLM_DEFINE_TAG(NetConnection, NAME_None, TEXT("Networking"), GET_STATFNAME(STAT_NetConnectionLLM), GET_STATFNAME(STAT_NetworkingSummaryLLM));
 
-
 static TAutoConsoleVariable<int32> CVarPingExcludeFrameTime(TEXT("net.PingExcludeFrameTime"), 0,
 	TEXT("If true, game frame times are subtracted from calculated ping to approximate actual network ping"));
 
@@ -100,7 +105,6 @@ static int32 GNetCloseTimingDebug = 0;
 static FAutoConsoleVariableRef CVarCloseTimingDebug(TEXT("net.CloseTimingDebug"), GNetCloseTimingDebug,
 	TEXT("Logs the last packet send/receive and TickFlush/TickDispatch times, on connection close - for debugging blocked send/recv paths."));
 
-
 extern int32 GNetDormancyValidate;
 extern bool GbNetReuseReplicatorsForDormantObjects;
 
@@ -114,6 +118,7 @@ namespace UE_NetConnectionPrivate
 	struct FValidateLevelVisibilityResult
 	{
 		bool bLevelExists = false;
+		const ULevelStreaming* StreamingLevel;
 		UPackage* Package = nullptr;
 		FLinkerLoad* Linker = nullptr;
 	};
@@ -130,10 +135,10 @@ namespace UE_NetConnectionPrivate
 		Result.Package = FindPackage(nullptr, *PackageNameStr);
 		Result.Linker = FLinkerLoad::FindExistingLinkerForPackage(Result.Package);
 
-		const ULevelStreaming* StreamingLevel = FLevelUtils::FindStreamingLevel(World, LevelVisibility.PackageName);
+		Result.StreamingLevel = FLevelUtils::FindStreamingLevel(World, LevelVisibility.PackageName);
 		Result.bLevelExists = Result.Linker ||
 			(!FPackageName::IsMemoryPackage(PackageNameStr) && FPackageName::DoesPackageExist(LevelVisibility.FileName.ToString())) ||
-			StreamingLevel != nullptr;
+			Result.StreamingLevel != nullptr;
 
 		return Result;
 	}
@@ -507,6 +512,14 @@ void UNetConnection::InitBase(UNetDriver* InDriver,class FSocket* InSocket, cons
 
 	UE_NET_TRACE_CONNECTION_CREATED(NetTraceId, GetConnectionId());
 	UE_NET_TRACE_CONNECTION_STATE_UPDATED(NetTraceId, GetConnectionId(), static_cast<uint8>(GetConnectionState()));
+
+#if UE_WITH_IRIS
+	if (UReplicationSystem* ReplicationSystem = Driver->GetReplicationSystem())
+	{
+		ReplicationSystem->AddConnection(GetConnectionId());
+		ReplicationSystem->SetConnectionUserData(GetConnectionId(), this);
+	}
+#endif // UE_WITH_IRIS
 }
 
 /**
@@ -1170,6 +1183,12 @@ void UNetConnection::CleanUp()
 		// Otherwise we'd be able to do the appropriate logic in Add/Remove Client/ServerConnection
 		if (const uint32 MyConnectionId = GetConnectionId())
 		{
+#if UE_WITH_IRIS
+			if (UReplicationSystem* ReplicationSystem = Driver->GetReplicationSystem())
+			{
+				ReplicationSystem->RemoveConnection(MyConnectionId);
+			}
+#endif // UE_WITH_IRIS
 			Driver->FreeConnectionId(MyConnectionId);
 		}
 	}
@@ -1340,7 +1359,24 @@ FNetLevelVisibilityTransactionId UNetConnection::UpdateLevelStreamStatusChangedT
 	if (bShouldBeVisible == false)
 	{
 		ClientVisibleLevelNames.Remove(PackageName);
-		UpdateCachedLevelVisibility(PackageName);
+
+#if UE_WITH_IRIS
+		if (UReplicationSystem* ReplicationSystem = Driver->GetReplicationSystem())
+		{
+			if (const ULevel* Level = LevelObject->GetLoadedLevel())
+			{
+				const UReplicationBridge* ReplicationBridge = ReplicationSystem->GetReplicationBridge();
+				if (UE::Net::FNetObjectGroupHandle GroupHandle = ReplicationBridge->GetLevelGroup(Level); GroupHandle != UE::Net::InvalidNetObjectGroupHandle)
+				{
+					ReplicationSystem->SetGroupFilterStatus(GroupHandle, GetConnectionId(), UE::Net::ENetFilterStatus::Disallow);
+				}
+			}
+		}
+		else
+#endif
+		{
+			UpdateCachedLevelVisibility(PackageName);
+		}
 	}
 
 	return TransactionId;
@@ -1463,6 +1499,24 @@ void UNetConnection::UpdateLevelVisibilityInternal(const FUpdateLevelVisibilityL
 			// Any dormant actor that has changes flushed or made before going dormant needs to be updated on the client 
 			// when the streaming level is loaded, so mark them active for this connection
 			UWorld* LevelWorld = nullptr;
+
+#if UE_WITH_IRIS
+			if (UReplicationSystem* ReplicationSystem = Driver->GetReplicationSystem())
+			{
+				// We need to be careful here if the server would load a level after the client, 
+				// for that to work we would need to update group filter status after level is loaded
+				const ULevel* Level = VisibilityResult.StreamingLevel ? VisibilityResult.StreamingLevel->GetLoadedLevel() : nullptr;
+				if (Level)
+				{
+					const UReplicationBridge* ReplicationBridge = ReplicationSystem->GetReplicationBridge();
+					if (UE::Net::FNetObjectGroupHandle GroupHandle = ReplicationBridge->GetLevelGroup(Level); GroupHandle != UE::Net::InvalidNetObjectGroupHandle)
+					{
+						ReplicationSystem->SetGroupFilterStatus(GroupHandle, GetConnectionId(), UE::Net::ENetFilterStatus::Allow);
+					}
+				}
+			}
+			else 
+#endif // UE_WITH_IRIS
 			if (VisibilityResult.Package)
 			{
 				LevelWorld = (UWorld*)FindObjectWithOuter(VisibilityResult.Package, UWorld::StaticClass());
@@ -1518,6 +1572,27 @@ void UNetConnection::UpdateLevelVisibilityInternal(const FUpdateLevelVisibilityL
 			ReplicationConnectionDriver->NotifyClientVisibleLevelNamesRemove(LevelVisibility.PackageName);
 		}
 
+#if UE_WITH_IRIS
+		if (UReplicationSystem* ReplicationSystem = Driver->GetReplicationSystem())
+		{
+			// See if the level is loaded, if it is disable filtering, if not filter should already be removed
+			if (const ULevelStreaming* StreamingLevel = FLevelUtils::FindStreamingLevel(GetWorld(), LevelVisibility.PackageName))
+			{
+				if (const ULevel* Level = StreamingLevel->GetLoadedLevel())
+				{
+					// $IRIS: $TODO: Implement forced scope-update for a group/connection to ensure that we treat replicated objects as destroyed if the filter status of a level is disabled/enabled on the same frame
+					// The reason for this is that the client currently destroys the instances rather then managing this through the replication system
+					// If we implement a way to re-instantiate instances on the client we might be able to persist the state
+					const UReplicationBridge* ReplicationBridge = ReplicationSystem->GetReplicationBridge();
+					if (UE::Net::FNetObjectGroupHandle GroupHandle = ReplicationBridge->GetLevelGroup(Level); GroupHandle != UE::Net::InvalidNetObjectGroupHandle)
+					{
+						ReplicationSystem->SetGroupFilterStatus(GroupHandle, GetConnectionId(), UE::Net::ENetFilterStatus::Disallow);
+					}
+				}
+			}
+		}
+#endif // UE_WITH_IRIS
+
 		// Close any channels now that have actors that were apart of the level the client just unloaded
 		for (auto It = ActorChannels.CreateIterator(); It; ++It)
 		{
@@ -2777,6 +2852,8 @@ void UNetConnection::ReceivedPacket( FBitReader& Reader, bool bIsReinjectedPacke
 
 	bool bSkipAck = false;
 
+	bool bHasBunchErrors = false;
+
 	// Track channels that were rejected while processing this packet - used to avoid sending multiple close-channel bunches,
 	// which would cause a disconnect serverside
 	TArray<int32> RejectedChans;
@@ -3355,16 +3432,24 @@ void UNetConnection::ReceivedPacket( FBitReader& Reader, bool bIsReinjectedPacke
 				Driver->InBunches++;
 				Driver->InTotalBunches++;
 
-				// Disconnect if we received a corrupted packet from the client (eg server crash attempt).
-				if (bIsServer && (Bunch.IsCriticalError() || Bunch.IsError()))
+				if (Bunch.IsCriticalError() || Bunch.IsError())
 				{
-					UE_LOG(LogNetTraffic, Error, TEXT("Received corrupted packet data from client %s. Disconnecting."), *LowLevelGetRemoteAddress());
+					bHasBunchErrors = true;
 
-					Close(AddToAndConsumeChainResultPtr(Bunch.ExtendedError, ENetCloseResult::CorruptData));
+					// Disconnect if we received a corrupted packet from the client (eg server crash attempt).
+					if (bIsServer)
+					{
+						UE_LOG(LogNetTraffic, Error, TEXT("Received corrupted packet data with SequenceId: %d from client %s. Disconnecting."), InPacketId, *LowLevelGetRemoteAddress());
 
-					return;
-				}
+						Close(AddToAndConsumeChainResultPtr(Bunch.ExtendedError, ENetCloseResult::CorruptData));
 
+						return;
+					}
+					else
+					{
+						UE_LOG(LogNetTraffic, Error, TEXT("Received corrupted packet data with SequenceId: %d from server %s"), InPacketId, *LowLevelGetRemoteAddress());
+					}
+				}
 				// In replay, if the bunch generated an error but the channel isn't actually open, clean it up so the channel index remains free
 				if (IsInternalAck() && Bunch.IsError() && Channel && !Channel->OpenedLocally && !Channel->OpenAcked)
 				{
@@ -3441,9 +3526,9 @@ void UNetConnection::ReceivedPacket( FBitReader& Reader, bool bIsReinjectedPacke
 	// Trace end marker of this incoming data packet.
 	UE_NET_TRACE_PACKET_RECV(NetTraceId, GetConnectionId(), InPacketId, Reader.GetNumBits());
 
-	if (bSkipAck && !IsInternalAck())
+	if ((bHasBunchErrors || bSkipAck) && !IsInternalAck())
 	{
-		// Trace packet as lost on receiving end to indicate bSkipAck in traced data
+		// Trace packet as lost on receiving end to indicate bSkipAck or errors in traced data
 		UE_NET_TRACE_PACKET_DROPPED(NetTraceId, GetConnectionId(), InPacketId, ENetTracePacketType::Incoming);
 	}
 }
@@ -4427,6 +4512,16 @@ void UNetConnection::HandleClientPlayer( APlayerController *PC, UNetConnection*
 	PlayerController = PC;
 	OwningActor = PC;
 
+#if UE_WITH_IRIS
+	{
+		// Enable replication
+		if (UReplicationSystem* ReplicationSystem = Driver->GetReplicationSystem())
+		{
+			ReplicationSystem->SetReplicationEnabledForConnection(GetConnectionId(), true);
+		}
+	}
+#endif // UE_WITH_IRIS
+
 	UWorld* World = PlayerController->GetWorld();
 	// if we have already loaded some sublevels, tell the server about them
 	{
@@ -4698,6 +4793,17 @@ void UNetConnection::SetClientLoginState( const EClientLoginState::Type NewState
 				TEXT("UNetConnection::SetClientLoginState: State changing from %s to %s"),
 				EClientLoginState::ToString(ClientLoginState), EClientLoginState::ToString(NewState));
 
+#if UE_WITH_IRIS
+	if (UReplicationSystem* ReplicationSystem = Driver->GetReplicationSystem())
+	{
+		if (NewState == EClientLoginState::ReceivedJoin)
+		{
+			// Enable replication
+			ReplicationSystem->SetReplicationEnabledForConnection(GetConnectionId(), true);
+		}
+	}
+#endif // UE_WITH_IRIS
+
 	ClientLoginState = NewState;
 }
 
diff --git a/Engine/Source/Runtime/Engine/Private/NetDriver.cpp b/Engine/Source/Runtime/Engine/Private/NetDriver.cpp
--- a/Engine/Source/Runtime/Engine/Private/NetDriver.cpp
+++ b/Engine/Source/Runtime/Engine/Private/NetDriver.cpp
@@ -60,6 +60,7 @@
 #include "Net/Core/Trace/NetTrace.h"
 #include "Net/Core/PropertyConditions/PropertyConditions.h"
 #include "Misc/ScopeExit.h"
+#include "Misc/NetworkVersion.h"
 #include "Net/DataChannel.h"
 #include "GameFramework/PlayerState.h"
 #include "Net/PerfCountersHelpers.h"
@@ -72,6 +73,16 @@
 #include "Net/NetworkGranularMemoryLogging.h"
 #include "SocketSubsystem.h"
 #include "AddressInfoTypes.h"
+#if UE_WITH_IRIS
+#include "Iris/IrisConfig.h"
+#include "Iris/Core/IrisProfiler.h"
+#include "Iris/ReplicationSystem/ReplicationSystem.h"
+#include "Iris/ReplicationSystem/ReplicationView.h"
+#include "Iris/ReplicationSystem/Filtering/NetObjectFilter.h"
+#include "Net/Iris/ReplicationSystem/ActorReplicationBridge.h"
+#include "Net/Iris/ReplicationSystem/ReplicationSystemUtil.h"
+#endif // UE_WITH_IRIS
+
 #if USE_SERVER_PERF_COUNTERS
 #include "PerfCountersModule.h"
 #endif
@@ -424,6 +435,9 @@ UNetDriver::UNetDriver(const FObjectInitializer& ObjectInitializer)
 ,	ProcessQueuedBunchesCurrentFrameMilliseconds(0.0f)
 ,	DDoS()
 ,	LocalAddr(nullptr)
+#if UE_WITH_IRIS
+,	DummyNetworkObjectInfo(MakePimpl<FNetworkObjectInfo>())
+#endif // UE_WITH_IRIS
 ,	NetworkObjects(new FNetworkObjectList)
 ,	LagState(ENetworkLagState::NotLagging)
 ,	DuplicateLevelID(INDEX_NONE)
@@ -567,7 +581,7 @@ void UNetDriver::NotifyGameInstanceUpdated()
 	if (GetWorld() && GetWorld()->WorldType)
 	{
 		FString InstanceName = FString::Printf(TEXT("%s (%s)"), *GetNameSafe(GetWorld()->GetGameInstance()), LexToString(GetWorld()->WorldType.GetValue()));
-		UE_NET_TRACE_UPDATE_INSTANCE(GetUniqueID(), IsServer(), *InstanceName);
+		UE_NET_TRACE_UPDATE_INSTANCE(GetNetTraceId(), IsServer(), *InstanceName);
 	}
 #endif
 }
@@ -587,6 +601,15 @@ void UNetDriver::AddNetworkActor(AActor* Actor)
 	LLM_SCOPE(ELLMTag::Networking);
 	ensureMsgf(Actor == nullptr || !(Actor->HasAnyFlags(RF_ClassDefaultObject | RF_ArchetypeObject)), TEXT("%s is a CDO or Archetype and should not be replicated."), *GetFullNameSafe(Actor));
 
+#if UE_WITH_IRIS
+	// We currently only want to do this on the server as we want to replicate the actor to the client
+	// but we rely on the old system for actor instantiation
+	if (ReplicationSystem)
+	{
+		return;
+	}
+#endif // UE_WITH_IRIS
+
 	if (!IsDormInitialStartupActor(Actor))
 	{
 		GetNetworkObjectList().FindOrAdd(Actor, this);
@@ -616,6 +639,14 @@ FNetworkObjectInfo* UNetDriver::FindOrAddNetworkObjectInfo(const AActor* InActor
 	ensureMsgf(InActor == nullptr || !(InActor->HasAnyFlags(RF_ClassDefaultObject | RF_ArchetypeObject)), TEXT("%s is a CDO or Archetype and should not be replicated."), *GetFullNameSafe(InActor));
 
 	bool bWasAdded = false;
+
+#if UE_WITH_IRIS
+	if (ReplicationSystem)
+	{
+		return FindNetworkObjectInfo(InActor);
+	}
+#endif // UE_WITH_IRIS
+
 	if (TSharedPtr<FNetworkObjectInfo>* InfoPtr = GetNetworkObjectList().FindOrAdd(const_cast<AActor*>(InActor), this, &bWasAdded))
 	{
 		if (bWasAdded && ReplicationDriver)
@@ -631,6 +662,21 @@ FNetworkObjectInfo* UNetDriver::FindOrAddNetworkObjectInfo(const AActor* InActor
 
 FNetworkObjectInfo* UNetDriver::FindNetworkObjectInfo(const AActor* InActor)
 {
+#if UE_WITH_IRIS
+	if (ReplicationSystem)
+	{
+		// Return default initialized NetworkObjectInfo so the presence of a pointer can at least indicate whether the actor is replicated or not.
+		const UE::Net::FNetHandle NetHandle = UE::Net::FReplicationSystemUtil::GetNetHandle(InActor);
+		if (NetHandle.IsValid())
+		{
+			*DummyNetworkObjectInfo = FNetworkObjectInfo();
+			return DummyNetworkObjectInfo.Get();
+		}
+
+		return nullptr;
+	}
+#endif // UE_WITH_IRIS
+
 	return GetNetworkObjectList().Find(InActor).Get();
 }
 
@@ -738,6 +784,14 @@ void UNetDriver::TickFlush(float DeltaSeconds)
 #if WITH_SERVER_CODE
 		int32 Updated = ServerReplicateActors( DeltaSeconds );
 
+#if UE_WITH_IRIS
+		if (ReplicationSystem)
+		{
+			UpdateReplicationViews();
+			ReplicationSystem->PreSendUpdate(DeltaSeconds);
+		}
+#endif // UE_WITH_IRIS
+
 		static int32 LastUpdateCount = 0;
 		// Only log the zero replicated actors once after replicating an actor
 		if ((LastUpdateCount && !Updated) || Updated)
@@ -747,6 +801,16 @@ void UNetDriver::TickFlush(float DeltaSeconds)
 		LastUpdateCount = Updated;
 #endif // WITH_SERVER_CODE
 	}
+#if UE_WITH_IRIS
+	else if (!IsServer() && ServerConnection != nullptr)
+	{
+		if (ReplicationSystem)
+		{
+			UpdateReplicationViews();
+			ReplicationSystem->PreSendUpdate(DeltaSeconds);
+		}
+	}
+#endif // UE_WITH_IRIS
 
 	// Reset queued bunch amortization timer
 	ProcessQueuedBunchesCurrentFrameMilliseconds = 0.0f;
@@ -1189,6 +1253,13 @@ void UNetDriver::TickFlush(float DeltaSeconds)
 		{
 			Connection->Tick(DeltaSeconds);
 		}
+
+#if UE_WITH_IRIS
+		if (ReplicationSystem && (ServerConnection != nullptr || ClientConnections.Num() > 0))
+		{
+			ReplicationSystem->PostSendUpdate();
+		}
+#endif // UE_WITH_IRIS
 	}
 
 	if (ConnectionlessHandler.IsValid())
@@ -1631,9 +1702,12 @@ bool UNetDriver::InitBase(bool bInitAsClient, FNetworkNotify* InNotify, const FU
 	if (!bInitAsClient)
 	{
 		ConnectionlessHandler.Reset(nullptr);
-		InitDestroyedStartupActors();
-		InitReplicationDriverClass();
-		SetReplicationDriver(UReplicationDriver::CreateReplicationDriver(this, URL, GetWorld()));
+		
+		if (!IsUsingIrisReplication())
+		{
+			InitReplicationDriverClass();
+			SetReplicationDriver(UReplicationDriver::CreateReplicationDriver(this, URL, GetWorld()));
+		}
 
 		DDoS.Init(FMath::Clamp(NetServerMaxTickRate, 1, 1000));
 
@@ -1664,6 +1738,41 @@ bool UNetDriver::InitBase(bool bInitAsClient, FNetworkNotify* InNotify, const FU
 	// If we are not using Iris, we use the UniqueId to identify the NetDriver.
 	NetTraceId = GetUniqueID();
 
+#if UE_WITH_IRIS
+	if (IsUsingIrisReplication() && !ReplicationSystem)
+	{
+		LLM_SCOPE(ELLMTag::Iris);
+
+		if (ensureAlways(InitReplicationBridgeClass()))
+		{
+			const bool bAllowSend = !bInitAsClient;
+
+			UReplicationBridge* ReplicationBridge = NewObject<UActorReplicationBridge>(GetTransientPackage(), ReplicationBridgeClass);
+			if (ReplicationBridge)
+			{
+				ReplicationBridge->SetNetDriver(this);
+
+				// Create ReplicationSystem
+				UReplicationSystem::FReplicationSystemParams Params;
+				Params.ReplicationBridge = ReplicationBridge;
+				Params.bIsServer = !bInitAsClient;
+				Params.bAllowObjectReplication = !bInitAsClient;
+
+				SetReplicationSystem(UE::Net::FReplicationSystemFactory::CreateReplicationSystem(Params));
+			}
+			else
+			{
+				UE_LOG(LogNet, Error, TEXT("Failed to initialize ReplicationSystem"));
+			}
+		}
+	}
+#endif // UE_WITH_IRIS
+
+	if (!bInitAsClient)
+	{
+		InitDestroyedStartupActors();
+	}
+
 	return bSuccess;
 }
 
@@ -1853,6 +1962,9 @@ void UNetDriver::Shutdown()
 	ConnectionlessHandler.Reset(nullptr);
 
 	SetReplicationDriver(nullptr);
+#if UE_WITH_IRIS
+	SetReplicationSystem(nullptr);
+#endif // UE_WITH_IRIS
 
 	// End NetTrace session for this instance
 	UE_NET_TRACE_END_SESSION(GetNetTraceId());
@@ -1873,6 +1985,18 @@ bool UNetDriver::IsServer() const
 	return ServerConnection == NULL;
 }
 
+EEngineNetworkRuntimeFeatures UNetDriver::GetNetworkRuntimeFeatures() const
+{
+	EEngineNetworkRuntimeFeatures NetDriverFeatures = EEngineNetworkRuntimeFeatures::None;
+	
+	if (IsUsingIrisReplication())
+	{
+		EnumAddFlags(NetDriverFeatures, EEngineNetworkRuntimeFeatures::IrisEnabled);
+	}
+	
+	return NetDriverFeatures;
+}
+
 // ----------------------------------------------------------------------------------------
 //	RPC Tracking (receiving) via CSV profiler
 // ----------------------------------------------------------------------------------------
@@ -2372,6 +2496,7 @@ void UNetDriver::ProcessRemoteFunctionForChannelPrivate(
 	}
 #endif
 
+	// Init bunch for rpc
 #if UE_NET_TRACE_ENABLED
 	SetTraceCollector(Bunch, UE_NET_TRACE_CREATE_COLLECTOR(ENetTraceVerbosity::Trace));
 #endif
@@ -3427,6 +3552,15 @@ FActorDestructionInfo* UNetDriver::CreateDestructionInfo( UNetDriver* NetDriver,
 
 void UNetDriver::NotifyActorDestroyed( AActor* ThisActor, bool IsSeamlessTravel )
 {
+#if UE_WITH_IRIS
+	// For Iris this is handled through a call to AActor::EndReplication
+	// NOTE: The reason for doing this after the removal from the call to RepChangedPropertyTrackerMap.Remove is that Iris currently relies on this to invoke PreReplication which will create a RepChangePropertyTracker
+	if (ReplicationSystem)
+	{
+		return;
+	}
+#endif // UE_WITH_IRIS
+
 	const bool bIsServer = IsServer();
 	
 	if (bIsServer)
@@ -3555,6 +3689,13 @@ void UNetDriver::NotifyStreamingLevelUnload(ULevel* Level)
 				NotifyActorLevelUnloaded(Actor);
 			}
 		}
+
+#if UE_WITH_IRIS
+		if (ReplicationSystem)
+		{
+			ReplicationSystem->NotifyStreamingLevelUnload(Level);
+		}
+#endif // UE_WITH_IRIS
 		
 		TArray<FNetworkGUID> RemovedGUIDs;
 		for (auto It = DestroyedStartupOrDormantActors.CreateIterator(); It; ++It)
@@ -3634,6 +3775,13 @@ void UNetDriver::NotifyActorTearOff(AActor* Actor)
 	{
 		ReplicationDriver->NotifyActorTearOff(Actor);
 	}
+#if UE_WITH_IRIS
+	if (ReplicationSystem)
+	{
+		// Set the actor to be torn-off during the next update of the replication systeem
+		ReplicationSystem->TearOffNextUpdate(UE::Net::FReplicationSystemUtil::GetNetHandle(Actor));
+	}
+#endif // UE_WITH_IRIS
 }
 
 void UNetDriver::ForceNetUpdate(AActor* Actor)
@@ -3644,6 +3792,14 @@ void UNetDriver::ForceNetUpdate(AActor* Actor)
 		ReplicationDriver->ForceNetUpdate(Actor);
 		return;
 	}
+
+#if UE_WITH_IRIS
+	if (ReplicationSystem)
+	{
+		ReplicationSystem->MarkDirty(UE::Net::FReplicationSystemUtil::GetNetHandle(Actor));
+		return;
+	}
+#endif // UE_WITH_IRIS
 	
 	// Legacy implementation
 	if (FNetworkObjectInfo* NetActor = FindNetworkObjectInfo(Actor))
@@ -3701,7 +3857,16 @@ void UNetDriver::FlushActorDormancy(AActor* Actor, bool bWasDormInitial)
 		ReplicationDriver->FlushNetDormancy(Actor, bWasDormInitial);
 	}
 
-	FlushActorDormancyInternal(Actor);
+#if UE_WITH_IRIS
+	if (ReplicationSystem && ReplicationSystem->IsServer())
+	{
+		UE::Net::FReplicationSystemUtil::FlushNetDormancy(Actor, bWasDormInitial);
+	}
+	else
+#endif // UE_WITH_IRIS
+	{
+		FlushActorDormancyInternal(Actor);
+	}
 }
 
 void UNetDriver::NotifyActorDormancyChange(AActor* Actor, ENetDormancy OldDormancyState)
@@ -3711,7 +3876,16 @@ void UNetDriver::NotifyActorDormancyChange(AActor* Actor, ENetDormancy OldDorman
 		ReplicationDriver->NotifyActorDormancyChange(Actor, OldDormancyState);
 	}
 
-	FlushActorDormancyInternal(Actor);
+#if UE_WITH_IRIS
+	if (ReplicationSystem)
+	{
+		UE::Net::FReplicationSystemUtil::NotifyActorDormancyChange(Actor, OldDormancyState);
+	}
+	else
+#endif // UE_WITH_IRIS
+	{
+		FlushActorDormancyInternal(Actor);
+	}
 }
 
 void UNetDriver::FlushActorDormancyInternal(AActor *Actor)
@@ -5437,6 +5611,13 @@ void UNetDriver::SetNetDriverDefinition(FName NewNetDriverDefinition)
 	}
 }
 
+void UNetDriver::PostCreation(bool bInitializeWithIris)
+{
+#if UE_WITH_IRIS
+	bIsUsingIris = bInitializeWithIris;
+#endif //UE_WITH_IRIS
+}
+
 #if NET_DEBUG_RELEVANT_ACTORS
 void UNetDriver::PrintDebugRelevantActors()
 {
@@ -5713,6 +5894,23 @@ void UNetDriver::CreateReplicatedStaticActorDestructionInfo(UNetDriver* NetDrive
 {
 	check(NetDriver && Level);
 
+#if UE_WITH_IRIS
+	if (NetDriver->ReplicationSystem)
+	{
+		if (UActorReplicationBridge* Bridge = NetDriver->ReplicationSystem->GetReplicationBridgeAs<UActorReplicationBridge>())
+		{			
+			// Add explicit destruction info for this object
+			UReplicationBridge::FEndReplicationParameters Params;
+			Params.Level = Level;
+			Params.Location = Info.DestroyedPosition;
+			Params.bUseDistanceBasedPrioritization = true;
+
+			Bridge->AddStaticDestructionInfo(Info.PathName.ToString(), Info.ObjOuter.Get(), Params);
+		}
+		return;
+	}
+#endif // UE_WITH_IRIS
+
 	FNetworkGUID NetGUID = NetDriver->GuidCache->AssignNewNetGUIDFromPath_Server( Info.PathName.ToString(), Info.ObjOuter.Get(), Info.ObjClass );
 	if (NetGUID.IsDefault())
 	{
@@ -5889,7 +6087,25 @@ void UNetDriver::SetWorld(class UWorld* InWorld)
 		Notify = InWorld;
 		RegisterTickEvents(InWorld);
 
-		GetNetworkObjectList().AddInitialObjects(InWorld, this);
+
+#if UE_WITH_IRIS
+		if (IsUsingIrisReplication())
+		{
+			if (bSkipBeginReplicationForWorld)
+			{
+				// We skipped, reset the flag in case there is another call.
+				bSkipBeginReplicationForWorld = false;
+			}
+			else
+			{
+				UE::Net::FReplicationSystemUtil::BeginReplicationForActorsInWorld(InWorld);
+			}
+		}
+		else
+#endif // UE_WITH_IRIS
+		{
+			GetNetworkObjectList().AddInitialObjects(InWorld, this);
+		}
 
 		NotifyGameInstanceUpdated();
 	}
@@ -5926,6 +6142,13 @@ void UNetDriver::ResetGameWorldState()
 	{
 		ReplicationDriver->ResetGameWorldState();
 	}
+
+#if UE_WITH_IRIS
+	if (ReplicationSystem)
+	{
+		ReplicationSystem->ResetGameWorldState();
+	}
+#endif // UE_WITH_IRIS
 }
 
 void UNetDriver::CleanPackageMaps()
@@ -6030,6 +6253,15 @@ void UNetDriver::SetReplicationDriver(UReplicationDriver* NewReplicationDriver)
 		}
 	}
 
+#if UE_WITH_IRIS
+	if (IsUsingIrisReplication() && NewReplicationDriver)
+	{
+		checkf(false, TEXT("ReplicationDriver (%s) are not supported with a NetDriver (%s) configured to use Iris"), *NewReplicationDriver->GetName(), *NetDriverName.ToString());
+		NewReplicationDriver->MarkAsGarbage();
+		return;
+	}
+#endif
+
 	ReplicationDriver = NewReplicationDriver;
 	if (ReplicationDriver)
 	{
@@ -6059,6 +6291,169 @@ UNetConnection* UNetDriver::GetConnectionById(uint32 ConnectionId) const
 	return nullptr;
 }
 
+#if UE_WITH_IRIS
+bool UNetDriver::InitReplicationBridgeClass()
+{
+	if (ReplicationBridgeClass != nullptr)
+	{
+		return true;
+	}
+
+	if (!ReplicationBridgeClassName.IsEmpty())
+	{
+		ReplicationBridgeClass = LoadClass<UReplicationBridge>(nullptr, *ReplicationBridgeClassName, nullptr, LOAD_None, nullptr);
+		if (ReplicationBridgeClass == nullptr)
+		{
+			UE_LOG(LogNet, Error, TEXT("Failed to load class '%s'"), *ReplicationBridgeClassName);
+			return false;
+		}
+	}
+	else
+	{
+		// Fall back on ActorReplicationBridge
+		ReplicationBridgeClass = UActorReplicationBridge::StaticClass();
+	}
+
+	return ReplicationBridgeClass != nullptr;
+}
+
+void UNetDriver::UpdateGroupFilterStatusForLevel(const ULevel* Level, UE::Net::FNetObjectGroupHandle LevelGroupHandle)
+{
+	using namespace UE::Net;
+
+	if (ReplicationSystem == nullptr || LevelGroupHandle == InvalidNetObjectGroupHandle)
+	{
+		return;
+	}
+
+	const FName WorldPackageName = GetWorldPackage()->GetFName();
+	const FName LevelPackageName = Level->GetOutermost()->GetFName();
+	const bool bIsPersistentLevel = Level->IsPersistentLevel();
+
+	// Build connection mask
+	FNetBitArray ConnectionMask(ReplicationSystem->GetMaxConnectionCount());
+		
+	for (UNetConnection* Connection : ClientConnections)
+	{
+		const bool bIsVisible = (bIsPersistentLevel && WorldPackageName == Connection->GetClientWorldPackageName()) || Connection->ClientVisibleLevelNames.Contains(LevelPackageName);
+		ConnectionMask.SetBitValue(Connection->GetConnectionId(), bIsVisible);
+	}
+
+	// Set group filter status
+	ReplicationSystem->SetGroupFilterStatus(LevelGroupHandle, ConnectionMask, ENetFilterStatus::Allow);
+}
+
+void UNetDriver::SetReplicationSystem(UReplicationSystem* InReplicationSystem)
+{
+	if (ReplicationSystem)
+	{
+		UE::Net::FReplicationSystemFactory::DestroyReplicationSystem(ReplicationSystem);
+	}
+
+	ReplicationSystem = InReplicationSystem;
+
+	if (ReplicationSystem)
+	{
+		// When we run using Iris, we use ReplicationSystemId as our unique identifier
+		NetTraceId = ReplicationSystem->GetId();
+
+		UE::Net::FReplicationSystemUtil::BeginReplicationForActorsInWorld(GetWorld());
+	}
+}
+
+void UNetDriver::ClearIrisSystem()
+{
+	if (ReplicationSystem)
+	{
+		UReplicationBridge* Bridge = ReplicationSystem->GetReplicationBridge();
+		if (ensureAlways(Bridge))
+		{
+			Bridge->SetNetDriver(nullptr);
+		}
+	}
+
+	ReplicationSystem = nullptr;
+}
+
+void UNetDriver::RestoreIrisSystem(UReplicationSystem* InReplicationSystem)
+{
+	check(InReplicationSystem != nullptr);
+	check(InReplicationSystem->GetReplicationBridge() != nullptr);
+	checkf(ReplicationSystem == nullptr, TEXT("Cannot restore IrisSystem in %s since one system is already initialized."), *GetName());
+
+	ReplicationSystem = InReplicationSystem;
+	ReplicationSystem->GetReplicationBridge()->SetNetDriver(this);
+
+	// When we run using Iris, we use ReplicationSystemId as our unique identifier
+	NetTraceId = ReplicationSystem->GetId();
+	
+	// World Actors have already been registered in this IrisSystem, prevent adding them twice when the World gets set.
+	bSkipBeginReplicationForWorld = true;
+}
+
+void UNetDriver::UpdateReplicationViews() const
+{
+	using namespace UE::Net;
+
+	FReplicationView ReplicationView;
+
+	FReplicationView::FView DefaultView;
+	{
+		const APlayerCameraManager* CameraManager = GetDefault<APlayerCameraManager>();
+		DefaultView.FoVRadians = FMath::DegreesToRadians(CameraManager->DefaultFOV);
+	}
+
+	TObjectPtr<UNetConnection> const* Connections = (ServerConnection != nullptr ? &ServerConnection : ClientConnections.GetData());
+	const uint32 ConnectionCount = (ServerConnection != nullptr ? 1U : uint32(ClientConnections.Num()));
+	TArray<const UNetConnection*, TInlineAllocator<4>> SubConnections;
+	for (uint32 ConnIt = 0, ConnEndIt = ConnectionCount; ConnIt != ConnEndIt; ++ConnIt)
+	{
+		ReplicationView.Views.Reset();
+
+		const UNetConnection* Conn = Connections[ConnIt];
+		SubConnections.Reset();
+		SubConnections.Add(Conn);
+		if (Conn->Children.Num() > 0)
+		{
+			SubConnections.Insert(reinterpret_cast<UNetConnection*const*>(Conn->Children.GetData()), Conn->Children.Num(), 1);
+		}
+
+		for (const UNetConnection* SubConn : SubConnections)
+		{
+			const AActor* ViewTarget = Conn->ViewTarget;
+			const APlayerController* ViewingController = Conn->PlayerController;
+			if (ViewTarget == nullptr && ViewingController == nullptr)
+			{
+				continue;
+			}
+
+			FReplicationView::FView View;
+			View.FoVRadians = DefaultView.FoVRadians;
+			if (ViewTarget)
+			{
+				View.Pos = ViewTarget->GetActorLocation();
+				View.Dir = ViewTarget->GetActorRotation().Vector();
+			}
+			if (ViewingController)
+			{
+				FRotator ViewRotation = ViewingController->GetControlRotation();
+				ViewingController->GetPlayerViewPoint(View.Pos, ViewRotation);
+				View.Dir = ViewRotation.Vector();
+
+				if (const APlayerCameraManager* CameraManager = ViewingController->PlayerCameraManager)
+				{
+					View.FoVRadians = FMath::DegreesToRadians(CameraManager->GetFOVAngle());
+				}
+			}
+
+			ReplicationView.Views.Add(View);
+		}
+
+		ReplicationSystem->SetReplicationView(Conn->GetConnectionId(), ReplicationView);
+	}
+}
+#endif // UE_WITH_IRIS
+
 #if NET_DEBUG_RELEVANT_ACTORS
 static void	DumpRelevantActors( UWorld* InWorld )
 {
@@ -6328,6 +6723,37 @@ void UNetDriver::ProcessRemoteFunction(
 			LocalOutParms = NetDriverInternal::CopyOutParametersToLocalParameters(Function, OutParms, Parameters, TargetObj);
 		}
 
+#if UE_WITH_IRIS
+		if (ReplicationSystem)
+		{
+			if (bIsServerMulticast)
+			{
+				if (ReplicationSystem->SendRPC(Actor, SubObject, Function, Parameters))
+				{
+					return;
+				}
+			}
+			else
+			{
+				if (UNetConnection* Connection = Actor->GetNetConnection())
+				{
+					if (UChildConnection* ChildConnection = Connection->GetUChildConnection())
+					{
+						Connection = ChildConnection->Parent;
+					}
+
+					if (ReplicationSystem->SendRPC(Connection->GetConnectionId(), Actor, SubObject, Function, Parameters))
+					{
+						return;
+					}
+				}
+			}
+
+			// If we are using Iris replication, we should never fall back on normal replication path
+			return;
+		}
+#endif // UE_WITH_IRIS
+
 		// Forward to replication Driver if there is one
 		if (ReplicationDriver && ReplicationDriver->ProcessRemoteFunction(Actor, Function, Parameters, OutParms, Stack, SubObject))
 		{
diff --git a/Engine/Source/Runtime/Engine/Private/NetworkObjectList.cpp b/Engine/Source/Runtime/Engine/Private/NetworkObjectList.cpp
--- a/Engine/Source/Runtime/Engine/Private/NetworkObjectList.cpp
+++ b/Engine/Source/Runtime/Engine/Private/NetworkObjectList.cpp
@@ -14,6 +14,14 @@ void FNetworkObjectList::AddInitialObjects(UWorld* const World, UNetDriver* NetD
 		return;
 	}
 
+#if UE_WITH_IRIS
+	if (NetDriver->GetReplicationSystem())
+	{
+		return;
+	}
+#endif // UE_WITH_IRIS
+
+
 	for (FActorIterator Iter(World); Iter; ++Iter)
 	{
 		AActor* Actor = *Iter;
@@ -50,6 +58,13 @@ TSharedPtr<FNetworkObjectInfo>* FNetworkObjectList::FindOrAdd(AActor* const Acto
 		return nullptr;
 	}
 
+#if UE_WITH_IRIS
+	if (NetDriver && NetDriver->GetReplicationSystem())
+	{
+		return nullptr;
+	}
+#endif // UE_WITH_IRIS
+
 	TSharedPtr<FNetworkObjectInfo>* NetworkObjectInfo = AllNetworkObjects.Find(Actor);
 	
 	if (NetworkObjectInfo == nullptr)
diff --git a/Engine/Source/Runtime/Engine/Private/Pawn.cpp b/Engine/Source/Runtime/Engine/Private/Pawn.cpp
--- a/Engine/Source/Runtime/Engine/Private/Pawn.cpp
+++ b/Engine/Source/Runtime/Engine/Private/Pawn.cpp
@@ -613,6 +613,11 @@ void APawn::PossessedBy(AController* NewController)
 	Controller = NewController;
 	ForceNetUpdate();
 
+#if UE_WITH_IRIS
+	// The owning connection depends on the Controller having the new value.
+	UpdateOwningNetConnection();
+#endif
+
 	if (Controller->PlayerState != nullptr)
 	{
 		SetPlayerState(Controller->PlayerState);
@@ -650,6 +655,11 @@ void APawn::UnPossessed()
 	SetOwner(nullptr);
 	Controller = nullptr;
 
+#if UE_WITH_IRIS
+	// The owning connection depends on the Controller having the new value.
+	UpdateOwningNetConnection();
+#endif
+
 	// Unregister input component if we created one
 	DestroyPlayerInputComponent();
 
diff --git a/Engine/Source/Runtime/Engine/Private/PendingNetGame.cpp b/Engine/Source/Runtime/Engine/Private/PendingNetGame.cpp
--- a/Engine/Source/Runtime/Engine/Private/PendingNetGame.cpp
+++ b/Engine/Source/Runtime/Engine/Private/PendingNetGame.cpp
@@ -111,8 +111,9 @@ void UPendingNetGame::SendInitialJoin()
 				EncryptionToken = URL.GetOption(TEXT("EncryptionToken="), TEXT(""));
 			}
 
-			FNetControlMessage<NMT_Hello>::Send(ServerConn, IsLittleEndian, LocalNetworkVersion, EncryptionToken);
-
+			EEngineNetworkRuntimeFeatures LocalNetworkFeatures = NetDriver->GetNetworkRuntimeFeatures();
+			FNetControlMessage<NMT_Hello>::Send(ServerConn, IsLittleEndian, LocalNetworkVersion, EncryptionToken, LocalNetworkFeatures);
+			
 			ServerConn->FlushNet();
 		}
 	}
@@ -187,18 +188,32 @@ void UPendingNetGame::NotifyControlMessage(UNetConnection* Connection, uint8 Mes
 	switch (MessageType)
 	{
 		case NMT_Upgrade:
+		{
 			// Report mismatch.
 			uint32 RemoteNetworkVersion;
 
-			if (FNetControlMessage<NMT_Upgrade>::Receive(Bunch, RemoteNetworkVersion))
+			EEngineNetworkRuntimeFeatures RemoteNetworkFeatures = EEngineNetworkRuntimeFeatures::None;
+
+			if (FNetControlMessage<NMT_Upgrade>::Receive(Bunch, RemoteNetworkVersion, RemoteNetworkFeatures))
 			{
+				TStringBuilder<128> RemoteFeaturesDescription;
+				FNetworkVersion::DescribeNetworkRuntimeFeaturesBitset(RemoteNetworkFeatures, RemoteFeaturesDescription);
+
+				TStringBuilder<128> LocalFeaturesDescription;
+				FNetworkVersion::DescribeNetworkRuntimeFeaturesBitset(NetDriver->GetNetworkRuntimeFeatures(), LocalFeaturesDescription);
+
+				UE_LOG(LogNet, Error, TEXT("Server is incompatible with the local version of the game: RemoteNetworkVersion=%u, RemoteNetworkFeatures=%s vs LocalNetworkVersion=%u, LocalNetworkFeatures=%s"), 
+					RemoteNetworkVersion, RemoteFeaturesDescription.ToString(),
+					FNetworkVersion::GetLocalNetworkVersion(), LocalFeaturesDescription.ToString()
+				);
+
 				// Upgrade
 				ConnectionError = NSLOCTEXT("Engine", "ClientOutdated", "The match you are trying to join is running an incompatible version of the game.  Please try upgrading your game version.").ToString();
 				GEngine->BroadcastNetworkFailure(NULL, NetDriver, ENetworkFailure::OutdatedClient, ConnectionError);
 			}
 
 			break;
-
+		}
 		case NMT_Failure:
 		{
 			// our connection attempt failed for some reason, for example a synchronization mismatch (bad GUID, etc) or because the server rejected our join attempt (too many players, etc)
diff --git a/Engine/Source/Runtime/Engine/Private/PlayerController.cpp b/Engine/Source/Runtime/Engine/Private/PlayerController.cpp
--- a/Engine/Source/Runtime/Engine/Private/PlayerController.cpp
+++ b/Engine/Source/Runtime/Engine/Private/PlayerController.cpp
@@ -72,6 +72,12 @@
 #include "Physics/AsyncPhysicsInputComponent.h"
 #include "GenericPlatform/GenericPlatformInputDeviceMapper.h"
 
+#if UE_WITH_IRIS
+#include "Iris/ReplicationSystem/ReplicationSystem.h"
+#include "Net/Iris/ReplicationSystem/ActorReplicationBridge.h"
+#include "Net/Iris/ReplicationSystem/ReplicationSystemUtil.h"
+#endif // UE_WITH_IRIS
+
 DEFINE_LOG_CATEGORY(LogPlayerController);
 
 #define LOCTEXT_NAMESPACE "PlayerController"
@@ -4655,6 +4661,27 @@ void APlayerController::ReceivedSpectatorClass(TSubclassOf<ASpectatorPawn> Spect
 
 void APlayerController::SetPawn(APawn* InPawn)
 {
+#if UE_WITH_IRIS
+	if (GetLocalRole() == ROLE_Authority)
+	{
+		if (const UReplicationSystem* ReplicationSystem = UE::Net::FReplicationSystemUtil::GetReplicationSystem(this))
+		{
+			if (APawn* PrevPawn = GetPawn(); PrevPawn != InPawn)
+			{
+				if (IsValid(PrevPawn))
+				{
+					UE::Net::FReplicationSystemUtil::RemoveDependentActor(this, PrevPawn);
+				}
+
+				if (IsValid(InPawn))
+				{
+					UE::Net::FReplicationSystemUtil::AddDependentActor(this, InPawn);
+				}
+			}
+		}
+	}
+#endif
+
 	if (InPawn == NULL)
 	{
 		// Attempt to move the PC to the current camera location if no pawn was specified
@@ -4742,6 +4769,11 @@ void APlayerController::SetPlayer( UPlayer* InPlayer )
 		if (NetConnection)
 		{
 			NetConnection->OwningActor = this;
+
+#if UE_WITH_IRIS
+			UpdateOwningNetConnection();
+			UE::Net::FReplicationSystemUtil::UpdateSubObjectGroupMemberships(this);
+#endif // UE_WITH_IRIS
 		}
 	}
 
@@ -5847,7 +5879,33 @@ void APlayerController::IncludeInNetConditionGroup(FName NetGroup)
 	checkf(!UE::Net::IsSpecialNetConditionGroup(NetGroup), TEXT("Cannot add a player to special netcondition group %s manually. This group membership is managed by the network engine automatically."), *NetGroup.ToString());
 	checkf(!NetGroup.IsNone(), TEXT("Invalid netcondition group: NONE"));
 	NetConditionGroups.AddUnique(NetGroup);
+
+#if UE_WITH_IRIS
+	UE::Net::FReplicationSystemUtil::UpdateSubObjectGroupMemberships(this);
+#endif
+}
+
+void APlayerController::RemoveFromNetConditionGroup(FName NetGroup)
+{
+	NetConditionGroups.RemoveSingleSwap(NetGroup);
+#if UE_WITH_IRIS
+	UE::Net::FReplicationSystemUtil::RemoveSubObjectGroupMembership(this, NetGroup);
+#endif
+}
+
+
+#if UE_WITH_IRIS
+void APlayerController::BeginReplication()
+{
+	// Always allow the PlayerController to be replicated as it is required for travel.
+	FActorBeginReplicationParams Params;
+	Params.bIncludeInLevelGroupFilter = false;
+	Super::BeginReplication(Params);
+
+	// Enable groups once owner is set!!
+	UE::Net::FReplicationSystemUtil::UpdateSubObjectGroupMemberships(this);
 }
+#endif // UE_WITH_IRIS
 
 UAsyncPhysicsData* APlayerController::GetAsyncPhysicsDataToWrite() const
 {
diff --git a/Engine/Source/Runtime/Engine/Private/PlayerState.cpp b/Engine/Source/Runtime/Engine/Private/PlayerState.cpp
--- a/Engine/Source/Runtime/Engine/Private/PlayerState.cpp
+++ b/Engine/Source/Runtime/Engine/Private/PlayerState.cpp
@@ -9,6 +9,11 @@
 #include "Net/OnlineEngineInterface.h"
 #include "GameFramework/GameStateBase.h"
 #include "Net/Core/PushModel/PushModel.h"
+#if UE_WITH_IRIS
+#include "Net/Iris/ReplicationSystem/ReplicationSystemUtil.h"
+#include "Iris/ReplicationSystem/ReplicationSystem.h"
+#include "Iris/ReplicationSystem/Prioritization/NetObjectPrioritizer.h"
+#endif
 
 PRAGMA_DISABLE_DEPRECATION_WARNINGS
 APlayerState::APlayerState(const FObjectInitializer& ObjectInitializer)
@@ -598,4 +603,40 @@ void APlayerState::OnPawnPrivateDestroyed(AActor* InActor)
 	}
 }
 
-PRAGMA_ENABLE_DEPRECATION_WARNINGS
\ No newline at end of file
+PRAGMA_ENABLE_DEPRECATION_WARNINGS
+
+#if UE_WITH_IRIS
+void APlayerState::BeginReplication()
+{
+	using namespace UE::Net;
+
+	Super::BeginReplication();
+
+	/**
+	 * If we're replicating using Iris we set the prioritizer to one named PlayerState if there is one.
+	 * For example add the following to your DefaultEngine.ini to limit the number of player states
+	 * replicated each frame, but still allow the owning connection's to be replicated each frame.
+	 * You can also implement your own prioritzer of course.
+	 * 
+	 * [/Script/IrisCore.NetObjectPrioritizerDefinitions]
+	 * +NetObjectPrioritizerDefinitions=(PrioritizerName=PlayerState, ClassName=/Script/IrisCore.NetObjectCountLimiter, ConfigClassName=/Script/Engine.PlayerStateCountLimiterConfig)
+	 */
+	const FNetHandle NetHandle = FReplicationSystemUtil::GetNetHandle(this);
+	if (!NetHandle.IsValid())
+	{
+		return;
+	}
+
+	UReplicationSystem* ReplicationSystem = GetReplicationSystem(NetHandle.GetReplicationSystemId());
+	if (ReplicationSystem == nullptr)
+	{
+		return;
+	}
+
+	FNetObjectPrioritizerHandle PrioritizerHandle = ReplicationSystem->GetPrioritizerHandle(FName("PlayerState", FNAME_Find));
+	if (PrioritizerHandle != InvalidNetObjectPrioritizerHandle)
+	{
+		ReplicationSystem->SetPrioritizer(NetHandle, PrioritizerHandle);
+	}	
+}
+#endif
