Fixed comment.

#rnx

diff --git a/Engine/Source/Runtime/Engine/Private/ContentStreaming.cpp b/Engine/Source/Runtime/Engine/Private/ContentStreaming.cpp
--- a/Engine/Source/Runtime/Engine/Private/ContentStreaming.cpp
+++ b/Engine/Source/Runtime/Engine/Private/ContentStreaming.cpp
@@ -476,7 +476,7 @@ void IStreamingManager::SetupViewInfos( float DeltaTime )
 
 	bool bHaveMultiplePlayerViews = (PendingViewInfos.Num() > 1) ? true : false;
 
-	// Add the slave locations.
+	// Add the secondary locations.
 	float ScreenSize = 1280.0f;
 	float FOVScreenSize = ScreenSize / FMath::Tan( 80.0f * float(UE_PI) / 360.0f );
 	if ( PendingViewInfos.Num() > 0 )
