Iris
- Disabled checks verifying that FastArrays stored in CDO/Archetype are not modified when using them as source during initialization of ShadowState.

diff --git a/Engine/Source/Runtime/Engine/Private/DataReplication.cpp b/Engine/Source/Runtime/Engine/Private/DataReplication.cpp
--- a/Engine/Source/Runtime/Engine/Private/DataReplication.cpp
+++ b/Engine/Source/Runtime/Engine/Private/DataReplication.cpp
@@ -364,8 +364,11 @@ bool FObjectReplicator::SendCustomDeltaProperty(UObject* InObject, uint16 Custom
 	Parms.bSupportsFastArrayDeltaStructSerialization = bSupportsFastArrayDelta;
 	Parms.Connection = Connection;
 	Parms.bInternalAck = Connection->IsInternalAck();
+
+#if UE_WITH_IRIS
 	// When initializing baselines we should not modify the source data if it originates from the CDO or archetype
 	Parms.bIsInitializingBaseFromDefault = Parms.Object && Parms.Object->HasAnyFlags(RF_ClassDefaultObject | RF_ArchetypeObject);
+#endif
 
 	return FNetSerializeCB::SendCustomDeltaProperty(*RepLayout, Parms, CustomDeltaIndex);
 }
