Iris - Fixed static analysis warnings.

diff --git a/Engine/Source/Runtime/Engine/Private/RepMovementNetSerializer.cpp b/Engine/Source/Runtime/Engine/Private/RepMovementNetSerializer.cpp
--- a/Engine/Source/Runtime/Engine/Private/RepMovementNetSerializer.cpp
+++ b/Engine/Source/Runtime/Engine/Private/RepMovementNetSerializer.cpp
@@ -3,7 +3,7 @@
 #include "RepMovementNetSerializer.h"
 
 #if UE_WITH_IRIS
-#include "Engine/EngineTypes.h"
+#include "Engine/ReplicatedState.h"
 #include "Iris/Serialization/NetBitStreamReader.h"
 #include "Iris/Serialization/NetBitStreamWriter.h"
 #include "Iris/Serialization/NetSerializerDelegates.h"
@@ -70,8 +70,8 @@ private:
 	inline static const FRotatorAsByteNetSerializerConfig RotatorAsByteSerializerConfig;
 	inline static const FRotatorAsShortNetSerializerConfig RotatorAsShortSerializerConfig;
 
-	inline static const FNetSerializerConfig* VectorNetQuantizeNetSerializerConfigs[3] = {&QuantizeSerializerConfig, &Quantize10SerializerConfig, &Quantize100SerializerConfig};
-	inline static const FNetSerializer* VectorNetQuantizeNetSerializers[3] = {};
+	inline static const FNetSerializerConfig* VectorNetQuantizeNetSerializerConfigs[4] = {&QuantizeSerializerConfig, &Quantize10SerializerConfig, &Quantize100SerializerConfig, nullptr};
+	inline static const FNetSerializer* VectorNetQuantizeNetSerializers[4] = {};
 
 	inline static const FNetSerializerConfig* RotatorNetSerializerConfigs[2] = {&RotatorAsByteSerializerConfig, &RotatorAsShortSerializerConfig};
 	inline static const FNetSerializer* RotatorNetSerializers[2] = {};
@@ -175,6 +175,12 @@ void FRepMovementNetSerializer::Deserialize(FNetSerializationContext& Context, c
 		TempValue.VelocityQuantizationLevel = (FlagsAndQuantizationLevels >> 2U) & 3U;
 		TempValue.LocationQuantizationLevel = (FlagsAndQuantizationLevels >> 4U) & 3U;
 		TempValue.RotationQuantizationLevel = (FlagsAndQuantizationLevels >> 6U) & 1U;
+
+		if (VectorNetQuantizeNetSerializers[TempValue.VelocityQuantizationLevel] == nullptr || RotatorNetSerializers[TempValue.RotationQuantizationLevel] == nullptr)
+		{
+			Context.SetError(GNetError_InvalidValue);
+			return;
+		}
 	}
 
 	// Angular velocity
diff --git a/Engine/Source/Runtime/Experimental/Iris/Core/Private/Iris/ReplicationSystem/NetBlob/NetBlobManager.cpp b/Engine/Source/Runtime/Experimental/Iris/Core/Private/Iris/ReplicationSystem/NetBlob/NetBlobManager.cpp
--- a/Engine/Source/Runtime/Experimental/Iris/Core/Private/Iris/ReplicationSystem/NetBlob/NetBlobManager.cpp
+++ b/Engine/Source/Runtime/Experimental/Iris/Core/Private/Iris/ReplicationSystem/NetBlob/NetBlobManager.cpp
@@ -349,6 +349,7 @@ void FNetBlobManager::FNetObjectAttachmentSendQueue::ProcessQueue(FReplicationCo
 	if (bHasMulticastAttachments)
 	{
 		ConnectionIds = static_cast<uint32*>(FMemory_Alloca(ReplicatingConnections.GetNumBits() * 4));
+		CA_ASSUME(ConnectionIds != nullptr);
 		ReplicatingConnections.ForAllSetBits([ConnectionIds, &ConnectionCount](uint32 Bit) { ConnectionIds[ConnectionCount++] = Bit; });
 	}
 
@@ -375,7 +376,7 @@ void FNetBlobManager::FNetObjectAttachmentSendQueue::ProcessQueue(FReplicationCo
 		{
 			if (bIsMultiPartAttachment)
 			{
-				UE_LOG(LogIris, Warning, TEXT("Splitting multicast net object attachment %s."), ReplicationStateDescriptor->DebugName ? ReplicationStateDescriptor->DebugName->Name : TEXT("Unknown"));
+				UE_LOG(LogIris, Warning, TEXT("Splitting multicast net object attachment %s."), ReplicationStateDescriptor != nullptr && ReplicationStateDescriptor->DebugName ? ReplicationStateDescriptor->DebugName->Name : TEXT("Unknown"));
 			}
 
 			const FNetBlobCreationInfo& BlobCreationInfo = Attachment->GetCreationInfo();
diff --git a/Engine/Source/Runtime/Experimental/Iris/Core/Private/Iris/ReplicationSystem/ObjectReferenceCache.cpp b/Engine/Source/Runtime/Experimental/Iris/Core/Private/Iris/ReplicationSystem/ObjectReferenceCache.cpp
--- a/Engine/Source/Runtime/Experimental/Iris/Core/Private/Iris/ReplicationSystem/ObjectReferenceCache.cpp
+++ b/Engine/Source/Runtime/Experimental/Iris/Core/Private/Iris/ReplicationSystem/ObjectReferenceCache.cpp
@@ -183,6 +183,8 @@ bool FObjectReferenceCache::CreateObjectReferenceInternal(const UObject* Object,
 		return false;
 	}
 
+	CA_ASSUME(Object != nullptr);
+
 	// Check if we are already know about this object
 	if (FNetHandle* ReferenceHandle = ObjectToNetReferenceHandle.Find(Object))
 	{
diff --git a/Engine/Source/Runtime/Experimental/Iris/Core/Private/Iris/ReplicationSystem/ObjectReplicationBridge.cpp b/Engine/Source/Runtime/Experimental/Iris/Core/Private/Iris/ReplicationSystem/ObjectReplicationBridge.cpp
--- a/Engine/Source/Runtime/Experimental/Iris/Core/Private/Iris/ReplicationSystem/ObjectReplicationBridge.cpp
+++ b/Engine/Source/Runtime/Experimental/Iris/Core/Private/Iris/ReplicationSystem/ObjectReplicationBridge.cpp
@@ -751,7 +751,7 @@ void UObjectReplicationBridge::PreUpdateAndPollImpl(FNetHandle Handle)
 				// If the object has pushed based fragments, and not is marked dirty and object is affected by GC we need to make sure that we refresh cached references for all push based fragments
 				const bool bIsPushBasedObject = EnumHasAnyFlags(InstanceTraits, EReplicationInstanceProtocolTraits::HasPartialPushBasedDirtiness | EReplicationInstanceProtocolTraits::HasFullPushBasedDirtiness);
 				const bool bHasObjectReferences = EnumHasAnyFlags(InstanceTraits, EReplicationInstanceProtocolTraits::HasObjectReference);
-				const bool bNeedsRefreshOfCachedObjectReferences = (!(bIsNewInScope | bIsDirtyObject) & bIsGCAffectedObject & bIsPushBasedObject & bHasObjectReferences);
+				const bool bNeedsRefreshOfCachedObjectReferences = ((!(bIsNewInScope | bIsDirtyObject)) & bIsGCAffectedObject & bIsPushBasedObject & bHasObjectReferences);
 				if (bNeedsRefreshOfCachedObjectReferences)
 				{
 					// Only states which has push based dirtiness need to be updated as the other states will be polled in full anyway.
diff --git a/Engine/Source/Runtime/Experimental/Iris/Core/Private/Iris/Serialization/PackedVectorNetSerializers.cpp b/Engine/Source/Runtime/Experimental/Iris/Core/Private/Iris/Serialization/PackedVectorNetSerializers.cpp
--- a/Engine/Source/Runtime/Experimental/Iris/Core/Private/Iris/Serialization/PackedVectorNetSerializers.cpp
+++ b/Engine/Source/Runtime/Experimental/Iris/Core/Private/Iris/Serialization/PackedVectorNetSerializers.cpp
@@ -597,7 +597,7 @@ void FPackedVectorNetSerializerBase::Quantize(uint32 ScaleBitCount, FNetSerializ
 	const SourceType& Source = *reinterpret_cast<const SourceType*>(Args.Source);
 	QuantizedType& Target = *reinterpret_cast<QuantizedType*>(Args.Target);
 
-	const ScalarType Scale = ScalarType(1 << ScaleBitCount);
+	const ScalarType Scale = ScalarType(IntType(1) << ScaleBitCount);
 	
 	SourceType ScaledValue;
 	// Avoid NaN checks called by FVector operators.
@@ -651,6 +651,7 @@ void FPackedVectorNetSerializerBase::Quantize(uint32 ScaleBitCount, FNetSerializ
 void FPackedVectorNetSerializerBase::Dequantize(uint32 ScaleBitCount, FNetSerializationContext& Context, const FNetDequantizeArgs& Args)
 {
 	using ScalarType = decltype(SourceType::X);
+	using IntType = typename TSignedIntType<sizeof(ScalarType)>::Type;
 
 	const QuantizedType& Source = *reinterpret_cast<const QuantizedType*>(Args.Source);
 	SourceType& Target = *reinterpret_cast<SourceType*>(Args.Target);
@@ -665,7 +666,7 @@ void FPackedVectorNetSerializerBase::Dequantize(uint32 ScaleBitCount, FNetSerial
 
 		if (Source.ComponentBitCountAndExtraInfo & IsScaledValueMask)
 		{
-			const ScalarType InvScale = ScalarType(1)/ScalarType(1 << ScaleBitCount);
+			const ScalarType InvScale = ScalarType(1)/ScalarType(IntType(1) << ScaleBitCount);
 			Target = TempValue*InvScale;
 		}
 		else
