Reworked UAnimInstance re-instancing to fix a crash

Added a FCoreDelegates delegate called by the BP reinstancing path. The existing one is called per-batch of reinstanced objects, but this means that some objects will still be REINST_ at the point of the call. This was causing ordering issues when attempting to re-link linked anim graphs (which point to object internals). The new delegate is called once all references have been replaced so the reinstanced objects should be self-consistent.

Reworked reinstancing logic to now account for this self-consistency. Key to this is the caching of the 'source' object (i.e. the host anim instance) on nodes that inherit from FAnimNode_CustomProperty - this allows the reinstancing logic to correctly handle linking & property patchup at the point of the delegate call.

Fixed a memory stomp when compiling an anim BP that contains a control rig node. This was caused by the control rig's delegate still being pointed at the old instance due to the lambda binding. Switched this to a raw binding that is cleaned up on node destruction.

diff --git a/Engine/Plugins/Animation/ControlRig/Source/ControlRig/Private/AnimNode_ControlRig.cpp b/Engine/Plugins/Animation/ControlRig/Source/ControlRig/Private/AnimNode_ControlRig.cpp
--- a/Engine/Plugins/Animation/ControlRig/Source/ControlRig/Private/AnimNode_ControlRig.cpp
+++ b/Engine/Plugins/Animation/ControlRig/Source/ControlRig/Private/AnimNode_ControlRig.cpp
@@ -26,6 +26,19 @@ FAnimNode_ControlRig::FAnimNode_ControlRig()
 {
 }
 
+FAnimNode_ControlRig::~FAnimNode_ControlRig()
+{
+	if(ControlRig)
+	{
+		ControlRig->OnInitialized_AnyThread().RemoveAll(this);
+	}
+}
+
+void FAnimNode_ControlRig::HandleOnInitialized_AnyThread(UControlRig*, const EControlRigState, const FName&)
+{
+	RefPoseSetterHash.Reset();
+}
+
 void FAnimNode_ControlRig::OnInitializeAnimInstance(const FAnimInstanceProxy* InProxy, const UAnimInstance* InAnimInstance)
 {
 	DECLARE_SCOPE_HIERARCHICAL_COUNTER_FUNC()
@@ -36,7 +49,7 @@ void FAnimNode_ControlRig::OnInitializeAnimInstance(const FAnimInstanceProxy* In
 		ControlRig->Initialize(true);
 		ControlRig->RequestInit();
 		RefPoseSetterHash.Reset();
-		ControlRig->OnInitialized_AnyThread().AddLambda([this](class UControlRig*, const EControlRigState, const FName&) { RefPoseSetterHash.Reset(); });
+		ControlRig->OnInitialized_AnyThread().AddRaw(this, &FAnimNode_ControlRig::HandleOnInitialized_AnyThread);
 
 		UpdateControlRigRefPoseIfNeeded(InProxy);
 	}
@@ -580,3 +593,18 @@ void FAnimNode_ControlRig::PropagateInputProperties(const UObject* InSourceInsta
 		}
 	}
 }
+
+#if WITH_EDITOR
+
+void FAnimNode_ControlRig::HandleObjectsReinstanced_Impl(UObject* InSourceObject, UObject* InTargetObject, const TMap<UObject*, UObject*>& OldToNewInstanceMap)
+{
+	Super::HandleObjectsReinstanced_Impl(InSourceObject, InTargetObject, OldToNewInstanceMap);
+	
+	if(ControlRig)
+	{
+		ControlRig->OnInitialized_AnyThread().RemoveAll(this);
+		ControlRig->OnInitialized_AnyThread().AddRaw(this, &FAnimNode_ControlRig::HandleOnInitialized_AnyThread);
+	}
+}
+
+#endif
\ No newline at end of file
diff --git a/Engine/Plugins/Animation/ControlRig/Source/ControlRig/Public/AnimNode_ControlRig.h b/Engine/Plugins/Animation/ControlRig/Source/ControlRig/Public/AnimNode_ControlRig.h
--- a/Engine/Plugins/Animation/ControlRig/Source/ControlRig/Public/AnimNode_ControlRig.h
+++ b/Engine/Plugins/Animation/ControlRig/Source/ControlRig/Public/AnimNode_ControlRig.h
@@ -18,6 +18,7 @@ struct CONTROLRIG_API FAnimNode_ControlRig : public FAnimNode_ControlRigBase
 	GENERATED_BODY()
 
 	FAnimNode_ControlRig();
+	~FAnimNode_ControlRig();
 
 	virtual UControlRig* GetControlRig() const override { return ControlRig; }
 	virtual TSubclassOf<UControlRig> GetControlRigClass() const override { return ControlRigClass; }
@@ -35,7 +36,12 @@ struct CONTROLRIG_API FAnimNode_ControlRig : public FAnimNode_ControlRigBase
 
 	virtual void InitializeProperties(const UObject* InSourceInstance, UClass* InTargetClass) override;
 	virtual void PropagateInputProperties(const UObject* InSourceInstance) override;
-
+	
+private:
+	void HandleOnInitialized_AnyThread(UControlRig*, const EControlRigState, const FName&);
+#if WITH_EDITOR
+	virtual void HandleObjectsReinstanced_Impl(UObject* InSourceObject, UObject* InTargetObject, const TMap<UObject*, UObject*>& OldToNewInstanceMap) override;
+#endif
 private:
 
 	/** Cached ControlRig */
diff --git a/Engine/Source/Editor/UnrealEd/Private/Kismet2/KismetReinstanceUtilities.cpp b/Engine/Source/Editor/UnrealEd/Private/Kismet2/KismetReinstanceUtilities.cpp
--- a/Engine/Source/Editor/UnrealEd/Private/Kismet2/KismetReinstanceUtilities.cpp
+++ b/Engine/Source/Editor/UnrealEd/Private/Kismet2/KismetReinstanceUtilities.cpp
@@ -2583,7 +2583,9 @@ void FBlueprintCompileReinstancer::ReplaceInstancesOfClass_Inner(TMap<UClass*, U
 		}
 	}
 
-
+	// Inform listeners of object reinstancing
+	FCoreUObjectDelegates::OnObjectsReinstanced.Broadcast(OldToNewInstanceMap);
+	
 	if(SelectedActors)
 	{
 		SelectedActors->EndBatchSelectOperation();
diff --git a/Engine/Source/Runtime/CoreUObject/Private/UObject/UObjectGlobals.cpp b/Engine/Source/Runtime/CoreUObject/Private/UObject/UObjectGlobals.cpp
--- a/Engine/Source/Runtime/CoreUObject/Private/UObject/UObjectGlobals.cpp
+++ b/Engine/Source/Runtime/CoreUObject/Private/UObject/UObjectGlobals.cpp
@@ -153,6 +153,7 @@ TSet<UObject*> FCoreUObjectDelegates::ObjectsModifiedThisFrame;
 FCoreUObjectDelegates::FOnObjectModified FCoreUObjectDelegates::OnObjectModified;
 FCoreUObjectDelegates::FOnObjectTransacted FCoreUObjectDelegates::OnObjectTransacted;
 FCoreUObjectDelegates::FOnObjectsReplaced FCoreUObjectDelegates::OnObjectsReplaced;
+FCoreUObjectDelegates::FOnObjectsReplaced FCoreUObjectDelegates::OnObjectsReinstanced;
 FCoreUObjectDelegates::FOnObjectPostCDOCompiled FCoreUObjectDelegates::OnObjectPostCDOCompiled;
 FCoreUObjectDelegates::FOnAssetLoaded FCoreUObjectDelegates::OnAssetLoaded;
 FCoreUObjectDelegates::FOnEndLoadPackage FCoreUObjectDelegates::OnEndLoadPackage;
diff --git a/Engine/Source/Runtime/CoreUObject/Public/UObject/UObjectGlobals.h b/Engine/Source/Runtime/CoreUObject/Public/UObject/UObjectGlobals.h
--- a/Engine/Source/Runtime/CoreUObject/Public/UObject/UObjectGlobals.h
+++ b/Engine/Source/Runtime/CoreUObject/Public/UObject/UObjectGlobals.h
@@ -2641,11 +2641,23 @@ struct COREUOBJECT_API FCoreUObjectDelegates
 	DECLARE_MULTICAST_DELEGATE_TwoParams(FOnObjectTransacted, UObject*, const class FTransactionObjectEvent&);
 	static FOnObjectTransacted OnObjectTransacted;
 
-	/** Called when UObjects have been replaced to allow others a chance to fix their references */
+	/**
+	 * Called when UObjects have been replaced to allow others a chance to fix their references
+	 * Note that this is called after properties are copied from old to new instances but before references to replacement
+	 * objects are fixed up in other objects (i.e. other objects can still be pointing to old data)
+	 */
 	using FReplacementObjectMap = TMap<UObject*, UObject*>; // Alias for use in the macro
 	DECLARE_MULTICAST_DELEGATE_OneParam(FOnObjectsReplaced, const FReplacementObjectMap&);
 	static FOnObjectsReplaced OnObjectsReplaced;
 
+	/**
+	 * Called when UObjects have been re-instanced to allow others a chance to fix their references
+	 * Note that this is called after references to replacement objects are fixed up in other objects (i.e. all object
+	 * references should be self-consistent).
+	 */
+	DECLARE_MULTICAST_DELEGATE_OneParam(FOnObjectsReinstanced, const FReplacementObjectMap&);
+	static FOnObjectsReinstanced OnObjectsReinstanced;
+
 	/**
 	 * Called after the Blueprint compiler has finished generating the Class Default Object (CDO) for a class. This can only happen in the editor.
 	 * This is called when the CDO and its associated class structure have been fully generated and populated, and allows the assignment of cached/derived data,
diff --git a/Engine/Source/Runtime/Engine/Classes/Animation/AnimInstance.h b/Engine/Source/Runtime/Engine/Classes/Animation/AnimInstance.h
--- a/Engine/Source/Runtime/Engine/Classes/Animation/AnimInstance.h
+++ b/Engine/Source/Runtime/Engine/Classes/Animation/AnimInstance.h
@@ -1619,6 +1619,7 @@ protected:
 	}
 
 	friend struct FAnimNode_LinkedAnimGraph;
+	friend struct FAnimNode_LinkedAnimLayer;
 	friend struct FAnimInstanceProxy;
 	
 	/** Return whether this AnimNotifyState should be triggered */
@@ -1641,7 +1642,7 @@ public:
 
 #if WITH_EDITOR
 	// Helper function to handle reinstancing in editor
-	virtual void HandleAnimInstanceReplaced(const TMap<UObject*, UObject*>& OldToNewInstanceMap);
+	virtual void HandleObjectsReinstanced(const TMap<UObject*, UObject*>& OldToNewInstanceMap);
 
 	/** Returns true if this anim instance is actively being debugged, false otherwise */
 	bool IsBeingDebugged() const;
diff --git a/Engine/Source/Runtime/Engine/Classes/Animation/AnimNode_CustomProperty.h b/Engine/Source/Runtime/Engine/Classes/Animation/AnimNode_CustomProperty.h
--- a/Engine/Source/Runtime/Engine/Classes/Animation/AnimNode_CustomProperty.h
+++ b/Engine/Source/Runtime/Engine/Classes/Animation/AnimNode_CustomProperty.h
@@ -46,6 +46,14 @@ public:
 		return nullptr;
 	}
 
+	// We only subscribe to the OnInitializeAnimInstance path because we need to cache our source object, so we only
+	// override these methods in editor at the moment
+#if WITH_EDITOR	
+	// FAnimNode_Base interface
+	virtual void OnInitializeAnimInstance(const FAnimInstanceProxy* InProxy, const UAnimInstance* InAnimInstance) override;
+	virtual bool NeedsOnInitializeAnimInstance() const override { return true; }
+#endif
+	
 protected:
 	/** List of source properties to use, 1-1 with Dest names below, built by the compiler */
 	UPROPERTY(meta=(BlueprintCompilerGeneratedDefaults))
@@ -64,22 +72,7 @@ protected:
 
 	/** List of properties on the TargetInstance to push to, built from name list when initialised */
 	TArray<FProperty*> DestProperties;
-
-#if WITH_EDITOR
-	bool bReinitializeProperties;
-#endif // WITH_EDITOR
-
-	virtual bool HasPreUpdate() const override
-	{
-#if WITH_EDITOR
-		return true;
-#else
-		return false;
-#endif
-	}
-
-	virtual void PreUpdate(const UAnimInstance* InAnimInstance) override;
-
+	
 	/* Initialize property links from the source instance, in this case AnimInstance 
 	 * Compiler creates those properties during compile time */
 	virtual void InitializeProperties(const UObject* InSourceInstance, UClass* InTargetClass);
@@ -91,9 +84,25 @@ protected:
 	virtual UClass* GetTargetClass() const PURE_VIRTUAL(FAnimNode_CustomProperty::GetTargetClass, return nullptr;);
 
 #if WITH_EDITOR
-	// Handle object replacement in editor
-	virtual void HandleObjectsReplaced(const TMap<UObject*, UObject*>& OldToNewInstanceMap);
-#endif	// #if WITH_EDITOR
+	/**
+	 * Handle object reinstancing override point.
+	 * When objects are replaced in editor, the FCoreUObjectDelegates::OnObjectsReplaced is called before reference
+	 * replacement, so we cannot handle replacement until later. This call is made on the first PreUpdate after object
+	 * replacement.
+	 */
+	virtual void HandleObjectsReinstanced_Impl(UObject* InSourceObject, UObject* InTargetObject, const TMap<UObject*, UObject*>& OldToNewInstanceMap);
+	
+private:
+	// Handle object reinstancing in editor
+	void HandleObjectsReinstanced(const TMap<UObject*, UObject*>& OldToNewInstanceMap);
+#endif // WITH_EDITOR
+
+#if WITH_EDITORONLY_DATA
+protected:
+	/** This is the source instance, cached to help with re-instancing */
+	UPROPERTY(Transient)
+	TObjectPtr<UObject> SourceInstance;
+#endif
 	
 	friend class UAnimGraphNode_CustomProperty;
 };
diff --git a/Engine/Source/Runtime/Engine/Classes/Animation/AnimNode_LinkedAnimGraph.h b/Engine/Source/Runtime/Engine/Classes/Animation/AnimNode_LinkedAnimGraph.h
--- a/Engine/Source/Runtime/Engine/Classes/Animation/AnimNode_LinkedAnimGraph.h
+++ b/Engine/Source/Runtime/Engine/Classes/Animation/AnimNode_LinkedAnimGraph.h
@@ -120,7 +120,7 @@ protected:
 	}
 
 #if WITH_EDITOR
-	virtual void HandleAnimInstanceReplaced(UAnimInstance* InNewInstance, const TMap<UObject*, UObject*>& OldToNewInstanceMap);
+	virtual void HandleObjectsReinstanced_Impl(UObject* InSourceObject, UObject* InTargetObject, const TMap<UObject*, UObject*>& OldToNewInstanceMap) override;
 #endif	// #if WITH_EDITOR
 
 	/** Link up pose links dynamically with linked instance */
diff --git a/Engine/Source/Runtime/Engine/Classes/Animation/AnimNode_LinkedAnimLayer.h b/Engine/Source/Runtime/Engine/Classes/Animation/AnimNode_LinkedAnimLayer.h
--- a/Engine/Source/Runtime/Engine/Classes/Animation/AnimNode_LinkedAnimLayer.h
+++ b/Engine/Source/Runtime/Engine/Classes/Animation/AnimNode_LinkedAnimLayer.h
@@ -55,7 +55,7 @@ public:
 #if WITH_EDITOR
 	// Event fired when the instance we are running has changed
 	FSimpleMulticastDelegate& OnInstanceChanged() { return OnInstanceChangedEvent; }
-	virtual void HandleAnimInstanceReplaced(UAnimInstance* InNewInstance, const TMap<UObject*, UObject*>& OldToNewInstanceMap) override;
+	virtual void HandleObjectsReinstanced_Impl(UObject* InSourceObject, UObject* InTargetObject, const TMap<UObject*, UObject*>& OldToNewInstanceMap) override;
 #endif
 
 protected:
diff --git a/Engine/Source/Runtime/Engine/Classes/Components/SkeletalMeshComponent.h b/Engine/Source/Runtime/Engine/Classes/Components/SkeletalMeshComponent.h
--- a/Engine/Source/Runtime/Engine/Classes/Components/SkeletalMeshComponent.h
+++ b/Engine/Source/Runtime/Engine/Classes/Components/SkeletalMeshComponent.h
@@ -2548,12 +2548,7 @@ private:
 	 * Cooking does not guarantee skeleton containing all names
 	 */
 	bool AreRequiredCurvesUpToDate() const;
-	
-#if WITH_EDITOR
-	// Handle reinstancing in editor
-	void HandleObjectsReplaced(const TMap<UObject*, UObject*>& OldToNewInstanceMap);
-#endif
-	
+
 public:
 	void ConditionallyDispatchQueuedAnimEvents();
 
diff --git a/Engine/Source/Runtime/Engine/Private/Animation/AnimInstance.cpp b/Engine/Source/Runtime/Engine/Private/Animation/AnimInstance.cpp
--- a/Engine/Source/Runtime/Engine/Private/Animation/AnimInstance.cpp
+++ b/Engine/Source/Runtime/Engine/Private/Animation/AnimInstance.cpp
@@ -120,6 +120,13 @@ UAnimInstance::UAnimInstance(const FObjectInitializer& ObjectInitializer)
 	bReceiveNotifiesFromLinkedInstances = false;
 	bPropagateNotifiesToLinkedInstances = false;
 	bUseMainInstanceMontageEvaluationData = false;
+
+#if WITH_EDITOR
+	if(!HasAnyFlags(RF_ClassDefaultObject))
+	{
+		FCoreUObjectDelegates::OnObjectsReinstanced.AddUObject(this, &UAnimInstance::HandleObjectsReinstanced);
+	}
+#endif // WITH_EDITOR	
 }
 
 // this is only used by montage marker based sync
@@ -3800,38 +3807,38 @@ void UAnimInstance::QueueRootMotionBlend(const FTransform& RootTransform, const
 }
 
 #if WITH_EDITOR
-void UAnimInstance::HandleAnimInstanceReplaced(const TMap<UObject*, UObject*>& OldToNewInstanceMap)
+void UAnimInstance::HandleObjectsReinstanced(const TMap<UObject*, UObject*>& OldToNewInstanceMap)
 {
-	if (IAnimClassInterface* AnimBlueprintClass = IAnimClassInterface::GetFromClass(GetClass()))
+	static IConsoleVariable* UseLegacyAnimInstanceReinstancingBehavior = IConsoleManager::Get().FindConsoleVariable(TEXT("bp.UseLegacyAnimInstanceReinstancingBehavior"));
+	if(UseLegacyAnimInstanceReinstancingBehavior == nullptr || !UseLegacyAnimInstanceReinstancingBehavior->GetBool())
 	{
-		TArray<UObject*> Values;
-		OldToNewInstanceMap.GenerateValueArray(Values);
-		if(Values.Contains(this))
+		bool bThisObjectWasReinstanced = false;
+
+		for(const TPair<UObject*, UObject*>& ObjectPair : OldToNewInstanceMap)
 		{
-			// Re-link graphs & layers if they are targets of a re-instance.
-			// As this gets called post-CPFUO and references are already replaced we need to only do this for 'new'
-			// instances that have just been replaced
-			for(const FStructProperty* LinkedGraphProperty : AnimBlueprintClass->GetLinkedAnimGraphNodeProperties())
+			if(ObjectPair.Value == this)
 			{
-				FAnimNode_LinkedAnimGraph* LinkedAnimGraph = LinkedGraphProperty->ContainerPtrToValuePtr<FAnimNode_LinkedAnimGraph>(this);
-				LinkedAnimGraph->HandleAnimInstanceReplaced(this, OldToNewInstanceMap);
+				bThisObjectWasReinstanced = true;
+				break;
 			}
-
-			for(const FStructProperty* LinkedLayerProperty : AnimBlueprintClass->GetLinkedAnimLayerNodeProperties())
+		}
+		
+		if(bThisObjectWasReinstanced)
+		{
+			RecalcRequiredBones();
+			
+			// Reinit proxy
+			FAnimInstanceProxy& Proxy = GetProxyOnGameThread<FAnimInstanceProxy>();
+			Proxy.Initialize(this);
+			Proxy.InitializeCachedClassData();
+			Proxy.InitializeRootNode_WithRoot(Proxy.RootNode);
+	
+			if(USkeletalMeshComponent* Mesh = GetSkelMeshComponent())
 			{
-				FAnimNode_LinkedAnimLayer* LinkedAnimLayer = LinkedLayerProperty->ContainerPtrToValuePtr<FAnimNode_LinkedAnimLayer>(this);
-				LinkedAnimLayer->HandleAnimInstanceReplaced(this, OldToNewInstanceMap);
+				Mesh->ClearMotionVector();
 			}
 		}
 	}
-	
-	RecalcRequiredBones();
-	
-	// Reinit proxy
-	FAnimInstanceProxy& Proxy = GetProxyOnGameThread<FAnimInstanceProxy>();
-	Proxy.Initialize(this);
-	Proxy.InitializeCachedClassData();
-	Proxy.InitializeRootNode_WithRoot(Proxy.RootNode);
 }
 #endif
 
diff --git a/Engine/Source/Runtime/Engine/Private/Animation/AnimNode_CustomProperty.cpp b/Engine/Source/Runtime/Engine/Private/Animation/AnimNode_CustomProperty.cpp
--- a/Engine/Source/Runtime/Engine/Private/Animation/AnimNode_CustomProperty.cpp
+++ b/Engine/Source/Runtime/Engine/Private/Animation/AnimNode_CustomProperty.cpp
@@ -8,19 +8,16 @@
 FAnimNode_CustomProperty::FAnimNode_CustomProperty()
 	: FAnimNode_Base()
 	, TargetInstance(nullptr)
-#if WITH_EDITOR
-	, bReinitializeProperties(false)
-#endif // WITH_EDITOR
 {
 #if WITH_EDITOR
-	FCoreUObjectDelegates::OnObjectsReplaced.AddRaw(this, &FAnimNode_CustomProperty::HandleObjectsReplaced);
+	FCoreUObjectDelegates::OnObjectsReinstanced.AddRaw(this, &FAnimNode_CustomProperty::HandleObjectsReinstanced);
 #endif // WITH_EDITOR
 }
 
 FAnimNode_CustomProperty::~FAnimNode_CustomProperty()
 {
 #if WITH_EDITOR
-	FCoreUObjectDelegates::OnObjectsReplaced.RemoveAll(this);
+	FCoreUObjectDelegates::OnObjectsReinstanced.RemoveAll(this);
 #endif // WITH_EDITOR
 }
 
@@ -56,19 +53,6 @@ void FAnimNode_CustomProperty::PropagateInputProperties(const UObject* InSourceI
 	}
 }
 
-void FAnimNode_CustomProperty::PreUpdate(const UAnimInstance* InAnimInstance) 
-{
-	FAnimNode_Base::PreUpdate(InAnimInstance);
-
-#if WITH_EDITOR
-	if (bReinitializeProperties)
-	{
-		InitializeProperties(InAnimInstance, GetTargetClass());
-		bReinitializeProperties = false;
-	}
-#endif// WITH_EDITOR
-}
-
 void FAnimNode_CustomProperty::InitializeProperties(const UObject* InSourceInstance, UClass* InTargetClass)
 {
 	if(InTargetClass)
@@ -106,16 +90,47 @@ void FAnimNode_CustomProperty::InitializeProperties(const UObject* InSourceInsta
 }
 
 #if WITH_EDITOR
-void FAnimNode_CustomProperty::HandleObjectsReplaced(const TMap<UObject*, UObject*>& OldToNewInstanceMap)
+
+void FAnimNode_CustomProperty::OnInitializeAnimInstance(const FAnimInstanceProxy* InProxy, const UAnimInstance* InAnimInstance)
 {
-	if (UObject* ThisTargetInstance = GetTargetInstance<UObject>())
+	SourceInstance = const_cast<UAnimInstance*>(InAnimInstance);
+}
+
+void FAnimNode_CustomProperty::HandleObjectsReinstanced_Impl(UObject* InSourceObject, UObject* InTargetObject, const TMap<UObject*, UObject*>& OldToNewInstanceMap)
+{
+	static IConsoleVariable* UseLegacyAnimInstanceReinstancingBehavior = IConsoleManager::Get().FindConsoleVariable(TEXT("bp.UseLegacyAnimInstanceReinstancingBehavior"));
+	if(UseLegacyAnimInstanceReinstancingBehavior == nullptr || !UseLegacyAnimInstanceReinstancingBehavior->GetBool())
 	{
-		UObject* const* ReinstancedTarget = OldToNewInstanceMap.Find(ThisTargetInstance);
-		if (ReinstancedTarget)
+		InitializeProperties(CastChecked<UAnimInstance>(InSourceObject), GetTargetClass());
+	}
+}
+
+void FAnimNode_CustomProperty::HandleObjectsReinstanced(const TMap<UObject*, UObject*>& OldToNewInstanceMap)
+{
+	static IConsoleVariable* UseLegacyAnimInstanceReinstancingBehavior = IConsoleManager::Get().FindConsoleVariable(TEXT("bp.UseLegacyAnimInstanceReinstancingBehavior"));
+	if(UseLegacyAnimInstanceReinstancingBehavior == nullptr || !UseLegacyAnimInstanceReinstancingBehavior->GetBool())
+	{
+		UObject* TargetObject = GetTargetInstance<UObject>();
+		UObject* SourceObject = SourceInstance;
+
+		if(TargetObject && SourceObject)
 		{
-			// recache the properties
-			bReinitializeProperties = true;
+			bool bRelevantObjectWasReinstanced = false;
+			for(const TPair<UObject*, UObject*>& ObjectPair : OldToNewInstanceMap)
+			{
+				if(ObjectPair.Value == TargetObject || ObjectPair.Value == SourceObject)
+				{
+					bRelevantObjectWasReinstanced = true;
+					break;
+				}
+			}
+
+			if(bRelevantObjectWasReinstanced)
+			{
+				HandleObjectsReinstanced_Impl(SourceObject, TargetObject, OldToNewInstanceMap);
+			}
 		}
 	}
 }
+
 #endif	// #if WITH_EDITOR
diff --git a/Engine/Source/Runtime/Engine/Private/Animation/AnimNode_LinkedAnimGraph.cpp b/Engine/Source/Runtime/Engine/Private/Animation/AnimNode_LinkedAnimGraph.cpp
--- a/Engine/Source/Runtime/Engine/Private/Animation/AnimNode_LinkedAnimGraph.cpp
+++ b/Engine/Source/Runtime/Engine/Private/Animation/AnimNode_LinkedAnimGraph.cpp
@@ -197,6 +197,10 @@ void FAnimNode_LinkedAnimGraph::GatherDebugData(FNodeDebugData& DebugData)
 
 void FAnimNode_LinkedAnimGraph::OnInitializeAnimInstance(const FAnimInstanceProxy* InProxy, const UAnimInstance* InAnimInstance)
 {
+#if WITH_EDITORONLY_DATA
+	SourceInstance = const_cast<UAnimInstance*>(InAnimInstance);
+#endif
+	
 	UAnimInstance* InstanceToRun = GetTargetInstance<UAnimInstance>();
 
 	if(*InstanceClass)
@@ -472,13 +476,16 @@ void FAnimNode_LinkedAnimGraph::RequestBlend(const IAnimClassInterface* PriorAni
 }
 
 #if WITH_EDITOR
-void FAnimNode_LinkedAnimGraph::HandleAnimInstanceReplaced(UAnimInstance* InNewInstance, const TMap<UObject*, UObject*>& OldToNewInstanceMap)
+void FAnimNode_LinkedAnimGraph::HandleObjectsReinstanced_Impl(UObject* InSourceObject, UObject* InTargetObject, const TMap<UObject*, UObject*>& OldToNewInstanceMap)
 {
-	if (UAnimInstance* ThisTargetInstance = GetTargetInstance<UAnimInstance>())
+	static IConsoleVariable* UseLegacyAnimInstanceReinstancingBehavior = IConsoleManager::Get().FindConsoleVariable(TEXT("bp.UseLegacyAnimInstanceReinstancingBehavior"));
+	if(UseLegacyAnimInstanceReinstancingBehavior == nullptr || !UseLegacyAnimInstanceReinstancingBehavior->GetBool())
 	{
-		// If we have a target, unlink & re-link
-		DynamicUnlink(InNewInstance);
-		DynamicLink(InNewInstance);
+		UAnimInstance* SourceAnimInstance = CastChecked<UAnimInstance>(InSourceObject);
+	
+		InitializeProperties(SourceAnimInstance, GetTargetClass());
+		DynamicUnlink(SourceAnimInstance);
+		DynamicLink(SourceAnimInstance);
 	}
 }
-#endif
+#endif
\ No newline at end of file
diff --git a/Engine/Source/Runtime/Engine/Private/Animation/AnimNode_LinkedAnimLayer.cpp b/Engine/Source/Runtime/Engine/Private/Animation/AnimNode_LinkedAnimLayer.cpp
--- a/Engine/Source/Runtime/Engine/Private/Animation/AnimNode_LinkedAnimLayer.cpp
+++ b/Engine/Source/Runtime/Engine/Private/Animation/AnimNode_LinkedAnimLayer.cpp
@@ -34,6 +34,10 @@ void FAnimNode_LinkedAnimLayer::InitializeSourceProperties(const UAnimInstance*
 
 void FAnimNode_LinkedAnimLayer::OnInitializeAnimInstance(const FAnimInstanceProxy* InProxy, const UAnimInstance* InAnimInstance)
 {
+#if WITH_EDITORONLY_DATA
+	SourceInstance = const_cast<UAnimInstance*>(InAnimInstance);
+#endif
+	
 	check(SourcePropertyNames.Num() == DestPropertyNames.Num());
 	
 	// Initialize source properties here as they do not change
@@ -141,16 +145,19 @@ void FAnimNode_LinkedAnimLayer::InitializeProperties(const UObject* InSourceInst
 }
 
 #if WITH_EDITOR
-void FAnimNode_LinkedAnimLayer::HandleAnimInstanceReplaced(UAnimInstance* InNewInstance, const TMap<UObject*, UObject*>& OldToNewInstanceMap)
+
+void FAnimNode_LinkedAnimLayer::HandleObjectsReinstanced_Impl(UObject* InSourceObject, UObject* InTargetObject, const TMap<UObject*, UObject*>& OldToNewInstanceMap)
 {
-	UAnimInstance* CurrentTargetInstance = GetTargetInstance<UAnimInstance>();
-	if(CurrentTargetInstance && (Interface.Get() == nullptr || InstanceClass.Get() == nullptr))
+	static IConsoleVariable* UseLegacyAnimInstanceReinstancingBehavior = IConsoleManager::Get().FindConsoleVariable(TEXT("bp.UseLegacyAnimInstanceReinstancingBehavior"));
+	if(UseLegacyAnimInstanceReinstancingBehavior == nullptr || !UseLegacyAnimInstanceReinstancingBehavior->GetBool())
 	{
-		InitializeSelfLayer(CurrentTargetInstance);
-	}
-	else
-	{
-		Super::HandleAnimInstanceReplaced(InNewInstance, OldToNewInstanceMap);
+		UAnimInstance* SourceAnimInstance = CastChecked<UAnimInstance>(InSourceObject);
+		InitializeProperties(SourceAnimInstance, Interface.Get() != nullptr ? Interface.Get() : InSourceObject->GetClass());
+		DynamicUnlink(SourceAnimInstance);
+		DynamicLink(SourceAnimInstance);
+
+		FAnimationInitializeContext Context(&SourceAnimInstance->GetProxyOnAnyThread<FAnimInstanceProxy>());
+		InitializeSubGraph_AnyThread(Context);
 	}
 }
 
diff --git a/Engine/Source/Runtime/Engine/Private/Components/SkeletalMeshComponent.cpp b/Engine/Source/Runtime/Engine/Private/Components/SkeletalMeshComponent.cpp
--- a/Engine/Source/Runtime/Engine/Private/Components/SkeletalMeshComponent.cpp
+++ b/Engine/Source/Runtime/Engine/Private/Components/SkeletalMeshComponent.cpp
@@ -279,10 +279,6 @@ USkeletalMeshComponent::USkeletalMeshComponent(const FObjectInitializer& ObjectI
 	bSkipBoundsUpdateWhenInterpolating = false;
 
 	DeferredKinematicUpdateIndex = INDEX_NONE;
-	
-#if WITH_EDITOR
-	FCoreUObjectDelegates::OnObjectsReplaced.AddUObject(this, &USkeletalMeshComponent::HandleObjectsReplaced);
-#endif // WITH_EDITOR
 }
 
 void USkeletalMeshComponent::Serialize(FArchive& Ar)
@@ -4614,37 +4610,6 @@ TArray<FTransform> USkeletalMeshComponent::GetBoneSpaceTransforms()
 	PRAGMA_ENABLE_DEPRECATION_WARNINGS
 }
 
-
-#if WITH_EDITOR
-void USkeletalMeshComponent::HandleObjectsReplaced(const TMap<UObject*, UObject*>& OldToNewInstanceMap)
-{
-	static IConsoleVariable* UseLegacyAnimInstanceReinstancingBehavior = IConsoleManager::Get().FindConsoleVariable(TEXT("bp.UseLegacyAnimInstanceReinstancingBehavior"));
-	if(UseLegacyAnimInstanceReinstancingBehavior == nullptr || !UseLegacyAnimInstanceReinstancingBehavior->GetBool())
-	{
-		bool bReinstancedAnimInstance = false;
-		
-		ForEachAnimInstance([&OldToNewInstanceMap, &bReinstancedAnimInstance](UAnimInstance* InAnimInstance)
-		{
-			UObject* const* ReinstancedTarget = OldToNewInstanceMap.Find(InAnimInstance);
-			if(ReinstancedTarget)
-			{
-				if (UAnimInstance* ReinstancedTargetInstance = Cast<UAnimInstance>(*ReinstancedTarget))
-				{
-					ReinstancedTargetInstance->HandleAnimInstanceReplaced(OldToNewInstanceMap);
-					bReinstancedAnimInstance = true;
-				}
-			}
-		});
-
-		if(bReinstancedAnimInstance)
-		{
-			ClearMotionVector();
-		}
-	}
-}
-#endif	// #if WITH_EDITOR
-
-
 #if WITH_EDITOR
 void USkeletalMeshComponent::DebugDrawPoseWatches(FPrimitiveDrawInterface* PDI)
 {
