Undo changes from cl 20967184 while considering changes made to EStrippedData becoming an enum class

diff --git a/Engine/Source/Runtime/Engine/Private/EngineUtils.cpp b/Engine/Source/Runtime/Engine/Private/EngineUtils.cpp
--- a/Engine/Source/Runtime/Engine/Private/EngineUtils.cpp
+++ b/Engine/Source/Runtime/Engine/Private/EngineUtils.cpp
@@ -556,7 +556,7 @@ FStripDataFlags::FStripDataFlags( class FArchive& Ar, uint8 InClassFlags /*= 0*/
 		{
 			// When cooking GlobalStripFlags are automatically generated based on the current target
 			// platform's properties.
-			GlobalStripFlags |= Ar.IsFilterEditorOnly() ? static_cast<uint8>(FStripDataFlags::EStrippedData::Editor) : static_cast<uint8>(FStripDataFlags::EStrippedData::None);
+			GlobalStripFlags |= Ar.CookingTarget()->HasEditorOnlyData() ? static_cast<uint8>(FStripDataFlags::EStrippedData::None) : static_cast<uint8>(FStripDataFlags::EStrippedData::Editor);
 			GlobalStripFlags |= !Ar.CookingTarget()->AllowAudioVisualData() ? static_cast<uint8>(FStripDataFlags::EStrippedData::Server) : static_cast<uint8>(FStripDataFlags::EStrippedData::None);
 			ClassStripFlags = InClassFlags;
 		}
