Fixing a comment when ending the game in Gamemode.cpp

#9380

diff --git a/Engine/Source/Runtime/Engine/Private/GameMode.cpp b/Engine/Source/Runtime/Engine/Private/GameMode.cpp
--- a/Engine/Source/Runtime/Engine/Private/GameMode.cpp
+++ b/Engine/Source/Runtime/Engine/Private/GameMode.cpp
@@ -386,7 +386,7 @@ void AGameMode::Tick(float DeltaSeconds)
 	}
 	if (GetMatchState() == MatchState::InProgress)
 	{
-		// Check to see if we should start the match
+		// Check to see if we should end the match
 		if (ReadyToEndMatch())
 		{
 			UE_LOG(LogGameMode, Log, TEXT("GameMode returned ReadyToEndMatch"));
