Unshelved from pending changelist '21012765':

Fix for APlayerCameraManager becoming unreachable via a spawned ACameraActor's owner pointer.

An APlayerController was spawned early during level load that was later cleaned up and replaced with the project's player controller.
But in this case the APlayerCameraManager never had its begin play called and therefore never cleaned up the spawned anim camera in end play.

#tests PIE no longer seeing unreachable warning

diff --git a/Engine/Source/Runtime/Engine/Classes/Camera/PlayerCameraManager.h b/Engine/Source/Runtime/Engine/Classes/Camera/PlayerCameraManager.h
--- a/Engine/Source/Runtime/Engine/Classes/Camera/PlayerCameraManager.h
+++ b/Engine/Source/Runtime/Engine/Classes/Camera/PlayerCameraManager.h
@@ -557,6 +557,7 @@ public:
 	virtual bool ShouldTickIfViewportsOnly() const override;
 	virtual void PostInitializeComponents() override;
 	virtual void EndPlay(const EEndPlayReason::Type EndPlayReason) override;
+	virtual void Destroyed() override;
 	virtual void DisplayDebug(class UCanvas* Canvas, const FDebugDisplayInfo& DebugDisplay, float& YL, float& YPos) override;
 	virtual void ApplyWorldOffset(const FVector& InOffset, bool bWorldShift) override;
 	//~ End AActor Interface
@@ -630,6 +631,8 @@ protected:
 	/** Internal. Places the given modifier in the ModifierList at the appropriate priority. */
 	virtual bool AddCameraModifierToList(UCameraModifier* NewModifier);
 
+	virtual void CleanUpAnimCamera(const bool bDestroy);
+
 public:	
 	/**
 	 * Applies the current set of camera modifiers to the given camera POV.
diff --git a/Engine/Source/Runtime/Engine/Private/PlayerCameraManager.cpp b/Engine/Source/Runtime/Engine/Private/PlayerCameraManager.cpp
--- a/Engine/Source/Runtime/Engine/Private/PlayerCameraManager.cpp
+++ b/Engine/Source/Runtime/Engine/Private/PlayerCameraManager.cpp
@@ -746,6 +746,20 @@
 	return false;
 }
 
+void APlayerCameraManager::CleanUpAnimCamera(const bool bDestroy)
+{
+	// clean up the temp camera actor
+	if (AnimCameraActor != nullptr)
+	{
+		if (bDestroy)
+		{
+			AnimCameraActor->Destroy();
+		}
+		AnimCameraActor->SetOwner(nullptr);
+		AnimCameraActor = nullptr;
+	}
+}
+
 bool APlayerCameraManager::RemoveCameraModifier(UCameraModifier* ModifierToRemove)
 {
 	if (ModifierToRemove)
@@ -822,21 +836,19 @@
 
 void APlayerCameraManager::EndPlay(const EEndPlayReason::Type EndPlayReason)
 {
-	// clean up the temp camera actor
-	if (AnimCameraActor)
-	{
-		if (EndPlayReason == EEndPlayReason::Destroyed)
-		{
-			AnimCameraActor->Destroy();
-		}
-		AnimCameraActor->SetOwner(nullptr);
-		AnimCameraActor = nullptr;
-	}
 	ActiveAnims.Empty();
 	ModifierList.Empty();
+	CleanUpAnimCamera(EndPlayReason == EEndPlayReason::Destroyed);
 	Super::EndPlay(EndPlayReason);
 }
 
+void APlayerCameraManager::Destroyed()
+{
+	CleanUpAnimCamera(true);
+
+	Super::Destroyed();
+}
+
 void APlayerCameraManager::InitializeFor(APlayerController* PC)
 {
 	FMinimalViewInfo DefaultFOVCache = GetCameraCacheView();
