nullptr fix when running content worker duplicate command

diff --git a/Engine/Source/Runtime/Engine/Private/WorldPartition/WorldPartitionActorDesc.cpp b/Engine/Source/Runtime/Engine/Private/WorldPartition/WorldPartitionActorDesc.cpp
--- a/Engine/Source/Runtime/Engine/Private/WorldPartition/WorldPartitionActorDesc.cpp
+++ b/Engine/Source/Runtime/Engine/Private/WorldPartition/WorldPartitionActorDesc.cpp
@@ -66,23 +66,24 @@ void FWorldPartitionActorDesc::Init(const AActor* InActor)
 	{
 		TArray<FName> LocalDataLayerAssetPaths;
 		TArray<FName> LocalDataLayerInstanceNames;
-		UDataLayerSubsystem* DataLayerSubsystem = UWorld::GetSubsystem<UDataLayerSubsystem>(InActor->GetWorld());
 		ULevel* Level = InActor->GetLevel();
-
-		LocalDataLayerAssetPaths.Reserve(InActor->GetDataLayerAssets().Num());
-		for (const TObjectPtr<const UDataLayerAsset>& DataLayerAsset : InActor->GetDataLayerAssets())
+		if (UDataLayerSubsystem* DataLayerSubsystem = UWorld::GetSubsystem<UDataLayerSubsystem>(InActor->GetWorld()))
 		{
-			// Pass Actor Level when resolving the DataLayerInstance as FWorldPartitionActorDesc always represents the state of the actor local to its outer level
-			if (DataLayerAsset && DataLayerSubsystem->GetDataLayerInstance(DataLayerAsset, Level))
+			LocalDataLayerAssetPaths.Reserve(InActor->GetDataLayerAssets().Num());
+			for (const TObjectPtr<const UDataLayerAsset>& DataLayerAsset : InActor->GetDataLayerAssets())
 			{
-				LocalDataLayerAssetPaths.Add(*DataLayerAsset->GetPathName());
+				// Pass Actor Level when resolving the DataLayerInstance as FWorldPartitionActorDesc always represents the state of the actor local to its outer level
+				if (DataLayerAsset && DataLayerSubsystem->GetDataLayerInstance(DataLayerAsset, Level))
+				{
+					LocalDataLayerAssetPaths.Add(*DataLayerAsset->GetPathName());
+				}
 			}
-		}
 
-		PRAGMA_DISABLE_DEPRECATION_WARNINGS
-		// Pass Actor Level when resolving the DataLayerInstance as FWorldPartitionActorDesc always represents the state of the actor local to its outer level
-		LocalDataLayerInstanceNames = DataLayerSubsystem->GetDataLayerInstanceNames(InActor->GetActorDataLayers(), Level);
-		PRAGMA_ENABLE_DEPRECATION_WARNINGS
+			PRAGMA_DISABLE_DEPRECATION_WARNINGS
+			// Pass Actor Level when resolving the DataLayerInstance as FWorldPartitionActorDesc always represents the state of the actor local to its outer level
+			LocalDataLayerInstanceNames = DataLayerSubsystem->GetDataLayerInstanceNames(InActor->GetActorDataLayers(), Level);
+			PRAGMA_ENABLE_DEPRECATION_WARNINGS
+		}
 
 		// Validation
 		const bool bHasDataLayerAssets = LocalDataLayerAssetPaths.Num() > 0;
