Fixed discrepancies between material instance vector parameters and double vector parameters.

diff --git a/Engine/Source/Editor/MaterialEditor/Private/MaterialInstanceEditor.cpp b/Engine/Source/Editor/MaterialEditor/Private/MaterialInstanceEditor.cpp
--- a/Engine/Source/Editor/MaterialEditor/Private/MaterialInstanceEditor.cpp
+++ b/Engine/Source/Editor/MaterialEditor/Private/MaterialInstanceEditor.cpp
@@ -521,6 +521,7 @@ void FMaterialInstanceEditor::ReInitMaterialFunctionProxies()
 		// Temporarily store unsaved parameters
 		TArray<FScalarParameterValue> ScalarParameterValues = FunctionInstanceProxy->ScalarParameterValues;
 		TArray<FVectorParameterValue> VectorParameterValues = FunctionInstanceProxy->VectorParameterValues;
+		TArray<FDoubleVectorParameterValue> DoubleVectorParameterValues = FunctionInstanceProxy->DoubleVectorParameterValues;
 		TArray<FTextureParameterValue> TextureParameterValues = FunctionInstanceProxy->TextureParameterValues;
 		TArray<FRuntimeVirtualTextureParameterValue> RuntimeVirtualTextureParameterValues = FunctionInstanceProxy->RuntimeVirtualTextureParameterValues;
 		TArray<FFontParameterValue> FontParameterValues = FunctionInstanceProxy->FontParameterValues;
@@ -562,6 +563,18 @@ void FMaterialInstanceEditor::ReInitMaterialFunctionProxies()
 			}
 		}
 
+		FunctionInstanceProxy->GetAllDoubleVectorParameterInfo(OutParameterInfo, Guids);
+		FunctionInstanceProxy->DoubleVectorParameterValues.Empty();
+		for (FDoubleVectorParameterValue& DoubleVectorParameter : DoubleVectorParameterValues)
+		{
+			int32 Index = Guids.Find(DoubleVectorParameter.ExpressionGUID);
+			if (Index != INDEX_NONE)
+			{
+				FunctionInstanceProxy->DoubleVectorParameterValues.Add(DoubleVectorParameter);
+				FunctionInstanceProxy->DoubleVectorParameterValues.Last().ParameterInfo = OutParameterInfo[Index];
+			}
+		}
+
 		FunctionInstanceProxy->GetAllTextureParameterInfo(OutParameterInfo, Guids);
 		FunctionInstanceProxy->TextureParameterValues.Empty();
 		for (FTextureParameterValue& TextureParameter : TextureParameterValues)
diff --git a/Engine/Source/Editor/MaterialEditor/Private/MaterialPropertyHelpers.cpp b/Engine/Source/Editor/MaterialEditor/Private/MaterialPropertyHelpers.cpp
--- a/Engine/Source/Editor/MaterialEditor/Private/MaterialPropertyHelpers.cpp
+++ b/Engine/Source/Editor/MaterialEditor/Private/MaterialPropertyHelpers.cpp
@@ -423,6 +423,7 @@ FReply FMaterialPropertyHelpers::OnClickedSaveNewFunctionInstance(class UMateria
 				{
 					ChildInstance->ScalarParameterValues = EditedInstance->ScalarParameterValues;
 					ChildInstance->VectorParameterValues = EditedInstance->VectorParameterValues;
+					ChildInstance->DoubleVectorParameterValues = EditedInstance->DoubleVectorParameterValues;
 					ChildInstance->TextureParameterValues = EditedInstance->TextureParameterValues;
 					ChildInstance->RuntimeVirtualTextureParameterValues = EditedInstance->RuntimeVirtualTextureParameterValues; 
 					ChildInstance->FontParameterValues = EditedInstance->FontParameterValues;
@@ -514,6 +515,7 @@ FReply FMaterialPropertyHelpers::OnClickedSaveNewLayerInstance(class UMaterialFu
 				{
 					ChildInstance->ScalarParameterValues = EditedInstance->ScalarParameterValues;
 					ChildInstance->VectorParameterValues = EditedInstance->VectorParameterValues;
+					ChildInstance->DoubleVectorParameterValues = EditedInstance->DoubleVectorParameterValues;
 					ChildInstance->TextureParameterValues = EditedInstance->TextureParameterValues;
 					ChildInstance->RuntimeVirtualTextureParameterValues = EditedInstance->RuntimeVirtualTextureParameterValues;
 					ChildInstance->FontParameterValues = EditedInstance->FontParameterValues;
diff --git a/Engine/Source/Editor/UnrealEd/Private/PreviewMaterial.cpp b/Engine/Source/Editor/UnrealEd/Private/PreviewMaterial.cpp
--- a/Engine/Source/Editor/UnrealEd/Private/PreviewMaterial.cpp
+++ b/Engine/Source/Editor/UnrealEd/Private/PreviewMaterial.cpp
@@ -824,6 +824,7 @@ void UMaterialEditorInstanceConstant::ApplySourceFunctionChanges()
 		// Copy updated function parameter values	
 		SourceFunction->ScalarParameterValues = SourceInstance->ScalarParameterValues;
 		SourceFunction->VectorParameterValues = SourceInstance->VectorParameterValues;
+		SourceFunction->DoubleVectorParameterValues = SourceInstance->DoubleVectorParameterValues;
 		SourceFunction->TextureParameterValues = SourceInstance->TextureParameterValues;
 		SourceFunction->RuntimeVirtualTextureParameterValues = SourceInstance->RuntimeVirtualTextureParameterValues;
 		SourceFunction->FontParameterValues = SourceInstance->FontParameterValues;
diff --git a/Engine/Source/Runtime/Engine/Classes/Materials/MaterialInterface.h b/Engine/Source/Runtime/Engine/Classes/Materials/MaterialInterface.h
--- a/Engine/Source/Runtime/Engine/Classes/Materials/MaterialInterface.h
+++ b/Engine/Source/Runtime/Engine/Classes/Materials/MaterialInterface.h
@@ -545,6 +545,7 @@ public:
 	ENGINE_API virtual void GetAllParametersOfType(EMaterialParameterType Type, TMap<FMaterialParameterInfo, FMaterialParameterMetadata>& OutParameters) const;
 	ENGINE_API void GetAllScalarParameterInfo(TArray<FMaterialParameterInfo>& OutParameterInfo, TArray<FGuid>& OutParameterIds) const;
 	ENGINE_API void GetAllVectorParameterInfo(TArray<FMaterialParameterInfo>& OutParameterInfo, TArray<FGuid>& OutParameterIds) const;
+	ENGINE_API void GetAllDoubleVectorParameterInfo(TArray<FMaterialParameterInfo>& OutParameterInfo, TArray<FGuid>& OutParameterIds) const;
 	ENGINE_API void GetAllTextureParameterInfo(TArray<FMaterialParameterInfo>& OutParameterInfo, TArray<FGuid>& OutParameterIds) const;
 	ENGINE_API void GetAllRuntimeVirtualTextureParameterInfo(TArray<FMaterialParameterInfo>& OutParameterInfo, TArray<FGuid>& OutParameterIds) const;
 	ENGINE_API void GetAllFontParameterInfo(TArray<FMaterialParameterInfo>& OutParameterInfo, TArray<FGuid>& OutParameterIds) const;
@@ -562,6 +563,7 @@ public:
 	ENGINE_API bool GetParameterDefaultValue(EMaterialParameterType Type, const FMemoryImageMaterialParameterInfo& ParameterInfo, FMaterialParameterMetadata& OutValue) const;
 	ENGINE_API bool GetScalarParameterDefaultValue(const FHashedMaterialParameterInfo& ParameterInfo, float& OutValue) const;
 	ENGINE_API bool GetVectorParameterDefaultValue(const FHashedMaterialParameterInfo& ParameterInfo, FLinearColor& OutValue) const;
+	ENGINE_API bool GetDoubleVectorParameterDefaultValue(const FHashedMaterialParameterInfo& ParameterInfo, FVector4d& OutValue) const;
 	ENGINE_API bool GetTextureParameterDefaultValue(const FHashedMaterialParameterInfo& ParameterInfo, class UTexture*& OutValue) const;
 	ENGINE_API bool GetRuntimeVirtualTextureParameterDefaultValue(const FHashedMaterialParameterInfo& ParameterInfo, class URuntimeVirtualTexture*& OutValue) const;
 	ENGINE_API bool GetFontParameterDefaultValue(const FHashedMaterialParameterInfo& ParameterInfo, class UFont*& OutFontValue, int32& OutFontPage) const;
@@ -773,6 +775,11 @@ public:
 	ENGINE_API bool IsVectorParameterUsedAsChannelMask(const FHashedMaterialParameterInfo& ParameterInfo, bool& OutValue) const;
 	ENGINE_API bool GetVectorParameterChannelNames(const FHashedMaterialParameterInfo& ParameterInfo, FParameterChannelNames& OutValue) const;
 #endif
+	ENGINE_API bool GetDoubleVectorParameterValue(const FHashedMaterialParameterInfo& ParameterInfo, FVector4d& OutValue, bool bOveriddenOnly = false) const;
+#if WITH_EDITOR
+	ENGINE_API bool IsDoubleVectorParameterUsedAsChannelMask(const FHashedMaterialParameterInfo& ParameterInfo, bool& OutValue) const;
+	ENGINE_API bool GetDoubleVectorParameterChannelNames(const FHashedMaterialParameterInfo& ParameterInfo, FParameterChannelNames& OutValue) const;
+#endif // WITH_EDITOR
 	ENGINE_API bool GetTextureParameterValue(const FHashedMaterialParameterInfo& ParameterInfo, class UTexture*& OutValue, bool bOveriddenOnly = false) const;
 	ENGINE_API bool GetRuntimeVirtualTextureParameterValue(const FHashedMaterialParameterInfo& ParameterInfo, class URuntimeVirtualTexture*& OutValue, bool bOveriddenOnly = false) const;
 #if WITH_EDITOR
diff --git a/Engine/Source/Runtime/Engine/Private/Materials/MaterialExpressions.cpp b/Engine/Source/Runtime/Engine/Private/Materials/MaterialExpressions.cpp
--- a/Engine/Source/Runtime/Engine/Private/Materials/MaterialExpressions.cpp
+++ b/Engine/Source/Runtime/Engine/Private/Materials/MaterialExpressions.cpp
@@ -14482,6 +14482,17 @@ void UMaterialFunctionInstance::UpdateParameterSet()
 						}
 					}
 				}
+				else if (const UMaterialExpressionDoubleVectorParameter* DoubleVectorParameter = Cast<const UMaterialExpressionDoubleVectorParameter>(FunctionExpression))
+				{
+					for (FDoubleVectorParameterValue& DoubleVectorParameterValue : DoubleVectorParameterValues)
+					{
+						if (DoubleVectorParameterValue.ExpressionGUID == DoubleVectorParameter->ExpressionGUID)
+						{
+							DoubleVectorParameterValue.ParameterInfo.Name = DoubleVectorParameter->ParameterName;
+							break;
+						}
+					}
+				}
 				else if (const UMaterialExpressionTextureSampleParameter* TextureParameter = Cast<const UMaterialExpressionTextureSampleParameter>(FunctionExpression))
 				{
 					for (FTextureParameterValue& TextureParameterValue : TextureParameterValues)
@@ -14547,6 +14558,7 @@ void UMaterialFunctionInstance::OverrideMaterialInstanceParameterValues(UMateria
 	// Dynamic parameters
 	Instance->ScalarParameterValues = ScalarParameterValues;
 	Instance->VectorParameterValues = VectorParameterValues;
+	Instance->DoubleVectorParameterValues = DoubleVectorParameterValues;
 	Instance->TextureParameterValues = TextureParameterValues;
 	Instance->RuntimeVirtualTextureParameterValues = RuntimeVirtualTextureParameterValues;
 	Instance->FontParameterValues = FontParameterValues;
@@ -14706,6 +14718,7 @@ bool UMaterialFunctionInstance::GetParameterOverrideValue(EMaterialParameterType
 	{
 	case EMaterialParameterType::Scalar: bResult = GameThread_GetParameterValue(ScalarParameterValues, ParameterInfo, OutResult); break;
 	case EMaterialParameterType::Vector: bResult = GameThread_GetParameterValue(VectorParameterValues, ParameterInfo, OutResult); break;
+	case EMaterialParameterType::DoubleVector: bResult = GameThread_GetParameterValue(DoubleVectorParameterValues, ParameterInfo, OutResult); break;
 	case EMaterialParameterType::Texture: bResult = GameThread_GetParameterValue(TextureParameterValues, ParameterInfo, OutResult); break;
 	case EMaterialParameterType::RuntimeVirtualTexture: bResult = GameThread_GetParameterValue(RuntimeVirtualTextureParameterValues, ParameterInfo, OutResult); break;
 	case EMaterialParameterType::Font: bResult = GameThread_GetParameterValue(FontParameterValues, ParameterInfo, OutResult); break;
diff --git a/Engine/Source/Runtime/Engine/Private/Materials/MaterialInstance.cpp b/Engine/Source/Runtime/Engine/Private/Materials/MaterialInstance.cpp
--- a/Engine/Source/Runtime/Engine/Private/Materials/MaterialInstance.cpp
+++ b/Engine/Source/Runtime/Engine/Private/Materials/MaterialInstance.cpp
@@ -3408,6 +3408,7 @@ void UMaterialInstance::ClearParameterValuesInternal(EMaterialInstanceClearParam
 	{
 		ScalarParameterValues.Empty();
 		VectorParameterValues.Empty();
+		DoubleVectorParameterValues.Empty();
 		bUpdateResource = true;
 	}
 
@@ -3782,6 +3783,7 @@ void UMaterialInstance::GetResourceSizeEx(FResourceSizeEx& CumulativeResourceSiz
 		CumulativeResourceSize.AddDedicatedSystemMemoryBytes(sizeof(FMaterialInstanceResource));
 		CumulativeResourceSize.AddDedicatedSystemMemoryBytes(ScalarParameterValues.Num() * sizeof(THashedMaterialParameterMap<float>::TNamedParameter));
 		CumulativeResourceSize.AddDedicatedSystemMemoryBytes(VectorParameterValues.Num() * sizeof(THashedMaterialParameterMap<FLinearColor>::TNamedParameter));
+		CumulativeResourceSize.AddDedicatedSystemMemoryBytes(DoubleVectorParameterValues.Num() * sizeof(THashedMaterialParameterMap<FVector4d>::TNamedParameter));
 		CumulativeResourceSize.AddDedicatedSystemMemoryBytes(TextureParameterValues.Num() * sizeof(THashedMaterialParameterMap<const UTexture*>::TNamedParameter));
 		CumulativeResourceSize.AddDedicatedSystemMemoryBytes(RuntimeVirtualTextureParameterValues.Num() * sizeof(THashedMaterialParameterMap<const URuntimeVirtualTexture*>::TNamedParameter));
 		CumulativeResourceSize.AddDedicatedSystemMemoryBytes(FontParameterValues.Num() * sizeof(THashedMaterialParameterMap<const UTexture*>::TNamedParameter));
@@ -3789,6 +3791,7 @@ void UMaterialInstance::GetResourceSizeEx(FResourceSizeEx& CumulativeResourceSiz
 		// Record space for hash tables as well..
 		if (ScalarParameterValues.Num()) CumulativeResourceSize.AddDedicatedSystemMemoryBytes(FDefaultSetAllocator::GetNumberOfHashBuckets(ScalarParameterValues.Num()) * sizeof(uint16));
 		if (VectorParameterValues.Num()) CumulativeResourceSize.AddDedicatedSystemMemoryBytes(FDefaultSetAllocator::GetNumberOfHashBuckets(VectorParameterValues.Num()) * sizeof(uint16));
+		if (DoubleVectorParameterValues.Num()) CumulativeResourceSize.AddDedicatedSystemMemoryBytes(FDefaultSetAllocator::GetNumberOfHashBuckets(DoubleVectorParameterValues.Num()) * sizeof(uint16));
 		if (TextureParameterValues.Num()) CumulativeResourceSize.AddDedicatedSystemMemoryBytes(FDefaultSetAllocator::GetNumberOfHashBuckets(TextureParameterValues.Num()) * sizeof(uint16));
 		if (RuntimeVirtualTextureParameterValues.Num()) CumulativeResourceSize.AddDedicatedSystemMemoryBytes(FDefaultSetAllocator::GetNumberOfHashBuckets(RuntimeVirtualTextureParameterValues.Num()) * sizeof(uint16));
 		if (FontParameterValues.Num()) CumulativeResourceSize.AddDedicatedSystemMemoryBytes(FDefaultSetAllocator::GetNumberOfHashBuckets(FontParameterValues.Num()) * sizeof(uint16));
@@ -4180,6 +4183,10 @@ bool UMaterialInstance::Equivalent(const UMaterialInstance* CompareTo) const
 	{
 		return false;
 	}
+	if (!CompareValueArraysByExpressionGUID(DoubleVectorParameterValues, CompareTo->DoubleVectorParameterValues))
+	{
+		return false;
+	}
 	if (!CompareValueArraysByExpressionGUID(RuntimeVirtualTextureParameterValues, CompareTo->RuntimeVirtualTextureParameterValues))
 	{
 		return false;
diff --git a/Engine/Source/Runtime/Engine/Private/Materials/MaterialInstanceDynamic.cpp b/Engine/Source/Runtime/Engine/Private/Materials/MaterialInstanceDynamic.cpp
--- a/Engine/Source/Runtime/Engine/Private/Materials/MaterialInstanceDynamic.cpp
+++ b/Engine/Source/Runtime/Engine/Private/Materials/MaterialInstanceDynamic.cpp
@@ -462,6 +462,11 @@ void UMaterialInstanceDynamic::CopyInterpParameters(UMaterialInstance* Source)
 			SetVectorParameterValue(it.ParameterInfo.Name, it.ParameterValue);
 		}
 
+		for (auto& it : Source->DoubleVectorParameterValues)
+		{
+			SetDoubleVectorParameterValue(it.ParameterInfo.Name, it.ParameterValue);
+		}
+
 		for (auto& it : Source->TextureParameterValues)
 		{
 			SetTextureParameterValue(it.ParameterInfo.Name, it.ParameterValue);
@@ -482,6 +487,7 @@ void UMaterialInstanceDynamic::CopyParameterOverrides(UMaterialInstance* Materia
 	if (ensureAsRuntimeWarning(MaterialInstance != nullptr))
 	{
 		VectorParameterValues = MaterialInstance->VectorParameterValues;
+		DoubleVectorParameterValues = MaterialInstance->DoubleVectorParameterValues;
 		ScalarParameterValues = MaterialInstance->ScalarParameterValues;
 		TextureParameterValues = MaterialInstance->TextureParameterValues;
 		FontParameterValues = MaterialInstance->FontParameterValues;
diff --git a/Engine/Source/Runtime/Engine/Private/Materials/MaterialInterface.cpp b/Engine/Source/Runtime/Engine/Private/Materials/MaterialInterface.cpp
--- a/Engine/Source/Runtime/Engine/Private/Materials/MaterialInterface.cpp
+++ b/Engine/Source/Runtime/Engine/Private/Materials/MaterialInterface.cpp
@@ -701,6 +701,17 @@ bool UMaterialInterface::GetVectorParameterValue(const FHashedMaterialParameterI
 	return false;
 }
 
+bool UMaterialInterface::GetDoubleVectorParameterValue(const FHashedMaterialParameterInfo& ParameterInfo, FVector4d& OutValue, bool bOveriddenOnly) const
+{
+	FMaterialParameterMetadata Result;
+	if (GetParameterValue(EMaterialParameterType::DoubleVector, ParameterInfo, Result, MakeParameterValueFlags(bOveriddenOnly)))
+	{
+		OutValue = Result.Value.AsVector4d();
+		return true;
+	}
+	return false;
+}
+
 #if WITH_EDITOR
 bool UMaterialInterface::IsVectorParameterUsedAsChannelMask(const FHashedMaterialParameterInfo& ParameterInfo, bool& OutValue) const
 {
@@ -724,6 +735,28 @@ bool UMaterialInterface::GetVectorParameterChannelNames(const FHashedMaterialPar
 	return false;
 }
 
+bool UMaterialInterface::IsDoubleVectorParameterUsedAsChannelMask(const FHashedMaterialParameterInfo& ParameterInfo, bool& OutValue) const
+{
+	FMaterialParameterMetadata Result;
+	if (GetParameterValue(EMaterialParameterType::DoubleVector, ParameterInfo, Result, EMaterialGetParameterValueFlags::CheckNonOverrides))
+	{
+		OutValue = Result.bUsedAsChannelMask;
+		return true;
+	}
+	return false;
+}
+
+bool UMaterialInterface::GetDoubleVectorParameterChannelNames(const FHashedMaterialParameterInfo& ParameterInfo, FParameterChannelNames& OutValue) const
+{
+	FMaterialParameterMetadata Result;
+	if (GetParameterValue(EMaterialParameterType::DoubleVector, ParameterInfo, Result, EMaterialGetParameterValueFlags::CheckNonOverrides))
+	{
+		OutValue = Result.ChannelNames;
+		return true;
+	}
+	return false;
+}
+
 bool UMaterialInterface::GetScalarParameterSliderMinMax(const FHashedMaterialParameterInfo& ParameterInfo, float& OutSliderMin, float& OutSliderMax) const
 {
 	FMaterialParameterMetadata Result;
@@ -843,6 +876,17 @@ bool UMaterialInterface::GetVectorParameterDefaultValue(const FHashedMaterialPar
 	return false;
 }
 
+bool UMaterialInterface::GetDoubleVectorParameterDefaultValue(const FHashedMaterialParameterInfo& ParameterInfo, FVector4& OutValue) const
+{
+	FMaterialParameterMetadata Result;
+	if (GetParameterDefaultValue(EMaterialParameterType::DoubleVector, ParameterInfo, Result))
+	{
+		OutValue = Result.Value.AsVector4d();
+		return true;
+	}
+	return false;
+}
+
 bool UMaterialInterface::GetTextureParameterDefaultValue(const FHashedMaterialParameterInfo& ParameterInfo, class UTexture*& OutValue) const
 {
 	FMaterialParameterMetadata Result;
@@ -931,6 +975,11 @@ void UMaterialInterface::GetAllVectorParameterInfo(TArray<FMaterialParameterInfo
 	GetAllParameterInfoOfType(EMaterialParameterType::Vector, OutParameterInfo, OutParameterIds);
 }
 
+void UMaterialInterface::GetAllDoubleVectorParameterInfo(TArray<FMaterialParameterInfo>& OutParameterInfo, TArray<FGuid>& OutParameterIds) const
+{
+	GetAllParameterInfoOfType(EMaterialParameterType::DoubleVector, OutParameterInfo, OutParameterIds);
+}
+
 void UMaterialInterface::GetAllTextureParameterInfo(TArray<FMaterialParameterInfo>& OutParameterInfo, TArray<FGuid>& OutParameterIds) const
 {
 	GetAllParameterInfoOfType(EMaterialParameterType::Texture, OutParameterInfo, OutParameterIds);
