ActorComponent init flow:
* Added ActorComponent ReadyForReplication() event function.
* This is called between InitComponent and before BeginPlay during the owning actor's BeginPlay
* Should be used to register preexisting replicated subobjects. It will make it so subobjects get replicated even if a component called an RPC on itself early inside BeginPlay.

diff --git a/Engine/Plugins/Runtime/GameplayAbilities/Source/GameplayAbilities/Private/AbilitySystemComponent.cpp b/Engine/Plugins/Runtime/GameplayAbilities/Source/GameplayAbilities/Private/AbilitySystemComponent.cpp
--- a/Engine/Plugins/Runtime/GameplayAbilities/Source/GameplayAbilities/Private/AbilitySystemComponent.cpp
+++ b/Engine/Plugins/Runtime/GameplayAbilities/Source/GameplayAbilities/Private/AbilitySystemComponent.cpp
@@ -232,6 +232,11 @@ void UAbilitySystemComponent::BeginPlay()
 
 	// Cache net role here as well since for map-placed actors on clients, the Role may not be set correctly yet in OnRegister.
 	CacheIsNetSimulated();
+}
+
+void UAbilitySystemComponent::ReadyForReplication()
+{
+	Super::ReadyForReplication();
 
 	// Register the spawned attributes to the replicated sub object list.
 	if (IsUsingRegisteredSubObjectList())
@@ -2854,7 +2859,7 @@ void UAbilitySystemComponent::AddSpawnedAttribute(UAttributeSet* Attribute)
 
 	if (SpawnedAttributes.Find(Attribute) == INDEX_NONE)
 	{
-		if (IsUsingRegisteredSubObjectList() && HasBegunPlay())
+		if (IsUsingRegisteredSubObjectList() && IsReadyForReplication())
 		{
 			AddReplicatedSubObject(Attribute);
 		}
@@ -2868,7 +2873,7 @@ void UAbilitySystemComponent::RemoveSpawnedAttribute(UAttributeSet* Attribute)
 {
 	if (SpawnedAttributes.RemoveSingle(Attribute) > 0)
 	{
-		if (IsUsingRegisteredSubObjectList() )
+		if (IsUsingRegisteredSubObjectList())
 		{
 			RemoveReplicatedSubObject(Attribute);
 		}
diff --git a/Engine/Plugins/Runtime/GameplayAbilities/Source/GameplayAbilities/Public/AbilitySystemComponent.h b/Engine/Plugins/Runtime/GameplayAbilities/Source/GameplayAbilities/Public/AbilitySystemComponent.h
--- a/Engine/Plugins/Runtime/GameplayAbilities/Source/GameplayAbilities/Public/AbilitySystemComponent.h
+++ b/Engine/Plugins/Runtime/GameplayAbilities/Source/GameplayAbilities/Public/AbilitySystemComponent.h
@@ -1597,6 +1597,7 @@ public:
 	virtual void PostNetReceive() override;
 	virtual void OnRegister() override;
 	virtual void OnUnregister() override;
+	virtual void ReadyForReplication() override;
 	virtual void BeginPlay() override;
 
 protected:
diff --git a/Engine/Source/Runtime/Engine/Classes/Components/ActorComponent.h b/Engine/Source/Runtime/Engine/Classes/Components/ActorComponent.h
--- a/Engine/Source/Runtime/Engine/Classes/Components/ActorComponent.h
+++ b/Engine/Source/Runtime/Engine/Classes/Components/ActorComponent.h
@@ -272,6 +272,9 @@ private:
 	/** Indicates that InitializeComponent has been called, but UninitializeComponent has not yet */
 	uint8 bHasBeenInitialized:1;
 
+	/** Indicates that ReadyForReplication has been called  */
+	uint8 bIsReadyForReplication:1;
+
 	/** Indicates that BeginPlay has been called, but EndPlay has not yet */
 	uint8 bHasBegunPlay:1;
 
@@ -354,6 +357,9 @@ public:
 	/** Indicates that InitializeComponent has been called, but UninitializeComponent has not yet */
 	bool HasBeenInitialized() const { return bHasBeenInitialized; }
 
+	/** Indicates that ReadyForReplication has been called */
+	bool IsReadyForReplication() const { return bIsReadyForReplication; }
+
 	/** Indicates that BeginPlay has been called, but EndPlay has not yet */
 	bool HasBegunPlay() const { return bHasBegunPlay; }
 
@@ -684,6 +690,15 @@ public:
 	 */
 	virtual void InitializeComponent();
 
+	/**
+	 * ReadyForReplication gets called on replicated components when their owning actor is officially ready for replication.
+	 * Called after InitializeComponent but before BeginPlay. From there it will only be set false when the component is destroyed.
+	 * This is where you want to register your replicated subobjects if you already possess some. 
+	 * A component can get replicated before HasBegunPlay() is true if inside a tick or in BeginPlay() an RPC is called on him.
+	 * Requires component to be registered, initialized and set to replicate.
+	 */
+	virtual void ReadyForReplication();
+
 	/**
 	 * Begins Play for the component. 
 	 * Called when the owning Actor begins play or when the component is created if the Actor has already begun play.
diff --git a/Engine/Source/Runtime/Engine/Classes/GameFramework/Actor.h b/Engine/Source/Runtime/Engine/Classes/GameFramework/Actor.h
--- a/Engine/Source/Runtime/Engine/Classes/GameFramework/Actor.h
+++ b/Engine/Source/Runtime/Engine/Classes/GameFramework/Actor.h
@@ -2893,6 +2893,9 @@ public:
 	/** Will reregister all components on this actor. Does a lot of work - should only really be used in editor, generally use UpdateComponentTransforms or MarkComponentsRenderStateDirty. */
 	virtual void ReregisterAllComponents();
 
+	/** Finish initializing the component and register tick functions and beginplay if it's the proper time to do so. */
+	void HandleRegisterComponentWithWorld(UActorComponent* Component);
+
 	/**
 	 * Incrementally registers components associated with this actor, used during level streaming
 	 *
diff --git a/Engine/Source/Runtime/Engine/Private/Actor.cpp b/Engine/Source/Runtime/Engine/Private/Actor.cpp
--- a/Engine/Source/Runtime/Engine/Private/Actor.cpp
+++ b/Engine/Source/Runtime/Engine/Private/Actor.cpp
@@ -3272,13 +3272,19 @@ void AActor::UpdateReplicatedComponent(UActorComponent* Component)
 
 
 	const ELifetimeCondition NetCondition = Component->GetIsReplicated() ? AllowActorComponentToReplicate(Component) : COND_Never;
+	
+	if (Component->GetIsReplicated() && ActorHasBegunPlay != EActorBeginPlayState::HasNotBegunPlay && !Component->IsReadyForReplication())
+	{
+		Component->ReadyForReplication();
+	}
 
 	if (FReplicatedComponentInfo* RepComponentInfo = ReplicatedComponentsInfo.FindByKey(Component))
 	{
 		// Even if not replicated anymore just set the condition to Never because we want to keep the subobject list intact in case the component switches back to replicated later.
 		RepComponentInfo->NetCondition = NetCondition;
 	}
-	else if (NetCondition != COND_Never)
+	// Keep track of subobjects for replicated components even if Allow returns Never. The condition might get changed later on and we want to start tracking subobjects immediately.
+	else if (Component->GetIsReplicated())
 	{
 		ReplicatedComponentsInfo.Emplace(FReplicatedComponentInfo(Component, NetCondition));
 	}
@@ -3325,19 +3331,25 @@ void AActor::UpdateAllReplicatedComponents()
 			if (ActorHasBegunPlay != EActorBeginPlayState::HasNotBegunPlay)
 			{
 				const ELifetimeCondition NetCondition = AllowActorComponentToReplicate(Component);
-				if (NetCondition != COND_Never)
+				const int32 Index = ReplicatedComponentsInfo.AddUnique(FReplicatedComponentInfo(Component));
+				ReplicatedComponentsInfo[Index].NetCondition = NetCondition;
+
+				if (!Component->IsReadyForReplication())
 				{
-					const int32 Index = ReplicatedComponentsInfo.AddUnique(FReplicatedComponentInfo(Component));
-					ReplicatedComponentsInfo[Index].NetCondition = NetCondition;
+					Component->ReadyForReplication();
+				}
+
 #if UE_WITH_IRIS
+	            if (NetCondition != COND_Never)
+				{
 					// Begin replication and set NetCondition, if component already is replicated the NetCondition will be updated
 					Component->BeginReplication();
 				}
 				else
 				{
 					Component->EndReplication();
+                }
 #endif
-				}
 			}
 		}
 	}
@@ -4036,10 +4048,12 @@ void AActor::DispatchBeginPlay(bool bFromLevelStreaming)
 		{
 			const ELifetimeCondition NetCondition = AllowActorComponentToReplicate(ReplicatedComponent);
 
-			if (NetCondition != COND_Never)
+			const int32 Index = ReplicatedComponentsInfo.AddUnique(UE::Net::FReplicatedComponentInfo(ReplicatedComponent));
+			ReplicatedComponentsInfo[Index].NetCondition = NetCondition;
+
+			if (!ReplicatedComponent->IsReadyForReplication())
 			{
-				const int32 Index = ReplicatedComponentsInfo.AddUnique(UE::Net::FReplicatedComponentInfo(ReplicatedComponent));
-				ReplicatedComponentsInfo[Index].NetCondition = NetCondition;
+				ReplicatedComponent->ReadyForReplication();
 			}
 		}
 
@@ -5402,6 +5416,39 @@ void AActor::UninitializeComponents()
 	}
 }
 
+void AActor::HandleRegisterComponentWithWorld(UActorComponent* Component)
+{
+	const bool bOwnerBeginPlayStarted = HasActorBegunPlay() || IsActorBeginningPlay();
+
+	if (!Component->HasBeenInitialized() && Component->bWantsInitializeComponent && IsActorInitialized())
+	{
+		Component->InitializeComponent();
+
+		if (bOwnerBeginPlayStarted)
+		{
+			AddComponentForReplication(Component);
+		}
+	}
+
+	if (bOwnerBeginPlayStarted)
+	{
+		Component->RegisterAllComponentTickFunctions(true);
+
+		if (!Component->HasBegunPlay())
+		{
+#if UE_WITH_IRIS
+		    if (Component->GetIsReplicated())
+		    {
+		    	UpdateReplicatedComponent(Component);
+		    }
+#endif // UE_WITH_IRIS
+
+			Component->BeginPlay();
+			ensureMsgf(Component->HasBegunPlay(), TEXT("Failed to route BeginPlay (%s)"), *Component->GetFullName());
+		}
+	}
+}
+
 void AActor::DrawDebugComponents(FColor const& BaseColor) const
 {
 #if ENABLE_DRAW_DEBUG
diff --git a/Engine/Source/Runtime/Engine/Private/ActorReplication.cpp b/Engine/Source/Runtime/Engine/Private/ActorReplication.cpp
--- a/Engine/Source/Runtime/Engine/Private/ActorReplication.cpp
+++ b/Engine/Source/Runtime/Engine/Private/ActorReplication.cpp
@@ -548,26 +548,35 @@ void AActor::AddComponentForReplication(UActorComponent* Component)
 		return;
 	}
 
+	if (Component->bWantsInitializeComponent && !Component->HasBeenInitialized())
+	{
+		return;
+	}
+
 	const ELifetimeCondition NetCondition = AllowActorComponentToReplicate(Component);
 
-	if (NetCondition != COND_Never)
+	FReplicatedComponentInfo* ComponentInfo = ReplicatedComponentsInfo.FindByKey(Component);
+	if (!ComponentInfo)
 	{
-		FReplicatedComponentInfo* ComponentInfo = ReplicatedComponentsInfo.FindByKey(Component);
-		if (!ComponentInfo)
-		{
-			ReplicatedComponentsInfo.Emplace(FReplicatedComponentInfo(Component, NetCondition));
-		}
-		else
-		{
-			// Always set the condition because this component could be registering SubObjects ahead of time and already in the list with COND_Never
-			ComponentInfo->NetCondition = NetCondition;
-		}
+		ReplicatedComponentsInfo.Emplace(FReplicatedComponentInfo(Component, NetCondition));
+	}
+	else
+	{
+		// Always set the condition because this component could be registering SubObjects ahead of time and already in the list with COND_Never
+		ComponentInfo->NetCondition = NetCondition;
+	}
+
+	if (!Component->IsReadyForReplication())
+	{
+		Component->ReadyForReplication();
+	}
 
 #if UE_WITH_IRIS
+    if (NetCondition != COND_Never)
+    {
 		Component->BeginReplication();
+    }
 #endif
-
-	}
 }
 
 void AActor::RemoveReplicatedComponent(UActorComponent* Component)
@@ -624,7 +633,8 @@ void AActor::AddActorComponentReplicatedSubObject(UActorComponent* OwnerComponen
 {
 	check(IsValid(OwnerComponent));
 	check(IsValid(SubObject));
-	ensureMsgf(NetCondition != COND_Custom, TEXT("Custom netconditions do not work with SubObjects. %s - %s will not be replicated."), *GetName(), *SubObject->GetName());
+	ensureMsgf(OwnerComponent->GetIsReplicated(), TEXT("Only components with replication enabled can register subobjects. %s::%s has replication disabled."), *GetName(), *SubObject->GetName());
+	ensureMsgf(NetCondition != COND_Custom, TEXT("Custom netconditions do not work with SubObjects. %s::%s will not be replicated."), *GetName(), *SubObject->GetName());
 
 	FReplicatedComponentInfo* ComponentInfo = ReplicatedComponentsInfo.FindByKey(OwnerComponent);
 	if (ComponentInfo)
diff --git a/Engine/Source/Runtime/Engine/Private/Components/ActorComponent.cpp b/Engine/Source/Runtime/Engine/Private/Components/ActorComponent.cpp
--- a/Engine/Source/Runtime/Engine/Private/Components/ActorComponent.cpp
+++ b/Engine/Source/Runtime/Engine/Private/Components/ActorComponent.cpp
@@ -716,6 +716,8 @@ void UActorComponent::BeginDestroy()
 		UninitializeComponent();
 	}
 
+	bIsReadyForReplication = false;
+
 	ExecuteUnregisterEvents();
 
 	// Ensure that we call OnComponentDestroyed before we destroy this component
@@ -1058,6 +1060,11 @@ void UActorComponent::UninitializeComponent()
 	bHasBeenInitialized = false;
 }
 
+void UActorComponent::ReadyForReplication()
+{
+	bIsReadyForReplication = true;
+}
+
 void UActorComponent::BeginPlay()
 {
 	TRACE_OBJECT_LIFETIME_BEGIN(this);
@@ -1086,6 +1093,7 @@ void UActorComponent::EndPlay(const EEndPlayReason::Type EndPlayReason)
 		ReceiveEndPlay(EndPlayReason);
 	}
 
+	bIsReadyForReplication = false;
 	bHasBegunPlay = false;
 }
 
@@ -1356,27 +1364,7 @@ void UActorComponent::RegisterComponentWithWorld(UWorld* InWorld, FRegisterCompo
 	}
 	else
 	{
-		if (!bHasBeenInitialized && bWantsInitializeComponent && MyOwner->IsActorInitialized())
-		{
-			InitializeComponent();
-		}
-
-		if (MyOwner->HasActorBegunPlay() || MyOwner->IsActorBeginningPlay())
-		{
-			RegisterAllComponentTickFunctions(true);
-			if (!bHasBegunPlay)
-			{
-#if UE_WITH_IRIS
-				if (bReplicates)
-				{
-					MyOwner->UpdateReplicatedComponent(this);
-				}
-#endif // UE_WITH_IRIS
-
-				BeginPlay();
-				ensureMsgf(bHasBegunPlay, TEXT("Failed to route BeginPlay (%s)"), *GetFullName());
-			}
-		}
+		MyOwner->HandleRegisterComponentWithWorld(this);
 	}
 
 	// If this is a blueprint created component and it has component children they can miss getting registered in some scenarios
@@ -1454,6 +1442,8 @@ void UActorComponent::DestroyComponent(bool bPromoteChildren/*= false*/)
 		UninitializeComponent();
 	}
 
+	bIsReadyForReplication = false;
+
 	// Unregister if registered
 	if(IsRegistered())
 	{
diff --git a/Engine/Source/Runtime/Engine/Private/Net/NetSubObjectRegistryGetter.h b/Engine/Source/Runtime/Engine/Private/Net/NetSubObjectRegistryGetter.h
--- a/Engine/Source/Runtime/Engine/Private/Net/NetSubObjectRegistryGetter.h
+++ b/Engine/Source/Runtime/Engine/Private/Net/NetSubObjectRegistryGetter.h
@@ -14,7 +14,6 @@ namespace UE::Net
 class FSubObjectRegistryGetter final
 {
 public:
-	using FReplicatedComponentInfo = UE::Net::FReplicatedComponentInfo;
 
 	FSubObjectRegistryGetter() = delete;
 	~FSubObjectRegistryGetter() = delete;
@@ -37,7 +36,7 @@ public:
 		return InActor->ReplicatedComponentsInfo;
 	}
 
-	static const FReplicatedComponentInfo* GetReplicatedComponentInfoForComponent(AActor* InActor, UActorComponent* InActorComp)
+	static const UE::Net::FReplicatedComponentInfo* GetReplicatedComponentInfoForComponent(AActor* InActor, UActorComponent* InActorComp)
 	{
 		return InActor->ReplicatedComponentsInfo.FindByKey(InActorComp);
 	}
@@ -47,6 +46,25 @@ public:
 		UE::Net::FReplicatedComponentInfo* ComponentInfo = InActor->ReplicatedComponentsInfo.FindByKey(InActorComp);
 		return ComponentInfo ? ComponentInfo->SubObjects.IsSubObjectInRegistry(InSubObject) : false;
 	}
+
+	/** Look for the subobject in the subobject list of the actor or any of its replicated components. */
+	static bool IsSubObjectInAnyRegistry(AActor* InActor, UObject* InSubObject)
+	{
+		if (InActor->ReplicatedSubObjects.IsSubObjectInRegistry(InSubObject))
+		{
+			return true;
+		}
+
+		for (const UE::Net::FReplicatedComponentInfo& ComponentInfo : InActor->ReplicatedComponentsInfo)
+		{
+			if (ComponentInfo.Component == InSubObject || ComponentInfo.SubObjects.IsSubObjectInRegistry(InSubObject))
+			{
+				return true;
+			}
+		}
+
+		return false;
+	}
 };
 
 }
diff --git a/Engine/Source/Runtime/GameplayTasks/Classes/GameplayTasksComponent.h b/Engine/Source/Runtime/GameplayTasks/Classes/GameplayTasksComponent.h
--- a/Engine/Source/Runtime/GameplayTasks/Classes/GameplayTasksComponent.h
+++ b/Engine/Source/Runtime/GameplayTasks/Classes/GameplayTasksComponent.h
@@ -160,7 +160,7 @@ public:
 	// END IGameplayTaskOwnerInterface
 
 	// ActorComponent overrides
-	virtual void BeginPlay() override;
+	virtual void ReadyForReplication() override;
 
 	UFUNCTION(BlueprintCallable, DisplayName="Run Gameplay Task", meta=(ScriptName="RunGameplayTask"), Category = "Gameplay Tasks", meta = (AutoCreateRefTerm = "AdditionalRequiredResources, AdditionalClaimedResources", AdvancedDisplay = "AdditionalRequiredResources, AdditionalClaimedResources"))
 	static EGameplayTaskRunResult K2_RunGameplayTask(TScriptInterface<IGameplayTaskOwnerInterface> TaskOwner, UGameplayTask* Task, uint8 Priority, TArray<TSubclassOf<UGameplayTaskResource> > AdditionalRequiredResources, TArray<TSubclassOf<UGameplayTaskResource> > AdditionalClaimedResources);
diff --git a/Engine/Source/Runtime/GameplayTasks/Private/GameplayTasksComponent.cpp b/Engine/Source/Runtime/GameplayTasks/Private/GameplayTasksComponent.cpp
--- a/Engine/Source/Runtime/GameplayTasks/Private/GameplayTasksComponent.cpp
+++ b/Engine/Source/Runtime/GameplayTasks/Private/GameplayTasksComponent.cpp
@@ -57,9 +57,9 @@ UGameplayTasksComponent::UGameplayTasksComponent(const FObjectInitializer& Objec
 	TopActivePriority = 0;
 }
 
-void UGameplayTasksComponent::BeginPlay()
+void UGameplayTasksComponent::ReadyForReplication()
 {
-	Super::BeginPlay();
+	Super::ReadyForReplication();
 
 	REDIRECT_TO_VLOG(GetOwner());
 	
@@ -789,7 +789,7 @@ bool UGameplayTasksComponent::AddSimulatedTask(UGameplayTask* NewTask)
 		SimulatedTasks.Add(NewTask);
 		SetSimulatedTasksNetDirty();
 
-		if (IsUsingRegisteredSubObjectList())
+		if (IsUsingRegisteredSubObjectList() && IsReadyForReplication())
 		{
 			AddReplicatedSubObject(NewTask, COND_SkipOwner);
 		}
@@ -815,7 +815,7 @@ void UGameplayTasksComponent::RemoveSimulatedTask(UGameplayTask* NewTask)
 
 void UGameplayTasksComponent::SetSimulatedTasks(const TArray<UGameplayTask*>& NewSimulatedTasks)
 {
-	if (IsUsingRegisteredSubObjectList())
+	if (IsUsingRegisteredSubObjectList() && IsReadyForReplication())
 	{
 		// Unregister all current tasks
 		for (UGameplayTask* OldGameplayTask : SimulatedTasks)
