Fix crash when switching to wireframe in DX12.

The wireframe view mode creates some temporary material proxies which were deleted along with the renderer. The renderer is destroyed  by a background thread now, so we need to move this cleanup to the RenderThreadEnd method, because proxies have to be deleted on the render thread.

#rnx

diff --git a/Engine/Source/Runtime/Engine/Public/SceneManagement.h b/Engine/Source/Runtime/Engine/Public/SceneManagement.h
--- a/Engine/Source/Runtime/Engine/Public/SceneManagement.h
+++ b/Engine/Source/Runtime/Engine/Public/SceneManagement.h
@@ -2401,12 +2401,19 @@ protected:
 	ENGINE_API FMeshElementCollector(ERHIFeatureLevel::Type InFeatureLevel, FSceneRenderingBulkObjectAllocator& InBulkAllocator);
 
 	~FMeshElementCollector()
+	{
+		DeleteTemporaryProxies();
+	}
+
+	void DeleteTemporaryProxies()
 	{
 		check(!ParallelTasks.Num()); // We should have blocked on this already
 		for (int32 ProxyIndex = 0; ProxyIndex < TemporaryProxies.Num(); ProxyIndex++)
 		{
 			delete TemporaryProxies[ProxyIndex];
 		}
+
+		TemporaryProxies.Empty();
 	}
 
 	void SetPrimitive(const FPrimitiveSceneProxy* InPrimitiveSceneProxy, FHitProxyId DefaultHitProxyId)
diff --git a/Engine/Source/Runtime/Renderer/Private/SceneRendering.cpp b/Engine/Source/Runtime/Renderer/Private/SceneRendering.cpp
--- a/Engine/Source/Runtime/Renderer/Private/SceneRendering.cpp
+++ b/Engine/Source/Runtime/Renderer/Private/SceneRendering.cpp
@@ -4132,6 +4132,8 @@ void FSceneRenderer::RenderThreadEnd(FRHICommandListImmediate& RHICmdList, const
 		for (FSceneRenderer* SceneRenderer : SceneRenderers)
 		{
 			SceneRenderer->GPUSceneDynamicContext.Release();
+			SceneRenderer->MeshCollector.DeleteTemporaryProxies();
+			SceneRenderer->RayTracingCollector.DeleteTemporaryProxies();
 		}
 
 		// Mem stack mark is stored on first scene renderer
