Allow the meta-data for generated types to be cooked as optional data, for use with cooked editors

This covers the following types and meta-data:
    * Blueprint Generated Class (class, property, and function meta-data)
    * User Defined Struct (struct and property meta-data)
    * User Defined Enum (enum meta-data)

diff --git a/Engine/Source/Runtime/CoreUObject/Private/UObject/Field.cpp b/Engine/Source/Runtime/CoreUObject/Private/UObject/Field.cpp
--- a/Engine/Source/Runtime/CoreUObject/Private/UObject/Field.cpp
+++ b/Engine/Source/Runtime/CoreUObject/Private/UObject/Field.cpp
@@ -941,6 +941,21 @@ const TMap<FName, FString>* FField::GetMetaDataMap() const
 	return MetaDataMap;
 }
 
+void FField::AppendMetaData(const TMap<FName, FString>& MetaDataMapToAppend)
+{
+	if (MetaDataMapToAppend.Num() > 0)
+	{
+		if (MetaDataMap)
+		{
+			MetaDataMap->Append(MetaDataMapToAppend);
+		}
+		else
+		{
+			MetaDataMap = new TMap<FName, FString>(MetaDataMapToAppend);
+		}
+	}
+}
+
 void FField::CopyMetaData(const FField* InSourceField, FField* InDestField)
 {
 	check(InSourceField);
diff --git a/Engine/Source/Runtime/CoreUObject/Public/UObject/Field.h b/Engine/Source/Runtime/CoreUObject/Public/UObject/Field.h
--- a/Engine/Source/Runtime/CoreUObject/Public/UObject/Field.h
+++ b/Engine/Source/Runtime/CoreUObject/Public/UObject/Field.h
@@ -860,6 +860,9 @@ public:
 	/** Gets all metadata associated with this field */
 	const TMap<FName, FString>* GetMetaDataMap() const;
 
+	/** Append the given metadata to this field */
+	void AppendMetaData(const TMap<FName, FString>& MetaDataMapToAppend);
+
 	/** Copies all metadata from Source Field to Dest Field */
 	static void CopyMetaData(const FField* InSourceField, FField* InDestField);
 
diff --git a/Engine/Source/Runtime/Engine/Classes/Engine/BlueprintGeneratedClass.h b/Engine/Source/Runtime/Engine/Classes/Engine/BlueprintGeneratedClass.h
--- a/Engine/Source/Runtime/Engine/Classes/Engine/BlueprintGeneratedClass.h
+++ b/Engine/Source/Runtime/Engine/Classes/Engine/BlueprintGeneratedClass.h
@@ -19,6 +19,7 @@ class UActorComponent;
 class UDynamicBlueprintBinding;
 class UInheritableComponentHandler;
 class UTimelineTemplate;
+class UClassCookedMetaData;
 
 DECLARE_MEMORY_STAT_EXTERN(TEXT("Persistent Uber Graph Frame memory"), STAT_PersistentUberGraphFrameMemory, STATGROUP_Memory, );
 DECLARE_MEMORY_STAT_EXTERN(TEXT("BPComp Instancing Fast Path memory"), STAT_BPCompInstancingFastPathMemory, STATGROUP_Memory, );
@@ -780,8 +781,8 @@ public:
 	virtual bool CanBeClusterRoot() const override;
 #if WITH_EDITOR
 	virtual UClass* RegenerateClass(UClass* ClassToRegenerate, UObject* PreviousCDO) override;
-	virtual void BeginCacheForCookedPlatformData(const ITargetPlatform* TargetPlatform) override;
-	virtual void ClearAllCachedCookedPlatformData() override;
+	virtual void PreSaveRoot(FObjectPreSaveRootContext ObjectSaveContext) override;
+	virtual void PostSaveRoot(FObjectPostSaveRootContext ObjectSaveContext) override;
 #endif	//WITH_EDITOR
 	virtual bool IsAsset() const override;
 	// End UObject interface
@@ -932,5 +933,16 @@ private:
 
 	/** Editor-only asset registry tags on cooked BPGC */
 	FEditorTags CookedEditorTags;
+
+protected:
+	virtual TSubclassOf<UClassCookedMetaData> GetCookedMetaDataClass() const;
+
+private:
+	UClassCookedMetaData* NewCookedMetaData();
+	const UClassCookedMetaData* FindCookedMetaData();
+	void PurgeCookedMetaData();
+
+	UPROPERTY()
+	TObjectPtr<UClassCookedMetaData> CachedCookedMetaDataPtr;
 #endif // WITH_EDITORONLY_DATA
 };
diff --git a/Engine/Source/Runtime/Engine/Classes/Engine/UserDefinedEnum.h b/Engine/Source/Runtime/Engine/Classes/Engine/UserDefinedEnum.h
--- a/Engine/Source/Runtime/Engine/Classes/Engine/UserDefinedEnum.h
+++ b/Engine/Source/Runtime/Engine/Classes/Engine/UserDefinedEnum.h
@@ -12,8 +12,11 @@
 #include "CoreMinimal.h"
 #include "UObject/ObjectMacros.h"
 #include "UObject/Class.h"
+#include "Templates/SubclassOf.h"
 #include "UserDefinedEnum.generated.h"
 
+class UEnumCookedMetaData;
+
 /** 
  *	An Enumeration is a list of named values.
  */
@@ -74,8 +77,23 @@ public:
 	virtual void PostEditUndo() override;
 	virtual void PostEditChangeProperty(FPropertyChangedEvent& PropertyChangedEvent);
 	virtual void GetAssetRegistryTags(TArray<FAssetRegistryTag>& OutTags) const override;
+	virtual void PreSaveRoot(FObjectPreSaveRootContext ObjectSaveContext) override;
+	virtual void PostSaveRoot(FObjectPostSaveRootContext ObjectSaveContext) override;
 	//~ End UObject Interface
 
 	FString GenerateNewEnumeratorName();
 #endif	// WITH_EDITOR
+
+#if WITH_EDITORONLY_DATA
+protected:
+	virtual TSubclassOf<UEnumCookedMetaData> GetCookedMetaDataClass() const;
+
+private:
+	UEnumCookedMetaData* NewCookedMetaData();
+	const UEnumCookedMetaData* FindCookedMetaData();
+	void PurgeCookedMetaData();
+
+	UPROPERTY()
+	TObjectPtr<UEnumCookedMetaData> CachedCookedMetaDataPtr;
+#endif // WITH_EDITORONLY_DATA
 };
diff --git a/Engine/Source/Runtime/Engine/Classes/Engine/UserDefinedStruct.h b/Engine/Source/Runtime/Engine/Classes/Engine/UserDefinedStruct.h
--- a/Engine/Source/Runtime/Engine/Classes/Engine/UserDefinedStruct.h
+++ b/Engine/Source/Runtime/Engine/Classes/Engine/UserDefinedStruct.h
@@ -8,10 +8,12 @@
 #include "UObject/Class.h"
 #include "UObject/WeakObjectPtr.h"
 #include "UObject/StructOnScope.h"
+#include "Templates/SubclassOf.h"
 #include "UserDefinedStruct.generated.h"
 
 class UUserDefinedStruct;
 class UUserDefinedStructEditorData;
+class UStructCookedMetaData;
 
 UENUM()
 enum EUserDefinedStructureStatus
@@ -86,6 +88,8 @@ public:
 	virtual void PostDuplicate(bool bDuplicateForPIE) override;
 	virtual void GetAssetRegistryTags(TArray<FAssetRegistryTag>& OutTags) const override;
 	virtual void PostLoad() override;
+	virtual void PreSaveRoot(FObjectPreSaveRootContext ObjectSaveContext) override;
+	virtual void PostSaveRoot(FObjectPostSaveRootContext ObjectSaveContext) override;
 	// End of UObject interface.
 
 	/** Creates a new guid if needed */
@@ -124,4 +128,17 @@ public:
 
 	/** Inspects properties and default values, setting appropriate StructFlags */
 	void UpdateStructFlags();
+
+#if WITH_EDITORONLY_DATA
+protected:
+	virtual TSubclassOf<UStructCookedMetaData> GetCookedMetaDataClass() const;
+
+private:
+	UStructCookedMetaData* NewCookedMetaData();
+	const UStructCookedMetaData* FindCookedMetaData();
+	void PurgeCookedMetaData();
+
+	UPROPERTY()
+	TObjectPtr<UStructCookedMetaData> CachedCookedMetaDataPtr;
+#endif // WITH_EDITORONLY_DATA
 };
diff --git a/Engine/Source/Runtime/Engine/Private/BlueprintGeneratedClass.cpp b/Engine/Source/Runtime/Engine/Private/BlueprintGeneratedClass.cpp
--- a/Engine/Source/Runtime/Engine/Private/BlueprintGeneratedClass.cpp
+++ b/Engine/Source/Runtime/Engine/Private/BlueprintGeneratedClass.cpp
@@ -8,8 +8,10 @@
 #include "UObject/CoreRedirects.h"
 #include "UObject/Package.h"
 #include "UObject/LinkerLoad.h"
+#include "UObject/ObjectSaveContext.h"
 #include "Serialization/ObjectReader.h"
 #include "Serialization/ObjectWriter.h"
+#include "CookedMetaData.h"
 #include "Engine/Blueprint.h"
 #include "Components/ActorComponent.h"
 #include "Curves/CurveFloat.h"
@@ -140,17 +142,24 @@ void UBlueprintGeneratedClass::PostLoad()
 	Super::PostLoad();
 
 #if WITH_EDITORONLY_DATA
-	UPackage* Package = GetOutermost();
+	if (bCooked)
+	{
+		if (const UClassCookedMetaData* CookedMetaData = FindCookedMetaData())
+		{
+			CookedMetaData->ApplyMetaData(this);
+			PurgeCookedMetaData();
+		}
+	}
 
 #if WITH_EDITOR
 	// Make BPGC from a cooked package standalone so it doesn't get GCed
-	if (GEditor && Package && Package->HasAnyPackageFlags(PKG_FilterEditorOnly))
+	if (GEditor && bCooked)
 	{
 		SetFlags(RF_Standalone);
 	}
 #endif //if WITH_EDITOR
 
-	if (Package == nullptr || !Package->bIsCookedForEditor)
+	if (!bCooked)
 	{
 		if (GetAuthoritativeClass() != this)
 		{
@@ -1704,9 +1713,9 @@ UClass* UBlueprintGeneratedClass::RegenerateClass(UClass* ClassToRegenerate, UOb
 	return this;
 }
 
-void UBlueprintGeneratedClass::BeginCacheForCookedPlatformData(const ITargetPlatform* TargetPlatform)
+void UBlueprintGeneratedClass::PreSaveRoot(FObjectPreSaveRootContext ObjectSaveContext)
 {
-	Super::BeginCacheForCookedPlatformData(TargetPlatform);
+	Super::PreSaveRoot(ObjectSaveContext);
 
 	auto ShouldCookBlueprintPropertyGuids = [this]() -> bool
 	{
@@ -1730,7 +1739,7 @@ void UBlueprintGeneratedClass::BeginCacheForCookedPlatformData(const ITargetPlat
 		return false;
 	};
 
-	if (ShouldCookBlueprintPropertyGuids())
+	if (ObjectSaveContext.IsCooking() && ShouldCookBlueprintPropertyGuids())
 	{
 		CookedPropertyGuids = PropertyGuids;
 
@@ -1752,13 +1761,30 @@ void UBlueprintGeneratedClass::BeginCacheForCookedPlatformData(const ITargetPlat
 	{
 		CookedPropertyGuids.Reset();
 	}
+
+	if (ObjectSaveContext.IsCooking() && (ObjectSaveContext.GetSaveFlags() & SAVE_Optional))
+	{
+		UClassCookedMetaData* CookedMetaData = NewCookedMetaData();
+		CookedMetaData->CacheMetaData(this);
+
+		if (!CookedMetaData->HasMetaData())
+		{
+			PurgeCookedMetaData();
+		}
+	}
+	else
+	{
+		PurgeCookedMetaData();
+	}
 }
 
-void UBlueprintGeneratedClass::ClearAllCachedCookedPlatformData()
+void UBlueprintGeneratedClass::PostSaveRoot(FObjectPostSaveRootContext ObjectSaveContext)
 {
-	Super::ClearAllCachedCookedPlatformData();
+	Super::PostSaveRoot(ObjectSaveContext);
 
 	CookedPropertyGuids.Reset();
+
+	PurgeCookedMetaData();
 }
 #endif
 
@@ -2277,4 +2303,35 @@ void UBlueprintGeneratedClass::GetEditorTags(FEditorTags& Tags) const
 		}
 	}
 }
+
+TSubclassOf<UClassCookedMetaData> UBlueprintGeneratedClass::GetCookedMetaDataClass() const
+{
+	return UClassCookedMetaData::StaticClass();
+}
+
+UClassCookedMetaData* UBlueprintGeneratedClass::NewCookedMetaData()
+{
+	if (!CachedCookedMetaDataPtr)
+	{
+		CachedCookedMetaDataPtr = CookedMetaDataUtil::NewCookedMetaData<UClassCookedMetaData>(this, "CookedClassMetaData", GetCookedMetaDataClass());
+	}
+	return CachedCookedMetaDataPtr;
+}
+
+const UClassCookedMetaData* UBlueprintGeneratedClass::FindCookedMetaData()
+{
+	if (!CachedCookedMetaDataPtr)
+	{
+		CachedCookedMetaDataPtr = CookedMetaDataUtil::FindCookedMetaData<UClassCookedMetaData>(this, TEXT("CookedClassMetaData"));
+	}
+	return CachedCookedMetaDataPtr;
+}
+
+void UBlueprintGeneratedClass::PurgeCookedMetaData()
+{
+	if (CachedCookedMetaDataPtr)
+	{
+		CookedMetaDataUtil::PurgeCookedMetaData<UClassCookedMetaData>(CachedCookedMetaDataPtr);
+	}
+}
 #endif //if WITH_EDITORONLY_DATA
diff --git a/Engine/Source/Runtime/Engine/Private/CookedMetaData.cpp b/Engine/Source/Runtime/Engine/Private/CookedMetaData.cpp
new file mode 100644
--- /dev/null
+++ b/Engine/Source/Runtime/Engine/Private/CookedMetaData.cpp
@@ -0,0 +1,170 @@
+// Copyright Epic Games, Inc. All Rights Reserved.
+
+#include "CookedMetaData.h"
+#include "UObject/Package.h"
+#include "UObject/MetaData.h"
+
+bool FObjectCookedMetaDataStore::HasMetaData() const
+{
+	return ObjectMetaData.Num() > 0;
+}
+
+void FObjectCookedMetaDataStore::CacheMetaData(const UObject* SourceObject)
+{
+	ObjectMetaData.Reset();
+
+#if WITH_EDITORONLY_DATA
+	if (UPackage* SourcePackage = SourceObject->GetPackage())
+	{
+		if (const UMetaData* SourceMetaData = SourcePackage->GetMetaData())
+		{
+			if (const TMap<FName, FString>* SourceObjectMetaData = SourceMetaData->ObjectMetaDataMap.Find(SourceObject))
+			{
+				ObjectMetaData = *SourceObjectMetaData;
+			}
+		}
+	}
+#endif
+}
+
+void FObjectCookedMetaDataStore::ApplyMetaData(UObject* TargetObject) const
+{
+#if WITH_EDITORONLY_DATA
+	if (UPackage* TargetPackage = TargetObject->GetPackage())
+	{
+		if (UMetaData* TargetMetaData = TargetPackage->GetMetaData())
+		{
+			TargetMetaData->ObjectMetaDataMap.FindOrAdd(TargetObject).Append(ObjectMetaData);
+		}
+	}
+#endif
+}
+
+
+bool FFieldCookedMetaDataStore::HasMetaData() const
+{
+	return FieldMetaData.Num() > 0;
+}
+
+void FFieldCookedMetaDataStore::CacheMetaData(const FField* SourceField)
+{
+	FieldMetaData.Reset();
+
+#if WITH_EDITORONLY_DATA
+	if (const TMap<FName, FString>* SourceFieldMetaData = SourceField->GetMetaDataMap())
+	{
+		FieldMetaData = *SourceFieldMetaData;
+	}
+#endif
+}
+
+void FFieldCookedMetaDataStore::ApplyMetaData(FField* TargetField) const
+{
+#if WITH_EDITORONLY_DATA
+	TargetField->AppendMetaData(FieldMetaData);
+#endif
+}
+
+
+bool FStructCookedMetaDataStore::HasMetaData() const
+{
+	return ObjectMetaData.HasMetaData()
+		|| PropertiesMetaData.Num() > 0;
+}
+
+void FStructCookedMetaDataStore::CacheMetaData(const UStruct* SourceStruct)
+{
+	ObjectMetaData.CacheMetaData(SourceStruct);
+
+	for (TFieldIterator<const FProperty> SourcePropertyIt(SourceStruct, EFieldIterationFlags::None); SourcePropertyIt; ++SourcePropertyIt)
+	{
+		FFieldCookedMetaDataStore SourcePropertyMetaData;
+		SourcePropertyMetaData.CacheMetaData(*SourcePropertyIt);
+
+		if (SourcePropertyMetaData.HasMetaData())
+		{
+			PropertiesMetaData.Add(SourcePropertyIt->GetFName(), MoveTemp(SourcePropertyMetaData));
+		}
+	}
+}
+
+void FStructCookedMetaDataStore::ApplyMetaData(UStruct* TargetStruct) const
+{
+	ObjectMetaData.ApplyMetaData(TargetStruct);
+
+	for (TFieldIterator<FProperty> TargetPropertyIt(TargetStruct, EFieldIterationFlags::None); TargetPropertyIt; ++TargetPropertyIt)
+	{
+		if (const FFieldCookedMetaDataStore* TargetPropertyMetaData = PropertiesMetaData.Find(TargetPropertyIt->GetFName()))
+		{
+			TargetPropertyMetaData->ApplyMetaData(*TargetPropertyIt);
+		}
+	}
+}
+
+
+bool UEnumCookedMetaData::HasMetaData() const
+{
+	return EnumMetaData.HasMetaData();
+}
+
+void UEnumCookedMetaData::CacheMetaData(const UEnum* SourceEnum)
+{
+	EnumMetaData.CacheMetaData(SourceEnum);
+}
+
+void UEnumCookedMetaData::ApplyMetaData(UEnum* TargetEnum) const
+{
+	EnumMetaData.ApplyMetaData(TargetEnum);
+}
+
+
+bool UStructCookedMetaData::HasMetaData() const
+{
+	return StructMetaData.HasMetaData();
+}
+
+void UStructCookedMetaData::CacheMetaData(const UScriptStruct* SourceStruct)
+{
+	StructMetaData.CacheMetaData(SourceStruct);
+}
+
+void UStructCookedMetaData::ApplyMetaData(UScriptStruct* TargetStruct) const
+{
+	StructMetaData.ApplyMetaData(TargetStruct);
+}
+
+
+bool UClassCookedMetaData::HasMetaData() const
+{
+	return ClassMetaData.HasMetaData()
+		|| FunctionsMetaData.Num() > 0;
+}
+
+void UClassCookedMetaData::CacheMetaData(const UClass* SourceClass)
+{
+	ClassMetaData.CacheMetaData(SourceClass);
+
+	for (TFieldIterator<const UFunction> SourceFunctionIt(SourceClass, EFieldIterationFlags::None); SourceFunctionIt; ++SourceFunctionIt)
+	{
+		FStructCookedMetaDataStore SourceFunctionMetaData;
+		SourceFunctionMetaData.CacheMetaData(*SourceFunctionIt);
+
+		if (SourceFunctionMetaData.HasMetaData())
+		{
+			FunctionsMetaData.Add(SourceFunctionIt->GetFName(), MoveTemp(SourceFunctionMetaData));
+		}
+	}
+}
+
+void UClassCookedMetaData::ApplyMetaData(UClass* TargetClass) const
+{
+	ClassMetaData.ApplyMetaData(TargetClass);
+
+	for (TFieldIterator<UFunction> TargetFunctionIt(TargetClass, EFieldIterationFlags::None); TargetFunctionIt; ++TargetFunctionIt)
+	{
+		if (const FStructCookedMetaDataStore* TargetFunctionMetaData = FunctionsMetaData.Find(TargetFunctionIt->GetFName()))
+		{
+			TargetFunctionMetaData->ApplyMetaData(*TargetFunctionIt);
+		}
+	}
+}
diff --git a/Engine/Source/Runtime/Engine/Private/UserDefinedEnum.cpp b/Engine/Source/Runtime/Engine/Private/UserDefinedEnum.cpp
--- a/Engine/Source/Runtime/Engine/Private/UserDefinedEnum.cpp
+++ b/Engine/Source/Runtime/Engine/Private/UserDefinedEnum.cpp
@@ -2,6 +2,8 @@
 
 #include "Engine/UserDefinedEnum.h"
 #include "UObject/EditorObjectVersion.h"
+#include "UObject/ObjectSaveContext.h"
+#include "CookedMetaData.h"
 #if WITH_EDITOR
 #include "Kismet2/EnumEditorUtils.h"
 #include "UObject/MetaData.h"
@@ -76,6 +78,16 @@ void UUserDefinedEnum::PostDuplicate(bool bDuplicateForPIE)
 void UUserDefinedEnum::PostLoad()
 {
 	Super::PostLoad();
+
+	if (GetPackage()->HasAnyPackageFlags(PKG_Cooked))
+	{
+		if (const UEnumCookedMetaData* CookedMetaData = FindCookedMetaData())
+		{
+			CookedMetaData->ApplyMetaData(this);
+			PurgeCookedMetaData();
+		}
+	}
+
 	FEnumEditorUtils::UpdateAfterPathChanged(this);
 	if (NumEnums() > 1 && DisplayNameMap.Num() == 0) // >1 because User Defined Enums always have a "MAX" entry
 	{
@@ -123,6 +135,33 @@ void UUserDefinedEnum::GetAssetRegistryTags(TArray<FAssetRegistryTag>& OutTags)
 	OutTags.Emplace(GET_MEMBER_NAME_CHECKED(UUserDefinedEnum, EnumDescription), DescriptionString, FAssetRegistryTag::TT_Hidden);
 }
 
+void UUserDefinedEnum::PreSaveRoot(FObjectPreSaveRootContext ObjectSaveContext)
+{
+	Super::PreSaveRoot(ObjectSaveContext);
+
+	if (ObjectSaveContext.IsCooking() && (ObjectSaveContext.GetSaveFlags() & SAVE_Optional))
+	{
+		UEnumCookedMetaData* CookedMetaData = NewCookedMetaData();
+		CookedMetaData->CacheMetaData(this);
+
+		if (!CookedMetaData->HasMetaData())
+		{
+			PurgeCookedMetaData();
+		}
+	}
+	else
+	{
+		PurgeCookedMetaData();
+	}
+}
+
+void UUserDefinedEnum::PostSaveRoot(FObjectPostSaveRootContext ObjectSaveContext)
+{
+	Super::PostSaveRoot(ObjectSaveContext);
+
+	PurgeCookedMetaData();
+}
+
 FString UUserDefinedEnum::GenerateNewEnumeratorName()
 {
 	FString EnumNameString;
@@ -208,3 +247,36 @@ bool UUserDefinedEnum::SetEnums(TArray<TPair<FName, int64>>& InNames, ECppForm I
 
 	return false;
 }
+
+#if WITH_EDITORONLY_DATA
+TSubclassOf<UEnumCookedMetaData> UUserDefinedEnum::GetCookedMetaDataClass() const
+{
+	return UEnumCookedMetaData::StaticClass();
+}
+
+UEnumCookedMetaData* UUserDefinedEnum::NewCookedMetaData()
+{
+	if (!CachedCookedMetaDataPtr)
+	{
+		CachedCookedMetaDataPtr = CookedMetaDataUtil::NewCookedMetaData<UEnumCookedMetaData>(this, "CookedEnumMetaData", GetCookedMetaDataClass());
+	}
+	return CachedCookedMetaDataPtr;
+}
+
+const UEnumCookedMetaData* UUserDefinedEnum::FindCookedMetaData()
+{
+	if (!CachedCookedMetaDataPtr)
+	{
+		CachedCookedMetaDataPtr = CookedMetaDataUtil::FindCookedMetaData<UEnumCookedMetaData>(this, TEXT("CookedEnumMetaData"));
+	}
+	return CachedCookedMetaDataPtr;
+}
+
+void UUserDefinedEnum::PurgeCookedMetaData()
+{
+	if (CachedCookedMetaDataPtr)
+	{
+		CookedMetaDataUtil::PurgeCookedMetaData<UEnumCookedMetaData>(CachedCookedMetaDataPtr);
+	}
+}
+#endif // WITH_EDITORONLY_DATA
diff --git a/Engine/Source/Runtime/Engine/Private/UserDefinedStruct.cpp b/Engine/Source/Runtime/Engine/Private/UserDefinedStruct.cpp
--- a/Engine/Source/Runtime/Engine/Private/UserDefinedStruct.cpp
+++ b/Engine/Source/Runtime/Engine/Private/UserDefinedStruct.cpp
@@ -5,6 +5,8 @@
 #include "UObject/StructOnScope.h"
 #include "UObject/UnrealType.h"
 #include "UObject/LinkerLoad.h"
+#include "UObject/ObjectSaveContext.h"
+#include "CookedMetaData.h"
 #include "UObject/FrameworkObjectVersion.h"
 #include "Misc/SecureHash.h"
 #include "UObject/PropertyPortFlags.h"
@@ -161,9 +163,45 @@ void UUserDefinedStruct::PostLoad()
 {
 	Super::PostLoad();
 
+	if (GetPackage()->HasAnyPackageFlags(PKG_Cooked))
+	{
+		if (const UStructCookedMetaData* CookedMetaData = FindCookedMetaData())
+		{
+			CookedMetaData->ApplyMetaData(this);
+			PurgeCookedMetaData();
+		}
+	}
+
 	ValidateGuid();
 }
 
+void UUserDefinedStruct::PreSaveRoot(FObjectPreSaveRootContext ObjectSaveContext)
+{
+	Super::PreSaveRoot(ObjectSaveContext);
+
+	if (ObjectSaveContext.IsCooking() && (ObjectSaveContext.GetSaveFlags() & SAVE_Optional))
+	{
+		UStructCookedMetaData* CookedMetaData = NewCookedMetaData();
+		CookedMetaData->CacheMetaData(this);
+
+		if (!CookedMetaData->HasMetaData())
+		{
+			PurgeCookedMetaData();
+		}
+	}
+	else
+	{
+		PurgeCookedMetaData();
+	}
+}
+
+void UUserDefinedStruct::PostSaveRoot(FObjectPostSaveRootContext ObjectSaveContext)
+{
+	Super::PostSaveRoot(ObjectSaveContext);
+
+	PurgeCookedMetaData();
+}
+
 void UUserDefinedStruct::GetAssetRegistryTags(TArray<FAssetRegistryTag>& OutTags) const
 {
 	Super::GetAssetRegistryTags(OutTags);
@@ -524,3 +562,36 @@ void UUserDefinedStruct::UpdateStructFlags()
 	}
 
 }
+
+#if WITH_EDITORONLY_DATA
+TSubclassOf<UStructCookedMetaData> UUserDefinedStruct::GetCookedMetaDataClass() const
+{
+	return UStructCookedMetaData::StaticClass();
+}
+
+UStructCookedMetaData* UUserDefinedStruct::NewCookedMetaData()
+{
+	if (!CachedCookedMetaDataPtr)
+	{
+		CachedCookedMetaDataPtr = CookedMetaDataUtil::NewCookedMetaData<UStructCookedMetaData>(this, "CookedStructMetaData", GetCookedMetaDataClass());
+	}
+	return CachedCookedMetaDataPtr;
+}
+
+const UStructCookedMetaData* UUserDefinedStruct::FindCookedMetaData()
+{
+	if (!CachedCookedMetaDataPtr)
+	{
+		CachedCookedMetaDataPtr = CookedMetaDataUtil::FindCookedMetaData<UStructCookedMetaData>(this, TEXT("CookedStructMetaData"));
+	}
+	return CachedCookedMetaDataPtr;
+}
+
+void UUserDefinedStruct::PurgeCookedMetaData()
+{
+	if (CachedCookedMetaDataPtr)
+	{
+		CookedMetaDataUtil::PurgeCookedMetaData<UStructCookedMetaData>(CachedCookedMetaDataPtr);
+	}
+}
+#endif // WITH_EDITORONLY_DATA
diff --git a/Engine/Source/Runtime/Engine/Public/CookedMetaData.h b/Engine/Source/Runtime/Engine/Public/CookedMetaData.h
new file mode 100644
--- /dev/null
+++ b/Engine/Source/Runtime/Engine/Public/CookedMetaData.h
@@ -0,0 +1,149 @@
+// Copyright Epic Games, Inc. All Rights Reserved.
+
+#pragma once
+
+#include "UObject/Object.h"
+#include "UObject/UObjectGlobals.h"
+#include "CookedMetaData.generated.h"
+
+/**
+ * Cooked meta-data for a UObject.
+ */
+USTRUCT()
+struct ENGINE_API FObjectCookedMetaDataStore
+{
+public:
+	GENERATED_BODY()
+
+	bool HasMetaData() const;
+	void CacheMetaData(const UObject* SourceObject);
+	void ApplyMetaData(UObject* TargetObject) const;
+
+private:
+	UPROPERTY()
+	TMap<FName, FString> ObjectMetaData;
+};
+
+/**
+ * Cooked meta-data for a FField.
+ */
+USTRUCT()
+struct ENGINE_API FFieldCookedMetaDataStore
+{
+public:
+	GENERATED_BODY()
+
+	bool HasMetaData() const;
+	void CacheMetaData(const FField* SourceField);
+	void ApplyMetaData(FField* TargetField) const;
+
+private:
+	UPROPERTY()
+	TMap<FName, FString> FieldMetaData;
+};
+
+/**
+ * Cooked meta-data for a UStruct, including its nested FProperty data.
+ */
+USTRUCT()
+struct ENGINE_API FStructCookedMetaDataStore
+{
+public:
+	GENERATED_BODY()
+
+	bool HasMetaData() const;
+	void CacheMetaData(const UStruct* SourceStruct);
+	void ApplyMetaData(UStruct* TargetStruct) const;
+
+private:
+	UPROPERTY()
+	FObjectCookedMetaDataStore ObjectMetaData;
+
+	UPROPERTY()
+	TMap<FName, FFieldCookedMetaDataStore> PropertiesMetaData;
+};
+
+/**
+ * Cooked meta-data for a UEnum.
+ */
+UCLASS(Optional)
+class ENGINE_API UEnumCookedMetaData : public UObject
+{
+public:
+	GENERATED_BODY()
+
+	virtual bool HasMetaData() const;
+	virtual void CacheMetaData(const UEnum* SourceEnum);
+	virtual void ApplyMetaData(UEnum* TargetEnum) const;
+
+protected:
+	UPROPERTY()
+	FObjectCookedMetaDataStore EnumMetaData;
+};
+
+/**
+ * Cooked meta-data for a UScriptStruct, including its nested FProperty data.
+ */
+UCLASS(Optional)
+class ENGINE_API UStructCookedMetaData : public UObject
+{
+public:
+	GENERATED_BODY()
+
+	virtual bool HasMetaData() const;
+	virtual void CacheMetaData(const UScriptStruct* SourceStruct);
+	virtual void ApplyMetaData(UScriptStruct* TargetStruct) const;
+
+protected:
+	UPROPERTY()
+	FStructCookedMetaDataStore StructMetaData;
+};
+
+/**
+ * Cooked meta-data for a UClass, including its nested FProperty and UFunction data.
+ */
+UCLASS(Optional)
+class ENGINE_API UClassCookedMetaData : public UObject
+{
+public:
+	GENERATED_BODY()
+
+	virtual bool HasMetaData() const;
+	virtual void CacheMetaData(const UClass* SourceClass);
+	virtual void ApplyMetaData(UClass* TargetClass) const;
+
+protected:
+	UPROPERTY()
+	FStructCookedMetaDataStore ClassMetaData;
+
+	UPROPERTY()
+	TMap<FName, FStructCookedMetaDataStore> FunctionsMetaData;
+};
+
+namespace CookedMetaDataUtil
+{
+
+template <typename CookedMetaDataType>
+CookedMetaDataType* NewCookedMetaData(UObject* Outer, FName Name, TSubclassOf<CookedMetaDataType> Class = CookedMetaDataType::StaticClass())
+{
+	return NewObject<CookedMetaDataType>(Outer, Class, Name, RF_Standalone | RF_Public);
+}
+
+template <typename CookedMetaDataType>
+CookedMetaDataType* FindCookedMetaData(UObject* Outer, const TCHAR* Name)
+{
+	return FindObject<CookedMetaDataType>(Outer, Name);
+}
+
+template <typename CookedMetaDataType, typename CookedMetaDataPtrType>
+void PurgeCookedMetaData(CookedMetaDataPtrType& CookedMetaDataPtr)
+{
+	FNameBuilder BaseMetaDataName(CookedMetaDataPtr->GetFName());
+	BaseMetaDataName << TEXT("_PURGED");
+
+	CookedMetaDataPtr->ClearFlags(RF_Standalone | RF_Public);
+	CookedMetaDataPtr->Rename(FNameBuilder(MakeUniqueObjectName(CookedMetaDataPtr->GetOuter(), CookedMetaDataPtr->GetClass(), FName(BaseMetaDataName))).ToString(), nullptr, REN_DoNotDirty | REN_DontCreateRedirectors | REN_ForceNoResetLoaders | REN_NonTransactional);
+	CookedMetaDataPtr = nullptr;
+}
+
+}
