Split the common transaction diffing logic so that it can be used outside of the core transaction system

eg, this may be useful for diffing the SCS components within an actor annotation

The following types have been moved:
    * FPersistentObjectRef -> FTransactionPersistentObjectRef
    * FSerializedProperty -> FTransactionSerializedProperty
    * FSerializedObject -> FTransactionSerializedObject
    * FReader -> UE::Transaction::FSerializedObjectDataReader
    * FWriter -> UE::Transaction::FSerializedObjectDataWriter

They have also been made valid to use as a UPROPERTY, as annotations may need to serialize this data.

The existing code in Transactor.h/EditorTransaction.cpp now calls into this code/uses these types.

#rnx

diff --git a/Engine/Source/Editor/UnrealEd/Classes/Editor/Transactor.h b/Engine/Source/Editor/UnrealEd/Classes/Editor/Transactor.h
--- a/Engine/Source/Editor/UnrealEd/Classes/Editor/Transactor.h
+++ b/Engine/Source/Editor/UnrealEd/Classes/Editor/Transactor.h
@@ -14,6 +14,7 @@
 #include "Serialization/ArchiveUObject.h"
 #include "Serialization/ArchiveSerializedPropertyChain.h"
 #include "Misc/ITransaction.h"
+#include "TransactionCommon.h"
 #include "Algo/Find.h"
 #include <type_traits>
 #include "Transactor.generated.h"
@@ -36,205 +37,34 @@ class UNREALED_API FTransaction : public ITransaction
 	friend class FTransactionObjectEvent;
 
 protected:
-	/**
-		This type is necessary because the blueprint system is destroying and creating
-		CDOs at edit time (usually on compile, but also on load), but also stores user
-		entered data in the CDO. We "need"  changes to a CDO to persist across instances
-		because as we undo and redo we  need to apply changes to different instances of
-		the CDO - alternatively we could destroy and create the CDO as part of a transaction
-		(this alternative is the reason for the bunny ears around need).
-
-		DanO: My long term preference is for the editor to use a dynamic, mutable type
-		(rather than the CDO) to store editor data. The CDO can then be re-instanced (or not)
-		as runtime code requires.
-	*/
-	struct FPersistentObjectRef
-	{
-	public:
-		friend FArchive& operator<<(FArchive& Ar, FPersistentObjectRef& ReferencedObject)
-		{
-			Ar << (std::underlying_type_t<EReferenceType>&)ReferencedObject.ReferenceType;
-			Ar << ReferencedObject.RootObject;
-			Ar << ReferencedObject.SubObjectHierarchyIDs;
-			return Ar;
-		}
-
-		friend bool operator==(const FPersistentObjectRef& LHS, const FPersistentObjectRef& RHS)
-		{
-			return LHS.ReferenceType == RHS.ReferenceType
-				&& LHS.RootObject == RHS.RootObject
-				&& (LHS.ReferenceType != EReferenceType::SubObject || LHS.SubObjectHierarchyIDs == RHS.SubObjectHierarchyIDs);
-		}
-
-		friend bool operator!=(const FPersistentObjectRef& LHS, const FPersistentObjectRef& RHS)
-		{
-			return !(LHS == RHS);
-		}
-
-		friend uint32 GetTypeHash(const FPersistentObjectRef& InObjRef)
-		{
-			return HashCombine(GetTypeHash(InObjRef.ReferenceType), GetTypeHash(InObjRef.RootObject));
-		}
-
-		FPersistentObjectRef() = default;
-		explicit FPersistentObjectRef(UObject* InObject);
-
-		bool IsRootObjectReference() const
-		{
-			return ReferenceType == EReferenceType::RootObject;
-		}
-
-		bool IsSubObjectReference() const
-		{
-			return ReferenceType == EReferenceType::SubObject;
-		}
-
-		UObject* Get() const;
-
-		void AddReferencedObjects(FReferenceCollector& Collector);
-
-	private:
-		/** This enum represents all of the different special cases we are handling with this type */
-		enum class EReferenceType : uint8
-		{
-			Unknown,
-			RootObject,
-			SubObject,
-		};
-
-		/** The reference type we're handling */
-		EReferenceType ReferenceType = EReferenceType::Unknown;
-		/** Stores the object pointer when ReferenceType==RootObject, and the outermost pointer of the sub-object chain when when ReferenceType==SubObject */
-		TObjectPtr<UObject> RootObject = nullptr;
-		/** Stores the sub-object name chain when ReferenceType==SubObject */
-		TArray<FName, TInlineAllocator<4>> SubObjectHierarchyIDs;
-
-		/** Cached pointers corresponding to RootObject when ReferenceType==SubObject (@note cache needs testing on access as it may have become stale) */
-		mutable TWeakObjectPtr<UObject> CachedRootObject;
-		/** Cache of pointers corresponding to the items within SubObjectHierarchyIDs when ReferenceType==SubObject (@note cache needs testing on access as it may have become stale) */
-		mutable TArray<TWeakObjectPtr<UObject>, TInlineAllocator<4>> CachedSubObjectHierarchy;
-	};
+	using FPersistentObjectRef = FTransactionPersistentObjectRef;
 
 	// Record of an object.
 	class UNREALED_API FObjectRecord
 	{
 	public:
-		struct FSerializedProperty
-		{
-			FSerializedProperty()
-				: DataOffset(INDEX_NONE)
-				, DataSize(0)
-			{
-			}
-
-			static FName BuildSerializedPropertyKey(const FArchiveSerializedPropertyChain& InPropertyChain)
-			{
-				const int32 NumProperties = InPropertyChain.GetNumProperties();
-				check(NumProperties > 0);
-				return InPropertyChain.GetPropertyFromRoot(0)->GetFName();
-			}
-
-			void AppendSerializedData(const int64 InOffset, const int64 InSize)
-			{
-				if (DataOffset == INDEX_NONE)
-				{
-					DataOffset = InOffset;
-					DataSize = InSize;
-				}
-				else
-				{
-					DataOffset = FMath::Min(DataOffset, InOffset);
-					DataSize = FMath::Max(InOffset + InSize - DataOffset, DataSize);
-				}
-			}
-
-			/** Offset to the start of this property within the serialized object */
-			int64 DataOffset;
-			/** Size (in bytes) of this property within the serialized object */
-			int64 DataSize;
-		};
+		using FSerializedProperty = FTransactionSerializedProperty;
 
-		struct FSerializedObject
+		struct FSerializedObject : public FTransactionSerializedObject
 		{
-			FSerializedObject()
-				: bIsPendingKill(false)
-			{
-			}
-
 			void SetObject(const UObject* InObject)
 			{
-				ObjectPackageName = InObject->GetPackage()->GetFName();
-				ObjectName = InObject->GetFName();
-				ObjectPathName = *InObject->GetPathName();
-				ObjectOuterPathName = InObject->GetOuter() ? FName(*InObject->GetOuter()->GetPathName()) : FName();
-				ObjectExternalPackageName = InObject->GetExternalPackage() ? InObject->GetExternalPackage()->GetFName() : FName();
-				ObjectClassPathName = FName(*InObject->GetClass()->GetPathName());
-				bIsPendingKill = IsValid(InObject);
+				FTransactionSerializedObject::SetObject(InObject);
 				ObjectAnnotation = InObject->FindOrCreateTransactionAnnotation();
 			}
 
 			void Reset()
 			{
-				ObjectPackageName = FName();
-				ObjectName = FName();
-				ObjectPathName = FName();
-				ObjectOuterPathName = FName();
-				ObjectExternalPackageName = FName();
-				ObjectClassPathName = FName();
-				bIsPendingKill = false;
-				Data.Reset();
-				ReferencedObjects.Reset();
-				ReferencedNames.Reset();
-				SerializedProperties.Reset();
-				SerializedObjectIndices.Reset();
-				SerializedNameIndices.Reset();
+				FTransactionSerializedObject::Reset();
 				ObjectAnnotation.Reset();
 			}
 
 			void Swap(FSerializedObject& Other)
 			{
-				Exchange(ObjectPackageName, Other.ObjectPackageName);
-				Exchange(ObjectName, Other.ObjectName);
-				Exchange(ObjectPathName, Other.ObjectPathName);
-				Exchange(ObjectOuterPathName, Other.ObjectOuterPathName);
-				Exchange(ObjectExternalPackageName, Other.ObjectExternalPackageName);
-				Exchange(ObjectClassPathName, Other.ObjectClassPathName);
-				Exchange(bIsPendingKill, Other.bIsPendingKill);
-				Exchange(Data, Other.Data);
-				Exchange(ReferencedObjects, Other.ReferencedObjects);
-				Exchange(ReferencedNames, Other.ReferencedNames);
-				Exchange(SerializedProperties, Other.SerializedProperties);
-				Exchange(SerializedObjectIndices, Other.SerializedObjectIndices);
-				Exchange(SerializedNameIndices, Other.SerializedNameIndices);
+				FTransactionSerializedObject::Swap(Other);
 				Exchange(ObjectAnnotation, Other.ObjectAnnotation);
 			}
 
-			/** The package name of the object when it was serialized, can be dictated either by outer chain or external package */
-			FName ObjectPackageName;
-			/** The name of the object when it was serialized */
-			FName ObjectName;
-			/** The path name of the object when it was serialized */
-			FName ObjectPathName;
-			/** The outer path name of the object when it was serialized */
-			FName ObjectOuterPathName;
-			/** The external package name of the object when it was serialized, if any */
-			FName ObjectExternalPackageName;
-			/** The path name of the object's class. */
-			FName ObjectClassPathName;
-			/** The pending kill state of the object when it was serialized */
-			bool bIsPendingKill;
-			/** The data stream used to serialize/deserialize record */
-			TArray64<uint8> Data;
-			/** External objects referenced in the transaction */
-			TArray<FPersistentObjectRef> ReferencedObjects;
-			/** FNames referenced in the object record */
-			TArray<FName> ReferencedNames;
-			/** Information about the properties that were serialized within this object */
-			TMap<FName, FSerializedProperty> SerializedProperties;
-			/** Information about the object pointer offsets that were serialized within this object (this maps the property name (or None if there was no property) to the ReferencedObjects indices of the property) */
-			TMultiMap<FName, int32> SerializedObjectIndices;
-			/** Information about the name offsets that were serialized within this object (this maps the property name to the ReferencedNames index of the property) */
-			TMultiMap<FName, int32> SerializedNameIndices;
 			/** Annotation data for the object stored externally */
 			TSharedPtr<ITransactionObjectAnnotation> ObjectAnnotation;
 		};
@@ -313,7 +143,7 @@ protected:
 		bool HasExpired() const;
 
 		/** Transfers data from an array. */
-		class FReader : public FArchiveUObject
+		class FReader : public UE::Transaction::FSerializedObjectDataReader
 		{
 		public:
 			FReader(
@@ -321,50 +151,14 @@ protected:
 				const FSerializedObject& InSerializedObject,
 				bool bWantBinarySerialization
 				)
-				: Owner(InOwner)
-				, SerializedObject(InSerializedObject)
-				, Offset(0)
+				: UE::Transaction::FSerializedObjectDataReader(InSerializedObject)
+				, Owner(InOwner)
 			{
 				this->SetWantBinaryPropertySerialization(bWantBinarySerialization);
-				this->SetIsLoading(true);
 				this->SetIsTransacting(true);
 			}
 
-			virtual int64 Tell() override {return Offset;}
-			virtual void Seek( int64 InPos ) override { Offset = InPos; }
-			virtual int64 TotalSize() override { return SerializedObject.Data.Num(); }
-
 		private:
-			void Serialize( void* SerData, int64 Num ) override
-			{
-				if( Num )
-				{
-					checkSlow(Offset+Num<=SerializedObject.Data.Num());
-					FMemory::Memcpy( SerData, &SerializedObject.Data[Offset], Num );
-					Offset += Num;
-				}
-			}
-			FArchive& operator<<( class FName& N ) override
-			{
-				int32 NameIndex = 0;
-				(FArchive&)*this << NameIndex;
-				N = SerializedObject.ReferencedNames[NameIndex];
-				return *this;
-			}
-			FArchive& operator<<( class UObject*& Res ) override
-			{
-				int32 ObjectIndex = 0;
-				(FArchive&)*this << ObjectIndex;
-				if (ObjectIndex != INDEX_NONE)
-				{
-					Res = SerializedObject.ReferencedObjects[ObjectIndex].Get();
-				}
-				else
-				{
-					Res = nullptr;
-				}
-				return *this;
-			}
 			void Preload( UObject* InObject ) override
 			{
 				if (Owner)
@@ -383,14 +177,12 @@ protected:
 				}
 			}
 			FTransaction* Owner;
-			const FSerializedObject& SerializedObject;
-			int64 Offset;
 		};
 
 		/**
 		 * Transfers data to an array.
 		 */
-		class FWriter : public FArchiveUObject
+		class FWriter : public UE::Transaction::FSerializedObjectDataWriter
 		{
 		public:
 			FWriter(
@@ -398,34 +190,13 @@ protected:
 				bool bWantBinarySerialization,
 				TArrayView<const FProperty*> InPropertiesToSerialize = TArrayView<const FProperty*>()
 				)
-				: SerializedObject(InSerializedObject)
+				: UE::Transaction::FSerializedObjectDataWriter(InSerializedObject)
 				, PropertiesToSerialize(InPropertiesToSerialize)
-				, Offset(0)
 			{
-				for(int32 ObjIndex = 0; ObjIndex < SerializedObject.ReferencedObjects.Num(); ++ObjIndex)
-				{
-					ObjectMap.Add(SerializedObject.ReferencedObjects[ObjIndex].Get(), ObjIndex);
-				}
-
-				for(int32 NameIndex = 0; NameIndex < SerializedObject.ReferencedNames.Num(); ++NameIndex)
-				{
-					NameMap.Add(SerializedObject.ReferencedNames[NameIndex], NameIndex);
-				}
-
 				this->SetWantBinaryPropertySerialization(bWantBinarySerialization);
-				this->SetIsSaving(true);
 				this->SetIsTransacting(true);
 			}
 
-			virtual int64 Tell() override {return Offset;}
-			virtual void Seek( int64 InPos ) override
-			{
-				checkSlow(Offset<=SerializedObject.Data.Num());
-				Offset = InPos; 
-			}
-
-			virtual int64 TotalSize() override { return SerializedObject.Data.Num(); }
-
 			virtual bool ShouldSkipProperty(const FProperty* InProperty) const override
 			{
 				return (PropertiesToSerialize.Num() > 0 && !PropertiesToSerialize.Contains(InProperty))
@@ -433,112 +204,7 @@ protected:
 			}
 
 		private:
-			struct FCachedPropertyKey
-			{
-			public:
-				FCachedPropertyKey()
-					: LastUpdateCount(0)
-				{
-				}
-
-				FName SyncCache(const FArchiveSerializedPropertyChain* InPropertyChain)
-				{
-					if (InPropertyChain)
-					{
-						const uint32 CurrentUpdateCount = InPropertyChain->GetUpdateCount();
-						if (CurrentUpdateCount != LastUpdateCount)
-						{
-							CachedKey = InPropertyChain->GetNumProperties() > 0 ? FSerializedProperty::BuildSerializedPropertyKey(*InPropertyChain) : FName();
-							LastUpdateCount = CurrentUpdateCount;
-						}
-					}
-					else
-					{
-						CachedKey = FName();
-						LastUpdateCount = 0;
-					}
-
-					return CachedKey;
-				}
-
-			private:
-				FName CachedKey;
-				uint32 LastUpdateCount;
-			};
-
-			void Serialize( void* SerData, int64 Num ) override
-			{
-				if( Num )
-				{
-					int64 DataIndex = ( Offset == SerializedObject.Data.Num() )
-						?  SerializedObject.Data.AddUninitialized(Num)
-						:  Offset;
-					
-					FMemory::Memcpy( &SerializedObject.Data[DataIndex], SerData, Num );
-					Offset+= Num;
-
-					// Track this property offset in the serialized data
-					const FArchiveSerializedPropertyChain* PropertyChain = GetSerializedPropertyChain();
-					if (PropertyChain && PropertyChain->GetNumProperties() > 0)
-					{
-						const FName SerializedTaggedPropertyKey = CachedSerializedTaggedPropertyKey.SyncCache(PropertyChain);
-						FSerializedProperty& SerializedTaggedProperty = SerializedObject.SerializedProperties.FindOrAdd(SerializedTaggedPropertyKey);
-						SerializedTaggedProperty.AppendSerializedData(DataIndex, Num);
-					}
-				}
-			}
-			FArchive& operator<<( class FName& N ) override
-			{
-				int32 NameIndex = INDEX_NONE;
-				const int32 * NameIndexPtr = NameMap.Find(N);
-				if (NameIndexPtr)
-				{
-					NameIndex = *NameIndexPtr;
-				}
-				else
-				{
-					NameIndex = SerializedObject.ReferencedNames.Add(N);
-					NameMap.Add(N, NameIndex);
-				}
-
-				// Track this name index in the serialized data
-				{
-					const FArchiveSerializedPropertyChain* PropertyChain = GetSerializedPropertyChain();
-					const FName SerializedTaggedPropertyKey = CachedSerializedTaggedPropertyKey.SyncCache(PropertyChain);
-					SerializedObject.SerializedNameIndices.Add(SerializedTaggedPropertyKey, NameIndex);
-				}
-
-				return (FArchive&)*this << NameIndex;
-			} 
-			FArchive& operator<<( class UObject*& Res ) override
-			{
-				int32 ObjectIndex = INDEX_NONE;
-				const int32* ObjIndexPtr = ObjectMap.Find(Res);
-				if(ObjIndexPtr)
-				{
-					ObjectIndex = *ObjIndexPtr;
-				}
-				else if(Res)
-				{
-					ObjectIndex = SerializedObject.ReferencedObjects.Add(FPersistentObjectRef(Res));
-					ObjectMap.Add(Res, ObjectIndex);
-				}
-
-				// Track this object offset in the serialized data
-				{
-					const FArchiveSerializedPropertyChain* PropertyChain = GetSerializedPropertyChain();
-					const FName SerializedTaggedPropertyKey = CachedSerializedTaggedPropertyKey.SyncCache(PropertyChain);
-					SerializedObject.SerializedObjectIndices.Add(SerializedTaggedPropertyKey, ObjectIndex);
-				}
-
-				return (FArchive&)*this << ObjectIndex;
-			}
-			FSerializedObject& SerializedObject;
 			TArrayView<const FProperty*> PropertiesToSerialize;
-			TMap<UObject*, int32> ObjectMap;
-			TMap<FName, int32> NameMap;
-			FCachedPropertyKey CachedSerializedTaggedPropertyKey;
-			int64 Offset;
 		};
 	};
 
diff --git a/Engine/Source/Editor/UnrealEd/Private/EditorTransaction.cpp b/Engine/Source/Editor/UnrealEd/Private/EditorTransaction.cpp
--- a/Engine/Source/Editor/UnrealEd/Private/EditorTransaction.cpp
+++ b/Engine/Source/Editor/UnrealEd/Private/EditorTransaction.cpp
@@ -228,7 +228,7 @@ void FTransaction::FObjectRecord::Load(FTransaction* Owner)
 			// When objects are created outside the transaction system we can end up
 			// finding them but not having any data for them, so don't serialize 
 			// when that happens:
-			if (SerializedObject.Data.Num() > 0)
+			if (SerializedObject.SerializedData.Num() > 0)
 			{
 				FReader Reader(Owner, SerializedObject, bWantsBinarySerialization);
 				SerializeContents(Reader, Oper);
@@ -338,174 +338,7 @@ void FTransaction::FObjectRecord::Snapshot( FTransaction* Owner, TArrayView<cons
 
 void FTransaction::FObjectRecord::Diff( const FTransaction* Owner, const FSerializedObject& OldSerializedObject, const FSerializedObject& NewSerializedObject, FTransactionObjectDeltaChange& OutDeltaChange, const bool bFullDiff )
 {
-	auto AreObjectPointersIdentical = [&OldSerializedObject, &NewSerializedObject](const FName InPropertyName)
-	{
-		TArray<int32, TInlineAllocator<8>> OldSerializedObjectIndices;
-		OldSerializedObject.SerializedObjectIndices.MultiFind(InPropertyName, OldSerializedObjectIndices, true);
-
-		TArray<int32, TInlineAllocator<8>> NewSerializedObjectIndices;
-		NewSerializedObject.SerializedObjectIndices.MultiFind(InPropertyName, NewSerializedObjectIndices, true);
-
-		bool bAreObjectPointersIdentical = OldSerializedObjectIndices.Num() == NewSerializedObjectIndices.Num();
-		for (int32 ObjIndex = 0; ObjIndex < OldSerializedObjectIndices.Num() && bAreObjectPointersIdentical; ++ObjIndex)
-		{
-			const FPersistentObjectRef& OldObjectRef = OldSerializedObject.ReferencedObjects.IsValidIndex(OldSerializedObjectIndices[ObjIndex]) ? OldSerializedObject.ReferencedObjects[OldSerializedObjectIndices[ObjIndex]] : FPersistentObjectRef();
-			const FPersistentObjectRef& NewObjectRef = NewSerializedObject.ReferencedObjects.IsValidIndex(NewSerializedObjectIndices[ObjIndex]) ? NewSerializedObject.ReferencedObjects[NewSerializedObjectIndices[ObjIndex]] : FPersistentObjectRef();
-			bAreObjectPointersIdentical = OldObjectRef == NewObjectRef;
-		}
-		return bAreObjectPointersIdentical;
-	};
-
-	auto AreNamesIdentical = [&OldSerializedObject, &NewSerializedObject](const FName InPropertyName)
-	{
-		TArray<int32, TInlineAllocator<8>> OldSerializedNameIndices;
-		OldSerializedObject.SerializedNameIndices.MultiFind(InPropertyName, OldSerializedNameIndices, true);
-
-		TArray<int32, TInlineAllocator<8>> NewSerializedNameIndices;
-		NewSerializedObject.SerializedNameIndices.MultiFind(InPropertyName, NewSerializedNameIndices, true);
-
-		bool bAreNamesIdentical = OldSerializedNameIndices.Num() == NewSerializedNameIndices.Num();
-		for (int32 ObjIndex = 0; ObjIndex < OldSerializedNameIndices.Num() && bAreNamesIdentical; ++ObjIndex)
-		{
-			const FName& OldName = OldSerializedObject.ReferencedNames.IsValidIndex(OldSerializedNameIndices[ObjIndex]) ? OldSerializedObject.ReferencedNames[OldSerializedNameIndices[ObjIndex]] : FName();
-			const FName& NewName = NewSerializedObject.ReferencedNames.IsValidIndex(NewSerializedNameIndices[ObjIndex]) ? NewSerializedObject.ReferencedNames[NewSerializedNameIndices[ObjIndex]] : FName();
-			bAreNamesIdentical = OldName == NewName;
-		}
-		return bAreNamesIdentical;
-	};
-
-	if (bFullDiff)
-	{
-		OutDeltaChange.bHasNameChange |= OldSerializedObject.ObjectName != NewSerializedObject.ObjectName;
-		OutDeltaChange.bHasOuterChange |= OldSerializedObject.ObjectOuterPathName != NewSerializedObject.ObjectOuterPathName;
-		OutDeltaChange.bHasExternalPackageChange |= OldSerializedObject.ObjectExternalPackageName != NewSerializedObject.ObjectExternalPackageName;
-		OutDeltaChange.bHasPendingKillChange |= OldSerializedObject.bIsPendingKill != NewSerializedObject.bIsPendingKill;
-
-		if (!AreObjectPointersIdentical(NAME_None))
-		{
-			OutDeltaChange.bHasNonPropertyChanges = true;
-		}
-
-		if (!AreNamesIdentical(NAME_None))
-		{
-			OutDeltaChange.bHasNonPropertyChanges = true;
-		}
-	}
-
-	if (OldSerializedObject.SerializedProperties.Num() > 0 || NewSerializedObject.SerializedProperties.Num() > 0)
-	{
-		int64 StartOfOldPropertyBlock = INT64_MAX;
-		int64 StartOfNewPropertyBlock = INT64_MAX;
-		int64 EndOfOldPropertyBlock = -1;
-		int64 EndOfNewPropertyBlock = -1;
-
-		for (const TPair<FName, FSerializedProperty>& NewNamePropertyPair : NewSerializedObject.SerializedProperties)
-		{
-			const FSerializedProperty* OldSerializedProperty = OldSerializedObject.SerializedProperties.Find(NewNamePropertyPair.Key);
-			if (!OldSerializedProperty)
-			{
-				if (bFullDiff)
-				{
-					// Missing property, assume that the property changed
-					OutDeltaChange.ChangedProperties.AddUnique(NewNamePropertyPair.Key);
-				}
-				continue;
-			}
-
-			// Update the tracking for the start/end of the property block within the serialized data
-			StartOfOldPropertyBlock = FMath::Min(StartOfOldPropertyBlock, OldSerializedProperty->DataOffset);
-			StartOfNewPropertyBlock = FMath::Min(StartOfNewPropertyBlock, NewNamePropertyPair.Value.DataOffset);
-			EndOfOldPropertyBlock = FMath::Max(EndOfOldPropertyBlock, OldSerializedProperty->DataOffset + OldSerializedProperty->DataSize);
-			EndOfNewPropertyBlock = FMath::Max(EndOfNewPropertyBlock, NewNamePropertyPair.Value.DataOffset + NewNamePropertyPair.Value.DataSize);
-
-			// Binary compare the serialized data to see if something has changed for this property
-			bool bIsPropertyIdentical = OldSerializedProperty->DataSize == NewNamePropertyPair.Value.DataSize;
-			if (bIsPropertyIdentical && NewNamePropertyPair.Value.DataSize > 0)
-			{
-				bIsPropertyIdentical = FMemory::Memcmp(&OldSerializedObject.Data[OldSerializedProperty->DataOffset], &NewSerializedObject.Data[NewNamePropertyPair.Value.DataOffset], NewNamePropertyPair.Value.DataSize) == 0;
-			}
-			if (bIsPropertyIdentical)
-			{
-				bIsPropertyIdentical = AreObjectPointersIdentical(NewNamePropertyPair.Key);
-			}
-			if (bIsPropertyIdentical)
-			{
-				bIsPropertyIdentical = AreNamesIdentical(NewNamePropertyPair.Key);
-			}
-
-			if (!bIsPropertyIdentical)
-			{
-				OutDeltaChange.ChangedProperties.AddUnique(NewNamePropertyPair.Key);
-			}
-		}
-
-		for (const TPair<FName, FSerializedProperty>& OldNamePropertyPair : OldSerializedObject.SerializedProperties)
-		{
-			const FSerializedProperty* NewSerializedProperty = NewSerializedObject.SerializedProperties.Find(OldNamePropertyPair.Key);
-			if (!NewSerializedProperty)
-			{
-				if (bFullDiff)
-				{
-					// Missing property, assume that the property changed
-					OutDeltaChange.ChangedProperties.AddUnique(OldNamePropertyPair.Key);
-				}
-				continue;
-			}
-		}
-
-		if (bFullDiff)
-		{
-			// Compare the data before the property block to see if something else in the object has changed
-			if (!OutDeltaChange.bHasNonPropertyChanges)
-			{
-				const int64 OldHeaderSize = FMath::Min(StartOfOldPropertyBlock, OldSerializedObject.Data.Num());
-				const int64 CurrentHeaderSize = FMath::Min(StartOfNewPropertyBlock, NewSerializedObject.Data.Num());
-
-				bool bIsHeaderIdentical = OldHeaderSize == CurrentHeaderSize;
-				if (bIsHeaderIdentical && CurrentHeaderSize > 0)
-				{
-					bIsHeaderIdentical = FMemory::Memcmp(&OldSerializedObject.Data[0], &NewSerializedObject.Data[0], CurrentHeaderSize) == 0;
-				}
-
-				if (!bIsHeaderIdentical)
-				{
-					OutDeltaChange.bHasNonPropertyChanges = true;
-				}
-			}
-
-			// Compare the data after the property block to see if something else in the object has changed
-			if (!OutDeltaChange.bHasNonPropertyChanges)
-			{
-				const int64 OldFooterSize = OldSerializedObject.Data.Num() - FMath::Max<int64>(EndOfOldPropertyBlock, 0);
-				const int64 CurrentFooterSize = NewSerializedObject.Data.Num() - FMath::Max<int64>(EndOfNewPropertyBlock, 0);
-
-				bool bIsFooterIdentical = OldFooterSize == CurrentFooterSize;
-				if (bIsFooterIdentical && CurrentFooterSize > 0)
-				{
-					bIsFooterIdentical = FMemory::Memcmp(&OldSerializedObject.Data[EndOfOldPropertyBlock], &NewSerializedObject.Data[EndOfNewPropertyBlock], CurrentFooterSize) == 0;
-				}
-
-				if (!bIsFooterIdentical)
-				{
-					OutDeltaChange.bHasNonPropertyChanges = true;
-				}
-			}
-		}
-	}
-	else if (bFullDiff)
-	{
-		// No properties, so just compare the whole blob
-		bool bIsBlobIdentical = OldSerializedObject.Data.Num() == NewSerializedObject.Data.Num();
-		if (bIsBlobIdentical && NewSerializedObject.Data.Num() > 0)
-		{
-			bIsBlobIdentical = FMemory::Memcmp(&OldSerializedObject.Data[0], &NewSerializedObject.Data[0], NewSerializedObject.Data.Num()) == 0;
-		}
-
-		if (!bIsBlobIdentical)
-		{
-			OutDeltaChange.bHasNonPropertyChanges = true;
-		}
-	}
+	UE::Transaction::DiffUtil::GenerateObjectDiff(OldSerializedObject, NewSerializedObject, OutDeltaChange, bFullDiff);
 }
 
 int32 FTransaction::GetRecordCount() const
@@ -607,152 +440,22 @@ FArchive& operator<<( FArchive& Ar, FTransaction::FObjectRecord& R )
 {
 	FMemMark Mark(FMemStack::Get());
 	Ar << R.Object;
-	Ar << R.SerializedObject.Data;
+	Ar << R.SerializedObject.SerializedData;
 	Ar << R.SerializedObject.ReferencedObjects;
 	Ar << R.SerializedObject.ReferencedNames;
 	Mark.Pop();
 	return Ar;
 }
 
-FTransaction::FPersistentObjectRef::FPersistentObjectRef(UObject* InObject)
-{
-	RootObject = InObject;
-	{
-		auto UseOuter = [](const UObject* Obj)
-		{
-			if (Obj == nullptr)
-			{
-				return false;
-			}
-
-			const bool bIsCDO = Obj->HasAllFlags(RF_ClassDefaultObject);
-			const UObject* CDO = bIsCDO ? Obj : nullptr;
-			const bool bIsClassCDO = (CDO != nullptr) ? (CDO->GetClass()->ClassDefaultObject == CDO) : false;
-			if (!bIsClassCDO && CDO)
-			{
-				// Likely a trashed CDO, try to recover
-				// Only known cause of this is ambiguous use of DSOs
-				CDO = CDO->GetClass()->ClassDefaultObject;
-			}
-			const UActorComponent* AsComponent = Cast<UActorComponent>(Obj);
-			const bool bIsDSO = Obj->HasAnyFlags(RF_DefaultSubObject);
-			const bool bIsSCSComponent = AsComponent && AsComponent->IsCreatedByConstructionScript();
-			return (bIsCDO && bIsClassCDO) || bIsDSO || bIsSCSComponent;
-		};
-
-		while (UseOuter(RootObject))
-		{
-			SubObjectHierarchyIDs.Add(RootObject->GetFName());
-			RootObject = RootObject->GetOuter();
-		}
-	}
-	check(RootObject);
-
-	if (SubObjectHierarchyIDs.Num() > 0)
-	{
-		ReferenceType = EReferenceType::SubObject;
-		Algo::Reverse(SubObjectHierarchyIDs);
-	}
-	else
-	{
-		ReferenceType = EReferenceType::RootObject;
-	}
-
-	// Make sure that when we look up the object we find the same thing:
-	checkSlow(Get() == InObject);
-}
-
-UObject* FTransaction::FPersistentObjectRef::Get() const
-{
-	if (ReferenceType == EReferenceType::SubObject)
-	{
-		check(SubObjectHierarchyIDs.Num() > 0);
-
-		UObject* CurrentObject = nullptr;
-
-		// Do we have a valid cached pointer?
-		if (!CachedRootObject.IsExplicitlyNull() && SubObjectHierarchyIDs.Num() == CachedSubObjectHierarchy.Num())
-		{
-			// Root object is a pointer test
-			{
-				CurrentObject = CachedRootObject.GetEvenIfUnreachable();
-				if (CurrentObject != RootObject)
-				{
-					CurrentObject = nullptr;
-				}
-			}
-
-			// All other sub-objects are a name test
-			for (int32 SubObjectIndex = 0; CurrentObject && SubObjectIndex < SubObjectHierarchyIDs.Num(); ++SubObjectIndex)
-			{
-				CurrentObject = CachedSubObjectHierarchy[SubObjectIndex].GetEvenIfUnreachable();
-				if (CurrentObject && CurrentObject->GetFName() != SubObjectHierarchyIDs[SubObjectIndex])
-				{
-					CurrentObject = nullptr;
-				}
-			}
-		}
-		if (CurrentObject)
-		{
-			return CurrentObject;
-		}
-
-		// Cached pointer is invalid
-		CachedRootObject.Reset();
-		CachedSubObjectHierarchy.Reset();
-
-		// Try to find and cache the subobject
-		CachedRootObject = RootObject;
-		CurrentObject = RootObject;
-		for (int32 SubObjectIndex = 0; CurrentObject && SubObjectIndex < SubObjectHierarchyIDs.Num(); ++SubObjectIndex)
-		{
-			CurrentObject = StaticFindObjectFast(UObject::StaticClass(), CurrentObject, SubObjectHierarchyIDs[SubObjectIndex]);
-			CachedSubObjectHierarchy.Add(CurrentObject);
-		}
-		if (CurrentObject)
-		{
-			check(!CachedRootObject.IsExplicitlyNull() && SubObjectHierarchyIDs.Num() == CachedSubObjectHierarchy.Num());
-			return CurrentObject;
-		}
-
-		// Cached pointer is invalid
-		CachedRootObject.Reset();
-		CachedSubObjectHierarchy.Reset();
-
-		return nullptr;
-	}
-
-	return RootObject;
-}
-
-void FTransaction::FPersistentObjectRef::AddReferencedObjects(FReferenceCollector& Collector)
-{
-	Collector.AddReferencedObject(RootObject);
-
-	if (ReferenceType == EReferenceType::SubObject)
-	{
-		// We can't refresh the resolved pointers during ARO, as it's not safe to call FindObject to update the cache if stale
-		// Instead we'll just ARO whatever we may have cached, as this may result in the resolved pointers being updated anyway
-		// Note: This is needed as sub-objects may be subject to GC while inside the transaction buffer, as the references from their root 
-		// object may have been removed (eg, a component on an actor will no longer be referenced by the actor after a delete operation)
-		for (TWeakObjectPtr<UObject>& CachedSubObject : CachedSubObjectHierarchy)
-		{
-			UObject* CachedSubObjectPtr = CachedSubObject.GetEvenIfUnreachable();
-			Collector.AddReferencedObject(CachedSubObjectPtr);
-			CachedSubObject = CachedSubObjectPtr;
-		}
-	}
-}
-
 void FTransaction::FObjectRecord::AddReferencedObjects( FReferenceCollector& Collector )
 {
-	Object.AddReferencedObjects(Collector);
+	Object.AddStructReferencedObjects(Collector);
 
 	auto AddSerializedObjectReferences = [&Collector](FSerializedObject& InSerializedObject)
 	{
 		for (FPersistentObjectRef& ReferencedObject : InSerializedObject.ReferencedObjects)
 		{
-			ReferencedObject.AddReferencedObjects(Collector);
+			ReferencedObject.AddStructReferencedObjects(Collector);
 		}
 
 		if (InSerializedObject.ObjectAnnotation.IsValid())
@@ -832,7 +535,7 @@ void FTransaction::AddReferencedObjects( FReferenceCollector& Collector )
 
 	for (TTuple<FPersistentObjectRef, FObjectRecords>& ObjectRecordsPair : ObjectRecordsMap)
 	{
-		ObjectRecordsPair.Key.AddReferencedObjects(Collector);
+		ObjectRecordsPair.Key.AddStructReferencedObjects(Collector);
 	}
 }
 
@@ -1116,7 +819,7 @@ SIZE_T FTransaction::DataSize() const
 	SIZE_T Result=0;
 	for( int32 i=0; i<Records.Num(); i++ )
 	{
-		Result += Records[i].SerializedObject.Data.Num();
+		Result += Records[i].SerializedObject.SerializedData.Num();
 	}
 	return Result;
 }
diff --git a/Engine/Source/Runtime/Engine/Private/TransactionCommon.cpp b/Engine/Source/Runtime/Engine/Private/TransactionCommon.cpp
new file mode 100644
--- /dev/null
+++ b/Engine/Source/Runtime/Engine/Private/TransactionCommon.cpp
@@ -0,0 +1,523 @@
+// Copyright Epic Games, Inc. All Rights Reserved.
+
+#include "TransactionCommon.h"
+#include "Misc/ITransaction.h"
+#include "UObject/UnrealType.h"
+#include "Components/ActorComponent.h"
+
+FTransactionPersistentObjectRef::FTransactionPersistentObjectRef(UObject* InObject)
+{
+	RootObject = InObject;
+	{
+		auto UseOuter = [](const UObject* Obj)
+		{
+			if (Obj == nullptr)
+			{
+				return false;
+			}
+
+			const bool bIsCDO = Obj->HasAllFlags(RF_ClassDefaultObject);
+			const UObject* CDO = bIsCDO ? Obj : nullptr;
+			const bool bIsClassCDO = (CDO != nullptr) ? (CDO->GetClass()->ClassDefaultObject == CDO) : false;
+			if (!bIsClassCDO && CDO)
+			{
+				// Likely a trashed CDO, try to recover
+				// Only known cause of this is ambiguous use of DSOs
+				CDO = CDO->GetClass()->ClassDefaultObject;
+			}
+			const UActorComponent* AsComponent = Cast<UActorComponent>(Obj);
+			const bool bIsDSO = Obj->HasAnyFlags(RF_DefaultSubObject);
+			const bool bIsSCSComponent = AsComponent && AsComponent->IsCreatedByConstructionScript();
+			return (bIsCDO && bIsClassCDO) || bIsDSO || bIsSCSComponent;
+		};
+
+		while (UseOuter(RootObject))
+		{
+			SubObjectHierarchyIDs.Add(RootObject->GetFName());
+			RootObject = RootObject->GetOuter();
+		}
+	}
+	check(RootObject);
+
+	if (SubObjectHierarchyIDs.Num() > 0)
+	{
+		ReferenceType = EReferenceType::SubObject;
+		Algo::Reverse(SubObjectHierarchyIDs);
+	}
+	else
+	{
+		ReferenceType = EReferenceType::RootObject;
+	}
+
+	// Make sure that when we look up the object we find the same thing:
+	checkSlow(Get() == InObject);
+}
+
+UObject* FTransactionPersistentObjectRef::Get() const
+{
+	if (ReferenceType == EReferenceType::SubObject)
+	{
+		check(SubObjectHierarchyIDs.Num() > 0);
+
+		UObject* CurrentObject = nullptr;
+
+		// Do we have a valid cached pointer?
+		if (!CachedRootObject.IsExplicitlyNull() && SubObjectHierarchyIDs.Num() == CachedSubObjectHierarchy.Num())
+		{
+			// Root object is a pointer test
+			{
+				CurrentObject = CachedRootObject.GetEvenIfUnreachable();
+				if (CurrentObject != RootObject)
+				{
+					CurrentObject = nullptr;
+				}
+			}
+
+			// All other sub-objects are a name test
+			for (int32 SubObjectIndex = 0; CurrentObject && SubObjectIndex < SubObjectHierarchyIDs.Num(); ++SubObjectIndex)
+			{
+				CurrentObject = CachedSubObjectHierarchy[SubObjectIndex].GetEvenIfUnreachable();
+				if (CurrentObject && CurrentObject->GetFName() != SubObjectHierarchyIDs[SubObjectIndex])
+				{
+					CurrentObject = nullptr;
+				}
+			}
+		}
+		if (CurrentObject)
+		{
+			return CurrentObject;
+		}
+
+		// Cached pointer is invalid
+		CachedRootObject.Reset();
+		CachedSubObjectHierarchy.Reset();
+
+		// Try to find and cache the subobject
+		CachedRootObject = RootObject;
+		CurrentObject = RootObject;
+		for (int32 SubObjectIndex = 0; CurrentObject && SubObjectIndex < SubObjectHierarchyIDs.Num(); ++SubObjectIndex)
+		{
+			CurrentObject = StaticFindObjectFast(UObject::StaticClass(), CurrentObject, SubObjectHierarchyIDs[SubObjectIndex]);
+			CachedSubObjectHierarchy.Add(CurrentObject);
+		}
+		if (CurrentObject)
+		{
+			check(!CachedRootObject.IsExplicitlyNull() && SubObjectHierarchyIDs.Num() == CachedSubObjectHierarchy.Num());
+			return CurrentObject;
+		}
+
+		// Cached pointer is invalid
+		CachedRootObject.Reset();
+		CachedSubObjectHierarchy.Reset();
+
+		return nullptr;
+	}
+
+	return RootObject;
+}
+
+void FTransactionPersistentObjectRef::AddStructReferencedObjects(FReferenceCollector& Collector)
+{
+	Collector.AddReferencedObject(RootObject);
+
+	if (ReferenceType == EReferenceType::SubObject)
+	{
+		// We can't refresh the resolved pointers during ARO, as it's not safe to call FindObject to update the cache if stale
+		// Instead we'll just ARO whatever we may have cached, as this may result in the resolved pointers being updated anyway
+		// Note: This is needed as sub-objects may be subject to GC while inside the transaction buffer, as the references from their root 
+		// object may have been removed (eg, a component on an actor will no longer be referenced by the actor after a delete operation)
+		for (TWeakObjectPtr<UObject>& CachedSubObject : CachedSubObjectHierarchy)
+		{
+			UObject* CachedSubObjectPtr = CachedSubObject.GetEvenIfUnreachable();
+			Collector.AddReferencedObject(CachedSubObjectPtr);
+			CachedSubObject = CachedSubObjectPtr;
+		}
+	}
+}
+
+
+FName FTransactionSerializedProperty::BuildSerializedPropertyKey(const FArchiveSerializedPropertyChain& InPropertyChain)
+{
+	const int32 NumProperties = InPropertyChain.GetNumProperties();
+	check(NumProperties > 0);
+	return InPropertyChain.GetPropertyFromRoot(0)->GetFName();
+}
+
+void FTransactionSerializedProperty::AppendSerializedData(const int64 InOffset, const int64 InSize)
+{
+	if (DataOffset == INDEX_NONE)
+	{
+		DataOffset = InOffset;
+		DataSize = InSize;
+	}
+	else
+	{
+		DataOffset = FMath::Min(DataOffset, InOffset);
+		DataSize = FMath::Max(InOffset + InSize - DataOffset, DataSize);
+	}
+}
+
+
+void FTransactionSerializedObjectData::Read(void* Dest, int64 Offset, int64 Num) const
+{
+	checkSlow(Offset + Num <= Data.Num());
+	FMemory::Memcpy(Dest, &Data[Offset], Num);
+}
+
+void FTransactionSerializedObjectData::Write(const void* Src, int64 Offset, int64 Num)
+{
+	if (Offset == Data.Num())
+	{
+		Data.AddUninitialized(Num);
+	}
+	FMemory::Memcpy(&Data[Offset], Src, Num);
+}
+
+
+namespace UE::Transaction
+{
+
+FSerializedObjectDataReader::FSerializedObjectDataReader(const FTransactionSerializedObject& InSerializedObject)
+	: SerializedObject(InSerializedObject)
+	, Offset(0)
+{
+	SetIsLoading(true);
+}
+
+void FSerializedObjectDataReader::Serialize(void* SerData, int64 Num)
+{
+	if (Num)
+	{
+		SerializedObject.SerializedData.Read(SerData, Offset, Num);
+		Offset += Num;
+	}
+}
+
+FArchive& FSerializedObjectDataReader::operator<<(FName& N)
+{
+	int32 NameIndex = 0;
+	*this << NameIndex;
+	N = SerializedObject.ReferencedNames[NameIndex];
+	return *this;
+}
+
+FArchive& FSerializedObjectDataReader::operator<<(UObject*& Res)
+{
+	int32 ObjectIndex = 0;
+	*this << ObjectIndex;
+	if (ObjectIndex != INDEX_NONE)
+	{
+		Res = SerializedObject.ReferencedObjects[ObjectIndex].Get();
+	}
+	else
+	{
+		Res = nullptr;
+	}
+	return *this;
+}
+
+
+FSerializedObjectDataWriter::FSerializedObjectDataWriter(FTransactionSerializedObject& InSerializedObject)
+	: SerializedObject(InSerializedObject)
+	, Offset(0)
+{
+	SetIsSaving(true);
+
+	for (int32 ObjIndex = 0; ObjIndex < SerializedObject.ReferencedObjects.Num(); ++ObjIndex)
+	{
+		ObjectMap.Add(SerializedObject.ReferencedObjects[ObjIndex].Get(), ObjIndex);
+	}
+
+	for (int32 NameIndex = 0; NameIndex < SerializedObject.ReferencedNames.Num(); ++NameIndex)
+	{
+		NameMap.Add(SerializedObject.ReferencedNames[NameIndex], NameIndex);
+	}
+}
+
+void FSerializedObjectDataWriter::Serialize(void* SerData, int64 Num)
+{
+	if (Num)
+	{
+		int64 DataIndex = Offset;
+		SerializedObject.SerializedData.Write(SerData, Offset, Num);
+		Offset += Num;
+
+		// Track this property offset in the serialized data
+		const FArchiveSerializedPropertyChain* PropertyChain = GetSerializedPropertyChain();
+		if (PropertyChain && PropertyChain->GetNumProperties() > 0)
+		{
+			const FName SerializedTaggedPropertyKey = CachedSerializedTaggedPropertyKey.SyncCache(PropertyChain);
+			FTransactionSerializedProperty& SerializedTaggedProperty = SerializedObject.SerializedProperties.FindOrAdd(SerializedTaggedPropertyKey);
+			SerializedTaggedProperty.AppendSerializedData(DataIndex, Num);
+		}
+	}
+}
+
+FArchive& FSerializedObjectDataWriter::operator<<(FName& N)
+{
+	int32 NameIndex = INDEX_NONE;
+	const int32* NameIndexPtr = NameMap.Find(N);
+	if (NameIndexPtr)
+	{
+		NameIndex = *NameIndexPtr;
+	}
+	else
+	{
+		NameIndex = SerializedObject.ReferencedNames.Add(N);
+		NameMap.Add(N, NameIndex);
+	}
+
+	// Track this name index in the serialized data
+	{
+		const FArchiveSerializedPropertyChain* PropertyChain = GetSerializedPropertyChain();
+		const FName SerializedTaggedPropertyKey = CachedSerializedTaggedPropertyKey.SyncCache(PropertyChain);
+		SerializedObject.SerializedNameIndices.Add(SerializedTaggedPropertyKey).Indices.AddUnique(NameIndex);
+	}
+
+	return *this << NameIndex;
+}
+
+FArchive& FSerializedObjectDataWriter::operator<<(UObject*& Res)
+{
+	int32 ObjectIndex = INDEX_NONE;
+	const int32* ObjIndexPtr = ObjectMap.Find(Res);
+	if (ObjIndexPtr)
+	{
+		ObjectIndex = *ObjIndexPtr;
+	}
+	else if (Res)
+	{
+		ObjectIndex = SerializedObject.ReferencedObjects.Add(FTransactionPersistentObjectRef(Res));
+		ObjectMap.Add(Res, ObjectIndex);
+	}
+
+	// Track this object offset in the serialized data
+	{
+		const FArchiveSerializedPropertyChain* PropertyChain = GetSerializedPropertyChain();
+		const FName SerializedTaggedPropertyKey = CachedSerializedTaggedPropertyKey.SyncCache(PropertyChain);
+		SerializedObject.SerializedObjectIndices.Add(SerializedTaggedPropertyKey).Indices.AddUnique(ObjectIndex);
+	}
+
+	return *this << ObjectIndex;
+}
+
+FName FSerializedObjectDataWriter::FCachedPropertyKey::SyncCache(const FArchiveSerializedPropertyChain* InPropertyChain)
+{
+	if (InPropertyChain)
+	{
+		const uint32 CurrentUpdateCount = InPropertyChain->GetUpdateCount();
+		if (CurrentUpdateCount != LastUpdateCount)
+		{
+			CachedKey = InPropertyChain->GetNumProperties() > 0 ? FTransactionSerializedProperty::BuildSerializedPropertyKey(*InPropertyChain) : FName();
+			LastUpdateCount = CurrentUpdateCount;
+		}
+	}
+	else
+	{
+		CachedKey = FName();
+		LastUpdateCount = 0;
+	}
+
+	return CachedKey;
+}
+
+
+namespace DiffUtil
+{
+
+void GenerateObjectDiff(const FTransactionSerializedObject& OldSerializedObject, const FTransactionSerializedObject& NewSerializedObject, FTransactionObjectDeltaChange& OutDeltaChange, const bool bFullDiff)
+{
+	if (bFullDiff)
+	{
+		OutDeltaChange.bHasNameChange |= OldSerializedObject.ObjectName != NewSerializedObject.ObjectName;
+		OutDeltaChange.bHasOuterChange |= OldSerializedObject.ObjectOuterPathName != NewSerializedObject.ObjectOuterPathName;
+		OutDeltaChange.bHasExternalPackageChange |= OldSerializedObject.ObjectExternalPackageName != NewSerializedObject.ObjectExternalPackageName;
+		OutDeltaChange.bHasPendingKillChange |= OldSerializedObject.bIsPendingKill != NewSerializedObject.bIsPendingKill;
+
+		if (!Internal::AreObjectPointersIdentical(OldSerializedObject, NewSerializedObject, NAME_None))
+		{
+			OutDeltaChange.bHasNonPropertyChanges = true;
+		}
+
+		if (!Internal::AreNamesIdentical(OldSerializedObject, NewSerializedObject, NAME_None))
+		{
+			OutDeltaChange.bHasNonPropertyChanges = true;
+		}
+	}
+
+	if (OldSerializedObject.SerializedProperties.Num() > 0 || NewSerializedObject.SerializedProperties.Num() > 0)
+	{
+		int64 StartOfOldPropertyBlock = INT64_MAX;
+		int64 StartOfNewPropertyBlock = INT64_MAX;
+		int64 EndOfOldPropertyBlock = -1;
+		int64 EndOfNewPropertyBlock = -1;
+
+		for (const TPair<FName, FTransactionSerializedProperty>& NewNamePropertyPair : NewSerializedObject.SerializedProperties)
+		{
+			const FTransactionSerializedProperty* OldSerializedProperty = OldSerializedObject.SerializedProperties.Find(NewNamePropertyPair.Key);
+			if (!OldSerializedProperty)
+			{
+				if (bFullDiff)
+				{
+					// Missing property, assume that the property changed
+					OutDeltaChange.ChangedProperties.AddUnique(NewNamePropertyPair.Key);
+				}
+				continue;
+			}
+
+			// Update the tracking for the start/end of the property block within the serialized data
+			StartOfOldPropertyBlock = FMath::Min(StartOfOldPropertyBlock, OldSerializedProperty->DataOffset);
+			StartOfNewPropertyBlock = FMath::Min(StartOfNewPropertyBlock, NewNamePropertyPair.Value.DataOffset);
+			EndOfOldPropertyBlock = FMath::Max(EndOfOldPropertyBlock, OldSerializedProperty->DataOffset + OldSerializedProperty->DataSize);
+			EndOfNewPropertyBlock = FMath::Max(EndOfNewPropertyBlock, NewNamePropertyPair.Value.DataOffset + NewNamePropertyPair.Value.DataSize);
+
+			// Binary compare the serialized data to see if something has changed for this property
+			bool bIsPropertyIdentical = OldSerializedProperty->DataSize == NewNamePropertyPair.Value.DataSize;
+			if (bIsPropertyIdentical && NewNamePropertyPair.Value.DataSize > 0)
+			{
+				bIsPropertyIdentical = FMemory::Memcmp(OldSerializedObject.SerializedData.GetPtr(OldSerializedProperty->DataOffset), NewSerializedObject.SerializedData.GetPtr(NewNamePropertyPair.Value.DataOffset), NewNamePropertyPair.Value.DataSize) == 0;
+			}
+			if (bIsPropertyIdentical)
+			{
+				bIsPropertyIdentical = Internal::AreObjectPointersIdentical(OldSerializedObject, NewSerializedObject, NewNamePropertyPair.Key);
+			}
+			if (bIsPropertyIdentical)
+			{
+				bIsPropertyIdentical = Internal::AreNamesIdentical(OldSerializedObject, NewSerializedObject, NewNamePropertyPair.Key);
+			}
+
+			if (!bIsPropertyIdentical)
+			{
+				OutDeltaChange.ChangedProperties.AddUnique(NewNamePropertyPair.Key);
+			}
+		}
+
+		for (const TPair<FName, FTransactionSerializedProperty>& OldNamePropertyPair : OldSerializedObject.SerializedProperties)
+		{
+			const FTransactionSerializedProperty* NewSerializedProperty = NewSerializedObject.SerializedProperties.Find(OldNamePropertyPair.Key);
+			if (!NewSerializedProperty)
+			{
+				if (bFullDiff)
+				{
+					// Missing property, assume that the property changed
+					OutDeltaChange.ChangedProperties.AddUnique(OldNamePropertyPair.Key);
+				}
+				continue;
+			}
+		}
+
+		if (bFullDiff)
+		{
+			// Compare the data before the property block to see if something else in the object has changed
+			if (!OutDeltaChange.bHasNonPropertyChanges)
+			{
+				const int64 OldHeaderSize = FMath::Min(StartOfOldPropertyBlock, OldSerializedObject.SerializedData.Num());
+				const int64 CurrentHeaderSize = FMath::Min(StartOfNewPropertyBlock, NewSerializedObject.SerializedData.Num());
+
+				bool bIsHeaderIdentical = OldHeaderSize == CurrentHeaderSize;
+				if (bIsHeaderIdentical && CurrentHeaderSize > 0)
+				{
+					bIsHeaderIdentical = FMemory::Memcmp(OldSerializedObject.SerializedData.GetPtr(0), NewSerializedObject.SerializedData.GetPtr(0), CurrentHeaderSize) == 0;
+				}
+
+				if (!bIsHeaderIdentical)
+				{
+					OutDeltaChange.bHasNonPropertyChanges = true;
+				}
+			}
+
+			// Compare the data after the property block to see if something else in the object has changed
+			if (!OutDeltaChange.bHasNonPropertyChanges)
+			{
+				const int64 OldFooterSize = OldSerializedObject.SerializedData.Num() - FMath::Max<int64>(EndOfOldPropertyBlock, 0);
+				const int64 CurrentFooterSize = NewSerializedObject.SerializedData.Num() - FMath::Max<int64>(EndOfNewPropertyBlock, 0);
+
+				bool bIsFooterIdentical = OldFooterSize == CurrentFooterSize;
+				if (bIsFooterIdentical && CurrentFooterSize > 0)
+				{
+					bIsFooterIdentical = FMemory::Memcmp(OldSerializedObject.SerializedData.GetPtr(EndOfOldPropertyBlock), NewSerializedObject.SerializedData.GetPtr(EndOfNewPropertyBlock), CurrentFooterSize) == 0;
+				}
+
+				if (!bIsFooterIdentical)
+				{
+					OutDeltaChange.bHasNonPropertyChanges = true;
+				}
+			}
+		}
+	}
+	else if (bFullDiff)
+	{
+		// No properties, so just compare the whole blob
+		bool bIsBlobIdentical = OldSerializedObject.SerializedData.Num() == NewSerializedObject.SerializedData.Num();
+		if (bIsBlobIdentical && NewSerializedObject.SerializedData.Num() > 0)
+		{
+			bIsBlobIdentical = FMemory::Memcmp(OldSerializedObject.SerializedData.GetPtr(0), NewSerializedObject.SerializedData.GetPtr(0), NewSerializedObject.SerializedData.Num()) == 0;
+		}
+
+		if (!bIsBlobIdentical)
+		{
+			OutDeltaChange.bHasNonPropertyChanges = true;
+		}
+	}
+}
+
+namespace Internal
+{
+
+bool AreObjectPointersIdentical(const FTransactionSerializedObject& OldSerializedObject, const FTransactionSerializedObject& NewSerializedObject, const FName PropertyName)
+{
+	auto GetIndicesArray = [&PropertyName](const FTransactionSerializedObject& SerializedObject) -> const TArray<int32>&
+	{
+		if (const FTransactionSerializedIndices* SerializedObjectIndices = SerializedObject.SerializedObjectIndices.Find(PropertyName))
+		{
+			return SerializedObjectIndices->Indices;
+		}
+
+		static const TArray<int32> EmptyIndicesArray;
+		return EmptyIndicesArray;
+	};
+
+	const TArray<int32>& OldSerializedObjectIndices = GetIndicesArray(OldSerializedObject);
+	const TArray<int32>& NewSerializedObjectIndices = GetIndicesArray(NewSerializedObject);
+
+	bool bAreObjectPointersIdentical = OldSerializedObjectIndices.Num() == NewSerializedObjectIndices.Num();
+	for (int32 ObjIndex = 0; ObjIndex < OldSerializedObjectIndices.Num() && bAreObjectPointersIdentical; ++ObjIndex)
+	{
+		const FTransactionPersistentObjectRef& OldObjectRef = OldSerializedObject.ReferencedObjects.IsValidIndex(OldSerializedObjectIndices[ObjIndex]) ? OldSerializedObject.ReferencedObjects[OldSerializedObjectIndices[ObjIndex]] : FTransactionPersistentObjectRef();
+		const FTransactionPersistentObjectRef& NewObjectRef = NewSerializedObject.ReferencedObjects.IsValidIndex(NewSerializedObjectIndices[ObjIndex]) ? NewSerializedObject.ReferencedObjects[NewSerializedObjectIndices[ObjIndex]] : FTransactionPersistentObjectRef();
+		bAreObjectPointersIdentical = OldObjectRef == NewObjectRef;
+	}
+	return bAreObjectPointersIdentical;
+}
+
+bool AreNamesIdentical(const FTransactionSerializedObject& OldSerializedObject, const FTransactionSerializedObject& NewSerializedObject, const FName PropertyName)
+{
+	auto GetIndicesArray = [&PropertyName](const FTransactionSerializedObject& SerializedObject) -> const TArray<int32>&
+	{
+		if (const FTransactionSerializedIndices* SerializedNameIndices = SerializedObject.SerializedNameIndices.Find(PropertyName))
+		{
+			return SerializedNameIndices->Indices;
+		}
+
+		static const TArray<int32> EmptyIndicesArray;
+		return EmptyIndicesArray;
+	};
+
+	const TArray<int32>& OldSerializedNameIndices = GetIndicesArray(OldSerializedObject);
+	const TArray<int32>& NewSerializedNameIndices = GetIndicesArray(NewSerializedObject);
+
+	bool bAreNamesIdentical = OldSerializedNameIndices.Num() == NewSerializedNameIndices.Num();
+	for (int32 ObjIndex = 0; ObjIndex < OldSerializedNameIndices.Num() && bAreNamesIdentical; ++ObjIndex)
+	{
+		const FName& OldName = OldSerializedObject.ReferencedNames.IsValidIndex(OldSerializedNameIndices[ObjIndex]) ? OldSerializedObject.ReferencedNames[OldSerializedNameIndices[ObjIndex]] : FName();
+		const FName& NewName = NewSerializedObject.ReferencedNames.IsValidIndex(NewSerializedNameIndices[ObjIndex]) ? NewSerializedObject.ReferencedNames[NewSerializedNameIndices[ObjIndex]] : FName();
+		bAreNamesIdentical = OldName == NewName;
+	}
+	return bAreNamesIdentical;
+}
+
+} // namespace Internal
+
+} // namespace DiffUtil
+
+} // namespace UE::Transaction
diff --git a/Engine/Source/Runtime/Engine/Public/TransactionCommon.h b/Engine/Source/Runtime/Engine/Public/TransactionCommon.h
new file mode 100644
--- /dev/null
+++ b/Engine/Source/Runtime/Engine/Public/TransactionCommon.h
@@ -0,0 +1,382 @@
+// Copyright Epic Games, Inc. All Rights Reserved.
+
+#pragma once
+
+#include "CoreMinimal.h"
+#include "UObject/ObjectPtr.h"
+#include "UObject/WeakObjectPtr.h"
+#include "Serialization/ArchiveUObject.h"
+#include "Serialization/ArchiveSerializedPropertyChain.h"
+#include "TransactionCommon.generated.h"
+
+class FReferenceCollector;
+struct FTransactionObjectDeltaChange;
+
+/**
+	This type is necessary because the blueprint system is destroying and creating
+	CDOs at edit time (usually on compile, but also on load), but also stores user
+	entered data in the CDO. We "need"  changes to a CDO to persist across instances
+	because as we undo and redo we  need to apply changes to different instances of
+	the CDO - alternatively we could destroy and create the CDO as part of a transaction
+	(this alternative is the reason for the bunny ears around need).
+
+	DanO: My long term preference is for the editor to use a dynamic, mutable type
+	(rather than the CDO) to store editor data. The CDO can then be re-instanced (or not)
+	as runtime code requires.
+*/
+USTRUCT()
+struct ENGINE_API FTransactionPersistentObjectRef
+{
+	GENERATED_BODY()
+
+public:
+	FTransactionPersistentObjectRef() = default;
+	explicit FTransactionPersistentObjectRef(UObject* InObject);
+
+	bool Serialize(FArchive& Ar)
+	{
+		Ar << *this;
+		return true;
+	}
+
+	friend FArchive& operator<<(FArchive& Ar, FTransactionPersistentObjectRef& ReferencedObject)
+	{
+		Ar << (std::underlying_type_t<EReferenceType>&)ReferencedObject.ReferenceType;
+		Ar << ReferencedObject.RootObject;
+		Ar << ReferencedObject.SubObjectHierarchyIDs;
+		return Ar;
+	}
+
+	friend bool operator==(const FTransactionPersistentObjectRef& LHS, const FTransactionPersistentObjectRef& RHS)
+	{
+		return LHS.ReferenceType == RHS.ReferenceType
+			&& LHS.RootObject == RHS.RootObject
+			&& (LHS.ReferenceType != EReferenceType::SubObject || LHS.SubObjectHierarchyIDs == RHS.SubObjectHierarchyIDs);
+	}
+
+	friend bool operator!=(const FTransactionPersistentObjectRef& LHS, const FTransactionPersistentObjectRef& RHS)
+	{
+		return !(LHS == RHS);
+	}
+
+	friend uint32 GetTypeHash(const FTransactionPersistentObjectRef& InObjRef)
+	{
+		return HashCombine(GetTypeHash(InObjRef.ReferenceType), GetTypeHash(InObjRef.RootObject));
+	}
+
+	bool IsRootObjectReference() const
+	{
+		return ReferenceType == EReferenceType::RootObject;
+	}
+
+	bool IsSubObjectReference() const
+	{
+		return ReferenceType == EReferenceType::SubObject;
+	}
+
+	UObject* Get() const;
+	
+	void AddStructReferencedObjects(FReferenceCollector& Collector);
+
+private:
+	/** This enum represents all of the different special cases we are handling with this type */
+	enum class EReferenceType : uint8
+	{
+		Unknown,
+		RootObject,
+		SubObject,
+	};
+
+	/** The reference type we're handling */
+	EReferenceType ReferenceType = EReferenceType::Unknown;
+	/** Stores the object pointer when ReferenceType==RootObject, and the outermost pointer of the sub-object chain when when ReferenceType==SubObject */
+	TObjectPtr<UObject> RootObject = nullptr;
+	/** Stores the sub-object name chain when ReferenceType==SubObject */
+	TArray<FName, TInlineAllocator<4>> SubObjectHierarchyIDs;
+
+	/** Cached pointers corresponding to RootObject when ReferenceType==SubObject (@note cache needs testing on access as it may have become stale) */
+	mutable TWeakObjectPtr<UObject> CachedRootObject;
+	/** Cache of pointers corresponding to the items within SubObjectHierarchyIDs when ReferenceType==SubObject (@note cache needs testing on access as it may have become stale) */
+	mutable TArray<TWeakObjectPtr<UObject>, TInlineAllocator<4>> CachedSubObjectHierarchy;
+};
+
+template<>
+struct TStructOpsTypeTraits<FTransactionPersistentObjectRef> : public TStructOpsTypeTraitsBase2<FTransactionPersistentObjectRef>
+{
+	enum
+	{
+		WithIdenticalViaEquality = true,
+		WithAddStructReferencedObjects = true,
+		WithSerializer = true,
+	};
+};
+
+USTRUCT()
+struct ENGINE_API FTransactionSerializedProperty
+{
+	GENERATED_BODY()
+
+public:
+	static FName BuildSerializedPropertyKey(const FArchiveSerializedPropertyChain& InPropertyChain);
+
+	void AppendSerializedData(const int64 InOffset, const int64 InSize);
+
+	/** Offset to the start of this property within the serialized object */
+	UPROPERTY()
+	int64 DataOffset = INDEX_NONE;
+
+	/** Size (in bytes) of this property within the serialized object */
+	UPROPERTY()
+	int64 DataSize = 0;
+};
+
+USTRUCT()
+struct ENGINE_API FTransactionSerializedObjectData
+{
+	GENERATED_BODY()
+
+public:
+	bool Serialize(FArchive& Ar)
+	{
+		Ar << *this;
+		return true;
+	}
+
+	friend FArchive& operator<<(FArchive& Ar, FTransactionSerializedObjectData& SerializedData)
+	{
+		Ar << SerializedData.Data;
+		return Ar;
+	}
+
+	friend bool operator==(const FTransactionSerializedObjectData& LHS, const FTransactionSerializedObjectData& RHS)
+	{
+		return LHS.Data == RHS.Data;
+	}
+
+	friend bool operator!=(const FTransactionSerializedObjectData& LHS, const FTransactionSerializedObjectData& RHS)
+	{
+		return !(LHS == RHS);
+	}
+
+	void Read(void* Dest, int64 Offset, int64 Num) const;
+	void Write(const void* Src, int64 Offset, int64 Num);
+	
+	const void* GetPtr(int64 Offset) const
+	{
+		return &Data[Offset];
+	}
+
+	void Reset()
+	{
+		Data.Reset();
+	}
+
+	int64 Num() const
+	{
+		return Data.Num();
+	}
+
+private:
+	TArray64<uint8> Data;
+};
+
+template<>
+struct TStructOpsTypeTraits<FTransactionSerializedObjectData> : public TStructOpsTypeTraitsBase2<FTransactionSerializedObjectData>
+{
+	enum
+	{
+		WithIdenticalViaEquality = true,
+		WithSerializer = true,
+	};
+};
+
+USTRUCT()
+struct FTransactionSerializedIndices
+{
+	GENERATED_BODY()
+
+public:
+	UPROPERTY()
+	TArray<int32> Indices;
+};
+
+USTRUCT()
+struct FTransactionSerializedObject
+{
+	GENERATED_BODY()
+
+public:
+	void SetObject(const UObject* InObject)
+	{
+		ObjectPackageName = InObject->GetPackage()->GetFName();
+		ObjectName = InObject->GetFName();
+		ObjectPathName = *InObject->GetPathName();
+		ObjectOuterPathName = InObject->GetOuter() ? FName(*InObject->GetOuter()->GetPathName()) : FName();
+		ObjectExternalPackageName = InObject->GetExternalPackage() ? InObject->GetExternalPackage()->GetFName() : FName();
+		ObjectClassPathName = FName(*InObject->GetClass()->GetPathName());
+		bIsPendingKill = IsValid(InObject);
+	}
+
+	void Reset()
+	{
+		ObjectPackageName = FName();
+		ObjectName = FName();
+		ObjectPathName = FName();
+		ObjectOuterPathName = FName();
+		ObjectExternalPackageName = FName();
+		ObjectClassPathName = FName();
+		bIsPendingKill = false;
+		SerializedData.Reset();
+		ReferencedObjects.Reset();
+		ReferencedNames.Reset();
+		SerializedProperties.Reset();
+		SerializedObjectIndices.Reset();
+		SerializedNameIndices.Reset();
+	}
+
+	void Swap(FTransactionSerializedObject& Other)
+	{
+		Exchange(ObjectPackageName, Other.ObjectPackageName);
+		Exchange(ObjectName, Other.ObjectName);
+		Exchange(ObjectPathName, Other.ObjectPathName);
+		Exchange(ObjectOuterPathName, Other.ObjectOuterPathName);
+		Exchange(ObjectExternalPackageName, Other.ObjectExternalPackageName);
+		Exchange(ObjectClassPathName, Other.ObjectClassPathName);
+		Exchange(bIsPendingKill, Other.bIsPendingKill);
+		Exchange(SerializedData, Other.SerializedData);
+		Exchange(ReferencedObjects, Other.ReferencedObjects);
+		Exchange(ReferencedNames, Other.ReferencedNames);
+		Exchange(SerializedProperties, Other.SerializedProperties);
+		Exchange(SerializedObjectIndices, Other.SerializedObjectIndices);
+		Exchange(SerializedNameIndices, Other.SerializedNameIndices);
+	}
+
+	/** The package name of the object when it was serialized, can be dictated either by outer chain or external package */
+	UPROPERTY()
+	FName ObjectPackageName;
+	
+	/** The name of the object when it was serialized */
+	UPROPERTY()
+	FName ObjectName;
+	
+	/** The path name of the object when it was serialized */
+	UPROPERTY()
+	FName ObjectPathName;
+	
+	/** The outer path name of the object when it was serialized */
+	UPROPERTY()
+	FName ObjectOuterPathName;
+	
+	/** The external package name of the object when it was serialized, if any */
+	UPROPERTY()
+	FName ObjectExternalPackageName;
+	
+	/** The path name of the object's class. */
+	UPROPERTY()
+	FName ObjectClassPathName;
+	
+	/** The pending kill state of the object when it was serialized */
+	UPROPERTY()
+	bool bIsPendingKill = false;
+
+	/** The data stream used to serialize/deserialize record */
+	UPROPERTY()
+	FTransactionSerializedObjectData SerializedData;
+
+	/** External objects referenced in the transaction */
+	UPROPERTY()
+	TArray<FTransactionPersistentObjectRef> ReferencedObjects;
+	
+	/** FNames referenced in the object record */
+	UPROPERTY()
+	TArray<FName> ReferencedNames;
+	
+	/** Information about the properties that were serialized within this object */
+	UPROPERTY()
+	TMap<FName, FTransactionSerializedProperty> SerializedProperties;
+	
+	/** Information about the object pointer offsets that were serialized within this object (this maps the property name (or None if there was no property) to the ReferencedObjects indices of the property) */
+	UPROPERTY()
+	TMap<FName, FTransactionSerializedIndices> SerializedObjectIndices;
+	
+	/** Information about the name offsets that were serialized within this object (this maps the property name to the ReferencedNames index of the property) */
+	UPROPERTY()
+	TMap<FName, FTransactionSerializedIndices> SerializedNameIndices;
+};
+
+namespace UE::Transaction
+{
+
+/** Transfers data from an array. */
+class ENGINE_API FSerializedObjectDataReader : public FArchiveUObject
+{
+public:
+	FSerializedObjectDataReader(const FTransactionSerializedObject& InSerializedObject);
+
+	virtual int64 Tell() override { return Offset; }
+	virtual void Seek(int64 InPos) override { Offset = InPos; }
+	virtual int64 TotalSize() override { return SerializedObject.SerializedData.Num(); }
+
+protected:
+	virtual void Serialize(void* SerData, int64 Num) override;
+
+	using FArchiveUObject::operator<<;
+	virtual FArchive& operator<<(class FName& N) override;
+	virtual FArchive& operator<<(class UObject*& Res) override;
+
+	const FTransactionSerializedObject& SerializedObject;
+	int64 Offset;
+};
+
+/**
+ * Transfers data to an array.
+ */
+class ENGINE_API FSerializedObjectDataWriter : public FArchiveUObject
+{
+public:
+	FSerializedObjectDataWriter(FTransactionSerializedObject& InSerializedObject);
+
+	virtual int64 Tell() override { return Offset; }
+	virtual void Seek(int64 InPos) override { checkSlow(Offset <= SerializedObject.SerializedData.Num()); Offset = InPos; }
+	virtual int64 TotalSize() override { return SerializedObject.SerializedData.Num(); }
+
+protected:
+	struct FCachedPropertyKey
+	{
+	public:
+		FName SyncCache(const FArchiveSerializedPropertyChain* InPropertyChain);
+
+	private:
+		FName CachedKey;
+		uint32 LastUpdateCount = 0;
+	};
+
+	virtual void Serialize(void* SerData, int64 Num) override;
+
+	using FArchiveUObject::operator<<;
+	virtual FArchive& operator<<(class FName& N) override;
+	virtual FArchive& operator<<(class UObject*& Res) override;
+
+	FTransactionSerializedObject& SerializedObject;
+	TMap<UObject*, int32> ObjectMap;
+	TMap<FName, int32> NameMap;
+	FCachedPropertyKey CachedSerializedTaggedPropertyKey;
+	int64 Offset;
+};
+
+namespace DiffUtil
+{
+
+ENGINE_API void GenerateObjectDiff(const FTransactionSerializedObject& OldSerializedObject, const FTransactionSerializedObject& NewSerializedObject, FTransactionObjectDeltaChange& OutDeltaChange, const bool bFullDiff = true);
+
+namespace Internal
+{
+
+ENGINE_API bool AreObjectPointersIdentical(const FTransactionSerializedObject& OldSerializedObject, const FTransactionSerializedObject& NewSerializedObject, const FName PropertyName);
+ENGINE_API bool AreNamesIdentical(const FTransactionSerializedObject& OldSerializedObject, const FTransactionSerializedObject& NewSerializedObject, const FName PropertyName);
+
+} // namespace Internal
+
+} // namespace DiffUtil
+
+} // namespace UE::Transaction
