- Set from String not Float to fix incorrect values in other types of cvars

#rnx

diff --git a/Engine/Source/Runtime/Engine/Private/Scalability.cpp b/Engine/Source/Runtime/Engine/Private/Scalability.cpp
--- a/Engine/Source/Runtime/Engine/Private/Scalability.cpp
+++ b/Engine/Source/Runtime/Engine/Private/Scalability.cpp
@@ -386,11 +386,9 @@ static void SetGroupQualityLevel(const TCHAR* InGroupName, int32 InQualityLevel,
 					FString CVarName, CVarValue;
 					if (CVarString.Split(TEXT("="), &CVarName, &CVarValue))
 					{
-						IConsoleVariable* CVar = IConsoleManager::Get().FindConsoleVariable(*CVarName);
-						if (CVar)
+						if (IConsoleVariable* CVar = IConsoleManager::Get().FindConsoleVariable(*CVarName))
 						{
-							float Value = FCString::Atof(*CVarValue);
-							CVar->SetWithCurrentPriority(Value);
+							CVar->SetWithCurrentPriority(*CVarValue);
 						}
 					}
 				}
