Comment out an ensure until I fix existing issues it found.

#rnx

diff --git a/Engine/Source/Runtime/Engine/Private/ActorReplication.cpp b/Engine/Source/Runtime/Engine/Private/ActorReplication.cpp
--- a/Engine/Source/Runtime/Engine/Private/ActorReplication.cpp
+++ b/Engine/Source/Runtime/Engine/Private/ActorReplication.cpp
@@ -633,7 +633,7 @@ void AActor::AddActorComponentReplicatedSubObject(UActorComponent* OwnerComponen
 {
 	check(IsValid(OwnerComponent));
 	check(IsValid(SubObject));
-	ensureMsgf(OwnerComponent->GetIsReplicated(), TEXT("Only components with replication enabled can register subobjects. %s::%s has replication disabled."), *GetName(), *SubObject->GetName());
+	//ensureMsgf(OwnerComponent->GetIsReplicated(), TEXT("Only components with replication enabled can register subobjects. %s::%s has replication disabled."), *GetName(), *SubObject->GetName());
 	ensureMsgf(NetCondition != COND_Custom, TEXT("Custom netconditions do not work with SubObjects. %s::%s will not be replicated."), *GetName(), *SubObject->GetName());
 
 	FReplicatedComponentInfo* ComponentInfo = ReplicatedComponentsInfo.FindByKey(OwnerComponent);
