[KillPendingKill] Use net driver as outer for actor channels to provide extra debug info in path names.

diff --git a/Engine/Source/Runtime/Engine/Private/NetDriver.cpp b/Engine/Source/Runtime/Engine/Private/NetDriver.cpp
--- a/Engine/Source/Runtime/Engine/Private/NetDriver.cpp
+++ b/Engine/Source/Runtime/Engine/Private/NetDriver.cpp
@@ -5580,7 +5580,7 @@ UChannel* UNetDriver::GetOrCreateChannelByName(const FName& ChName)
 UChannel* UNetDriver::InternalCreateChannelByName(const FName& ChName)
 {
 	LLM_SCOPE_BYTAG(NetChannel);
-	return NewObject<UChannel>(GetTransientPackage(), ChannelDefinitionMap[ChName].ChannelClass);
+	return NewObject<UChannel>(this, ChannelDefinitionMap[ChName].ChannelClass);
 }
 
 void UNetDriver::ReleaseToChannelPool(UChannel* Channel)
