Removing IsRecording() check from RecordReplay, since it will already stop the existing replay properly

diff --git a/Engine/Source/Runtime/Engine/Private/ReplaySubsystem.cpp b/Engine/Source/Runtime/Engine/Private/ReplaySubsystem.cpp
--- a/Engine/Source/Runtime/Engine/Private/ReplaySubsystem.cpp
+++ b/Engine/Source/Runtime/Engine/Private/ReplaySubsystem.cpp
@@ -121,12 +121,6 @@ void UReplaySubsystem::RecordReplay(const FString& Name, const FString& Friendly
 		return;
 	}
 
-	if (IsRecording())
-	{
-		UE_LOG(LogDemo, Warning, TEXT("UReplaySubsystem::RecordReplay: A replay is already being recorded, cannot begin recording another one."));
-		return;
-	}
-
 	FURL DemoURL;
 	FString DemoName = Name;
 
