Chaos - disable inertia conditioning by default when using the non-linear joint solver in RBAN

diff --git a/Engine/Source/Runtime/Engine/Private/PhysicsEngine/ImmediatePhysics/ImmediatePhysicsChaos/ImmediatePhysicsSimulation_Chaos.cpp b/Engine/Source/Runtime/Engine/Private/PhysicsEngine/ImmediatePhysics/ImmediatePhysicsChaos/ImmediatePhysicsSimulation_Chaos.cpp
--- a/Engine/Source/Runtime/Engine/Private/PhysicsEngine/ImmediatePhysics/ImmediatePhysicsChaos/ImmediatePhysicsSimulation_Chaos.cpp
+++ b/Engine/Source/Runtime/Engine/Private/PhysicsEngine/ImmediatePhysics/ImmediatePhysicsChaos/ImmediatePhysicsSimulation_Chaos.cpp
@@ -29,6 +29,7 @@
 //PRAGMA_DISABLE_OPTIMIZATION
 
 DECLARE_CYCLE_STAT(TEXT("FSimulation::Simulate_Chaos"), STAT_ImmediateSimulate_Chaos, STATGROUP_ImmediatePhysics);
+DECLARE_CYCLE_STAT(TEXT("FSimulation::Simulate_Chaos::InertiaConditioning"), STAT_ImmediateSimulate_Chaos_InertiaConditioning, STATGROUP_ImmediatePhysics);
 
 //////////////////////////////////////////////////////////////////////////
 // @todo(ccaulfield): remove when finished
@@ -145,8 +146,11 @@ FAutoConsoleVariableRef CVarChaosImmPhysJointAngularDriveDamping(TEXT("p.Chaos.I
 FAutoConsoleVariableRef CVarChaosImmPhysJointMinParentMassRatio(TEXT("p.Chaos.ImmPhys.Joint.MinParentMassRatio"), ChaosImmediate_Joint_MinParentMassRatio, TEXT("6Dof joint MinParentMassRatio (if > 0)"));
 FAutoConsoleVariableRef CVarChaosImmPhysJointMaxInertiaRatio(TEXT("p.Chaos.ImmPhys.Joint.MaxInertiaRatio"), ChaosImmediate_Joint_MaxInertiaRatio, TEXT("6Dof joint MaxInertiaRatio (if > 0)"));
 
-bool bChaosImmediate_InertiaConditioningEnabled = true;
-FAutoConsoleVariableRef  CVarChaosImmPhysParticleInertiaConditioningEnabled(TEXT("p.Chaos.ImmPhys.InertiaConditioning.Enabled"), bChaosImmediate_InertiaConditioningEnabled, TEXT("Enable/Disable constraint stabilization through inertia conditioning"));
+bool bChaosImmediate_LinearInertiaConditioningEnabled = true;
+FAutoConsoleVariableRef  CVarChaosImmPhysParticleLinearInertiaConditioningEnabled(TEXT("p.Chaos.ImmPhys.InertiaConditioning.LinearEnabled"), bChaosImmediate_LinearInertiaConditioningEnabled, TEXT("Enable/Disable constraint stabilization through inertia conditioning when using the linear joint solver"));
+
+bool bChaosImmediate_NonLinearInertiaConditioningEnabled = false;
+FAutoConsoleVariableRef  CVarChaosImmPhysParticleNonLinearInertiaConditioningEnabled(TEXT("p.Chaos.ImmPhys.InertiaConditioning.NonlinearEnabled"), bChaosImmediate_NonLinearInertiaConditioningEnabled, TEXT("Enable/Disable constraint stabilization through inertia conditioning when using the non-linear joint solver"));
 
 Chaos::FRealSingle ChaosImmediate_InertiaConditioningDistance = 10;
 FAutoConsoleVariableRef  CVarChaosImmPhysParticleInertiaConditioningDistance(TEXT("p.Chaos.ImmPhys.InertiaConditioning.Distance"), ChaosImmediate_InertiaConditioningDistance, TEXT("An input to inertia conditioning system. The joint distance error which needs to be stable (generate a low rotation)."));
@@ -987,11 +991,17 @@ namespace ImmediatePhysics_Chaos
 #endif
 	}
 
+	// @todo(chaos): Move implementation to MinEvolution
 	void FSimulation::UpdateInertiaConditioning(const FVector& Gravity)
 	{
 		using namespace Chaos;
 
-		if (bChaosImmediate_InertiaConditioningEnabled)
+		SCOPE_CYCLE_COUNTER(STAT_ImmediateSimulate_Chaos_InertiaConditioning);
+
+		// Linear and Non-linear joint solvers have a separate cvar to enable inertia conditioning
+		const bool bUseInertiaConditioning = Implementation->Joints.GetSettings().bUseLinearSolver ? bChaosImmediate_LinearInertiaConditioningEnabled : bChaosImmediate_NonLinearInertiaConditioningEnabled;
+
+		if (bUseInertiaConditioning)
 		{
 			// The maximum contribution to error correction from rotation
 			const FRealSingle MaxRotationRatio = ChaosImmediate_InertiaConditioningRotationRatio;
@@ -1003,6 +1013,7 @@ namespace ImmediatePhysics_Chaos
 			// A limit on the relative sizes of the inertia components (inverse)
 			const FRealSingle MaxInvInertiaComponentRatio = ChaosImmediate_MaxInvInertiaComponentRatio;
 
+			// @todo(chaos): keep track of changes to the flags so we don't need to loop every tick (maybe check it in Integrate)
 			for (FActorHandle* Actor : Implementation->ActorHandles)
 			{
 				FGenericParticleHandle Particle = Actor->GetParticle();
