[DebugDrawComponent] replaced pure virtual by default implementation returning a null proxy (similar to PrimitiveComponent)
It is up to derived classes to provide a valid implementation depending on their own preprocessor defines.
#rnx

diff --git a/Engine/Source/Runtime/Engine/Classes/Debug/DebugDrawComponent.h b/Engine/Source/Runtime/Engine/Classes/Debug/DebugDrawComponent.h
--- a/Engine/Source/Runtime/Engine/Classes/Debug/DebugDrawComponent.h
+++ b/Engine/Source/Runtime/Engine/Classes/Debug/DebugDrawComponent.h
@@ -55,9 +55,8 @@ protected:
 	virtual FPrimitiveSceneProxy* CreateSceneProxy() override final;
 
 	/* Method that derived class should override to create the scene proxy and customize a custom delegate helper (if any) */
-	virtual FDebugRenderSceneProxy* CreateDebugSceneProxy() PURE_VIRTUAL(UDebugDrawComponent::CreateSceneProxyInternal, return nullptr;);
+	virtual FDebugRenderSceneProxy* CreateDebugSceneProxy() { return nullptr; }
 
-#if UE_ENABLE_DEBUG_DRAWING
 	virtual void CreateRenderState_Concurrent(FRegisterComponentContext* Context) override;
 	virtual void DestroyRenderState_Concurrent() override;
 
@@ -65,5 +64,4 @@ protected:
 	virtual FDebugDrawDelegateHelper& GetDebugDrawDelegateHelper() { return DebugDrawDelegateHelper; }
 
 	FDebugDrawDelegateHelper DebugDrawDelegateHelper;
-#endif
 };
\ No newline at end of file
diff --git a/Engine/Source/Runtime/Engine/Private/Debug/DebugDrawComponent.cpp b/Engine/Source/Runtime/Engine/Private/Debug/DebugDrawComponent.cpp
--- a/Engine/Source/Runtime/Engine/Private/Debug/DebugDrawComponent.cpp
+++ b/Engine/Source/Runtime/Engine/Private/Debug/DebugDrawComponent.cpp
@@ -5,18 +5,15 @@
 FPrimitiveSceneProxy* UDebugDrawComponent::CreateSceneProxy()
 {
 	FDebugRenderSceneProxy* Proxy = CreateDebugSceneProxy();
-#if UE_ENABLE_DEBUG_DRAWING
   	if (Proxy != nullptr)
 	{
 		GetDebugDrawDelegateHelper().InitDelegateHelper(Proxy);
 	}
 
 	GetDebugDrawDelegateHelper().ProcessDeferredRegister();
-#endif
 	return Proxy;
 }
 
-#if UE_ENABLE_DEBUG_DRAWING
   void UDebugDrawComponent::CreateRenderState_Concurrent(FRegisterComponentContext* Context)
 {
 	Super::CreateRenderState_Concurrent(Context);
@@ -28,4 +25,3 @@ void UDebugDrawComponent::DestroyRenderState_Concurrent()
 	GetDebugDrawDelegateHelper().UnregisterDebugDrawDelegate();
 	Super::DestroyRenderState_Concurrent();
 }
-#endif
