Added random data to handshake packets, in attempt to work around potential ISP packet-filtering issues.

diff --git a/Engine/Plugins/Online/OnlineSubsystemUtils/Source/OnlineSubsystemUtils/Private/IpNetDriver.cpp b/Engine/Plugins/Online/OnlineSubsystemUtils/Source/OnlineSubsystemUtils/Private/IpNetDriver.cpp
--- a/Engine/Plugins/Online/OnlineSubsystemUtils/Source/OnlineSubsystemUtils/Private/IpNetDriver.cpp
+++ b/Engine/Plugins/Online/OnlineSubsystemUtils/Source/OnlineSubsystemUtils/Private/IpNetDriver.cpp
@@ -30,10 +30,6 @@ Notes:
 #include <limits>
 #include "NetAddressResolution.h"
 
-/** For backwards compatibility with the engine stateless connect code */
-#ifndef STATELESSCONNECT_HAS_RANDOM_SEQUENCE
-	#define STATELESSCONNECT_HAS_RANDOM_SEQUENCE 0
-#endif
 
 /*-----------------------------------------------------------------------------
 	Declarations.
@@ -1418,7 +1414,6 @@ UNetConnection* UIpNetDriver::ProcessConnectionlessPacket(FReceivedPacketView& P
 
 			ReturnVal->InitRemoteConnection(this, SocketPrivate.Get(), World ? World->URL : FURL(), *Address, USOCK_Open);
 
-#if STATELESSCONNECT_HAS_RANDOM_SEQUENCE
 			// Set the initial packet sequence from the handshake data
 			if (StatelessConnect.IsValid())
 			{
@@ -1429,7 +1424,6 @@ UNetConnection* UIpNetDriver::ProcessConnectionlessPacket(FReceivedPacketView& P
 
 				ReturnVal->InitSequence(ClientSequence, ServerSequence);
 			}
-#endif
 
 			if (ReturnVal->Handler.IsValid())
 			{
diff --git a/Engine/Source/Runtime/Engine/Private/PacketHandlers/StatelessConnectHandlerComponent.cpp b/Engine/Source/Runtime/Engine/Private/PacketHandlers/StatelessConnectHandlerComponent.cpp
--- a/Engine/Source/Runtime/Engine/Private/PacketHandlers/StatelessConnectHandlerComponent.cpp
+++ b/Engine/Source/Runtime/Engine/Private/PacketHandlers/StatelessConnectHandlerComponent.cpp
@@ -9,23 +9,9 @@
 #include "Net/Core/Misc/PacketAudit.h"
 
 
-
-
 DEFINE_LOG_CATEGORY(LogHandshake);
 
 
-
-// @todo #JohnB: It is important that banning functionality gets implemented here. This is the earliest/best place,
-//					to reject connections (however, format of 'Address', is defined by net driver - making this tricky, e.g. Steam)
-
-// @todo #JohnB: Do profiling of the handshake check, plus the entire PacketHandler 'Incoming'/'Outgoing' chain,
-//					when under a DoS attack.
-
-// @todo #JohnB: Consider adding an increasing-cost challenge at some stage. Not strictly necessary, but may be nice to have.
-
-// @todo #JohnB: The handshake restart code, will need to factor in IP/address bans.
-
-
 /**
  * Purpose:
  *
@@ -51,19 +37,19 @@ DEFINE_LOG_CATEGORY(LogHandshake);
  * and the server responding with a unique 'Cookie' value, which the client has to respond with.
  *
  * Client - Initial Connect:
- * [?:MagicHeader][HandshakeBit][RestartHandshakeBit][SecretIdBit][28:PacketSizeFiller][AlignPad]
+ * [?:MagicHeader][HandshakeBit][RestartHandshakeBit][8:MinVersion][8:CurVersion][8:HandshakePacketType][8:SentPacketCount][SecretIdBit][28:PacketSizeFiller][AlignPad][?:RandomData]
  *													--->
  *															Server - Stateless Handshake Challenge:
- *															[?:MagicHeader][HandshakeBit][RestartHandshakeBit][SecretIdBit][8:Timestamp][20:Cookie][AlignPad]
+ *															[?:MagicHeader][HandshakeBit][RestartHandshakeBit][8:MinVersion][8:CurVersion][8:HandshakePacketType][8:SentPacketCount][SecretIdBit][8:Timestamp][20:Cookie][AlignPad][?:RandomData]
  *													<---
  * Client - Stateless Challenge Response:
- * [?:MagicHeader][HandshakeBit][RestartHandshakeBit][SecretIdBit][8:Timestamp][20:Cookie][AlignPad]
+ * [?:MagicHeader][HandshakeBit][RestartHandshakeBit][8:MinVersion][8:CurVersion][8:HandshakePacketType][8:SentPacketCount][SecretIdBit][8:Timestamp][20:Cookie][AlignPad][?:RandomData]
  *													--->
  *															Server:
  *															Ignore, or create UNetConnection.
  *
  *															Server - Stateless Handshake Ack
- *															[?:MagicHeader][HandshakeBit][RestartHandshakeBit][SecretIdBit][8:Timestamp][20:Cookie][AlignPad]
+ *															[?:MagicHeader][HandshakeBit][RestartHandshakeBit][8:MinVersion][8:CurVersion][8:HandshakePacketType][8:SentPacketCount][SecretIdBit][8:Timestamp][20:Cookie][AlignPad][?:RandomData]
  *													<---
  * Client:
  * Handshake Complete.
@@ -75,14 +61,14 @@ DEFINE_LOG_CATEGORY(LogHandshake);
  * so the protocol has been crafted so the server sends only a minimal (1 byte) response, to minimize DRDoS reflection amplification.
  *
  *															Server - Restart Handshake Request:
- *															[?:MagicHeader][HandshakeBit][RestartHandshakeBit][AlignPad]
+ *															[?:MagicHeader][HandshakeBit][RestartHandshakeBit][8:MinVersion][8:CurVersion][8:HandshakePacketType][8:SentPacketCount][AlignPad][?:RandomData]
  *													<--
  * Client -  Initial Connect (as above)
  *													-->
  *															Server -  Stateless Handshake Challenge (as above)
  *													<--
  * Client - Stateless Challenge Response + Original Cookie
- * [?:MagicHeader][HandshakeBit][RestartHandshakeBit][SecretIdBit][8:Timestamp][20:Cookie][20:OriginalCookie][AlignPad]
+ * [?:MagicHeader][HandshakeBit][RestartHandshakeBit][8:MinVersion][8:CurVersion][8:HandshakePacketType][8:SentPacketCount][SecretIdBit][8:Timestamp][20:Cookie][20:OriginalCookie][AlignPad][?:RandomData]
  *													-->
  *															Server:
  *															Ignore, or restore UNetConnection.
@@ -98,9 +84,14 @@ DEFINE_LOG_CATEGORY(LogHandshake);
  *	- HandshakeBit:			Bit signifying whether a packet is a handshake packet. Applied to all game packets.
  *	- SecretIdBit:			For handshake packets, specifies which HandshakeSecret array was used to generate Cookie.
  *	- RestartHandshakeBit:  Sent by the server when it detects normal game traffic from an unknown IP/port combination.
+ *	- MinVersion:			The minimum handshake protocol version supported by the remote side
+ *	- CurVersion:			The currently active protocol version used by the remote side (determines received packet format)
+ *	- HandshakePacketType:	Number indicating the type of handshake packet, based on EHandshakePacketType
+ *	- SentPacketCount:		The number of handshake packets sent - for packet-analysis/debugging purposes
  *	- Timestamp:			Server timestamp, from the moment the handshake challenge was sent.
  *	- Cookie:				Cookie generated by the server, which the client must reply with.
  *	- AlignPad:				Handshake packets and PacketHandler's in general, require complex padding of packets. See ParseHandshakePacket.
+ *	- RandomData:			Data of random length/content appended to handshake packets, to work around potential faulty ISP packet filtering
  *
  *	- PacketSizeFiller:		Pads the client packet with blank information, so that the initial client packet,
  *							is the same size as the server response packet.
@@ -164,24 +155,14 @@ DEFINE_LOG_CATEGORY(LogHandshake);
 #define PACKETLOSS_TEST 0
 
 
-// Whether or not clients should send diagnostics to the server with the restart handshake, detailing why the request was accepted.
-#define RESTART_HANDSHAKE_DIAGNOSTICS 0
-
-// Disable client sending of handshake diagnostics
-#define DISABLE_SEND_HANDSHAKE_DIAGNOSTICS 1
-
 
 /**
  * Defines
  */
 
-#define HANDSHAKE_PACKET_SIZE_BITS				227
-#define RESTART_HANDSHAKE_PACKET_SIZE_BITS		2
-#define RESTART_RESPONSE_SIZE_BITS				387
-
-#if RESTART_HANDSHAKE_DIAGNOSTICS
-#define RESTART_RESPONSE_DIAGNOSTICS_SIZE_BITS	483
-#endif
+#define HANDSHAKE_PACKET_SIZE_BITS				259
+#define RESTART_HANDSHAKE_PACKET_SIZE_BITS		34
+#define RESTART_RESPONSE_SIZE_BITS				419
 
 
 // The number of seconds between secret value updates, and the random variance applied to this
@@ -205,47 +186,61 @@ TAutoConsoleVariable<FString> CVarNetMagicHeader(
 	TEXT("String representing binary bits which are prepended to every packet sent by the game. Max length: 32 bits."));
 
 
-#if RESTART_HANDSHAKE_DIAGNOSTICS
-TAutoConsoleVariable<int32> CVarNetRestartHandshakeDiagnostics(
-	TEXT("net.RestartHandshakeDiagnostics"),
-	0,
-	TEXT("Enables or disables restart handshake diagnostics. Serverside this controls logging. Clientside this controls sending."));
+namespace UE::Net
+{
+	static float HandshakeResendInterval = 1.f;
+
+	FAutoConsoleVariableRef CVarNetHandshakeResendInterval(
+		TEXT("net.HandshakeResendInterval"),
+		HandshakeResendInterval,
+		TEXT("The delay between resending handshake packets which we have not received a response for."));
+
+	/** The minimum supported stateless handshake protocol version */
+#ifdef HANDSHAKE_MIN_VERSION_OVERRIDE
+	static int32 MinSupportedHandshakeVersion = HANDSHAKE_MIN_VERSION_OVERRIDE;
+#else
+	static int32 MinSupportedHandshakeVersion = static_cast<uint8>(EHandshakeVersion::Original);
 #endif
 
+	/** The current compile-time handshake version */
+#ifdef HANDSHAKE_VERSION_OVERRIDE
+	static int32 CurrentHandshakeVersion = HANDSHAKE_VERSION_OVERRIDE;
+#else
+	static int32 CurrentHandshakeVersion = static_cast<uint8>(EHandshakeVersion::Latest);
+#endif
 
-/**
- * Structs/enums
- */
+	FAutoConsoleVariableRef CVarNetMinHandshakeVersion(
+		TEXT("net.MinHandshakeVersion"),
+		MinSupportedHandshakeVersion,
+		TEXT("The minimum supported stateless handshake protocol version (numeric)."));
 
-#if RESTART_HANDSHAKE_DIAGNOSTICS
-struct FRestartHandshakeDiagnostics
-{
-	float LastRestartPacketTimeDiff = -1.f;
-	float LastNetConnPacketTimeDiff = -1.f;
+	FAutoConsoleVariableRef CVarNetCurrentHandshakeVersion(
+		TEXT("net.CurrentHandshakeVersion"),
+		CurrentHandshakeVersion,
+		TEXT("The current supported stateless handshake protocol version (numeric)"));
 
-	uint64 ReservedUnused = 0;
+	TAutoConsoleVariable<int32> CVarNetDoHandshakeVersionFallback(
+		TEXT("net.DoHandshakeVersionFallback"),
+		0,
+		TEXT("Whether or not to (clientside) perform randomized falling-back to previous versions of the handshake protocol, upon failure."));
 
+	/** The base amount of random data to add to handshake packets */
+	static constexpr int32 BaseRandomDataLengthBytes		= 16;
 
-	bool IsValid()
-	{
-		return LastRestartPacketTimeDiff != -1.f || LastNetConnPacketTimeDiff != -1.f;
-	}
-};
+	/** The amount by which the length of random data should randomly vary */
+	static constexpr int32 RandomDataLengthVarianceBytes	= 8;
 
-static_assert(((RESTART_RESPONSE_DIAGNOSTICS_SIZE_BITS - RESTART_RESPONSE_SIZE_BITS) - (sizeof(FRestartHandshakeDiagnostics) * 8)) == 0,
-				"FRestartHandshakeDiagnostics must be properly factored into packet size defines.");
 
-FArchive& operator << (FArchive& Ar, FRestartHandshakeDiagnostics& D)
-{
-	Ar << D.LastRestartPacketTimeDiff;
-	Ar << D.LastNetConnPacketTimeDiff;
-	Ar << D.ReservedUnused;
+	/** HANDSHAKE_PACKET_SIZE_BITS for EHandshakeVersion::Original */
+	static constexpr int32 OriginalHandshakePacketSizeBits = 227;
+
+	/** RESTART_HANDSHAKE_PACKET_SIZE_BITS for EHandshakeVersion::Original */
+	static constexpr int32 OriginalRestartHandshakePacketSizeBits = 2;
 
-	return Ar;
+	/** RESTART_RESPONSE_SIZE_BITS for EHandshakeVersion::Original */
+	static constexpr int32 OriginalRestartResponseSizeBits = 387;
 }
 
-static FRestartHandshakeDiagnostics HandshakeDiagnostics;
-#endif
 
 
 /**
@@ -261,6 +256,7 @@ StatelessConnectHandlerComponent::StatelessConnectHandlerComponent()
 	, LastChallengeSuccessAddress(nullptr)
 	, LastServerSequence(0)
 	, LastClientSequence(0)
+	, MinClientHandshakeVersion(static_cast<EHandshakeVersion>(UE::Net::CurrentHandshakeVersion))
 	, LastClientSendTimestamp(0.0)
 	, LastChallengeTimestamp(0.0)
 	, LastRestartPacketTimestamp(0.0)
@@ -270,6 +266,7 @@ StatelessConnectHandlerComponent::StatelessConnectHandlerComponent()
 	, bRestartedHandshake(false)
 	, AuthorisedCookie()
 	, MagicHeader()
+	, LastRemoteHandshakeVersion(static_cast<EHandshakeVersion>(UE::Net::CurrentHandshakeVersion))
 {
 	SetActive(true);
 
@@ -322,13 +319,22 @@ void StatelessConnectHandlerComponent::CountBytes(FArchive& Ar) const
 
 void StatelessConnectHandlerComponent::NotifyHandshakeBegin()
 {
+	using namespace UE::Net;
+
+	SendInitialPacket(static_cast<EHandshakeVersion>(CurrentHandshakeVersion));
+}
+
+void StatelessConnectHandlerComponent::SendInitialPacket(EHandshakeVersion HandshakeVersion)
+{
+	using namespace UE::Net;
+
 	if (Handler->Mode == Handler::Mode::Client)
 	{
 		UNetConnection* ServerConn = (Driver != nullptr ? ToRawPtr(Driver->ServerConnection) : nullptr);
 
 		if (ServerConn != nullptr)
 		{
-			FBitWriter InitialPacket(GetAdjustedSizeBits(HANDSHAKE_PACKET_SIZE_BITS) + 1 /* Termination bit */);
+			FBitWriter InitialPacket(GetAdjustedSizeBits(HANDSHAKE_PACKET_SIZE_BITS) + (BaseRandomDataLengthBytes * 8) + 1 /* Termination bit */);
 			uint8 bHandshakePacket = 1;
 
 			if (MagicHeader.Num() > 0)
@@ -345,6 +351,19 @@ void StatelessConnectHandlerComponent::NotifyHandshakeBegin()
 			uint8 PacketSizeFiller[28];
 
 			InitialPacket.WriteBit(bRestartHandshake);
+
+			if (HandshakeVersion >= EHandshakeVersion::Randomized)
+			{
+				uint8 MinVersion = MinSupportedHandshakeVersion;
+				uint8 CurVersion = CurrentHandshakeVersion;
+				uint8 HandshakePacketType = static_cast<uint8>(EHandshakePacketType::InitialPacket);
+
+				InitialPacket << MinVersion;
+				InitialPacket << CurVersion;
+				InitialPacket << HandshakePacketType;
+				InitialPacket << SentHandshakePacketCount;
+			}
+
 			InitialPacket.WriteBit(SecretIdPad);
 
 			FMemory::Memzero(PacketSizeFiller, UE_ARRAY_COUNT(PacketSizeFiller));
@@ -352,7 +371,7 @@ void StatelessConnectHandlerComponent::NotifyHandshakeBegin()
 
 
 
-			CapHandshakePacket(InitialPacket);
+			CapHandshakePacket(InitialPacket, HandshakeVersion);
 
 
 			// Disable PacketHandler parsing, and send the raw packet
@@ -388,11 +407,14 @@ void StatelessConnectHandlerComponent::NotifyHandshakeBegin()
 	}
 }
 
-void StatelessConnectHandlerComponent::SendConnectChallenge(TSharedPtr<const FInternetAddr> ClientAddress)
+void StatelessConnectHandlerComponent::SendConnectChallenge(TSharedPtr<const FInternetAddr> ClientAddress, EHandshakeVersion HandshakeVersion,
+															uint8 ClientSentHandshakePacketCount)
 {
+	using namespace UE::Net;
+
 	if (Driver != nullptr)
 	{
-		FBitWriter ChallengePacket(GetAdjustedSizeBits(HANDSHAKE_PACKET_SIZE_BITS) + 1 /* Termination bit */);
+		FBitWriter ChallengePacket(GetAdjustedSizeBits(HANDSHAKE_PACKET_SIZE_BITS) + (BaseRandomDataLengthBytes * 8) + 1 /* Termination bit */);
 		uint8 bHandshakePacket = 1;
 		uint8 bRestartHandshake = 0; // Ignored clientside
 		double Timestamp = Driver->GetElapsedTime();
@@ -407,6 +429,19 @@ void StatelessConnectHandlerComponent::SendConnectChallenge(TSharedPtr<const FIn
 
 		ChallengePacket.WriteBit(bHandshakePacket);
 		ChallengePacket.WriteBit(bRestartHandshake);
+
+		if (HandshakeVersion >= EHandshakeVersion::Randomized)
+		{
+			uint8 MinVersion = MinSupportedHandshakeVersion;
+			uint8 CurVersion = static_cast<uint8>(HandshakeVersion);
+			uint8 HandshakePacketType = static_cast<uint8>(EHandshakePacketType::Challenge);
+
+			ChallengePacket << MinVersion;
+			ChallengePacket << CurVersion;
+			ChallengePacket << HandshakePacketType;
+			ChallengePacket << ClientSentHandshakePacketCount;
+		}
+
 		ChallengePacket.WriteBit(ActiveSecret);
 
 		ChallengePacket << Timestamp;
@@ -419,7 +454,7 @@ void StatelessConnectHandlerComponent::SendConnectChallenge(TSharedPtr<const FIn
 				TEXT("SendConnectChallenge. Timestamp: %f, Cookie: %s" ), Timestamp, *FString::FromBlob(Cookie, UE_ARRAY_COUNT(Cookie)));
 #endif
 
-		CapHandshakePacket(ChallengePacket);
+		CapHandshakePacket(ChallengePacket, HandshakeVersion);
 
 		
 		// Disable PacketHandler parsing, and send the raw packet
@@ -463,22 +498,19 @@ void StatelessConnectHandlerComponent::SendConnectChallenge(TSharedPtr<const FIn
 	}
 }
 
-void StatelessConnectHandlerComponent::SendChallengeResponse(uint8 InSecretId, double InTimestamp, uint8 InCookie[COOKIE_BYTE_SIZE])
+void StatelessConnectHandlerComponent::SendChallengeResponse(EHandshakeVersion HandshakeVersion, uint8 InSecretId, double InTimestamp,
+																uint8 InCookie[COOKIE_BYTE_SIZE])
 {
+	using namespace UE::Net;
+
 	UNetConnection* ServerConn = (Driver != nullptr ? ToRawPtr(Driver->ServerConnection) : nullptr);
 
 	if (ServerConn != nullptr)
 	{
 		int32 RestartHandshakeResponseSize = RESTART_RESPONSE_SIZE_BITS;
 
-#if RESTART_HANDSHAKE_DIAGNOSTICS && !DISABLE_SEND_HANDSHAKE_DIAGNOSTICS
-		bool bEnableDiagnostics = bRestartedHandshake && !!CVarNetRestartHandshakeDiagnostics.GetValueOnAnyThread();
-
-		RestartHandshakeResponseSize = bEnableDiagnostics ? RESTART_RESPONSE_DIAGNOSTICS_SIZE_BITS : RestartHandshakeResponseSize;
-#endif
-
 		const int32 BaseSize = GetAdjustedSizeBits(bRestartedHandshake ? RestartHandshakeResponseSize : HANDSHAKE_PACKET_SIZE_BITS);
-		FBitWriter ResponsePacket(BaseSize + 1 /* Termination bit */);
+		FBitWriter ResponsePacket(BaseSize + (BaseRandomDataLengthBytes * 8) + 1 /* Termination bit */);
 		uint8 bHandshakePacket = 1;
 		uint8 bRestartHandshake = (bRestartedHandshake ? 1 : 0);
 
@@ -489,6 +521,21 @@ void StatelessConnectHandlerComponent::SendChallengeResponse(uint8 InSecretId, d
 
 		ResponsePacket.WriteBit(bHandshakePacket);
 		ResponsePacket.WriteBit(bRestartHandshake);
+
+		if (HandshakeVersion >= EHandshakeVersion::Randomized)
+		{
+			uint8 MinVersion = MinSupportedHandshakeVersion;
+			uint8 CurVersion = static_cast<uint8>(HandshakeVersion);
+			uint8 HandshakePacketType = static_cast<uint8>(bRestartedHandshake ? EHandshakePacketType::RestartResponse :
+											EHandshakePacketType::Response);
+
+			ResponsePacket << MinVersion;
+			ResponsePacket << CurVersion;
+
+			ResponsePacket << HandshakePacketType;
+			ResponsePacket << SentHandshakePacketCount;
+		}
+
 		ResponsePacket.WriteBit(InSecretId);
 
 		ResponsePacket << InTimestamp;
@@ -497,20 +544,13 @@ void StatelessConnectHandlerComponent::SendChallengeResponse(uint8 InSecretId, d
 		if (bRestartedHandshake)
 		{
 			ResponsePacket.Serialize(AuthorisedCookie, COOKIE_BYTE_SIZE);
-
-#if RESTART_HANDSHAKE_DIAGNOSTICS && !DISABLE_SEND_HANDSHAKE_DIAGNOSTICS
-			if (bEnableDiagnostics)
-			{
-				ResponsePacket << HandshakeDiagnostics;
-			}
-#endif
 		}
 
 #if !UE_BUILD_SHIPPING
 		UE_LOG( LogHandshake, Log, TEXT( "SendChallengeResponse. Timestamp: %f, Cookie: %s" ), InTimestamp, *FString::FromBlob( InCookie, COOKIE_BYTE_SIZE ) );
 #endif
 
-		CapHandshakePacket(ResponsePacket);
+		CapHandshakePacket(ResponsePacket, HandshakeVersion);
 
 
 		// Disable PacketHandler parsing, and send the raw packet
@@ -544,6 +584,7 @@ void StatelessConnectHandlerComponent::SendChallengeResponse(uint8 InSecretId, d
 		LastTimestamp = InTimestamp;
 		LastServerSequence = *CurSequence & (MAX_PACKETID - 1);
 		LastClientSequence = *(CurSequence + 1) & (MAX_PACKETID - 1);
+		LastRemoteHandshakeVersion = HandshakeVersion;
 
 		FMemory::Memcpy(LastCookie, InCookie, UE_ARRAY_COUNT(LastCookie));
 	}
@@ -553,11 +594,14 @@ void StatelessConnectHandlerComponent::SendChallengeResponse(uint8 InSecretId, d
 	}
 }
 
-void StatelessConnectHandlerComponent::SendChallengeAck(TSharedPtr<const FInternetAddr> ClientAddress, uint8 InCookie[COOKIE_BYTE_SIZE])
+void StatelessConnectHandlerComponent::SendChallengeAck(TSharedPtr<const FInternetAddr> ClientAddress, EHandshakeVersion HandshakeVersion,
+														uint8 ClientSentHandshakePacketCount, uint8 InCookie[COOKIE_BYTE_SIZE])
 {
+	using namespace UE::Net;
+
 	if (Driver != nullptr)
 	{
-		FBitWriter AckPacket(GetAdjustedSizeBits(HANDSHAKE_PACKET_SIZE_BITS) + 1 /* Termination bit */);
+		FBitWriter AckPacket(GetAdjustedSizeBits(HANDSHAKE_PACKET_SIZE_BITS) + (BaseRandomDataLengthBytes * 8) + 1 /* Termination bit */);
 		uint8 bHandshakePacket = 1;
 		uint8 bRestartHandshake = 0; // Ignored clientside
 		double Timestamp  = -1.0;
@@ -569,6 +613,20 @@ void StatelessConnectHandlerComponent::SendChallengeAck(TSharedPtr<const FIntern
 
 		AckPacket.WriteBit(bHandshakePacket);
 		AckPacket.WriteBit(bRestartHandshake);
+
+		if (HandshakeVersion >= EHandshakeVersion::Randomized)
+		{
+			uint8 MinVersion = MinSupportedHandshakeVersion;
+			uint8 CurVersion = static_cast<uint8>(HandshakeVersion);
+			uint8 HandshakePacketType = static_cast<uint8>(EHandshakePacketType::Ack);
+
+			AckPacket << MinVersion;
+			AckPacket << CurVersion;
+
+			AckPacket << HandshakePacketType;
+			AckPacket << ClientSentHandshakePacketCount;
+		}
+
 		AckPacket.WriteBit(bHandshakePacket);	// ActiveSecret
 
 		AckPacket << Timestamp;
@@ -578,7 +636,7 @@ void StatelessConnectHandlerComponent::SendChallengeAck(TSharedPtr<const FIntern
 		UE_LOG(LogHandshake, Log, TEXT("SendChallengeAck. InCookie: %s" ), *FString::FromBlob(InCookie, COOKIE_BYTE_SIZE));
 #endif
 
-		CapHandshakePacket(AckPacket);
+		CapHandshakePacket(AckPacket, HandshakeVersion);
 
 		
 		// Disable PacketHandler parsing, and send the raw packet
@@ -622,11 +680,14 @@ void StatelessConnectHandlerComponent::SendChallengeAck(TSharedPtr<const FIntern
 	}
 }
 
-void StatelessConnectHandlerComponent::SendRestartHandshakeRequest(const TSharedPtr<const FInternetAddr> ClientAddress)
+void StatelessConnectHandlerComponent::SendRestartHandshakeRequest(const TSharedPtr<const FInternetAddr> ClientAddress,
+																	EHandshakeVersion HandshakeVersion)
 {
+	using namespace UE::Net;
+
 	if (Driver != nullptr)
 	{
-		FBitWriter RestartPacket(GetAdjustedSizeBits(RESTART_HANDSHAKE_PACKET_SIZE_BITS) + 1 /* Termination bit */);
+		FBitWriter RestartPacket(GetAdjustedSizeBits(RESTART_HANDSHAKE_PACKET_SIZE_BITS) + (BaseRandomDataLengthBytes * 8) + 1 /* Termination bit */);
 		uint8 bHandshakePacket = 1;
 		uint8 bRestartHandshake = 1;
 
@@ -638,13 +699,26 @@ void StatelessConnectHandlerComponent::SendRestartHandshakeRequest(const TShared
 		RestartPacket.WriteBit(bHandshakePacket);
 		RestartPacket.WriteBit(bRestartHandshake);
 
+		if (HandshakeVersion >= EHandshakeVersion::Randomized)
+		{
+			uint8 MinVersion = MinSupportedHandshakeVersion;
+			uint8 CurVersion = static_cast<uint8>(HandshakeVersion);
+			uint8 HandshakePacketType = static_cast<uint8>(EHandshakePacketType::RestartHandshake);
+
+			RestartPacket << MinVersion;
+			RestartPacket << CurVersion;
+
+			RestartPacket << HandshakePacketType;
+			RestartPacket << SentHandshakePacketCount;
+		}
+
 #if !UE_BUILD_SHIPPING
 		FDDoSDetection* DDoS = Handler->GetDDoS();
 
 		UE_CLOG((DDoS == nullptr || !DDoS->CheckLogRestrictions()), LogHandshake, Verbose, TEXT("SendRestartHandshakeRequest."));
 #endif
 
-		CapHandshakePacket(RestartPacket);
+		CapHandshakePacket(RestartPacket, HandshakeVersion);
 
 		
 		// Disable PacketHandler parsing, and send the raw packet
@@ -688,21 +762,42 @@ void StatelessConnectHandlerComponent::SendRestartHandshakeRequest(const TShared
 	}
 }
 
-void StatelessConnectHandlerComponent::CapHandshakePacket(FBitWriter& HandshakePacket)
+void StatelessConnectHandlerComponent::CapHandshakePacket(FBitWriter& HandshakePacket, EHandshakeVersion HandshakeVersion)
 {
+	using namespace UE::Net;
+
 	uint32 NumBits = HandshakePacket.GetNumBits() - GetAdjustedSizeBits(0);
 
-#if RESTART_HANDSHAKE_DIAGNOSTICS
-	check(NumBits == HANDSHAKE_PACKET_SIZE_BITS || NumBits == RESTART_HANDSHAKE_PACKET_SIZE_BITS || NumBits == RESTART_RESPONSE_SIZE_BITS
-			|| NumBits == RESTART_RESPONSE_DIAGNOSTICS_SIZE_BITS);
-#else
-	check(NumBits == HANDSHAKE_PACKET_SIZE_BITS || NumBits == RESTART_HANDSHAKE_PACKET_SIZE_BITS || NumBits == RESTART_RESPONSE_SIZE_BITS);
+#if !UE_BUILD_SHIPPING
+	if (HandshakeVersion == EHandshakeVersion::Original)
+	{
+		check(NumBits == OriginalHandshakePacketSizeBits || NumBits == OriginalRestartHandshakePacketSizeBits ||
+				NumBits == OriginalRestartResponseSizeBits);
+	}
+	else
+	{
+		check(NumBits == HANDSHAKE_PACKET_SIZE_BITS || NumBits == RESTART_HANDSHAKE_PACKET_SIZE_BITS || NumBits == RESTART_RESPONSE_SIZE_BITS);
+	}
 #endif
 
 	FPacketAudit::AddStage(TEXT("PostPacketHandler"), HandshakePacket);
 
+	if (HandshakeVersion >= EHandshakeVersion::Randomized)
+	{
+		const int32 RandomDataLengthBytes = BaseRandomDataLengthBytes - FMath::RandRange(0, RandomDataLengthVarianceBytes);
+
+		for (int32 RandIdx=0; RandIdx<RandomDataLengthBytes; RandIdx++)
+		{
+			uint8 RandVal = FMath::Rand() % 255;
+
+			HandshakePacket << RandVal;
+		}
+	}
+
 	// Add a termination bit, the same as the UNetConnection code does
 	HandshakePacket.WriteBit(1);
+
+	SentHandshakePacketCount++;
 }
 
 void StatelessConnectHandlerComponent::SetDriver(UNetDriver* InDriver)
@@ -740,6 +835,7 @@ void StatelessConnectHandlerComponent::InitFromConnectionless(StatelessConnectHa
 {
 	// Store the cookie/address used for the handshake, to enable server ack-retries
 	LastChallengeSuccessAddress = InConnectionlessHandler->LastChallengeSuccessAddress;
+	LastRemoteHandshakeVersion = InConnectionlessHandler->LastRemoteHandshakeVersion;
 
 	FMemory::Memcpy(AuthorisedCookie, InConnectionlessHandler->AuthorisedCookie, UE_ARRAY_COUNT(AuthorisedCookie));
 
@@ -748,6 +844,8 @@ void StatelessConnectHandlerComponent::InitFromConnectionless(StatelessConnectHa
 
 void StatelessConnectHandlerComponent::Incoming(FBitReader& Packet)
 {
+	using namespace UE::Net;
+
 	if (MagicHeader.Num() > 0)
 	{
 		// Don't bother with the expense of verifying the magic header here.
@@ -759,13 +857,9 @@ void StatelessConnectHandlerComponent::Incoming(FBitReader& Packet)
 
 	if (bHandshakePacket)
 	{
-		bool bRestartHandshake = false;
-		uint8 SecretId = 0;
-		double Timestamp = 1.;
-		uint8 Cookie[COOKIE_BYTE_SIZE];
-		uint8 OrigCookie[COOKIE_BYTE_SIZE];
+		FParsedHandshakeData HandshakeData;
 
-		bHandshakePacket = ParseHandshakePacket(Packet, bRestartHandshake, SecretId, Timestamp, Cookie, OrigCookie);
+		bHandshakePacket = ParseHandshakePacket(Packet, HandshakeData);
 
 		if (bHandshakePacket)
 		{
@@ -773,24 +867,24 @@ void StatelessConnectHandlerComponent::Incoming(FBitReader& Packet)
 			{
 				if (State == Handler::Component::State::UnInitialized || State == Handler::Component::State::InitializedOnLocal)
 				{
-					if (bRestartHandshake)
+					if (HandshakeData.bRestartHandshake)
 					{
 #if !UE_BUILD_SHIPPING
 						UE_LOG(LogHandshake, Log, TEXT("Ignoring restart handshake request, while already restarted."));
 #endif
 					}
 					// Receiving challenge, verify the timestamp is > 0.0f
-					else if (Timestamp > 0.0)
+					else if (HandshakeData.HandshakePacketType == EHandshakePacketType::Challenge && HandshakeData.Timestamp > 0.0)
 					{
 						LastChallengeTimestamp = (Driver != nullptr ? Driver->GetElapsedTime() : 0.0);
 
-						SendChallengeResponse(SecretId, Timestamp, Cookie);
+						SendChallengeResponse(HandshakeData.RemoteCurVersion, HandshakeData.SecretId, HandshakeData.Timestamp, HandshakeData.Cookie);
 
 						// Utilize this state as an intermediary, indicating that the challenge response has been sent
 						SetState(Handler::Component::State::InitializedOnLocal);
 					}
 					// Receiving challenge ack, verify the timestamp is < 0.0f
-					else if (Timestamp < 0.0)
+					else if (HandshakeData.HandshakePacketType == EHandshakePacketType::Ack && HandshakeData.Timestamp < 0.0)
 					{
 						if (!bRestartedHandshake)
 						{
@@ -799,7 +893,7 @@ void StatelessConnectHandlerComponent::Incoming(FBitReader& Packet)
 							// Extract the initial packet sequence from the random Cookie data
 							if (ensure(ServerConn != nullptr))
 							{
-								int16* CurSequence = (int16*)Cookie;
+								int16* CurSequence = (int16*)HandshakeData.Cookie;
 
 								int32 ServerSequence = *CurSequence & (MAX_PACKETID - 1);
 								int32 ClientSequence = *(CurSequence + 1) & (MAX_PACKETID - 1);
@@ -808,7 +902,7 @@ void StatelessConnectHandlerComponent::Incoming(FBitReader& Packet)
 							}
 
 							// Save the final authorized cookie
-							FMemory::Memcpy(AuthorisedCookie, Cookie, UE_ARRAY_COUNT(AuthorisedCookie));
+							FMemory::Memcpy(AuthorisedCookie, HandshakeData.Cookie, UE_ARRAY_COUNT(AuthorisedCookie));
 						}
 
 						// Now finish initializing the handler - flushing the queued packet buffer in the process.
@@ -816,9 +910,12 @@ void StatelessConnectHandlerComponent::Incoming(FBitReader& Packet)
 						Initialized();
 
 						bRestartedHandshake = false;
+
+						// Reset packet count clientside, due to how it affects protocol version fallback selection
+						SentHandshakePacketCount = 0;
 					}
 				}
-				else if (bRestartHandshake)
+				else if (HandshakeData.bRestartHandshake)
 				{
 					uint8 ZeroCookie[COOKIE_BYTE_SIZE] = {0};
 					bool bValidAuthCookie = FMemory::Memcmp(AuthorisedCookie, ZeroCookie, COOKIE_BYTE_SIZE) != 0;
@@ -849,14 +946,6 @@ void StatelessConnectHandlerComponent::Incoming(FBitReader& Packet)
 							bPassedDualIPCheck = LastRestartPacketTimestamp == 0.0 ||
 													LastRestartPacketTimeDiff > 1.1 ||
 													LastNetConnPacketTimeDiff > 1.0;
-
-#if RESTART_HANDSHAKE_DIAGNOSTICS
-							if (bPassedDualIPCheck)
-							{
-								HandshakeDiagnostics.LastRestartPacketTimeDiff = LastRestartPacketTimeDiff;
-								HandshakeDiagnostics.LastNetConnPacketTimeDiff = LastNetConnPacketTimeDiff;
-							}
-#endif
 						}
 
 						LastRestartPacketTimestamp = CurrentTime;
@@ -894,7 +983,7 @@ void StatelessConnectHandlerComponent::Incoming(FBitReader& Packet)
 							bRestartedHandshake = true;
 
 							SetState(Handler::Component::State::UnInitialized);
-							NotifyHandshakeBegin();
+							SendInitialPacket(LastRemoteHandshakeVersion);
 						}
 						else if (WithinHandshakeLogLimit())
 						{
@@ -937,7 +1026,7 @@ void StatelessConnectHandlerComponent::Incoming(FBitReader& Packet)
 							*LastChallengeSuccessAddress->ToString(true), *FString::FromBlob(AuthorisedCookie, COOKIE_BYTE_SIZE));
 #endif
 
-					SendChallengeAck(LastChallengeSuccessAddress, AuthorisedCookie);
+					SendChallengeAck(LastChallengeSuccessAddress, LastRemoteHandshakeVersion, 0, AuthorisedCookie);
 				}
 			}
 		}
@@ -991,6 +1080,8 @@ void StatelessConnectHandlerComponent::Outgoing(FBitWriter& Packet, FOutPacketTr
 
 void StatelessConnectHandlerComponent::IncomingConnectionless(FIncomingPacketRef PacketRef)
 {
+	using namespace UE::Net;
+
 	FBitReader& Packet = PacketRef.Packet;
 	const TSharedPtr<const FInternetAddr> Address = PacketRef.Address;
 
@@ -1007,23 +1098,35 @@ void StatelessConnectHandlerComponent::IncomingConnectionless(FIncomingPacketRef
 
 	if (bHandshakePacket)
 	{
-		bool bRestartHandshake = false;
-		uint8 SecretId = 0;
-		double Timestamp = 1.0;
-		uint8 Cookie[COOKIE_BYTE_SIZE];
-		uint8 OrigCookie[COOKIE_BYTE_SIZE];
+		FParsedHandshakeData HandshakeData;
 
-		bHandshakePacket = ParseHandshakePacket(Packet, bRestartHandshake, SecretId, Timestamp, Cookie, OrigCookie);
+		bHandshakePacket = ParseHandshakePacket(Packet, HandshakeData);
 
 		if (bHandshakePacket)
 		{
+			bool bValidVersion = false;
+			EHandshakeVersion TargetVersion = static_cast<EHandshakeVersion>(CurrentHandshakeVersion);
+
+			if (static_cast<uint8>(HandshakeData.RemoteMinVersion) >= MinSupportedHandshakeVersion &&
+				HandshakeData.RemoteMinVersion < TargetVersion)
+			{
+				if (HandshakeData.RemoteCurVersion <= TargetVersion)
+				{
+					TargetVersion = HandshakeData.RemoteCurVersion;
+				}
+
+				bValidVersion = true;
+			}
+
+
 			if (Handler->Mode == Handler::Mode::Server)
 			{
-				const bool bInitialConnect = Timestamp == 0.0;
+				const bool bInitialConnect = HandshakeData.HandshakePacketType == EHandshakePacketType::InitialPacket &&
+												HandshakeData.Timestamp == 0.0;
 
 				if (bInitialConnect)
 				{
-					SendConnectChallenge(Address);
+					SendConnectChallenge(Address, TargetVersion, HandshakeData.RemoteSentHandshakePacketCount);
 				}
 				// Challenge response
 				else if (Driver != nullptr)
@@ -1031,57 +1134,48 @@ void StatelessConnectHandlerComponent::IncomingConnectionless(FIncomingPacketRef
 					// NOTE: Allow CookieDelta to be 0.0, as it is possible for a server to send a challenge and receive a response,
 					//			during the same tick
 					bool bChallengeSuccess = false;
-					const double CookieDelta = Driver->GetElapsedTime() - Timestamp;
-					const double SecretDelta = Timestamp - LastSecretUpdateTimestamp;
+					const double CookieDelta = Driver->GetElapsedTime() - HandshakeData.Timestamp;
+					const double SecretDelta = HandshakeData.Timestamp - LastSecretUpdateTimestamp;
 					const bool bValidCookieLifetime = CookieDelta >= 0.0 && (MAX_COOKIE_LIFETIME - CookieDelta) > 0.0;
-					const bool bValidSecretIdTimestamp = (SecretId == ActiveSecret) ? (SecretDelta >= 0.0) : (SecretDelta <= 0.0);
+					const bool bValidSecretIdTimestamp = (HandshakeData.SecretId == ActiveSecret) ? (SecretDelta >= 0.0) : (SecretDelta <= 0.0);
 
 					if (bValidCookieLifetime && bValidSecretIdTimestamp)
 					{
 						// Regenerate the cookie from the packet info, and see if the received cookie matches the regenerated one
 						uint8 RegenCookie[COOKIE_BYTE_SIZE];
 
-						GenerateCookie(Address, SecretId, Timestamp, RegenCookie);
+						GenerateCookie(Address, HandshakeData.SecretId, HandshakeData.Timestamp, RegenCookie);
 
-						bChallengeSuccess = FMemory::Memcmp(Cookie, RegenCookie, COOKIE_BYTE_SIZE) == 0;
+						bChallengeSuccess = FMemory::Memcmp(HandshakeData.Cookie, RegenCookie, COOKIE_BYTE_SIZE) == 0;
 
 						if (bChallengeSuccess)
 						{
-							if (bRestartHandshake)
+							if (HandshakeData.bRestartHandshake)
 							{
-								FMemory::Memcpy(AuthorisedCookie, OrigCookie, UE_ARRAY_COUNT(AuthorisedCookie));
-
-#if RESTART_HANDSHAKE_DIAGNOSTICS
-								if (HandshakeDiagnostics.IsValid() && !!CVarNetRestartHandshakeDiagnostics.GetValueOnAnyThread())
-								{
-									FDDoSDetection* DDoS = Handler->GetDDoS();
-
-									UE_CLOG((DDoS == nullptr || !DDoS->CheckLogRestrictions()), LogHandshake, Log,
-											TEXT("Got restart handshake diagnostics: LastRestartPacketTimeDiff: %f, ")
-											TEXT("LastNetConnPacketTimeDiff: %f"),
-											HandshakeDiagnostics.LastRestartPacketTimeDiff,
-											HandshakeDiagnostics.LastNetConnPacketTimeDiff);
-
-									HandshakeDiagnostics = FRestartHandshakeDiagnostics();
-								}
-#endif
+								FMemory::Memcpy(AuthorisedCookie, HandshakeData.OrigCookie, UE_ARRAY_COUNT(AuthorisedCookie));
 							}
 							else
 							{
-								int16* CurSequence = (int16*)Cookie;
+								int16* CurSequence = (int16*)HandshakeData.Cookie;
 
 								LastServerSequence = *CurSequence & (MAX_PACKETID - 1);
 								LastClientSequence = *(CurSequence + 1) & (MAX_PACKETID - 1);
 
-								FMemory::Memcpy(AuthorisedCookie, Cookie, UE_ARRAY_COUNT(AuthorisedCookie));
+								FMemory::Memcpy(AuthorisedCookie, HandshakeData.Cookie, UE_ARRAY_COUNT(AuthorisedCookie));
 							}
 
-							bRestartedHandshake = bRestartHandshake;
+							bRestartedHandshake = HandshakeData.bRestartHandshake;
 							LastChallengeSuccessAddress = Address->Clone();
+							LastRemoteHandshakeVersion = TargetVersion;
+
+							if (TargetVersion < MinClientHandshakeVersion && static_cast<uint8>(TargetVersion) >= MinSupportedHandshakeVersion)
+							{
+								MinClientHandshakeVersion = TargetVersion;
+							}
 
 
 							// Now ack the challenge response - the cookie is stored in AuthorisedCookie, to enable retries
-							SendChallengeAck(Address, AuthorisedCookie);
+							SendChallengeAck(Address, TargetVersion, HandshakeData.RemoteSentHandshakePacketCount, AuthorisedCookie);
 						}
 					}
 				}
@@ -1115,52 +1209,168 @@ void StatelessConnectHandlerComponent::IncomingConnectionless(FIncomingPacketRef
 	{
 		// The packet was fine but not a handshake packet - an existing client might suddenly be communicating on a different address.
 		// If we get them to resend their cookie, we can update the connection's info with their new address.
-		SendRestartHandshakeRequest(Address);
+		SendRestartHandshakeRequest(Address, static_cast<EHandshakeVersion>(MinSupportedHandshakeVersion));
 	}
 }
 
-bool StatelessConnectHandlerComponent::ParseHandshakePacket(FBitReader& Packet, bool& bOutRestartHandshake, uint8& OutSecretId,
-															double& OutTimestamp, uint8 (&OutCookie)[COOKIE_BYTE_SIZE],
-															uint8 (&OutOrigCookie)[COOKIE_BYTE_SIZE])
+bool StatelessConnectHandlerComponent::ParseHandshakePacket(FBitReader& Packet, FParsedHandshakeData& OutResult)
 {
+	using namespace UE::Net;
+
+	// Ensure original packet sizes don't overlap with size range of new packet format - so that we can detect original version, based on size
+	{
+		static constexpr int32 MinHandshakePacketVariance = HANDSHAKE_PACKET_SIZE_BITS - RandomDataLengthVarianceBytes;
+		static constexpr int32 MaxHandshakePacketVariance = HANDSHAKE_PACKET_SIZE_BITS;
+		static constexpr int32 MinRestartHandshakePacketVariance = RESTART_HANDSHAKE_PACKET_SIZE_BITS - RandomDataLengthVarianceBytes;
+		static constexpr int32 MaxRestartHandshakePacketVariance = RESTART_HANDSHAKE_PACKET_SIZE_BITS;
+		static constexpr int32 MinRestartResponsePacketVariance = RESTART_RESPONSE_SIZE_BITS - RandomDataLengthVarianceBytes;
+		static constexpr int32 MaxRestartResponsePacketVariance = RESTART_RESPONSE_SIZE_BITS;
+
+		static_assert(OriginalHandshakePacketSizeBits < MinHandshakePacketVariance || OriginalHandshakePacketSizeBits > MaxHandshakePacketVariance);
+		static_assert(OriginalHandshakePacketSizeBits < MinRestartHandshakePacketVariance ||
+						OriginalHandshakePacketSizeBits > MaxRestartHandshakePacketVariance);
+		static_assert(OriginalHandshakePacketSizeBits < MinRestartResponsePacketVariance ||
+						OriginalHandshakePacketSizeBits > MaxRestartResponsePacketVariance);
+		static_assert(OriginalRestartHandshakePacketSizeBits < MinHandshakePacketVariance ||
+						OriginalRestartHandshakePacketSizeBits > MaxHandshakePacketVariance);
+		static_assert(OriginalRestartHandshakePacketSizeBits < MinRestartHandshakePacketVariance ||
+						OriginalRestartHandshakePacketSizeBits > MaxRestartHandshakePacketVariance);
+		static_assert(OriginalRestartHandshakePacketSizeBits < MinRestartResponsePacketVariance ||
+						OriginalRestartHandshakePacketSizeBits > MaxRestartResponsePacketVariance);
+		static_assert(OriginalRestartResponseSizeBits < MinHandshakePacketVariance || OriginalRestartResponseSizeBits > MaxHandshakePacketVariance);
+		static_assert(OriginalRestartResponseSizeBits < MinRestartHandshakePacketVariance ||
+						OriginalRestartResponseSizeBits > MaxRestartHandshakePacketVariance);
+		static_assert(OriginalRestartResponseSizeBits < MinRestartResponsePacketVariance ||
+						OriginalRestartResponseSizeBits > MaxRestartResponsePacketVariance);
+	};
+
 	bool bValidPacket = false;
-	uint32 BitsLeft = Packet.GetBitsLeft();
-	bool bHandshakePacketSize = BitsLeft == (HANDSHAKE_PACKET_SIZE_BITS - 1);
-	bool bRestartResponsePacketSize = BitsLeft == (RESTART_RESPONSE_SIZE_BITS - 1);
-	bool bRestartResponseDiagnosticsPacketSize = false
-#if RESTART_HANDSHAKE_DIAGNOSTICS
-			|| BitsLeft == (RESTART_RESPONSE_DIAGNOSTICS_SIZE_BITS - 1)
-#endif
-		;
+	const int32 BitsLeft = Packet.GetBitsLeft();
+	const bool bOriginalVersion = BitsLeft == (OriginalHandshakePacketSizeBits - 1) || BitsLeft == (OriginalRestartHandshakePacketSizeBits - 1) ||
+									BitsLeft == (OriginalRestartResponseSizeBits - 1);
+
+	if (bOriginalVersion)
+	{
+		return ParseHandshakePacketOriginal(Packet, OutResult);
+	}
+
+
+	const int32 HandshakePacketBitsLeft = BitsLeft - (HANDSHAKE_PACKET_SIZE_BITS - 1);
+	const int32 RestartHandshakePacketBitsLeft = BitsLeft - (RESTART_HANDSHAKE_PACKET_SIZE_BITS - 1);
+	const int32 RestartResponsePacketBitsLeft = BitsLeft - (RESTART_RESPONSE_SIZE_BITS - 1);
+	const int32 MinBitsLeft = (BaseRandomDataLengthBytes - RandomDataLengthVarianceBytes) * 8;
+	const int32 MaxBitsLeft = BaseRandomDataLengthBytes * 8;
+	const bool bMaybeHandshakePacketSize = HandshakePacketBitsLeft >= MinBitsLeft && HandshakePacketBitsLeft <= MaxBitsLeft;
+	const bool bMaybeRestartHandshakePacket = RestartHandshakePacketBitsLeft >= MinBitsLeft && RestartHandshakePacketBitsLeft <= MaxBitsLeft;
+	const bool bMaybeRestartResponsePacketSize = RestartResponsePacketBitsLeft >= MinBitsLeft && RestartResponsePacketBitsLeft <= MaxBitsLeft;
+
+	OutResult.bRestartHandshake = !!Packet.ReadBit();
+
+	uint8 RemoteMinVersion = 0;
+	uint8 RemoteCurVersion = 0;
+	uint8 HandshakePacketType = 0;
+	EHandshakePacketType& HandshakePacketTypeEnum = OutResult.HandshakePacketType;
+
+	Packet << RemoteMinVersion;
+	Packet << RemoteCurVersion;
+	Packet << HandshakePacketType;
+	Packet << OutResult.RemoteSentHandshakePacketCount;
+
+	OutResult.RemoteMinVersion = static_cast<EHandshakeVersion>(RemoteMinVersion);
+	OutResult.RemoteCurVersion = static_cast<EHandshakeVersion>(RemoteCurVersion);
+	HandshakePacketTypeEnum = static_cast<EHandshakePacketType>(HandshakePacketType);
+
+	// Only accept handshake packets of roughly the right size
+	const bool bHandshakePacket = bMaybeHandshakePacketSize && (HandshakePacketTypeEnum == EHandshakePacketType::InitialPacket ||
+		HandshakePacketTypeEnum == EHandshakePacketType::Challenge || HandshakePacketTypeEnum == EHandshakePacketType::Response ||
+		HandshakePacketTypeEnum == EHandshakePacketType::Ack);
+
+	const bool bRestartHandshakePacket = bMaybeRestartHandshakePacket && HandshakePacketTypeEnum == EHandshakePacketType::RestartHandshake;
+	const bool bRestartResponsePacket = bMaybeRestartResponsePacketSize && HandshakePacketTypeEnum == EHandshakePacketType::RestartResponse;
 
 	// Only accept handshake packets of precisely the right size
-	if (bHandshakePacketSize || bRestartResponsePacketSize || bRestartResponseDiagnosticsPacketSize)
+	if (bHandshakePacket || bRestartResponsePacket)
 	{
-		bOutRestartHandshake = !!Packet.ReadBit();
-		OutSecretId = Packet.ReadBit();
+		OutResult.SecretId = Packet.ReadBit();
 
-		Packet << OutTimestamp;
+		Packet << OutResult.Timestamp;
 
-		Packet.Serialize(OutCookie, COOKIE_BYTE_SIZE);
+		Packet.Serialize(OutResult.Cookie, COOKIE_BYTE_SIZE);
 
-		if (bRestartResponsePacketSize || bRestartResponseDiagnosticsPacketSize)
+		if (bRestartResponsePacket)
 		{
-			Packet.Serialize(OutOrigCookie, COOKIE_BYTE_SIZE);
+			Packet.Serialize(OutResult.OrigCookie, COOKIE_BYTE_SIZE);
+		}
 
-#if RESTART_HANDSHAKE_DIAGNOSTICS
-			if (bRestartResponseDiagnosticsPacketSize)
+		bValidPacket = !Packet.IsError();
+	}
+	else if (bRestartHandshakePacket)
+	{
+		bValidPacket = !Packet.IsError() && OutResult.bRestartHandshake && Handler->Mode == Handler::Mode::Client;
+	}
+
+	if (bValidPacket)
+	{
+		Packet.SetAtEnd();
+	}
+
+	return bValidPacket;
+}
+
+bool StatelessConnectHandlerComponent::ParseHandshakePacketOriginal(FBitReader& Packet, FParsedHandshakeData& OutResult)
+{
+	using namespace UE::Net;
+
+	bool bValidPacket = false;
+	uint32 BitsLeft = Packet.GetBitsLeft();
+	bool bHandshakePacketSize = BitsLeft == (OriginalHandshakePacketSizeBits - 1);
+	bool bRestartResponsePacketSize = BitsLeft == (OriginalRestartResponseSizeBits - 1);
+
+	OutResult.RemoteMinVersion = EHandshakeVersion::Original;
+	OutResult.RemoteCurVersion = EHandshakeVersion::Original;
+
+	// Only accept handshake packets of precisely the right size
+	if (bHandshakePacketSize || bRestartResponsePacketSize)
+	{
+		OutResult.bRestartHandshake = !!Packet.ReadBit();
+		OutResult.SecretId = Packet.ReadBit();
+
+		Packet << OutResult.Timestamp;
+
+		Packet.Serialize(OutResult.Cookie, COOKIE_BYTE_SIZE);
+
+		if (bRestartResponsePacketSize)
+		{
+			OutResult.HandshakePacketType = EHandshakePacketType::RestartResponse;
+			Packet.Serialize(OutResult.OrigCookie, COOKIE_BYTE_SIZE);
+		}
+		else if (OutResult.Timestamp > 0.0)
+		{
+			if (Handler->Mode == Handler::Mode::Client)
 			{
-				Packet << HandshakeDiagnostics;
+				OutResult.HandshakePacketType = EHandshakePacketType::Challenge;
 			}
-#endif
+			else
+			{
+				OutResult.HandshakePacketType = EHandshakePacketType::Response;
+			}
+		}
+		else if (OutResult.Timestamp < 0.0)
+		{
+			OutResult.HandshakePacketType = EHandshakePacketType::Ack;
+		}
+		else
+		{
+			OutResult.HandshakePacketType = EHandshakePacketType::InitialPacket;
 		}
 
 		bValidPacket = !Packet.IsError();
 	}
-	else if (BitsLeft == (RESTART_HANDSHAKE_PACKET_SIZE_BITS - 1))
+	else if (BitsLeft == (OriginalRestartHandshakePacketSizeBits - 1))
 	{
-		bOutRestartHandshake = !!Packet.ReadBit();
-		bValidPacket = !Packet.IsError() && bOutRestartHandshake && Handler->Mode == Handler::Mode::Client;
+		OutResult.HandshakePacketType = EHandshakePacketType::RestartHandshake;
+		OutResult.bRestartHandshake = !!Packet.ReadBit();
+		bValidPacket = !Packet.IsError() && OutResult.bRestartHandshake && Handler->Mode == Handler::Mode::Client;
 	}
 
 	return bValidPacket;
@@ -1168,9 +1378,6 @@ bool StatelessConnectHandlerComponent::ParseHandshakePacket(FBitReader& Packet,
 
 void StatelessConnectHandlerComponent::GenerateCookie(TSharedPtr<const FInternetAddr> ClientAddress, uint8 SecretId, double Timestamp, uint8 (&OutCookie)[20])
 {
-	// @todo #JohnB: Add cpu stats tracking, like what Oodle does upon compression
-	//					NOTE: Being serverside, will only show up in .uprof, not on any 'stat' commands. Still necessary though.
-
 	TArray<uint8> CookieData;
 	FMemoryWriter CookieArc(CookieData);
 	FString ClientAddressString(ClientAddress->ToString(true));
@@ -1227,13 +1434,15 @@ int32 StatelessConnectHandlerComponent::GetReservedPacketBits() const
 
 void StatelessConnectHandlerComponent::Tick(float DeltaTime)
 {
+	using namespace UE::Net;
+
 	if (Handler->Mode == Handler::Mode::Client)
 	{
 		if (State != Handler::Component::State::Initialized && LastClientSendTimestamp != 0.0)
 		{
 			double LastSendTimeDiff = FPlatformTime::Seconds() - LastClientSendTimestamp;
 
-			if (LastSendTimeDiff > 1.0)
+			if (LastSendTimeDiff > UE::Net::HandshakeResendInterval)
 			{
 				const bool bRestartChallenge = Driver != nullptr && ((Driver->GetElapsedTime() - LastChallengeTimestamp) > MIN_COOKIE_LIFETIME);
 
@@ -1246,13 +1455,29 @@ void StatelessConnectHandlerComponent::Tick(float DeltaTime)
 				{
 					UE_LOG(LogHandshake, Verbose, TEXT("Initial handshake packet timeout - resending."));
 
-					NotifyHandshakeBegin();
+					EHandshakeVersion ResendVersion = static_cast<EHandshakeVersion>(CurrentHandshakeVersion);
+
+					// In case the server doesn't support the current handshake version, randomly switch between supported versions - if enabled
+					// (we don't know if the server supports the minimum version either, so pick from the full range).
+					// It's better for devs to explicitly hotfix the 'net.MinHandshakeVersion' value, instead of relying upon this fallback.
+					if (!!CVarNetDoHandshakeVersionFallback.GetValueOnAnyThread() && FMath::RandBool())
+					{
+						// Decrement the minimum version, based on the number of handshake packets sent - to select for higher supported versions
+						const int32 MinVersion = FMath::Max(MinSupportedHandshakeVersion, CurrentHandshakeVersion - SentHandshakePacketCount);
+
+						if (MinVersion != CurrentHandshakeVersion)
+						{
+							ResendVersion = static_cast<EHandshakeVersion>(FMath::RandRange(MinVersion, CurrentHandshakeVersion));
+						}
+					}
+
+					SendInitialPacket(ResendVersion);
 				}
 				else if (State == Handler::Component::State::InitializedOnLocal && LastTimestamp != 0.0)
 				{
 					UE_LOG(LogHandshake, Verbose, TEXT("Challenge response packet timeout - resending."));
 
-					SendChallengeResponse(LastSecretId, LastTimestamp, LastCookie);
+					SendChallengeResponse(LastRemoteHandshakeVersion, LastSecretId, LastTimestamp, LastCookie);
 				}
 			}
 		}
diff --git a/Engine/Source/Runtime/Engine/Public/PacketHandlers/StatelessConnectHandlerComponent.h b/Engine/Source/Runtime/Engine/Public/PacketHandlers/StatelessConnectHandlerComponent.h
--- a/Engine/Source/Runtime/Engine/Public/PacketHandlers/StatelessConnectHandlerComponent.h
+++ b/Engine/Source/Runtime/Engine/Public/PacketHandlers/StatelessConnectHandlerComponent.h
@@ -2,14 +2,9 @@
 
 #pragma once
 
-#include "CoreMinimal.h"
 #include "PacketHandler.h"
 #include "IPAddress.h"
 
-/** Temporary version-check macro for randomized packet sequence, for backwards-compatibility. */
-#define STATELESSCONNECT_HAS_RANDOM_SEQUENCE 1
-
-class UNetDriver;
 
 DECLARE_LOG_CATEGORY_EXTERN(LogHandshake, Log, All);
 
@@ -31,6 +26,34 @@ class UNetDriver;
 #define COOKIE_BYTE_SIZE 20
 
 
+namespace UE::Net
+{
+	/**
+	 * The different handshake protocol versions
+	 */
+	enum class EHandshakeVersion : uint8
+	{
+		Original		= 0,	// The original/unversioned handshake protocol
+		Randomized		= 1,	// The version of the handshake protocol with randomization, versioning and debug/diagnostic tweaks
+
+		Latest			= Randomized
+	};
+
+	/**
+	 * The different handshake packet types
+	 */
+	enum class EHandshakePacketType : uint8
+	{
+		InitialPacket		= 0,
+		Challenge			= 1,
+		Response			= 2,
+		Ack					= 3,
+		RestartHandshake	= 4,
+		RestartResponse		= 5
+	};
+}
+
+
 /**
  * PacketHandler component for implementing a stateless (non-memory-consuming) connection handshake
  *
@@ -38,6 +61,8 @@ class UNetDriver;
  */
 class ENGINE_API StatelessConnectHandlerComponent : public HandlerComponent
 {
+	using EHandshakeVersion = UE::Net::EHandshakeVersion;
+
 public:
 	/**
 	 * Base constructor
@@ -59,44 +84,59 @@ public:
 	void InitFromConnectionless(StatelessConnectHandlerComponent* InConnectionlessHandler);
 
 private:
+	/**
+	 * Constructs and sends the initial handshake packet, from the client to the server
+	 *
+	 * @param HandshakeVersion	The handshake format version to use when sending
+	 */
+	void SendInitialPacket(EHandshakeVersion HandshakeVersion);
+
 	/**
 	 * Constructs and sends the server response to the initial connect packet, from the server to the client.
 	 *
-	 * @param ClientAddress		The address of the client to send the challenge to.
+	 * @param ClientAddress						The address of the client to send the challenge to.
+	 * @param HandshakeVersion					The handshake format version to use when sending
+	 * @param ClientSentHandshakePacketCount	The number of handshake packets the client has sent (for debugging/packet-analysis)
 	 */
-	void SendConnectChallenge(TSharedPtr<const FInternetAddr> ClientAddress);
+	void SendConnectChallenge(TSharedPtr<const FInternetAddr> ClientAddress, EHandshakeVersion HandshakeVersion, uint8 ClientSentHandshakePacketCount);
 
 	/**
 	 * Constructs and sends the handshake challenge response packet, from the client to the server
 	 *
+	 * @param HandshakeVersion	The handshake format version to use when sending
 	 * @param InSecretId		Which of the two server HandshakeSecret values this uses
 	 * @param InTimestamp		The timestamp value to send
 	 * @param InCookie			The cookie value to send
 	 */
-	void SendChallengeResponse(uint8 InSecretId, double InTimestamp, uint8 InCookie[COOKIE_BYTE_SIZE]);
+	void SendChallengeResponse(EHandshakeVersion HandshakeVersion, uint8 InSecretId, double InTimestamp, uint8 InCookie[COOKIE_BYTE_SIZE]);
 
 	/**
 	 * Constructs and sends the server ack to a successful challenge response, from the server to the client.
 	 *
 	 * @param ClientAddress		The address of the client to send the ack to.
+	 * @param HandshakeVersion	The handshake format version to use when sending
+	 * @param ClientSentHandshakePacketCount	The number of handshake packets the client has sent (for debugging/packet-analysis)
 	 * @param InCookie			The cookie value to send
 	 */
-	void SendChallengeAck(TSharedPtr<const FInternetAddr> ClientAddress, uint8 InCookie[COOKIE_BYTE_SIZE]);
+	void SendChallengeAck(TSharedPtr<const FInternetAddr> ClientAddress, EHandshakeVersion HandshakeVersion, uint8 ClientSentHandshakePacketCount,
+							uint8 InCookie[COOKIE_BYTE_SIZE]);
 
 	/**
 	 * Constructs and sends a request to resend the cookie, from the server to the client.
 	 *
 	 * @param ClientAddress		The address of the client to send the request to.
+	 * @param HandshakeVersion	The handshake format version to use when sending
 	 */
-	void SendRestartHandshakeRequest(const TSharedPtr<const FInternetAddr> ClientAddress);
+	void SendRestartHandshakeRequest(const TSharedPtr<const FInternetAddr> ClientAddress, EHandshakeVersion HandshakeVersion);
 
 
 	/**
 	 * Pads the handshake packet, to match the PacketBitAlignment of the PacketHandler, so that it will parse correctly.
 	 *
 	 * @param HandshakePacket	The handshake packet to be aligned.
+	 * @param HandshakeVersion	The handshake format version to use when sending
 	 */
-	void CapHandshakePacket(FBitWriter& HandshakePacket);
+	void CapHandshakePacket(FBitWriter& HandshakePacket, EHandshakeVersion HandshakeVersion);
 
 public:
 	/**
@@ -204,19 +244,51 @@ protected:
 	virtual void Tick(float DeltaTime) override;
 
 private:
+
+	/**
+	 * Handshake data parsed from a packet
+	 */
+	struct FParsedHandshakeData
+	{
+		/** The minimum supported handshake protocol version of the remote side */
+		EHandshakeVersion RemoteMinVersion = EHandshakeVersion::Latest;
+
+		/** The current handshake protocol version the remote side used for communication */
+		EHandshakeVersion RemoteCurVersion = EHandshakeVersion::Latest;
+
+		/** The type of handshake packet */
+		UE::Net::EHandshakePacketType HandshakePacketType = UE::Net::EHandshakePacketType::InitialPacket;
+
+		/** The remote value of 'SentHandshakePacketCount' */
+		uint8 RemoteSentHandshakePacketCount = 0;
+
+		/** Whether or not this packet is a restart handshake packet */
+		bool bRestartHandshake = false;
+
+		/** Which of the two serverside HandshakeSecret values this is based on */
+		uint8 SecretId = 0;
+
+		/** The server timestamp, from the moment the challenge was sent (or 0.f if from the client) */
+		double Timestamp = 0.0;
+
+		/** A unique identifier, generated by the server, which the client must reply with (or 0, for initial packet) */
+		uint8 Cookie[COOKIE_BYTE_SIZE] = {};
+
+		/** If this is a restart handshake challenge response, this is the original handshake's cookie */
+		uint8 OrigCookie[COOKIE_BYTE_SIZE] = {};
+	};
+
+
 	/**
 	 * Parses an incoming handshake packet (does not parse the handshake bit though)
 	 *
 	 * @param Packet		The packet the handshake is being parsed from
-	 * @param bOutRestartHandshake	Whether or not this packet is a restart handshake packet
-	 * @param OutSecretId	Which of the two serverside HandshakeSecret values this is based on
-	 * @param OutTimestamp	The server timestamp, from the moment the challenge was sent (or 0.f if from the client)
-	 * @param OutCookie		A unique identifier, generated by the server, which the client must reply with (or 0, for initial packet)
-	 * @param OutOrigCookie			If this is a restart handshake challenge response, this is the original handshake's cookie
+	 * @param OutResult		The parsed handshake data, if successful
 	 * @return				Whether or not the handshake packet was parsed successfully
 	 */
-	bool ParseHandshakePacket(FBitReader& Packet, bool& bOutRestartHandshake, uint8& OutSecretId, double& OutTimestamp,
-								uint8 (&OutCookie)[COOKIE_BYTE_SIZE], uint8 (&OutOrigCookie)[COOKIE_BYTE_SIZE]);
+	bool ParseHandshakePacket(FBitReader& Packet, FParsedHandshakeData& OutResult);
+
+	bool ParseHandshakePacketOriginal(FBitReader& Packet, FParsedHandshakeData& OutResult);
 
 	/**
 	 * Takes the client address plus server timestamp, and outputs a deterministic cookie value
@@ -270,6 +342,9 @@ private:
 	/** The initial client sequence value, from the last successful handshake */
 	int32 LastClientSequence;
 
+	/** The minimum supported client version, for all connections - for setting minimum bounds for the restart handshake (not aware of version) */
+	EHandshakeVersion MinClientHandshakeVersion;
+
 
 	/** Clientside variables */
 
@@ -303,5 +378,11 @@ private:
 
 	/** The magic header which is prepended to all packets */
 	TBitArray<> MagicHeader;
+
+	/** The number of sent handshake packets, added to packets to aid packet dump debugging (may roll over) */
+	uint8 SentHandshakePacketCount = 0;
+
+	/** The last remote handshake version sent by the client/server */
+	EHandshakeVersion LastRemoteHandshakeVersion;
 };
 
