Avoid the setup cost to calling FODSCManager::AddThreadedShaderPipelineRequest if nothing will actually be requested.

diff --git a/Engine/Source/Runtime/Engine/Private/Materials/MaterialShared.cpp b/Engine/Source/Runtime/Engine/Private/Materials/MaterialShared.cpp
--- a/Engine/Source/Runtime/Engine/Private/Materials/MaterialShared.cpp
+++ b/Engine/Source/Runtime/Engine/Private/Materials/MaterialShared.cpp
@@ -3102,16 +3102,19 @@ bool FMaterial::TryGetShaders(const FMaterialShaderTypes& InTypes, const FVertex
 #if WITH_ODSC
 				if (FPlatformProperties::RequiresCookedData())
 				{
-					const FString MaterialName = GetFullPath();
-					const FString VFTypeName(InVertexFactoryType ? InVertexFactoryType->GetName() : TEXT(""));
-					const FString PipelineName(InTypes.PipelineType->GetName());
-					TArray<FString> ShaderStageNamesToCompile;
-					for (auto* ShaderType : InTypes.PipelineType->GetStages())
+					if (GODSCManager->IsHandlingRequests())
 					{
-						ShaderStageNamesToCompile.Add(ShaderType->GetName());
-					}
+						const FString MaterialName = GetFullPath();
+						const FString VFTypeName(InVertexFactoryType ? InVertexFactoryType->GetName() : TEXT(""));
+						const FString PipelineName(InTypes.PipelineType->GetName());
+						TArray<FString> ShaderStageNamesToCompile;
+						for (auto* ShaderType : InTypes.PipelineType->GetStages())
+						{
+							ShaderStageNamesToCompile.Add(ShaderType->GetName());
+						}
 
-					GODSCManager->AddThreadedShaderPipelineRequest(ShaderPlatform, GetFeatureLevel(), GetQualityLevel(), MaterialName, VFTypeName, PipelineName, ShaderStageNamesToCompile);
+						GODSCManager->AddThreadedShaderPipelineRequest(ShaderPlatform, GetFeatureLevel(), GetQualityLevel(), MaterialName, VFTypeName, PipelineName, ShaderStageNamesToCompile);
+					}
 				}
 				else 
 #endif
@@ -3152,14 +3155,16 @@ bool FMaterial::TryGetShaders(const FMaterialShaderTypes& InTypes, const FVertex
 #if WITH_ODSC
 					if (FPlatformProperties::RequiresCookedData())
 					{
+						if (GODSCManager->IsHandlingRequests())
+						{
+							const FString MaterialName = GetFullPath();
+							const FString VFTypeName(InVertexFactoryType ? InVertexFactoryType->GetName() : TEXT(""));
+							const FString PipelineName;
+							TArray<FString> ShaderStageNamesToCompile;
+							ShaderStageNamesToCompile.Add(ShaderType->GetName());
 
-						const FString MaterialName = GetFullPath();
-						const FString VFTypeName(InVertexFactoryType ? InVertexFactoryType->GetName() : TEXT(""));
-						const FString PipelineName;
-						TArray<FString> ShaderStageNamesToCompile;
-						ShaderStageNamesToCompile.Add(ShaderType->GetName());
-
-						GODSCManager->AddThreadedShaderPipelineRequest(ShaderPlatform, GetFeatureLevel(), GetQualityLevel(), MaterialName, VFTypeName, PipelineName, ShaderStageNamesToCompile);
+							GODSCManager->AddThreadedShaderPipelineRequest(ShaderPlatform, GetFeatureLevel(), GetQualityLevel(), MaterialName, VFTypeName, PipelineName, ShaderStageNamesToCompile);
+						}
 					}
 					else
 #endif
diff --git a/Engine/Source/Runtime/Engine/Private/ODSC/ODSCManager.cpp b/Engine/Source/Runtime/Engine/Private/ODSC/ODSCManager.cpp
--- a/Engine/Source/Runtime/Engine/Private/ODSC/ODSCManager.cpp
+++ b/Engine/Source/Runtime/Engine/Private/ODSC/ODSCManager.cpp
@@ -49,7 +49,7 @@ bool FODSCManager::Tick(float DeltaSeconds)
 {
 	QUICK_SCOPE_CYCLE_COUNTER(STAT_FODSCManager_Tick);
 
-	if (Thread)
+	if (IsHandlingRequests())
 	{
 		Thread->Wakeup();
 
@@ -72,7 +72,7 @@ bool FODSCManager::Tick(float DeltaSeconds)
 
 void FODSCManager::AddThreadedRequest(const TArray<FString>& MaterialsToCompile, const FString& ShaderTypesToLoad, EShaderPlatform ShaderPlatform, ERHIFeatureLevel::Type FeatureLevel, EMaterialQualityLevel::Type QualityLevel, ODSCRecompileCommand RecompileCommandType)
 {
-	if (Thread)
+	if (IsHandlingRequests())
 	{
 		Thread->AddRequest(MaterialsToCompile, ShaderTypesToLoad, ShaderPlatform, FeatureLevel, QualityLevel, RecompileCommandType);
 	}
@@ -80,7 +80,7 @@ void FODSCManager::AddThreadedRequest(const TArray<FString>& MaterialsToCompile,
 
 void FODSCManager::AddThreadedShaderPipelineRequest(EShaderPlatform ShaderPlatform, ERHIFeatureLevel::Type FeatureLevel, EMaterialQualityLevel::Type QualityLevel, const FString& MaterialName, const FString& VertexFactoryName, const FString& PipelineName, const TArray<FString>& ShaderTypeNames)
 {
-	if (Thread)
+	if (IsHandlingRequests())
 	{
 		Thread->AddShaderPipelineRequest(ShaderPlatform, FeatureLevel, QualityLevel, MaterialName, VertexFactoryName, PipelineName, ShaderTypeNames);
 	}
diff --git a/Engine/Source/Runtime/Engine/Public/ODSC/ODSCManager.h b/Engine/Source/Runtime/Engine/Public/ODSC/ODSCManager.h
--- a/Engine/Source/Runtime/Engine/Public/ODSC/ODSCManager.h
+++ b/Engine/Source/Runtime/Engine/Public/ODSC/ODSCManager.h
@@ -67,6 +67,10 @@ public:
 	 * @return false if no longer needs ticking
 	 */
 	void AddThreadedShaderPipelineRequest(EShaderPlatform ShaderPlatform, ERHIFeatureLevel::Type FeatureLevel, EMaterialQualityLevel::Type QualityLevel, const FString& MaterialName, const FString& VertexFactoryName, const FString& PipelineName, const TArray<FString>& ShaderTypeNames);
+
+	/** Returns true if we would actually add a request when calling AddThreadedShaderPipelineRequest. */
+	inline bool IsHandlingRequests() const { return Thread != nullptr; }
+
 private:
 
 	void OnEnginePreExit();
