Fix changelist history and property retirement cleanup when a bunch exceeds the max size check

diff --git a/Engine/Source/Runtime/Engine/Private/DataReplication.cpp b/Engine/Source/Runtime/Engine/Private/DataReplication.cpp
--- a/Engine/Source/Runtime/Engine/Private/DataReplication.cpp
+++ b/Engine/Source/Runtime/Engine/Private/DataReplication.cpp
@@ -1904,7 +1904,8 @@ void FObjectReplicator::PostSendBunch( FPacketIdRange & PacketRange, uint8 bReli
 	check(RepLayout);
 
 	// Don't update retirement records for reliable properties. This is ok to do only if we also pause replication on the channel until the acks have gone through.
-	const bool bSkipRetirementUpdate = OwningChannel->bPausedUntilReliableACK;
+	// Also clean up history and/or retirement records if the bunch failed to send
+	const bool bSkipRetirementUpdate = OwningChannel->bPausedUntilReliableACK || (PacketRange.First == INDEX_NONE && PacketRange.Last == INDEX_NONE);
 
 	const FRepLayout& LocalRepLayout = *RepLayout;
 
