#UE Since we already do not create a new child actor during a replay fastforward, we now also do not destroy it. The replay code was serializing a pointer to the destroyed actor after the fastforward. Also changed the default verbosity of the child actor component log to "Log" so that I can add new messaging. There were no log messages lower than warning verbosity at this time.

diff --git a/Engine/Source/Runtime/Engine/Private/Components/ChildActorComponent.cpp b/Engine/Source/Runtime/Engine/Private/Components/ChildActorComponent.cpp
--- a/Engine/Source/Runtime/Engine/Private/Components/ChildActorComponent.cpp
+++ b/Engine/Source/Runtime/Engine/Private/Components/ChildActorComponent.cpp
@@ -13,7 +13,7 @@
 #include "Iris/ReplicationSystem/PropertyReplicationFragment.h"
 #endif // UE_WITH_IRIS
 
-DEFINE_LOG_CATEGORY_STATIC(LogChildActorComponent, Warning, All);
+DEFINE_LOG_CATEGORY_STATIC(LogChildActorComponent, Log, All);
 
 ENGINE_API int32 GExperimentalAllowPerInstanceChildActorProperties = 0;
 
@@ -645,6 +645,7 @@ void UChildActorComponent::CreateChildActor(TFunction<void(AActor*)> CustomizerF
 			// If we're in a replay scrub and this is a startup actor and the class exists, assume the reference will be restored later
 			if (World->IsPlayingReplay() && World->GetDemoNetDriver()->IsRestoringStartupActors() && MyOwner && MyOwner->IsNetStartupActor() && ChildActorClass)
 			{
+				UE_LOG(LogChildActorComponent, Display, TEXT("Not creating child actor due to replay! Comp:%s ChildActor:%s"), *GetPathName(), *GetPathNameSafe(ChildActor));
 				return;
 			}
 
@@ -755,6 +756,18 @@ void UChildActorComponent::CreateChildActor(TFunction<void(AActor*)> CustomizerF
 
 void UChildActorComponent::DestroyChildActor()
 {
+	UWorld* LocalWorld = GetWorld();
+	if (LocalWorld != nullptr)
+	{
+		AActor* MyOwner = GetOwner();
+		// If we're in a replay scrub and this is a startup actor and the class exists, assume the reference will be restored later
+		if (LocalWorld->IsPlayingReplay() && LocalWorld->GetDemoNetDriver()->IsRestoringStartupActors() && MyOwner && MyOwner->IsNetStartupActor() && ChildActorClass)
+		{
+			UE_LOG(LogChildActorComponent, Display, TEXT("Not creating destroying actor due to replay! Comp:%s ChildActor:%s"), *GetPathName(), *GetPathNameSafe(ChildActor));
+			return;
+		}
+	}
+
 	// If we own an Actor, kill it now unless we don't have authority on it, for that we rely on the server
 	// If the level is being removed then don't destroy the child actor so re-adding it doesn't
 	// need to create a new actor
diff --git a/Engine/Source/Runtime/Engine/Private/DemoNetDriver.cpp b/Engine/Source/Runtime/Engine/Private/DemoNetDriver.cpp
--- a/Engine/Source/Runtime/Engine/Private/DemoNetDriver.cpp
+++ b/Engine/Source/Runtime/Engine/Private/DemoNetDriver.cpp
@@ -3703,49 +3703,52 @@ bool UDemoNetDriver::FastForwardLevels(const FGotoResult& GotoResult)
 	// If we've gotten this far, it means we should have something to process.
 	check(ReadPacketsHelper.Packets.Num() > 0);
 
-	for (ULevel* Level : ReplayHelper.LevelsPendingFastForward)
 	{
-		// Track the appropriate level, and mark it as ready.
-		FReplayHelper::FLevelStatus& LevelStatus = ReplayHelper.GetLevelStatus(ReplayHelper.GetLevelPackageName(*Level));
-		LevelIndices.Add(LevelStatus.LevelIndex);
-		LevelStatus.bIsReady = true;
-
-		TSet<TWeakObjectPtr<AActor>> LevelActors;
-		for (AActor* Actor : Level->Actors)
+		TGuardValue<bool> RestoringStartupActors(bIsRestoringStartupActors, true);
+		for (ULevel* Level : ReplayHelper.LevelsPendingFastForward)
 		{
-			if (Actor == nullptr || !Actor->IsNetStartupActor())
-			{
-				continue;
-			}
-			else if (ReplayHelper.PlaybackDeletedNetStartupActors.Contains(Actor->GetFullName()))
-			{
-				// Put this actor on the rollback list so we can undelete it during future scrubbing,
-				// then delete it.
-				QueueNetStartupActorForRollbackViaDeletion(Actor);
-				World->DestroyActor(Actor, true);
-			}
-			else
+			// Track the appropriate level, and mark it as ready.
+			FReplayHelper::FLevelStatus& LevelStatus = ReplayHelper.GetLevelStatus(ReplayHelper.GetLevelPackageName(*Level));
+			LevelIndices.Add(LevelStatus.LevelIndex);
+			LevelStatus.bIsReady = true;
+
+			TSet<TWeakObjectPtr<AActor>> LevelActors;
+			for (AActor* Actor : Level->Actors)
 			{
-				if (RollbackNetStartupActors.Contains(Actor->GetFullName()))
+				if (Actor == nullptr || !Actor->IsNetStartupActor())
 				{
+					continue;
+				}
+				else if (ReplayHelper.PlaybackDeletedNetStartupActors.Contains(Actor->GetFullName()))
+				{
+					// Put this actor on the rollback list so we can undelete it during future scrubbing,
+					// then delete it.
+					QueueNetStartupActorForRollbackViaDeletion(Actor);
 					World->DestroyActor(Actor, true);
 				}
 				else
 				{
-					StartupActors.Add(Actor);
+					if (RollbackNetStartupActors.Contains(Actor->GetFullName()))
+					{
+						World->DestroyActor(Actor, true);
+					}
+					else
+					{
+						StartupActors.Add(Actor);
+					}
 				}
 			}
-		}
 
-		TArray<AActor*> SpawnedActors;
-		RespawnNecessaryNetStartupActors(SpawnedActors, Level);
+			TArray<AActor*> SpawnedActors;
+			RespawnNecessaryNetStartupActors(SpawnedActors, Level);
 
-		for (AActor* Actor : SpawnedActors)
-		{
-			StartupActors.Add(Actor);
-		}
+			for (AActor* Actor : SpawnedActors)
+			{
+				StartupActors.Add(Actor);
+			}
 
-		LocalLevels.Add(Level);
+			LocalLevels.Add(Level);
+		}
 	}
 
 	ReplayHelper.LevelsPendingFastForward.Reset();
@@ -4098,12 +4101,16 @@ bool UDemoNetDriver::LoadCheckpoint(const FGotoResult& GotoResult)
 	PlaybackPackets.Empty();
 	ReplayHelper.PlaybackFrames.Empty();
 
-	// Destroy startup actors that need to rollback via being destroyed and re-created
-	for (FActorIterator ActorIt(World); ActorIt; ++ActorIt)
 	{
-		if (RollbackNetStartupActors.Contains(ActorIt->GetFullName()))
+		TGuardValue<bool> RestoringStartupActors(bIsRestoringStartupActors, true);
+
+		// Destroy startup actors that need to rollback via being destroyed and re-created
+		for (FActorIterator ActorIt(World); ActorIt; ++ActorIt)
 		{
-			World->DestroyActor(*ActorIt, true);
+			if (RollbackNetStartupActors.Contains(ActorIt->GetFullName()))
+			{
+				World->DestroyActor(*ActorIt, true);
+			}
 		}
 	}
 
@@ -4351,6 +4358,7 @@ bool UDemoNetDriver::LoadCheckpoint(const FGotoResult& GotoResult)
 	
 	if (World != nullptr)
 	{
+		TGuardValue<bool> RestoringStartupActors(bIsRestoringStartupActors, true);
 		// Destroy startup actors that shouldn't exist past this checkpoint
 		for (FActorIterator ActorIt( World ); ActorIt; ++ActorIt)
 		{
