[HLOD / Material Baking] Do not bypass vertex interpolator nodes when baking materials
* Fixes failures to compile export proxy materials that requires vertex interpolator nodes to be valid
* Updated HLOD base hash to force a rebuild

diff --git a/Engine/Source/Runtime/Engine/Private/HLODProxy.cpp b/Engine/Source/Runtime/Engine/Private/HLODProxy.cpp
--- a/Engine/Source/Runtime/Engine/Private/HLODProxy.cpp
+++ b/Engine/Source/Runtime/Engine/Private/HLODProxy.cpp
@@ -496,7 +496,7 @@ uint32 UHLODProxy::GetCRC(UStaticMeshComponent* InComponent, uint32 InCRC, const
 }
 
 // Key that forms the basis of the HLOD proxy key. Bump this key (i.e. generate a new GUID) when you want to force a rebuild of ALL HLOD proxies
-#define HLOD_PROXY_BASE_KEY		TEXT("1623F71084794B7A925C65D479104CA6")
+#define HLOD_PROXY_BASE_KEY		TEXT("498A8FFD64744FA3B2E7C7589581F91A")
 
 FName UHLODProxy::GenerateKeyForActor(const ALODActor* LODActor, bool bMustUndoLevelTransform)
 {
diff --git a/Engine/Source/Runtime/Engine/Public/MaterialCompiler.h b/Engine/Source/Runtime/Engine/Public/MaterialCompiler.h
--- a/Engine/Source/Runtime/Engine/Public/MaterialCompiler.h
+++ b/Engine/Source/Runtime/Engine/Public/MaterialCompiler.h
@@ -229,7 +229,7 @@ public:
 	inline bool IsVertexInterpolatorBypass() const
 	{
 		const EMaterialCompilerType Type = GetCompilerType();
-		return Type == EMaterialCompilerType::MaterialProxy || Type == EMaterialCompilerType::Lightmass;
+		return Type == EMaterialCompilerType::Lightmass;
 	}
 
 	virtual EMaterialValueType GetType(int32 Code) = 0;
