Correction to SceneViewport to use the cached display metrics instead of the Initial ones which could cause issues when monitor settings or resolution changes when the engine is running.

diff --git a/Engine/Source/Runtime/Engine/Private/Slate/SceneViewport.cpp b/Engine/Source/Runtime/Engine/Private/Slate/SceneViewport.cpp
--- a/Engine/Source/Runtime/Engine/Private/Slate/SceneViewport.cpp
+++ b/Engine/Source/Runtime/Engine/Private/Slate/SceneViewport.cpp
@@ -1392,7 +1392,7 @@ void FSceneViewport::ResizeFrame(uint32 NewWindowSizeX, uint32 NewWindowSizeY, E
 				else
 				{
 					FDisplayMetrics DisplayMetrics;
-					FSlateApplication::Get().GetInitialDisplayMetrics(DisplayMetrics);
+					FSlateApplication::Get().GetCachedDisplayMetrics(DisplayMetrics);
 
 					if (DisplayMetrics.MonitorInfo.Num() > 0)
 					{
