Remove additive base pose during cook when enabling a.StripAdditiveRefPose

diff --git a/Engine/Source/Runtime/Engine/Private/Animation/AnimSequence.cpp b/Engine/Source/Runtime/Engine/Private/Animation/AnimSequence.cpp
--- a/Engine/Source/Runtime/Engine/Private/Animation/AnimSequence.cpp
+++ b/Engine/Source/Runtime/Engine/Private/Animation/AnimSequence.cpp
@@ -106,7 +106,28 @@ static FAutoConsoleVariableRef CVarOddFrameStripping(
 	GPerformFrameStrippingOddFramedAnimations,
 	TEXT("1 = When frame stripping apply to animations with an odd number of frames too. 0 = only even framed animations"));
 
+int32 GStripAdditiveRefPose = 0;
+static FAutoConsoleVariableRef CVarStripAdditiveRefPose(
+	TEXT("a.StripAdditiveRefPose"),
+	GStripAdditiveRefPose,
+	TEXT("1 = Strip additive ref poses on cook. 0 = off"));
 
+namespace UE::Anim::Private
+{
+	bool ShouldStripAdditiveRefPose()
+	{
+		return GStripAdditiveRefPose != 0;
+	}
+
+	bool IsAdditiveRefPoseStripped()
+	{
+#if WITH_EDITOR
+		return false;
+#else
+		return ShouldStripAdditiveRefPose();
+#endif
+	}
+} // namespace UE::Anim::Private
 
 #if WITH_EDITOR
 
@@ -596,8 +617,25 @@ void UAnimSequence::Serialize(FArchive& Ar)
 		PRAGMA_ENABLE_DEPRECATION_WARNINGS
 	}
 
+	UAnimSequence* StrippedRefPoseSeq = nullptr;
+	if (Ar.IsCooking() && Ar.IsSaving())
+	{
+		if (UE::Anim::Private::ShouldStripAdditiveRefPose() && GetAdditiveAnimType() != EAdditiveAnimationType::AAT_None)
+		{
+			// Strip the additive base before UPROPERTY serialization in Super::Serialize
+			StrippedRefPoseSeq = RefPoseSeq;
+			RefPoseSeq = nullptr;
+		}
+	}
+
 	Super::Serialize(Ar);
 
+	if (StrippedRefPoseSeq)
+	{
+		// Restore after UPROPERTY serialization to avoid compression requests missing an additive base pose. 
+		RefPoseSeq = StrippedRefPoseSeq;
+	}
+
 	if (Ar.IsCooking())
 	{
 		PRAGMA_DISABLE_DEPRECATION_WARNINGS
@@ -2511,9 +2549,9 @@ bool UAnimSequence::IsValidAdditive() const
 		case ABPT_RefPose:
 			return true;
 		case ABPT_AnimScaled:
-			return (RefPoseSeq != NULL);
+			return UE::Anim::Private::IsAdditiveRefPoseStripped() || (RefPoseSeq != nullptr);
 		case ABPT_AnimFrame:
-			return (RefPoseSeq != NULL) && (RefFrameIndex >= 0);
+			return (UE::Anim::Private::IsAdditiveRefPoseStripped() || RefPoseSeq != nullptr) && (RefFrameIndex >= 0);
 		case ABPT_LocalAnimFrame:
 			return (RefFrameIndex >= 0);
 		default:
