UModel: Resolve issues with non-inclusive terminology

diff --git a/Engine/Source/Runtime/Engine/Private/Model.cpp b/Engine/Source/Runtime/Engine/Private/Model.cpp
--- a/Engine/Source/Runtime/Engine/Private/Model.cpp
+++ b/Engine/Source/Runtime/Engine/Private/Model.cpp
@@ -409,63 +409,63 @@ void UModel::PostEditUndo()
 	Super::PostEditUndo();
 }
 
-void UModel::ModifySurf( int32 InIndex, bool UpdateMaster )
+void UModel::ModifySurf( int32 InIndex, bool UpdateBrushes )
 {
 	Modify();
 
 	FBspSurf& Surf = Surfs[InIndex];
-	if( UpdateMaster && Surf.Actor )
+	if( UpdateBrushes && Surf.Actor )
 	{
 		Surf.Actor->Brush->Modify();
 	}
 }
 
-void UModel::ModifyAllSurfs( bool UpdateMaster )
+void UModel::ModifyAllSurfs( bool UpdateBrushes )
 {
 	Modify();
 
-	if (UpdateMaster)
+	if (UpdateBrushes)
 	{
-		TArray<UModel*> MasterModels;
-		MasterModels.Reset(Surfs.Num());
+		TArray<UModel*> Brushes;
+		Brushes.Reset(Surfs.Num());
 
 		for (const FBspSurf& Surf : Surfs)
 		{
 			if (Surf.Actor)
 			{
 				check(Surf.Actor->Brush);
-				MasterModels.AddUnique(Surf.Actor->Brush);
+				Brushes.AddUnique(Surf.Actor->Brush);
 			}
 		}
 
-		for (UModel* MasterModel : MasterModels)
+		for (UModel* Brush: Brushes)
 		{
-			MasterModel->Modify();
+			Brush->Modify();
 		}
 	}
 }
 
-void UModel::ModifySelectedSurfs( bool UpdateMaster )
+void UModel::ModifySelectedSurfs( bool UpdateBrushes )
 {
 	Modify();
 
-	if (UpdateMaster)
+	if (UpdateBrushes)
 	{
-		TArray<UModel*> MasterModels;
-		MasterModels.Reset(Surfs.Num());
+		TArray<UModel*> Brushes;
+		Brushes.Reset(Surfs.Num());
 
 		for (const FBspSurf& Surf : Surfs)
 		{
 			if (Surf.Actor && (Surf.PolyFlags & PF_Selected))
 			{
 				check(Surf.Actor->Brush);
-				MasterModels.AddUnique(Surf.Actor->Brush);
+				Brushes.AddUnique(Surf.Actor->Brush);
 			}
 		}
 
-		for (UModel* MasterModel : MasterModels)
+		for (UModel* Brush : Brushes)
 		{
-			MasterModel->Modify();
+			Brush->Modify();
 		}
 	}
 }
diff --git a/Engine/Source/Runtime/Engine/Public/Model.h b/Engine/Source/Runtime/Engine/Public/Model.h
--- a/Engine/Source/Runtime/Engine/Public/Model.h
+++ b/Engine/Source/Runtime/Engine/Public/Model.h
@@ -540,9 +540,9 @@ public:
 	int32 BuildVertexBuffers();
 
 	// UModel transactions.
-	ENGINE_API void ModifySelectedSurfs( bool UpdateMaster );
-	void ModifyAllSurfs( bool UpdateMaster );
-	ENGINE_API void ModifySurf( int32 InIndex, bool UpdateMaster );
+	ENGINE_API void ModifySelectedSurfs( bool UpdateBrushes );
+	void ModifyAllSurfs( bool UpdateBrushes );
+	ENGINE_API void ModifySurf( int32 InIndex, bool UpdateBrushes );
 	ENGINE_API bool HasSelectedSurfaces() const;
 #endif
 
