Chaos - add stats for RBAN body, joint, collision counts

#pf 62cf11f67822772758df26fb

diff --git a/Engine/Source/Runtime/Engine/Private/PhysicsEngine/ImmediatePhysics/ImmediatePhysicsChaos/ImmediatePhysicsSimulation_Chaos.cpp b/Engine/Source/Runtime/Engine/Private/PhysicsEngine/ImmediatePhysics/ImmediatePhysicsChaos/ImmediatePhysicsSimulation_Chaos.cpp
--- a/Engine/Source/Runtime/Engine/Private/PhysicsEngine/ImmediatePhysics/ImmediatePhysicsChaos/ImmediatePhysicsSimulation_Chaos.cpp
+++ b/Engine/Source/Runtime/Engine/Private/PhysicsEngine/ImmediatePhysics/ImmediatePhysicsChaos/ImmediatePhysicsSimulation_Chaos.cpp
@@ -31,6 +31,13 @@
 DECLARE_CYCLE_STAT(TEXT("FSimulation::Simulate_Chaos"), STAT_ImmediateSimulate_Chaos, STATGROUP_ImmediatePhysics);
 DECLARE_CYCLE_STAT(TEXT("FSimulation::Simulate_Chaos::InertiaConditioning"), STAT_ImmediateSimulate_Chaos_InertiaConditioning, STATGROUP_ImmediatePhysics);
 
+// Stat Counters
+DECLARE_DWORD_COUNTER_STAT(TEXT("FSimulation::NumBodies"), STAT_ImmediateSimulate_ChaosCounter_NumBodies, STATGROUP_ImmediatePhysicsCounters);
+DECLARE_DWORD_COUNTER_STAT(TEXT("FSimulation::NumDynamicBodies"), STAT_ImmediateSimulate_ChaosCounter_NumDynamicBodies, STATGROUP_ImmediatePhysicsCounters);
+DECLARE_DWORD_COUNTER_STAT(TEXT("FSimulation::NumKinematicBodies"), STAT_ImmediateSimulate_ChaosCounter_NumKinematicBodies, STATGROUP_ImmediatePhysicsCounters);
+DECLARE_DWORD_COUNTER_STAT(TEXT("FSimulation::NumContacts"), STAT_ImmediateSimulate_ChaosCounter_NumContacts, STATGROUP_ImmediatePhysicsCounters);
+DECLARE_DWORD_COUNTER_STAT(TEXT("FSimulation::NumJoints"), STAT_ImmediateSimulate_ChaosCounter_NumJoints, STATGROUP_ImmediatePhysicsCounters);
+
 //////////////////////////////////////////////////////////////////////////
 // @todo(ccaulfield): remove when finished
 //
@@ -983,6 +990,8 @@ namespace ImmediatePhysics_Chaos
 		Implementation->Evolution.SetSimulationSpace(Implementation->SimulationSpace);
 		Implementation->Evolution.Advance(StepTime, NumSteps, RewindTime);
 
+		UpdateStatCounters();
+
 #if CHAOS_DEBUG_DRAW
 		if (bChaosImmediate_DebugDrawOnSimulate)
 		{
@@ -1056,6 +1065,15 @@ namespace ImmediatePhysics_Chaos
 		}
 	}
 
+	void FSimulation::UpdateStatCounters()
+	{
+		INC_DWORD_STAT_BY(STAT_ImmediateSimulate_ChaosCounter_NumBodies, Implementation->Particles.GetNonDisabledView().Num());
+		INC_DWORD_STAT_BY(STAT_ImmediateSimulate_ChaosCounter_NumDynamicBodies, Implementation->Particles.GetNonDisabledDynamicView().Num());
+		INC_DWORD_STAT_BY(STAT_ImmediateSimulate_ChaosCounter_NumKinematicBodies, Implementation->Particles.GetActiveKinematicParticlesView().Num());
+		INC_DWORD_STAT_BY(STAT_ImmediateSimulate_ChaosCounter_NumContacts, Implementation->Collisions.NumConstraints());
+		INC_DWORD_STAT_BY(STAT_ImmediateSimulate_ChaosCounter_NumJoints, Implementation->Joints.NumConstraints());
+	}
+
 	void FSimulation::DebugDrawStaticParticles()
 	{
 #if CHAOS_DEBUG_DRAW
diff --git a/Engine/Source/Runtime/Engine/Public/Physics/ImmediatePhysics/ImmediatePhysicsChaos/ImmediatePhysicsSimulation_Chaos.h b/Engine/Source/Runtime/Engine/Public/Physics/ImmediatePhysics/ImmediatePhysicsChaos/ImmediatePhysicsSimulation_Chaos.h
--- a/Engine/Source/Runtime/Engine/Public/Physics/ImmediatePhysics/ImmediatePhysicsChaos/ImmediatePhysicsSimulation_Chaos.h
+++ b/Engine/Source/Runtime/Engine/Public/Physics/ImmediatePhysics/ImmediatePhysicsChaos/ImmediatePhysicsSimulation_Chaos.h
@@ -105,6 +105,7 @@ namespace ImmediatePhysics_Chaos
 		void UpdateActivePotentiallyCollidingPairs();
 		FReal UpdateStepTime(const FReal DeltaTime, const FReal MaxStepTime);
 
+		void UpdateStatCounters();
 		void DebugDrawStaticParticles();
 		void DebugDrawKinematicParticles();
 		void DebugDrawDynamicParticles();
diff --git a/Engine/Source/Runtime/Engine/Public/Physics/ImmediatePhysics/ImmediatePhysicsStats.h b/Engine/Source/Runtime/Engine/Public/Physics/ImmediatePhysics/ImmediatePhysicsStats.h
--- a/Engine/Source/Runtime/Engine/Public/Physics/ImmediatePhysics/ImmediatePhysicsStats.h
+++ b/Engine/Source/Runtime/Engine/Public/Physics/ImmediatePhysics/ImmediatePhysicsStats.h
@@ -5,5 +5,6 @@
 #include "Stats/Stats2.h"
 
 DECLARE_STATS_GROUP(TEXT("Immediate Physics"), STATGROUP_ImmediatePhysics, STATCAT_Advanced);
+DECLARE_STATS_GROUP(TEXT("Immediate Physics Counters"), STATGROUP_ImmediatePhysicsCounters, STATCAT_Advanced);
 
 DECLARE_CYCLE_STAT_EXTERN(TEXT("RigidBodyNodeInitTime"), STAT_RigidBodyNodeInitTime, STATGROUP_ImmediatePhysics,);
