Fix cook on the fly always using parent material shader.

diff --git a/Engine/Source/Runtime/Engine/Private/Materials/MaterialInstance.cpp b/Engine/Source/Runtime/Engine/Private/Materials/MaterialInstance.cpp
--- a/Engine/Source/Runtime/Engine/Private/Materials/MaterialInstance.cpp
+++ b/Engine/Source/Runtime/Engine/Private/Materials/MaterialInstance.cpp
@@ -2276,7 +2276,7 @@ void UMaterialInstance::CacheGivenTypesForCooking(EShaderPlatform ShaderPlatform
 			QualityLevel = GetCachedScalabilityCVars().MaterialQualityLevel;
 		}
 
-		FMaterialResource* CurrentResource = FindOrCreateMaterialResource(StaticPermutationMaterialResources, BaseMaterial, nullptr, FeatureLevel, QualityLevel);
+		FMaterialResource* CurrentResource = FindOrCreateMaterialResource(StaticPermutationMaterialResources, BaseMaterial, this, FeatureLevel, QualityLevel);
 		check(CurrentResource);
 
 		// Prepare the resource for compilation, but don't compile the completed shader map.
