Inclusive Terminology Push - Part of UE-158645, UE-158525, UE-158532
Renamed / deprecated master/slave terms at SkinnedMeshComponent and SkeletalMeshComponent (plus all the users)

#tests Verified the properties transfer correctly after modifying them without the fix and loading the character with the component with the fix applied

diff --git a/Engine/Config/BaseEngine.ini b/Engine/Config/BaseEngine.ini
--- a/Engine/Config/BaseEngine.ini
+++ b/Engine/Config/BaseEngine.ini
@@ -1099,6 +1099,15 @@ DefaultWorldPartitionHLODLayer=/Engine/Maps/Templates/HLODs/HLODLayer_Instanced.
 ; Inclusive Terminology Pass
 +PropertyRedirects=(OldName="SimSpaceSettings.MasterAlpha", NewName="SimSpaceSettings.WorldAlpha")
 
++PropertyRedirects=(OldName="SkinnedMeshComponent.MasterPoseComponent", NewName="LeaderPoseComponent")
++PropertyRedirects=(OldName="SkinnedMeshComponent.bUseBoundsFromMasterPoseComponent", NewName="bUseBoundsFromLeaderPoseComponent")
++PropertyRedirects=(OldName="SkinnedMeshComponent.bIgnoreMasterPoseComponentLOD", NewName="bIgnoreLeaderPoseComponentLOD")
++FunctionRedirects=(OldName="SkinnedMeshComponent.SetMasterPoseComponent", NewName="SetLeaderPoseComponent")
+
++PropertyRedirects=(OldName="SkeletalMeshComponent.bPropagateCurvesToSlaves", NewName="bPropagateCurvesToFollowers")
++FunctionRedirects=(OldName="SkeletalMeshComponent.BindClothToMasterPoseComponent", NewName="BindClothToLeaderPoseComponent")
++FunctionRedirects=(OldName="SkeletalMeshComponent.UnbindClothFromMasterPoseComponent", NewName="UnbindClothFromLeaderPoseComponent")
+
 [CoreUObject.Metadata]
 ; Note: UnrealHeaderTool should be rerun after making changes to MetadataRedirects to catch any keys specified in class headers
 +MetadataRedirects=(OldKey="K2Protected", NewKey="BlueprintProtected")
diff --git a/Engine/Plugins/Animation/ControlRig/Source/ControlRig/Private/ControlRigComponent.cpp b/Engine/Plugins/Animation/ControlRig/Source/ControlRig/Private/ControlRigComponent.cpp
--- a/Engine/Plugins/Animation/ControlRig/Source/ControlRig/Private/ControlRigComponent.cpp
+++ b/Engine/Plugins/Animation/ControlRig/Source/ControlRig/Private/ControlRigComponent.cpp
@@ -1647,7 +1647,7 @@ void UControlRigComponent::HandleControlRigPreConstructionEvent(UControlRig* InC
 	{
 		SkeletalMeshComponent->TickAnimation(0.f, false);
 		SkeletalMeshComponent->RefreshBoneTransforms();
-		SkeletalMeshComponent->RefreshSlaveComponents();
+		SkeletalMeshComponent->RefreshFollowerComponents();
 		SkeletalMeshComponent->UpdateComponentToWorld();
 		SkeletalMeshComponent->FinalizeBoneTransform();
 		SkeletalMeshComponent->MarkRenderTransformDirty();
diff --git a/Engine/Plugins/Animation/ControlRig/Source/ControlRig/Private/Sequencer/MovieSceneControlRigParameterTemplate.cpp b/Engine/Plugins/Animation/ControlRig/Source/ControlRig/Private/Sequencer/MovieSceneControlRigParameterTemplate.cpp
--- a/Engine/Plugins/Animation/ControlRig/Source/ControlRig/Private/Sequencer/MovieSceneControlRigParameterTemplate.cpp
+++ b/Engine/Plugins/Animation/ControlRig/Source/ControlRig/Private/Sequencer/MovieSceneControlRigParameterTemplate.cpp
@@ -832,7 +832,7 @@ struct FControlRigParameterPreAnimatedTokenProducer : IMovieScenePreAnimatedToke
 							SkeletalMeshComponent->TickAnimation(0.f, false);
 
 							SkeletalMeshComponent->RefreshBoneTransforms();
-							SkeletalMeshComponent->RefreshSlaveComponents();
+							SkeletalMeshComponent->RefreshFollowerComponents();
 							SkeletalMeshComponent->UpdateComponentToWorld();
 							SkeletalMeshComponent->FinalizeBoneTransform();
 							SkeletalMeshComponent->MarkRenderTransformDirty();
diff --git a/Engine/Plugins/Animation/ControlRig/Source/ControlRigEditor/Private/EditMode/ControlRigEditMode.cpp b/Engine/Plugins/Animation/ControlRig/Source/ControlRigEditor/Private/EditMode/ControlRigEditMode.cpp
--- a/Engine/Plugins/Animation/ControlRig/Source/ControlRigEditor/Private/EditMode/ControlRigEditMode.cpp
+++ b/Engine/Plugins/Animation/ControlRig/Source/ControlRigEditor/Private/EditMode/ControlRigEditMode.cpp
@@ -4172,7 +4172,7 @@ void FControlRigEditMode::TickManipulatableObjects(float DeltaTime)
 					// have to TickAnimation() to ensure that PreUpdate() is called on all anim instances
 					SkeletalMeshToUpdate->TickAnimation(0.0f, false);
 					SkeletalMeshToUpdate->RefreshBoneTransforms();
-					SkeletalMeshToUpdate->RefreshSlaveComponents();
+					SkeletalMeshToUpdate->RefreshFollowerComponents	();
 					SkeletalMeshToUpdate->UpdateComponentToWorld();
 					SkeletalMeshToUpdate->FinalizeBoneTransform();
 					SkeletalMeshToUpdate->MarkRenderTransformDirty();
diff --git a/Engine/Plugins/Animation/IKRig/Source/IKRig/Private/AnimNodes/AnimNode_RetargetPoseFromMesh.cpp b/Engine/Plugins/Animation/IKRig/Source/IKRig/Private/AnimNodes/AnimNode_RetargetPoseFromMesh.cpp
--- a/Engine/Plugins/Animation/IKRig/Source/IKRig/Private/AnimNodes/AnimNode_RetargetPoseFromMesh.cpp
+++ b/Engine/Plugins/Animation/IKRig/Source/IKRig/Private/AnimNodes/AnimNode_RetargetPoseFromMesh.cpp
@@ -325,12 +325,12 @@ void FAnimNode_RetargetPoseFromMesh::CopyBoneTransformsFromSource(USkeletalMeshC
 
 TObjectPtr<USkeletalMeshComponent> FAnimNode_RetargetPoseFromMesh::GetComponentToCopyPoseFrom() const
 {
-	// if our source is running under master-pose, then get bone data from there
+	// if our source is running under leader-pose, then get bone data from there
 	if (SourceMeshComponent.IsValid())
 	{
-		if(USkeletalMeshComponent* MasterPoseComponent = Cast<USkeletalMeshComponent>(SourceMeshComponent->MasterPoseComponent.Get()))
+		if(USkeletalMeshComponent* LeaderPoseComponent = Cast<USkeletalMeshComponent>(SourceMeshComponent->LeaderPoseComponent.Get()))
 		{
-			return MasterPoseComponent;
+			return LeaderPoseComponent;
 		}
 	}
 	
diff --git a/Engine/Plugins/Developer/AnimationSharing/Source/AnimationSharing/Private/AnimationSharingInstances.cpp b/Engine/Plugins/Developer/AnimationSharing/Source/AnimationSharing/Private/AnimationSharingInstances.cpp
--- a/Engine/Plugins/Developer/AnimationSharing/Source/AnimationSharing/Private/AnimationSharingInstances.cpp
+++ b/Engine/Plugins/Developer/AnimationSharing/Source/AnimationSharing/Private/AnimationSharingInstances.cpp
@@ -20,9 +20,9 @@ void UAnimSharingStateInstance::GetInstancedActors(TArray<class AActor*>& Actors
 		if (StateData.Components.IsValidIndex(ComponentIndex))
 		{
 			USkeletalMeshComponent* Component = StateData.Components[ComponentIndex];
-			const TArray<TWeakObjectPtr<USkinnedMeshComponent>>& SlaveComponents = Component->GetSlavePoseComponents();
+			const TArray<TWeakObjectPtr<USkinnedMeshComponent>>& FollowerComponents = Component->GetFollowerPoseComponents();
 
-			Algo::TransformIf(SlaveComponents, Actors,
+			Algo::TransformIf(FollowerComponents, Actors,
 				[&Actors](const TWeakObjectPtr<USkinnedMeshComponent>& WeakPtr)
 			{
 				// Needs to be valid and unique
diff --git a/Engine/Plugins/Developer/AnimationSharing/Source/AnimationSharing/Private/AnimationSharingManager.cpp b/Engine/Plugins/Developer/AnimationSharing/Source/AnimationSharing/Private/AnimationSharingManager.cpp
--- a/Engine/Plugins/Developer/AnimationSharing/Source/AnimationSharing/Private/AnimationSharingManager.cpp
+++ b/Engine/Plugins/Developer/AnimationSharing/Source/AnimationSharing/Private/AnimationSharingManager.cpp
@@ -433,7 +433,7 @@ void UAnimationSharingManager::RegisterActorWithSkeleton(AActor* InActor, const
 
 				Component->SetComponentTickEnabled(false);
 				Component->PrimaryComponentTick.bCanEverTick = false;
-				Component->bIgnoreMasterPoseComponentLOD = true;
+				Component->bIgnoreLeaderPoseComponentLOD = true;
 
 				ActorData.ComponentIndices.Add(Data->PerComponentData.Num() - 1);
 
@@ -483,7 +483,7 @@ void UAnimationSharingManager::UnregisterActor(AActor* InActor)
 
 			for (int32 ComponentIndex : ActorData.ComponentIndices)
 			{
-				SkeletonData->PerComponentData[ComponentIndex].Component->SetMasterPoseComponent(nullptr, true);
+				SkeletonData->PerComponentData[ComponentIndex].Component->SetLeaderPoseComponent(nullptr, true);
 				SkeletonData->PerComponentData[ComponentIndex].Component->PrimaryComponentTick.bCanEverTick = true;
 				SkeletonData->PerComponentData[ComponentIndex].Component->SetComponentTickEnabled(true);
 				SkeletonData->RemoveComponent(ComponentIndex);
@@ -596,7 +596,7 @@ void UAnimationSharingManager::UnregisterAllActors()
 				FPerActorData& ActorData = Data->PerActorData[ActorIndex];
 				for (int32 ComponentIndex : ActorData.ComponentIndices)
 				{
-					Data->PerComponentData[ComponentIndex].Component->SetMasterPoseComponent(nullptr, true);
+					Data->PerComponentData[ComponentIndex].Component->SetLeaderPoseComponent(nullptr, true);
 					Data->PerComponentData[ComponentIndex].Component->PrimaryComponentTick.bCanEverTick = true;
 					Data->PerComponentData[ComponentIndex].Component->SetComponentTickEnabled(true);
 					Data->PerComponentData[ComponentIndex].Component->bRecentlyRendered = false;
@@ -701,7 +701,7 @@ void UAnimSharingInstance::BeginDestroy()
 	{
 		for (uint32 ComponentIndex : ActorData.ComponentIndices)
 		{
-			PerComponentData[ComponentIndex].Component->SetMasterPoseComponent(nullptr, true);
+			PerComponentData[ComponentIndex].Component->SetLeaderPoseComponent(nullptr, true);
 		}
 	}
 
@@ -909,7 +909,7 @@ void UAnimSharingInstance::SetupState(FPerStateData& StateData, const FAnimation
 					Component->VisibilityBasedAnimTickOption = EVisibilityBasedAnimTickOption::AlwaysTickPoseAndRefreshBones;
 					Component->SetForcedLOD(1);
 					Component->SetVisibility(GMasterComponentsVisible == 1);
-					Component->bPropagateCurvesToSlaves = StateEntry.bRequiresCurves;
+					Component->bPropagateCurvesToFollowers = StateEntry.bRequiresCurves;
 
 					if (AnimBPClass != nullptr && AnimSequence != nullptr && !StateEntry.bOnDemand)
 					{
@@ -1095,7 +1095,7 @@ void UAnimSharingInstance::TickDebugInformation()
 
 				GEngine->AddOnScreenDebugMessage(1337, 1, FColor::White, OnScreenString);
 
-				const USkinnedMeshComponent* Component = PerComponentData[ActorData.ComponentIndices[0]].Component->MasterPoseComponent.Get();
+				const USkinnedMeshComponent* Component = PerComponentData[ActorData.ComponentIndices[0]].Component->LeaderPoseComponent.Get();
 #if UE_BUILD_DEVELOPMENT
 				if (Component != nullptr)
 				{
@@ -1740,7 +1740,7 @@ void UAnimSharingInstance::SetMasterComponentForActor(uint32 ActorIndex, USkelet
 	
 	for (uint32 ComponentIndex : ActorData.ComponentIndices)
 	{
-		PerComponentData[ComponentIndex].Component->SetMasterPoseComponent(Component, true);
+		PerComponentData[ComponentIndex].Component->SetLeaderPoseComponent(Component, true);
 	}
 }
 
diff --git a/Engine/Plugins/Editor/SequencerAnimTools/Source/SequencerAnimTools/Private/Sequencer/AnimationBoneTrail.cpp b/Engine/Plugins/Editor/SequencerAnimTools/Source/SequencerAnimTools/Private/Sequencer/AnimationBoneTrail.cpp
--- a/Engine/Plugins/Editor/SequencerAnimTools/Source/SequencerAnimTools/Private/Sequencer/AnimationBoneTrail.cpp
+++ b/Engine/Plugins/Editor/SequencerAnimTools/Source/SequencerAnimTools/Private/Sequencer/AnimationBoneTrail.cpp
@@ -34,7 +34,7 @@ void FAnimTrajectoryCache::Evaluate(FTrajectoryCache* ParentTransformCache)
 		return;
 	}
 
-	SkeletalMeshComponent->MasterPoseComponent = nullptr;
+	SkeletalMeshComponent->LeaderPoseComponent = nullptr;
 
 	CachedAnimSequence->SetSkeleton(SkeletalMeshComponent->GetSkeletalMesh()->GetSkeleton());
 
diff --git a/Engine/Plugins/Experimental/Animation/CurveExpression/Source/Runtime/Private/AnimNode_RemapCurvesFromMesh.cpp b/Engine/Plugins/Experimental/Animation/CurveExpression/Source/Runtime/Private/AnimNode_RemapCurvesFromMesh.cpp
--- a/Engine/Plugins/Experimental/Animation/CurveExpression/Source/Runtime/Private/AnimNode_RemapCurvesFromMesh.cpp
+++ b/Engine/Plugins/Experimental/Animation/CurveExpression/Source/Runtime/Private/AnimNode_RemapCurvesFromMesh.cpp
@@ -155,13 +155,13 @@ void FAnimNode_RemapCurvesFromMesh::PreUpdate(
 
 	if (CurrentMeshComponent && CurrentMeshComponent->GetSkeletalMesh() && CurrentMeshComponent->IsRegistered())
 	{
-		// If our source is running under master-pose, then get bone data from there
-		if(USkeletalMeshComponent* MasterPoseComponent = Cast<USkeletalMeshComponent>(CurrentMeshComponent->MasterPoseComponent.Get()))
+		// If our source is running under leader-pose, then get bone data from there
+		if(USkeletalMeshComponent* LeaderPoseComponent = Cast<USkeletalMeshComponent>(CurrentMeshComponent->LeaderPoseComponent.Get()))
 		{
-			CurrentMeshComponent = MasterPoseComponent;
+			CurrentMeshComponent = LeaderPoseComponent;
 		}
 
-		// re-check mesh component validity as it may have changed to master
+		// re-check mesh component validity as it may have changed to leader
 		if(CurrentMeshComponent->GetSkeletalMesh() && CurrentMeshComponent->IsRegistered())
 		{
 			if (ExpressionEngine.IsSet())
diff --git a/Engine/Plugins/Experimental/ChaosCloth/Source/Chaos/Private/ChaosCloth/ChaosClothingSimulation.cpp b/Engine/Plugins/Experimental/ChaosCloth/Source/Chaos/Private/ChaosCloth/ChaosClothingSimulation.cpp
--- a/Engine/Plugins/Experimental/ChaosCloth/Source/Chaos/Private/ChaosCloth/ChaosClothingSimulation.cpp
+++ b/Engine/Plugins/Experimental/ChaosCloth/Source/Chaos/Private/ChaosCloth/ChaosClothingSimulation.cpp
@@ -672,7 +672,7 @@ void FClothingSimulation::GetSimulationData(
 		}
 
 		// Get the reference bone index for this cloth
-		const int32 ReferenceBoneIndex = InOverrideComponent ? InOwnerComponent->GetMasterBoneMap()[Cloth->GetReferenceBoneIndex()] : Cloth->GetReferenceBoneIndex();
+		const int32 ReferenceBoneIndex = InOverrideComponent ? InOwnerComponent->GetLeaderBoneMap()[Cloth->GetReferenceBoneIndex()] : Cloth->GetReferenceBoneIndex();
 		if (!ComponentSpaceTransforms.IsValidIndex(ReferenceBoneIndex))
 		{
 			UE_CLOG(!bHasInvalidReferenceBoneTransforms, LogSkeletalMesh, Warning, TEXT("Failed to write back clothing simulation data for component %s as bone transforms are invalid."), *InOwnerComponent->GetName());
diff --git a/Engine/Plugins/Experimental/MLDeformer/MLDeformerFramework/Source/MLDeformerFramework/Private/MLDeformerModelInstance.cpp b/Engine/Plugins/Experimental/MLDeformer/MLDeformerFramework/Source/MLDeformerFramework/Private/MLDeformerModelInstance.cpp
--- a/Engine/Plugins/Experimental/MLDeformer/MLDeformerFramework/Source/MLDeformerFramework/Private/MLDeformerModelInstance.cpp
+++ b/Engine/Plugins/Experimental/MLDeformer/MLDeformerFramework/Source/MLDeformerFramework/Private/MLDeformerModelInstance.cpp
@@ -131,20 +131,20 @@ FString UMLDeformerModelInstance::CheckCompatibility(USkeletalMeshComponent* InS
 
 void UMLDeformerModelInstance::UpdateBoneTransforms()
 {
-	const USkinnedMeshComponent* MasterPoseComponent = SkeletalMeshComponent->MasterPoseComponent.Get();
-	if (MasterPoseComponent)
+	const USkinnedMeshComponent* LeaderPoseComponent = SkeletalMeshComponent->LeaderPoseComponent.Get();
+	if (LeaderPoseComponent)
 	{
-		const TArray<FTransform>& MasterTransforms = MasterPoseComponent->GetComponentSpaceTransforms();
-		USkeletalMesh* Mesh = MasterPoseComponent->GetSkeletalMesh();
+		const TArray<FTransform>& LeaderTransforms = LeaderPoseComponent->GetComponentSpaceTransforms();
+		USkeletalMesh* Mesh = LeaderPoseComponent->GetSkeletalMesh();
 		const int32 NumTrainingBones = AssetBonesToSkelMeshMappings.Num();
 		for (int32 Index = 0; Index < NumTrainingBones; ++Index)
 		{
 			const int32 ComponentBoneIndex = AssetBonesToSkelMeshMappings[Index];
-			const FTransform& ComponentSpaceTransform = MasterTransforms[ComponentBoneIndex];
+			const FTransform& ComponentSpaceTransform = LeaderTransforms[ComponentBoneIndex];
 			const int32 ParentIndex = Mesh->GetRefSkeleton().GetParentIndex(ComponentBoneIndex);
 			if (ParentIndex != INDEX_NONE)
 			{
-				TrainingBoneTransforms[Index] = ComponentSpaceTransform.GetRelativeTransform(MasterTransforms[ParentIndex]);
+				TrainingBoneTransforms[Index] = ComponentSpaceTransform.GetRelativeTransform(LeaderTransforms[ParentIndex]);
 			}
 			else
 			{
diff --git a/Engine/Plugins/FX/Niagara/Source/Niagara/Private/NiagaraDataInterfaceSkeletalMesh.cpp b/Engine/Plugins/FX/Niagara/Source/Niagara/Private/NiagaraDataInterfaceSkeletalMesh.cpp
--- a/Engine/Plugins/FX/Niagara/Source/Niagara/Private/NiagaraDataInterfaceSkeletalMesh.cpp
+++ b/Engine/Plugins/FX/Niagara/Source/Niagara/Private/NiagaraDataInterfaceSkeletalMesh.cpp
@@ -419,14 +419,14 @@ void FSkeletalMeshSkinningData::UpdateBoneTransforms()
 	TArray<FMatrix44f>& CurrBones = CurrBoneRefToLocals();
 	TArray<FTransform3f>& CurrTransforms = CurrComponentTransforms();
 
-	if (USkinnedMeshComponent* MasterComponent = SkelComp->MasterPoseComponent.Get())
+	if (USkinnedMeshComponent* LeaderComponent = SkelComp->LeaderPoseComponent.Get())
 	{
-		const TArray<int32>& MasterBoneMap = SkelComp->GetMasterBoneMap();
-		const int32 NumBones = MasterBoneMap.Num();
+		const TArray<int32>& LeaderBoneMap = SkelComp->GetLeaderBoneMap();
+		const int32 NumBones = LeaderBoneMap.Num();
 
 		if (NumBones == 0)
 		{
-			// This case indicates an invalid master pose component (e.g. no skeletal mesh)
+			// This case indicates an invalid leader pose component (e.g. no skeletal mesh)
 			CurrBones.Empty(SkelMesh->GetRefSkeleton().GetNum());
 			CurrBones.AddDefaulted(SkelMesh->GetRefSkeleton().GetNum());
 			CurrTransforms.Empty(SkelMesh->GetRefSkeleton().GetNum());
@@ -437,21 +437,21 @@ void FSkeletalMeshSkinningData::UpdateBoneTransforms()
 			CurrBones.SetNumUninitialized(NumBones);
 			CurrTransforms.SetNumUninitialized(NumBones);
 
-			const TArray<FTransform>& MasterTransforms = MasterComponent->GetComponentSpaceTransforms();
+			const TArray<FTransform>& LeaderTransforms = LeaderComponent->GetComponentSpaceTransforms();
 			for (int32 BoneIndex = 0; BoneIndex < NumBones; ++BoneIndex)
 			{
-				bool bFoundMaster = false;
-				if (MasterBoneMap.IsValidIndex(BoneIndex))
+				bool bFoundLeader = false;
+				if (LeaderBoneMap.IsValidIndex(BoneIndex))
 				{
-					const int32 MasterIndex = MasterBoneMap[BoneIndex];
-					if (MasterIndex != INDEX_NONE && MasterIndex < MasterTransforms.Num())
+					const int32 LeaderIndex = LeaderBoneMap[BoneIndex];
+					if (LeaderIndex != INDEX_NONE && LeaderIndex < LeaderTransforms.Num())
 					{
-						bFoundMaster = true;
-						CurrTransforms[BoneIndex] = (FTransform3f)MasterTransforms[MasterIndex];
+						bFoundLeader = true;
+						CurrTransforms[BoneIndex] = (FTransform3f)LeaderTransforms[LeaderIndex];
 					}
 				}
 
-				if ( !bFoundMaster )
+				if ( !bFoundLeader)
 				{
 					const int32 ParentIndex = SkelMesh->GetRefSkeleton().GetParentIndex(BoneIndex);
 					FTransform3f BoneTransform = (FTransform3f)SkelMesh->GetRefSkeleton().GetRefBonePose()[BoneIndex];
@@ -1229,32 +1229,32 @@ void FSkeletalMeshGpuDynamicBufferProxy::NewFrame(const FNDISkeletalMesh_Instanc
 	// If we have a component pull transforms from component otherwise grab from skel mesh
 	if (SkelComp)
 	{
-		if (USkinnedMeshComponent* MasterComponent = SkelComp->MasterPoseComponent.Get())
+		if (USkinnedMeshComponent* LeaderComponent = SkelComp->LeaderPoseComponent.Get())
 		{
 			const FReferenceSkeleton* ReferenceSkeleton = nullptr;
-			const TArray<int32>& MasterBoneMap = SkelComp->GetMasterBoneMap();
-			const int32 NumBones = MasterBoneMap.Num();
+			const TArray<int32>& LeaderBoneMap = SkelComp->GetLeaderBoneMap();
+			const int32 NumBones = LeaderBoneMap.Num();
 
 			TArray<FTransform> TempBoneTransforms;
 			TempBoneTransforms.Reserve(SamplingBoneCount);
 
 			if (NumBones == 0)
 			{
-				// This case indicates an invalid master pose component (e.g. no skeletal mesh)
+				// This case indicates an invalid leader pose component (e.g. no skeletal mesh)
 				TempBoneTransforms.AddDefaulted(SamplingBoneCount);
 			}
 			else
 			{
 				ReferenceSkeleton = &SkelMesh->GetRefSkeleton();
-				const TArray<FTransform>& MasterTransforms = MasterComponent->GetComponentSpaceTransforms();
+				const TArray<FTransform>& LeaderTransforms = LeaderComponent->GetComponentSpaceTransforms();
 				for (int32 BoneIndex=0; BoneIndex < NumBones; ++BoneIndex)
 				{
-					if (MasterBoneMap.IsValidIndex(BoneIndex))
+					if (LeaderBoneMap.IsValidIndex(BoneIndex))
 					{
-						const int32 MasterIndex = MasterBoneMap[BoneIndex];
-						if (MasterIndex != INDEX_NONE && MasterIndex < MasterTransforms.Num())
+						const int32 LeaderIndex = LeaderBoneMap[BoneIndex];
+						if (LeaderIndex != INDEX_NONE && LeaderIndex < LeaderTransforms.Num())
 						{
-							TempBoneTransforms.Add(MasterTransforms[MasterIndex]);
+							TempBoneTransforms.Add(LeaderTransforms[LeaderIndex]);
 							continue;
 						}
 					}
diff --git a/Engine/Plugins/Importers/USDImporter/Source/USDSchemas/Private/USDSkelRootTranslator.cpp b/Engine/Plugins/Importers/USDImporter/Source/USDSchemas/Private/USDSkelRootTranslator.cpp
--- a/Engine/Plugins/Importers/USDImporter/Source/USDSchemas/Private/USDSkelRootTranslator.cpp
+++ b/Engine/Plugins/Importers/USDImporter/Source/USDSchemas/Private/USDSkelRootTranslator.cpp
@@ -1379,7 +1379,7 @@ void FUsdSkelRootTranslator::UpdateComponents( USceneComponent* SceneComponent )
 
 			SkeletalMeshComponent->TickAnimation( 0.f, false );
 			SkeletalMeshComponent->RefreshBoneTransforms();
-			SkeletalMeshComponent->RefreshSlaveComponents();
+			SkeletalMeshComponent->RefreshFollowerComponents();
 			SkeletalMeshComponent->UpdateComponentToWorld();
 			SkeletalMeshComponent->FinalizeBoneTransform();
 			SkeletalMeshComponent->MarkRenderTransformDirty();
diff --git a/Engine/Plugins/Importers/USDImporter/Source/USDStageImporter/Private/USDStageImporter.cpp b/Engine/Plugins/Importers/USDImporter/Source/USDStageImporter/Private/USDStageImporter.cpp
--- a/Engine/Plugins/Importers/USDImporter/Source/USDStageImporter/Private/USDStageImporter.cpp
+++ b/Engine/Plugins/Importers/USDImporter/Source/USDStageImporter/Private/USDStageImporter.cpp
@@ -1419,7 +1419,7 @@ namespace UsdStageImporterImpl
 				{
 					SkeletalMeshComponent->TickAnimation( 0.f, false );
 					SkeletalMeshComponent->RefreshBoneTransforms();
-					SkeletalMeshComponent->RefreshSlaveComponents();
+					SkeletalMeshComponent->RefreshFollowerComponents();
 					SkeletalMeshComponent->UpdateComponentToWorld();
 					SkeletalMeshComponent->FinalizeBoneTransform();
 					SkeletalMeshComponent->MarkRenderTransformDirty();
diff --git a/Engine/Plugins/Importers/USDImporter/Source/USDUtilities/Private/USDPrimConversion.cpp b/Engine/Plugins/Importers/USDImporter/Source/USDUtilities/Private/USDPrimConversion.cpp
--- a/Engine/Plugins/Importers/USDImporter/Source/USDUtilities/Private/USDPrimConversion.cpp
+++ b/Engine/Plugins/Importers/USDImporter/Source/USDUtilities/Private/USDPrimConversion.cpp
@@ -2924,7 +2924,7 @@ bool UnrealToUsd::CreateSkeletalAnimationBaker( UE::FUsdPrim& SkelRoot, UE::FUsd
 			Component.TickAnimation( 0.f, false );
 			Component.UpdateLODStatus();
 			Component.RefreshBoneTransforms();
-			Component.RefreshSlaveComponents();
+			Component.RefreshFollowerComponents();
 			Component.UpdateComponentToWorld();
 			Component.FinalizeBoneTransform();
 			Component.MarkRenderTransformDirty();
diff --git a/Engine/Plugins/MovieScene/LevelSequenceEditor/Source/LevelSequenceEditor/Private/LevelSequenceEditorSubsystem.cpp b/Engine/Plugins/MovieScene/LevelSequenceEditor/Source/LevelSequenceEditor/Private/LevelSequenceEditorSubsystem.cpp
--- a/Engine/Plugins/MovieScene/LevelSequenceEditor/Source/LevelSequenceEditor/Private/LevelSequenceEditorSubsystem.cpp
+++ b/Engine/Plugins/MovieScene/LevelSequenceEditor/Source/LevelSequenceEditor/Private/LevelSequenceEditorSubsystem.cpp
@@ -746,7 +746,7 @@ void ULevelSequenceEditorSubsystem::BakeTransform(const TArray<FMovieSceneBindin
 						SkeletalMeshComponent->TickAnimation(0.f, false);
 
 						SkeletalMeshComponent->RefreshBoneTransforms();
-						SkeletalMeshComponent->RefreshSlaveComponents();
+						SkeletalMeshComponent->RefreshFollowerComponents();
 						SkeletalMeshComponent->UpdateComponentToWorld();
 						SkeletalMeshComponent->FinalizeBoneTransform();
 						SkeletalMeshComponent->MarkRenderTransformDirty();
diff --git a/Engine/Plugins/VirtualProduction/Takes/Source/TakeTrackRecorders/Private/MovieScene3DTransformTrackRecorder.cpp b/Engine/Plugins/VirtualProduction/Takes/Source/TakeTrackRecorders/Private/MovieScene3DTransformTrackRecorder.cpp
--- a/Engine/Plugins/VirtualProduction/Takes/Source/TakeTrackRecorders/Private/MovieScene3DTransformTrackRecorder.cpp
+++ b/Engine/Plugins/VirtualProduction/Takes/Source/TakeTrackRecorders/Private/MovieScene3DTransformTrackRecorder.cpp
@@ -489,7 +489,7 @@ void UMovieScene3DTransformTrackRecorder::PostProcessAnimationData(UMovieSceneAn
 
 			}
 			// Search for the Root Bone in the Skeleton
-			const USkeletalMesh* SkeletalMesh = SkeletalMeshComponent->MasterPoseComponent != nullptr ? ToRawPtr(SkeletalMeshComponent->MasterPoseComponent->GetSkeletalMesh()) : ToRawPtr(SkeletalMeshComponent->GetSkeletalMesh());
+			const USkeletalMesh* SkeletalMesh = SkeletalMeshComponent->LeaderPoseComponent != nullptr ? ToRawPtr(SkeletalMeshComponent->LeaderPoseComponent->GetSkeletalMesh()) : ToRawPtr(SkeletalMeshComponent->GetSkeletalMesh());
 			const UAnimSequence* AnimSequence = AnimTrackRecorder->GetAnimSequence();
 			if (AnimSequence && SkeletalMesh && AnimSequence->GetSkeleton())
 			{
diff --git a/Engine/Plugins/VirtualProduction/Takes/Source/TakeTrackRecorders/Private/MovieSceneAnimationTrackRecorder.cpp b/Engine/Plugins/VirtualProduction/Takes/Source/TakeTrackRecorders/Private/MovieSceneAnimationTrackRecorder.cpp
--- a/Engine/Plugins/VirtualProduction/Takes/Source/TakeTrackRecorders/Private/MovieSceneAnimationTrackRecorder.cpp
+++ b/Engine/Plugins/VirtualProduction/Takes/Source/TakeTrackRecorders/Private/MovieSceneAnimationTrackRecorder.cpp
@@ -52,9 +52,9 @@ void UMovieSceneAnimationTrackRecorder::CreateAnimationAssetAndSequence(const AA
 		ComponentTransform = SkeletalMeshComponent->GetComponentToWorld().GetRelativeTransform(Actor->GetTransform());
 		FString AnimationAssetName = Actor->GetActorLabel();
 
-		if (ULevelSequence* MasterLevelSequence = OwningTakeRecorderSource->GetMasterLevelSequence())
+		if (ULevelSequence* LeaderLevelSequence = OwningTakeRecorderSource->GetMasterLevelSequence())
 		{
-			UTakeMetaData* AssetMetaData = MasterLevelSequence->FindMetaData<UTakeMetaData>();
+			UTakeMetaData* AssetMetaData = LeaderLevelSequence->FindMetaData<UTakeMetaData>();
 
 			AnimationAssetName = AssetMetaData->GenerateAssetPath(AnimSettings->AnimationAssetName);
 
@@ -144,11 +144,11 @@ void UMovieSceneAnimationTrackRecorder::CreateTrackImpl()
 			USkeleton* AnimSkeleton = AnimSequence->GetSkeleton();
 			// add all frames
 
-			const USkinnedMeshComponent* const MasterPoseComponentInst = SkeletalMeshComponent->MasterPoseComponent.Get();
+			const USkinnedMeshComponent* const LeaderPoseComponentInst = SkeletalMeshComponent->LeaderPoseComponent.Get();
 			const TArray<FTransform>* SpaceBases;
-			if (MasterPoseComponentInst)
+			if (LeaderPoseComponentInst)
 			{
-				SpaceBases = &MasterPoseComponentInst->GetComponentSpaceTransforms();
+				SpaceBases = &LeaderPoseComponentInst->GetComponentSpaceTransforms();
 			}
 			else
 			{
@@ -157,7 +157,7 @@ void UMovieSceneAnimationTrackRecorder::CreateTrackImpl()
 			for (int32 BoneIndex = 0; BoneIndex < SpaceBases->Num(); ++BoneIndex)
 			{
 				// verify if this bone exists in skeleton
-				const int32 BoneTreeIndex = AnimSkeleton->GetSkeletonBoneIndexFromMeshBoneIndex(SkeletalMeshComponent->MasterPoseComponent != nullptr ? SkeletalMeshComponent->MasterPoseComponent->GetSkeletalMesh() : SkeletalMeshComponent->GetSkeletalMesh(), BoneIndex);
+				const int32 BoneTreeIndex = AnimSkeleton->GetSkeletonBoneIndexFromMeshBoneIndex(SkeletalMeshComponent->LeaderPoseComponent != nullptr ? SkeletalMeshComponent->LeaderPoseComponent->GetSkeletalMesh() : SkeletalMeshComponent->GetSkeletalMesh(), BoneIndex);
 				if (BoneTreeIndex != INDEX_NONE)
 				{
 					// add tracks for the bone existing
diff --git a/Engine/Source/Editor/MovieSceneTools/Private/MovieSceneToolHelpers.cpp b/Engine/Source/Editor/MovieSceneTools/Private/MovieSceneToolHelpers.cpp
--- a/Engine/Source/Editor/MovieSceneTools/Private/MovieSceneToolHelpers.cpp
+++ b/Engine/Source/Editor/MovieSceneTools/Private/MovieSceneToolHelpers.cpp
@@ -3839,7 +3839,7 @@ static void TickFrame(const FFrameNumber& FrameNumber,float DeltaTime, UMovieSce
 		SkelMeshComp->TickAnimation(DeltaTime, false);
 
 		SkelMeshComp->RefreshBoneTransforms();
-		SkelMeshComp->RefreshSlaveComponents();
+		SkelMeshComp->RefreshFollowerComponents();
 		SkelMeshComp->UpdateComponentToWorld();
 		SkelMeshComp->FinalizeBoneTransform();
 		SkelMeshComp->MarkRenderTransformDirty();
@@ -4327,7 +4327,7 @@ static void GetSequencerActorWorldTransforms(IMovieScenePlayer* Player, UMovieSc
 						{
 							MeshComp->TickAnimation(0.03f, false);
 							MeshComp->RefreshBoneTransforms();
-							MeshComp->RefreshSlaveComponents();
+							MeshComp->RefreshFollowerComponents();
 							MeshComp->UpdateComponentToWorld();
 							MeshComp->FinalizeBoneTransform();
 							MeshComp->MarkRenderTransformDirty();
diff --git a/Engine/Source/Editor/SequenceRecorder/Private/AnimationRecorder.cpp b/Engine/Source/Editor/SequenceRecorder/Private/AnimationRecorder.cpp
--- a/Engine/Source/Editor/SequenceRecorder/Private/AnimationRecorder.cpp
+++ b/Engine/Source/Editor/SequenceRecorder/Private/AnimationRecorder.cpp
@@ -193,25 +193,25 @@ bool FAnimationRecorder::TriggerRecordAnimation(USkeletalMeshComponent* Componen
 	return false;
 }
 
-/** Helper function to get space bases depending on master pose component */
+/** Helper function to get space bases depending on leader pose component */
 void FAnimationRecorder::GetBoneTransforms(USkeletalMeshComponent* Component, TArray<FTransform>& BoneTransforms)
 {
-	const USkinnedMeshComponent* const MasterPoseComponentInst = Component->MasterPoseComponent.Get();
-	if(MasterPoseComponentInst)
+	const USkinnedMeshComponent* const LeaderPoseComponentInst = Component->LeaderPoseComponent.Get();
+	if(LeaderPoseComponentInst)
 	{
-		const TArray<FTransform>& SpaceBases = MasterPoseComponentInst->GetComponentSpaceTransforms();
+		const TArray<FTransform>& SpaceBases = LeaderPoseComponentInst->GetComponentSpaceTransforms();
 		BoneTransforms.Reset(BoneTransforms.Num());
 		BoneTransforms.AddUninitialized(SpaceBases.Num());
 		for(int32 BoneIndex = 0; BoneIndex < SpaceBases.Num(); BoneIndex++)
 		{
-			if(BoneIndex < Component->GetMasterBoneMap().Num())
+			if(BoneIndex < Component->GetLeaderBoneMap().Num())
 			{
-				int32 MasterBoneIndex = Component->GetMasterBoneMap()[BoneIndex];
+				int32 LeaderBoneIndex = Component->GetLeaderBoneMap()[BoneIndex];
 
-				// If ParentBoneIndex is valid, grab matrix from MasterPoseComponent.
-				if(MasterBoneIndex != INDEX_NONE && MasterBoneIndex < SpaceBases.Num())
+				// If ParentBoneIndex is valid, grab matrix from LeaderPoseComponent.
+				if(LeaderBoneIndex != INDEX_NONE && LeaderBoneIndex < SpaceBases.Num())
 				{
-					BoneTransforms[BoneIndex] = SpaceBases[MasterBoneIndex];
+					BoneTransforms[BoneIndex] = SpaceBases[LeaderBoneIndex];
 				}
 				else
 				{
@@ -270,7 +270,7 @@ void FAnimationRecorder::StartRecord(USkeletalMeshComponent* Component, UAnimSeq
 	for (int32 BoneIndex=0; BoneIndex <PreviousSpacesBases.Num(); ++BoneIndex)
 	{
 		// verify if this bone exists in skeleton
-		const int32 BoneTreeIndex = AnimSkeleton->GetSkeletonBoneIndexFromMeshBoneIndex(Component->MasterPoseComponent != nullptr ? Component->MasterPoseComponent->GetSkeletalMesh() : Component->GetSkeletalMesh(), BoneIndex);
+		const int32 BoneTreeIndex = AnimSkeleton->GetSkeletonBoneIndexFromMeshBoneIndex(Component->LeaderPoseComponent != nullptr ? Component->LeaderPoseComponent->GetSkeletalMesh() : Component->GetSkeletalMesh(), BoneIndex);
 		if (BoneTreeIndex != INDEX_NONE)
 		{
 			// add tracks for the bone existing
@@ -628,11 +628,11 @@ void FAnimationRecorder::ProcessRecordedTimes(UAnimSequence* AnimSequence, USkel
 	
 	USkeleton* AnimSkeleton = AnimSequence->GetSkeleton();
 
-	const USkinnedMeshComponent* const MasterPoseComponentInst = SkeletalMeshComponent->MasterPoseComponent.Get();
+	const USkinnedMeshComponent* const LeaderPoseComponentInst = SkeletalMeshComponent->LeaderPoseComponent.Get();
 	const TArray<FTransform>* SpaceBases;
-	if (MasterPoseComponentInst)
+	if (LeaderPoseComponentInst)
 	{
-		SpaceBases = &MasterPoseComponentInst->GetComponentSpaceTransforms();
+		SpaceBases = &LeaderPoseComponentInst->GetComponentSpaceTransforms();
 	}
 	else
 	{
@@ -652,7 +652,7 @@ void FAnimationRecorder::ProcessRecordedTimes(UAnimSequence* AnimSequence, USkel
 	{
 		for (int32 BoneIndex = 0; BoneIndex < SpaceBases->Num(); ++BoneIndex)
 		{
-			const int32 BoneTreeIndex = AnimSkeleton->GetSkeletonBoneIndexFromMeshBoneIndex(SkeletalMeshComponent->MasterPoseComponent != nullptr ? SkeletalMeshComponent->MasterPoseComponent->GetSkeletalMesh() : SkeletalMeshComponent->GetSkeletalMesh(), BoneIndex);
+			const int32 BoneTreeIndex = AnimSkeleton->GetSkeletonBoneIndexFromMeshBoneIndex(SkeletalMeshComponent->LeaderPoseComponent != nullptr ? SkeletalMeshComponent->LeaderPoseComponent->GetSkeletalMesh() : SkeletalMeshComponent->GetSkeletalMesh(), BoneIndex);
 			if (BoneTreeIndex != INDEX_NONE)
 			{
 				FName BoneTreeName = AnimSkeleton->GetReferenceSkeleton().GetBoneName(BoneTreeIndex);
@@ -673,7 +673,7 @@ void FAnimationRecorder::ProcessRecordedTimes(UAnimSequence* AnimSequence, USkel
 	for (int32 BoneIndex = 0; BoneIndex < SpaceBases->Num(); ++BoneIndex)
 	{
 		// verify if this bone exists in skeleton
-		const int32 BoneTreeIndex = AnimSkeleton->GetSkeletonBoneIndexFromMeshBoneIndex(SkeletalMeshComponent->MasterPoseComponent != nullptr ? SkeletalMeshComponent->MasterPoseComponent->GetSkeletalMesh() : SkeletalMeshComponent->GetSkeletalMesh(), BoneIndex);
+		const int32 BoneTreeIndex = AnimSkeleton->GetSkeletonBoneIndexFromMeshBoneIndex(SkeletalMeshComponent->LeaderPoseComponent != nullptr ? SkeletalMeshComponent->LeaderPoseComponent->GetSkeletalMesh() : SkeletalMeshComponent->GetSkeletalMesh(), BoneIndex);
 		if (BoneTreeIndex != INDEX_NONE)
 		{
 			// add tracks for the bone existing
@@ -826,7 +826,7 @@ bool FAnimationRecorder::Record(USkeletalMeshComponent* Component, FTransform co
 	if (ensure(AnimationObject))
 	{
 		IAnimationDataController& Controller = AnimationObject->GetController();
-		USkeletalMesh* SkeletalMesh = Component->MasterPoseComponent != nullptr ? ToRawPtr(Component->MasterPoseComponent->GetSkeletalMesh()) : ToRawPtr(Component->GetSkeletalMesh());
+		USkeletalMesh* SkeletalMesh = Component->LeaderPoseComponent != nullptr ? ToRawPtr(Component->LeaderPoseComponent->GetSkeletalMesh()) : ToRawPtr(Component->GetSkeletalMesh());
 
 		const TArray<FBoneAnimationTrack>& BoneAnimationTracks = AnimationObject->GetDataModel()->GetBoneAnimationTracks();
 
diff --git a/Engine/Source/Editor/SequenceRecorder/Private/Sections/MovieScene3DTransformSectionRecorder.cpp b/Engine/Source/Editor/SequenceRecorder/Private/Sections/MovieScene3DTransformSectionRecorder.cpp
--- a/Engine/Source/Editor/SequenceRecorder/Private/Sections/MovieScene3DTransformSectionRecorder.cpp
+++ b/Engine/Source/Editor/SequenceRecorder/Private/Sections/MovieScene3DTransformSectionRecorder.cpp
@@ -132,7 +132,7 @@ void FMovieScene3DTransformSectionRecorder::FinalizeSection(float CurrentTime)
 		USkeletalMeshComponent* SkeletalMeshComponent = AnimRecorder->GetSkeletalMeshComponent();
 		if (SkeletalMeshComponent)
 		{
-			USkeletalMesh* SkeletalMesh = SkeletalMeshComponent->MasterPoseComponent != nullptr ? ToRawPtr(SkeletalMeshComponent->MasterPoseComponent->GetSkeletalMesh()) : ToRawPtr(SkeletalMeshComponent->GetSkeletalMesh());
+			USkeletalMesh* SkeletalMesh = SkeletalMeshComponent->LeaderPoseComponent != nullptr ? ToRawPtr(SkeletalMeshComponent->LeaderPoseComponent->GetSkeletalMesh()) : ToRawPtr(SkeletalMeshComponent->GetSkeletalMesh());
 			if (AnimSequence && SkeletalMesh)
 			{
 				// find the root bone
diff --git a/Engine/Source/Editor/UnrealEd/Private/Fbx/FbxAnimationExport.cpp b/Engine/Source/Editor/UnrealEd/Private/Fbx/FbxAnimationExport.cpp
--- a/Engine/Source/Editor/UnrealEd/Private/Fbx/FbxAnimationExport.cpp
+++ b/Engine/Source/Editor/UnrealEd/Private/Fbx/FbxAnimationExport.cpp
@@ -717,7 +717,7 @@ void FFbxExporter::ExportAnimTrack(IAnimTrackAdapter& AnimTrackAdapter, AActor*
 					SkeletalMeshComponent->TickAnimation(TickRate, false);
 
 					SkeletalMeshComponent->RefreshBoneTransforms();
-					SkeletalMeshComponent->RefreshSlaveComponents();
+					SkeletalMeshComponent->RefreshFollowerComponents();
 					SkeletalMeshComponent->UpdateComponentToWorld();
 					SkeletalMeshComponent->FinalizeBoneTransform();
 					SkeletalMeshComponent->MarkRenderTransformDirty();
diff --git a/Engine/Source/Editor/UnrealEd/Private/Fbx/FbxMainExport.cpp b/Engine/Source/Editor/UnrealEd/Private/Fbx/FbxMainExport.cpp
--- a/Engine/Source/Editor/UnrealEd/Private/Fbx/FbxMainExport.cpp
+++ b/Engine/Source/Editor/UnrealEd/Private/Fbx/FbxMainExport.cpp
@@ -3278,7 +3278,7 @@ void FFbxExporter::ExportLevelSequenceBaked3DTransformTrack(IAnimTrackAdapter& A
 				ChildSkeletalMeshComponent->TickAnimation(SampleRate, false);
 
 				ChildSkeletalMeshComponent->RefreshBoneTransforms();
-				ChildSkeletalMeshComponent->RefreshSlaveComponents();
+				ChildSkeletalMeshComponent->RefreshFollowerComponents();
 				ChildSkeletalMeshComponent->UpdateComponentToWorld();
 				ChildSkeletalMeshComponent->FinalizeBoneTransform();
 				ChildSkeletalMeshComponent->MarkRenderTransformDirty();
diff --git a/Engine/Source/Runtime/AnimGraphRuntime/Private/AnimNodes/AnimNode_CopyPoseFromMesh.cpp b/Engine/Source/Runtime/AnimGraphRuntime/Private/AnimNodes/AnimNode_CopyPoseFromMesh.cpp
--- a/Engine/Source/Runtime/AnimGraphRuntime/Private/AnimNodes/AnimNode_CopyPoseFromMesh.cpp
+++ b/Engine/Source/Runtime/AnimGraphRuntime/Private/AnimNodes/AnimNode_CopyPoseFromMesh.cpp
@@ -114,13 +114,13 @@ void FAnimNode_CopyPoseFromMesh::PreUpdate(const UAnimInstance* InAnimInstance)
 
 	if (CurrentMeshComponent && CurrentMeshComponent->GetSkeletalMesh() && CurrentMeshComponent->IsRegistered())
 	{
-		// If our source is running under master-pose, then get bone data from there
-		if(USkeletalMeshComponent* MasterPoseComponent = Cast<USkeletalMeshComponent>(CurrentMeshComponent->MasterPoseComponent.Get()))
+		// If our source is running under leader-pose, then get bone data from there
+		if(USkeletalMeshComponent* LeaderPoseComponent = Cast<USkeletalMeshComponent>(CurrentMeshComponent->LeaderPoseComponent.Get()))
 		{
-			CurrentMeshComponent = MasterPoseComponent;
+			CurrentMeshComponent = LeaderPoseComponent;
 		}
 
-		// re-check mesh component validity as it may have changed to master
+		// re-check mesh component validity as it may have changed to leader
 		if(CurrentMeshComponent->GetSkeletalMesh() && CurrentMeshComponent->IsRegistered())
 		{
 			const bool bUROInSync = CurrentMeshComponent->ShouldUseUpdateRateOptimizations() && CurrentMeshComponent->AnimUpdateRateParams != nullptr && CurrentMeshComponent->AnimUpdateRateParams == InAnimInstance->GetSkelMeshComponent()->AnimUpdateRateParams;
diff --git a/Engine/Source/Runtime/AnimGraphRuntime/Public/AnimCustomInstanceHelper.h b/Engine/Source/Runtime/AnimGraphRuntime/Public/AnimCustomInstanceHelper.h
--- a/Engine/Source/Runtime/AnimGraphRuntime/Public/AnimCustomInstanceHelper.h
+++ b/Engine/Source/Runtime/AnimGraphRuntime/Public/AnimCustomInstanceHelper.h
@@ -272,7 +272,7 @@ public:
 
 			// Update space bases to reset it back to ref pose
 			InSkeletalMeshComponent->RefreshBoneTransforms();
-			InSkeletalMeshComponent->RefreshSlaveComponents();
+			InSkeletalMeshComponent->RefreshFollowerComponents();
 			InSkeletalMeshComponent->UpdateComponentToWorld();
 		}
 
diff --git a/Engine/Source/Runtime/AnimGraphRuntime/Public/BoneControllers/AnimNode_RigidBody.h b/Engine/Source/Runtime/AnimGraphRuntime/Public/BoneControllers/AnimNode_RigidBody.h
--- a/Engine/Source/Runtime/AnimGraphRuntime/Public/BoneControllers/AnimNode_RigidBody.h
+++ b/Engine/Source/Runtime/AnimGraphRuntime/Public/BoneControllers/AnimNode_RigidBody.h
@@ -66,7 +66,7 @@ struct ANIMGRAPHRUNTIME_API FSimSpaceSettings
 	float WorldAlpha;
 
 #if WITH_EDITORONLY_DATA
-	UE_DEPRECATED(5.1, "MasterAlpha has been deprecated. Please, use WorldAlpha.")
+	UE_DEPRECATED(5.1, "This property has been deprecated. Please, use WorldAlpha.")
 	UPROPERTY()
 	float MasterAlpha;
 #endif // WITH_EDITORONLY_DATA
diff --git a/Engine/Source/Runtime/ClothingSystemRuntimeCommon/Private/ClothingSimulation.cpp b/Engine/Source/Runtime/ClothingSystemRuntimeCommon/Private/ClothingSimulation.cpp
--- a/Engine/Source/Runtime/ClothingSystemRuntimeCommon/Private/ClothingSimulation.cpp
+++ b/Engine/Source/Runtime/ClothingSystemRuntimeCommon/Private/ClothingSimulation.cpp
@@ -62,16 +62,16 @@ void FClothingSimulationContextCommon::FillBoneTransforms(const USkeletalMeshCom
 {
 	const USkeletalMesh* const SkeletalMesh = InComponent->GetSkeletalMesh();
 
-	if (USkinnedMeshComponent* const MasterComponent = InComponent->MasterPoseComponent.Get())
+	if (USkinnedMeshComponent* const LeaderComponent = InComponent->LeaderPoseComponent.Get())
 	{
-		const TArray<int32>& MasterBoneMap = InComponent->GetMasterBoneMap();
-		int32 NumBones = MasterBoneMap.Num();
+		const TArray<int32>& LeaderBoneMap = InComponent->GetLeaderBoneMap();
+		int32 NumBones = LeaderBoneMap.Num();
 
 		if (NumBones == 0)
 		{
 			if (SkeletalMesh)
 			{
-				// This case indicates an invalid master pose component (e.g. no skeletal mesh)
+				// This case indicates an invalid leader pose component (e.g. no skeletal mesh)
 				NumBones = SkeletalMesh->GetRefSkeleton().GetNum();
 
 				BoneTransforms.Empty(NumBones);
@@ -83,21 +83,21 @@ void FClothingSimulationContextCommon::FillBoneTransforms(const USkeletalMeshCom
 			BoneTransforms.Reset(NumBones);
 			BoneTransforms.AddDefaulted(NumBones);
 
-			const TArray<FTransform>& MasterTransforms = MasterComponent->GetComponentSpaceTransforms();
+			const TArray<FTransform>& LeaderTransforms = LeaderComponent->GetComponentSpaceTransforms();
 			for (int32 BoneIndex = 0; BoneIndex < NumBones; ++BoneIndex)
 			{
-				bool bFoundMaster = false;
-				if (MasterBoneMap.IsValidIndex(BoneIndex))
+				bool bFoundLeader = false;
+				if (LeaderBoneMap.IsValidIndex(BoneIndex))
 				{
-					const int32 MasterIndex = MasterBoneMap[BoneIndex];
-					if (MasterIndex != INDEX_NONE && MasterIndex < MasterTransforms.Num())
+					const int32 LeaderIndex = LeaderBoneMap[BoneIndex];
+					if (LeaderIndex != INDEX_NONE && LeaderIndex < LeaderTransforms.Num())
 					{
-						BoneTransforms[BoneIndex] = MasterTransforms[MasterIndex];
-						bFoundMaster = true;
+						BoneTransforms[BoneIndex] = LeaderTransforms[LeaderIndex];
+						bFoundLeader = true;
 					}
 				}
 
-				if (!bFoundMaster && SkeletalMesh)
+				if (!bFoundLeader && SkeletalMesh)
 				{
 					const int32 ParentIndex = SkeletalMesh->GetRefSkeleton().GetParentIndex(BoneIndex);
 
@@ -214,8 +214,8 @@ void FClothingSimulationCommon::FillContext(USkeletalMeshComponent* InComponent,
 	if (Context->BoneTransforms.Num() == 0)
 	{
 		const AActor* const CompOwner = InComponent->GetOwner();
-		const USkinnedMeshComponent* const Master = InComponent->MasterPoseComponent.Get();
-		UE_LOG(LogSkeletalMesh, Warning, TEXT("Attempting to fill a clothing simulation context for a skeletal mesh component that has zero bones (Comp: %s, Master: %s, Actor: %s)."), *InComponent->GetName(), Master ? *Master->GetName() : TEXT("None"), CompOwner ? *CompOwner->GetName() : TEXT("None"));
+		const USkinnedMeshComponent* const Leader = InComponent->LeaderPoseComponent.Get();
+		UE_LOG(LogSkeletalMesh, Warning, TEXT("Attempting to fill a clothing simulation context for a skeletal mesh component that has zero bones (Comp: %s, Leader: %s, Actor: %s)."), *InComponent->GetName(), Leader ? *Leader->GetName() : TEXT("None"), CompOwner ? *CompOwner->GetName() : TEXT("None"));
 
 		// Make sure we clear this out to skip any attempted simulations
 		Context->BoneTransforms.Reset();
diff --git a/Engine/Source/Runtime/Engine/Classes/Components/SkeletalMeshComponent.h b/Engine/Source/Runtime/Engine/Classes/Components/SkeletalMeshComponent.h
--- a/Engine/Source/Runtime/Engine/Classes/Components/SkeletalMeshComponent.h
+++ b/Engine/Source/Runtime/Engine/Classes/Components/SkeletalMeshComponent.h
@@ -644,8 +644,8 @@ private:
 	UPROPERTY()
 	uint8 bDisableAnimCurves_DEPRECATED : 1;
 
-	/** Whether or not we're taking cloth sim information from our master component */
-	uint8 bBindClothToMasterComponent:1;
+	/** Whether or not we're taking cloth sim information from our leader component */
+	uint8 bBindClothToLeaderComponent:1;
 
 	/** Flag denoting whether or not the clothing transform needs to update */
 	uint8 bPendingClothTransformUpdate:1;
@@ -753,9 +753,15 @@ public:
 	UPROPERTY()
 	uint8 bEnableLineCheckWithBounds:1;
 
-	/** If true, propagates calls to ApplyAnimationCurvesToComponent for slave components, only needed if slave components do not tick themselves */
-	UPROPERTY(EditAnywhere, BlueprintReadWrite, Category = MasterPoseComponent)
+	/** If true, propagates calls to ApplyAnimationCurvesToComponent for follower components, only needed if follower components do not tick themselves */
+	UPROPERTY(EditAnywhere, BlueprintReadWrite, Category = LeaderPoseComponent)
+	uint8 bPropagateCurvesToFollowers : 1;
+
+#if WITH_EDITORONLY_DATA
+	UE_DEPRECATED(5.1, "This property is deprecated. Please use bPropagateCurvesToFollowers instead")
+	UPROPERTY()
 	uint8 bPropagateCurvesToSlaves : 1;
+#endif // WITH_EDITORONLY_DATA
 
 	/** Whether to skip UpdateKinematicBonesToAnim() when interpolating. Kinematic bones are updated to the target interpolation pose only on ticks when they are evaluated. */
 	UPROPERTY(EditAnywhere, BlueprintReadWrite, AdvancedDisplay, Category = Optimization)
@@ -1197,23 +1203,31 @@ public:
 	void ResetClothTeleportMode();
 
 	/** 
-	 * If this component has a valid MasterPoseComponent then this function makes cloth items on the slave component
-	 * take the transforms of the cloth items on the master component instead of simulating separately.
-	 * @Note This will FORCE any cloth actor on the master component to simulate in local space. Also
+	 * If this component has a valid LeaderPoseComponent then this function makes cloth items on the follower component
+	 * take the transforms of the cloth items on the leader component instead of simulating separately.
+	 * @Note This will FORCE any cloth actor on the leader component to simulate in local space. Also
 	 * The meshes used in the components must be identical for the cloth to bind correctly
 	 */
-	UFUNCTION(BlueprintCallable, Category="Clothing", meta=(UnsafeDuringActorConstruction="true"))
-	void BindClothToMasterPoseComponent();
+	UFUNCTION(BlueprintCallable, Category = "Clothing", meta = (UnsafeDuringActorConstruction = "true"))
+	void BindClothToLeaderPoseComponent();
+
+	UE_DEPRECATED(5.1, "BindClothToMasterPoseComponent has been deprecated. Please use BindClothToLeaderPoseComponent instead.")
+	UFUNCTION(BlueprintCallable, Category="Clothing", meta=(UnsafeDuringActorConstruction="true", DeprecatedFunction, DeprecationMessage = "BindClothToMasterPoseComponent has been deprecated. Please use BindClothToLeaderPoseComponent instead."))
+	void BindClothToMasterPoseComponent() { BindClothToLeaderPoseComponent(); }
 
 	/** 
-	 * If this component has a valid MasterPoseComponent and has previously had its cloth bound to the
+	 * If this component has a valid LeaderPoseComponent and has previously had its cloth bound to the
 	 * MCP, this function will unbind the cloth and resume simulation.
-	 * @param bRestoreSimulationSpace if true and the master pose cloth was originally simulating in world
-	 * space, we will restore this setting. This will cause the master component to reset which may be
+	 * @param bRestoreSimulationSpace if true and the leader pose cloth was originally simulating in world
+	 * space, we will restore this setting. This will cause the leader component to reset which may be
 	 * undesirable.
 	 */
 	UFUNCTION(BlueprintCallable, Category="Clothing", meta=(UnsafeDuringActorConstruction="true"))
-	void UnbindClothFromMasterPoseComponent(bool bRestoreSimulationSpace = true);
+	void UnbindClothFromLeaderPoseComponent(bool bRestoreSimulationSpace = true);
+
+	UE_DEPRECATED(5.1, "UnbindClothFromMasterPoseComponent has been deprecated. Please use UnbindClothFromLeaderPoseComponent instead.")
+	UFUNCTION(BlueprintCallable, Category = "Clothing", meta = (UnsafeDuringActorConstruction = "true", DeprecatedFunction, DeprecationMessage = "UnbindClothFromMasterPoseComponent has been deprecated. Please use UnbindClothFromLeaderPoseComponent instead."))
+	void UnbindClothFromMasterPoseComponent(bool bRestoreSimulationSpace = true) { UnbindClothFromLeaderPoseComponent(bRestoreSimulationSpace); }
 
 	/**
 	 * Sets whether or not to allow rigid body animation nodes for this component
@@ -1241,7 +1255,7 @@ public:
 #if WITH_EDITOR
 	/**
 	 * return true if currently updating in editor is true
-	 * this is non BP because this is only used for slave component to detect master component ticking state
+	 * this is non BP because this is only used for follower component to detect leader component ticking state
 	 */
 	bool GetUpdateAnimationInEditor() const 
 	{		
@@ -1409,7 +1423,10 @@ public:
 	/** Checked whether we have already ticked the pose this frame */
 	bool PoseTickedThisFrame() const;
 
-	bool IsClothBoundToMasterComponent() const { return bBindClothToMasterComponent; }
+	bool IsClothBoundToLeaderComponent() const { return bBindClothToLeaderComponent; }
+	
+	UE_DEPRECATED(5.1, "IsClothBoundToMasterComponent has been deprecated. Please, use IsClothBoundToLeaderComponent instead.")
+	bool IsClothBoundToMasterComponent() const { return IsClothBoundToLeaderComponent(); }
 
 	/** Get the current clothing simulation (read only) */
 	const IClothingSimulation* GetClothingSimulation() const;
@@ -1836,7 +1853,9 @@ protected:
 	virtual void DispatchParallelTickPose( FActorComponentTickFunction* TickFunction ) override;
 public:
 	virtual void TickPose(float DeltaTime, bool bNeedsValidRootMotion) override;
-	virtual void UpdateSlaveComponent() override;
+	virtual void UpdateFollowerComponent() override;
+	UE_DEPRECATED(5.1, "UpdateSlaveComponent has been deprecated. Please use UpdateFollowerComponent instead.")
+	virtual void UpdateSlaveComponent() override { UpdateFollowerComponent(); };
 	virtual bool ShouldUpdateTransform(bool bLODHasChanged) const override;
 	virtual bool ShouldTickPose() const override;
 	virtual bool AllocateTransformData() override;
@@ -2362,8 +2381,8 @@ protected:
 
 private:
 	/** Override USkinnedMeshComponent */
-	virtual void AddSlavePoseComponent(USkinnedMeshComponent* SkinnedMeshComponent) override;
-	virtual void RemoveSlavePoseComponent(USkinnedMeshComponent* SkinnedMeshComponent) override;
+	virtual void AddFollowerPoseComponent(USkinnedMeshComponent* SkinnedMeshComponent) override;
+	virtual void RemoveFollowerPoseComponent(USkinnedMeshComponent* SkinnedMeshComponent) override;
 
 	// Returns whether we need to run the Pre Cloth Tick or not
 	bool ShouldRunEndPhysicsTick() const;
diff --git a/Engine/Source/Runtime/Engine/Classes/Components/SkinnedMeshComponent.h b/Engine/Source/Runtime/Engine/Classes/Components/SkinnedMeshComponent.h
--- a/Engine/Source/Runtime/Engine/Classes/Components/SkinnedMeshComponent.h
+++ b/Engine/Source/Runtime/Engine/Classes/Components/SkinnedMeshComponent.h
@@ -256,16 +256,22 @@ class ENGINE_API USkinnedMeshComponent : public UMeshComponent, public ILODSyncI
 	TObjectPtr<class USkeletalMesh> SkeletalMesh;
 
 	//
-	// MasterPoseComponent.
+	// LeaderPoseComponent.
 	//
 	
 	/**
 	 *	If set, this SkeletalMeshComponent will not use its SpaceBase for bone transform, but will
-	 *	use the component space transforms from the MasterPoseComponent. This is used when constructing a character using multiple skeletal meshes sharing the same
+	 *	use the component space transforms from the LeaderPoseComponent. This is used when constructing a character using multiple skeletal meshes sharing the same
 	 *	skeleton within the same Actor.
 	 */
-	UPROPERTY(BlueprintReadOnly, Category="Mesh")
+	UPROPERTY(BlueprintReadOnly, Category = "Mesh")
+	TWeakObjectPtr<USkinnedMeshComponent> LeaderPoseComponent;
+
+#if WITH_EDITORONLY_DATA
+	UE_DEPRECATED(5.1, "This property is deprecated. Please use LeaderPoseComponent instead")
+	UPROPERTY()
 	TWeakObjectPtr<USkinnedMeshComponent> MasterPoseComponent;
+#endif // WITH_EDITORONLY_DATA
 
 	/**
 	 * How this Component's LOD uses the skin cache feature. Auto will defer to the asset's (SkeletalMesh) option. If Ray Tracing is enabled, will imply Enabled
@@ -294,8 +300,8 @@ class ENGINE_API USkinnedMeshComponent : public UMeshComponent, public ILODSyncI
 
 	uint32 GetBoneTransformRevisionNumber() const 
 	{
-		const USkinnedMeshComponent* MasterPoseComponentPtr = MasterPoseComponent.Get();
-		return (MasterPoseComponentPtr ? MasterPoseComponentPtr->CurrentBoneTransformRevisionNumber : CurrentBoneTransformRevisionNumber);
+		const USkinnedMeshComponent* LeaderPoseComponentPtr = LeaderPoseComponent.Get();
+		return (LeaderPoseComponentPtr ? LeaderPoseComponentPtr->CurrentBoneTransformRevisionNumber : CurrentBoneTransformRevisionNumber);
 	}
 
 	/* this update renderer with new revision number twice so to clear bone velocity for motion blur or temporal AA */
@@ -330,49 +336,49 @@ protected:
 	/** current bone transform revision number */
 	uint32 CurrentBoneTransformRevisionNumber;
 
-	/** Incremented every time the master bone map changes. Used to keep in sync with any duplicate data needed by other threads */
-	int32 MasterBoneMapCacheCount;
+	/** Incremented every time the leader bone map changes. Used to keep in sync with any duplicate data needed by other threads */
+	int32 LeaderBoneMapCacheCount;
 
 	/** 
-	 * If set, this component has slave pose components that are associated with this 
+	 * If set, this component has follower pose components that are associated with this 
 	 * Note this is weak object ptr, so it will go away unless you have other strong reference
 	 */
-	TArray< TWeakObjectPtr<USkinnedMeshComponent> > SlavePoseComponents;
+	TArray< TWeakObjectPtr<USkinnedMeshComponent> > FollowerPoseComponents;
 
 	/**
-	 *	Mapping between bone indices in this component and the parent one. Each element is the index of the bone in the MasterPoseComponent.
+	 *	Mapping between bone indices in this component and the parent one. Each element is the index of the bone in the LeaderPoseComponent.
 	 *	Size should be the same as USkeletalMesh.RefSkeleton size (ie number of bones in this skeleton).
 	 */
-	TArray<int32> MasterBoneMap;
+	TArray<int32> LeaderBoneMap;
 
-	/** Cached relative transform for slave bones that are missing in the master */
-	struct FMissingMasterBoneCacheEntry
+	/** Cached relative transform for follower bones that are missing in the leader */
+	struct FMissingLeaderBoneCacheEntry
 	{
-		FMissingMasterBoneCacheEntry()
+		FMissingLeaderBoneCacheEntry()
 			: RelativeTransform(FTransform::Identity)
 			, CommonAncestorBoneIndex(INDEX_NONE)
 		{}
 
-		FMissingMasterBoneCacheEntry(const FTransform& InRelativeTransform, int32 InCommonAncestorBoneIndex)
+		FMissingLeaderBoneCacheEntry(const FTransform& InRelativeTransform, int32 InCommonAncestorBoneIndex)
 			: RelativeTransform(InRelativeTransform)
 			, CommonAncestorBoneIndex(InCommonAncestorBoneIndex)
 		{}
 
 		/** 
 		 * Relative transform of the missing bone's ref pose, based on the earliest common ancestor 
-		 * this will be equivalent to the component space transform of the bone had it existed in the master. 
+		 * this will be equivalent to the component space transform of the bone had it existed in the leader. 
 		 */
 		FTransform RelativeTransform;
 
-		/** The index of the earliest common ancestor of the master mesh. Index is the bone index in *this* mesh. */
+		/** The index of the earliest common ancestor of the leader mesh. Index is the bone index in *this* mesh. */
 		int32 CommonAncestorBoneIndex;
 	};
 
 	/**  
 	 * Map of missing bone indices->transforms so that calls to GetBoneTransform() succeed when bones are not
-	 * present in a master mesh when using master-pose. Index key is the bone index of *this* mesh.
+	 * present in a leader mesh when using leader-pose. Index key is the bone index of *this* mesh.
 	 */
-	TMap<int32, FMissingMasterBoneCacheEntry> MissingMasterBoneMap;
+	TMap<int32, FMissingLeaderBoneCacheEntry> MissingLeaderBoneMap;
 
 	/**
 	*	Mapping for socket overrides, key is the Source socket name and the value is the override socket name
@@ -409,7 +415,10 @@ public:
 
 	USkinnedAsset* GetSkinnedAsset() const;
 
-	const TArray<int32>& GetMasterBoneMap() const { return MasterBoneMap; }
+	const TArray<int32>& GetLeaderBoneMap() const { return LeaderBoneMap; }
+
+	UE_DEPRECATED(5.1, "GetMasterBoneMap has been deprecated. Please use the GetLeaderBoneMap.")
+	const TArray<int32>& GetMasterBoneMap() const { return GetLeaderBoneMap(); }
 
 	const FExternalMorphWeightData& GetExternalMorphWeights(int32 LOD) const { return ExternalMorphWeightData[LOD]; }
 	FExternalMorphWeightData& GetExternalMorphWeights(int32 LOD) { return ExternalMorphWeightData[LOD]; }
@@ -561,12 +570,18 @@ public:
 	uint8 bOverrideMinLod:1;
 
 	/** 
-	 * When true, we will just using the bounds from our MasterPoseComponent.  This is useful for when we have a Mesh Parented
+	 * When true, we will just using the bounds from our LeaderPoseComponent.  This is useful for when we have a Mesh Parented
 	 * to the main SkelMesh (e.g. outline mesh or a full body overdraw effect that is toggled) that is always going to be the same
 	 * bounds as parent.  We want to do no calculations in that case.
 	 */
 	UPROPERTY(EditAnywhere, AdvancedDisplay, BlueprintReadWrite, Category = SkeletalMesh)
-	uint8 bUseBoundsFromMasterPoseComponent:1;
+	uint8 bUseBoundsFromLeaderPoseComponent : 1;
+
+#if WITH_EDITORONLY_DATA
+	UE_DEPRECATED(5.1, "This property is deprecated. Please use bUseBoundsFromLeaderPoseComponent instead")
+	UPROPERTY()
+	uint8 bUseBoundsFromMasterPoseComponent : 1;
+#endif // WITH_EDITORONLY_DATA
 
 	/** Forces the mesh to draw in wireframe mode. */
 	UPROPERTY()
@@ -653,9 +668,15 @@ public:
 	UPROPERTY(EditAnywhere, AdvancedDisplay, BlueprintReadOnly, Category=Optimization)
 	uint8 bRenderStatic:1;
 
-	/** Flag that when set will ensure UpdateLODStatus will not take the MasterPoseComponent's current LOD in consideration when determining the correct LOD level (this requires MasterPoseComponent's LOD to always be >= determined LOD otherwise bone transforms could be missing */
+	/** Flag that when set will ensure UpdateLODStatus will not take the LeaderPoseComponent's current LOD in consideration when determining the correct LOD level (this requires LeaderPoseComponent's LOD to always be >= determined LOD otherwise bone transforms could be missing */
 	UPROPERTY(EditAnywhere, AdvancedDisplay, BlueprintReadOnly, Category = LOD)
+	uint8 bIgnoreLeaderPoseComponentLOD : 1;
+
+#if WITH_EDITORONLY_DATA
+	UE_DEPRECATED(5.1, "This property is deprecated. Please use bIgnoreLeaderPoseComponentLOD instead")
+	UPROPERTY()
 	uint8 bIgnoreMasterPoseComponentLOD : 1;
+#endif // WITH_EDITORONLY_DATA
 
 protected:
 	/** Are we using double buffered ComponentSpaceTransforms */
@@ -913,7 +934,6 @@ public:
 #if WITH_EDITOR
 	virtual bool CanEditChange(const FProperty* InProperty) const override;
 #endif // WITH_EDITOR
-	//~ End UObject Interface
 
 protected:
 	//~ Begin UActorComponent Interface
@@ -1059,7 +1079,7 @@ public:
 
 	FORCEINLINE	const USkinnedMeshComponent* GetBaseComponent()const
 	{
-		return MasterPoseComponent.IsValid() ? MasterPoseComponent.Get() : this;
+		return LeaderPoseComponent.IsValid() ? LeaderPoseComponent.Get() : this;
 	}
 
 	/**
@@ -1178,8 +1198,8 @@ protected:
 	 */
 	virtual void DispatchParallelTickPose(FActorComponentTickFunction* TickFunction) {}
 
-	/** Helper function for UpdateLODStatus, called with a valid index for InMasterPoseComponentPredictedLODLevel when updating LOD status for slave components */
-	bool UpdateLODStatus_Internal(int32 InMasterPoseComponentPredictedLODLevel, bool bRequestedByMasterPoseComponent = false);
+	/** Helper function for UpdateLODStatus, called with a valid index for InLeaderPoseComponentPredictedLODLevel when updating LOD status for follower components */
+	bool UpdateLODStatus_Internal(int32 InLeaderPoseComponentPredictedLODLevel, bool bRequestedByLeaderPoseComponent = false);
 
 public:
 	/**
@@ -1198,10 +1218,13 @@ public:
 	FOnTickPose OnTickPose;
 
 	/** 
-	 * Update Slave Component. This gets called when MasterPoseComponent!=NULL
+	 * Update Follower Component. This gets called when LeaderPoseComponent!=NULL
 	 * 
 	 */
-	virtual void UpdateSlaveComponent();
+	virtual void UpdateFollowerComponent();
+
+	UE_DEPRECATED(5.1, "UpdateSlaveComponent has been deprecated. Please use UpdateFollowerComponent instead.")
+	virtual void UpdateSlaveComponent() { UpdateFollowerComponent(); }
 
 	/** 
 	 * Update the PredictedLODLevel and MaxDistanceFactor in the component from its MeshObject. 
@@ -1355,7 +1378,7 @@ protected:
 	/** Update Mesh Bound information based on input
 	 * 
 	 * @param RootOffset	: Root Bone offset from mesh location
-	 *						  If MasterPoseComponent exists, it will applied to MasterPoseComponent's bound
+	 *						  If LeaderPoseComponent exists, it will applied to LeaderPoseComponent's bound
 	 * @param UsePhysicsAsset	: Whether or not to use PhysicsAsset for calculating bound of mesh
 	 */
 	FBoxSphereBounds CalcMeshBound(const FVector3f& RootOffset, bool UsePhysicsAsset, const FTransform& Transform) const;
@@ -1379,33 +1402,54 @@ private:
 	
 public:
 	/**
-	 * Set MasterPoseComponent for this component
+	 * Set LeaderPoseComponent for this component
 	 *
-	 * @param NewMasterBoneComponent New MasterPoseComponent
+	 * @param NewLeaderBoneComponent New LeaderPoseComponent
 	 */
-	UFUNCTION(BlueprintCallable, Category="Components|SkinnedMesh")
-	void SetMasterPoseComponent(USkinnedMeshComponent* NewMasterBoneComponent, bool bForceUpdate = false);
+	UFUNCTION(BlueprintCallable, Category = "Components|SkinnedMesh")
+	void SetLeaderPoseComponent(USkinnedMeshComponent* NewLeaderBoneComponent, bool bForceUpdate = false);
+
+	UE_DEPRECATED(5.1, "SetMasterPoseComponent has been deprecated. Please use SetLeaderPoseComponent instead.")
+	UFUNCTION(BlueprintCallable, Category = "Components|SkinnedMesh", meta = (DeprecatedFunction, DeprecationMessage = "SetMasterPoseComponent has been deprecated. Please use SetLeaderPoseComponent instead."))
+	void SetMasterPoseComponent(USkinnedMeshComponent* NewMasterBoneComponent, bool bForceUpdate = false) { SetLeaderPoseComponent(NewMasterBoneComponent, bForceUpdate); }
+
+	/** Return current active list of follower components */
+	const TArray< TWeakObjectPtr<USkinnedMeshComponent> >& GetFollowerPoseComponents() const;
+
+	UE_DEPRECATED(5.1, "GetSlavePoseComponents has been deprecated. Please use GetFollowerPoseComponents instead.")
+	const TArray< TWeakObjectPtr<USkinnedMeshComponent> >& GetSlavePoseComponents() const { return GetFollowerPoseComponents(); }
 
-	/** Return current active list of slave components */
-	const TArray< TWeakObjectPtr<USkinnedMeshComponent> >& GetSlavePoseComponents() const;
 protected:
-	/** Add a slave component to the SlavePoseComponents array */
-	virtual void AddSlavePoseComponent(USkinnedMeshComponent* SkinnedMeshComponent);
-	/** Remove a slave component from the SlavePoseComponents array */
-	virtual void RemoveSlavePoseComponent(USkinnedMeshComponent* SkinnedMeshComponent);
+	/** Add a follower component to the FollowerPoseComponents array */
+	virtual void AddFollowerPoseComponent(USkinnedMeshComponent* SkinnedMeshComponent);
+
+	UE_DEPRECATED(5.1, "AddSlavePoseComponent has been deprecated. Please use AddFollowerPoseComponent instead.")
+	virtual void AddSlavePoseComponent(USkinnedMeshComponent* SkinnedMeshComponent) { AddFollowerPoseComponent(SkinnedMeshComponent); }
+
+	/** Remove a follower component from the FollowerPoseComponents array */
+	virtual void RemoveFollowerPoseComponent(USkinnedMeshComponent* SkinnedMeshComponent);
+
+	UE_DEPRECATED(5.1, "RemoveSlavePoseComponent has been deprecated. Please use RemoveFollowerPoseComponent instead.")
+	virtual void RemoveSlavePoseComponent(USkinnedMeshComponent* SkinnedMeshComponent) { RemoveFollowerPoseComponent(SkinnedMeshComponent); }
 
 public:
 	/** 
-	 * Refresh Slave Components if exists
+	 * Refresh Follower Components if exists
 	 * 
 	 * This isn't necessary in any other case except in editor where you need to mark them as dirty for rendering
 	 */
-	void RefreshSlaveComponents();
+	void RefreshFollowerComponents();
+
+	UE_DEPRECATED(5.1, "RefreshSlaveComponents has been deprecated. Please use RefreshFollowerComponents instead.")
+	void RefreshSlaveComponents() { RefreshFollowerComponents(); }
 
 	/**
-	 * Update MasterBoneMap for MasterPoseComponent and this component
+	 * Update LeaderBoneMap for LeaderPoseComponent and this component
 	 */
-	void UpdateMasterBoneMap();
+	void UpdateLeaderBoneMap();
+
+	UE_DEPRECATED(5.1, "UpdateMasterBoneMap has been deprecated. Please use UpdateLeaderBoneMap instead.")
+	void UpdateMasterBoneMap() { UpdateLeaderBoneMap(); }
 
 	/**
 	 * @param InSocketName	The name of the socket to find
@@ -1668,10 +1712,10 @@ private:
 	virtual void RefreshMorphTargets() {};
 
 	/**  
-	 * When bones are not resent in a master mesh when using master-pose, we call this to evaluate 
+	 * When bones are not resent in a leader mesh when using leader-pose, we call this to evaluate 
 	 * relative transforms.
 	 */
-	bool GetMissingMasterBoneRelativeTransform(int32 InBoneIndex, FMissingMasterBoneCacheEntry& OutInfo) const;
+	bool GetMissingLeaderBoneRelativeTransform(int32 InBoneIndex, FMissingLeaderBoneCacheEntry& OutInfo) const;
 
 	// BEGIN ILODSyncComponent
 	virtual int32 GetDesiredSyncLOD() const override;
diff --git a/Engine/Source/Runtime/Engine/Private/Components/PoseableMeshComponent.cpp b/Engine/Source/Runtime/Engine/Private/Components/PoseableMeshComponent.cpp
--- a/Engine/Source/Runtime/Engine/Private/Components/PoseableMeshComponent.cpp
+++ b/Engine/Source/Runtime/Engine/Private/Components/PoseableMeshComponent.cpp
@@ -135,7 +135,7 @@ void UPoseableMeshComponent::SetBoneTransformByName(FName BoneName, const FTrans
 		return;
 	}
 
-	check(!MasterPoseComponent.IsValid()); //Shouldn't call set bone functions when we are using MasterPoseComponent
+	check(!LeaderPoseComponent.IsValid()); //Shouldn't call set bone functions when we are using LeaderPoseComponent
 
 	int32 BoneIndex = GetBoneIndex(BoneName);
 	if(BoneIndex >=0 && BoneIndex < BoneSpaceTransforms.Num())
@@ -221,7 +221,7 @@ FTransform UPoseableMeshComponent::GetBoneTransformByName(FName BoneName, EBoneS
 		return FTransform();
 	}
 
-	USkinnedMeshComponent* MPCPtr = MasterPoseComponent.Get();
+	USkinnedMeshComponent* MPCPtr = LeaderPoseComponent.Get();
 	if (MPCPtr)
 	{
 		if (USkeletalMeshComponent* SMC = Cast<USkeletalMeshComponent>(MPCPtr))
@@ -230,11 +230,11 @@ FTransform UPoseableMeshComponent::GetBoneTransformByName(FName BoneName, EBoneS
 			{
 				return GetBoneTransformByNameHelper(BoneName, BoneSpace, AnimInstance->GetRequiredBones(), SMC);
 			}
-			FString Message = FString::Printf(TEXT("Cannot return valid bone transform. Master Pose Component has no anim instance"));
+			FString Message = FString::Printf(TEXT("Cannot return valid bone transform. Leader Pose Component has no anim instance"));
 			FFrame::KismetExecutionMessage(*Message, ELogVerbosity::Warning);
 			return FTransform();
 		}
-		FString Message = FString::Printf(TEXT("Cannot return valid bone transform. Master Pose Component is not of type USkeletalMeshComponent"));
+		FString Message = FString::Printf(TEXT("Cannot return valid bone transform. Leader Pose Component is not of type USkeletalMeshComponent"));
 		FFrame::KismetExecutionMessage(*Message, ELogVerbosity::Warning);
 		return FTransform();
 	}
diff --git a/Engine/Source/Runtime/Engine/Private/Components/SkeletalMeshComponent.cpp b/Engine/Source/Runtime/Engine/Private/Components/SkeletalMeshComponent.cpp
--- a/Engine/Source/Runtime/Engine/Private/Components/SkeletalMeshComponent.cpp
+++ b/Engine/Source/Runtime/Engine/Private/Components/SkeletalMeshComponent.cpp
@@ -238,7 +238,7 @@ USkeletalMeshComponent::USkeletalMeshComponent(const FObjectInitializer& ObjectI
 	ComputeTeleportRotationThresholdInRadians();
 	ComputeTeleportDistanceThresholdInRadians();
 
-	bBindClothToMasterComponent = false;
+	bBindClothToLeaderComponent = false;
 	bClothingSimulationSuspended = false;
 
 #if WITH_EDITORONLY_DATA
@@ -273,7 +273,7 @@ USkeletalMeshComponent::USkeletalMeshComponent(const FObjectInitializer& ObjectI
 	// By default enable overlaps when blending physics - user can disable if they are sure it's unnecessary
 	bUpdateOverlapsOnAnimationFinalize = true;
 
-	bPropagateCurvesToSlaves = false;
+	bPropagateCurvesToFollowers = false;
 
 	bSkipKinematicUpdateWhenInterpolating = false;
 	bSkipBoundsUpdateWhenInterpolating = false;
@@ -781,8 +781,8 @@ void USkeletalMeshComponent::InitAnim(bool bForceReinit)
 		const bool bInitializedAnimInstance = InitializeAnimScriptInstance(bForceReinit, !bTickAnimationNow);
 
 		// Make sure we have a valid pose.
-		// We don't allocate transform data when using MasterPoseComponent, so we have nothing to render.
-		if (!MasterPoseComponent.IsValid())
+		// We don't allocate transform data when using LeaderPoseComponent, so we have nothing to render.
+		if (!LeaderPoseComponent.IsValid())
 		{	
 			if (bInitializedAnimInstance || (AnimScriptInstance == nullptr))
 			{ 
@@ -829,11 +829,11 @@ void USkeletalMeshComponent::ApplyEditedComponentSpaceTransforms()
 	MarkRenderTransformDirty();
 	MarkRenderDynamicDataDirty();
 
-	for (auto& SlaveComponent : GetSlavePoseComponents())
+	for (auto& FollowerComponent : GetFollowerPoseComponents())
 	{
-		if (SlaveComponent.IsValid())
+		if (FollowerComponent.IsValid())
 		{
-			SlaveComponent->UpdateSlaveComponent();
+			FollowerComponent->UpdateFollowerComponent();
 		}
 	}
 }
@@ -856,7 +856,7 @@ bool USkeletalMeshComponent::InitializeAnimScriptInstance(bool bForceReinit, boo
 
 			if (AnimScriptInstance)
 			{
-				// If we have any linked instances left we need to clear them out now, we're about to have a new master instance
+				// If we have any linked instances left we need to clear them out now, we're about to have a new leader instance
 				ResetLinkedAnimInstances();
 
 				AnimScriptInstance->InitializeAnimation(bInDeferRootNodeInitialization);
@@ -1297,11 +1297,11 @@ bool USkeletalMeshComponent::ShouldUpdateTransform(bool bLODHasChanged) const
 			return true;
 		}
 
-		// if master pose is ticking, slave also has to update it
-		if (MasterPoseComponent.IsValid())
+		// if leader pose is ticking, follower also has to update it
+		if (LeaderPoseComponent.IsValid())
 		{
-			const USkeletalMeshComponent* Master = CastChecked<USkeletalMeshComponent>(MasterPoseComponent.Get());
-			if (Master->GetUpdateAnimationInEditor())
+			const USkeletalMeshComponent* Leader = CastChecked<USkeletalMeshComponent>(LeaderPoseComponent.Get());
+			if (Leader->GetUpdateAnimationInEditor())
 			{
 				return true;
 			}
@@ -1483,7 +1483,7 @@ void USkeletalMeshComponent::TickComponent(float DeltaTime, enum ELevelTick Tick
 		constexpr bool bIsInitialization = false;
 		ClothingSimulation->FillContext(this, DeltaTime, ClothingSimulationContext, bIsInitialization);
 
-		ClothingSimulation->GetSimulationData(CurrentSimulationData, this, Cast<USkeletalMeshComponent>(MasterPoseComponent.Get()));
+		ClothingSimulation->GetSimulationData(CurrentSimulationData, this, Cast<USkeletalMeshComponent>(LeaderPoseComponent.Get()));
 	}
 
 	Super::TickComponent(DeltaTime, TickType, ThisTickFunction);
@@ -1997,32 +1997,32 @@ void USkeletalMeshComponent::EvaluateAnimation(const USkeletalMesh* InSkeletalMe
 	}
 }
 
-void USkeletalMeshComponent::UpdateSlaveComponent()
+void USkeletalMeshComponent::UpdateFollowerComponent()
 {
-	check (MasterPoseComponent.IsValid());
+	check (LeaderPoseComponent.IsValid());
 
 	ResetMorphTargetCurves();
 
-	if (USkeletalMeshComponent* MasterSMC = Cast<USkeletalMeshComponent>(MasterPoseComponent.Get()))
+	if (USkeletalMeshComponent* LeaderSMC = Cast<USkeletalMeshComponent>(LeaderPoseComponent.Get()))
 	{
-		// first set any animation-driven curves from the master SMC
-		if (MasterSMC->AnimScriptInstance)
+		// first set any animation-driven curves from the leader SMC
+		if (LeaderSMC->AnimScriptInstance)
 		{
-			MasterSMC->AnimScriptInstance->RefreshCurves(this);
+			LeaderSMC->AnimScriptInstance->RefreshCurves(this);
 		}
 
 		// we changed order of morphtarget to be overriden by SetMorphTarget from BP
 		// so this has to go first
-		// now propagate BP-driven curves from the master SMC...
+		// now propagate BP-driven curves from the leader SMC...
 		if (GetSkeletalMesh())
 		{
 			check(MorphTargetWeights.Num() == GetSkeletalMesh()->GetMorphTargets().Num());
-			if (MasterSMC->MorphTargetCurves.Num() > 0)
+			if (LeaderSMC->MorphTargetCurves.Num() > 0)
 			{
-				FAnimationRuntime::AppendActiveMorphTargets(GetSkeletalMesh(), MasterSMC->MorphTargetCurves, ActiveMorphTargets, MorphTargetWeights);
+				FAnimationRuntime::AppendActiveMorphTargets(GetSkeletalMesh(), LeaderSMC->MorphTargetCurves, ActiveMorphTargets, MorphTargetWeights);
 			}
 
-			// if slave also has it, add it here. 
+			// if follower also has it, add it here. 
 			if (MorphTargetCurves.Num() > 0)
 			{
 				FAnimationRuntime::AppendActiveMorphTargets(GetSkeletalMesh(), MorphTargetCurves, ActiveMorphTargets, MorphTargetWeights);
@@ -2031,7 +2031,7 @@ void USkeletalMeshComponent::UpdateSlaveComponent()
 
 	}
  
-	Super::UpdateSlaveComponent();
+	Super::UpdateFollowerComponent();
 }
 
 #if WITH_EDITOR
@@ -2205,11 +2205,11 @@ void USkeletalMeshComponent::UpdateClothSimulationContext(float InDeltaTime)
 
 void USkeletalMeshComponent::HandleExistingParallelClothSimulation()
 {
-	if (bBindClothToMasterComponent)
+	if (bBindClothToLeaderComponent)
 	{
-		if (USkeletalMeshComponent* MasterComp = Cast<USkeletalMeshComponent>(MasterPoseComponent.Get()))
+		if (USkeletalMeshComponent* LeaderComp = Cast<USkeletalMeshComponent>(LeaderPoseComponent.Get()))
 		{
-			MasterComp->HandleExistingParallelClothSimulation();
+			LeaderComp->HandleExistingParallelClothSimulation();
 		}
 	}
 
@@ -2232,12 +2232,12 @@ void USkeletalMeshComponent::WritebackClothingSimulationData()
 		CSV_SCOPED_TIMING_STAT(Animation, Cloth);
 
 		USkinnedMeshComponent* OverrideComponent = nullptr;
-		if(MasterPoseComponent.IsValid())
+		if(LeaderPoseComponent.IsValid())
 		{
-			OverrideComponent = MasterPoseComponent.Get();
+			OverrideComponent = LeaderPoseComponent.Get();
 
 			// Check if our bone map is actually valid, if not there is no clothing data to build
-			if(MasterBoneMap.Num() == 0)
+			if(LeaderBoneMap.Num() == 0)
 			{
 				CurrentSimulationData.Reset();
 				return;
@@ -2822,10 +2822,10 @@ void USkeletalMeshComponent::ApplyAnimationCurvesToComponent(const TMap<FName, f
 		FAnimationRuntime::AppendActiveMorphTargets(GetSkeletalMesh(), *InAnimationMorphCurves, ActiveMorphTargets, MorphTargetWeights);
 	}
 
-	/** Push through curves to slave components */
-	if (bPropagateCurvesToSlaves && bContainsMorphCurves && bContainsMaterialCurves && SlavePoseComponents.Num() > 0)
+	/** Push through curves to follower components */
+	if (bPropagateCurvesToFollowers && bContainsMorphCurves && bContainsMaterialCurves && FollowerPoseComponents.Num() > 0)
 	{
-		for (TWeakObjectPtr<USkinnedMeshComponent> MeshComponent : SlavePoseComponents)
+		for (TWeakObjectPtr<USkinnedMeshComponent> MeshComponent : FollowerPoseComponents)
 		{
 			if (USkeletalMeshComponent* SKComponent = Cast<USkeletalMeshComponent>(MeshComponent.Get()))
 			{
@@ -2869,17 +2869,17 @@ FBoxSphereBounds USkeletalMeshComponent::CalcBounds(const FTransform& LocalToWor
 	{
 		FVector RootBoneOffset = RootBoneTranslation;
 
-		// if to use MasterPoseComponent's fixed skel bounds, 
-		// send MasterPoseComponent's Root Bone Translation
-		if (MasterPoseComponent.IsValid())
+		// if to use LeaderPoseComponent's fixed skel bounds, 
+		// send LeaderPoseComponent's Root Bone Translation
+		if (LeaderPoseComponent.IsValid())
 		{
-			const USkinnedMeshComponent* const MasterPoseComponentInst = MasterPoseComponent.Get();
-			check(MasterPoseComponentInst);
-			if (MasterPoseComponentInst->GetSkeletalMesh() &&
-				MasterPoseComponentInst->bComponentUseFixedSkelBounds &&
-				MasterPoseComponentInst->IsA((USkeletalMeshComponent::StaticClass())))
+			const USkinnedMeshComponent* const LeaderPoseComponentInst = LeaderPoseComponent.Get();
+			check(LeaderPoseComponentInst);
+			if (LeaderPoseComponentInst->GetSkeletalMesh() &&
+				LeaderPoseComponentInst->bComponentUseFixedSkelBounds &&
+				LeaderPoseComponentInst->IsA((USkeletalMeshComponent::StaticClass())))
 			{
-				const USkeletalMeshComponent* BaseComponent = CastChecked<USkeletalMeshComponent>(MasterPoseComponentInst);
+				const USkeletalMeshComponent* BaseComponent = CastChecked<USkeletalMeshComponent>(LeaderPoseComponentInst);
 				RootBoneOffset = BaseComponent->RootBoneTranslation; // Adjust bounds by root bone translation
 			}
 		}
@@ -3226,7 +3226,7 @@ void USkeletalMeshComponent::HideBone( int32 BoneIndex, EPhysBodyOp PhysBodyOpti
 		return;
 	}
 
-	if (MasterPoseComponent.IsValid())
+	if (LeaderPoseComponent.IsValid())
 	{
 		return;
 	}
@@ -3260,7 +3260,7 @@ void USkeletalMeshComponent::UnHideBone( int32 BoneIndex )
 		return;
 	}
 
-	if (MasterPoseComponent.IsValid())
+	if (LeaderPoseComponent.IsValid())
 	{
 		return;
 	}
@@ -3911,11 +3911,11 @@ void USkeletalMeshComponent::RefreshMorphTargets()
 			PostProcessAnimInstance->RefreshCurves(this);
 		}
 	}
-	else if (USkeletalMeshComponent* MasterSMC = Cast<USkeletalMeshComponent>(MasterPoseComponent.Get()))
+	else if (USkeletalMeshComponent* LeaderSMC = Cast<USkeletalMeshComponent>(LeaderPoseComponent.Get()))
 	{
-		if (MasterSMC->AnimScriptInstance)
+		if (LeaderSMC->AnimScriptInstance)
 		{
-			MasterSMC->AnimScriptInstance->RefreshCurves(this);
+			LeaderSMC->AnimScriptInstance->RefreshCurves(this);
 		}
 	}
 	
@@ -4069,37 +4069,37 @@ bool USkeletalMeshComponent::IsClothingSimulationSuspended() const
 	return bClothingSimulationSuspended;
 }
 
-void USkeletalMeshComponent::BindClothToMasterPoseComponent()
+void USkeletalMeshComponent::BindClothToLeaderPoseComponent()
 {
-	if(USkeletalMeshComponent* MasterComp = Cast<USkeletalMeshComponent>(MasterPoseComponent.Get()))
+	if(USkeletalMeshComponent* LeaderComp = Cast<USkeletalMeshComponent>(LeaderPoseComponent.Get()))
 	{
-		if(GetSkeletalMesh() != MasterComp->GetSkeletalMesh())
+		if(GetSkeletalMesh() != LeaderComp->GetSkeletalMesh())
 		{
 			// Not the same mesh, can't bind
 			return;
 		}
 
-		if(ClothingSimulation && MasterComp->ClothingSimulation)
+		if(ClothingSimulation && LeaderComp->ClothingSimulation)
 		{
 			bDisableClothSimulation = true;
 
-			// When we extract positions from now we'll just take the master components positions
-			bBindClothToMasterComponent = true;
+			// When we extract positions from now we'll just take the Leader components positions
+			bBindClothToLeaderComponent = true;
 		}
 	}
 }
 
-void USkeletalMeshComponent::UnbindClothFromMasterPoseComponent(bool bRestoreSimulationSpace)
+void USkeletalMeshComponent::UnbindClothFromLeaderPoseComponent(bool bRestoreSimulationSpace)
 {
-	USkeletalMeshComponent* MasterComp = Cast<USkeletalMeshComponent>(MasterPoseComponent.Get());
-	if(MasterComp && bBindClothToMasterComponent)
+	USkeletalMeshComponent* LeaderComp = Cast<USkeletalMeshComponent>(LeaderPoseComponent.Get());
+	if(LeaderComp && bBindClothToLeaderComponent)
 	{
 		if(ClothingSimulation)
 		{
 			bDisableClothSimulation = false;
 		}
 
-		bBindClothToMasterComponent = false;
+		bBindClothToLeaderComponent = false;
 	}
 }
 
@@ -4290,9 +4290,9 @@ bool USkeletalMeshComponent::MoveComponentImpl(const FVector& Delta, const FQuat
 	return bSuccess;
 }
 
-void USkeletalMeshComponent::AddSlavePoseComponent(USkinnedMeshComponent* SkinnedMeshComponent)
+void USkeletalMeshComponent::AddFollowerPoseComponent(USkinnedMeshComponent* SkinnedMeshComponent)
 {
-	Super::AddSlavePoseComponent(SkinnedMeshComponent);
+	Super::AddFollowerPoseComponent(SkinnedMeshComponent);
 
 	if(USkeletalMeshComponent* SkeletalMeshComponent = Cast<USkeletalMeshComponent>(SkinnedMeshComponent))
 	{
@@ -4302,9 +4302,9 @@ void USkeletalMeshComponent::AddSlavePoseComponent(USkinnedMeshComponent* Skinne
 	bRequiredBonesUpToDate = false;
 }
 
-void USkeletalMeshComponent::RemoveSlavePoseComponent(USkinnedMeshComponent* SkinnedMeshComponent)
+void USkeletalMeshComponent::RemoveFollowerPoseComponent(USkinnedMeshComponent* SkinnedMeshComponent)
 {
-	Super::RemoveSlavePoseComponent(SkinnedMeshComponent);
+	Super::RemoveFollowerPoseComponent(SkinnedMeshComponent);
 
 	if(USkeletalMeshComponent* SkeletalMeshComponent = Cast<USkeletalMeshComponent>(SkinnedMeshComponent))
 	{
diff --git a/Engine/Source/Runtime/Engine/Private/Components/SkinnedMeshComponent.cpp b/Engine/Source/Runtime/Engine/Private/Components/SkinnedMeshComponent.cpp
--- a/Engine/Source/Runtime/Engine/Private/Components/SkinnedMeshComponent.cpp
+++ b/Engine/Source/Runtime/Engine/Private/Components/SkinnedMeshComponent.cpp
@@ -394,9 +394,9 @@ USkinnedMeshComponent::USkinnedMeshComponent(const FObjectInitializer& ObjectIni
 	bUpdateDeformerAtNextTick = false;
 
 	bCanEverAffectNavigation = false;
-	MasterBoneMapCacheCount = 0;
+	LeaderBoneMapCacheCount = 0;
 	bSyncAttachParentLOD = true;
-	bIgnoreMasterPoseComponentLOD = false;
+	bIgnoreLeaderPoseComponentLOD = false;
 
 	CurrentBoneTransformRevisionNumber = 0;
 
@@ -508,7 +508,7 @@ void USkinnedMeshComponent::Serialize(FArchive& Ar)
 		// add all native variables - mostly bigger chunks 
 		ComponentSpaceTransformsArray[0].CountBytes(Ar);
 		ComponentSpaceTransformsArray[1].CountBytes(Ar);
-		MasterBoneMap.CountBytes(Ar);
+		LeaderBoneMap.CountBytes(Ar);
 	}
 }
 
@@ -524,10 +524,10 @@ void USkinnedMeshComponent::OnRegister()
 
 	AnimUpdateRateParams = FAnimUpdateRateManager::GetUpdateRateParameters(this);
 
-	if (MasterPoseComponent.IsValid())
+	if (LeaderPoseComponent.IsValid())
 	{
 		// we have to make sure it updates the mastesr pose
-		SetMasterPoseComponent(MasterPoseComponent.Get(), true);
+		SetLeaderPoseComponent(LeaderPoseComponent.Get(), true);
 	}
 	else
 	{
@@ -768,9 +768,9 @@ bool USkinnedMeshComponent::RequiresGameThreadEndOfFrameRecreate() const
 	}
 #endif
 
-	// When we are a master/slave, we cannot recreate render state in parallel as this could 
+	// When we are a leader/follower, we cannot recreate render state in parallel as this could 
 	// happen concurrently with our dependent component(s)
-	return MasterPoseComponent.Get() != nullptr || SlavePoseComponents.Num() > 0;
+	return LeaderPoseComponent.Get() != nullptr || FollowerPoseComponents.Num() > 0;
 }
 
 FString USkinnedMeshComponent::GetDetailedInfoInternal() const
@@ -1010,9 +1010,9 @@ void USkinnedMeshComponent::TickComponent(float DeltaTime, enum ELevelTick TickT
 	if( ShouldUpdateTransform(bLODHasChanged) )
 	{
 		// Do not update bones if we are taking bone transforms from another SkelMeshComp
-		if( MasterPoseComponent.IsValid() )
+		if( LeaderPoseComponent.IsValid() )
 		{
-			UpdateSlaveComponent();
+			UpdateFollowerComponent();
 		}
 		else 
 		{
@@ -1048,8 +1048,7 @@ UObject const* USkinnedMeshComponent::AdditionalStatObject() const
 	return GetSkeletalMesh();
 }
 
-
-void USkinnedMeshComponent::UpdateSlaveComponent()
+void USkinnedMeshComponent::UpdateFollowerComponent()
 {
 	MarkRenderDynamicDataDirty();
 }
@@ -1200,12 +1199,12 @@ void USkinnedMeshComponent::GetStreamingRenderAssetInfo(FStreamingTextureLevelCo
 
 bool USkinnedMeshComponent::ShouldUpdateBoneVisibility() const
 {
-	// do not update if it has MasterPoseComponent
-	return !MasterPoseComponent.IsValid();
+	// do not update if it has LeaderPoseComponent
+	return !LeaderPoseComponent.IsValid();
 }
 void USkinnedMeshComponent::RebuildVisibilityArray()
 {
-	// BoneVisibility needs update if MasterComponent == NULL
+	// BoneVisibility needs update if LeaderComponent == NULL
 	// if MaterComponent, it should follow MaterPoseComponent
 	if ( ShouldUpdateBoneVisibility())
 	{
@@ -1297,21 +1296,21 @@ FBoxSphereBounds USkinnedMeshComponent::CalcMeshBound(const FVector3f& RootOffse
 	AActor* Owner = GetOwner();
 	FVector DrawScale = LocalToWorld.GetScale3D();	
 
-	const USkinnedMeshComponent* const MasterPoseComponentInst = MasterPoseComponent.Get();
+	const USkinnedMeshComponent* const LeaderPoseComponentInst = LeaderPoseComponent.Get();
 	UPhysicsAsset * const PhysicsAsset = GetPhysicsAsset();
-	UPhysicsAsset * const MasterPhysicsAsset = (MasterPoseComponentInst != nullptr)? MasterPoseComponentInst->GetPhysicsAsset() : nullptr;
+	UPhysicsAsset * const LeaderPhysicsAsset = (LeaderPoseComponentInst != nullptr)? LeaderPoseComponentInst->GetPhysicsAsset() : nullptr;
 
 	// Can only use the PhysicsAsset to calculate the bounding box if we are not non-uniformly scaling the mesh.
 	const USkeletalMesh* SkeletalMeshConst = GetSkeletalMesh();
 	const bool bCanUsePhysicsAsset = DrawScale.IsUniform() && (GetSkeletalMesh() != NULL)
 		// either space base exists or child component
-		&& ( (GetNumComponentSpaceTransforms() == SkeletalMeshConst->GetRefSkeleton().GetNum()) || (MasterPhysicsAsset) );
+		&& ( (GetNumComponentSpaceTransforms() == SkeletalMeshConst->GetRefSkeleton().GetNum()) || (LeaderPhysicsAsset) );
 
 	const bool bDetailModeAllowsRendering = (DetailMode <= GetCachedScalabilityCVars().DetailMode);
 	const bool bIsVisible = ( bDetailModeAllowsRendering && (ShouldRender() || bCastHiddenShadow));
 
 	const bool bHasPhysBodies = PhysicsAsset && PhysicsAsset->SkeletalBodySetups.Num();
-	const bool bMasterHasPhysBodies = MasterPhysicsAsset && MasterPhysicsAsset->SkeletalBodySetups.Num();
+	const bool bLeaderHasPhysBodies = LeaderPhysicsAsset && LeaderPhysicsAsset->SkeletalBodySetups.Num();
 
 	// if not visible, or we were told to use fixed bounds, use skelmesh bounds
 	if ( (!bIsVisible || bComponentUseFixedSkelBounds) && GetSkeletalMesh())
@@ -1320,16 +1319,16 @@ FBoxSphereBounds USkinnedMeshComponent::CalcMeshBound(const FVector3f& RootOffse
 		RootAdjustedBounds.Origin += FVector(RootOffset); // Adjust bounds by root bone translation
 		NewBounds = RootAdjustedBounds.TransformBy(LocalToWorld);
 	}
-	else if(MasterPoseComponentInst && MasterPoseComponentInst->GetSkeletalMesh() && MasterPoseComponentInst->bComponentUseFixedSkelBounds)
+	else if(LeaderPoseComponentInst && LeaderPoseComponentInst->GetSkeletalMesh() && LeaderPoseComponentInst->bComponentUseFixedSkelBounds)
 	{
-		FBoxSphereBounds RootAdjustedBounds = MasterPoseComponentInst->GetSkeletalMesh()->GetBounds();
+		FBoxSphereBounds RootAdjustedBounds = LeaderPoseComponentInst->GetSkeletalMesh()->GetBounds();
 		RootAdjustedBounds.Origin += FVector(RootOffset); // Adjust bounds by root bone translation
 		NewBounds = RootAdjustedBounds.TransformBy(LocalToWorld);
 	}
-	// Use MasterPoseComponent's PhysicsAsset if told to
-	else if (MasterPoseComponentInst && bCanUsePhysicsAsset && bUseBoundsFromMasterPoseComponent)
+	// Use LeaderPoseComponent's PhysicsAsset if told to
+	else if (LeaderPoseComponentInst && bCanUsePhysicsAsset && bUseBoundsFromLeaderPoseComponent)
 	{
-		NewBounds = MasterPoseComponentInst->CalcBounds(LocalToWorld);
+		NewBounds = LeaderPoseComponentInst->CalcBounds(LocalToWorld);
 	}
 #if WITH_EDITOR
 	// For AnimSet Viewer, use 'bounds preview' physics asset if present.
@@ -1343,10 +1342,10 @@ FBoxSphereBounds USkinnedMeshComponent::CalcMeshBound(const FVector3f& RootOffse
 	{
 		NewBounds = FBoxSphereBounds(PhysicsAsset->CalcAABB(this, LocalToWorld));
 	}
-	// Use MasterPoseComponent's PhysicsAsset, if we don't have one and it does
-	else if(MasterPoseComponentInst && bCanUsePhysicsAsset && bMasterHasPhysBodies)
+	// Use LeaderPoseComponent's PhysicsAsset, if we don't have one and it does
+	else if(LeaderPoseComponentInst && bCanUsePhysicsAsset && bLeaderHasPhysBodies)
 	{
-		NewBounds = FBoxSphereBounds(MasterPhysicsAsset->CalcAABB(this, LocalToWorld));
+		NewBounds = FBoxSphereBounds(LeaderPhysicsAsset->CalcAABB(this, LocalToWorld));
 	}
 	// Fallback is to use the one from the skeletal mesh. Usually pretty bad in terms of Accuracy of where the SkelMesh Bounds are located (i.e. usually bigger than it needs to be)
 	else if(GetSkeletalMesh())
@@ -1372,17 +1371,17 @@ FBoxSphereBounds USkinnedMeshComponent::CalcMeshBound(const FVector3f& RootOffse
 
 void USkinnedMeshComponent::GetPreSkinnedLocalBounds(FBoxSphereBounds& OutBounds) const
 {
-	const USkinnedMeshComponent* const MasterPoseComponentInst = MasterPoseComponent.Get();
+	const USkinnedMeshComponent* const LeaderPoseComponentInst = LeaderPoseComponent.Get();
 
 	if (GetSkeletalMesh())
 	{
 		// Get the Pre-skinned bounds from the skeletal mesh. Note that these bounds are the "ExtendedBounds", so they can be tweaked on the SkeletalMesh   
 		OutBounds = GetSkeletalMesh()->GetBounds();
 	}
-	else if(MasterPoseComponentInst && MasterPoseComponentInst->GetSkeletalMesh())
+	else if(LeaderPoseComponentInst && LeaderPoseComponentInst->GetSkeletalMesh())
 	{
-		// Get the bounds from the master pose if ther is no skeletal mesh
-		OutBounds = MasterPoseComponentInst->GetSkeletalMesh()->GetBounds();
+		// Get the bounds from the leader pose if ther is no skeletal mesh
+		OutBounds = LeaderPoseComponentInst->GetSkeletalMesh()->GetBounds();
 	}
 	else
 	{
@@ -1400,29 +1399,29 @@ FMatrix USkinnedMeshComponent::GetBoneMatrix(int32 BoneIdx) const
 		return FMatrix::Identity;
 	}
 
-	// Handle case of use a MasterPoseComponent - get bone matrix from there.
-	const USkinnedMeshComponent* const MasterPoseComponentInst = MasterPoseComponent.Get();
-	if(MasterPoseComponentInst)
+	// Handle case of use a LeaderPoseComponent - get bone matrix from there.
+	const USkinnedMeshComponent* const LeaderPoseComponentInst = LeaderPoseComponent.Get();
+	if(LeaderPoseComponentInst)
 	{
-		if(BoneIdx < MasterBoneMap.Num())
+		if(BoneIdx < LeaderBoneMap.Num())
 		{
-			int32 ParentBoneIndex = MasterBoneMap[BoneIdx];
+			int32 ParentBoneIndex = LeaderBoneMap[BoneIdx];
 
-			// If ParentBoneIndex is valid, grab matrix from MasterPoseComponent.
+			// If ParentBoneIndex is valid, grab matrix from LeaderPoseComponent.
 			if(	ParentBoneIndex != INDEX_NONE && 
-				ParentBoneIndex < MasterPoseComponentInst->GetNumComponentSpaceTransforms())
+				ParentBoneIndex < LeaderPoseComponentInst->GetNumComponentSpaceTransforms())
 			{
-				return MasterPoseComponentInst->GetComponentSpaceTransforms()[ParentBoneIndex].ToMatrixWithScale() * GetComponentTransform().ToMatrixWithScale();
+				return LeaderPoseComponentInst->GetComponentSpaceTransforms()[ParentBoneIndex].ToMatrixWithScale() * GetComponentTransform().ToMatrixWithScale();
 			}
 			else
 			{
-				UE_LOG(LogSkinnedMeshComp, Verbose, TEXT("GetBoneMatrix : ParentBoneIndex(%d:%s) out of range of MasterPoseComponent->SpaceBases for %s(%s)"), BoneIdx, *GetNameSafe(MasterPoseComponentInst->GetSkeletalMesh()), *GetNameSafe(GetSkeletalMesh()), *GetPathName());
+				UE_LOG(LogSkinnedMeshComp, Verbose, TEXT("GetBoneMatrix : ParentBoneIndex(%d:%s) out of range of LeaderPoseComponent->SpaceBases for %s(%s)"), BoneIdx, *GetNameSafe(LeaderPoseComponentInst->GetSkeletalMesh()), *GetNameSafe(GetSkeletalMesh()), *GetPathName());
 				return FMatrix::Identity;
 			}
 		}
 		else
 		{
-			UE_LOG(LogSkinnedMeshComp, Warning, TEXT("GetBoneMatrix : BoneIndex(%d) out of range of MasterBoneMap for %s (%s)"), BoneIdx, *this->GetFName().ToString(), GetSkeletalMesh() ? *GetSkeletalMesh()->GetFName().ToString() : TEXT("NULL") );
+			UE_LOG(LogSkinnedMeshComp, Warning, TEXT("GetBoneMatrix : BoneIndex(%d) out of range of LeaderBoneMap for %s (%s)"), BoneIdx, *this->GetFName().ToString(), GetSkeletalMesh() ? *GetSkeletalMesh()->GetFName().ToString() : TEXT("NULL") );
 			return FMatrix::Identity;
 		}
 	}
@@ -1455,56 +1454,56 @@ FTransform USkinnedMeshComponent::GetBoneTransform(int32 BoneIdx) const
 
 FTransform USkinnedMeshComponent::GetBoneTransform(int32 BoneIdx, const FTransform& LocalToWorld) const
 {
-	// Handle case of use a MasterPoseComponent - get bone matrix from there.
-	const USkinnedMeshComponent* const MasterPoseComponentInst = MasterPoseComponent.Get();
-	if(MasterPoseComponentInst)
+	// Handle case of use a LeaderPoseComponent - get bone matrix from there.
+	const USkinnedMeshComponent* const LeaderPoseComponentInst = LeaderPoseComponent.Get();
+	if(LeaderPoseComponentInst)
 	{
-		if (!MasterPoseComponentInst->IsRegistered())
+		if (!LeaderPoseComponentInst->IsRegistered())
 		{
-			// We aren't going to get anything valid from the master pose if it
+			// We aren't going to get anything valid from the leader pose if it
 			// isn't valid so for now return identity
 			return FTransform::Identity;
 		}
-		if(BoneIdx < MasterBoneMap.Num())
+		if(BoneIdx < LeaderBoneMap.Num())
 		{
-			const int32 MasterBoneIndex = MasterBoneMap[BoneIdx];
-			const int32 NumMasterTransforms = MasterPoseComponentInst->GetNumComponentSpaceTransforms();
+			const int32 LeaderBoneIndex = LeaderBoneMap[BoneIdx];
+			const int32 NumLeaderTransforms = LeaderPoseComponentInst->GetNumComponentSpaceTransforms();
 
-			// If MasterBoneIndex is valid, grab matrix from MasterPoseComponent.
-			if(MasterBoneIndex >= 0 && MasterBoneIndex < NumMasterTransforms)
+			// If LeaderBoneIndex is valid, grab matrix from LeaderPoseComponent.
+			if(LeaderBoneIndex >= 0 && LeaderBoneIndex < NumLeaderTransforms)
 			{
-				return MasterPoseComponentInst->GetComponentSpaceTransforms()[MasterBoneIndex] * LocalToWorld;
+				return LeaderPoseComponentInst->GetComponentSpaceTransforms()[LeaderBoneIndex] * LocalToWorld;
 			}
 			else
 			{
 				// Is this a missing bone we have cached?
-				FMissingMasterBoneCacheEntry MissingBoneInfo;
-				const FMissingMasterBoneCacheEntry* MissingBoneInfoPtr = MissingMasterBoneMap.Find(BoneIdx);
+				FMissingLeaderBoneCacheEntry MissingBoneInfo;
+				const FMissingLeaderBoneCacheEntry* MissingBoneInfoPtr = MissingLeaderBoneMap.Find(BoneIdx);
 				if(MissingBoneInfoPtr != nullptr)
 				{
-					const int32 MissingMasterBoneIndex = MissingBoneInfoPtr->CommonAncestorBoneIndex;
-					if(MissingMasterBoneIndex >= 0 && MissingMasterBoneIndex < NumMasterTransforms)
+					const int32 MissingLeaderBoneIndex = MissingBoneInfoPtr->CommonAncestorBoneIndex;
+					if(MissingLeaderBoneIndex >= 0 && MissingLeaderBoneIndex < NumLeaderTransforms)
 					{
-						return MissingBoneInfoPtr->RelativeTransform * MasterPoseComponentInst->GetComponentSpaceTransforms()[MissingBoneInfoPtr->CommonAncestorBoneIndex] * LocalToWorld;
+						return MissingBoneInfoPtr->RelativeTransform * LeaderPoseComponentInst->GetComponentSpaceTransforms()[MissingBoneInfoPtr->CommonAncestorBoneIndex] * LocalToWorld;
 					}
 				}
 				// Otherwise we might be able to generate the missing transform on the fly (although this is expensive)
-				else if(GetMissingMasterBoneRelativeTransform(BoneIdx, MissingBoneInfo))
+				else if(GetMissingLeaderBoneRelativeTransform(BoneIdx, MissingBoneInfo))
 				{
-					const int32 MissingMasterBoneIndex = MissingBoneInfo.CommonAncestorBoneIndex;
-					if (MissingMasterBoneIndex >= 0 && MissingMasterBoneIndex < NumMasterTransforms)
+					const int32 MissingLeaderBoneIndex = MissingBoneInfo.CommonAncestorBoneIndex;
+					if (MissingLeaderBoneIndex >= 0 && MissingLeaderBoneIndex < NumLeaderTransforms)
 					{
-						return MissingBoneInfo.RelativeTransform * MasterPoseComponentInst->GetComponentSpaceTransforms()[MissingBoneInfo.CommonAncestorBoneIndex] * LocalToWorld;
+						return MissingBoneInfo.RelativeTransform * LeaderPoseComponentInst->GetComponentSpaceTransforms()[MissingBoneInfo.CommonAncestorBoneIndex] * LocalToWorld;
 					}
 				}
 
-				UE_LOG(LogSkinnedMeshComp, Verbose, TEXT("GetBoneTransform : ParentBoneIndex(%d) out of range of MasterPoseComponent->SpaceBases for %s"), BoneIdx, *this->GetFName().ToString() );
+				UE_LOG(LogSkinnedMeshComp, Verbose, TEXT("GetBoneTransform : ParentBoneIndex(%d) out of range of LeaderPoseComponent->SpaceBases for %s"), BoneIdx, *this->GetFName().ToString() );
 				return FTransform::Identity;
 			}
 		}
 		else
 		{
-			UE_LOG(LogSkinnedMeshComp, Warning, TEXT("GetBoneTransform : BoneIndex(%d) out of range of MasterBoneMap for %s"), BoneIdx, *this->GetFName().ToString() );
+			UE_LOG(LogSkinnedMeshComp, Warning, TEXT("GetBoneTransform : BoneIndex(%d) out of range of LeaderBoneMap for %s"), BoneIdx, *this->GetFName().ToString() );
 			return FTransform::Identity;
 		}
 	}
@@ -1759,18 +1758,18 @@ void USkinnedMeshComponent::SetSkeletalMesh(USkeletalMesh* InSkelMesh, bool bRei
 		GetSkeletalMesh() = InSkelMesh;
 		SetPredictedLODLevel(0);
 
-		//SlavePoseComponents is an array of weak obj ptrs, so it can contain null elements
-		for (auto Iter = SlavePoseComponents.CreateIterator(); Iter; ++Iter)
+		//FollowerPoseComponents is an array of weak obj ptrs, so it can contain null elements
+		for (auto Iter = FollowerPoseComponents.CreateIterator(); Iter; ++Iter)
 		{
 			TWeakObjectPtr<USkinnedMeshComponent> Comp = (*Iter);
 			if (Comp.IsValid() == false)
 			{
-				SlavePoseComponents.RemoveAt(Iter.GetIndex());
+				FollowerPoseComponents.RemoveAt(Iter.GetIndex());
 				--Iter;
 			}
 			else
 			{
-				Comp->UpdateMasterBoneMap();
+				Comp->UpdateLeaderBoneMap();
 			}
 		}
 
@@ -1778,7 +1777,7 @@ void USkinnedMeshComponent::SetSkeletalMesh(USkeletalMesh* InSkelMesh, bool bRei
 		if (IsRegistered())
 		{
 			AllocateTransformData();
-			UpdateMasterBoneMap();
+			UpdateLeaderBoneMap();
 			InvalidateCachedBounds();
 			// clear morphtarget cache
 			ActiveMorphTargets.Empty();			
@@ -1826,7 +1825,7 @@ bool USkinnedMeshComponent::AllocateTransformData()
 	LLM_SCOPE_BYNAME(TEXT("SkeletalMesh/TransformData"));
 
 	// Allocate transforms if not present.
-	if (GetSkeletalMesh() != NULL && MasterPoseComponent == NULL )
+	if (GetSkeletalMesh() != NULL && LeaderPoseComponent == NULL )
 	{
 		if(GetNumComponentSpaceTransforms() != GetSkeletalMesh()->GetRefSkeleton().GetNum() )
 		{
@@ -1890,64 +1889,64 @@ void USkinnedMeshComponent::SetPhysicsAsset(class UPhysicsAsset* InPhysicsAsset,
 	PhysicsAssetOverride = InPhysicsAsset;
 }
 
-void USkinnedMeshComponent::SetMasterPoseComponent(class USkinnedMeshComponent* NewMasterBoneComponent, bool bForceUpdate)
+void USkinnedMeshComponent::SetLeaderPoseComponent(class USkinnedMeshComponent* NewLeaderBoneComponent, bool bForceUpdate)
 {
 	// Early out if we're already setup.
-	if (!bForceUpdate && NewMasterBoneComponent == MasterPoseComponent)
+	if (!bForceUpdate && NewLeaderBoneComponent == LeaderPoseComponent)
 	{
 		return;
 	}
 
-	USkinnedMeshComponent* OldMasterPoseComponent = MasterPoseComponent.Get();
-	USkinnedMeshComponent* ValidNewMasterPose = NewMasterBoneComponent;
+	USkinnedMeshComponent* OldLeaderPoseComponent = LeaderPoseComponent.Get();
+	USkinnedMeshComponent* ValidNewLeaderPose = NewLeaderBoneComponent;
 
-	// now add to slave components list, 
-	if (ValidNewMasterPose)
+	// now add to follower components list, 
+	if (ValidNewLeaderPose)
 	{
-		// verify if my current master pose is valid
-		// we can't have chain of master poses, so 
-		// we'll find the root master pose component
-		USkinnedMeshComponent* Iterator = ValidNewMasterPose;
-		while (Iterator->MasterPoseComponent.IsValid())
+		// verify if my current Leader pose is valid
+		// we can't have chain of Leader poses, so 
+		// we'll find the root Leader pose component
+		USkinnedMeshComponent* Iterator = ValidNewLeaderPose;
+		while (Iterator->LeaderPoseComponent.IsValid())
 		{
-			ValidNewMasterPose = Iterator->MasterPoseComponent.Get();
-			Iterator = ValidNewMasterPose;
+			ValidNewLeaderPose = Iterator->LeaderPoseComponent.Get();
+			Iterator = ValidNewLeaderPose;
 
 			// we have cycling, where in this chain, if it comes back to me, then reject it
 			if (Iterator == this)
 			{
 				ensureAlwaysMsgf(false,
-					TEXT("SetMasterPoseComponent detected loop (the input master pose chain point to itself. (%s <- %s)). Aborting... "),
-					*GetNameSafe(NewMasterBoneComponent), *GetNameSafe(this));
-				ValidNewMasterPose = nullptr;
+					TEXT("SetLeaderPoseComponent detected loop (the input Leader pose chain point to itself. (%s <- %s)). Aborting... "),
+					*GetNameSafe(NewLeaderBoneComponent), *GetNameSafe(this));
+				ValidNewLeaderPose = nullptr;
 				break;
 			}
 		}
 
-		// if we have valid master pose, compare with input data and we warn users
-		if (ValidNewMasterPose)
+		// if we have valid Leader pose, compare with input data and we warn users
+		if (ValidNewLeaderPose)
 		{
-			// Output if master is not same as input, which means it has changed. 
-			UE_CLOG(ValidNewMasterPose == NewMasterBoneComponent, LogSkinnedMeshComp, Verbose,
-				TEXT("MasterPoseComponent chain is detected (%s). We re-route to top-most MasterPoseComponent (%s)"),
-				*GetNameSafe(ValidNewMasterPose), *GetNameSafe(NewMasterBoneComponent));
+			// Output if Leader is not same as input, which means it has changed. 
+			UE_CLOG(ValidNewLeaderPose == NewLeaderBoneComponent, LogSkinnedMeshComp, Verbose,
+				TEXT("LeaderPoseComponent chain is detected (%s). We re-route to top-most LeaderPoseComponent (%s)"),
+				*GetNameSafe(ValidNewLeaderPose), *GetNameSafe(NewLeaderBoneComponent));
 		}
 	}
 
-	// now we have valid master pose, set it
-	MasterPoseComponent = ValidNewMasterPose;
-	if (ValidNewMasterPose)
+	// now we have valid Leader pose, set it
+	LeaderPoseComponent = ValidNewLeaderPose;
+	if (ValidNewLeaderPose)
 	{
 		bool bAddNew = true;
 		// make sure no empty element is there, this is weak obj ptr, so it will go away unless there is 
-		// other reference, this is intentional as master to slave reference is weak
-		for (auto Iter = ValidNewMasterPose->SlavePoseComponents.CreateIterator(); Iter; ++Iter)
+		// other reference, this is intentional as Leader to follower reference is weak
+		for (auto Iter = ValidNewLeaderPose->FollowerPoseComponents.CreateIterator(); Iter; ++Iter)
 		{
 			TWeakObjectPtr<USkinnedMeshComponent> Comp = (*Iter);
 			if (Comp.IsValid() == false)
 			{
 				// remove
-				ValidNewMasterPose->SlavePoseComponents.RemoveAt(Iter.GetIndex());
+				ValidNewLeaderPose->FollowerPoseComponents.RemoveAt(Iter.GetIndex());
 				--Iter;
 			}
 			// if it has same as me, ignore to add
@@ -1959,76 +1958,76 @@ void USkinnedMeshComponent::SetMasterPoseComponent(class USkinnedMeshComponent*
 
 		if (bAddNew)
 		{
-			ValidNewMasterPose->AddSlavePoseComponent(this);
+			ValidNewLeaderPose->AddFollowerPoseComponent(this);
 		}
 
-		// set up tick dependency between master & slave components
-		PrimaryComponentTick.AddPrerequisite(ValidNewMasterPose, ValidNewMasterPose->PrimaryComponentTick);
+		// set up tick dependency between leader & follower components
+		PrimaryComponentTick.AddPrerequisite(ValidNewLeaderPose, ValidNewLeaderPose->PrimaryComponentTick);
 	}
 
-	if ((OldMasterPoseComponent != nullptr) && (OldMasterPoseComponent != ValidNewMasterPose))
+	if ((OldLeaderPoseComponent != nullptr) && (OldLeaderPoseComponent != ValidNewLeaderPose))
 	{
-		OldMasterPoseComponent->RemoveSlavePoseComponent(this);
+		OldLeaderPoseComponent->RemoveFollowerPoseComponent(this);
 
-		// Only remove tick dependency if the old master pose comp isn't our attach parent. We should always have a tick dependency with our parent (see USceneComponent::AttachToComponent)
-		if (GetAttachParent() != OldMasterPoseComponent)
+		// Only remove tick dependency if the old leader pose comp isn't our attach parent. We should always have a tick dependency with our parent (see USceneComponent::AttachToComponent)
+		if (GetAttachParent() != OldLeaderPoseComponent)
 		{
-			// remove tick dependency between master & slave components
-			PrimaryComponentTick.RemovePrerequisite(OldMasterPoseComponent, OldMasterPoseComponent->PrimaryComponentTick);
+			// remove tick dependency between leader & follower components
+			PrimaryComponentTick.RemovePrerequisite(OldLeaderPoseComponent, OldLeaderPoseComponent->PrimaryComponentTick);
 		}
 	}
 
 	AllocateTransformData();
 	RecreatePhysicsState();
-	UpdateMasterBoneMap();
+	UpdateLeaderBoneMap();
 
-	// Update Slave in case Master has already been ticked, and we won't get an update for another frame.
-	if (ValidNewMasterPose)
+	// Update Follower in case Leader has already been ticked, and we won't get an update for another frame.
+	if (ValidNewLeaderPose)
 	{
-		// if I have master, but I also have slaves, they won't work anymore
-		// we have to reroute the slaves to new master
-		if (SlavePoseComponents.Num() > 0)
+		// if I have leader, but I also have followers, they won't work anymore
+		// we have to reroute the followers to new leader
+		if (FollowerPoseComponents.Num() > 0)
 		{
 			UE_LOG(LogSkinnedMeshComp, Verbose,
-				TEXT("MasterPoseComponent chain is detected (%s). We re-route all children to new MasterPoseComponent (%s)"),
-				*GetNameSafe(this), *GetNameSafe(ValidNewMasterPose));
+				TEXT("LeaderPoseComponent chain is detected (%s). We re-route all children to new LeaderPoseComponent (%s)"),
+				*GetNameSafe(this), *GetNameSafe(ValidNewLeaderPose));
 
-			// Walk through array in reverse, as changing the Slaves' MasterPoseComponent will remove them from our SlavePoseComponents array.
-			const int32 NumSlaves = SlavePoseComponents.Num();
-			for (int32 SlaveIndex = NumSlaves - 1; SlaveIndex >= 0; SlaveIndex--)
+			// Walk through array in reverse, as changing the Followers' LeaderPoseComponent will remove them from our FollowerPoseComponents array.
+			const int32 NumFollowers = FollowerPoseComponents.Num();
+			for (int32 FollowerIndex = NumFollowers - 1; FollowerIndex >= 0; FollowerIndex--)
 			{
-				if (USkinnedMeshComponent* SlaveComp = SlavePoseComponents[SlaveIndex].Get())
+				if (USkinnedMeshComponent* FollowerComp = FollowerPoseComponents[FollowerIndex].Get())
 				{
-					SlaveComp->SetMasterPoseComponent(ValidNewMasterPose);
+					FollowerComp->SetLeaderPoseComponent(ValidNewLeaderPose);
 				}
 			}
 		}
 
-		UpdateSlaveComponent();
+		UpdateFollowerComponent();
 	}
 }
 
-const TArray< TWeakObjectPtr<USkinnedMeshComponent> >& USkinnedMeshComponent::GetSlavePoseComponents() const
+const TArray< TWeakObjectPtr<USkinnedMeshComponent> >& USkinnedMeshComponent::GetFollowerPoseComponents() const
 {
-	return SlavePoseComponents;
+	return FollowerPoseComponents;
 }
 
-void USkinnedMeshComponent::AddSlavePoseComponent(USkinnedMeshComponent* SkinnedMeshComponent)
+void USkinnedMeshComponent::AddFollowerPoseComponent(USkinnedMeshComponent* SkinnedMeshComponent)
 {
-	SlavePoseComponents.AddUnique(SkinnedMeshComponent);
+	FollowerPoseComponents.AddUnique(SkinnedMeshComponent);
 }
 
-void USkinnedMeshComponent::RemoveSlavePoseComponent(USkinnedMeshComponent* SkinnedMeshComponent)
+void USkinnedMeshComponent::RemoveFollowerPoseComponent(USkinnedMeshComponent* SkinnedMeshComponent)
 {
-	SlavePoseComponents.Remove(SkinnedMeshComponent);
+	FollowerPoseComponents.Remove(SkinnedMeshComponent);
 }
 
 void USkinnedMeshComponent::InvalidateCachedBounds()
 {
 	bCachedLocalBoundsUpToDate = false;
 	bCachedWorldSpaceBoundsUpToDate = false;
-	// Also invalidate all slave components.
-	for (const TWeakObjectPtr<USkinnedMeshComponent>& SkinnedMeshComp : SlavePoseComponents)
+	// Also invalidate all follower components.
+	for (const TWeakObjectPtr<USkinnedMeshComponent>& SkinnedMeshComp : FollowerPoseComponents)
 	{
 		if (USkinnedMeshComponent* SkinnedMeshCompPtr = SkinnedMeshComp.Get())
 		{
@@ -2051,13 +2050,13 @@ void USkinnedMeshComponent::InvalidateCachedBounds()
 	}
 }
 
-void USkinnedMeshComponent::RefreshSlaveComponents()
+void USkinnedMeshComponent::RefreshFollowerComponents()
 {
-	for (const TWeakObjectPtr<USkinnedMeshComponent>& MeshComp : SlavePoseComponents)
+	for (const TWeakObjectPtr<USkinnedMeshComponent>& MeshComp : FollowerPoseComponents)
 	{
 		if (USkinnedMeshComponent* MeshCompPtr = MeshComp.Get())
 		{
-			// Update any children of the slave components if they are using sockets
+			// Update any children of the follower components if they are using sockets
 			MeshCompPtr->UpdateChildTransforms(EUpdateTransformFlags::OnlyUpdateIfUsingSocket);
 
 			MeshCompPtr->MarkRenderDynamicDataDirty();
@@ -2124,11 +2123,11 @@ UMorphTarget* USkinnedMeshComponent::FindMorphTarget( FName MorphTargetName ) co
 	return NULL;
 }
 
-bool USkinnedMeshComponent::GetMissingMasterBoneRelativeTransform(int32 InBoneIndex, FMissingMasterBoneCacheEntry& OutInfo) const
+bool USkinnedMeshComponent::GetMissingLeaderBoneRelativeTransform(int32 InBoneIndex, FMissingLeaderBoneCacheEntry& OutInfo) const
 {
-	const FReferenceSkeleton& SlaveRefSkeleton = GetSkeletalMesh()->GetRefSkeleton();
-	check(SlaveRefSkeleton.IsValidIndex(InBoneIndex));
-	const TArray<FTransform>& BoneSpaceRefPoseTransforms = SlaveRefSkeleton.GetRefBonePose();
+	const FReferenceSkeleton& FollowerRefSkeleton = GetSkeletalMesh()->GetRefSkeleton();
+	check(FollowerRefSkeleton.IsValidIndex(InBoneIndex));
+	const TArray<FTransform>& BoneSpaceRefPoseTransforms = FollowerRefSkeleton.GetRefBonePose();
 
 	OutInfo.CommonAncestorBoneIndex = INDEX_NONE;
 	OutInfo.RelativeTransform = FTransform::Identity;
@@ -2136,14 +2135,14 @@ bool USkinnedMeshComponent::GetMissingMasterBoneRelativeTransform(int32 InBoneIn
 	FTransform RelativeTransform = BoneSpaceRefPoseTransforms[InBoneIndex];
 
 	// we need to find a common base component-space transform in this skeletal mesh as it
-	// isnt present in the master, so run up the hierarchy
+	// isnt present in the leader, so run up the hierarchy
 	int32 CommonAncestorBoneIndex = InBoneIndex;
 	while(CommonAncestorBoneIndex != INDEX_NONE)
 	{
-		CommonAncestorBoneIndex = SlaveRefSkeleton.GetParentIndex(CommonAncestorBoneIndex);
+		CommonAncestorBoneIndex = FollowerRefSkeleton.GetParentIndex(CommonAncestorBoneIndex);
 		if(CommonAncestorBoneIndex != INDEX_NONE)
 		{
-			OutInfo.CommonAncestorBoneIndex = MasterBoneMap[CommonAncestorBoneIndex];
+			OutInfo.CommonAncestorBoneIndex = LeaderBoneMap[CommonAncestorBoneIndex];
 			if(OutInfo.CommonAncestorBoneIndex != INDEX_NONE)
 			{
 				OutInfo.RelativeTransform = RelativeTransform;
@@ -2157,50 +2156,50 @@ bool USkinnedMeshComponent::GetMissingMasterBoneRelativeTransform(int32 InBoneIn
 	return false;
 }
 
-void USkinnedMeshComponent::UpdateMasterBoneMap()
+void USkinnedMeshComponent::UpdateLeaderBoneMap()
 {
-	MasterBoneMap.Reset();
-	MissingMasterBoneMap.Reset();
+	LeaderBoneMap.Reset();
+	MissingLeaderBoneMap.Reset();
 
 	if (GetSkeletalMesh())
 	{
-		if (USkinnedMeshComponent* MasterPoseComponentPtr = MasterPoseComponent.Get())
+		if (USkinnedMeshComponent* LeaderPoseComponentPtr = LeaderPoseComponent.Get())
 		{
-			if (USkeletalMesh* MasterMesh = MasterPoseComponentPtr->GetSkeletalMesh())
+			if (USkeletalMesh* LeaderMesh = LeaderPoseComponentPtr->GetSkeletalMesh())
 			{
-				const FReferenceSkeleton& SlaveRefSkeleton = GetSkeletalMesh()->GetRefSkeleton();
-				const FReferenceSkeleton& MasterRefSkeleton = MasterMesh->GetRefSkeleton();
+				const FReferenceSkeleton& FollowerRefSkeleton = GetSkeletalMesh()->GetRefSkeleton();
+				const FReferenceSkeleton& LeaderRefSkeleton = LeaderMesh->GetRefSkeleton();
 
-				MasterBoneMap.AddUninitialized(SlaveRefSkeleton.GetNum());
-				if (GetSkeletalMesh() == MasterMesh)
+				LeaderBoneMap.AddUninitialized(FollowerRefSkeleton.GetNum());
+				if (GetSkeletalMesh() == LeaderMesh)
 				{
 					// if the meshes are the same, the indices must match exactly so we don't need to look them up
-					for (int32 BoneIndex = 0; BoneIndex < MasterBoneMap.Num(); BoneIndex++)
+					for (int32 BoneIndex = 0; BoneIndex < LeaderBoneMap.Num(); BoneIndex++)
 					{
-						MasterBoneMap[BoneIndex] = BoneIndex;
+						LeaderBoneMap[BoneIndex] = BoneIndex;
 					}
 				}
 				else
 				{
-					for (int32 BoneIndex = 0; BoneIndex < MasterBoneMap.Num(); BoneIndex++)
+					for (int32 BoneIndex = 0; BoneIndex < LeaderBoneMap.Num(); BoneIndex++)
 					{
-						const FName BoneName = SlaveRefSkeleton.GetBoneName(BoneIndex);
-						MasterBoneMap[BoneIndex] = MasterRefSkeleton.FindBoneIndex(BoneName);
+						const FName BoneName = FollowerRefSkeleton.GetBoneName(BoneIndex);
+						LeaderBoneMap[BoneIndex] = LeaderRefSkeleton.FindBoneIndex(BoneName);
 					}
 
-					// Cache bones for any SOCKET bones that are missing in the master.
+					// Cache bones for any SOCKET bones that are missing in the leader.
 					// We assume that sockets will be potentially called more often, so we
 					// leave out missing BONE transforms here to try to balance memory & performance.
 					for(USkeletalMeshSocket* Socket : GetSkeletalMesh()->GetActiveSocketList())
 					{
-						int32 BoneIndex = SlaveRefSkeleton.FindBoneIndex(Socket->BoneName);
-						int32 MasterBoneIndex = MasterRefSkeleton.FindBoneIndex(Socket->BoneName);
-						if(BoneIndex != INDEX_NONE && MasterBoneIndex == INDEX_NONE)
+						int32 BoneIndex = FollowerRefSkeleton.FindBoneIndex(Socket->BoneName);
+						int32 LeaderBoneIndex = LeaderRefSkeleton.FindBoneIndex(Socket->BoneName);
+						if(BoneIndex != INDEX_NONE && LeaderBoneIndex == INDEX_NONE)
 						{
-							FMissingMasterBoneCacheEntry MissingBoneInfo;
-							if(GetMissingMasterBoneRelativeTransform(BoneIndex, MissingBoneInfo))
+							FMissingLeaderBoneCacheEntry MissingBoneInfo;
+							if(GetMissingLeaderBoneRelativeTransform(BoneIndex, MissingBoneInfo))
 							{
-								MissingMasterBoneMap.Add(BoneIndex, MissingBoneInfo);
+								MissingLeaderBoneMap.Add(BoneIndex, MissingBoneInfo);
 							}
 						}
 					}
@@ -2209,7 +2208,7 @@ void USkinnedMeshComponent::UpdateMasterBoneMap()
 		}
 	}
 
-	MasterBoneMapCacheCount += 1;
+	LeaderBoneMapCacheCount += 1;
 }
 
 FTransform USkinnedMeshComponent::GetSocketTransform(FName InSocketName, ERelativeTransformSpace TransformSpace) const
@@ -2411,17 +2410,17 @@ FQuat USkinnedMeshComponent::GetBoneQuaternion(FName BoneName, EBoneSpaces::Type
 	FTransform BoneTransform;
 	if( Space == EBoneSpaces::ComponentSpace )
 	{
-		const USkinnedMeshComponent* const MasterPoseComponentInst = MasterPoseComponent.Get();
-		if(MasterPoseComponentInst)
+		const USkinnedMeshComponent* const LeaderPoseComponentInst = LeaderPoseComponent.Get();
+		if(LeaderPoseComponentInst)
 		{
-			if(BoneIndex < MasterBoneMap.Num())
+			if(BoneIndex < LeaderBoneMap.Num())
 			{
-				int32 ParentBoneIndex = MasterBoneMap[BoneIndex];
-				// If ParentBoneIndex is valid, grab matrix from MasterPoseComponent.
+				int32 ParentBoneIndex = LeaderBoneMap[BoneIndex];
+				// If ParentBoneIndex is valid, grab matrix from LeaderPoseComponent.
 				if(	ParentBoneIndex != INDEX_NONE && 
-					ParentBoneIndex < MasterPoseComponentInst->GetNumComponentSpaceTransforms())
+					ParentBoneIndex < LeaderPoseComponentInst->GetNumComponentSpaceTransforms())
 				{
-					BoneTransform = MasterPoseComponentInst->GetComponentSpaceTransforms()[ParentBoneIndex];
+					BoneTransform = LeaderPoseComponentInst->GetComponentSpaceTransforms()[ParentBoneIndex];
 				}
 				else
 				{
@@ -2461,17 +2460,17 @@ FVector USkinnedMeshComponent::GetBoneLocation(FName BoneName, EBoneSpaces::Type
 	{
 	case EBoneSpaces::ComponentSpace:
 	{
-		const USkinnedMeshComponent* const MasterPoseComponentInst = MasterPoseComponent.Get();
-		if(MasterPoseComponentInst)
+		const USkinnedMeshComponent* const LeaderPoseComponentInst = LeaderPoseComponent.Get();
+		if(LeaderPoseComponentInst)
 		{
-			if(BoneIndex < MasterBoneMap.Num())
+			if(BoneIndex < LeaderBoneMap.Num())
 			{
-				int32 ParentBoneIndex = MasterBoneMap[BoneIndex];
-				// If ParentBoneIndex is valid, grab transform from MasterPoseComponent.
+				int32 ParentBoneIndex = LeaderBoneMap[BoneIndex];
+				// If ParentBoneIndex is valid, grab transform from LeaderPoseComponent.
 				if(	ParentBoneIndex != INDEX_NONE && 
-					ParentBoneIndex < MasterPoseComponentInst->GetNumComponentSpaceTransforms())
+					ParentBoneIndex < LeaderPoseComponentInst->GetNumComponentSpaceTransforms())
 				{
-					return MasterPoseComponentInst->GetComponentSpaceTransforms()[ParentBoneIndex].GetLocation();
+					return LeaderPoseComponentInst->GetComponentSpaceTransforms()[ParentBoneIndex].GetLocation();
 				}
 			}
 
@@ -2641,7 +2640,7 @@ FName USkinnedMeshComponent::FindClosestBone(FVector TestLocation, FVector* Bone
 		float BestDistSquared = UE_BIG_NUMBER;
 		int32 BestIndex = -1;
 
-		const USkinnedMeshComponent* BaseComponent = MasterPoseComponent.IsValid() ? MasterPoseComponent.Get() : this;
+		const USkinnedMeshComponent* BaseComponent = LeaderPoseComponent.IsValid() ? LeaderPoseComponent.Get() : this;
 		const TArray<FTransform>& CompSpaceTransforms = BaseComponent->GetComponentSpaceTransforms();
 
 		for (int32 i = 0; i < BaseComponent->GetNumComponentSpaceTransforms(); i++)
@@ -3104,9 +3103,9 @@ bool USkinnedMeshComponent::IsBoneHidden( int32 BoneIndex ) const
 			return EditableBoneVisibilityStates[ BoneIndex ] != BVS_Visible;
 		}
 	}
-	else if (USkinnedMeshComponent* MasterPoseComponentPtr = MasterPoseComponent.Get())
+	else if (USkinnedMeshComponent* LeaderPoseComponentPtr = LeaderPoseComponent.Get())
 	{
-		return MasterPoseComponentPtr->IsBoneHidden( BoneIndex );
+		return LeaderPoseComponentPtr->IsBoneHidden( BoneIndex );
 	}
 
 	return false;
@@ -3262,12 +3261,12 @@ bool USkinnedMeshComponent::UpdateLODStatus()
 	return UpdateLODStatus_Internal(INDEX_NONE);
 }
 
-bool USkinnedMeshComponent::UpdateLODStatus_Internal(int32 InMasterPoseComponentPredictedLODLevel, bool bRequestedByMasterPoseComponent)
+bool USkinnedMeshComponent::UpdateLODStatus_Internal(int32 InLeaderPoseComponentPredictedLODLevel, bool bRequestedByLeaderPoseComponent)
 {
-	// Don't update LOD status for slave component unless it explicitly ignores its master component's LOD or if update is recursively requested by its master.
-	// This is because when UpdateLODStatus is called on master component, it updates the slave component LOD,
-	// therefore if slave component also calls UpdateLODStatus, it could overturn the result and cause LOD to be out of sync from its master.
-	if (MasterPoseComponent.IsValid() && !bIgnoreMasterPoseComponentLOD && !bRequestedByMasterPoseComponent)
+	// Don't update LOD status for follower component unless it explicitly ignores its leader component's LOD or if update is recursively requested by its leader.
+	// This is because when UpdateLODStatus is called on leader component, it updates the follower component LOD,
+	// therefore if follower component also calls UpdateLODStatus, it could overturn the result and cause LOD to be out of sync from its leader.
+	if (LeaderPoseComponent.IsValid() && !bIgnoreLeaderPoseComponentLOD && !bRequestedByLeaderPoseComponent)
 	{
 		return false;
 	}
@@ -3302,10 +3301,10 @@ bool USkinnedMeshComponent::UpdateLODStatus_Internal(int32 InMasterPoseComponent
 		}
 		else
 		{
-			// Match LOD of MasterPoseComponent if it exists.
-			if (InMasterPoseComponentPredictedLODLevel != INDEX_NONE && !bIgnoreMasterPoseComponentLOD)
+			// Match LOD of LeaderPoseComponent if it exists.
+			if (InLeaderPoseComponentPredictedLODLevel != INDEX_NONE && !bIgnoreLeaderPoseComponentLOD)
 			{
-				NewPredictedLODLevel = FMath::Clamp(InMasterPoseComponentPredictedLODLevel, 0, MaxLODIndex);
+				NewPredictedLODLevel = FMath::Clamp(InLeaderPoseComponentPredictedLODLevel, 0, MaxLODIndex);
 			}
 			else if (bSyncAttachParentLOD && GetAttachParent() && GetAttachParent()->IsA(USkinnedMeshComponent::StaticClass()))
 			{
@@ -3384,13 +3383,13 @@ bool USkinnedMeshComponent::UpdateLODStatus_Internal(int32 InMasterPoseComponent
 	PredictedLODLevel = NewPredictedLODLevel;
 	PRAGMA_ENABLE_DEPRECATION_WARNINGS
 
-	// also update slave component LOD status, as we may need to recalc required bones if this changes
+	// also update follower component LOD status, as we may need to recalc required bones if this changes
 	// independently of our LOD
-	for (const TWeakObjectPtr<USkinnedMeshComponent>& SlaveComponent : SlavePoseComponents)
+	for (const TWeakObjectPtr<USkinnedMeshComponent>& FollowerComponent : FollowerPoseComponents)
 	{
-		if (USkinnedMeshComponent* SlaveComponentPtr = SlaveComponent.Get())
+		if (USkinnedMeshComponent* FollowerComponentPtr = FollowerComponent.Get())
 		{
-			bLODChanged |= SlaveComponentPtr->UpdateLODStatus_Internal(NewPredictedLODLevel, /*bRequestedByMasterPoseComponent=*/true);
+			bLODChanged |= FollowerComponentPtr->UpdateLODStatus_Internal(NewPredictedLODLevel, /*bRequestedByLeaderPoseComponent=*/true);
 		}
 	}
 
@@ -3472,7 +3471,7 @@ void USkinnedMeshComponent::GetCPUSkinnedVertices(TArray<FFinalSkinVertex>& OutV
 	// Component state should be left unchanged at the end of this call.
 	USkinnedMeshComponent* MutableThis = const_cast<USkinnedMeshComponent*>(this);
 
-	USkinnedMeshComponent* PoseComponent = MutableThis->MasterPoseComponent.Get() ? MutableThis->MasterPoseComponent.Get() : MutableThis;
+	USkinnedMeshComponent* PoseComponent = MutableThis->LeaderPoseComponent.Get() ? MutableThis->LeaderPoseComponent.Get() : MutableThis;
 
 	int32 CachedForcedLOD = PoseComponent->GetForcedLOD();
 	PoseComponent->SetForcedLOD(InLODIndex + 1);
@@ -3554,16 +3553,16 @@ void USkinnedMeshComponent::BeginDestroy()
 		RefPoseOverride = nullptr;
 	}
 
-	// Disconnect slave components from this component if present.
+	// Disconnect follower components from this component if present.
 	// They will currently have no transforms allocated so will be
 	// in an invalid state when this component is destroyed
 	// Walk backwards as we'll be removing from this array
-	const int32 NumSlaveComponents = SlavePoseComponents.Num();
-	for(int32 SlaveIndex = NumSlaveComponents - 1; SlaveIndex >= 0; --SlaveIndex)
+	const int32 NumFollowerComponents = FollowerPoseComponents.Num();
+	for(int32 FollowerIndex = NumFollowerComponents - 1; FollowerIndex >= 0; --FollowerIndex)
 	{
-		if(USkinnedMeshComponent* Slave = SlavePoseComponents[SlaveIndex].Get())
+		if(USkinnedMeshComponent* Follower = FollowerPoseComponents[FollowerIndex].Get())
 		{
-			Slave->SetMasterPoseComponent(nullptr);
+			Follower->SetLeaderPoseComponent(nullptr);
 		}
 	}
 }
@@ -4308,8 +4307,8 @@ void GetTypedSkinnedTangentBasis(
 	OutTangentY = FVector3f::ZeroVector;
 	OutTangentZ = FVector3f::ZeroVector;
 
-	const USkinnedMeshComponent* const MasterPoseComponentInst = SkinnedComp->MasterPoseComponent.Get();
-	const USkinnedMeshComponent* BaseComponent = MasterPoseComponentInst ? MasterPoseComponentInst : SkinnedComp;
+	const USkinnedMeshComponent* const LeaderPoseComponentInst = SkinnedComp->LeaderPoseComponent.Get();
+	const USkinnedMeshComponent* BaseComponent = LeaderPoseComponentInst ? LeaderPoseComponentInst : SkinnedComp;
 
 	// Do soft skinning for this vertex.
 	const int32 BufferVertIndex = Section.GetVertexBufferIndex() + VertIndex;
@@ -4348,8 +4347,8 @@ FVector3f GetTypedSkinnedVertexPosition(
 {
 	FVector3f SkinnedPos(0, 0, 0);
 
-	const USkinnedMeshComponent* const MasterPoseComponentInst = SkinnedComp->MasterPoseComponent.Get();
-	const USkinnedMeshComponent* BaseComponent = MasterPoseComponentInst ? MasterPoseComponentInst : SkinnedComp;
+	const USkinnedMeshComponent* const LeaderPoseComponentInst = SkinnedComp->LeaderPoseComponent.Get();
+	const USkinnedMeshComponent* BaseComponent = LeaderPoseComponentInst ? LeaderPoseComponentInst : SkinnedComp;
 
 	// Do soft skinning for this vertex.
 	int32 BufferVertIndex = Section.GetVertexBufferIndex() + VertIndex;
@@ -4371,11 +4370,11 @@ FVector3f GetTypedSkinnedVertexPosition(
 		const int32 MeshBoneIndex = Section.BoneMap[BoneMapIndex];
 		int32 TransformBoneIndex = MeshBoneIndex;
 
-		if (MasterPoseComponentInst)
+		if (LeaderPoseComponentInst)
 		{
-			const TArray<int32>& MasterBoneMap = SkinnedComp->GetMasterBoneMap();
-			check(MasterBoneMap.Num() == SkinnedComp->GetSkeletalMesh()->GetRefSkeleton().GetNum());
-			TransformBoneIndex = MasterBoneMap[MeshBoneIndex];
+			const TArray<int32>& LeaderBoneMap = SkinnedComp->GetLeaderBoneMap();
+			check(LeaderBoneMap.Num() == SkinnedComp->GetSkeletalMesh()->GetRefSkeleton().GetNum());
+			TransformBoneIndex = LeaderBoneMap[MeshBoneIndex];
 		}
 
 		const float	Weight = (float)SkinWeightVertexBuffer.GetBoneWeight(BufferVertIndex, InfluenceIndex) / 255.0f;
diff --git a/Engine/Source/Runtime/Engine/Private/Particles/ParticleModules_Location.cpp b/Engine/Source/Runtime/Engine/Private/Particles/ParticleModules_Location.cpp
--- a/Engine/Source/Runtime/Engine/Private/Particles/ParticleModules_Location.cpp
+++ b/Engine/Source/Runtime/Engine/Private/Particles/ParticleModules_Location.cpp
@@ -2949,7 +2949,7 @@ bool UParticleModuleLocationSkelVertSurface::VertInfluencedByActiveBoneTyped(
 	FModuleLocationVertSurfaceInstancePayload* InstancePayload, 
 	int32* OutBoneIndex)
 {
-	const TArray<int32>& MasterBoneMap = InSkelMeshComponent->GetMasterBoneMap();
+	const TArray<int32>& LeaderBoneMap = InSkelMeshComponent->GetLeaderBoneMap();
 	// Get weights on this vertex
 	FSkinWeightVertexBuffer* WeightBuffer = InSkelMeshComponent->GetSkinWeightBuffer(LODIndex);
 	if (WeightBuffer)
@@ -2962,10 +2962,10 @@ bool UParticleModuleLocationSkelVertSurface::VertInfluencedByActiveBoneTyped(
 #endif
 		{
 			int32 BoneIndex = Section.BoneMap[WeightBuffer->GetBoneIndex(Section.GetVertexBufferIndex() + VertIndex, InfluenceIndex)];
-			if (InSkelMeshComponent->MasterPoseComponent.IsValid())
+			if (InSkelMeshComponent->LeaderPoseComponent.IsValid())
 			{
-				check(MasterBoneMap.Num() == InSkelMeshComponent->GetSkeletalMesh()->GetRefSkeleton().GetNum());
-				BoneIndex = MasterBoneMap[BoneIndex];
+				check(LeaderBoneMap.Num() == InSkelMeshComponent->GetSkeletalMesh()->GetRefSkeleton().GetNum());
+				BoneIndex = LeaderBoneMap[BoneIndex];
 			}
 
 			if (!InstancePayload->NumValidAssociatedBoneIndices || InstancePayload->ValidAssociatedBoneIndices.Contains(BoneIndex))
diff --git a/Engine/Source/Runtime/Engine/Private/PhysicsEngine/PhysAnim.cpp b/Engine/Source/Runtime/Engine/Private/PhysicsEngine/PhysAnim.cpp
--- a/Engine/Source/Runtime/Engine/Private/PhysicsEngine/PhysAnim.cpp
+++ b/Engine/Source/Runtime/Engine/Private/PhysicsEngine/PhysAnim.cpp
@@ -451,8 +451,8 @@ void USkeletalMeshComponent::FinalizeAnimationUpdate()
 	// New bone positions need to be sent to render thread
 	MarkRenderDynamicDataDirty();
 
-	// If we have any Slave Components, they need to be refreshed as well.
-	RefreshSlaveComponents();
+	// If we have any Follower Components, they need to be refreshed as well.
+	RefreshFollowerComponents();
 }
 
 void USkeletalMeshComponent::CompleteParallelBlendPhysics()
diff --git a/Engine/Source/Runtime/Engine/Private/SkeletalMeshComponentPhysics.cpp b/Engine/Source/Runtime/Engine/Private/SkeletalMeshComponentPhysics.cpp
--- a/Engine/Source/Runtime/Engine/Private/SkeletalMeshComponentPhysics.cpp
+++ b/Engine/Source/Runtime/Engine/Private/SkeletalMeshComponentPhysics.cpp
@@ -1560,7 +1560,7 @@ bool USkeletalMeshComponent::UpdateOverlapsImpl(const TOverlapArrayView* Pending
 bool USkeletalMeshComponent::ShouldCreatePhysicsState() const
 {
 	bool bShouldCreatePhysicsState = Super::ShouldCreatePhysicsState();
-	bShouldCreatePhysicsState &= (MasterPoseComponent.IsValid() == false);
+	bShouldCreatePhysicsState &= (LeaderPoseComponent.IsValid() == false);
 	
 	return bShouldCreatePhysicsState;
 }
@@ -2458,7 +2458,7 @@ bool USkeletalMeshComponent::GetClosestPointOnPhysicsAsset(const FVector& WorldP
 	if(PhysicsAsset && RefSkeleton)
 	{
 		const TArray<FTransform>& BoneTransforms = GetComponentSpaceTransforms();
-		const bool bHasMasterPoseComponent = MasterPoseComponent.IsValid();
+		const bool bHasLeaderPoseComponent = LeaderPoseComponent.IsValid();
 		const FVector ComponentPosition = GetComponentTransform().InverseTransformPosition(WorldPosition);
 	
 		float CurrentClosestDistance = FLT_MAX;
@@ -2472,7 +2472,7 @@ bool USkeletalMeshComponent::GetClosestPointOnPhysicsAsset(const FVector& WorldP
 			const int32 BoneIndex = RefSkeleton->FindBoneIndex(BoneName);
 			if(BoneIndex != INDEX_NONE)
 			{
-				const FTransform BoneTM = bHasMasterPoseComponent ? GetBoneTransform(BoneIndex) : BoneTransforms[BoneIndex];
+				const FTransform BoneTM = bHasLeaderPoseComponent ? GetBoneTransform(BoneIndex) : BoneTransforms[BoneIndex];
 				const float Dist = bApproximate ? (BoneTM.GetLocation() - ComponentPosition).SizeSquared() : BodySetupInstance->GetShortestDistanceToPoint(ComponentPosition, BoneTM);
 
 				if (Dist < CurrentClosestDistance)
@@ -2490,7 +2490,7 @@ bool USkeletalMeshComponent::GetClosestPointOnPhysicsAsset(const FVector& WorldP
 		{
 			bSuccess = true;
 
-			const FTransform BoneTM = bHasMasterPoseComponent ? GetBoneTransform(CurrentClosestBoneIndex) : (BoneTransforms[CurrentClosestBoneIndex] * GetComponentTransform());
+			const FTransform BoneTM = bHasLeaderPoseComponent ? GetBoneTransform(CurrentClosestBoneIndex) : (BoneTransforms[CurrentClosestBoneIndex] * GetComponentTransform());
 			ClosestPointOnPhysicsAsset.Distance = CurrentClosestBodySetup->GetClosestPointAndNormal(WorldPosition, BoneTM, ClosestPointOnPhysicsAsset.ClosestWorldPosition, ClosestPointOnPhysicsAsset.Normal);
 			ClosestPointOnPhysicsAsset.BoneName = CurrentClosestBodySetup->BoneName;
 		}
diff --git a/Engine/Source/Runtime/Engine/Private/SkeletalRender.cpp b/Engine/Source/Runtime/Engine/Private/SkeletalRender.cpp
--- a/Engine/Source/Runtime/Engine/Private/SkeletalRender.cpp
+++ b/Engine/Source/Runtime/Engine/Private/SkeletalRender.cpp
@@ -236,7 +236,7 @@ float FSkeletalMeshObject::GetScreenSize(int32 LODIndex) const
 Global functions
 -----------------------------------------------------------------------------*/
 
-void UpdateRefToLocalMatricesInner(TArray<FMatrix44f>& ReferenceToLocal, const TArray<FTransform>& ComponentTransform, const TArray<uint8>& BoneVisibilityStates, const TArray<int32>* MasterBoneMap,
+void UpdateRefToLocalMatricesInner(TArray<FMatrix44f>& ReferenceToLocal, const TArray<FTransform>& ComponentTransform, const TArray<uint8>& BoneVisibilityStates, const TArray<int32>* LeaderBoneMap,
 	const TArray<FMatrix44f>* RefBasesInvMatrix, const FReferenceSkeleton& RefSkeleton, const FSkeletalMeshRenderData* InSkeletalMeshRenderData, int32 LODIndex, const TArray<FBoneIndexType>* ExtraRequiredBoneIndices)
 {
 	const FSkeletalMeshLODRenderData& LOD = InSkeletalMeshRenderData->LODRenderData[LODIndex];
@@ -244,7 +244,7 @@ void UpdateRefToLocalMatricesInner(TArray<FMatrix44f>& ReferenceToLocal, const T
 	const TArray<FBoneIndexType>* RequiredBoneSets[3] = { &LOD.ActiveBoneIndices, ExtraRequiredBoneIndices, NULL };
 
 	const bool bBoneVisibilityStatesValid = BoneVisibilityStates.Num() == ComponentTransform.Num();
-	const bool bIsMasterCompValid = MasterBoneMap != nullptr;
+	const bool bIsLeaderCompValid = LeaderBoneMap != nullptr;
 	
 	for (int32 RequiredBoneSetIndex = 0; RequiredBoneSets[RequiredBoneSetIndex] != NULL; RequiredBoneSetIndex++)
 	{
@@ -260,23 +260,23 @@ void UpdateRefToLocalMatricesInner(TArray<FMatrix44f>& ReferenceToLocal, const T
 				// On the off chance the parent matrix isn't valid, revert to identity.
 				ReferenceToLocal[ThisBoneIndex] = FMatrix44f::Identity;
 
-				//if we have master pose component, we use MasterBoneMap to figure out the mapping
-				if( bIsMasterCompValid )
+				//if we have Leader pose component, we use LeaderBoneMap to figure out the mapping
+				if( bIsLeaderCompValid )
 				{
 					// If valid, use matrix from parent component.
-					const int32 MasterBoneIndex = (*MasterBoneMap)[ThisBoneIndex];
-					if (ComponentTransform.IsValidIndex(MasterBoneIndex))
+					const int32 LeaderBoneIndex = (*LeaderBoneMap)[ThisBoneIndex];
+					if (ComponentTransform.IsValidIndex(LeaderBoneIndex))
 					{
 						const int32 ParentIndex = RefSkeleton.GetParentIndex(ThisBoneIndex);
-						bool bNeedToHideBone = BoneVisibilityStates[MasterBoneIndex] != BVS_Visible;
+						bool bNeedToHideBone = BoneVisibilityStates[LeaderBoneIndex] != BVS_Visible;
 						if (bNeedToHideBone && ParentIndex != INDEX_NONE)
 						{
 							ReferenceToLocal[ThisBoneIndex] = ReferenceToLocal[ParentIndex].ApplyScale(0.f);
 						}
 						else
 						{
-							checkSlow(ComponentTransform[MasterBoneIndex].IsRotationNormalized());
-							ReferenceToLocal[ThisBoneIndex] = (FMatrix44f)ComponentTransform[MasterBoneIndex].ToMatrixWithScale();
+							checkSlow(ComponentTransform[LeaderBoneIndex].IsRotationNormalized());
+							ReferenceToLocal[ThisBoneIndex] = (FMatrix44f)ComponentTransform[LeaderBoneIndex].ToMatrixWithScale();
 						}
 					}
 					else
@@ -343,14 +343,14 @@ void UpdateRefToLocalMatricesInner(TArray<FMatrix44f>& ReferenceToLocal, const T
 void UpdateRefToLocalMatrices( TArray<FMatrix44f>& ReferenceToLocal, const USkinnedMeshComponent* InMeshComponent, const FSkeletalMeshRenderData* InSkeletalMeshRenderData, int32 LODIndex, const TArray<FBoneIndexType>* ExtraRequiredBoneIndices )
 {
 	const USkeletalMesh* const ThisMesh = InMeshComponent->SkeletalMesh;
-	const USkinnedMeshComponent* const MasterComp = InMeshComponent->MasterPoseComponent.Get();
+	const USkinnedMeshComponent* const LeaderComp = InMeshComponent->LeaderPoseComponent.Get();
 	const FSkeletalMeshLODRenderData& LOD = InSkeletalMeshRenderData->LODRenderData[LODIndex];
 
 	const FReferenceSkeleton& RefSkeleton = ThisMesh->GetRefSkeleton();
-	const TArray<int32>& MasterBoneMap = InMeshComponent->GetMasterBoneMap();
-	const bool bIsMasterCompValid = MasterComp && MasterBoneMap.Num() == ThisMesh->GetRefSkeleton().GetNum();
-	const TArray<FTransform>& ComponentTransform = (bIsMasterCompValid)? MasterComp->GetComponentSpaceTransforms() : InMeshComponent->GetComponentSpaceTransforms();
-	const TArray<uint8>& BoneVisibilityStates = (bIsMasterCompValid) ? MasterComp->GetBoneVisibilityStates() : InMeshComponent->GetBoneVisibilityStates();
+	const TArray<int32>& LeaderBoneMap = InMeshComponent->GetLeaderBoneMap();
+	const bool bIsLeaderCompValid = LeaderComp && LeaderBoneMap.Num() == ThisMesh->GetRefSkeleton().GetNum();
+	const TArray<FTransform>& ComponentTransform = (bIsLeaderCompValid)? LeaderComp->GetComponentSpaceTransforms() : InMeshComponent->GetComponentSpaceTransforms();
+	const TArray<uint8>& BoneVisibilityStates = (bIsLeaderCompValid) ? LeaderComp->GetBoneVisibilityStates() : InMeshComponent->GetBoneVisibilityStates();
 	// Get inv ref pose matrices
 	const TArray<FMatrix44f>* RefBasesInvMatrix = &ThisMesh->GetRefBasesInvMatrix();
 
@@ -391,20 +391,20 @@ void UpdateRefToLocalMatrices( TArray<FMatrix44f>& ReferenceToLocal, const USkin
 		return;
 	}
 
-	UpdateRefToLocalMatricesInner(ReferenceToLocal, ComponentTransform, BoneVisibilityStates, (bIsMasterCompValid)? &MasterBoneMap : nullptr, RefBasesInvMatrix, RefSkeleton, InSkeletalMeshRenderData, LODIndex, ExtraRequiredBoneIndices);
+	UpdateRefToLocalMatricesInner(ReferenceToLocal, ComponentTransform, BoneVisibilityStates, (bIsLeaderCompValid)? &LeaderBoneMap : nullptr, RefBasesInvMatrix, RefSkeleton, InSkeletalMeshRenderData, LODIndex, ExtraRequiredBoneIndices);
 }
 
 void UpdatePreviousRefToLocalMatrices(TArray<FMatrix44f>& ReferenceToLocal, const USkinnedMeshComponent* InMeshComponent, const FSkeletalMeshRenderData* InSkeletalMeshRenderData, int32 LODIndex, const TArray<FBoneIndexType>* ExtraRequiredBoneIndices)
 {
 	const USkeletalMesh* const ThisMesh = InMeshComponent->SkeletalMesh;
-	const USkinnedMeshComponent* const MasterComp = InMeshComponent->MasterPoseComponent.Get();
+	const USkinnedMeshComponent* const LeaderComp = InMeshComponent->LeaderPoseComponent.Get();
 	const FSkeletalMeshLODRenderData& LOD = InSkeletalMeshRenderData->LODRenderData[LODIndex];
 
 	const FReferenceSkeleton& RefSkeleton = ThisMesh->GetRefSkeleton();
-	const TArray<int32>& MasterBoneMap = InMeshComponent->GetMasterBoneMap();
-	const bool bIsMasterCompValid = MasterComp && MasterBoneMap.Num() == ThisMesh->GetRefSkeleton().GetNum();
-	const TArray<FTransform>& ComponentTransform = (bIsMasterCompValid) ? MasterComp->GetPreviousComponentTransformsArray() : InMeshComponent->GetPreviousComponentTransformsArray();
-	const TArray<uint8>& BoneVisibilityStates = (bIsMasterCompValid) ? MasterComp->GetPreviousBoneVisibilityStates() : InMeshComponent->GetPreviousBoneVisibilityStates();
+	const TArray<int32>& LeaderBoneMap = InMeshComponent->GetLeaderBoneMap();
+	const bool bIsLeaderCompValid = LeaderComp && LeaderBoneMap.Num() == ThisMesh->GetRefSkeleton().GetNum();
+	const TArray<FTransform>& ComponentTransform = (bIsLeaderCompValid) ? LeaderComp->GetPreviousComponentTransformsArray() : InMeshComponent->GetPreviousComponentTransformsArray();
+	const TArray<uint8>& BoneVisibilityStates = (bIsLeaderCompValid) ? LeaderComp->GetPreviousBoneVisibilityStates() : InMeshComponent->GetPreviousBoneVisibilityStates();
 	// Get inv ref pose matrices
 	const TArray<FMatrix44f>* RefBasesInvMatrix = &ThisMesh->GetRefBasesInvMatrix();
 	// Check if there is an override (and it's the right size)
@@ -443,7 +443,7 @@ void UpdatePreviousRefToLocalMatrices(TArray<FMatrix44f>& ReferenceToLocal, cons
 
 		return;
 	}
-	UpdateRefToLocalMatricesInner(ReferenceToLocal, ComponentTransform, BoneVisibilityStates, (bIsMasterCompValid) ? &MasterBoneMap : nullptr, RefBasesInvMatrix, RefSkeleton, InSkeletalMeshRenderData, LODIndex, ExtraRequiredBoneIndices);
+	UpdateRefToLocalMatricesInner(ReferenceToLocal, ComponentTransform, BoneVisibilityStates, (bIsLeaderCompValid) ? &LeaderBoneMap : nullptr, RefBasesInvMatrix, RefSkeleton, InSkeletalMeshRenderData, LODIndex, ExtraRequiredBoneIndices);
 }
 
 bool IsSkeletalMeshClothBlendEnabled()
diff --git a/Engine/Source/Runtime/Engine/Private/SkeletalRenderCPUSkin.cpp b/Engine/Source/Runtime/Engine/Private/SkeletalRenderCPUSkin.cpp
--- a/Engine/Source/Runtime/Engine/Private/SkeletalRenderCPUSkin.cpp
+++ b/Engine/Source/Runtime/Engine/Private/SkeletalRenderCPUSkin.cpp
@@ -1204,12 +1204,12 @@ bool FDynamicSkelMeshObjectDataCPUSkin::UpdateClothSimulationData(USkinnedMeshCo
 {
 	USkeletalMeshComponent* SimMeshComponent = Cast<USkeletalMeshComponent>(InMeshComponent);
 
-	if (InMeshComponent->MasterPoseComponent.IsValid() && (SimMeshComponent && SimMeshComponent->IsClothBoundToMasterComponent()))
+	if (InMeshComponent->LeaderPoseComponent.IsValid() && (SimMeshComponent && SimMeshComponent->IsClothBoundToLeaderComponent()))
 	{
 		USkeletalMeshComponent* SrcComponent = SimMeshComponent;
 
-		// if I have master, override sim component
-		SimMeshComponent = Cast<USkeletalMeshComponent>(InMeshComponent->MasterPoseComponent.Get());
+		// if I have Leader, override sim component
+		SimMeshComponent = Cast<USkeletalMeshComponent>(InMeshComponent->LeaderPoseComponent.Get());
 
 		// IF we don't have sim component that is skeletalmeshcomponent, just ignore
 		if (!SimMeshComponent)
diff --git a/Engine/Source/Runtime/Engine/Private/SkeletalRenderGPUSkin.cpp b/Engine/Source/Runtime/Engine/Private/SkeletalRenderGPUSkin.cpp
--- a/Engine/Source/Runtime/Engine/Private/SkeletalRenderGPUSkin.cpp
+++ b/Engine/Source/Runtime/Engine/Private/SkeletalRenderGPUSkin.cpp
@@ -2462,11 +2462,11 @@ bool FDynamicSkelMeshObjectDataGPUSkin::UpdateClothSimulationData(USkinnedMeshCo
 {
 	USkeletalMeshComponent* SimMeshComponent = Cast<USkeletalMeshComponent>(InMeshComponent);
 
-	if (InMeshComponent->MasterPoseComponent.IsValid() && (SimMeshComponent && SimMeshComponent->IsClothBoundToMasterComponent()))
+	if (InMeshComponent->LeaderPoseComponent.IsValid() && (SimMeshComponent && SimMeshComponent->IsClothBoundToLeaderComponent()))
 	{
 		USkeletalMeshComponent* SrcComponent = SimMeshComponent;
-		// if I have master, override sim component
-		SimMeshComponent = Cast<USkeletalMeshComponent>(InMeshComponent->MasterPoseComponent.Get());
+		// if I have leader, override sim component
+		SimMeshComponent = Cast<USkeletalMeshComponent>(InMeshComponent->LeaderPoseComponent.Get());
 		// IF we don't have sim component that is skeletalmeshcomponent, just ignore
 		if (!SimMeshComponent)
 		{
diff --git a/Engine/Source/Runtime/MovieSceneTracks/Private/Evaluation/MovieSceneSkeletalAnimationTemplate.cpp b/Engine/Source/Runtime/MovieSceneTracks/Private/Evaluation/MovieSceneSkeletalAnimationTemplate.cpp
--- a/Engine/Source/Runtime/MovieSceneTracks/Private/Evaluation/MovieSceneSkeletalAnimationTemplate.cpp
+++ b/Engine/Source/Runtime/MovieSceneTracks/Private/Evaluation/MovieSceneSkeletalAnimationTemplate.cpp
@@ -157,7 +157,7 @@ struct FPreAnimatedAnimationTokenProducer : IMovieScenePreAnimatedTokenProducer
 				Component->TickAnimation(0.f, false);
 
 				Component->RefreshBoneTransforms();
-				Component->RefreshSlaveComponents();
+				Component->RefreshFollowerComponents();
 				Component->UpdateComponentToWorld();
 				Component->FinalizeBoneTransform();
 				Component->MarkRenderTransformDirty();
@@ -362,7 +362,7 @@ namespace MovieScene
 				SkeletalMeshComponent->TickAnimation(0.f, false);
 
 				SkeletalMeshComponent->RefreshBoneTransforms();
-				SkeletalMeshComponent->RefreshSlaveComponents();
+				SkeletalMeshComponent->RefreshFollowerComponents();
 				SkeletalMeshComponent->UpdateComponentToWorld();
 				SkeletalMeshComponent->FinalizeBoneTransform();
 				SkeletalMeshComponent->MarkRenderTransformDirty();
