World Partition: convert commandlet will now recursively include sublevel script actor references.

#rnx

diff --git a/Engine/Source/Editor/UnrealEd/Private/Commandlets/WorldPartitionConvertCommandlet.cpp b/Engine/Source/Editor/UnrealEd/Private/Commandlets/WorldPartitionConvertCommandlet.cpp
--- a/Engine/Source/Editor/UnrealEd/Private/Commandlets/WorldPartitionConvertCommandlet.cpp
+++ b/Engine/Source/Editor/UnrealEd/Private/Commandlets/WorldPartitionConvertCommandlet.cpp
@@ -1204,14 +1204,14 @@ int32 UWorldPartitionConvertCommandlet::Main(const FString& Params)
 				LevelScriptActorReferences.Add(LevelScriptActor);
 
 				ULevelScriptBlueprint* LevelScriptBlueprint = SubLevel->GetLevelScriptBlueprint(true);
-				LevelScriptActorReferences.Append(ActorsReferencesUtils::GetActorReferences(LevelScriptBlueprint));
+				LevelScriptActorReferences.Append(ActorsReferencesUtils::GetActorReferences(LevelScriptBlueprint, RF_NoFlags, true));
 
 				for(AActor* Actor: SubLevel->Actors)
 				{
 					if(IsValid(Actor))
 					{
 						TSet<AActor*> ActorReferences;
-						ActorReferences.Append(ActorsReferencesUtils::GetActorReferences(Actor));
+						ActorReferences.Append(ActorsReferencesUtils::GetActorReferences(Actor, RF_NoFlags, true));
 
 						for (AActor* ActorReference : ActorReferences)
 						{
diff --git a/Engine/Source/Runtime/Engine/Private/ActorReferencesUtils.cpp b/Engine/Source/Runtime/Engine/Private/ActorReferencesUtils.cpp
--- a/Engine/Source/Runtime/Engine/Private/ActorReferencesUtils.cpp
+++ b/Engine/Source/Runtime/Engine/Private/ActorReferencesUtils.cpp
@@ -7,10 +7,11 @@
 class FArchiveGatherExternalActorRefs : public FArchiveUObject
 {
 public:
-	FArchiveGatherExternalActorRefs(UObject* InRoot, TSet<AActor*>& InActorReferences, EObjectFlags InRequiredFlags)
+	FArchiveGatherExternalActorRefs(UObject* InRoot, TSet<AActor*>& InActorReferences, EObjectFlags InRequiredFlags, bool bInRecursive)
 		: Root(InRoot)
 		, ActorReferences(InActorReferences)
 		, RequiredFlags(InRequiredFlags)
+		, bRecursive(bInRecursive)
 	{
 		// Don't gather transient actor references
 		SetIsPersistent(true);
@@ -38,7 +39,7 @@ public:
 			{
 				HandleObjectReference(Obj);
 
-				if (Obj->IsInOuter(Root))
+				if (bRecursive || Obj->IsInOuter(Root))
 				{
 					Obj->Serialize(*this);
 				}
@@ -78,16 +79,17 @@ private:
 	TSet<AActor*>& ActorReferences;
 	TSet<UObject*> SubObjects;
 	EObjectFlags RequiredFlags;
+	bool bRecursive;
 };
 
-TArray<AActor*> ActorsReferencesUtils::GetExternalActorReferences(UObject* Root)
+TArray<AActor*> ActorsReferencesUtils::GetExternalActorReferences(UObject* Root, bool bRecursive)
 {
-	return GetActorReferences(Root, RF_HasExternalPackage);
+	return GetActorReferences(Root, RF_HasExternalPackage, bRecursive);
 }
 
-TArray<AActor*> ActorsReferencesUtils::GetActorReferences(UObject* Root, EObjectFlags RequiredFlags)
+TArray<AActor*> ActorsReferencesUtils::GetActorReferences(UObject* Root, EObjectFlags RequiredFlags, bool bRecursive)
 {
 	TSet<AActor*> Result;
-	FArchiveGatherExternalActorRefs Ar(Root, Result, RequiredFlags);
+	FArchiveGatherExternalActorRefs Ar(Root, Result, RequiredFlags, bRecursive);
 	return Result.Array();
 }
\ No newline at end of file
diff --git a/Engine/Source/Runtime/Engine/Public/ActorReferencesUtils.h b/Engine/Source/Runtime/Engine/Public/ActorReferencesUtils.h
--- a/Engine/Source/Runtime/Engine/Public/ActorReferencesUtils.h
+++ b/Engine/Source/Runtime/Engine/Public/ActorReferencesUtils.h
@@ -12,10 +12,10 @@ namespace ActorsReferencesUtils
 	/**
 	 * Gather direct references to external actors from the root object.
 	 */
-	ENGINE_API TArray<AActor*> GetExternalActorReferences(UObject* Root);
+	ENGINE_API TArray<AActor*> GetExternalActorReferences(UObject* Root, bool bRecursive = false);
 
 	/**
 	 * Gather direct references to actors from the root object.
 	 */
-	ENGINE_API TArray<AActor*> GetActorReferences(UObject* Root, EObjectFlags RequiredFlags = RF_NoFlags);
+	ENGINE_API TArray<AActor*> GetActorReferences(UObject* Root, EObjectFlags RequiredFlags = RF_NoFlags, bool bRecursive = false);
 }
\ No newline at end of file
