Skeletal Mesh:
 - Replace the USkeletalMesh pointer by a USkinnedAsset pointer in USkinnedMeshComponent.
 - Manage backward compatibility through changes to codebase to a new Getter/Setter API, transfer pointer values in PostLoad/Getter, and deprecate direct pointer property access.
- Add a GetSkeletalMesh function to USkeletalMeshComponent to facilitate the transition, cast to the correct type, and maintain blueprint functions.

Chaos Cloth:
- Update UClothComponent to inehrit from USkinnedMeshComponent.
- Complete first pass at implementing rendering of the UClothComponent.

diff --git a/Engine/Plugins/Animation/ControlRig/Source/ControlRig/Private/AnimNode_ControlRig.cpp b/Engine/Plugins/Animation/ControlRig/Source/ControlRig/Private/AnimNode_ControlRig.cpp
--- a/Engine/Plugins/Animation/ControlRig/Source/ControlRig/Private/AnimNode_ControlRig.cpp
+++ b/Engine/Plugins/Animation/ControlRig/Source/ControlRig/Private/AnimNode_ControlRig.cpp
@@ -293,7 +293,7 @@ void FAnimNode_ControlRig::UpdateControlRigRefPoseIfNeeded(const FAnimInstancePr
 
 	if(InProxy->GetSkelMeshComponent())
 	{
-		ExpectedHash = HashCombine(ExpectedHash, (int32)reinterpret_cast<uintptr_t>(InProxy->GetSkelMeshComponent()->GetSkeletalMesh().Get()));
+		ExpectedHash = HashCombine(ExpectedHash, (int32)reinterpret_cast<uintptr_t>(InProxy->GetSkelMeshComponent()->GetSkeletalMesh()));
 	}
 
 	if(bIncludePoseInHash)
diff --git a/Engine/Plugins/Experimental/Animation/PoseSearch/Source/Editor/Private/PoseSearchDebugger.cpp b/Engine/Plugins/Experimental/Animation/PoseSearch/Source/Editor/Private/PoseSearchDebugger.cpp
--- a/Engine/Plugins/Experimental/Animation/PoseSearch/Source/Editor/Private/PoseSearchDebugger.cpp
+++ b/Engine/Plugins/Experimental/Animation/PoseSearch/Source/Editor/Private/PoseSearchDebugger.cpp
@@ -38,7 +38,7 @@ static FLinearColor LinearColorBlend(FLinearColor LinearColorA, FLinearColor Lin
 void UPoseSearchMeshComponent::Initialize(const FTransform& InComponentToWorld)
 {
 	SetComponentToWorld(InComponentToWorld);
-	const FReferenceSkeleton& SkeletalMeshRefSkeleton = GetSkeletalMesh()->GetRefSkeleton();
+	const FReferenceSkeleton& SkeletalMeshRefSkeleton = GetSkinnedAsset()->GetRefSkeleton();
 
 	// set up bone visibility states as this gets skipped since we allocate the component array before registration
 	for (int32 BaseIndex = 0; BaseIndex < 2; ++BaseIndex)
@@ -2495,9 +2495,9 @@ void FDebuggerViewModel::UpdateFromTimeline()
 			USkeletalMesh* SkeletalMesh = TSoftObjectPtr<USkeletalMesh>(FSoftObjectPath(SkeletalMeshObjectInfo->PathName)).LoadSynchronous();
 			if (SkeletalMesh)
 			{
-				ActiveComponent->SetSkeletalMesh(SkeletalMesh);
-				SelectedComponent->SetSkeletalMesh(SkeletalMesh);
-				AssetComponent->SetSkeletalMesh(SkeletalMesh);
+				ActiveComponent->SetSkinnedAsset(SkeletalMesh, true);
+				SelectedComponent->SetSkinnedAsset(SkeletalMesh, true);
+				AssetComponent->SetSkinnedAsset(SkeletalMesh, true);
 			}
 			FTransform ComponentWorldTransform;
 			// Active skeleton is simply the traced bone transforms
diff --git a/Engine/Plugins/Experimental/MLDeformer/MLDeformerFramework/Source/MLDeformerFramework/Private/MLDeformerModelInstance.cpp b/Engine/Plugins/Experimental/MLDeformer/MLDeformerFramework/Source/MLDeformerFramework/Private/MLDeformerModelInstance.cpp
--- a/Engine/Plugins/Experimental/MLDeformer/MLDeformerFramework/Source/MLDeformerFramework/Private/MLDeformerModelInstance.cpp
+++ b/Engine/Plugins/Experimental/MLDeformer/MLDeformerFramework/Source/MLDeformerFramework/Private/MLDeformerModelInstance.cpp
@@ -135,13 +135,13 @@ void UMLDeformerModelInstance::UpdateBoneTransforms()
 	if (LeaderPoseComponent)
 	{
 		const TArray<FTransform>& LeaderTransforms = LeaderPoseComponent->GetComponentSpaceTransforms();
-		USkeletalMesh* Mesh = LeaderPoseComponent->GetSkeletalMesh();
+		USkinnedAsset* SkinnedAsset = LeaderPoseComponent->GetSkinnedAsset();
 		const int32 NumTrainingBones = AssetBonesToSkelMeshMappings.Num();
 		for (int32 Index = 0; Index < NumTrainingBones; ++Index)
 		{
 			const int32 ComponentBoneIndex = AssetBonesToSkelMeshMappings[Index];
 			const FTransform& ComponentSpaceTransform = LeaderTransforms[ComponentBoneIndex];
-			const int32 ParentIndex = Mesh->GetRefSkeleton().GetParentIndex(ComponentBoneIndex);
+			const int32 ParentIndex = SkinnedAsset->GetRefSkeleton().GetParentIndex(ComponentBoneIndex);
 			if (ParentIndex != INDEX_NONE)
 			{
 				TrainingBoneTransforms[Index] = ComponentSpaceTransform.GetRelativeTransform(LeaderTransforms[ParentIndex]);
@@ -298,7 +298,7 @@ bool UMLDeformerModelInstance::SetupNeuralNetworkForFrame()
 	// Some safety checks.
 	if (Model == nullptr ||
 		SkeletalMeshComponent == nullptr ||
-		SkeletalMeshComponent->SkeletalMesh == nullptr ||
+		SkeletalMeshComponent->GetSkeletalMesh() == nullptr ||
 		!bIsCompatible)
 	{
 		return false;
diff --git a/Engine/Plugins/Experimental/MLDeformer/MLDeformerFramework/Source/MLDeformerFrameworkEditor/Private/MLDeformerEditorActor.cpp b/Engine/Plugins/Experimental/MLDeformer/MLDeformerFramework/Source/MLDeformerFrameworkEditor/Private/MLDeformerEditorActor.cpp
--- a/Engine/Plugins/Experimental/MLDeformer/MLDeformerFramework/Source/MLDeformerFrameworkEditor/Private/MLDeformerEditorActor.cpp
+++ b/Engine/Plugins/Experimental/MLDeformer/MLDeformerFramework/Source/MLDeformerFrameworkEditor/Private/MLDeformerEditorActor.cpp
@@ -39,7 +39,7 @@ namespace UE::MLDeformer
 
 	bool FMLDeformerEditorActor::HasVisualMesh() const
 	{
-		if (SkeletalMeshComponent && SkeletalMeshComponent->SkeletalMesh)
+		if (SkeletalMeshComponent && SkeletalMeshComponent->GetSkeletalMesh())
 		{
 			return true;
 		}
diff --git a/Engine/Plugins/Experimental/MLDeformer/NeuralMorphModel/Source/NeuralMorphModelEditor/Private/NeuralMorphEditorModelActor.cpp b/Engine/Plugins/Experimental/MLDeformer/NeuralMorphModel/Source/NeuralMorphModelEditor/Private/NeuralMorphEditorModelActor.cpp
--- a/Engine/Plugins/Experimental/MLDeformer/NeuralMorphModel/Source/NeuralMorphModelEditor/Private/NeuralMorphEditorModelActor.cpp
+++ b/Engine/Plugins/Experimental/MLDeformer/NeuralMorphModel/Source/NeuralMorphModelEditor/Private/NeuralMorphEditorModelActor.cpp
@@ -40,7 +40,7 @@ namespace UE::NeuralMorphModel
 
 	bool FNeuralMorphEditorModelActor::HasVisualMesh() const
 	{
-		if (SkeletalMeshComponent && SkeletalMeshComponent->SkeletalMesh)
+		if (SkeletalMeshComponent && SkeletalMeshComponent->GetSkeletalMesh())
 		{
 			return true;
 		}
diff --git a/Engine/Plugins/Experimental/MLDeformer/VertexDeltaModel/Source/VertexDeltaModelEditor/Private/VertexDeltaEditorModelActor.cpp b/Engine/Plugins/Experimental/MLDeformer/VertexDeltaModel/Source/VertexDeltaModelEditor/Private/VertexDeltaEditorModelActor.cpp
--- a/Engine/Plugins/Experimental/MLDeformer/VertexDeltaModel/Source/VertexDeltaModelEditor/Private/VertexDeltaEditorModelActor.cpp
+++ b/Engine/Plugins/Experimental/MLDeformer/VertexDeltaModel/Source/VertexDeltaModelEditor/Private/VertexDeltaEditorModelActor.cpp
@@ -40,7 +40,7 @@ namespace UE::VertexDeltaModel
 
 	bool FVertexDeltaEditorModelActor::HasVisualMesh() const
 	{
-		if (SkeletalMeshComponent && SkeletalMeshComponent->SkeletalMesh)
+		if (SkeletalMeshComponent && SkeletalMeshComponent->GetSkeletalMesh())
 		{
 			return true;
 		}
diff --git a/Engine/Plugins/FX/Niagara/Source/Niagara/Private/NiagaraDataInterfaceRigidMeshCollisionQuery.cpp b/Engine/Plugins/FX/Niagara/Source/Niagara/Private/NiagaraDataInterfaceRigidMeshCollisionQuery.cpp
--- a/Engine/Plugins/FX/Niagara/Source/Niagara/Private/NiagaraDataInterfaceRigidMeshCollisionQuery.cpp
+++ b/Engine/Plugins/FX/Niagara/Source/Niagara/Private/NiagaraDataInterfaceRigidMeshCollisionQuery.cpp
@@ -200,7 +200,7 @@ void CollectComponents(TConstArrayView<AActor*> Actors, TConstArrayView<FName> C
 	{
 		if (UPhysicsAsset* PhysicsAsset = Component->GetPhysicsAsset())
 		{
-			USkeletalMesh* MeshAsset = Component->GetSkeletalMesh() ? Component->GetSkeletalMesh().Get() : PhysicsAsset->GetPreviewMesh();
+			USkeletalMesh* MeshAsset = Component->GetSkeletalMesh() ? Component->GetSkeletalMesh() : PhysicsAsset->GetPreviewMesh();
 			if (!MeshAsset || !MeshAsset->GetRefSkeleton().GetNum())
 			{
 				return false;
@@ -220,7 +220,7 @@ void ForEachBodySetup(USkeletalMeshComponent* Component, TBodySetupPredicate Pre
 {
 	if (UPhysicsAsset* PhysicsAsset = Component->GetPhysicsAsset())
 	{
-		USkeletalMesh* SkeletalMesh = Component->GetSkeletalMesh() ? Component->GetSkeletalMesh().Get() : PhysicsAsset->GetPreviewMesh();
+		USkeletalMesh* SkeletalMesh = Component->GetSkeletalMesh() ? Component->GetSkeletalMesh() : PhysicsAsset->GetPreviewMesh();
 		const FReferenceSkeleton& RefSkeleton = SkeletalMesh->GetRefSkeleton();
 
 		for (const UBodySetup* BodySetup : PhysicsAsset->SkeletalBodySetups)
@@ -314,7 +314,7 @@ FTransform CreateElementTransform(const TComponentType* Component, const UBodySe
 template<>
 FTransform CreateElementTransform<USkeletalMeshComponent>(const USkeletalMeshComponent* Component, const UBodySetup* BodySetup)
 {
-	if (USkeletalMesh* SkeletalMesh = Component->GetSkeletalMesh().Get())
+	if (USkeletalMesh* SkeletalMesh = Component->GetSkeletalMesh())
 	{
 		const FReferenceSkeleton& RefSkeleton = SkeletalMesh->GetRefSkeleton();
 		const int32 BoneCount = RefSkeleton.GetNum();
diff --git a/Engine/Plugins/Runtime/ApexDestruction/Source/ApexDestruction/Public/DestructibleComponent.h b/Engine/Plugins/Runtime/ApexDestruction/Source/ApexDestruction/Public/DestructibleComponent.h
--- a/Engine/Plugins/Runtime/ApexDestruction/Source/ApexDestruction/Public/DestructibleComponent.h
+++ b/Engine/Plugins/Runtime/ApexDestruction/Source/ApexDestruction/Public/DestructibleComponent.h
@@ -27,7 +27,7 @@ DECLARE_DYNAMIC_MULTICAST_DELEGATE_TwoParams(FComponentFractureSignature, const
  */
 
 class UE_DEPRECATED(4.26, "APEX is deprecated. Destruction in future will be supported using Chaos Destruction.") UDestructibleComponent;
-UCLASS(ClassGroup = Physics, hidecategories = (Object, Mesh, "Components|SkinnedMesh", Mirroring, Activation, "Components|Activation"), config = Engine, editinlinenew, meta = (BlueprintSpawnableComponent))
+UCLASS(ClassGroup = Physics, hidecategories = (Object, Mesh, "Mesh|SkeletalAsset", "Components|SkinnedMesh", Mirroring, Activation, "Components|Activation"), config = Engine, editinlinenew, meta = (BlueprintSpawnableComponent))
 class APEXDESTRUCTION_API UDestructibleComponent : public USkinnedMeshComponent, public IDestructibleInterface
 {
 	GENERATED_UCLASS_BODY()
diff --git a/Engine/Plugins/Runtime/MeshModelingToolset/Source/ModelingComponentsEditorOnly/Private/ToolTargets/SkeletalMeshComponentToolTarget.cpp b/Engine/Plugins/Runtime/MeshModelingToolset/Source/ModelingComponentsEditorOnly/Private/ToolTargets/SkeletalMeshComponentToolTarget.cpp
--- a/Engine/Plugins/Runtime/MeshModelingToolset/Source/ModelingComponentsEditorOnly/Private/ToolTargets/SkeletalMeshComponentToolTarget.cpp
+++ b/Engine/Plugins/Runtime/MeshModelingToolset/Source/ModelingComponentsEditorOnly/Private/ToolTargets/SkeletalMeshComponentToolTarget.cpp
@@ -2,7 +2,7 @@
 
 #include "ToolTargets/SkeletalMeshComponentToolTarget.h"
 
-#include "Components/SkinnedMeshComponent.h"
+#include "Components/SkeletalMeshComponent.h"
 #include "ConversionUtils/DynamicMeshViaMeshDescriptionUtil.h"
 #include "DynamicMeshToMeshDescription.h"
 #include "MeshDescriptionToDynamicMesh.h"
@@ -28,12 +28,12 @@ bool USkeletalMeshComponentReadOnlyToolTarget::IsValid() const
 	{
 		return false;
 	}
-	USkinnedMeshComponent* SkinnedMeshComponent = Cast<USkinnedMeshComponent>(Component);
-	if (SkinnedMeshComponent == nullptr)
+	USkeletalMeshComponent* SkeletalMeshComponent = Cast<USkeletalMeshComponent>(Component);
+	if (SkeletalMeshComponent == nullptr)
 	{
 		return false;
 	}
-	const USkeletalMesh* SkeletalMesh = SkinnedMeshComponent->GetSkeletalMesh();
+	const USkeletalMesh* SkeletalMesh = SkeletalMeshComponent->GetSkeletalMesh();
 
 	return USkeletalMeshReadOnlyToolTarget::IsValid(SkeletalMesh);
 }
@@ -54,7 +54,7 @@ void USkeletalMeshComponentReadOnlyToolTarget::GetMaterialSet(FComponentMaterial
 
 	if (bPreferAssetMaterials)
 	{
-		const USkeletalMesh* SkeletalMesh = Cast<USkinnedMeshComponent>(Component)->GetSkeletalMesh();
+		const USkeletalMesh* SkeletalMesh = Cast<USkeletalMeshComponent>(Component)->GetSkeletalMesh();
 		USkeletalMeshToolTarget::GetMaterialSet(SkeletalMesh, MaterialSetOut, bPreferAssetMaterials);
 	}
 	else
@@ -74,7 +74,7 @@ bool USkeletalMeshComponentReadOnlyToolTarget::CommitMaterialSetUpdate(const FCo
 
 	if (bApplyToAsset)
 	{
-		USkeletalMesh* SkeletalMesh = Cast<USkinnedMeshComponent>(Component)->GetSkeletalMesh();
+		USkeletalMesh* SkeletalMesh = Cast<USkeletalMeshComponent>(Component)->GetSkeletalMesh();
 
 		// unregister the component while we update it's static mesh
 		TUniquePtr<FComponentReregisterContext> ComponentReregisterContext = MakeUnique<FComponentReregisterContext>(Component);
@@ -115,7 +115,7 @@ const FMeshDescription* USkeletalMeshComponentReadOnlyToolTarget::GetMeshDescrip
 	if (!CachedMeshDescription.IsValid())
 	{
 		CachedMeshDescription = MakeUnique<FMeshDescription>();
-		const USkeletalMesh* SkeletalMesh = Cast<USkinnedMeshComponent>(Component)->GetSkeletalMesh();
+		const USkeletalMesh* SkeletalMesh = Cast<USkeletalMeshComponent>(Component)->GetSkeletalMesh();
 		USkeletalMeshToolTarget::GetMeshDescription(SkeletalMesh, *CachedMeshDescription);
 	}
 	
@@ -139,7 +139,7 @@ FDynamicMesh3 USkeletalMeshComponentReadOnlyToolTarget::GetDynamicMesh()
 
 USkeletalMesh* USkeletalMeshComponentReadOnlyToolTarget::GetSkeletalMesh() const
 {
-	return IsValid() ? Cast<USkinnedMeshComponent>(Component)->GetSkeletalMesh() : nullptr;
+	return IsValid() ? Cast<USkeletalMeshComponent>(Component)->GetSkeletalMesh() : nullptr;
 }
 
 
@@ -152,7 +152,7 @@ void USkeletalMeshComponentToolTarget::CommitMeshDescription(const FCommitter& C
 {
 	if (ensure(IsValid()) == false) return;
 
-	USkeletalMesh* SkeletalMesh = Cast<USkinnedMeshComponent>(Component)->GetSkeletalMesh();
+	USkeletalMesh* SkeletalMesh = Cast<USkeletalMeshComponent>(Component)->GetSkeletalMesh();
 	if (!CachedMeshDescription.IsValid())
 	{
 		CachedMeshDescription = MakeUnique<FMeshDescription>();
@@ -193,15 +193,15 @@ bool USkeletalMeshComponentReadOnlyToolTargetFactory::CanBuildTarget(UObject* So
 	// just add another factory that allows that class specifically(but make sure that
 	// GetMeshDescription and such work properly)
 
-	return Cast<USkinnedMeshComponent>(SourceObject) && Cast<USkinnedMeshComponent>(SourceObject)->GetSkeletalMesh() &&
-		ExactCast<USkeletalMesh>(Cast<USkinnedMeshComponent>(SourceObject)->GetSkeletalMesh()) &&
+	return Cast<USkeletalMeshComponent>(SourceObject) && Cast<USkeletalMeshComponent>(SourceObject)->GetSkeletalMesh() &&
+		ExactCast<USkeletalMesh>(Cast<USkeletalMeshComponent>(SourceObject)->GetSkeletalMesh()) &&
 		Requirements.AreSatisfiedBy(USkeletalMeshComponentReadOnlyToolTarget::StaticClass());
 }
 
 UToolTarget* USkeletalMeshComponentReadOnlyToolTargetFactory::BuildTarget(UObject* SourceObject, const FToolTargetTypeRequirements& Requirements)
 {
 	USkeletalMeshComponentReadOnlyToolTarget* Target = NewObject<USkeletalMeshComponentReadOnlyToolTarget>();
-	Target->Component = Cast<USkinnedMeshComponent>(SourceObject);
+	Target->Component = Cast<USkeletalMeshComponent>(SourceObject);
 	check(Target->Component && Requirements.AreSatisfiedBy(Target));
 
 	return Target;
@@ -216,15 +216,15 @@ bool USkeletalMeshComponentToolTargetFactory::CanBuildTarget(UObject* SourceObje
 	// just add another factory that allows that class specifically(but make sure that
 	// GetMeshDescription and such work properly)
 
-	return Cast<USkinnedMeshComponent>(SourceObject) && Cast<USkinnedMeshComponent>(SourceObject)->GetSkeletalMesh() &&
-		ExactCast<USkeletalMesh>(Cast<USkinnedMeshComponent>(SourceObject)->GetSkeletalMesh()) &&
+	return Cast<USkeletalMeshComponent>(SourceObject) && Cast<USkeletalMeshComponent>(SourceObject)->GetSkeletalMesh() &&
+		ExactCast<USkeletalMesh>(Cast<USkeletalMeshComponent>(SourceObject)->GetSkeletalMesh()) &&
 		Requirements.AreSatisfiedBy(USkeletalMeshComponentToolTarget::StaticClass());
 }
 
 UToolTarget* USkeletalMeshComponentToolTargetFactory::BuildTarget(UObject* SourceObject, const FToolTargetTypeRequirements& Requirements)
 {
 	USkeletalMeshComponentToolTarget* Target = NewObject<USkeletalMeshComponentToolTarget>();
-	Target->Component = Cast<USkinnedMeshComponent>(SourceObject);
+	Target->Component = Cast<USkeletalMeshComponent>(SourceObject);
 	check(Target->Component && Requirements.AreSatisfiedBy(Target));
 
 	return Target;
diff --git a/Engine/Plugins/Runtime/Oculus/OculusVR/Source/OculusInput/Private/OculusHandComponent.cpp b/Engine/Plugins/Runtime/Oculus/OculusVR/Source/OculusInput/Private/OculusHandComponent.cpp
--- a/Engine/Plugins/Runtime/Oculus/OculusVR/Source/OculusInput/Private/OculusHandComponent.cpp
+++ b/Engine/Plugins/Runtime/Oculus/OculusVR/Source/OculusInput/Private/OculusHandComponent.cpp
@@ -33,7 +33,7 @@ void UOculusHandComponent::BeginPlay()
 	Super::BeginPlay();
 
 	// Use custom mesh if a skeletal mesh is already set, else try to load the runtime mesh
-	if (GetSkeletalMesh())
+	if (GetSkinnedAsset())
 	{
 		bCustomHandMesh = true;
 		bSkeletalMeshInitialized = true;
@@ -51,7 +51,7 @@ void UOculusHandComponent::InitializeSkeletalMesh()
 	{
 		if (UOculusInputFunctionLibrary::GetHandSkeletalMesh(RuntimeSkeletalMesh, SkeletonType, MeshType))
 		{
-			SetSkeletalMesh(RuntimeSkeletalMesh);
+			SetSkinnedAsset(RuntimeSkeletalMesh, true);
 			if (MaterialOverride)
 			{
 				SetMaterial(0, MaterialOverride);
@@ -106,7 +106,7 @@ void UOculusHandComponent::TickComponent(float DeltaTime, enum ELevelTick TickTy
 			}
 
 			// Update Bone Pose Rotations
-			if (GetSkeletalMesh())
+			if (GetSkinnedAsset())
 			{
 				UpdateBonePose();
 			}
@@ -163,7 +163,7 @@ void UOculusHandComponent::UpdateBonePose()
 			else
 			{
 				// Set Remaing Bone Rotations
-				int32 BoneIndex = GetSkeletalMesh()->GetRefSkeleton().FindBoneIndex(BoneElem.Value);
+				int32 BoneIndex = GetSkinnedAsset()->GetRefSkeleton().FindBoneIndex(BoneElem.Value);
 				if (BoneIndex >= 0)
 				{
 					FQuat BoneRotation = UOculusInputFunctionLibrary::GetBoneRotation(SkeletonType, (EBone)BoneElem.Key);
@@ -181,7 +181,7 @@ void UOculusHandComponent::UpdateBonePose()
 		BoneSpaceTransforms[0].SetRotation(RootBoneRotation);
 
 		// Set Remaining Bone Rotations
-		for (uint32 BoneIndex = 1; BoneIndex < (uint32)GetSkeletalMesh()->GetRefSkeleton().GetNum(); BoneIndex++)
+		for (uint32 BoneIndex = 1; BoneIndex < (uint32)GetSkinnedAsset()->GetRefSkeleton().GetNum(); BoneIndex++)
 		{
 			FQuat BoneRotation = UOculusInputFunctionLibrary::GetBoneRotation(SkeletonType, (EBone)BoneIndex);
 			BoneSpaceTransforms[BoneIndex].SetRotation(BoneRotation);
diff --git a/Engine/Plugins/VirtualProduction/Takes/Source/TakeTrackRecorders/Private/MovieScene3DTransformTrackRecorder.cpp b/Engine/Plugins/VirtualProduction/Takes/Source/TakeTrackRecorders/Private/MovieScene3DTransformTrackRecorder.cpp
--- a/Engine/Plugins/VirtualProduction/Takes/Source/TakeTrackRecorders/Private/MovieScene3DTransformTrackRecorder.cpp
+++ b/Engine/Plugins/VirtualProduction/Takes/Source/TakeTrackRecorders/Private/MovieScene3DTransformTrackRecorder.cpp
@@ -489,7 +489,9 @@ void UMovieScene3DTransformTrackRecorder::PostProcessAnimationData(UMovieSceneAn
 
 			}
 			// Search for the Root Bone in the Skeleton
-			const USkeletalMesh* SkeletalMesh = SkeletalMeshComponent->LeaderPoseComponent != nullptr ? ToRawPtr(SkeletalMeshComponent->LeaderPoseComponent->GetSkeletalMesh()) : ToRawPtr(SkeletalMeshComponent->GetSkeletalMesh());
+			const USkeletalMesh* SkeletalMesh = Cast<USkeletalMeshComponent>(SkeletalMeshComponent->LeaderPoseComponent) ?
+				Cast<USkeletalMeshComponent>(SkeletalMeshComponent->LeaderPoseComponent)->GetSkeletalMesh() : 
+				SkeletalMeshComponent->GetSkeletalMesh();
 			const UAnimSequence* AnimSequence = AnimTrackRecorder->GetAnimSequence();
 			if (AnimSequence && SkeletalMesh && AnimSequence->GetSkeleton())
 			{
diff --git a/Engine/Plugins/VirtualProduction/Takes/Source/TakeTrackRecorders/Private/MovieSceneAnimationTrackRecorder.cpp b/Engine/Plugins/VirtualProduction/Takes/Source/TakeTrackRecorders/Private/MovieSceneAnimationTrackRecorder.cpp
--- a/Engine/Plugins/VirtualProduction/Takes/Source/TakeTrackRecorders/Private/MovieSceneAnimationTrackRecorder.cpp
+++ b/Engine/Plugins/VirtualProduction/Takes/Source/TakeTrackRecorders/Private/MovieSceneAnimationTrackRecorder.cpp
@@ -157,7 +157,10 @@ void UMovieSceneAnimationTrackRecorder::CreateTrackImpl()
 			for (int32 BoneIndex = 0; BoneIndex < SpaceBases->Num(); ++BoneIndex)
 			{
 				// verify if this bone exists in skeleton
-				const int32 BoneTreeIndex = AnimSkeleton->GetSkeletonBoneIndexFromMeshBoneIndex(SkeletalMeshComponent->LeaderPoseComponent != nullptr ? SkeletalMeshComponent->LeaderPoseComponent->GetSkeletalMesh() : SkeletalMeshComponent->GetSkeletalMesh(), BoneIndex);
+				const int32 BoneTreeIndex = AnimSkeleton->GetSkeletonBoneIndexFromMeshBoneIndex(
+					Cast<USkeletalMeshComponent>(SkeletalMeshComponent->LeaderPoseComponent) ? 
+					Cast<USkeletalMeshComponent>(SkeletalMeshComponent->LeaderPoseComponent)->GetSkeletalMesh() :
+					SkeletalMeshComponent->GetSkeletalMesh(), BoneIndex);
 				if (BoneTreeIndex != INDEX_NONE)
 				{
 					// add tracks for the bone existing
diff --git a/Engine/Source/Editor/DetailCustomizations/Private/SkinnedMeshComponentDetails.cpp b/Engine/Source/Editor/DetailCustomizations/Private/SkinnedMeshComponentDetails.cpp
--- a/Engine/Source/Editor/DetailCustomizations/Private/SkinnedMeshComponentDetails.cpp
+++ b/Engine/Source/Editor/DetailCustomizations/Private/SkinnedMeshComponentDetails.cpp
@@ -6,7 +6,7 @@
 #include "Widgets/Text/STextBlock.h"
 #include "Widgets/Input/SEditableTextBox.h"
 #include "GameFramework/Actor.h"
-#include "Components/SkinnedMeshComponent.h"
+#include "Components/SkeletalMeshComponent.h"
 #include "Editor.h"
 #include "PropertyHandle.h"
 #include "DetailLayoutBuilder.h"
@@ -181,7 +181,7 @@ void FSkinnedMeshComponentDetails::PopulateSkinWeightProfileNames()
 	SkinWeightProfileNames.Add(MakeShared<FName>(DefaultProfileName));
 
 	// Retrieve all possible skin weight profiles from the component
-	if (USkinnedMeshComponent* Component = WeakSkinnedMeshComponent.Get())
+	if (USkeletalMeshComponent* Component = Cast<USkeletalMeshComponent>(WeakSkinnedMeshComponent.Get()))
 	{
 		if (USkeletalMesh* Mesh = Component->GetSkeletalMesh())
 		{			
diff --git a/Engine/Source/Editor/Persona/Private/SkinWeightProfileHelpers.cpp b/Engine/Source/Editor/Persona/Private/SkinWeightProfileHelpers.cpp
--- a/Engine/Source/Editor/Persona/Private/SkinWeightProfileHelpers.cpp
+++ b/Engine/Source/Editor/Persona/Private/SkinWeightProfileHelpers.cpp
@@ -280,7 +280,7 @@ void FSkinWeightProfileHelpers::ClearSkinWeightProfileInstanceOverrides(USkeleta
 	FScopedSkeletalMeshPostEditChange ScopedPostEditChange(InSkeletalMesh, false, true);
 	for (TObjectIterator<USkinnedMeshComponent> It; It; ++It)
 	{
-		if (It->GetSkeletalMesh() == InSkeletalMesh)
+		if (Cast<USkeletalMesh>(It->GetSkinnedAsset()) == InSkeletalMesh)
 		{
 			checkf(!It->IsUnreachable(), TEXT("%s"), *It->GetFullName());
 
diff --git a/Engine/Source/Editor/SequenceRecorder/Private/AnimationRecorder.cpp b/Engine/Source/Editor/SequenceRecorder/Private/AnimationRecorder.cpp
--- a/Engine/Source/Editor/SequenceRecorder/Private/AnimationRecorder.cpp
+++ b/Engine/Source/Editor/SequenceRecorder/Private/AnimationRecorder.cpp
@@ -270,7 +270,10 @@ void FAnimationRecorder::StartRecord(USkeletalMeshComponent* Component, UAnimSeq
 	for (int32 BoneIndex=0; BoneIndex <PreviousSpacesBases.Num(); ++BoneIndex)
 	{
 		// verify if this bone exists in skeleton
-		const int32 BoneTreeIndex = AnimSkeleton->GetSkeletonBoneIndexFromMeshBoneIndex(Component->LeaderPoseComponent != nullptr ? Component->LeaderPoseComponent->GetSkeletalMesh() : Component->GetSkeletalMesh(), BoneIndex);
+		const int32 BoneTreeIndex = AnimSkeleton->GetSkeletonBoneIndexFromMeshBoneIndex(
+			Cast<USkeletalMeshComponent>(Component->LeaderPoseComponent) ?
+			Cast<USkeletalMeshComponent>(Component->LeaderPoseComponent)->GetSkeletalMesh() :
+			Component->GetSkeletalMesh(), BoneIndex);
 		if (BoneTreeIndex != INDEX_NONE)
 		{
 			// add tracks for the bone existing
@@ -652,7 +655,10 @@ void FAnimationRecorder::ProcessRecordedTimes(UAnimSequence* AnimSequence, USkel
 	{
 		for (int32 BoneIndex = 0; BoneIndex < SpaceBases->Num(); ++BoneIndex)
 		{
-			const int32 BoneTreeIndex = AnimSkeleton->GetSkeletonBoneIndexFromMeshBoneIndex(SkeletalMeshComponent->LeaderPoseComponent != nullptr ? SkeletalMeshComponent->LeaderPoseComponent->GetSkeletalMesh() : SkeletalMeshComponent->GetSkeletalMesh(), BoneIndex);
+			const int32 BoneTreeIndex = AnimSkeleton->GetSkeletonBoneIndexFromMeshBoneIndex(
+				Cast<USkeletalMeshComponent>(SkeletalMeshComponent->LeaderPoseComponent) ?
+				Cast<USkeletalMeshComponent>(SkeletalMeshComponent->LeaderPoseComponent)->GetSkeletalMesh() :
+				SkeletalMeshComponent->GetSkeletalMesh(), BoneIndex);
 			if (BoneTreeIndex != INDEX_NONE)
 			{
 				FName BoneTreeName = AnimSkeleton->GetReferenceSkeleton().GetBoneName(BoneTreeIndex);
@@ -673,7 +679,10 @@ void FAnimationRecorder::ProcessRecordedTimes(UAnimSequence* AnimSequence, USkel
 	for (int32 BoneIndex = 0; BoneIndex < SpaceBases->Num(); ++BoneIndex)
 	{
 		// verify if this bone exists in skeleton
-		const int32 BoneTreeIndex = AnimSkeleton->GetSkeletonBoneIndexFromMeshBoneIndex(SkeletalMeshComponent->LeaderPoseComponent != nullptr ? SkeletalMeshComponent->LeaderPoseComponent->GetSkeletalMesh() : SkeletalMeshComponent->GetSkeletalMesh(), BoneIndex);
+		const int32 BoneTreeIndex = AnimSkeleton->GetSkeletonBoneIndexFromMeshBoneIndex(
+			Cast<USkeletalMeshComponent>(SkeletalMeshComponent->LeaderPoseComponent) ?
+			Cast<USkeletalMeshComponent>(SkeletalMeshComponent->LeaderPoseComponent)->GetSkeletalMesh() :
+			SkeletalMeshComponent->GetSkeletalMesh(), BoneIndex);
 		if (BoneTreeIndex != INDEX_NONE)
 		{
 			// add tracks for the bone existing
@@ -826,7 +835,10 @@ bool FAnimationRecorder::Record(USkeletalMeshComponent* Component, FTransform co
 	if (ensure(AnimationObject))
 	{
 		IAnimationDataController& Controller = AnimationObject->GetController();
-		USkeletalMesh* SkeletalMesh = Component->LeaderPoseComponent != nullptr ? ToRawPtr(Component->LeaderPoseComponent->GetSkeletalMesh()) : ToRawPtr(Component->GetSkeletalMesh());
+		USkeletalMesh* SkeletalMesh =
+			Cast<USkeletalMeshComponent>(Component->LeaderPoseComponent) ?
+			Cast<USkeletalMeshComponent>(Component->LeaderPoseComponent)->GetSkeletalMesh() :
+			Component->GetSkeletalMesh();
 
 		const TArray<FBoneAnimationTrack>& BoneAnimationTracks = AnimationObject->GetDataModel()->GetBoneAnimationTracks();
 
diff --git a/Engine/Source/Editor/SequenceRecorder/Private/Sections/MovieScene3DTransformSectionRecorder.cpp b/Engine/Source/Editor/SequenceRecorder/Private/Sections/MovieScene3DTransformSectionRecorder.cpp
--- a/Engine/Source/Editor/SequenceRecorder/Private/Sections/MovieScene3DTransformSectionRecorder.cpp
+++ b/Engine/Source/Editor/SequenceRecorder/Private/Sections/MovieScene3DTransformSectionRecorder.cpp
@@ -132,7 +132,10 @@ void FMovieScene3DTransformSectionRecorder::FinalizeSection(float CurrentTime)
 		USkeletalMeshComponent* SkeletalMeshComponent = AnimRecorder->GetSkeletalMeshComponent();
 		if (SkeletalMeshComponent)
 		{
-			USkeletalMesh* SkeletalMesh = SkeletalMeshComponent->LeaderPoseComponent != nullptr ? ToRawPtr(SkeletalMeshComponent->LeaderPoseComponent->GetSkeletalMesh()) : ToRawPtr(SkeletalMeshComponent->GetSkeletalMesh());
+			USkeletalMesh* SkeletalMesh =
+				Cast<USkeletalMeshComponent>(SkeletalMeshComponent->LeaderPoseComponent) ?
+				Cast<USkeletalMeshComponent>(SkeletalMeshComponent->LeaderPoseComponent)->GetSkeletalMesh() :
+				SkeletalMeshComponent->GetSkeletalMesh();
 			if (AnimSequence && SkeletalMesh)
 			{
 				// find the root bone
diff --git a/Engine/Source/Editor/UnrealEd/Private/EditorEngine.cpp b/Engine/Source/Editor/UnrealEd/Private/EditorEngine.cpp
--- a/Engine/Source/Editor/UnrealEd/Private/EditorEngine.cpp
+++ b/Engine/Source/Editor/UnrealEd/Private/EditorEngine.cpp
@@ -3523,7 +3523,7 @@ struct FConvertStaticMeshActorInfo
 		InternalSetToActor(Actor);
 
 		// Set component properties.
-		if ( bComponentPropsDifferFromDefaults[0] ) MeshComp->GetSkeletalMesh()		= SkeletalMesh;
+		if ( bComponentPropsDifferFromDefaults[0] ) MeshComp->SetSkinnedAsset(SkeletalMesh);
 		if ( bComponentPropsDifferFromDefaults[1] ) MeshComp->OverrideMaterials		= OverrideMaterials;
 		if ( bComponentPropsDifferFromDefaults[2] ) MeshComp->CachedMaxDrawDistance	= CachedMaxDrawDistance;
 		if ( bComponentPropsDifferFromDefaults[3] ) MeshComp->CastShadow			= CastShadow;
diff --git a/Engine/Source/Editor/UnrealEd/Private/Factories/ActorFactory.cpp b/Engine/Source/Editor/UnrealEd/Private/Factories/ActorFactory.cpp
--- a/Engine/Source/Editor/UnrealEd/Private/Factories/ActorFactory.cpp
+++ b/Engine/Source/Editor/UnrealEd/Private/Factories/ActorFactory.cpp
@@ -958,7 +958,7 @@ void UActorFactoryPhysicsAsset::PostSpawnActor(UObject* Asset, AActor* NewActor)
 	NewSkelActor->GetSkeletalMeshComponent()->UnregisterComponent();
 
 	// Change properties
-	NewSkelActor->GetSkeletalMeshComponent()->GetSkeletalMesh() = UseSkelMesh;
+	NewSkelActor->GetSkeletalMeshComponent()->SetSkeletalMeshAsset(UseSkelMesh);
 	if (NewSkelActor->GetWorld()->IsPlayInEditor())
 	{
 		NewSkelActor->ReplicatedMesh = UseSkelMesh;
@@ -991,7 +991,7 @@ void UActorFactoryPhysicsAsset::PostCreateBlueprint( UObject* Asset, AActor* CDO
 
 			USkeletalMesh* UseSkelMesh = PhysicsAsset->PreviewSkeletalMesh.Get();
 
-			SkeletalPhysicsActor->GetSkeletalMeshComponent()->GetSkeletalMesh() = UseSkelMesh;
+			SkeletalPhysicsActor->GetSkeletalMeshComponent()->SetSkeletalMeshAsset(UseSkelMesh);
 			SkeletalPhysicsActor->GetSkeletalMeshComponent()->PhysicsAssetOverride = PhysicsAsset;
 		}
 
@@ -1334,7 +1334,7 @@ void UActorFactorySkeletalMesh::PostSpawnActor( UObject* Asset, AActor* NewActor
 	NewSMActor->GetSkeletalMeshComponent()->UnregisterComponent();
 
 	// Change properties
-	NewSMActor->GetSkeletalMeshComponent()->GetSkeletalMesh() = SkeletalMesh;
+	NewSMActor->GetSkeletalMeshComponent()->SetSkeletalMeshAsset(SkeletalMesh);
 	if (NewSMActor->GetWorld()->IsGameWorld())
 	{
 		NewSMActor->ReplicatedMesh = SkeletalMesh;
@@ -1366,7 +1366,7 @@ void UActorFactorySkeletalMesh::PostCreateBlueprint( UObject* Asset, AActor* CDO
 		UAnimBlueprint* AnimBlueprint = Cast<UAnimBlueprint>(Asset);
 
 		ASkeletalMeshActor* SkeletalMeshActor = CastChecked<ASkeletalMeshActor>(CDO);
-		SkeletalMeshActor->GetSkeletalMeshComponent()->GetSkeletalMesh() = SkeletalMesh;
+		SkeletalMeshActor->GetSkeletalMeshComponent()->SetSkeletalMeshAsset(SkeletalMesh);
 		SkeletalMeshActor->GetSkeletalMeshComponent()->AnimClass = AnimBlueprint ? Cast<UAnimBlueprintGeneratedClass>(AnimBlueprint->GeneratedClass) : nullptr;
 	}
 }
diff --git a/Engine/Source/Runtime/Engine/Classes/Components/SkeletalMeshComponent.h b/Engine/Source/Runtime/Engine/Classes/Components/SkeletalMeshComponent.h
--- a/Engine/Source/Runtime/Engine/Classes/Components/SkeletalMeshComponent.h
+++ b/Engine/Source/Runtime/Engine/Classes/Components/SkeletalMeshComponent.h
@@ -334,7 +334,7 @@ struct FClothCollisionSource
  * @see https://docs.unrealengine.com/latest/INT/Engine/Content/Types/SkeletalMeshes/
  * @see USkeletalMesh
  */
-UCLASS(Blueprintable, ClassGroup=(Rendering, Common), hidecategories=Object, config=Engine, editinlinenew, meta=(BlueprintSpawnableComponent))
+UCLASS(Blueprintable, ClassGroup=(Rendering, Common), hidecategories=(Object, "Mesh|SkeletalAsset"), config=Engine, editinlinenew, meta=(BlueprintSpawnableComponent))
 class ENGINE_API USkeletalMeshComponent : public USkinnedMeshComponent, public IInterface_CollisionDataProvider
 {
 	GENERATED_UCLASS_BODY()
@@ -347,7 +347,27 @@ class ENGINE_API USkeletalMeshComponent : public USkinnedMeshComponent, public I
 	friend struct FAnimNode_LinkedAnimLayer;
 	friend struct FLinkedInstancesAdapter;
 
+#if WITH_EDITORONLY_DATA
+private:
+	/** The skeletal mesh used by this component. */
+	UE_DEPRECATED(5.1, "This property isn't deprecated, but getter and setter must be used at all times to preserve correct operations.")
+	UPROPERTY(EditAnywhere, Transient, BlueprintSetter = SetSkeletalMeshAsset, BlueprintGetter = GetSkeletalMeshAsset, Category = Mesh)
+	TObjectPtr<class USkeletalMesh> SkeletalMeshAsset;
+#endif
+
 public:
+	/**
+	 * Set the SkeletalMesh rendered for this mesh.
+	 */
+	UFUNCTION(BlueprintCallable, Category = "Components|SkeletalMesh")
+	void SetSkeletalMeshAsset(class USkeletalMesh* NewMesh) { SetSkeletalMesh(NewMesh, true); }
+
+	/**
+	 * Get the SkeletalMesh rendered for this mesh.
+	 */
+	UFUNCTION(BlueprintPure, Category = "Components|SkeletalMesh")
+	class USkeletalMesh* GetSkeletalMeshAsset() const;
+
 #if WITH_EDITORONLY_DATA
 	/** The blueprint for creating an AnimationScript. */
 	UPROPERTY()
@@ -1865,6 +1885,9 @@ public:
 	virtual void SetPhysicsAsset(class UPhysicsAsset* NewPhysicsAsset,bool bForceReInit = false) override;
 	virtual void SetSkeletalMesh(class USkeletalMesh* NewMesh, bool bReinitPose = true) override;
 
+	// GetSkeletalMesh() is soon to be removed, use GetSkeletalMeshAsset() instead
+	class USkeletalMesh* GetSkeletalMesh() const { return GetSkeletalMeshAsset(); }
+
 	static FVector3f GetSkinnedVertexPosition(USkeletalMeshComponent* Component, int32 VertexIndex, const FSkeletalMeshLODRenderData& Model, const FSkinWeightVertexBuffer& SkinWeightBuffer);
 	static FVector3f GetSkinnedVertexPosition(USkeletalMeshComponent* Component, int32 VertexIndex, const FSkeletalMeshLODRenderData& Model, const FSkinWeightVertexBuffer& SkinWeightBuffer, TArray<FMatrix44f>& CachedRefToLocals);
 	static void ComputeSkinnedPositions(USkeletalMeshComponent* Component, TArray<FVector3f> & OutPositions, TArray<FMatrix44f>& CachedRefToLocals, const FSkeletalMeshLODRenderData& Model, const FSkinWeightVertexBuffer& SkinWeightBuffer);
diff --git a/Engine/Source/Runtime/Engine/Classes/Components/SkinnedMeshComponent.h b/Engine/Source/Runtime/Engine/Classes/Components/SkinnedMeshComponent.h
--- a/Engine/Source/Runtime/Engine/Classes/Components/SkinnedMeshComponent.h
+++ b/Engine/Source/Runtime/Engine/Classes/Components/SkinnedMeshComponent.h
@@ -252,9 +252,17 @@ class ENGINE_API USkinnedMeshComponent : public UMeshComponent, public ILODSyncI
 	friend class FSkinnedMeshComponentRecreateRenderStateContext;
 
 	/** The skeletal mesh used by this component. */
-	UPROPERTY(EditAnywhere, BlueprintReadOnly, Category="Mesh", meta = (DisallowedClasses = "/Script/ApexDestruction.DestructibleMesh"))
+	UE_DEPRECATED(5.1, "Replaced by SkinnedAsset. Use GetSkinnedAsset()/SetSkinnedAsset() instead.")
+	UPROPERTY(EditAnywhere, BlueprintGetter = GetSkeletalMesh_DEPRECATED, Category = "Mesh|SkeletalAsset", meta = (DisallowedClasses = "/Script/ApexDestruction.DestructibleMesh", DeprecatedProperty, DeprecationMessage = "Use USkeletalMeshComponent::GetSkeletalMesh() or GetSkinnedAsset() instead."))
 	TObjectPtr<class USkeletalMesh> SkeletalMesh;
 
+private:
+	/** The skinned asset used by this component. */
+	UE_DEPRECATED(5.1, "This property isn't deprecated, but getter and setter must be used instead in order to preserve backward compatibility with the SkeletalMesh pointer.")
+	UPROPERTY(BlueprintGetter = GetSkinnedAsset, Category = "Mesh")
+	TObjectPtr<class USkinnedAsset> SkinnedAsset;
+
+public:
 	//
 	// LeaderPoseComponent.
 	//
@@ -404,17 +412,6 @@ protected:
 	TSharedPtr<FSkelMeshRefPoseOverride> RefPoseOverride;
 
 public:
-	TObjectPtr<class USkeletalMesh>& GetSkeletalMesh()
-	{
-		return SkeletalMesh;
-	}
-	const TObjectPtr<class USkeletalMesh>& GetSkeletalMesh() const
-	{
-		return SkeletalMesh;
-	}
-
-	USkinnedAsset* GetSkinnedAsset() const;
-
 	const TArray<int32>& GetLeaderBoneMap() const { return LeaderBoneMap; }
 
 	UE_DEPRECATED(5.1, "GetMasterBoneMap has been deprecated. Please use the GetLeaderBoneMap.")
@@ -868,9 +865,45 @@ public:
 	 * @param NewMesh New mesh to set for this component
 	 * @param bReinitPose Whether we should keep current pose or reinitialize.
 	 */
+	UE_DEPRECATED(5.1, "Use USkeletalMeshComponent::SetSkeletalMesh() or SetSkinnedAsset() instead.")
 	UFUNCTION(BlueprintCallable, Category="Components|SkinnedMesh")
 	virtual void SetSkeletalMesh(class USkeletalMesh* NewMesh, bool bReinitPose = true);
 
+	/**
+	 * Get the SkeletalMesh rendered for this mesh.
+	 * This function is not technically deprecated but shouldn't be used other than for Blueprint backward compatibility purposes.
+	 * It is used to access the correct SkinnedAsset pointer value through the deprecated SkeletalMesh property in blueprints.
+	 * 
+	 * @return the SkeletalMesh set to this mesh.
+	 */
+	UE_DEPRECATED(5.1, "Use USkeletalMeshComponent::GetSkeletalMeshAsset() or GetSkinnedAsset() instead.")
+	UFUNCTION(BlueprintPure, Category = "Components|SkinnedMesh", DisplayName = "Get Skeletal Mesh", meta = (DeprecatedFunction, DeprecationMessage = "Use USkeletalMeshComponent::GetSkeletalMesh() or GetSkinnedAsset() instead."))
+	class USkeletalMesh* GetSkeletalMesh_DEPRECATED() const;
+
+	/**
+	 * Change the SkinnedAsset that is rendered for this Component. Will re-initialize the animation tree etc.
+	 *
+	 * @param NewSkinnedAsset New mesh to set for this component
+	 * @param bReinitPose Whether we should keep current pose or reinitialize.
+	 */
+	UFUNCTION(BlueprintCallable, Category = "Components|SkinnedMesh")
+	void SetSkinnedAsset(class USkinnedAsset* InSkinnedAsset, bool bReinitPose);
+
+	/**
+	 * Change the SkinnedAsset that is rendered without reinitializing this Component.
+	 *
+	 * @param InSkinnedAsset New mesh to set for this component
+	 */
+	void SetSkinnedAsset(class USkinnedAsset* InSkinnedAsset);
+
+	/**
+	 * Get the SkinnedAsset rendered for this mesh.
+	 *
+	 * @return the SkinnedAsset set to this mesh.
+	 */
+	UFUNCTION(BlueprintPure, Category = "Components|SkinnedMesh")
+	USkinnedAsset* GetSkinnedAsset() const;
+
 	/**
 	 * Change the MeshDeformer that is used for this Component.
 	 *
@@ -929,6 +962,7 @@ public:
 	//~ Begin UObject Interface
 	virtual void BeginDestroy() override;
 	virtual void Serialize(FArchive& Ar) override;
+	virtual void PostLoad() override;
 	virtual void GetResourceSizeEx(FResourceSizeEx& CumulativeResourceSize) override;
 	virtual FString GetDetailedInfoInternal() const override;
 #if WITH_EDITOR
diff --git a/Engine/Source/Runtime/Engine/Classes/Engine/SkeletalMesh.h b/Engine/Source/Runtime/Engine/Classes/Engine/SkeletalMesh.h
--- a/Engine/Source/Runtime/Engine/Classes/Engine/SkeletalMesh.h
+++ b/Engine/Source/Runtime/Engine/Classes/Engine/SkeletalMesh.h
@@ -2433,8 +2433,7 @@ public:
 	 *	Find a socket object in this SkeletalMesh by name. 
 	 *	Entering NAME_None will return NULL. If there are multiple sockets with the same name, will return the first one.
 	 */
-	UFUNCTION(BlueprintCallable, Category="Animation")
-	USkeletalMeshSocket* FindSocket(FName InSocketName) const;
+	virtual USkeletalMeshSocket* FindSocket(FName InSocketName) const override;
 
 	/**
 	*	Find a socket object in this SkeletalMesh by name.
@@ -2450,8 +2449,7 @@ public:
 	*	Also returns the index for the socket allowing for future fast access via GetSocketByIndex()
 	*	Also rteturns the socket loca transform and the bone index (if any)
 	*/
-	UFUNCTION(BlueprintCallable, Category = "Animation")
-	USkeletalMeshSocket* FindSocketInfo(FName InSocketName, FTransform& OutTransform, int32& OutBoneIndex, int32& OutIndex) const;
+	virtual USkeletalMeshSocket* FindSocketInfo(FName InSocketName, FTransform& OutTransform, int32& OutBoneIndex, int32& OutIndex) const override;
 
 	/** Returns the number of sockets available. Both on this mesh and it's skeleton. */
 	UFUNCTION(BlueprintCallable, Category = "Animation")
diff --git a/Engine/Source/Runtime/Engine/Classes/Engine/SkinnedAsset.h b/Engine/Source/Runtime/Engine/Classes/Engine/SkinnedAsset.h
--- a/Engine/Source/Runtime/Engine/Classes/Engine/SkinnedAsset.h
+++ b/Engine/Source/Runtime/Engine/Classes/Engine/SkinnedAsset.h
@@ -145,6 +145,24 @@ public:
 	virtual TArray<class USkeletalMeshSocket*> GetActiveSocketList() const
 	PURE_VIRTUAL(USkinnedAsset::GetActiveSocketList, static TArray<class USkeletalMeshSocket*> Dummy; return Dummy;);
 
+	/**
+	 * Find a socket object in this SkeletalMesh by name.
+	 * Entering NAME_None will return NULL. If there are multiple sockets with the same name, will return the first one.
+	 */
+	UFUNCTION(BlueprintCallable, Category = "Animation")
+	virtual USkeletalMeshSocket* FindSocket(FName InSocketName) const
+	PURE_VIRTUAL(USkinnedAsset::FindSocket, return nullptr;);
+
+	/**
+	 * Find a socket object and associated info in this SkeletalMesh by name.
+	 * Entering NAME_None will return NULL. If there are multiple sockets with the same name, will return the first one.
+	 * Also returns the index for the socket allowing for future fast access via GetSocketByIndex()
+	 * Also returns the socket transform and the bone index (if any)
+	 */
+	UFUNCTION(BlueprintCallable, Category = "Animation")
+	virtual USkeletalMeshSocket* FindSocketInfo(FName InSocketName, FTransform& OutTransform, int32& OutBoneIndex, int32& OutIndex) const
+	PURE_VIRTUAL(USkinnedAsset::FindSocketInfo, return nullptr;);
+
 	virtual USkeleton* GetSkeleton()
 	PURE_VIRTUAL(USkinnedAsset::GetSkeleton, return nullptr;);
 	virtual const USkeleton* GetSkeleton() const
diff --git a/Engine/Source/Runtime/Engine/Private/Components/SkeletalMeshComponent.cpp b/Engine/Source/Runtime/Engine/Private/Components/SkeletalMeshComponent.cpp
--- a/Engine/Source/Runtime/Engine/Private/Components/SkeletalMeshComponent.cpp
+++ b/Engine/Source/Runtime/Engine/Private/Components/SkeletalMeshComponent.cpp
@@ -281,6 +281,11 @@ USkeletalMeshComponent::USkeletalMeshComponent(const FObjectInitializer& ObjectI
 	DeferredKinematicUpdateIndex = INDEX_NONE;
 }
 
+USkeletalMesh* USkeletalMeshComponent::GetSkeletalMeshAsset() const
+{
+	return Cast<USkeletalMesh>(GetSkinnedAsset());
+}
+
 void USkeletalMeshComponent::Serialize(FArchive& Ar)
 {
 	PRAGMA_DISABLE_DEPRECATION_WARNINGS
@@ -367,6 +372,13 @@ void USkeletalMeshComponent::PostLoad()
 	{
 		bOverrideMinLod = true;
 	}
+
+#if WITH_EDITORONLY_DATA
+	// Update property alias
+	PRAGMA_DISABLE_DEPRECATION_WARNINGS
+	SkeletalMeshAsset = GetSkeletalMeshAsset();
+	PRAGMA_ENABLE_DEPRECATION_WARNINGS
+#endif
 }
 
 void USkeletalMeshComponent::RegisterComponentTickFunctions(bool bRegister)
@@ -1006,9 +1018,18 @@ void USkeletalMeshComponent::BeginPlay()
 #if WITH_EDITOR
 void USkeletalMeshComponent::PostEditChangeProperty(FPropertyChangedEvent& PropertyChangedEvent)
 {
+	FProperty* PropertyThatChanged = PropertyChangedEvent.Property;
+
+	// Set the skinned asset pointer with the alias pointer (must happen before the call to Super::PostEditChangeProperty)
+	PRAGMA_DISABLE_DEPRECATION_WARNINGS
+	if (PropertyThatChanged && PropertyThatChanged->GetFName() == GET_MEMBER_NAME_CHECKED(USkeletalMeshComponent, SkeletalMeshAsset))
+	{
+		SetSkinnedAsset(SkeletalMeshAsset);
+	}
+	PRAGMA_ENABLE_DEPRECATION_WARNINGS
+
 	Super::PostEditChangeProperty(PropertyChangedEvent);
 
-	FProperty* PropertyThatChanged = PropertyChangedEvent.Property;
 
 	if ( PropertyThatChanged != nullptr )
 	{
@@ -1039,7 +1060,9 @@ void USkeletalMeshComponent::PostEditChangeProperty(FPropertyChangedEvent& Prope
 		}
 
 		USkeletalMesh* SkelMesh = GetSkeletalMesh();
-		if(PropertyThatChanged->GetFName() == GET_MEMBER_NAME_CHECKED( USkeletalMeshComponent, GetSkeletalMesh()))
+		PRAGMA_DISABLE_DEPRECATION_WARNINGS
+		if(PropertyThatChanged->GetFName() == GET_MEMBER_NAME_CHECKED( USkeletalMeshComponent, SkeletalMeshAsset))
+		PRAGMA_ENABLE_DEPRECATION_WARNINGS
 		{
 			ValidateAnimation();
 
@@ -2875,7 +2898,7 @@ FBoxSphereBounds USkeletalMeshComponent::CalcBounds(const FTransform& LocalToWor
 		{
 			const USkinnedMeshComponent* const LeaderPoseComponentInst = LeaderPoseComponent.Get();
 			check(LeaderPoseComponentInst);
-			if (LeaderPoseComponentInst->GetSkeletalMesh() &&
+			if (LeaderPoseComponentInst->GetSkinnedAsset() &&
 				LeaderPoseComponentInst->bComponentUseFixedSkelBounds &&
 				LeaderPoseComponentInst->IsA((USkeletalMeshComponent::StaticClass())))
 			{
@@ -2927,6 +2950,13 @@ void USkeletalMeshComponent::SetSkeletalMesh(USkeletalMesh* InSkelMesh, bool bRe
 		return;
 	}
 
+	// Update property alias
+#if WITH_EDITORONLY_DATA
+	PRAGMA_DISABLE_DEPRECATION_WARNINGS
+	SkeletalMeshAsset = InSkelMesh;
+	PRAGMA_ENABLE_DEPRECATION_WARNINGS
+#endif
+
 	// We may be doing parallel evaluation on the current anim instance
 	// Calling this here with true will block this init till that thread completes
 	// and it is safe to continue
@@ -2938,7 +2968,7 @@ void USkeletalMeshComponent::SetSkeletalMesh(USkeletalMesh* InSkelMesh, bool bRe
 
 	{
 		FRenderStateRecreator RenderStateRecreator(this);
-		Super::SetSkeletalMesh(InSkelMesh, bReinitPose);
+		SetSkinnedAsset(InSkelMesh, bReinitPose);
 		
 #if WITH_EDITOR
 		ValidateAnimation();
diff --git a/Engine/Source/Runtime/Engine/Private/Components/SkinnedMeshComponent.cpp b/Engine/Source/Runtime/Engine/Private/Components/SkinnedMeshComponent.cpp
--- a/Engine/Source/Runtime/Engine/Private/Components/SkinnedMeshComponent.cpp
+++ b/Engine/Source/Runtime/Engine/Private/Components/SkinnedMeshComponent.cpp
@@ -427,10 +427,10 @@ void USkinnedMeshComponent::UpdateMorphMaterialUsageOnProxy()
 	// update morph material usage
 	if (SceneProxy)
 	{
-		if (ActiveMorphTargets.Num() > 0 && GetSkeletalMesh()->GetMorphTargets().Num() > 0)
+		if (ActiveMorphTargets.Num() > 0 && GetSkinnedAsset()->GetMorphTargets().Num() > 0)
 		{
 			TArray<UMaterialInterface*> MaterialUsingMorphTarget;
-			for (UMorphTarget* MorphTarget : GetSkeletalMesh()->GetMorphTargets())
+			for (UMorphTarget* MorphTarget : GetSkinnedAsset()->GetMorphTargets())
 			{
 				if (!MorphTarget)
 				{
@@ -440,9 +440,9 @@ void USkinnedMeshComponent::UpdateMorphMaterialUsageOnProxy()
 				{
 					for (int32 SectionIndex : MorphTargetLODModel.SectionIndices)
 					{
-						for (int32 LodIdx = 0; LodIdx < GetSkeletalMesh()->GetResourceForRendering()->LODRenderData.Num(); LodIdx++)
+						for (int32 LodIdx = 0; LodIdx < GetSkinnedAsset()->GetResourceForRendering()->LODRenderData.Num(); LodIdx++)
 						{
-							const FSkeletalMeshLODRenderData& LODModel = GetSkeletalMesh()->GetResourceForRendering()->LODRenderData[LodIdx];
+							const FSkeletalMeshLODRenderData& LODModel = GetSkinnedAsset()->GetResourceForRendering()->LODRenderData[LodIdx];
 							if (LODModel.RenderSections.IsValidIndex(SectionIndex))
 							{
 								MaterialUsingMorphTarget.AddUnique(GetMaterial(LODModel.RenderSections[SectionIndex].MaterialIndex));
@@ -512,6 +512,18 @@ void USkinnedMeshComponent::Serialize(FArchive& Ar)
 	}
 }
 
+void USkinnedMeshComponent::PostLoad()
+{
+	Super::PostLoad();
+
+	PRAGMA_DISABLE_DEPRECATION_WARNINGS
+	if(!SkinnedAsset)
+	{
+		SkinnedAsset = SkeletalMesh;  // Set to the SkeletalMesh pointer for backward compatibility when the SkinnedAsset is invalid
+	}
+	PRAGMA_ENABLE_DEPRECATION_WARNINGS
+}
+
 void USkinnedMeshComponent::OnRegister()
 {
 	LLM_SCOPE(ELLMTag::SkeletalMesh);
@@ -570,7 +582,7 @@ void USkinnedMeshComponent::OnUnregister()
 void USkinnedMeshComponent::RefreshExternalMorphTargetWeights()
 {
 	// Clear the external weights if there is no skeletal mesh.
-	if (SkeletalMesh.Get() == nullptr || SkeletalMesh->GetResourceForRendering() == nullptr)
+	if (GetSkinnedAsset() == nullptr || GetSkinnedAsset()->GetResourceForRendering() == nullptr)
 	{
 		ExternalMorphWeightData.Empty();
 		ExternalMorphWeightData.AddDefaulted(GetNumLODs());
@@ -581,13 +593,13 @@ void USkinnedMeshComponent::RefreshExternalMorphTargetWeights()
 	const int32 NumLODs = GetNumLODs();
 	ExternalMorphWeightData.Reset();
 	ExternalMorphWeightData.AddDefaulted(NumLODs);
-	check(SkeletalMesh->GetResourceForRendering()->LODRenderData.Num() == NumLODs);
+	check(GetSkinnedAsset()->GetResourceForRendering()->LODRenderData.Num() == NumLODs);
 
 	for (int32 LOD = 0; LOD < NumLODs; ++LOD)
 	{
 		FExternalMorphWeightData& ExternalWeights = GetExternalMorphWeights(LOD);
 
-		auto& MorphSets = SkeletalMesh->GetResourceForRendering()->LODRenderData[LOD].ExternalMorphSets;
+		auto& MorphSets = GetSkinnedAsset()->GetResourceForRendering()->LODRenderData[LOD].ExternalMorphSets;
 		for (const auto& Item : MorphSets)
 		{
 			check(Item.Value.IsValid());
@@ -614,10 +626,10 @@ void USkinnedMeshComponent::CreateRenderState_Concurrent(FRegisterComponentConte
 {
 	LLM_SCOPE(ELLMTag::SkeletalMesh);
 
-	if( GetSkeletalMesh() )
+	if( GetSkinnedAsset() )
 	{
 		// Attempting to track down UE-45505, where it looks as if somehow a skeletal mesh component's mesh has only been partially loaded, causing a mismatch in the LOD arrays
-		checkf(!GetSkeletalMesh()->HasAnyFlags(RF_NeedLoad | RF_NeedPostLoad | RF_NeedPostLoadSubobjects | RF_WillBeLoaded), TEXT("Attempting to create render state for a skeletal mesh that is is not fully loaded. Mesh: %s"), *GetSkeletalMesh()->GetName());
+		checkf(!GetSkinnedAsset()->HasAnyFlags(RF_NeedLoad | RF_NeedPostLoad | RF_NeedPostLoadSubobjects | RF_WillBeLoaded), TEXT("Attempting to create render state for a skeletal mesh that is is not fully loaded. Mesh: %s"), *GetSkinnedAsset()->GetName());
 
 		// Initialize the alternate weight tracks if present BEFORE creating the new mesh object
 		InitLODInfos();
@@ -626,7 +638,7 @@ void USkinnedMeshComponent::CreateRenderState_Concurrent(FRegisterComponentConte
 		if ( FApp::CanEverRender() && ShouldComponentAddToScene() )
 		{
 			ERHIFeatureLevel::Type SceneFeatureLevel = GetWorld()->FeatureLevel;
-			FSkeletalMeshRenderData* SkelMeshRenderData = GetSkeletalMesh()->GetResourceForRendering();
+			FSkeletalMeshRenderData* SkelMeshRenderData = GetSkinnedAsset()->GetResourceForRendering();
 			int32 MinLODIndex = ComputeMinLOD();
 			
 #if DO_CHECK
@@ -636,7 +648,7 @@ void USkinnedMeshComponent::CreateRenderState_Concurrent(FRegisterComponentConte
 				const FPositionVertexBuffer* PositionVertexBufferPtr = &LODData.StaticVertexBuffers.PositionVertexBuffer;
 				if (!PositionVertexBufferPtr || (PositionVertexBufferPtr->GetNumVertices() <= 0))
 				{
-					UE_LOG(LogSkinnedMeshComp, Warning, TEXT("Invalid Lod %i for Rendering Asset: %s"), LODIndex, *GetSkeletalMesh()->GetFullName());
+					UE_LOG(LogSkinnedMeshComp, Warning, TEXT("Invalid Lod %i for Rendering Asset: %s"), LODIndex, *GetSkinnedAsset()->GetFullName());
 				}
 			}
 #endif
@@ -671,7 +683,7 @@ void USkinnedMeshComponent::CreateRenderState_Concurrent(FRegisterComponentConte
 					FString FeatureLevelName; GetFeatureLevelName(SceneFeatureLevel, FeatureLevelName);
 
 					UE_LOG(LogSkinnedMeshComp, Warning, TEXT("SkeletalMesh %s, is not supported for current feature level (%s) and will not be rendered. MinLOD %d, NumBones %d (supported %d), NumBoneInfluences: %d"),
-						*GetNameSafe(GetSkeletalMesh()), *FeatureLevelName, MinLODIndex, MaxBonesPerChunk, MaxSupportedGPUSkinBones, NumBoneInfluences);
+						*GetNameSafe(GetSkinnedAsset()), *FeatureLevelName, MinLODIndex, MaxBonesPerChunk, MaxSupportedGPUSkinBones, NumBoneInfluences);
 				}
 			}
 
@@ -690,7 +702,7 @@ void USkinnedMeshComponent::CreateRenderState_Concurrent(FRegisterComponentConte
 
 	Super::CreateRenderState_Concurrent(Context);
 
-	if (GetSkeletalMesh())
+	if (GetSkinnedAsset())
 	{
 		// Update dynamic data
 
@@ -710,14 +722,14 @@ void USkinnedMeshComponent::CreateRenderState_Concurrent(FRegisterComponentConte
 			}
 
 			// Clamp to loaded streaming data if available
-			if (GetSkeletalMesh()->IsStreamable() && MeshObject)
+			if (GetSkinnedAsset()->IsStreamable() && MeshObject)
 			{
 				ModifiedLODLevel = FMath::Max<int32>(ModifiedLODLevel, MeshObject->GetSkeletalMeshRenderData().PendingFirstLODIdx);
 			}
 
 			// If we have a valid LOD, set up required data, during reimport we may try to create data before we have all the LODs
 			// imported, in that case we skip until we have all the LODs
-			if(GetSkeletalMesh()->IsValidLODIndex(ModifiedLODLevel))
+			if(GetSkinnedAsset()->IsValidLODIndex(ModifiedLODLevel))
 			{
 				const bool bMorphTargetsAllowed = CVarEnableMorphTargets.GetValueOnAnyThread(true) != 0;
 
@@ -777,9 +789,9 @@ FString USkinnedMeshComponent::GetDetailedInfoInternal() const
 {
 	FString Result;  
 
-	if(GetSkeletalMesh() != NULL )
+	if(GetSkinnedAsset() != NULL )
 	{
-		Result = GetSkeletalMesh()->GetDetailedInfoInternal();
+		Result = GetSkinnedAsset()->GetDetailedInfoInternal();
 	}
 	else
 	{
@@ -797,13 +809,13 @@ void USkinnedMeshComponent::SendRenderDynamicData_Concurrent()
 
 	// if we have not updated the transforms then no need to send them to the rendering thread
 	// @todo GIsEditor used to be bUpdateSkelWhenNotRendered. Look into it further to find out why it doesn't update animations in the AnimSetViewer, when a level is loaded in UED (like POC_Cover.gear).
-	if( MeshObject && GetSkeletalMesh() && (bForceMeshObjectUpdate || (bRecentlyRendered || VisibilityBasedAnimTickOption == EVisibilityBasedAnimTickOption::AlwaysTickPoseAndRefreshBones || GIsEditor || MeshObject->bHasBeenUpdatedAtLeastOnce == false)) )
+	if( MeshObject && GetSkinnedAsset() && (bForceMeshObjectUpdate || (bRecentlyRendered || VisibilityBasedAnimTickOption == EVisibilityBasedAnimTickOption::AlwaysTickPoseAndRefreshBones || GIsEditor || MeshObject->bHasBeenUpdatedAtLeastOnce == false)) )
 	{
 		SCOPE_CYCLE_COUNTER(STAT_MeshObjectUpdate);
 
 		int32 UseLOD = GetPredictedLODLevel();
 		// Clamp to loaded streaming data if available
-		if (GetSkeletalMesh()->IsStreamable() && MeshObject)
+		if (GetSkinnedAsset()->IsStreamable() && MeshObject)
 		{
 			UseLOD = FMath::Max<int32>(UseLOD, MeshObject->GetSkeletalMeshRenderData().PendingFirstLODIdx);
 		}
@@ -840,7 +852,7 @@ void USkinnedMeshComponent::SendRenderDynamicData_Concurrent()
 		{
 			if (MeshDeformerInstance->IsActive())
 			{
-				MeshDeformerInstance->EnqueueWork(GetScene(), UMeshDeformerInstance::WorkLoad_Update, GetSkeletalMesh()->GetFName());
+				MeshDeformerInstance->EnqueueWork(GetScene(), UMeshDeformerInstance::WorkLoad_Update, GetSkinnedAsset()->GetFName());
 			}
 			else
 			{
@@ -858,7 +870,7 @@ void USkinnedMeshComponent::ClearMotionVector()
 	{
 		int32 UseLOD = GetPredictedLODLevel();
 		// Clamp to loaded streaming data if available
-		if (GetSkeletalMesh()->IsStreamable() && MeshObject)
+		if (GetSkinnedAsset()->IsStreamable() && MeshObject)
 		{
 			UseLOD = FMath::Max<int32>(UseLOD, MeshObject->GetSkeletalMeshRenderData().PendingFirstLODIdx);
 		}
@@ -884,7 +896,7 @@ void USkinnedMeshComponent::ForceMotionVector()
 	{
 		int32 UseLOD = GetPredictedLODLevel();
 		// Clamp to loaded streaming data if available
-		if (GetSkeletalMesh()->IsStreamable() && MeshObject)
+		if (GetSkinnedAsset()->IsStreamable() && MeshObject)
 		{
 			UseLOD = FMath::Max<int32>(UseLOD, MeshObject->GetSkeletalMeshRenderData().PendingFirstLODIdx);
 		}
@@ -923,12 +935,12 @@ bool USkinnedMeshComponent::CanEditChange(const FProperty* InProperty) const
 
 void USkinnedMeshComponent::InitLODInfos()
 {
-	if (GetSkeletalMesh() != NULL)
+	if (GetSkinnedAsset() != NULL)
 	{
-		if (GetSkeletalMesh()->GetLODNum() != LODInfo.Num())
+		if (GetSkinnedAsset()->GetLODNum() != LODInfo.Num())
 		{
-			LODInfo.Empty(GetSkeletalMesh()->GetLODNum());
-			for (int32 Idx=0; Idx < GetSkeletalMesh()->GetLODNum(); Idx++)
+			LODInfo.Empty(GetSkinnedAsset()->GetLODNum());
+			for (int32 Idx=0; Idx < GetSkinnedAsset()->GetLODNum(); Idx++)
 			{
 				new(LODInfo) FSkelMeshComponentLODInfo();
 			}
@@ -975,7 +987,7 @@ void USkinnedMeshComponent::TickUpdateRate(float DeltaTime, bool bNeedsValidRoot
 				DrawDebugBox(GetWorld(), Bounds.Origin, Bounds.BoxExtent, FQuat::Identity, DrawColor, false);
 
 				FString DebugString = FString::Printf(TEXT("%s UpdateRate(%d) EvaluationRate(%d) ShouldInterpolateSkippedFrames(%d) ShouldSkipUpdate(%d) Interp Alpha (%f) AdditionalTime(%f)"),
-					*GetNameSafe(GetSkeletalMesh()), AnimUpdateRateParams->UpdateRate, AnimUpdateRateParams->EvaluationRate,
+					*GetNameSafe(GetSkinnedAsset()), AnimUpdateRateParams->UpdateRate, AnimUpdateRateParams->EvaluationRate,
 					AnimUpdateRateParams->ShouldInterpolateSkippedFrames(), AnimUpdateRateParams->ShouldSkipUpdate(), AnimUpdateRateParams->GetInterpolationAlpha(), AnimUpdateRateParams->AdditionalTime);
 
 				GEngine->AddOnScreenDebugMessage(INDEX_NONE, 0.f, FColor::Red, DebugString, false);
@@ -1045,7 +1057,7 @@ void USkinnedMeshComponent::TickComponent(float DeltaTime, enum ELevelTick TickT
 
 UObject const* USkinnedMeshComponent::AdditionalStatObject() const
 {
-	return GetSkeletalMesh();
+	return GetSkinnedAsset();
 }
 
 void USkinnedMeshComponent::UpdateFollowerComponent()
@@ -1056,9 +1068,9 @@ void USkinnedMeshComponent::UpdateFollowerComponent()
 // this has to be skeletalmesh material. You can't have more than what SkeletalMesh materials have
 int32 USkinnedMeshComponent::GetNumMaterials() const
 {
-	if (GetSkeletalMesh())
+	if (GetSkinnedAsset())
 	{
-		return GetSkeletalMesh()->GetMaterials().Num();
+		return GetSkinnedAsset()->GetMaterials().Num();
 	}
 
 	return 0;
@@ -1070,9 +1082,9 @@ UMaterialInterface* USkinnedMeshComponent::GetMaterial(int32 MaterialIndex) cons
 	{
 		return OverrideMaterials[MaterialIndex];
 	}
-	else if (GetSkeletalMesh() && GetSkeletalMesh()->GetMaterials().IsValidIndex(MaterialIndex) && GetSkeletalMesh()->GetMaterials()[MaterialIndex].MaterialInterface)
+	else if (GetSkinnedAsset() && GetSkinnedAsset()->GetMaterials().IsValidIndex(MaterialIndex) && GetSkinnedAsset()->GetMaterials()[MaterialIndex].MaterialInterface)
 	{
-		return GetSkeletalMesh()->GetMaterials()[MaterialIndex].MaterialInterface;
+		return GetSkinnedAsset()->GetMaterials()[MaterialIndex].MaterialInterface;
 	}
 
 	return nullptr;
@@ -1080,9 +1092,9 @@ UMaterialInterface* USkinnedMeshComponent::GetMaterial(int32 MaterialIndex) cons
 
 int32 USkinnedMeshComponent::GetMaterialIndex(FName MaterialSlotName) const
 {
-	if (GetSkeletalMesh() != nullptr)
+	if (GetSkinnedAsset() != nullptr)
 	{
-		const TArray<FSkeletalMaterial>& SkeletalMeshMaterials = GetSkeletalMesh()->GetMaterials();
+		const TArray<FSkeletalMaterial>& SkeletalMeshMaterials = GetSkinnedAsset()->GetMaterials();
 		for (int32 MaterialIndex = 0; MaterialIndex < SkeletalMeshMaterials.Num(); ++MaterialIndex)
 		{
 			const FSkeletalMaterial &SkeletalMaterial = SkeletalMeshMaterials[MaterialIndex];
@@ -1098,9 +1110,9 @@ int32 USkinnedMeshComponent::GetMaterialIndex(FName MaterialSlotName) const
 TArray<FName> USkinnedMeshComponent::GetMaterialSlotNames() const
 {
 	TArray<FName> MaterialNames;
-	if (GetSkeletalMesh() != nullptr)
+	if (GetSkinnedAsset() != nullptr)
 	{
-		const TArray<FSkeletalMaterial>& SkeletalMeshMaterials = GetSkeletalMesh()->GetMaterials();
+		const TArray<FSkeletalMaterial>& SkeletalMeshMaterials = GetSkinnedAsset()->GetMaterials();
 		for (int32 MaterialIndex = 0; MaterialIndex < SkeletalMeshMaterials.Num(); ++MaterialIndex)
 		{
 			const FSkeletalMaterial &SkeletalMaterial = SkeletalMeshMaterials[MaterialIndex];
@@ -1130,7 +1142,7 @@ PRAGMA_ENABLE_DEPRECATION_WARNINGS
 void USkinnedMeshComponent::SetCPUSkinningEnabled(bool bEnable, bool bRecreateRenderStateImmediately)
 {
 	checkSlow(IsInGameThread());
-	check(GetSkeletalMesh() && GetSkeletalMesh()->GetResourceForRendering());
+	check(GetSkinnedAsset() && GetSkinnedAsset()->GetResourceForRendering());
 
 	if (GetCPUSkinningEnabled() == bEnable)
 	{
@@ -1144,16 +1156,16 @@ void USkinnedMeshComponent::SetCPUSkinningEnabled(bool bEnable, bool bRecreateRe
 		IRenderAssetStreamingManager& Manager = IStreamingManager::Get().GetRenderAssetStreamingManager();
 		Manager.BlockTillAllRequestsFinished();
 
-		const bool bOriginalForcedFullyLoad = GetSkeletalMesh()->bForceMiplevelsToBeResident;
-		GetSkeletalMesh()->bForceMiplevelsToBeResident = true;
-		Manager.UpdateIndividualRenderAsset(GetSkeletalMesh());
+		const bool bOriginalForcedFullyLoad = GetSkinnedAsset()->bForceMiplevelsToBeResident;
+		GetSkinnedAsset()->bForceMiplevelsToBeResident = true;
+		Manager.UpdateIndividualRenderAsset(GetSkinnedAsset());
 
-		GetSkeletalMesh()->WaitForPendingInitOrStreaming();
+		GetSkinnedAsset()->WaitForPendingInitOrStreaming();
 
-		check(GetSkeletalMesh()->GetResourceForRendering()->CurrentFirstLODIdx <= GetSkeletalMesh()->GetDefaultMinLod());
+		check(GetSkinnedAsset()->GetResourceForRendering()->CurrentFirstLODIdx <= GetSkinnedAsset()->GetDefaultMinLod());
 
-		GetSkeletalMesh()->UnlinkStreaming();
-		GetSkeletalMesh()->bForceMiplevelsToBeResident = bOriginalForcedFullyLoad;
+		GetSkinnedAsset()->UnlinkStreaming();
+		GetSkinnedAsset()->bForceMiplevelsToBeResident = bOriginalForcedFullyLoad;
 	}
 
 PRAGMA_DISABLE_DEPRECATION_WARNINGS
@@ -1176,10 +1188,10 @@ PRAGMA_ENABLE_DEPRECATION_WARNINGS
 
 bool USkinnedMeshComponent::GetMaterialStreamingData(int32 MaterialIndex, FPrimitiveMaterialInfo& MaterialData) const
 {
-	if (GetSkeletalMesh() && GetSkeletalMesh()->IsMaterialUsed(MaterialIndex))
+	if (GetSkinnedAsset() && GetSkinnedAsset()->IsMaterialUsed(MaterialIndex))
 	{
 		MaterialData.Material = GetMaterial(MaterialIndex);
-		MaterialData.UVChannelData = GetSkeletalMesh()->GetUVChannelData(MaterialIndex);
+		MaterialData.UVChannelData = GetSkinnedAsset()->GetUVChannelData(MaterialIndex);
 		MaterialData.PackedRelativeBox = PackedRelativeBox_Identity;
 	}
 	return MaterialData.IsValid();
@@ -1189,11 +1201,11 @@ void USkinnedMeshComponent::GetStreamingRenderAssetInfo(FStreamingTextureLevelCo
 {
 	GetStreamingTextureInfoInner(LevelContext, nullptr, StreamingDistanceMultiplier, OutStreamingRenderAssets);
 
-	if (GetSkeletalMesh() && GetSkeletalMesh()->IsStreamable())
+	if (GetSkinnedAsset() && GetSkinnedAsset()->IsStreamable())
 	{
 		const int32 LocalForcedLodModel = GetForcedLOD();
-		const float TexelFactor = LocalForcedLodModel > 0 ? -(GetSkeletalMesh()->GetLODNum() - LocalForcedLodModel + 1) : Bounds.SphereRadius * 2.f;
-		new (OutStreamingRenderAssets) FStreamingRenderAssetPrimitiveInfo(GetSkeletalMesh(), Bounds, TexelFactor, PackedRelativeBox_Identity, false, false);
+		const float TexelFactor = LocalForcedLodModel > 0 ? -(GetSkinnedAsset()->GetLODNum() - LocalForcedLodModel + 1) : Bounds.SphereRadius * 2.f;
+		new (OutStreamingRenderAssets) FStreamingRenderAssetPrimitiveInfo(GetSkinnedAsset(), Bounds, TexelFactor, PackedRelativeBox_Identity, false, false);
 	}
 }
 
@@ -1217,9 +1229,9 @@ void USkinnedMeshComponent::RebuildVisibilityArray()
 
 		// The following code relies on a complete hierarchy sorted from parent to children
 		TArray<uint8>& EditableBoneVisibilityStates = GetEditableBoneVisibilityStates();
-		if (EditableBoneVisibilityStates.Num() != GetSkeletalMesh()->GetRefSkeleton().GetNum())
+		if (EditableBoneVisibilityStates.Num() != GetSkinnedAsset()->GetRefSkeleton().GetNum())
 		{
-			UE_LOG(LogSkinnedMeshComp, Warning, TEXT("RebuildVisibilityArray() failed because EditableBoneVisibilityStates size: %d not equal to RefSkeleton bone count: %d."), EditableBoneVisibilityStates.Num(), GetSkeletalMesh()->GetRefSkeleton().GetNum());
+			UE_LOG(LogSkinnedMeshComp, Warning, TEXT("RebuildVisibilityArray() failed because EditableBoneVisibilityStates size: %d not equal to RefSkeleton bone count: %d."), EditableBoneVisibilityStates.Num(), GetSkinnedAsset()->GetRefSkeleton().GetNum());
 			return;
 		}
 
@@ -1231,7 +1243,7 @@ void USkinnedMeshComponent::RebuildVisibilityArray()
 			if (VisState != BVS_ExplicitlyHidden)
 			{
 				// Check direct parent (only need to do one deep, since we have already processed the parent and written to BoneVisibilityStates previously)
-				const int32 ParentIndex = GetSkeletalMesh()->GetRefSkeleton().GetParentIndex(BoneId);
+				const int32 ParentIndex = GetSkinnedAsset()->GetRefSkeleton().GetParentIndex(BoneId);
 				if ((ParentIndex == -1) || (EditableBoneVisibilityStates[ParentIndex] == BVS_Visible))
 				{
 					EditableBoneVisibilityStates[BoneId] = BVS_Visible;
@@ -1280,9 +1292,9 @@ class UPhysicsAsset* USkinnedMeshComponent::GetPhysicsAsset() const
 		return PhysicsAssetOverride;
 	}
 
-	if (GetSkeletalMesh() && GetSkeletalMesh()->GetPhysicsAsset())
+	if (GetSkinnedAsset() && GetSkinnedAsset()->GetPhysicsAsset())
 	{
-		return GetSkeletalMesh()->GetPhysicsAsset();
+		return GetSkinnedAsset()->GetPhysicsAsset();
 	}
 
 	return nullptr;
@@ -1301,10 +1313,10 @@ FBoxSphereBounds USkinnedMeshComponent::CalcMeshBound(const FVector3f& RootOffse
 	UPhysicsAsset * const LeaderPhysicsAsset = (LeaderPoseComponentInst != nullptr)? LeaderPoseComponentInst->GetPhysicsAsset() : nullptr;
 
 	// Can only use the PhysicsAsset to calculate the bounding box if we are not non-uniformly scaling the mesh.
-	const USkeletalMesh* SkeletalMeshConst = GetSkeletalMesh();
-	const bool bCanUsePhysicsAsset = DrawScale.IsUniform() && (GetSkeletalMesh() != NULL)
+	const USkinnedAsset* SkinnedAssetConst = GetSkinnedAsset();
+	const bool bCanUsePhysicsAsset = DrawScale.IsUniform() && (GetSkinnedAsset() != NULL)
 		// either space base exists or child component
-		&& ( (GetNumComponentSpaceTransforms() == SkeletalMeshConst->GetRefSkeleton().GetNum()) || (LeaderPhysicsAsset) );
+		&& ( (GetNumComponentSpaceTransforms() == SkinnedAssetConst->GetRefSkeleton().GetNum()) || (LeaderPhysicsAsset) );
 
 	const bool bDetailModeAllowsRendering = (DetailMode <= GetCachedScalabilityCVars().DetailMode);
 	const bool bIsVisible = ( bDetailModeAllowsRendering && (ShouldRender() || bCastHiddenShadow));
@@ -1313,15 +1325,15 @@ FBoxSphereBounds USkinnedMeshComponent::CalcMeshBound(const FVector3f& RootOffse
 	const bool bLeaderHasPhysBodies = LeaderPhysicsAsset && LeaderPhysicsAsset->SkeletalBodySetups.Num();
 
 	// if not visible, or we were told to use fixed bounds, use skelmesh bounds
-	if ( (!bIsVisible || bComponentUseFixedSkelBounds) && GetSkeletalMesh())
+	if ( (!bIsVisible || bComponentUseFixedSkelBounds) && GetSkinnedAsset())
 	{
-		FBoxSphereBounds RootAdjustedBounds = SkeletalMeshConst->GetBounds();
+		FBoxSphereBounds RootAdjustedBounds = SkinnedAssetConst->GetBounds();
 		RootAdjustedBounds.Origin += FVector(RootOffset); // Adjust bounds by root bone translation
 		NewBounds = RootAdjustedBounds.TransformBy(LocalToWorld);
 	}
-	else if(LeaderPoseComponentInst && LeaderPoseComponentInst->GetSkeletalMesh() && LeaderPoseComponentInst->bComponentUseFixedSkelBounds)
+	else if(LeaderPoseComponentInst && LeaderPoseComponentInst->GetSkinnedAsset() && LeaderPoseComponentInst->bComponentUseFixedSkelBounds)
 	{
-		FBoxSphereBounds RootAdjustedBounds = LeaderPoseComponentInst->GetSkeletalMesh()->GetBounds();
+		FBoxSphereBounds RootAdjustedBounds = LeaderPoseComponentInst->GetSkinnedAsset()->GetBounds();
 		RootAdjustedBounds.Origin += FVector(RootOffset); // Adjust bounds by root bone translation
 		NewBounds = RootAdjustedBounds.TransformBy(LocalToWorld);
 	}
@@ -1332,7 +1344,7 @@ FBoxSphereBounds USkinnedMeshComponent::CalcMeshBound(const FVector3f& RootOffse
 	}
 #if WITH_EDITOR
 	// For AnimSet Viewer, use 'bounds preview' physics asset if present.
-	else if(GetSkeletalMesh() && bHasPhysBodies && bCanUsePhysicsAsset && PhysicsAsset->CanCalculateValidAABB(this, LocalToWorld))
+	else if(GetSkinnedAsset() && bHasPhysBodies && bCanUsePhysicsAsset && PhysicsAsset->CanCalculateValidAABB(this, LocalToWorld))
 	{
 		NewBounds = FBoxSphereBounds(PhysicsAsset->CalcAABB(this, LocalToWorld));
 	}
@@ -1348,9 +1360,9 @@ FBoxSphereBounds USkinnedMeshComponent::CalcMeshBound(const FVector3f& RootOffse
 		NewBounds = FBoxSphereBounds(LeaderPhysicsAsset->CalcAABB(this, LocalToWorld));
 	}
 	// Fallback is to use the one from the skeletal mesh. Usually pretty bad in terms of Accuracy of where the SkelMesh Bounds are located (i.e. usually bigger than it needs to be)
-	else if(GetSkeletalMesh())
+	else if(GetSkinnedAsset())
 	{
-		FBoxSphereBounds RootAdjustedBounds = GetSkeletalMesh()->GetBounds();
+		FBoxSphereBounds RootAdjustedBounds = GetSkinnedAsset()->GetBounds();
 
 		// Adjust bounds by root bone translation
 		RootAdjustedBounds.Origin += FVector(RootOffset);
@@ -1373,15 +1385,15 @@ void USkinnedMeshComponent::GetPreSkinnedLocalBounds(FBoxSphereBounds& OutBounds
 {
 	const USkinnedMeshComponent* const LeaderPoseComponentInst = LeaderPoseComponent.Get();
 
-	if (GetSkeletalMesh())
+	if (GetSkinnedAsset())
 	{
 		// Get the Pre-skinned bounds from the skeletal mesh. Note that these bounds are the "ExtendedBounds", so they can be tweaked on the SkeletalMesh   
-		OutBounds = GetSkeletalMesh()->GetBounds();
+		OutBounds = GetSkinnedAsset()->GetBounds();
 	}
-	else if(LeaderPoseComponentInst && LeaderPoseComponentInst->GetSkeletalMesh())
+	else if(LeaderPoseComponentInst && LeaderPoseComponentInst->GetSkinnedAsset())
 	{
 		// Get the bounds from the leader pose if ther is no skeletal mesh
-		OutBounds = LeaderPoseComponentInst->GetSkeletalMesh()->GetBounds();
+		OutBounds = LeaderPoseComponentInst->GetSkinnedAsset()->GetBounds();
 	}
 	else
 	{
@@ -1415,13 +1427,13 @@ FMatrix USkinnedMeshComponent::GetBoneMatrix(int32 BoneIdx) const
 			}
 			else
 			{
-				UE_LOG(LogSkinnedMeshComp, Verbose, TEXT("GetBoneMatrix : ParentBoneIndex(%d:%s) out of range of LeaderPoseComponent->SpaceBases for %s(%s)"), BoneIdx, *GetNameSafe(LeaderPoseComponentInst->GetSkeletalMesh()), *GetNameSafe(GetSkeletalMesh()), *GetPathName());
+				UE_LOG(LogSkinnedMeshComp, Verbose, TEXT("GetBoneMatrix : ParentBoneIndex(%d:%s) out of range of LeaderPoseComponent->SpaceBases for %s(%s)"), BoneIdx, *GetNameSafe(LeaderPoseComponentInst->GetSkinnedAsset()), *GetNameSafe(GetSkinnedAsset()), *GetPathName());
 				return FMatrix::Identity;
 			}
 		}
 		else
 		{
-			UE_LOG(LogSkinnedMeshComp, Warning, TEXT("GetBoneMatrix : BoneIndex(%d) out of range of LeaderBoneMap for %s (%s)"), BoneIdx, *this->GetFName().ToString(), GetSkeletalMesh() ? *GetSkeletalMesh()->GetFName().ToString() : TEXT("NULL") );
+			UE_LOG(LogSkinnedMeshComp, Warning, TEXT("GetBoneMatrix : BoneIndex(%d) out of range of LeaderBoneMap for %s (%s)"), BoneIdx, *this->GetFName().ToString(), GetSkinnedAsset() ? *GetSkinnedAsset()->GetFName().ToString() : TEXT("NULL") );
 			return FMatrix::Identity;
 		}
 	}
@@ -1434,7 +1446,7 @@ FMatrix USkinnedMeshComponent::GetBoneMatrix(int32 BoneIdx) const
 		}
 		else
 		{
-			UE_LOG(LogSkinnedMeshComp, Warning, TEXT("GetBoneMatrix : BoneIndex(%d) out of range of SpaceBases for %s (%s)"), BoneIdx, *GetPathName(), GetSkeletalMesh() ? *GetSkeletalMesh()->GetFullName() : TEXT("NULL") );
+			UE_LOG(LogSkinnedMeshComp, Warning, TEXT("GetBoneMatrix : BoneIndex(%d) out of range of SpaceBases for %s (%s)"), BoneIdx, *GetPathName(), GetSkinnedAsset() ? *GetSkinnedAsset()->GetFullName() : TEXT("NULL") );
 			return FMatrix::Identity;
 		}
 	}
@@ -1516,7 +1528,7 @@ FTransform USkinnedMeshComponent::GetBoneTransform(int32 BoneIdx, const FTransfo
 		}
 		else
 		{
-			UE_LOG(LogSkinnedMeshComp, Verbose, TEXT("GetBoneTransform : BoneIndex(%d) out of range of SpaceBases for %s (%s)"), BoneIdx, *GetPathName(), GetSkeletalMesh() ? *GetSkeletalMesh()->GetFullName() : TEXT("NULL") );
+			UE_LOG(LogSkinnedMeshComp, Verbose, TEXT("GetBoneTransform : BoneIndex(%d) out of range of SpaceBases for %s (%s)"), BoneIdx, *GetPathName(), GetSkinnedAsset() ? *GetSkinnedAsset()->GetFullName() : TEXT("NULL") );
 			return FTransform::Identity;
 		}
 	}
@@ -1524,15 +1536,15 @@ FTransform USkinnedMeshComponent::GetBoneTransform(int32 BoneIdx, const FTransfo
 
 int32 USkinnedMeshComponent::GetNumBones()const
 {
-	return GetSkeletalMesh() ? GetSkeletalMesh()->GetRefSkeleton().GetNum() : 0;
+	return GetSkinnedAsset() ? GetSkinnedAsset()->GetRefSkeleton().GetNum() : 0;
 }
 
 int32 USkinnedMeshComponent::GetBoneIndex( FName BoneName) const
 {
 	int32 BoneIndex = INDEX_NONE;
-	if ( BoneName != NAME_None && GetSkeletalMesh())
+	if ( BoneName != NAME_None && GetSkinnedAsset())
 	{
-		BoneIndex = GetSkeletalMesh()->GetRefSkeleton().FindBoneIndex( BoneName );
+		BoneIndex = GetSkinnedAsset()->GetRefSkeleton().FindBoneIndex( BoneName );
 	}
 
 	return BoneIndex;
@@ -1541,7 +1553,7 @@ int32 USkinnedMeshComponent::GetBoneIndex( FName BoneName) const
 
 FName USkinnedMeshComponent::GetBoneName(int32 BoneIndex) const
 {
-	return (GetSkeletalMesh() != NULL && GetSkeletalMesh()->GetRefSkeleton().IsValidIndex(BoneIndex)) ? GetSkeletalMesh()->GetRefSkeleton().GetBoneName(BoneIndex) : NAME_None;
+	return (GetSkinnedAsset() != NULL && GetSkinnedAsset()->GetRefSkeleton().IsValidIndex(BoneIndex)) ? GetSkinnedAsset()->GetRefSkeleton().GetBoneName(BoneIndex) : NAME_None;
 }
 
 
@@ -1552,16 +1564,16 @@ FName USkinnedMeshComponent::GetParentBone( FName BoneName ) const
 	int32 BoneIndex = GetBoneIndex(BoneName);
 	if ((BoneIndex != INDEX_NONE) && (BoneIndex > 0)) // This checks that this bone is not the root (ie no parent), and that BoneIndex != INDEX_NONE (ie bone name was found)
 	{
-		Result = GetSkeletalMesh()->GetRefSkeleton().GetBoneName(GetSkeletalMesh()->GetRefSkeleton().GetParentIndex(BoneIndex));
+		Result = GetSkinnedAsset()->GetRefSkeleton().GetBoneName(GetSkinnedAsset()->GetRefSkeleton().GetParentIndex(BoneIndex));
 	}
 	return Result;
 }
 
 FTransform USkinnedMeshComponent::GetDeltaTransformFromRefPose(FName BoneName, FName BaseName/* = NAME_None*/) const
 {
-	if (GetSkeletalMesh())
+	if (GetSkinnedAsset())
 	{
-		const FReferenceSkeleton& RefSkeleton = GetSkeletalMesh()->GetRefSkeleton();
+		const FReferenceSkeleton& RefSkeleton = GetSkinnedAsset()->GetRefSkeleton();
 		const int32 BoneIndex = GetBoneIndex(BoneName);
 		if (BoneIndex != INDEX_NONE)
 		{
@@ -1589,7 +1601,7 @@ FTransform USkinnedMeshComponent::GetDeltaTransformFromRefPose(FName BoneName, F
 
 bool USkinnedMeshComponent::GetTwistAndSwingAngleOfDeltaRotationFromRefPose(FName BoneName, float& OutTwistAngle, float& OutSwingAngle) const
 {
-	const FReferenceSkeleton& RefSkeleton = GetSkeletalMesh()->GetRefSkeleton();
+	const FReferenceSkeleton& RefSkeleton = GetSkinnedAsset()->GetRefSkeleton();
 	const int32 BoneIndex = GetBoneIndex(BoneName);
 	const TArray<FTransform>& Transforms = GetComponentSpaceTransforms();
 
@@ -1642,12 +1654,12 @@ bool USkinnedMeshComponent::IsSkinCacheAllowed(int32 LodIdx) const
 		return false;
 	}
 
-	if (!GetSkeletalMesh())
+	if (!GetSkinnedAsset())
 	{
 		return GEnableGPUSkinCache && bGlobalDefault;
 	}
 
-	FSkeletalMeshLODInfo* LodInfo = GetSkeletalMesh()->GetLODInfo(LodIdx);
+	FSkeletalMeshLODInfo* LodInfo = GetSkinnedAsset()->GetLODInfo(LodIdx);
 	if (!LodInfo)
 	{
 		return GEnableGPUSkinCache && bGlobalDefault;
@@ -1671,7 +1683,7 @@ bool USkinnedMeshComponent::IsSkinCacheAllowed(int32 LodIdx) const
 
 void USkinnedMeshComponent::GetBoneNames(TArray<FName>& BoneNames)
 {
-	if (GetSkeletalMesh() == NULL)
+	if (GetSkinnedAsset() == NULL)
 	{
 		// no mesh, so no bones
 		BoneNames.Empty();
@@ -1679,11 +1691,11 @@ void USkinnedMeshComponent::GetBoneNames(TArray<FName>& BoneNames)
 	else
 	{
 		// pre-size the array to avoid unnecessary reallocation
-		BoneNames.Empty(GetSkeletalMesh()->GetRefSkeleton().GetNum());
-		BoneNames.AddUninitialized(GetSkeletalMesh()->GetRefSkeleton().GetNum());
-		for (int32 i = 0; i < GetSkeletalMesh()->GetRefSkeleton().GetNum(); i++)
+		BoneNames.Empty(GetSkinnedAsset()->GetRefSkeleton().GetNum());
+		BoneNames.AddUninitialized(GetSkinnedAsset()->GetRefSkeleton().GetNum());
+		for (int32 i = 0; i < GetSkinnedAsset()->GetRefSkeleton().GetNum(); i++)
 		{
-			BoneNames[i] = GetSkeletalMesh()->GetRefSkeleton().GetBoneName(i);
+			BoneNames[i] = GetSkinnedAsset()->GetRefSkeleton().GetBoneName(i);
 		}
 	}
 }
@@ -1693,23 +1705,23 @@ bool USkinnedMeshComponent::BoneIsChildOf(FName BoneName, FName ParentBoneName)
 {
 	bool bResult = false;
 
-	if(GetSkeletalMesh())
+	if(GetSkinnedAsset())
 	{
-		const int32 BoneIndex = GetSkeletalMesh()->GetRefSkeleton().FindBoneIndex(BoneName);
+		const int32 BoneIndex = GetSkinnedAsset()->GetRefSkeleton().FindBoneIndex(BoneName);
 		if(BoneIndex == INDEX_NONE)
 		{
-			UE_LOG(LogSkinnedMeshComp, Log, TEXT("execBoneIsChildOf: BoneName '%s' not found in SkeletalMesh '%s'"), *BoneName.ToString(), *GetSkeletalMesh()->GetName());
+			UE_LOG(LogSkinnedMeshComp, Log, TEXT("execBoneIsChildOf: BoneName '%s' not found in SkeletalMesh '%s'"), *BoneName.ToString(), *GetSkinnedAsset()->GetName());
 			return bResult;
 		}
 
-		const int32 ParentBoneIndex = GetSkeletalMesh()->GetRefSkeleton().FindBoneIndex(ParentBoneName);
+		const int32 ParentBoneIndex = GetSkinnedAsset()->GetRefSkeleton().FindBoneIndex(ParentBoneName);
 		if(ParentBoneIndex == INDEX_NONE)
 		{
-			UE_LOG(LogSkinnedMeshComp, Log, TEXT("execBoneIsChildOf: ParentBoneName '%s' not found in SkeletalMesh '%s'"), *ParentBoneName.ToString(), *GetSkeletalMesh()->GetName());
+			UE_LOG(LogSkinnedMeshComp, Log, TEXT("execBoneIsChildOf: ParentBoneName '%s' not found in SkeletalMesh '%s'"), *ParentBoneName.ToString(), *GetSkinnedAsset()->GetName());
 			return bResult;
 		}
 
-		bResult = GetSkeletalMesh()->GetRefSkeleton().BoneIsChildOf(BoneIndex, ParentBoneIndex);
+		bResult = GetSkinnedAsset()->GetRefSkeleton().BoneIsChildOf(BoneIndex, ParentBoneIndex);
 	}
 
 	return bResult;
@@ -1718,9 +1730,9 @@ bool USkinnedMeshComponent::BoneIsChildOf(FName BoneName, FName ParentBoneName)
 
 FVector USkinnedMeshComponent::GetRefPosePosition(int32 BoneIndex) const
 {
-	if(GetSkeletalMesh() && (BoneIndex >= 0) && (BoneIndex < GetSkeletalMesh()->GetRefSkeleton().GetNum()))
+	if(GetSkinnedAsset() && (BoneIndex >= 0) && (BoneIndex < GetSkinnedAsset()->GetRefSkeleton().GetNum()))
 	{
-		return GetSkeletalMesh()->GetRefSkeleton().GetRefBonePose()[BoneIndex].GetTranslation();
+		return GetSkinnedAsset()->GetRefSkeleton().GetRefBonePose()[BoneIndex].GetTranslation();
 	}
 	else
 	{
@@ -1730,9 +1742,9 @@ FVector USkinnedMeshComponent::GetRefPosePosition(int32 BoneIndex) const
 
 FTransform USkinnedMeshComponent::GetRefPoseTransform(int32 BoneIndex) const
 {
-	if(GetSkeletalMesh() && (BoneIndex >= 0) && (BoneIndex < GetSkeletalMesh()->GetRefSkeleton().GetNum()))
+	if(GetSkinnedAsset() && (BoneIndex >= 0) && (BoneIndex < GetSkinnedAsset()->GetRefSkeleton().GetNum()))
 	{
-		return GetSkeletalMesh()->GetRefSkeleton().GetRefBonePose()[BoneIndex];
+		return GetSkinnedAsset()->GetRefSkeleton().GetRefBonePose()[BoneIndex];
 	}
 	else
 	{
@@ -1740,12 +1752,22 @@ FTransform USkinnedMeshComponent::GetRefPoseTransform(int32 BoneIndex) const
 	}
 }
 
-void USkinnedMeshComponent::SetSkeletalMesh(USkeletalMesh* InSkelMesh, bool bReinitPose)
+void USkinnedMeshComponent::SetSkeletalMesh(USkeletalMesh* NewMesh, bool bReinitPose)
 {
-	// NOTE: InSkelMesh may be NULL (useful in the editor for removing the skeletal mesh associated with
+	SetSkinnedAsset(NewMesh, bReinitPose);
+}
+
+USkeletalMesh* USkinnedMeshComponent::GetSkeletalMesh_DEPRECATED() const
+{
+	return Cast<class USkeletalMesh>(GetSkinnedAsset());
+}
+
+void USkinnedMeshComponent::SetSkinnedAsset(USkinnedAsset* InSkinnedAsset, bool /*bReinitPose*/)
+{
+	// NOTE: InSkinnedAsset may be NULL (useful in the editor for removing the skeletal mesh associated with
 	//   this component on-the-fly)
 
-	if (InSkelMesh == GetSkeletalMesh())
+	if (InSkinnedAsset == GetSkinnedAsset())
 	{
 		// do nothing if the input mesh is the same mesh we're already using.
 		return;
@@ -1755,7 +1777,7 @@ void USkinnedMeshComponent::SetSkeletalMesh(USkeletalMesh* InSkelMesh, bool bRei
 		//Handle destroying and recreating the renderstate
 		FRenderStateRecreator RenderStateRecreator(this);
 
-		GetSkeletalMesh() = InSkelMesh;
+		SetSkinnedAsset(InSkinnedAsset);
 		SetPredictedLODLevel(0);
 
 		//FollowerPoseComponents is an array of weak obj ptrs, so it can contain null elements
@@ -1780,7 +1802,7 @@ void USkinnedMeshComponent::SetSkeletalMesh(USkeletalMesh* InSkelMesh, bool bRei
 			UpdateLeaderBoneMap();
 			InvalidateCachedBounds();
 			// clear morphtarget cache
-			ActiveMorphTargets.Empty();			
+			ActiveMorphTargets.Empty();
 			MorphTargetWeights.Empty();
 			ExternalMorphWeightData.Empty();
 		}
@@ -1797,6 +1819,21 @@ void USkinnedMeshComponent::SetSkeletalMesh(USkeletalMesh* InSkelMesh, bool bRei
 	}
 }
 
+void USkinnedMeshComponent::SetSkinnedAsset(class USkinnedAsset* InSkinnedAsset)
+{
+	PRAGMA_DISABLE_DEPRECATION_WARNINGS
+	SkinnedAsset = InSkinnedAsset;
+	SkeletalMesh = nullptr;  // The SkeletalMesh pointer is no longer valid from here on, blueprint/function access will use getter to retrieve the SkinnedAsset pointer and direct access is deprecated
+	PRAGMA_ENABLE_DEPRECATION_WARNINGS
+}
+
+USkinnedAsset* USkinnedMeshComponent::GetSkinnedAsset() const
+{
+	PRAGMA_DISABLE_DEPRECATION_WARNINGS
+	return SkinnedAsset ? SkinnedAsset : SkeletalMesh;  // Return the SkeletalMesh pointer for backward compatibility when the SkinnedAsset is invalid
+	PRAGMA_ENABLE_DEPRECATION_WARNINGS
+}
+
 void USkinnedMeshComponent::SetMeshDeformer(UMeshDeformer* InMeshDeformer)
 {
 	MeshDeformer = InMeshDeformer;
@@ -1810,9 +1847,9 @@ FSkeletalMeshRenderData* USkinnedMeshComponent::GetSkeletalMeshRenderData() cons
 	{
 		return &MeshObject->GetSkeletalMeshRenderData();
 	}
-	else if (GetSkeletalMesh())
+	else if (GetSkinnedAsset())
 	{
-		return GetSkeletalMesh()->GetResourceForRendering();
+		return GetSkinnedAsset()->GetResourceForRendering();
 	}
 	else
 	{
@@ -1825,25 +1862,27 @@ bool USkinnedMeshComponent::AllocateTransformData()
 	LLM_SCOPE_BYNAME(TEXT("SkeletalMesh/TransformData"));
 
 	// Allocate transforms if not present.
-	if (GetSkeletalMesh() != NULL && LeaderPoseComponent == NULL )
+	if (GetSkinnedAsset() != NULL && LeaderPoseComponent == NULL )
 	{
-		if(GetNumComponentSpaceTransforms() != GetSkeletalMesh()->GetRefSkeleton().GetNum() )
+		const int32 NumBones = GetSkinnedAsset()->GetRefSkeleton().GetNum();
+
+		if(GetNumComponentSpaceTransforms() != NumBones)
 		{
 			for (int32 BaseIndex = 0; BaseIndex < 2; ++BaseIndex)
 			{
-				ComponentSpaceTransformsArray[BaseIndex].Empty(GetSkeletalMesh()->GetRefSkeleton().GetNum());
-				ComponentSpaceTransformsArray[BaseIndex].AddUninitialized(GetSkeletalMesh()->GetRefSkeleton().GetNum());
+				ComponentSpaceTransformsArray[BaseIndex].Empty(NumBones);
+				ComponentSpaceTransformsArray[BaseIndex].AddUninitialized(NumBones);
 
-				for (int32 I = 0; I < GetSkeletalMesh()->GetRefSkeleton().GetNum(); ++I)
+				for (int32 I = 0; I < NumBones; ++I)
 				{
 					ComponentSpaceTransformsArray[BaseIndex][I].SetIdentity();
 				}
 
-				BoneVisibilityStates[BaseIndex].Empty( GetSkeletalMesh()->GetRefSkeleton().GetNum() );
-				if( GetSkeletalMesh()->GetRefSkeleton().GetNum() )
+				BoneVisibilityStates[BaseIndex].Empty(NumBones);
+				if(NumBones)
 				{
-					BoneVisibilityStates[BaseIndex].AddUninitialized( GetSkeletalMesh()->GetRefSkeleton().GetNum() );
-					for (int32 BoneIndex = 0; BoneIndex < GetSkeletalMesh()->GetRefSkeleton().GetNum(); BoneIndex++)
+					BoneVisibilityStates[BaseIndex].AddUninitialized(NumBones);
+					for (int32 BoneIndex = 0; BoneIndex < NumBones; BoneIndex++)
 					{
 						BoneVisibilityStates[BaseIndex][BoneIndex] = BVS_Visible;
 					}
@@ -2115,9 +2154,9 @@ void USkinnedMeshComponent::SetSelectedEditorMaterial(int32 InSelectedEditorMate
 
 UMorphTarget* USkinnedMeshComponent::FindMorphTarget( FName MorphTargetName ) const
 {
-	if(GetSkeletalMesh() != NULL )
+	if(GetSkinnedAsset() != NULL )
 	{
-		return GetSkeletalMesh()->FindMorphTarget(MorphTargetName);
+		return GetSkinnedAsset()->FindMorphTarget(MorphTargetName);
 	}
 
 	return NULL;
@@ -2125,7 +2164,7 @@ UMorphTarget* USkinnedMeshComponent::FindMorphTarget( FName MorphTargetName ) co
 
 bool USkinnedMeshComponent::GetMissingLeaderBoneRelativeTransform(int32 InBoneIndex, FMissingLeaderBoneCacheEntry& OutInfo) const
 {
-	const FReferenceSkeleton& FollowerRefSkeleton = GetSkeletalMesh()->GetRefSkeleton();
+	const FReferenceSkeleton& FollowerRefSkeleton = GetSkinnedAsset()->GetRefSkeleton();
 	check(FollowerRefSkeleton.IsValidIndex(InBoneIndex));
 	const TArray<FTransform>& BoneSpaceRefPoseTransforms = FollowerRefSkeleton.GetRefBonePose();
 
@@ -2161,17 +2200,17 @@ void USkinnedMeshComponent::UpdateLeaderBoneMap()
 	LeaderBoneMap.Reset();
 	MissingLeaderBoneMap.Reset();
 
-	if (GetSkeletalMesh())
+	if (GetSkinnedAsset())
 	{
 		if (USkinnedMeshComponent* LeaderPoseComponentPtr = LeaderPoseComponent.Get())
 		{
-			if (USkeletalMesh* LeaderMesh = LeaderPoseComponentPtr->GetSkeletalMesh())
+			if (USkinnedAsset* LeaderMesh = LeaderPoseComponentPtr->GetSkinnedAsset())
 			{
-				const FReferenceSkeleton& FollowerRefSkeleton = GetSkeletalMesh()->GetRefSkeleton();
+				const FReferenceSkeleton& FollowerRefSkeleton = GetSkinnedAsset()->GetRefSkeleton();
 				const FReferenceSkeleton& LeaderRefSkeleton = LeaderMesh->GetRefSkeleton();
 
 				LeaderBoneMap.AddUninitialized(FollowerRefSkeleton.GetNum());
-				if (GetSkeletalMesh() == LeaderMesh)
+				if (GetSkinnedAsset() == LeaderMesh)
 				{
 					// if the meshes are the same, the indices must match exactly so we don't need to look them up
 					for (int32 BoneIndex = 0; BoneIndex < LeaderBoneMap.Num(); BoneIndex++)
@@ -2190,7 +2229,7 @@ void USkinnedMeshComponent::UpdateLeaderBoneMap()
 					// Cache bones for any SOCKET bones that are missing in the leader.
 					// We assume that sockets will be potentially called more often, so we
 					// leave out missing BONE transforms here to try to balance memory & performance.
-					for(USkeletalMeshSocket* Socket : GetSkeletalMesh()->GetActiveSocketList())
+					for(USkeletalMeshSocket* Socket : GetSkinnedAsset()->GetActiveSocketList())
 					{
 						int32 BoneIndex = FollowerRefSkeleton.FindBoneIndex(Socket->BoneName);
 						int32 LeaderBoneIndex = LeaderRefSkeleton.FindBoneIndex(Socket->BoneName);
@@ -2284,10 +2323,10 @@ class USkeletalMeshSocket const* USkinnedMeshComponent::GetSocketInfoByName(FNam
 
 	USkeletalMeshSocket const* Socket = nullptr;
 
-	if(GetSkeletalMesh())
+	if(GetSkinnedAsset())
 	{
 		int32 SocketIndex;
-		Socket = GetSkeletalMesh()->FindSocketInfo(OverrideSocketName, OutTransform, OutBoneIndex, SocketIndex);
+		Socket = GetSkinnedAsset()->FindSocketInfo(OverrideSocketName, OutTransform, OutBoneIndex, SocketIndex);
 	}
 	else
 	{
@@ -2313,9 +2352,9 @@ class USkeletalMeshSocket const* USkinnedMeshComponent::GetSocketByName(FName In
 
 	USkeletalMeshSocket const* Socket = nullptr;
 
-	if(GetSkeletalMesh())
+	if(GetSkinnedAsset())
 	{
-		Socket = GetSkeletalMesh()->FindSocket(OverrideSocketName);
+		Socket = GetSkinnedAsset()->FindSocket(OverrideSocketName);
 	}
 	else
 	{
@@ -2371,7 +2410,7 @@ bool USkinnedMeshComponent::DoesSocketExist(FName InSocketName) const
 
 FName USkinnedMeshComponent::GetSocketBoneName(FName InSocketName) const
 {
-	if(!GetSkeletalMesh())
+	if(!GetSkinnedAsset())
 	{
 		return NAME_None;
 	}
@@ -2380,7 +2419,7 @@ FName USkinnedMeshComponent::GetSocketBoneName(FName InSocketName) const
 	const FName OverrideSocketName = OverrideSocket ? *OverrideSocket : InSocketName;
 
 	// First check for a socket
-	USkeletalMeshSocket const* TmpSocket = GetSkeletalMesh()->FindSocket(OverrideSocketName);
+	USkeletalMeshSocket const* TmpSocket = GetSkinnedAsset()->FindSocket(OverrideSocketName);
 	if( TmpSocket )
 	{
 		return TmpSocket->BoneName;
@@ -2515,19 +2554,19 @@ FVector USkinnedMeshComponent::GetBoneAxis( FName BoneName, EAxis::Type Axis ) c
 
 bool USkinnedMeshComponent::HasAnySockets() const
 {
-	return (GetSkeletalMesh() != NULL) && (
+	return (GetSkinnedAsset() != NULL) && (
 #if WITH_EDITOR
-		(GetSkeletalMesh()->GetActiveSocketList().Num() > 0) ||
+		(GetSkinnedAsset()->GetActiveSocketList().Num() > 0) ||
 #endif
-		(GetSkeletalMesh()->GetRefSkeleton().GetNum() > 0));
+		(GetSkinnedAsset()->GetRefSkeleton().GetNum() > 0));
 }
 
 void USkinnedMeshComponent::QuerySupportedSockets(TArray<FComponentSocketDescription>& OutSockets) const
 {
-	if (GetSkeletalMesh() != NULL)
+	if (GetSkinnedAsset() != NULL)
 	{
 		// Grab all the mesh and skeleton sockets
-		const TArray<USkeletalMeshSocket*> AllSockets = GetSkeletalMesh()->GetActiveSocketList();
+		const TArray<USkeletalMeshSocket*> AllSockets = GetSkinnedAsset()->GetActiveSocketList();
 
 		for (int32 SocketIdx = 0; SocketIdx < AllSockets.Num(); ++SocketIdx)
 		{
@@ -2538,9 +2577,9 @@ void USkinnedMeshComponent::QuerySupportedSockets(TArray<FComponentSocketDescrip
 		}
 
 		// Now grab the bones, which can behave exactly like sockets
-		for (int32 BoneIdx = 0; BoneIdx < GetSkeletalMesh()->GetRefSkeleton().GetNum(); ++BoneIdx)
+		for (int32 BoneIdx = 0; BoneIdx < GetSkinnedAsset()->GetRefSkeleton().GetNum(); ++BoneIdx)
 		{
-			const FName BoneName = GetSkeletalMesh()->GetRefSkeleton().GetBoneName(BoneIdx);
+			const FName BoneName = GetSkinnedAsset()->GetRefSkeleton().GetBoneName(BoneIdx);
 			new (OutSockets) FComponentSocketDescription(BoneName, EComponentSocketType::Bone);
 		}
 	}
@@ -2566,9 +2605,9 @@ bool USkinnedMeshComponent::GetMaterialPropertyPath(int32 ElementIndex, UObject*
 		}
 		return true;
 	}
-	if (GetSkeletalMesh() && GetSkeletalMesh()->GetMaterials().IsValidIndex(ElementIndex))
+	if (GetSkinnedAsset() && GetSkinnedAsset()->GetMaterials().IsValidIndex(ElementIndex))
 	{
-		OutOwner = GetSkeletalMesh();
+		OutOwner = GetSkinnedAsset();
 		OutPropertyPath = FString::Printf(TEXT("%s[%d].%s"), *USkeletalMesh::GetMaterialsMemberName().ToString(), ElementIndex, GET_MEMBER_NAME_STRING_CHECKED(FSkeletalMaterial, MaterialInterface));
 		OutProperty = FSkeletalMaterial::StaticStruct()->FindPropertyByName(GET_MEMBER_NAME_CHECKED(FSkeletalMaterial, MaterialInterface));
 		return true;
@@ -2612,7 +2651,7 @@ void USkinnedMeshComponent::TransformFromBoneSpace(FName BoneName, FVector InPos
 
 FName USkinnedMeshComponent::FindClosestBone(FVector TestLocation, FVector* BoneLocation, float IgnoreScale, bool bRequirePhysicsAsset) const
 {
-	if (GetSkeletalMesh() == NULL)
+	if (GetSkinnedAsset() == NULL)
 	{
 		if (BoneLocation != NULL)
 		{
@@ -2649,7 +2688,7 @@ FName USkinnedMeshComponent::FindClosestBone(FVector TestLocation, FVector* Bone
 			bool bPassPACheck = !bRequirePhysicsAsset;
 			if (bRequirePhysicsAsset)
 			{
-				FName BoneName = GetSkeletalMesh()->GetRefSkeleton().GetBoneName(i);
+				FName BoneName = GetSkinnedAsset()->GetRefSkeleton().GetBoneName(i);
 				bPassPACheck = (PhysAsset->BodySetupIndexMap.Find(BoneName) != nullptr);
 			}
 
@@ -2679,7 +2718,7 @@ FName USkinnedMeshComponent::FindClosestBone(FVector TestLocation, FVector* Bone
 			{
 				*BoneLocation = (CompSpaceTransforms[BestIndex] * GetComponentTransform()).GetLocation();
 			}
-			return GetSkeletalMesh()->GetRefSkeleton().GetBoneName(BestIndex);
+			return GetSkinnedAsset()->GetRefSkeleton().GetBoneName(BestIndex);
 		}
 	}
 }
@@ -2692,7 +2731,7 @@ FName USkinnedMeshComponent::FindClosestBone_K2(FVector TestLocation, FVector& B
 
 void USkinnedMeshComponent::ShowMaterialSection(int32 MaterialID, int32 SectionIndex, bool bShow, int32 LODIndex)
 {
-	if (!GetSkeletalMesh())
+	if (!GetSkinnedAsset())
 	{
 		// no skeletalmesh, then nothing to do. 
 		return;
@@ -2701,16 +2740,16 @@ void USkinnedMeshComponent::ShowMaterialSection(int32 MaterialID, int32 SectionI
 	InitLODInfos();
 	if (LODInfo.IsValidIndex(LODIndex))
 	{
-		const FSkeletalMeshLODInfo& SkelLODInfo = *GetSkeletalMesh()->GetLODInfo(LODIndex);
+		const FSkeletalMeshLODInfo& SkelLODInfo = *GetSkinnedAsset()->GetLODInfo(LODIndex);
 		FSkelMeshComponentLODInfo& SkelCompLODInfo = LODInfo[LODIndex];
 		TArray<bool>& HiddenMaterials = SkelCompLODInfo.HiddenMaterials;
 	
 		// allocate if not allocated yet
-		if ( HiddenMaterials.Num() != GetSkeletalMesh()->GetMaterials().Num() )
+		if ( HiddenMaterials.Num() != GetSkinnedAsset()->GetMaterials().Num() )
 		{
-			// Using skeletalmesh component because Materials.Num() should be <= GetSkeletalMesh()->GetMaterials().Num()		
-			HiddenMaterials.Empty(GetSkeletalMesh()->GetMaterials().Num());
-			HiddenMaterials.AddZeroed(GetSkeletalMesh()->GetMaterials().Num());
+			// Using skeletalmesh component because Materials.Num() should be <= GetSkinnedAsset()->GetMaterials().Num()		
+			HiddenMaterials.Empty(GetSkinnedAsset()->GetMaterials().Num());
+			HiddenMaterials.AddZeroed(GetSkinnedAsset()->GetMaterials().Num());
 		}
 		// If we have a valid LODInfo LODMaterialMap, route material index through it.
 		int32 UseMaterialIndex = MaterialID;			
@@ -2786,10 +2825,10 @@ bool USkinnedMeshComponent::IsMaterialSectionShown(int32 MaterialID, int32 LODIn
 
 void USkinnedMeshComponent::GetUsedMaterials( TArray<UMaterialInterface*>& OutMaterials, bool bGetDebugMaterials ) const
 {
-	if(GetSkeletalMesh())
+	if(GetSkinnedAsset())
 	{
 		// The max number of materials used is the max of the materials on the skeletal mesh and the materials on the mesh component
-		const int32 NumMaterials = FMath::Max(GetSkeletalMesh()->GetMaterials().Num(), OverrideMaterials.Num() );
+		const int32 NumMaterials = FMath::Max(GetSkinnedAsset()->GetMaterials().Num(), OverrideMaterials.Num() );
 		for( int32 MatIdx = 0; MatIdx < NumMaterials; ++MatIdx )
 		{
 			// GetMaterial will determine the correct material to use for this index.  
@@ -2818,11 +2857,11 @@ FSkinWeightVertexBuffer* USkinnedMeshComponent::GetSkinWeightBuffer(int32 LODInd
 {
 	FSkinWeightVertexBuffer* WeightBuffer = nullptr;
 
-	if (GetSkeletalMesh() &&
-		GetSkeletalMesh()->GetResourceForRendering() && 
-		GetSkeletalMesh()->GetResourceForRendering()->LODRenderData.IsValidIndex(LODIndex) )
+	if (GetSkinnedAsset() &&
+		GetSkinnedAsset()->GetResourceForRendering() && 
+		GetSkinnedAsset()->GetResourceForRendering()->LODRenderData.IsValidIndex(LODIndex) )
 	{
-		FSkeletalMeshLODRenderData& LODData = GetSkeletalMesh()->GetResourceForRendering()->LODRenderData[LODIndex];
+		FSkeletalMeshLODRenderData& LODData = GetSkinnedAsset()->GetResourceForRendering()->LODRenderData[LODIndex];
 
 		// Grab weight buffer (check for override)
 		if (LODInfo.IsValidIndex(LODIndex) &&
@@ -2876,17 +2915,17 @@ FVector3f USkinnedMeshComponent::GetSkinnedVertexPosition(USkinnedMeshComponent*
 
 void USkinnedMeshComponent::SetRefPoseOverride(const TArray<FTransform>& NewRefPoseTransforms)
 {
-	if (!GetSkeletalMesh())
+	if (!GetSkinnedAsset())
 	{
 		UE_LOG(LogSkeletalMesh, Warning, TEXT("SetRefPoseOverride (%s) : Not valid without SkeletalMesh assigned."), *GetName());
 		return;
 	}
 
-	const int32 NumRealBones = GetSkeletalMesh()->GetRefSkeleton().GetRawBoneNum();
+	const int32 NumRealBones = GetSkinnedAsset()->GetRefSkeleton().GetRawBoneNum();
 
 	if (NumRealBones != NewRefPoseTransforms.Num())
 	{
-		UE_LOG(LogSkeletalMesh, Warning, TEXT("SetRefPoseOverride (%s) : Expected %d transforms, got %d."), *GetSkeletalMesh()->GetName(), NumRealBones, NewRefPoseTransforms.Num());
+		UE_LOG(LogSkeletalMesh, Warning, TEXT("SetRefPoseOverride (%s) : Expected %d transforms, got %d."), *GetSkinnedAsset()->GetName(), NumRealBones, NewRefPoseTransforms.Num());
 		return;
 	}
 
@@ -2917,7 +2956,7 @@ void USkinnedMeshComponent::SetRefPoseOverride(const TArray<FTransform>& NewRefP
 		// Construct mesh-space skeletal hierarchy.
 		if (BoneIndex > 0)
 		{
-			int32 ParentIndex = GetSkeletalMesh()->GetRefSkeleton().GetRawParentIndex(BoneIndex);
+			int32 ParentIndex = GetSkinnedAsset()->GetRefSkeleton().GetRawParentIndex(BoneIndex);
 			CachedComposedRefPoseMatrices[BoneIndex] = CachedComposedRefPoseMatrices[BoneIndex] * CachedComposedRefPoseMatrices[ParentIndex];
 		}
 
@@ -2929,7 +2968,7 @@ void USkinnedMeshComponent::SetRefPoseOverride(const TArray<FTransform>& NewRefP
 			ZAxis.IsNearlyZero(UE_SMALL_NUMBER))
 		{
 			// this is not allowed, warn them 
-			UE_LOG(LogSkeletalMesh, Warning, TEXT("Reference Pose for asset %s for joint (%s) includes NIL matrix. Zero scale isn't allowed on ref pose. "), *GetSkeletalMesh()->GetPathName(), *GetSkeletalMesh()->GetRefSkeleton().GetBoneName(BoneIndex).ToString());
+			UE_LOG(LogSkeletalMesh, Warning, TEXT("Reference Pose for asset %s for joint (%s) includes NIL matrix. Zero scale isn't allowed on ref pose. "), *GetSkinnedAsset()->GetPathName(), *GetSkinnedAsset()->GetRefSkeleton().GetBoneName(BoneIndex).ToString());
 		}
 
 		// Precompute inverse so we can use from-refpose-skin vertices.
@@ -2949,7 +2988,7 @@ void USkinnedMeshComponent::ClearRefPoseOverride()
 void USkinnedMeshComponent::CacheRefToLocalMatrices(TArray<FMatrix44f>& OutRefToLocal)const
 {
 	const USkinnedMeshComponent* BaseComponent = GetBaseComponent();
-	OutRefToLocal.SetNumUninitialized(GetSkeletalMesh()->GetRefBasesInvMatrix().Num());
+	OutRefToLocal.SetNumUninitialized(GetSkinnedAsset()->GetRefBasesInvMatrix().Num());
 	const TArray<FTransform>& CompSpaceTransforms = BaseComponent->GetComponentSpaceTransforms();
 	if(CompSpaceTransforms.Num())
 	{
@@ -2957,7 +2996,7 @@ void USkinnedMeshComponent::CacheRefToLocalMatrices(TArray<FMatrix44f>& OutRefTo
 
 		for (int32 MatrixIdx = 0; MatrixIdx < OutRefToLocal.Num(); ++MatrixIdx)
 		{
-			OutRefToLocal[MatrixIdx] = GetSkeletalMesh()->GetRefBasesInvMatrix()[MatrixIdx] * FMatrix44f(CompSpaceTransforms[MatrixIdx].ToMatrixWithScale());
+			OutRefToLocal[MatrixIdx] = GetSkinnedAsset()->GetRefBasesInvMatrix()[MatrixIdx] * FMatrix44f(CompSpaceTransforms[MatrixIdx].ToMatrixWithScale());
 		}
 	}
 	else
@@ -2965,7 +3004,7 @@ void USkinnedMeshComponent::CacheRefToLocalMatrices(TArray<FMatrix44f>& OutRefTo
 		//Possible in some cases to request this before the component space transforms are prepared (undo/redo)
 		for (int32 MatrixIdx = 0; MatrixIdx < OutRefToLocal.Num(); ++MatrixIdx)
 		{
-			OutRefToLocal[MatrixIdx] = GetSkeletalMesh()->GetRefBasesInvMatrix()[MatrixIdx];
+			OutRefToLocal[MatrixIdx] = GetSkinnedAsset()->GetRefBasesInvMatrix()[MatrixIdx];
 		}
 	}
 }
@@ -2975,7 +3014,7 @@ void USkinnedMeshComponent::ComputeSkinnedPositions(USkinnedMeshComponent* Compo
 	OutPositions.Empty();
 
 	// Fail if no mesh
-	if (!Component || !Component->GetSkeletalMesh())
+	if (!Component || !Component->GetSkinnedAsset())
 	{
 		return;
 	}
@@ -3002,7 +3041,7 @@ FColor USkinnedMeshComponent::GetVertexColor(int32 VertexIndex) const
 {
 	// Fail if no mesh or no color vertex buffer.
 	FColor FallbackColor = FColor(255, 255, 255, 255);
-	if (!GetSkeletalMesh() || !MeshObject)
+	if (!GetSkinnedAsset() || !MeshObject)
 	{
 		return FallbackColor;
 	}
@@ -3040,7 +3079,7 @@ FVector2D USkinnedMeshComponent::GetVertexUV(int32 VertexIndex, uint32 UVChannel
 {
 	// Fail if no mesh or no vertex buffer.
 	FVector2D FallbackUV = FVector2D::ZeroVector;
-	if (!GetSkeletalMesh() || !MeshObject)
+	if (!GetSkinnedAsset() || !MeshObject)
 	{
 		return FallbackUV;
 	}
@@ -3183,7 +3222,7 @@ void USkinnedMeshComponent::SetMinLOD(int32 InNewMinLOD)
 
 int32 USkinnedMeshComponent::ComputeMinLOD() const
 {
-	int32 MinLodIndex = bOverrideMinLod ? MinLodModel : GetSkeletalMesh()->GetMinLodIdx();
+	int32 MinLodIndex = bOverrideMinLod ? MinLodModel : GetSkinnedAsset()->GetMinLodIdx();
 	int32 NumLODs = GetNumLODs();
 	// want to make sure MinLOD stays within the valid range
 	MinLodIndex = FMath::Min(MinLodIndex, NumLODs - 1);
@@ -3228,7 +3267,7 @@ void USkinnedMeshComponent::SetCapsuleIndirectShadowMinVisibility(float NewValue
 // @todo: think about consolidating this with UpdateLODStatus_Internal
 int32 USkinnedMeshComponent::GetDesiredSyncLOD() const
 {
-	if (GetSkeletalMesh() && MeshObject)
+	if (GetSkinnedAsset() && MeshObject)
 	{
 #if WITH_EDITOR
 		const int32 LODBias = GetLODBias();
@@ -3277,7 +3316,7 @@ bool USkinnedMeshComponent::UpdateLODStatus_Internal(int32 InLeaderPoseComponent
 	const int32 OldPredictedLODLevel = GetPredictedLODLevel();
 	int32 NewPredictedLODLevel = OldPredictedLODLevel;
 
-	if (GetSkeletalMesh() != nullptr)
+	if (GetSkinnedAsset() != nullptr)
 	{
 #if WITH_EDITOR
 		const int32 LODBias = GetLODBias();
@@ -3334,7 +3373,7 @@ bool USkinnedMeshComponent::UpdateLODStatus_Internal(int32 InLeaderPoseComponent
 			}
 		}
 
-		if (GetSkeletalMesh()->IsStreamable() && MeshObject)
+		if (GetSkinnedAsset()->IsStreamable() && MeshObject)
 		{
 			NewPredictedLODLevel = FMath::Max<int32>(NewPredictedLODLevel, MeshObject->GetSkeletalMeshRenderData().PendingFirstLODIdx);
 		}
@@ -3343,7 +3382,7 @@ bool USkinnedMeshComponent::UpdateLODStatus_Internal(int32 InLeaderPoseComponent
 		if (CVarAnimVisualizeLODs.GetValueOnAnyThread() != 0)
 		{
 			// Reduce to visible animated, non SyncAttachParentLOD to reduce clutter.
-			if (GetSkeletalMesh() && MeshObject && bRecentlyRendered)
+			if (GetSkinnedAsset() && MeshObject && bRecentlyRendered)
 			{
 				const bool bHasValidSyncAttachParent = bSyncAttachParentLOD && GetAttachParent() && GetAttachParent()->IsA(USkinnedMeshComponent::StaticClass());
 				if (!bHasValidSyncAttachParent)
@@ -3520,18 +3559,18 @@ void USkinnedMeshComponent::ReleaseResources()
 
 void USkinnedMeshComponent::RegisterLODStreamingCallback(FLODStreamingCallback&& Callback, int32 LODIdx, float TimeoutSecs, bool bOnStreamIn)
 {
-	if (GetSkeletalMesh())
+	if (GetSkinnedAsset())
 	{
-		GetSkeletalMesh()->RegisterMipLevelChangeCallback(this, LODIdx, TimeoutSecs, bOnStreamIn, MoveTemp(Callback));
+		GetSkinnedAsset()->RegisterMipLevelChangeCallback(this, LODIdx, TimeoutSecs, bOnStreamIn, MoveTemp(Callback));
 		bMipLevelCallbackRegistered = true;
 	}
 }
 
 void USkinnedMeshComponent::BeginDestroy()
 {
-	if (GetSkeletalMesh() && bMipLevelCallbackRegistered)
+	if (GetSkinnedAsset() && bMipLevelCallbackRegistered)
 	{
-		GetSkeletalMesh()->RemoveMipLevelChangeCallback(this);
+		GetSkinnedAsset()->RemoveMipLevelChangeCallback(this);
 		bMipLevelCallbackRegistered = false;
 	}
 
@@ -3898,7 +3937,7 @@ void USkinnedMeshComponent::SetSkinWeightOverride(int32 LODIndex, const TArray<F
 			Info.OverrideSkinWeights->SetMaxBoneInfluences(NumBoneInfluences);
 			Info.OverrideSkinWeights->SetUse16BitBoneIndex(bUse16BitBoneIndex);
 
-			const FReferenceSkeleton& RefSkel = GetSkeletalMesh()->GetRefSkeleton();
+			const FReferenceSkeleton& RefSkel = GetSkinnedAsset()->GetRefSkeleton();
 			TArray<FSkinWeightInfo> GPUWeights;
 			CreateSkinWeightsArray(SkinWeights, LODData, GPUWeights, RefSkel);
 			*(Info.OverrideSkinWeights) = GPUWeights;
@@ -4007,7 +4046,7 @@ bool USkinnedMeshComponent::SetSkinWeightProfile(FName InProfileName)
 				// Put in a skin weight profile request
 				if (FSkinWeightProfileManager* Manager = FSkinWeightProfileManager::Get(GetWorld()))
 				{
-					Manager->RequestSkinWeightProfile(InProfileName, GetSkeletalMesh(), this, Callback);
+					Manager->RequestSkinWeightProfile(InProfileName, GetSkinnedAsset(), this, Callback);
 					bSkinWeightProfilePending = true;
 				}
 			}
@@ -4138,16 +4177,11 @@ void USkinnedMeshComponent::SetRenderStatic(bool bNewValue)
 	}
 }
 
-USkinnedAsset* USkinnedMeshComponent::GetSkinnedAsset() const
-{
-	return SkeletalMesh;
-}
-
 #if WITH_EDITOR
 
 bool USkinnedMeshComponent::IsCompiling() const
 {
-	return GetSkeletalMesh() && GetSkeletalMesh()->IsCompiling();
+	return GetSkinnedAsset() && GetSkinnedAsset()->IsCompiling();
 }
 
 void USkinnedMeshComponent::BindWorldDelegates()
@@ -4362,7 +4396,7 @@ FVector3f GetTypedSkinnedVertexPosition(
 #endif
 	{
 		const int32 BoneMapIndex = SkinWeightVertexBuffer.GetBoneIndex(BufferVertIndex, InfluenceIndex);
-		if (!ensureMsgf(Section.BoneMap.IsValidIndex(BoneMapIndex), TEXT("%s has attempted to access a BoneMap of size %i with an invalid index of %i in GetTypedSkinnedVertexPosition()"), *SkinnedComp->GetSkeletalMesh()->GetFullName(), Section.BoneMap.Num(), BoneMapIndex))
+		if (!ensureMsgf(Section.BoneMap.IsValidIndex(BoneMapIndex), TEXT("%s has attempted to access a BoneMap of size %i with an invalid index of %i in GetTypedSkinnedVertexPosition()"), *SkinnedComp->GetSkinnedAsset()->GetFullName(), Section.BoneMap.Num(), BoneMapIndex))
 		{
 			continue;
 		}
@@ -4373,7 +4407,7 @@ FVector3f GetTypedSkinnedVertexPosition(
 		if (LeaderPoseComponentInst)
 		{
 			const TArray<int32>& LeaderBoneMap = SkinnedComp->GetLeaderBoneMap();
-			check(LeaderBoneMap.Num() == SkinnedComp->GetSkeletalMesh()->GetRefSkeleton().GetNum());
+			check(LeaderBoneMap.Num() == SkinnedComp->GetSkinnedAsset()->GetRefSkeleton().GetNum());
 			TransformBoneIndex = LeaderBoneMap[MeshBoneIndex];
 		}
 
@@ -4381,7 +4415,7 @@ FVector3f GetTypedSkinnedVertexPosition(
 		{
 			if (bCachedMatrices)
 			{
-				if (ensureMsgf(RefToLocals.IsValidIndex(MeshBoneIndex), TEXT("%s has attempted to access a RefToLocals of size %i with an invalid index of %i in GetTypedSkinnedVertexPosition()"), *SkinnedComp->GetSkeletalMesh()->GetFullName(), RefToLocals.Num(), MeshBoneIndex))
+				if (ensureMsgf(RefToLocals.IsValidIndex(MeshBoneIndex), TEXT("%s has attempted to access a RefToLocals of size %i with an invalid index of %i in GetTypedSkinnedVertexPosition()"), *SkinnedComp->GetSkinnedAsset()->GetFullName(), RefToLocals.Num(), MeshBoneIndex))
 				{
 					const FMatrix44f& RefToLocal = RefToLocals[MeshBoneIndex];
 					SkinnedPos += RefToLocal.TransformPosition(PositionVertexBuffer.VertexPosition(BufferVertIndex)) * Weight;
@@ -4390,7 +4424,7 @@ FVector3f GetTypedSkinnedVertexPosition(
 			else
 			{
 				const FMatrix44f BoneTransformMatrix = (TransformBoneIndex != INDEX_NONE) ? (FMatrix44f)BaseComponent->GetComponentSpaceTransforms()[TransformBoneIndex].ToMatrixWithScale() : FMatrix44f::Identity;
-				const FMatrix44f RefToLocal = SkinnedComp->GetSkeletalMesh()->GetRefBasesInvMatrix()[MeshBoneIndex] * BoneTransformMatrix;
+				const FMatrix44f RefToLocal = SkinnedComp->GetSkinnedAsset()->GetRefBasesInvMatrix()[MeshBoneIndex] * BoneTransformMatrix;
 				SkinnedPos += RefToLocal.TransformPosition(PositionVertexBuffer.VertexPosition(BufferVertIndex)) * Weight;
 			}
 		}
diff --git a/Engine/Source/Runtime/Engine/Private/ObjectCacheContext.cpp b/Engine/Source/Runtime/Engine/Private/ObjectCacheContext.cpp
--- a/Engine/Source/Runtime/Engine/Private/ObjectCacheContext.cpp
+++ b/Engine/Source/Runtime/Engine/Private/ObjectCacheContext.cpp
@@ -574,7 +574,10 @@ TObjectCacheIterator<USkinnedMeshComponent> FObjectCacheContext::GetSkinnedMeshC
 		TempMap.Reserve(8192);
 		for (USkinnedMeshComponent* Component : GetSkinnedMeshComponents())
 		{
-			TempMap.FindOrAdd(ToRawPtr(Component->GetSkeletalMesh())).Add(Component);
+			if (USkeletalMeshComponent* const SkeletalMeshComponent = Cast<USkeletalMeshComponent>(Component))
+			{
+				TempMap.FindOrAdd(SkeletalMeshComponent->GetSkeletalMesh()).Add(Component);
+			}
 		}
 		SkeletalMeshToComponents = MoveTemp(TempMap);
 	}
diff --git a/Engine/Source/Runtime/Engine/Private/SkeletalMesh.cpp b/Engine/Source/Runtime/Engine/Private/SkeletalMesh.cpp
--- a/Engine/Source/Runtime/Engine/Private/SkeletalMesh.cpp
+++ b/Engine/Source/Runtime/Engine/Private/SkeletalMesh.cpp
@@ -342,12 +342,12 @@ void FScopedSkeletalMeshPostEditChange::SetSkeletalMesh(USkeletalMesh* InSkeleta
 			RecreateExistingRenderStateContext = new FSkinnedMeshComponentRecreateRenderStateContext(InSkeletalMesh, false);
 
 			// Now iterate over all skeletal mesh components and unregister them from the world, we will reregister them in the destructor
-			for (TObjectIterator<USkinnedMeshComponent> It; It; ++It)
+			for (TObjectIterator<USkeletalMeshComponent> It; It; ++It)
 			{
-				USkinnedMeshComponent* SkinComp = *It;
-				if (SkinComp->SkeletalMesh == SkeletalMesh)
+				USkeletalMeshComponent* SkelComp = *It;
+				if (SkelComp->GetSkeletalMesh() == SkeletalMesh)
 				{
-					ComponentReregisterContexts.Add(new FComponentReregisterContext(SkinComp));
+					ComponentReregisterContexts.Add(new FComponentReregisterContext(SkelComp));
 				}
 			}
 		}
@@ -1205,7 +1205,7 @@ void RefreshSkelMeshOnPhysicsAssetChange(const USkeletalMesh* InSkeletalMesh)
 		{
 			USkeletalMeshComponent* SkeletalMeshComponent = Cast<USkeletalMeshComponent>(*Iter);
 			// if PhysicsAssetOverride is NULL, it uses SkeletalMesh Physics Asset, so I'll need to update here
-			if  (SkeletalMeshComponent->SkeletalMesh == InSkeletalMesh &&
+			if  (SkeletalMeshComponent->GetSkeletalMesh() == InSkeletalMesh &&
 				 SkeletalMeshComponent->PhysicsAssetOverride == NULL)
 			{
 				// it needs to recreate IF it already has been created
@@ -1333,7 +1333,7 @@ void USkeletalMesh::PostEditChangeProperty(FPropertyChangedEvent& PropertyChange
 		for(TObjectIterator<USkeletalMeshComponent> It; It; ++It)
 		{
 			USkeletalMeshComponent* MeshComponent = *It;
-			if(MeshComponent && !MeshComponent->IsTemplate() && MeshComponent->SkeletalMesh == this)
+			if(MeshComponent && !MeshComponent->IsTemplate() && MeshComponent->GetSkeletalMesh() == this)
 			{
 				ComponentsToReregister.Add(*It);
 			}
@@ -1391,12 +1391,12 @@ void USkeletalMesh::PostEditUndo()
 	check(IsInGameThread());
 
 	Super::PostEditUndo();
-	for( TObjectIterator<USkinnedMeshComponent> It; It; ++It )
+	for( TObjectIterator<USkeletalMeshComponent> It; It; ++It )
 	{
-		USkinnedMeshComponent* MeshComponent = *It;
+		USkeletalMeshComponent* MeshComponent = *It;
 		if( MeshComponent && 
 			!MeshComponent->IsTemplate() &&
-			MeshComponent->SkeletalMesh == this )
+			MeshComponent->GetSkeletalMesh() == this )
 		{
 			FComponentReregisterContext Context(MeshComponent);
 		}
@@ -3457,7 +3457,7 @@ void USkeletalMesh::InitMorphTargetsAndRebuildRenderData()
 		// reset all morphtarget for all components
 		for (TObjectIterator<USkeletalMeshComponent> It; It; ++It)
 		{
-			if (It->SkeletalMesh == this)
+			if (It->GetSkeletalMesh() == this)
 			{
 				It->RefreshMorphTargets();
 			}
@@ -6816,9 +6816,9 @@ FSkinnedMeshComponentRecreateRenderStateContext::FSkinnedMeshComponentRecreateRe
 	: bRefreshBounds(InRefreshBounds)
 {
 	bool bRequireFlushRenderingCommands = false;
-	for (TObjectIterator<USkinnedMeshComponent> It; It; ++It)
+	for (TObjectIterator<USkeletalMeshComponent> It; It; ++It)
 	{
-		if (It->SkeletalMesh == InSkeletalMesh)
+		if (It->GetSkeletalMesh() == InSkeletalMesh)
 		{
 			checkf(!It->IsUnreachable(), TEXT("%s"), *It->GetFullName());
 
diff --git a/Engine/Source/Runtime/Engine/Private/SkeletalRender.cpp b/Engine/Source/Runtime/Engine/Private/SkeletalRender.cpp
--- a/Engine/Source/Runtime/Engine/Private/SkeletalRender.cpp
+++ b/Engine/Source/Runtime/Engine/Private/SkeletalRender.cpp
@@ -41,12 +41,12 @@ FSkeletalMeshObject::FSkeletalMeshObject(USkinnedMeshComponent* InMeshComponent,
 ,	WorkingMaxDistanceFactor(0.f)
 ,   bHasBeenUpdatedAtLeastOnce(false)
 #if RHI_RAYTRACING
-, bSupportRayTracing(IsSkeletalMeshRayTracingSupported() && InMeshComponent->SkeletalMesh->GetSupportRayTracing())
+, bSupportRayTracing(IsSkeletalMeshRayTracingSupported() && InMeshComponent->GetSkinnedAsset()->GetSupportRayTracing())
 , bHiddenMaterialVisibilityDirtyForRayTracing(false)
-, RayTracingMinLOD(InMeshComponent->SkeletalMesh->GetRayTracingMinLOD())
+, RayTracingMinLOD(InMeshComponent->GetSkinnedAsset()->GetRayTracingMinLOD())
 #endif
 #if !UE_BUILD_SHIPPING
-, DebugName(InMeshComponent->SkeletalMesh->GetFName())
+, DebugName(InMeshComponent->GetSkinnedAsset()->GetFName())
 #endif // !UE_BUILD_SHIPPING
 #if WITH_EDITORONLY_DATA
 ,   SectionIndexPreview(InMeshComponent->GetSectionPreview())
@@ -55,12 +55,12 @@ FSkeletalMeshObject::FSkeletalMeshObject(USkinnedMeshComponent* InMeshComponent,
 ,	SelectedEditorMaterial(InMeshComponent->GetSelectedEditorMaterial())
 #endif	
 ,	SkeletalMeshRenderData(InSkelMeshRenderData)
-,	SkeletalMeshLODInfo(InMeshComponent->SkeletalMesh->GetLODInfoArray())
+,	SkeletalMeshLODInfo(InMeshComponent->GetSkinnedAsset()->GetLODInfoArray())
 ,	SkinCacheEntry(nullptr)
 ,	SkinCacheEntryForRayTracing(nullptr)
 ,	LastFrameNumber(0)
 ,	bUsePerBoneMotionBlur(InMeshComponent->bPerBoneMotionBlur)
-,	StatId(InMeshComponent->SkeletalMesh->GetStatID(true))
+,	StatId(InMeshComponent->GetSkinnedAsset()->GetStatID(true))
 ,	FeatureLevel(InFeatureLevel)
 ,	ComponentId(InMeshComponent->ComponentId.PrimIDValue)
 {
@@ -342,17 +342,17 @@ void UpdateRefToLocalMatricesInner(TArray<FMatrix44f>& ReferenceToLocal, const T
  */
 void UpdateRefToLocalMatrices( TArray<FMatrix44f>& ReferenceToLocal, const USkinnedMeshComponent* InMeshComponent, const FSkeletalMeshRenderData* InSkeletalMeshRenderData, int32 LODIndex, const TArray<FBoneIndexType>* ExtraRequiredBoneIndices )
 {
-	const USkeletalMesh* const ThisMesh = InMeshComponent->SkeletalMesh;
+	const USkinnedAsset* const SkinnedAsset = InMeshComponent->GetSkinnedAsset();
 	const USkinnedMeshComponent* const LeaderComp = InMeshComponent->LeaderPoseComponent.Get();
 	const FSkeletalMeshLODRenderData& LOD = InSkeletalMeshRenderData->LODRenderData[LODIndex];
 
-	const FReferenceSkeleton& RefSkeleton = ThisMesh->GetRefSkeleton();
+	const FReferenceSkeleton& RefSkeleton = SkinnedAsset->GetRefSkeleton();
 	const TArray<int32>& LeaderBoneMap = InMeshComponent->GetLeaderBoneMap();
-	const bool bIsLeaderCompValid = LeaderComp && LeaderBoneMap.Num() == ThisMesh->GetRefSkeleton().GetNum();
+	const bool bIsLeaderCompValid = LeaderComp && LeaderBoneMap.Num() == SkinnedAsset->GetRefSkeleton().GetNum();
 	const TArray<FTransform>& ComponentTransform = (bIsLeaderCompValid)? LeaderComp->GetComponentSpaceTransforms() : InMeshComponent->GetComponentSpaceTransforms();
 	const TArray<uint8>& BoneVisibilityStates = (bIsLeaderCompValid) ? LeaderComp->GetBoneVisibilityStates() : InMeshComponent->GetBoneVisibilityStates();
 	// Get inv ref pose matrices
-	const TArray<FMatrix44f>* RefBasesInvMatrix = &ThisMesh->GetRefBasesInvMatrix();
+	const TArray<FMatrix44f>* RefBasesInvMatrix = &SkinnedAsset->GetRefBasesInvMatrix();
 
 	// Check if there is an override (and it's the right size)
 	if( InMeshComponent->GetRefPoseOverride() && 
@@ -378,7 +378,7 @@ void UpdateRefToLocalMatrices( TArray<FMatrix44f>& ReferenceToLocal, const USkin
 	{
 		UE_LOG(LogSkeletalMesh, Error,
 			TEXT("Mesh %s : Invalid LODIndex [count %d, index %d], streaming[Ready(%d), F(%d), P(%d)], \
-			ExtraRequiredBoneIndices is (%d), and total number is (%d)"), *GetNameSafe(ThisMesh),
+			ExtraRequiredBoneIndices is (%d), and total number is (%d)"), *GetNameSafe(SkinnedAsset),
 			InSkeletalMeshRenderData->LODRenderData.Num(), LODIndex, InSkeletalMeshRenderData->bReadyForStreaming,
 			InSkeletalMeshRenderData->CurrentFirstLODIdx, InSkeletalMeshRenderData->PendingFirstLODIdx,
 			(ExtraRequiredBoneIndices) ? 1 : 0, (ExtraRequiredBoneIndices) ? ExtraRequiredBoneIndices->Num() : 0);
@@ -396,7 +396,7 @@ void UpdateRefToLocalMatrices( TArray<FMatrix44f>& ReferenceToLocal, const USkin
 
 void UpdatePreviousRefToLocalMatrices(TArray<FMatrix44f>& ReferenceToLocal, const USkinnedMeshComponent* InMeshComponent, const FSkeletalMeshRenderData* InSkeletalMeshRenderData, int32 LODIndex, const TArray<FBoneIndexType>* ExtraRequiredBoneIndices)
 {
-	const USkeletalMesh* const ThisMesh = InMeshComponent->SkeletalMesh;
+	const USkinnedAsset* const ThisMesh = InMeshComponent->GetSkinnedAsset();
 	const USkinnedMeshComponent* const LeaderComp = InMeshComponent->LeaderPoseComponent.Get();
 	const FSkeletalMeshLODRenderData& LOD = InSkeletalMeshRenderData->LODRenderData[LODIndex];
 
