Inclusive terminology change for Avoidance Manager.

diff --git a/Engine/Source/Runtime/Engine/Classes/AI/Navigation/AvoidanceManager.h b/Engine/Source/Runtime/Engine/Classes/AI/Navigation/AvoidanceManager.h
--- a/Engine/Source/Runtime/Engine/Classes/AI/Navigation/AvoidanceManager.h
+++ b/Engine/Source/Runtime/Engine/Classes/AI/Navigation/AvoidanceManager.h
@@ -225,7 +225,7 @@ protected:
 	TArray<int32> DebugUIDs;
 	bool bDebugAll;
 	
-	/** master switch for avoidance system */
+	/** main switch for avoidance system */
 	static bool bSystemActive;
 #endif
 };
