CIS - Fix non-PCH build error member access into incomplete type 'USkeletalMeshComponent'

diff --git a/Engine/Source/Runtime/Engine/Private/ObjectCacheContext.cpp b/Engine/Source/Runtime/Engine/Private/ObjectCacheContext.cpp
--- a/Engine/Source/Runtime/Engine/Private/ObjectCacheContext.cpp
+++ b/Engine/Source/Runtime/Engine/Private/ObjectCacheContext.cpp
@@ -566,24 +566,26 @@ TObjectCacheIterator<UMaterialInterface> FObjectCacheContext::GetMaterialsAffect
 
 TObjectCacheIterator<USkinnedMeshComponent> FObjectCacheContext::GetSkinnedMeshComponents(USkeletalMesh* InSkeletalMesh)
 {
-	if (!SkeletalMeshToComponents.IsSet())
+	return GetSkinnedMeshComponents(Cast<USkinnedAsset>(InSkeletalMesh));
+}
+
+TObjectCacheIterator<USkinnedMeshComponent> FObjectCacheContext::GetSkinnedMeshComponents(USkinnedAsset* InSkinnedAsset)
+{
+	if (!SkinnedAssetToComponents.IsSet())
 	{
 		TRACE_CPUPROFILER_EVENT_SCOPE(ComputeSkeletalMeshToComponents);
 
-		TMap<TObjectKey<USkeletalMesh>, TSet<USkinnedMeshComponent*>> TempMap;
+		TMap<TObjectKey<USkinnedAsset>, TSet<USkinnedMeshComponent*>> TempMap;
 		TempMap.Reserve(8192);
 		for (USkinnedMeshComponent* Component : GetSkinnedMeshComponents())
 		{
-			if (USkeletalMeshComponent* const SkeletalMeshComponent = Cast<USkeletalMeshComponent>(Component))
-			{
-				TempMap.FindOrAdd(SkeletalMeshComponent->GetSkeletalMesh()).Add(Component);
-			}
+			TempMap.FindOrAdd(Component->GetSkinnedAsset()).Add(Component);
 		}
-		SkeletalMeshToComponents = MoveTemp(TempMap);
+		SkinnedAssetToComponents = MoveTemp(TempMap);
 	}
 
 	static TSet<USkinnedMeshComponent*> EmptySet;
-	TSet<USkinnedMeshComponent*>* Set = SkeletalMeshToComponents.GetValue().Find(InSkeletalMesh);
+	TSet<USkinnedMeshComponent*>* Set = SkinnedAssetToComponents.GetValue().Find(InSkinnedAsset);
 	const TSet<USkinnedMeshComponent*>& Result = Set ? *Set : EmptySet;
 
 	return TObjectCacheIterator<USkinnedMeshComponent>(Result.Array());
diff --git a/Engine/Source/Runtime/Engine/Public/ObjectCacheContext.h b/Engine/Source/Runtime/Engine/Public/ObjectCacheContext.h
--- a/Engine/Source/Runtime/Engine/Public/ObjectCacheContext.h
+++ b/Engine/Source/Runtime/Engine/Public/ObjectCacheContext.h
@@ -16,6 +16,7 @@ class UMaterialInterface;
 class UStaticMesh;
 class UTexture;
 class USkeletalMesh;
+class USkinnedAsset;
 
 /**
  *  Class to abstract implementation details of the containers used inside the
@@ -68,6 +69,8 @@ public:
 	TObjectCacheIterator<UPrimitiveComponent>   GetPrimitiveComponents();
 	TObjectCacheIterator<UStaticMeshComponent>  GetStaticMeshComponents();
 	TObjectCacheIterator<USkinnedMeshComponent> GetSkinnedMeshComponents();
+	TObjectCacheIterator<USkinnedMeshComponent> GetSkinnedMeshComponents(USkinnedAsset* InSkinnedAsset);
+	UE_DEPRECATED(5.1, "Use GetSkinnedMeshComponents(USkinnedAsset* InSkinnedAsset) instead")
 	TObjectCacheIterator<USkinnedMeshComponent> GetSkinnedMeshComponents(USkeletalMesh* InSkeletalMesh);
 	TObjectCacheIterator<UStaticMeshComponent>  GetStaticMeshComponents(UStaticMesh* InStaticMesh);
 	TObjectCacheIterator<UMaterialInterface>    GetMaterialsAffectedByTexture(UTexture* InTexture);
@@ -84,7 +87,7 @@ private:
 	TOptional<TMap<TObjectKey<UTexture>, TSet<UMaterialInterface*>>> TextureToMaterials;
 	TOptional<TMap<TObjectKey<UMaterialInterface>, TSet<UPrimitiveComponent*>>> MaterialToPrimitives;
 	TOptional<TMap<TObjectKey<UStaticMesh>, TSet<UStaticMeshComponent*>>> StaticMeshToComponents;
-	TOptional<TMap<TObjectKey<USkeletalMesh>, TSet<USkinnedMeshComponent*>>> SkeletalMeshToComponents;
+	TOptional<TMap<TObjectKey<USkinnedAsset>, TSet<USkinnedMeshComponent*>>> SkinnedAssetToComponents;
 	TOptional<TArray<USkinnedMeshComponent*>> SkinnedMeshComponents;
 	TOptional<TArray<UStaticMeshComponent*>> StaticMeshComponents;
 	TOptional<TArray<UPrimitiveComponent*>> PrimitiveComponents;
