Fix SubsystemCollection ownership:
- Refactor FSubsystemCollectionBase so it doesn't subclass FGCObject anymore
- FSubsystemCollection now subclasses FGCObject (to preserve backward compatibility and to be used outside of UObject owners)
- FObjectSubsystemCollection now used by UWorld/UEditorEngine/UEngine/ULocalPlayer/UGameInstance to control lifetime of their Subsystems
- Remove usage of TUniqueObj as FObjectSubsystemCollection no longer inherits FGCObject.

diff --git a/Engine/Source/Editor/UnrealEd/Classes/Editor/EditorEngine.h b/Engine/Source/Editor/UnrealEd/Classes/Editor/EditorEngine.h
--- a/Engine/Source/Editor/UnrealEd/Classes/Editor/EditorEngine.h
+++ b/Engine/Source/Editor/UnrealEd/Classes/Editor/EditorEngine.h
@@ -3264,7 +3264,7 @@ public:
 	UEditorSubsystem* GetEditorSubsystemBase(TSubclassOf<UEditorSubsystem> SubsystemClass) const
 	{
 		checkSlow(this != nullptr);
-		return EditorSubsystemCollection->GetSubsystem<UEditorSubsystem>(SubsystemClass);
+		return EditorSubsystemCollection.GetSubsystem<UEditorSubsystem>(SubsystemClass);
 	}
 
 	/**
@@ -3274,7 +3274,7 @@ public:
 	TSubsystemClass* GetEditorSubsystem() const
 	{
 		checkSlow(this != nullptr);
-		return EditorSubsystemCollection->GetSubsystem<TSubsystemClass>(TSubsystemClass::StaticClass());
+		return EditorSubsystemCollection.GetSubsystem<TSubsystemClass>(TSubsystemClass::StaticClass());
 	}
 
 	/**
@@ -3285,15 +3285,11 @@ public:
 	template <typename TSubsystemClass>
 	const TArray<TSubsystemClass*>& GetEditorSubsystemArray() const
 	{
-		return EditorSubsystemCollection->GetSubsystemArray<TSubsystemClass>(TSubsystemClass::StaticClass());
+		return EditorSubsystemCollection.GetSubsystemArray<TSubsystemClass>(TSubsystemClass::StaticClass());
 	}
 
 private:
-	// TUniqueObj is used here to work around a hot reload issue caused by FSubsystemCollection inheriting FGCObject.
-	// When hot reload occurs, the CDO for this type can be reconstructed over the same object at the same address without
-	// destroying it first, which breaks FGCObject.
-	// TUniquePtr makes sure the object is allocated on the heap, giving it a unique address.
-	TUniqueObj<FSubsystemCollection<UEditorSubsystem>> EditorSubsystemCollection;
+	FObjectSubsystemCollection<UEditorSubsystem> EditorSubsystemCollection;
 
 	// DEPRECATED VARIABLES ONLY
 public:
diff --git a/Engine/Source/Editor/UnrealEd/Private/EditorEngine.cpp b/Engine/Source/Editor/UnrealEd/Private/EditorEngine.cpp
--- a/Engine/Source/Editor/UnrealEd/Private/EditorEngine.cpp
+++ b/Engine/Source/Editor/UnrealEd/Private/EditorEngine.cpp
@@ -901,7 +901,7 @@ void UEditorEngine::HandleSettingChanged( FName Name )
 
 void UEditorEngine::InitializeObjectReferences()
 {
-	EditorSubsystemCollection->Initialize(this);
+	EditorSubsystemCollection.Initialize(this);
 
 	Super::InitializeObjectReferences();
 
@@ -1257,6 +1257,8 @@ void UEditorEngine::PreExit()
 				World->CleanupWorld();
 			}
 		}
+
+		EditorSubsystemCollection.Deinitialize();
 	}
 
 	Super::PreExit();
@@ -1278,8 +1280,6 @@ void UEditorEngine::FinishDestroy()
 			ToolMenus->UnregisterStringCommandHandler("Command");
 		}
 
-		EditorSubsystemCollection->Deinitialize();
-
 		// Unregister events
 		FEditorDelegates::MapChange.RemoveAll(this);
 		FCoreDelegates::ModalErrorMessage.Unbind();
@@ -1354,6 +1354,9 @@ void UEditorEngine::AddReferencedObjects(UObject* InThis, FReferenceCollector& C
 		Collector.AddReferencedObject(This->PlayInEditorSessionInfo->OriginalRequestParams.EditorPlaySettings, This);
 	}
 
+	// Keep Editor subsystems alive
+	This->EditorSubsystemCollection.AddReferencedObjects(This, Collector);
+
 	Super::AddReferencedObjects( This, Collector );
 }
 
diff --git a/Engine/Source/Editor/UnrealEd/Private/EditorModeManager.cpp b/Engine/Source/Editor/UnrealEd/Private/EditorModeManager.cpp
--- a/Engine/Source/Editor/UnrealEd/Private/EditorModeManager.cpp
+++ b/Engine/Source/Editor/UnrealEd/Private/EditorModeManager.cpp
@@ -725,7 +725,10 @@ void FEditorModeTools::RemoveAllDelegateHandlers()
 	if (GEditor)
 	{
 		GEditor->UnregisterForUndo(this);
-		GEditor->GetEditorSubsystem<UAssetEditorSubsystem>()->OnEditorModeUnregistered().RemoveAll(this);
+		if (UAssetEditorSubsystem* AssetEditorSubsystem = GEditor->GetEditorSubsystem<UAssetEditorSubsystem>())
+		{
+			AssetEditorSubsystem->OnEditorModeUnregistered().RemoveAll(this);
+		}
 	}
 
 	FWorldDelegates::OnWorldCleanup.RemoveAll(this);
diff --git a/Engine/Source/Editor/UnrealEd/Private/EditorModeRegistry.cpp b/Engine/Source/Editor/UnrealEd/Private/EditorModeRegistry.cpp
--- a/Engine/Source/Editor/UnrealEd/Private/EditorModeRegistry.cpp
+++ b/Engine/Source/Editor/UnrealEd/Private/EditorModeRegistry.cpp
@@ -128,9 +128,14 @@ void FEditorModeRegistry::UnregisterMode(FEditorModeID ModeID)
 	// First off delete the factory
 	if (ModeFactories.Remove(ModeID) > 0 && bInitialized)
 	{
-		UAssetEditorSubsystem* AssetEditorSubsystem = GEditor->GetEditorSubsystem<UAssetEditorSubsystem>();
-		AssetEditorSubsystem->OnEditorModeUnregistered().Broadcast(ModeID);
-		AssetEditorSubsystem->OnEditorModesChanged().Broadcast();
+		if (GEditor)
+		{
+			if (UAssetEditorSubsystem* AssetEditorSubsystem = GEditor->GetEditorSubsystem<UAssetEditorSubsystem>())
+			{
+				AssetEditorSubsystem->OnEditorModeUnregistered().Broadcast(ModeID);
+				AssetEditorSubsystem->OnEditorModesChanged().Broadcast();
+			}
+		}
 	}
 }
 
diff --git a/Engine/Source/Runtime/Engine/Classes/Engine/Engine.h b/Engine/Source/Runtime/Engine/Classes/Engine/Engine.h
--- a/Engine/Source/Runtime/Engine/Classes/Engine/Engine.h
+++ b/Engine/Source/Runtime/Engine/Classes/Engine/Engine.h
@@ -3460,7 +3460,7 @@ public:
 	UEngineSubsystem* GetEngineSubsystemBase(TSubclassOf<UEngineSubsystem> SubsystemClass) const
 	{
 		checkSlow(this != nullptr);
-		return EngineSubsystemCollection->GetSubsystem<UEngineSubsystem>(SubsystemClass);
+		return EngineSubsystemCollection.GetSubsystem<UEngineSubsystem>(SubsystemClass);
 	}
 
 	/**
@@ -3470,7 +3470,7 @@ public:
 	TSubsystemClass* GetEngineSubsystem() const
 	{
 		checkSlow(this != nullptr);
-		return EngineSubsystemCollection->GetSubsystem<TSubsystemClass>(TSubsystemClass::StaticClass());
+		return EngineSubsystemCollection.GetSubsystem<TSubsystemClass>(TSubsystemClass::StaticClass());
 	}
 
 	/**
@@ -3479,15 +3479,11 @@ public:
 	template <typename TSubsystemClass>
 	const TArray<TSubsystemClass*>& GetEngineSubsystemArray() const
 	{
-		return EngineSubsystemCollection->GetSubsystemArray<TSubsystemClass>(TSubsystemClass::StaticClass());
+		return EngineSubsystemCollection.GetSubsystemArray<TSubsystemClass>(TSubsystemClass::StaticClass());
 	}
 
 private:
-	// TUniqueObj is used here to work around a hot reload issue caused by FSubsystemCollection inheriting FGCObject.
-	// When hot reload occurs, the CDO for this type can be reconstructed over the same object at the same address without
-	// destroying it first, which breaks FGCObject.
-	// TUniquePtr makes sure the object is allocated on the heap, giving it a unique address.
-	TUniqueObj<FSubsystemCollection<UEngineSubsystem>> EngineSubsystemCollection;
+	FObjectSubsystemCollection<UEngineSubsystem> EngineSubsystemCollection;
 
 public:
 	/**
diff --git a/Engine/Source/Runtime/Engine/Classes/Engine/GameInstance.h b/Engine/Source/Runtime/Engine/Classes/Engine/GameInstance.h
--- a/Engine/Source/Runtime/Engine/Classes/Engine/GameInstance.h
+++ b/Engine/Source/Runtime/Engine/Classes/Engine/GameInstance.h
@@ -218,6 +218,7 @@ public:
 	//~ Begin UObject Interface
 	virtual class UWorld* GetWorld() const final;
 	virtual void FinishDestroy() override;
+	static void AddReferencedObjects(UObject* InThis, FReferenceCollector& Collector);
 	//~ End UObject Interface
 
 	/** virtual function to allow custom GameInstances an opportunity to set up what it needs */
@@ -643,5 +644,5 @@ protected:
 
 private:
 
-	FSubsystemCollection<UGameInstanceSubsystem> SubsystemCollection;
+	FObjectSubsystemCollection<UGameInstanceSubsystem> SubsystemCollection;
 };
\ No newline at end of file
diff --git a/Engine/Source/Runtime/Engine/Classes/Engine/LocalPlayer.h b/Engine/Source/Runtime/Engine/Classes/Engine/LocalPlayer.h
--- a/Engine/Source/Runtime/Engine/Classes/Engine/LocalPlayer.h
+++ b/Engine/Source/Runtime/Engine/Classes/Engine/LocalPlayer.h
@@ -223,7 +223,7 @@ private:
 	/** Event called when platform user id changes */
 	FOnPlatformUserIdChanged OnPlatformUserIdChangedEvent;
 
-	FSubsystemCollection<ULocalPlayerSubsystem> SubsystemCollection;
+	FObjectSubsystemCollection<ULocalPlayerSubsystem> SubsystemCollection;
 
 public:
 	// UObject interface
diff --git a/Engine/Source/Runtime/Engine/Classes/Engine/World.h b/Engine/Source/Runtime/Engine/Classes/Engine/World.h
--- a/Engine/Source/Runtime/Engine/Classes/Engine/World.h
+++ b/Engine/Source/Runtime/Engine/Classes/Engine/World.h
@@ -1035,15 +1035,6 @@ private:
 	static void BroadcastWorldPartitionUninitialized(UWorld* InWorld, UWorldPartition* InWorldPartition);
 };
 
-/**
- * World specific Subsystem collection which delegates object references to its owning UWorld
- */
-class FWorldSubsystemCollection : public FSubsystemCollection<UWorldSubsystem>
-{
-public:
-	virtual bool IsGCObjectReferencer() const override { return false; }
-};
-
 /** 
  * The World is the top level object representing a map or a sandbox in which Actors and Components will exist and be rendered.  
  *
@@ -4170,7 +4161,7 @@ public:
 	FWorldPSCPool PSCPool;
 
 	//PSC Pooling END
-	FWorldSubsystemCollection SubsystemCollection;
+	FObjectSubsystemCollection<UWorldSubsystem> SubsystemCollection;
 };
 
 /** Global UWorld pointer. Use of this pointer should be avoided whenever possible. */
diff --git a/Engine/Source/Runtime/Engine/Private/GameInstance.cpp b/Engine/Source/Runtime/Engine/Private/GameInstance.cpp
--- a/Engine/Source/Runtime/Engine/Private/GameInstance.cpp
+++ b/Engine/Source/Runtime/Engine/Private/GameInstance.cpp
@@ -64,6 +64,15 @@ void UGameInstance::FinishDestroy()
 	Super::FinishDestroy();
 }
 
+void UGameInstance::AddReferencedObjects(UObject* InThis, FReferenceCollector& Collector)
+{
+	UGameInstance* This = CastChecked<UGameInstance>(InThis);
+		
+	This->SubsystemCollection.AddReferencedObjects(This, Collector);
+
+	UObject::AddReferencedObjects(This, Collector);
+}
+
 UWorld* UGameInstance::GetWorld() const
 {
 	return WorldContext ? WorldContext->World() : NULL;
diff --git a/Engine/Source/Runtime/Engine/Private/LocalPlayer.cpp b/Engine/Source/Runtime/Engine/Private/LocalPlayer.cpp
--- a/Engine/Source/Runtime/Engine/Private/LocalPlayer.cpp
+++ b/Engine/Source/Runtime/Engine/Private/LocalPlayer.cpp
@@ -1731,6 +1731,8 @@ void ULocalPlayer::AddReferencedObjects(UObject* InThis, FReferenceCollector& Co
 		}
 	}
 
+	This->SubsystemCollection.AddReferencedObjects(This, Collector);
+
 	UPlayer::AddReferencedObjects(This, Collector);
 }
 
diff --git a/Engine/Source/Runtime/Engine/Private/Subsystems/SubsystemCollection.cpp b/Engine/Source/Runtime/Engine/Private/Subsystems/SubsystemCollection.cpp
--- a/Engine/Source/Runtime/Engine/Private/Subsystems/SubsystemCollection.cpp
+++ b/Engine/Source/Runtime/Engine/Private/Subsystems/SubsystemCollection.cpp
@@ -186,30 +186,17 @@ void FSubsystemCollectionBase::Deinitialize()
 	}
 
 	// Deinit and clean up existing systems
+	SubsystemArrayMap.Empty();
+	for (auto Iter = SubsystemMap.CreateIterator(); Iter; ++Iter)
 	{
-		// @note Iterate over SubsystemMapWeak to test whether subsystems have already been 
-		// deleted which can happen when Deinitialize() is called from the destructor
-		TMap<UClass*, TWeakObjectPtr<USubsystem>> SubsystemMapWeakCopy;
-		Swap(SubsystemMapWeakCopy, SubsystemMapWeak);
-		for (const TPair<UClass*, TWeakObjectPtr<USubsystem>>& Iter : SubsystemMapWeakCopy)
+		UClass* KeyClass = Iter.Key();
+		USubsystem* Subsystem = Iter.Value();
+		if (Subsystem != nullptr && Subsystem->GetClass() == KeyClass)
 		{
-			UClass* KeyClass = Iter.Key;
-			USubsystem* Subsystem = Iter.Value.GetEvenIfUnreachable();
-			if (Subsystem)
-			{
-				if ((Subsystem->GetClass() == KeyClass) && (Subsystem->InternalOwningSubsystem != nullptr))
-				{
-					Subsystem->Deinitialize();
-					Subsystem->InternalOwningSubsystem = nullptr;
-				}
-			}
-			else
-			{
-				UE_LOG(LogSubsystemCollection, Log, TEXT("Cannot deinitialize %s because it has been deleted"), *KeyClass->GetName());
-			}
+			Subsystem->Deinitialize();
+			Subsystem->InternalOwningSubsystem = nullptr;
 		}
 	}
-	SubsystemArrayMap.Empty();
 	SubsystemMap.Empty();
 	Outer = nullptr;
 }
@@ -230,24 +217,11 @@ USubsystem* FSubsystemCollectionBase::InitializeDependency(TSubclassOf<USubsyste
 	return Subsystem;
 }
 
-void FSubsystemCollectionBase::AddReferencedObjects(FReferenceCollector& Collector)
-{
-	if (IsGCObjectReferencer())
-	{
-		AddReferencedObjects(nullptr, Collector);
-	}
-}
-
 void FSubsystemCollectionBase::AddReferencedObjects(UObject* Referencer, FReferenceCollector& Collector)
 {
 	Collector.AddReferencedObjects(SubsystemMap, Referencer);
 }
 
-FString FSubsystemCollectionBase::GetReferencerName() const
-{
-	return TEXT("FSubsystemCollectionBase");
-}
-
 USubsystem* FSubsystemCollectionBase::AddAndInitializeSubsystem(UClass* SubsystemClass)
 {
 	TGuardValue<bool> PopulatingGuard(bPopulating, true);
@@ -271,7 +245,6 @@ USubsystem* FSubsystemCollectionBase::AddAndInitializeSubsystem(UClass* Subsyste
 			{
 				USubsystem* Subsystem = NewObject<USubsystem>(Outer, SubsystemClass);
 				SubsystemMap.Add(SubsystemClass,Subsystem);
-				SubsystemMapWeak.Add(SubsystemClass, Subsystem);
 				Subsystem->InternalOwningSubsystem = this;
 				Subsystem->Initialize(*this);
 				return Subsystem;
@@ -293,7 +266,6 @@ void FSubsystemCollectionBase::RemoveAndDeinitializeSubsystem(USubsystem* Subsys
 	check(Subsystem == SubsystemFound);
 
 	const UClass* SubsystemClass = Subsystem->GetClass();
-	SubsystemMapWeak.Remove(SubsystemClass);
 
 	for (auto& Pair : SubsystemArrayMap)
 	{
diff --git a/Engine/Source/Runtime/Engine/Private/UnrealEngine.cpp b/Engine/Source/Runtime/Engine/Private/UnrealEngine.cpp
--- a/Engine/Source/Runtime/Engine/Private/UnrealEngine.cpp
+++ b/Engine/Source/Runtime/Engine/Private/UnrealEngine.cpp
@@ -1831,7 +1831,7 @@ void UEngine::Init(IEngineLoop* InEngineLoop)
 
 	// Subsystems.
 	FURL::StaticInit();
-	EngineSubsystemCollection->Initialize(this);
+	EngineSubsystemCollection.Initialize(this);
 
 #if !UE_BUILD_SHIPPING
 	// Check for overrides to the default map on the command line
@@ -2199,7 +2199,7 @@ void UEngine::PreExit()
 	NextDynamicResolutionState.Reset();
 #endif
 
-	EngineSubsystemCollection->Deinitialize();
+	EngineSubsystemCollection.Deinitialize();
 }
 
 void UEngine::ShutdownHMD()
@@ -3365,6 +3365,9 @@ void UEngine::AddReferencedObjects(UObject* InThis, FReferenceCollector& Collect
 		Context.AddReferencedObjects(Collector, This);
 	}
 
+	// Keep Engine subsystems alive
+	This->EngineSubsystemCollection.AddReferencedObjects(This, Collector);
+
 	Super::AddReferencedObjects(This, Collector);
 }
 
diff --git a/Engine/Source/Runtime/Engine/Public/Subsystems/SubsystemCollection.h b/Engine/Source/Runtime/Engine/Public/Subsystems/SubsystemCollection.h
--- a/Engine/Source/Runtime/Engine/Public/Subsystems/SubsystemCollection.h
+++ b/Engine/Source/Runtime/Engine/Public/Subsystems/SubsystemCollection.h
@@ -11,7 +11,7 @@ class UDynamicSubsystem;
 
 DECLARE_LOG_CATEGORY_EXTERN(LogSubsystemCollection, Log, All);
 
-class ENGINE_API FSubsystemCollectionBase : public FGCObject
+class ENGINE_API FSubsystemCollectionBase
 {
 public:
 	/** Initialize the collection of systems, systems will be created and initialized */
@@ -39,10 +39,6 @@ public:
 		return Cast<TSubsystemClass>(InitializeDependency(TSubsystemClass::StaticClass()));
 	}
 
-	/* FGCObject Interface */
-	virtual void AddReferencedObjects(FReferenceCollector& Collector) override;
-	virtual FString GetReferencerName() const override;
-
 	/** Registers and adds instances of the specified Subsystem class to all existing SubsystemCollections of the correct type.
 	 *  Should be used by specific subsystems in plug ins when plugin is activated.
 	 */
@@ -53,12 +49,8 @@ public:
 	 */
 	static void DeactivateExternalSubsystem(UClass* SubsystemClass);
 
-	/** Add referenced objects with referencer */
+	/** Collect references held by this collection */
 	void AddReferencedObjects(UObject* Referencer, FReferenceCollector& Collector);
-
-	/** If True, this collection will add its references through its FGCObject::AddReferencedObjects interface implementation */
-	virtual bool IsGCObjectReferencer() const { return true; }
-
 protected:
 	/** protected constructor - for use by the template only(FSubsystemCollection<TBaseType>) */
 	FSubsystemCollectionBase(UClass* InBaseType);
@@ -87,8 +79,6 @@ private:
 
 	TMap<UClass*, USubsystem*> SubsystemMap;
 
-	TMap<UClass*, TWeakObjectPtr<USubsystem>> SubsystemMapWeak;
-
 	mutable TMap<UClass*, TArray<USubsystem*>> SubsystemArrayMap;
 
 	UClass* BaseType;
@@ -108,7 +98,7 @@ private:
 };
 
 template<typename TBaseType>
-class FSubsystemCollection : public FSubsystemCollectionBase
+class FSubsystemCollection : public FSubsystemCollectionBase, public FGCObject
 {
 public:
 	/** Get a Subsystem by type */
@@ -133,6 +123,17 @@ public:
 		return *SpecificArray;
 	}
 
+	/* FGCObject Interface */
+	virtual void AddReferencedObjects(FReferenceCollector& Collector) override
+	{
+		FSubsystemCollectionBase::AddReferencedObjects(nullptr, Collector);
+	}
+
+	virtual FString GetReferencerName() const override
+	{
+		return TEXT("FSubsystemCollection");
+	}
+	
 public:
 
 	/** Construct a FSubsystemCollection, pass in the owning object almost certainly (this). */
@@ -142,3 +143,39 @@ public:
 	}
 };
 
+/** Subsystem collection which delegates UObject references to its owning UObject (object needs to implement AddReferencedObjects and forward call to Collection */
+template<typename TBaseType>
+class FObjectSubsystemCollection : public FSubsystemCollectionBase
+{
+public:
+	/** Get a Subsystem by type */
+	template <typename TSubsystemClass>
+	TSubsystemClass* GetSubsystem(const TSubclassOf<TSubsystemClass>& SubsystemClass) const
+	{
+		static_assert(TIsDerivedFrom<TSubsystemClass, TBaseType>::IsDerived, "TSubsystemClass must be derived from TBaseType");
+
+		// A static cast is safe here because we know SubsystemClass derives from TSubsystemClass if it is not null
+		return static_cast<TSubsystemClass*>(GetSubsystemInternal(SubsystemClass));
+	}
+
+	/** Get a list of Subsystems by type */
+	template <typename TSubsystemClass>
+	const TArray<TSubsystemClass*>& GetSubsystemArray(const TSubclassOf<TSubsystemClass>& SubsystemClass) const
+	{
+		// Force a compile time check that TSubsystemClass derives from TBaseType, the internal code only enforces it's a USubsystem
+		TSubclassOf<TBaseType> SubsystemBaseClass = SubsystemClass;
+
+		const TArray<USubsystem*>& Array = GetSubsystemArrayInternal(SubsystemBaseClass);
+		const TArray<TSubsystemClass*>* SpecificArray = reinterpret_cast<const TArray<TSubsystemClass*>*>(&Array);
+		return *SpecificArray;
+	}
+
+public:
+
+	/** Construct a FSubsystemCollection, pass in the owning object almost certainly (this). */
+	FObjectSubsystemCollection()
+		: FSubsystemCollectionBase(TBaseType::StaticClass())
+	{
+	}
+};
+
