If a scene component is being attached by replication and it's already in the Parent's AttachChildren, don't add it at the LastAttachIndex to prevent possibly overwritting an unmapped reference in the array.

diff --git a/Engine/Source/Runtime/Engine/Private/Components/SceneComponent.cpp b/Engine/Source/Runtime/Engine/Private/Components/SceneComponent.cpp
--- a/Engine/Source/Runtime/Engine/Private/Components/SceneComponent.cpp
+++ b/Engine/Source/Runtime/Engine/Private/Components/SceneComponent.cpp
@@ -2053,7 +2053,8 @@ bool USceneComponent::AttachToComponent(USceneComponent* Parent, const FAttachme
 		OnAttachmentChanged();
 
 		// Preserve order of previous attachment if valid (in case we're doing a reattach operation inside a loop that might assume the AttachChildren order won't change)
-		if(LastAttachIndex != INDEX_NONE)
+		// Don't do this if updating attachment from replication to avoid overwriting addresses in AttachChildren that may be unmapped
+		if(LastAttachIndex != INDEX_NONE && !bNetUpdateAttachment)
 		{
 			Parent->AttachChildren.Insert(this, LastAttachIndex);
 		}
