Support StaticMeshRenderData >2GB.

diff --git a/Engine/Source/Runtime/Engine/Private/StaticMesh.cpp b/Engine/Source/Runtime/Engine/Private/StaticMesh.cpp
--- a/Engine/Source/Runtime/Engine/Private/StaticMesh.cpp
+++ b/Engine/Source/Runtime/Engine/Private/StaticMesh.cpp
@@ -2879,7 +2879,7 @@ void FStaticMeshRenderData::Cache(const ITargetPlatform* TargetPlatform, UStatic
 		{
 			COOK_STAT(Timer.AddHit(MeshDataBuffer.GetSize()));
 
-			FMemoryReaderView Ar(MakeArrayView((const uint8*)MeshDataBuffer.GetData(), MeshDataBuffer.GetSize()), /*bIsPersistent=*/ true);
+			FMemoryReaderView Ar(MakeMemoryView((const uint8*)MeshDataBuffer.GetData(), MeshDataBuffer.GetSize()), /*bIsPersistent=*/ true);
 			Serialize(Ar, Owner, /*bCooked=*/ false);
 
 			// Reconstruct EstimatedCompressedSize
