Use FVecto2Dr::GetAbs so (-x, y) or (x, -y) return (x, y)
Fix crash occuring when all actors are in the same quadrant with both positive and negative coordinate (i.e (-x, y) or (x, -y))

diff --git a/Engine/Source/Runtime/Engine/Private/WorldPartition/RuntimeSpatialHash/RuntimeSpatialHashGridHelper.cpp b/Engine/Source/Runtime/Engine/Private/WorldPartition/RuntimeSpatialHash/RuntimeSpatialHashGridHelper.cpp
--- a/Engine/Source/Runtime/Engine/Private/WorldPartition/RuntimeSpatialHash/RuntimeSpatialHashGridHelper.cpp
+++ b/Engine/Source/Runtime/Engine/Private/WorldPartition/RuntimeSpatialHash/RuntimeSpatialHashGridHelper.cpp
@@ -41,8 +41,8 @@ FSquare2DGridHelper::FSquare2DGridHelper(const FBox& InWorldBounds, const FVecto
 
 	if (WorldBounds.IsValid)
 	{
-		const FVector2D DistMin = FMath::Abs(FVector2D(WorldBounds.Min - Origin));
-		const FVector2D DistMax = FMath::Abs(FVector2D(WorldBounds.Max - Origin));
+		const FVector2D DistMin = FVector2D(WorldBounds.Min - Origin).GetAbs();
+		const FVector2D DistMax = FVector2D(WorldBounds.Max - Origin).GetAbs();
 		const double WorldBoundsMaxExtent = FMath::Max(DistMin.GetMax(), DistMax.GetMax());
 
 		if (WorldBoundsMaxExtent > 0)
