Fixed serialization issue where bEvaluateWorldPositionOffsetInRayTracing is defaulted to true incorrectly

Submitting on behalf of Juan Canada.

diff --git a/Engine/Source/Runtime/Engine/Private/Components/StaticMeshComponent.cpp b/Engine/Source/Runtime/Engine/Private/Components/StaticMeshComponent.cpp
--- a/Engine/Source/Runtime/Engine/Private/Components/StaticMeshComponent.cpp
+++ b/Engine/Source/Runtime/Engine/Private/Components/StaticMeshComponent.cpp
@@ -310,12 +310,22 @@ void UStaticMeshComponent::Serialize(FArchive& Ar)
 
 	LLM_SCOPE(ELLMTag::StaticMesh);
 
-	Super::Serialize(Ar);
-
 	Ar.UsingCustomVersion(FRenderingObjectVersion::GUID);
 	Ar.UsingCustomVersion(FFortniteNCBranchObjectVersion::GUID);
 	Ar.UsingCustomVersion(FFortniteMainBranchObjectVersion::GUID);
 
+	const bool bBeforeRemappedEvaluateWorldPositionOffset = (Ar.CustomVer(FFortniteNCBranchObjectVersion::GUID) < FFortniteNCBranchObjectVersion::RemappedEvaluateWorldPositionOffsetInRayTracing
+		&& Ar.CustomVer(FFortniteMainBranchObjectVersion::GUID) < FFortniteMainBranchObjectVersion::RemappedEvaluateWorldPositionOffsetInRayTracing);
+
+	// When bEvaluateWorldPositionOffsetInRayTracing was named bEvaluateWorldPositionOffset the default value was false and now it is true. 
+	// Therefore if the default was not set in a blueprint it needs to be changed for old assets before Super::Serialize(Ar);
+	if (Ar.IsLoading() && bBeforeRemappedEvaluateWorldPositionOffset && !GetArchetype()->IsInBlueprint())
+	{
+		bEvaluateWorldPositionOffset = false; 
+	}
+
+	Super::Serialize(Ar);
+
 #if WITH_EDITORONLY_DATA
 	if (Ar.IsCooking())
 	{
@@ -368,8 +378,7 @@ void UStaticMeshComponent::Serialize(FArchive& Ar)
 	}
 #endif
 
-	if (Ar.CustomVer(FFortniteNCBranchObjectVersion::GUID) < FFortniteNCBranchObjectVersion::RemappedEvaluateWorldPositionOffsetInRayTracing &&
-	Ar.CustomVer(FFortniteMainBranchObjectVersion::GUID) < FFortniteMainBranchObjectVersion::RemappedEvaluateWorldPositionOffsetInRayTracing)
+	if (Ar.IsLoading() && bBeforeRemappedEvaluateWorldPositionOffset)
 	{
 		bEvaluateWorldPositionOffsetInRayTracing = bEvaluateWorldPositionOffset;
 		bEvaluateWorldPositionOffset = true; // Default WPO evaluation on
