Added FOutBunch(int64 MaxBits) constructor.

For cases when there's no UPackageMap but you still want to specify the MaxBits.

#rnx

diff --git a/Engine/Source/Runtime/Engine/Private/DataBunch.cpp b/Engine/Source/Runtime/Engine/Private/DataBunch.cpp
--- a/Engine/Source/Runtime/Engine/Private/DataBunch.cpp
+++ b/Engine/Source/Runtime/Engine/Private/DataBunch.cpp
@@ -138,7 +138,7 @@ FOutBunch::FOutBunch( UPackageMap *InPackageMap, int64 MaxBits )
 ,	Channel		( nullptr )
 ,	Time		( 0 )
 ,	ChIndex		( 0 )
-,	ChName		( NAME_None )
+,	ChName		( )
 ,	ChSequence	( 0 )
 ,	PacketId	( 0 )
 ,	ReceivedAck	( 0 )
@@ -155,6 +155,11 @@ FOutBunch::FOutBunch( UPackageMap *InPackageMap, int64 MaxBits )
 {
 }
 
+FOutBunch::FOutBunch(int64 InMaxBits)
+: FOutBunch(static_cast<UPackageMap*>(nullptr), InMaxBits)
+{
+}
+
 void FOutBunch::CountMemory(FArchive& Ar) const
 {
 	for (const FOutBunch* Current = this; Current; Current = Current->Next)
diff --git a/Engine/Source/Runtime/Engine/Public/Net/DataBunch.h b/Engine/Source/Runtime/Engine/Public/Net/DataBunch.h
--- a/Engine/Source/Runtime/Engine/Public/Net/DataBunch.h
+++ b/Engine/Source/Runtime/Engine/Public/Net/DataBunch.h
@@ -69,6 +69,7 @@ public:
 
 	// Functions.
 	FOutBunch();
+	explicit FOutBunch(int64 InMaxBits);
 	FOutBunch( class UChannel* InChannel, bool bClose );
 	FOutBunch( UPackageMap * PackageMap, int64 InMaxBits = 1024 );
 
