[KillPendingKill] Clear referenced to MyParticleEventManager when it's destroyed. This is fairly important since we don't create a new one unless the previous one was null.

diff --git a/Engine/Source/Runtime/Engine/Classes/Particles/ParticleEventManager.h b/Engine/Source/Runtime/Engine/Classes/Particles/ParticleEventManager.h
--- a/Engine/Source/Runtime/Engine/Classes/Particles/ParticleEventManager.h
+++ b/Engine/Source/Runtime/Engine/Classes/Particles/ParticleEventManager.h
@@ -14,6 +14,8 @@ UCLASS(config=Game, NotBlueprintable, notplaceable)
 class AParticleEventManager : public AActor
 {
 	GENERATED_UCLASS_BODY()
+		
+	virtual void EndPlay(const EEndPlayReason::Type EndPlayReason) override;
 
 	virtual void HandleParticleSpawnEvents( UParticleSystemComponent* Component, const TArray<FParticleEventSpawnData>& SpawnEvents );
 	virtual void HandleParticleDeathEvents( UParticleSystemComponent* Component, const TArray<FParticleEventDeathData>& DeathEvents );
diff --git a/Engine/Source/Runtime/Engine/Private/Particles/ParticleEventManager.cpp b/Engine/Source/Runtime/Engine/Private/Particles/ParticleEventManager.cpp
--- a/Engine/Source/Runtime/Engine/Private/Particles/ParticleEventManager.cpp
+++ b/Engine/Source/Runtime/Engine/Private/Particles/ParticleEventManager.cpp
@@ -11,6 +11,21 @@ AParticleEventManager::AParticleEventManager(const FObjectInitializer& ObjectIni
 	SetCanBeDamaged(false);
 }
 
+void AParticleEventManager::EndPlay(const EEndPlayReason::Type EndPlayReason)
+{
+	Super::EndPlay(EndPlayReason);
+	if (EndPlayReason == EEndPlayReason::Destroyed)
+	{
+		if (UWorld* World = GetWorld())
+		{
+			if (World->MyParticleEventManager == this)
+			{
+				World->MyParticleEventManager = nullptr;
+			}
+		}
+	}
+}
+
 void AParticleEventManager::HandleParticleSpawnEvents( UParticleSystemComponent* Component, const TArray<FParticleEventSpawnData>& SpawnEvents )
 {
 	AEmitter* Emitter = Cast<AEmitter>(Component->GetOwner());
