Fixed ParticleEventManager being destroyed when scrubbing through a demo/replay and not being recreated.

diff --git a/Engine/Source/Runtime/Engine/Private/Particles/ParticleEventManager.cpp b/Engine/Source/Runtime/Engine/Private/Particles/ParticleEventManager.cpp
--- a/Engine/Source/Runtime/Engine/Private/Particles/ParticleEventManager.cpp
+++ b/Engine/Source/Runtime/Engine/Private/Particles/ParticleEventManager.cpp
@@ -9,6 +9,7 @@ AParticleEventManager::AParticleEventManager(const FObjectInitializer& ObjectIni
 	: Super(ObjectInitializer)
 {
 	SetCanBeDamaged(false);
+	bReplayRewindable = true;
 }
 
 void AParticleEventManager::EndPlay(const EEndPlayReason::Type EndPlayReason)
