Keeping last update requested velocity for usage in Blueprint

diff --git a/Engine/Source/Runtime/Engine/Classes/GameFramework/CharacterMovementComponent.h b/Engine/Source/Runtime/Engine/Classes/GameFramework/CharacterMovementComponent.h
--- a/Engine/Source/Runtime/Engine/Classes/GameFramework/CharacterMovementComponent.h
+++ b/Engine/Source/Runtime/Engine/Classes/GameFramework/CharacterMovementComponent.h
@@ -1071,6 +1071,17 @@ public:
 	UPROPERTY(Transient)
 	FVector RequestedVelocity;
 
+	/**
+	 * Velocity requested by path following during last Update
+	 * Updated when we consume the value
+	 */
+	UPROPERTY(Transient)
+	FVector LastUpdateRequestedVelocity;
+
+	/** Returns velocity requested by path following */
+	UFUNCTION(BlueprintCallable, Category="Pawn|Components|CharacterMovement", meta=(Keywords="Velocity RequestedVelocity"))
+	FVector GetLastUpdateRequestedVelocity() const;
+
 	/** No default value, for now it's assumed to be valid if GetAvoidanceManager() returns non-NULL. */
 	UPROPERTY(Category="Character Movement: Avoidance", VisibleAnywhere, BlueprintReadOnly, AdvancedDisplay)
 	int32 AvoidanceUID;
diff --git a/Engine/Source/Runtime/Engine/Private/CharacterMovementComponentAsync.cpp b/Engine/Source/Runtime/Engine/Private/CharacterMovementComponentAsync.cpp
--- a/Engine/Source/Runtime/Engine/Private/CharacterMovementComponentAsync.cpp
+++ b/Engine/Source/Runtime/Engine/Private/CharacterMovementComponentAsync.cpp
@@ -397,6 +397,7 @@ void FCharacterMovementComponentAsyncInput::PerformMovement(float DeltaSeconds,
 		}
 
 		// consume path following requested velocity
+		Output.LastUpdateRequestedVelocity = Output.bHasRequestedVelocity ? Output.RequestedVelocity : FVector::ZeroVector;
 		Output.bHasRequestedVelocity = false;
 
 		// TODO OnMOvementUpdated
@@ -4593,6 +4594,7 @@ void FCharacterMovementComponentAsyncOutput::Copy(const FCharacterMovementCompon
 	bHasRequestedVelocity = Value.bHasRequestedVelocity;
 	bRequestedMoveWithMaxSpeed = Value.bRequestedMoveWithMaxSpeed;
 	RequestedVelocity = Value.RequestedVelocity;
+	LastUpdateRequestedVelocity = Value.LastUpdateRequestedVelocity;
 	NumJumpApexAttempts = Value.NumJumpApexAttempts;
 	AnimRootMotionVelocity = Value.AnimRootMotionVelocity;
 	bShouldApplyDeltaToMeshPhysicsTransforms = Value.bShouldApplyDeltaToMeshPhysicsTransforms;
diff --git a/Engine/Source/Runtime/Engine/Private/Components/CharacterMovementComponent.cpp b/Engine/Source/Runtime/Engine/Private/Components/CharacterMovementComponent.cpp
--- a/Engine/Source/Runtime/Engine/Private/Components/CharacterMovementComponent.cpp
+++ b/Engine/Source/Runtime/Engine/Private/Components/CharacterMovementComponent.cpp
@@ -2007,6 +2007,7 @@ void UCharacterMovementComponent::SimulateMovement(float DeltaSeconds)
 		UpdateCharacterStateAfterMovement(DeltaSeconds);
 
 		// consume path following requested velocity
+		LastUpdateRequestedVelocity = bHasRequestedVelocity ? RequestedVelocity : FVector::ZeroVector;
 		bHasRequestedVelocity = false;
 
 		OnMovementUpdated(DeltaSeconds, OldLocation, OldVelocity);
@@ -2598,6 +2599,7 @@ void UCharacterMovementComponent::PerformMovement(float DeltaSeconds)
 		}
 
 		// consume path following requested velocity
+		LastUpdateRequestedVelocity = bHasRequestedVelocity ? RequestedVelocity : FVector::ZeroVector;
 		bHasRequestedVelocity = false;
 
 		OnMovementUpdated(DeltaSeconds, OldLocation, OldVelocity);
@@ -5650,6 +5652,7 @@ void UCharacterMovementComponent::StopActiveMovement()
 	Acceleration = FVector::ZeroVector; 
 	bHasRequestedVelocity = false;
 	RequestedVelocity = FVector::ZeroVector;
+	LastUpdateRequestedVelocity = FVector::ZeroVector;
 }
 
 void UCharacterMovementComponent::ProcessLanded(const FHitResult& Hit, float remainingTime, int32 Iterations)
@@ -10617,6 +10620,11 @@ void UCharacterMovementComponent::CapsuleTouched(UPrimitiveComponent* Overlapped
 	}
 }
 
+FVector UCharacterMovementComponent::GetLastUpdateRequestedVelocity() const
+{
+	return LastUpdateRequestedVelocity;
+}
+
 void UCharacterMovementComponent::SetAvoidanceGroup(int32 GroupFlags)
 {
 	SetAvoidanceGroupMask(GroupFlags);
@@ -12630,6 +12638,7 @@ void UCharacterMovementComponent::FillAsyncInput(const FVector& InputVector, FCh
 		AsyncSimState->bHasRequestedVelocity = bHasRequestedVelocity;
 		AsyncSimState->bRequestedMoveWithMaxSpeed = bRequestedMoveWithMaxSpeed;
 		AsyncSimState->RequestedVelocity = RequestedVelocity;
+		AsyncSimState->LastUpdateRequestedVelocity = LastUpdateRequestedVelocity;
 		AsyncSimState->NumJumpApexAttempts = NumJumpApexAttempts;
 		AsyncSimState->bShouldApplyDeltaToMeshPhysicsTransforms = false;
 		AsyncSimState->DeltaPosition = FVector::ZeroVector;
@@ -12749,6 +12758,7 @@ void UCharacterMovementComponent::ApplyAsyncOutput(FCharacterMovementComponentAs
 	bHasRequestedVelocity = Output.bHasRequestedVelocity;
 	bRequestedMoveWithMaxSpeed = Output.bRequestedMoveWithMaxSpeed;
 	RequestedVelocity = Output.RequestedVelocity;
+	LastUpdateRequestedVelocity = Output.LastUpdateRequestedVelocity;
 
 	if (CharacterOwner)
 	{
diff --git a/Engine/Source/Runtime/Engine/Public/CharacterMovementComponentAsync.h b/Engine/Source/Runtime/Engine/Public/CharacterMovementComponentAsync.h
--- a/Engine/Source/Runtime/Engine/Public/CharacterMovementComponentAsync.h
+++ b/Engine/Source/Runtime/Engine/Public/CharacterMovementComponentAsync.h
@@ -210,6 +210,7 @@ public:
 	bool bHasRequestedVelocity;
 	bool bRequestedMoveWithMaxSpeed;
 	FVector RequestedVelocity;
+	FVector LastUpdateRequestedVelocity;
 	int32 NumJumpApexAttempts;
 	FVector AnimRootMotionVelocity;
 	bool bShouldApplyDeltaToMeshPhysicsTransforms; // See UpdateBasedMovement
