Use new FOutBunch constructor. Minor cleanup.

#rnx

diff --git a/Engine/Source/Runtime/Engine/Private/Net/Experimental/Iris/DataStreamChannel.cpp b/Engine/Source/Runtime/Engine/Private/Net/Experimental/Iris/DataStreamChannel.cpp
--- a/Engine/Source/Runtime/Engine/Private/Net/Experimental/Iris/DataStreamChannel.cpp
+++ b/Engine/Source/Runtime/Engine/Private/Net/Experimental/Iris/DataStreamChannel.cpp
@@ -119,12 +119,12 @@ void UDataStreamChannel::SendOpenBunch()
 	if (!bHandshakeSent && NumOutRec == 0)
 	{
 		// Send dummy data to open the channel
-		FOutBunch OutBunch(static_cast<UPackageMap*>(nullptr), 8U);
+		constexpr int64 MaxBunchBits = 8;
+		FOutBunch OutBunch(MaxBunchBits);
 		OutBunch.ChName = this->ChName;
 		OutBunch.ChIndex = this->ChIndex;
 		OutBunch.Channel = this;
 		OutBunch.Next = nullptr;
-	//		OutBunch.bOpen = !OpenAcked;
 		// Unreliable bunches will be dropped on the receiving side unless the channel is open.
 		OutBunch.bReliable = true;
 
@@ -214,8 +214,8 @@ void UDataStreamChannel::Tick()
 
 		IRIS_PROFILER_SCOPE(UDataStreamChannel_SendBunchAndFlushNet);
 
-		// Can't use FOutBunch constructor taking a UChannel* as we want to limit the amount of data to send.
-		FOutBunch OutBunch(static_cast<UPackageMap*>(nullptr), MaxBytes*8U);
+		const int64 MaxBunchBits = MaxBytes*8;
+		FOutBunch OutBunch(MaxBunchBits);
 #if UE_NET_TRACE_ENABLED
 		SetTraceCollector(OutBunch, Collector);
 #endif
@@ -223,7 +223,6 @@ void UDataStreamChannel::Tick()
 		OutBunch.ChIndex = this->ChIndex;
 		OutBunch.Channel = this;
 		OutBunch.Next = nullptr;
-//		OutBunch.bOpen = !OpenAcked;
 		// Unreliable bunches will be dropped on the receiving side unless the channel is open.
 		OutBunch.bReliable = !OpenAcked;
 		OutBunch.SerializeBits(BitStreamBuffer, BitWriter.GetPosBits());
@@ -240,7 +239,6 @@ void UDataStreamChannel::Tick()
 		{
 			check(PacketIds.First != INDEX_NONE);
 			// Something went wrong.
-			//DataStreamManager->ProcessPacketDeliveryStatus(EPacketDeliveryStatus::Discard, Record);
 			return UDataStream::EWriteResult::NoData;
 		}
 
@@ -298,7 +296,7 @@ void UDataStreamChannel::ReceivedAck(int32 PacketId)
 	}
 
 	const FDataStreamChannelRecord& ChannelRecord = WriteRecords.Peek();
-	checkf((uint32)PacketId == ChannelRecord.PacketId, TEXT("PacketId %d != ChannelRecord.PacketId %d, WriteRecords.Num %d"), PacketId, ChannelRecord.PacketId, (int32)WriteRecords.Count());
+	ensureMsgf((uint32)PacketId == ChannelRecord.PacketId, TEXT("PacketId %d != ChannelRecord.PacketId %d, WriteRecords.Num %d"), PacketId, ChannelRecord.PacketId, (int32)WriteRecords.Count());
 
 	DataStreamManager->ProcessPacketDeliveryStatus(UE::Net::EPacketDeliveryStatus::Delivered, static_cast<const FDataStreamRecord*>(ChannelRecord.Record));
 
diff --git a/Engine/Source/Runtime/Engine/Private/Net/Iris/ReplicationSystem/ActorReplicationBridge.cpp b/Engine/Source/Runtime/Engine/Private/Net/Iris/ReplicationSystem/ActorReplicationBridge.cpp
--- a/Engine/Source/Runtime/Engine/Private/Net/Iris/ReplicationSystem/ActorReplicationBridge.cpp
+++ b/Engine/Source/Runtime/Engine/Private/Net/Iris/ReplicationSystem/ActorReplicationBridge.cpp
@@ -905,7 +905,8 @@ void UActorReplicationBridge::GetActorCreationHeader(const AActor* Actor, UE::Ne
 
 	// Custom actor creation data
 	{
-		FOutBunch Bunch(static_cast<UPackageMap*>(nullptr), 1024);
+		constexpr int64 MaxBunchBits = 1024;
+		FOutBunch Bunch(MaxBunchBits);
 		const_cast<AActor*>(Actor)->OnSerializeNewActor(Bunch);
 		Header.CustomCreationDataBitCount = Bunch.GetNumBits();
 		if (Header.CustomCreationDataBitCount > 0)
