Warn about exceeding platform path limits when dumping shader debug info.

#rnx

diff --git a/Engine/Source/Runtime/Engine/Private/ShaderCompiler/ShaderCompiler.cpp b/Engine/Source/Runtime/Engine/Private/ShaderCompiler/ShaderCompiler.cpp
--- a/Engine/Source/Runtime/Engine/Private/ShaderCompiler/ShaderCompiler.cpp
+++ b/Engine/Source/Runtime/Engine/Private/ShaderCompiler/ShaderCompiler.cpp
@@ -3799,6 +3799,14 @@ FString FShaderCompilingManager::CreateShaderDebugInfoPath(const FShaderCompiler
 	DumpDebugInfoPath.ReplaceInline(TEXT("?"), TEXT("!"));
 	DumpDebugInfoPath.ReplaceInline(TEXT("\""), TEXT("\'"));
 
+	// Return empty string if we exceeded the platform path limitations and emit warning
+	static const int32 MaxLength = FPlatformMisc::GetMaxPathLength();
+	if (DumpDebugInfoPath.Len() > MaxLength)
+	{
+		UE_LOG(LogShaderCompilers, Warning, TEXT("Shader debug path exceeded platform limits: %s"), *DumpDebugInfoPath);
+		return TEXT("");
+	}
+
 	if (!IFileManager::Get().DirectoryExists(*DumpDebugInfoPath))
 	{
 		verifyf(IFileManager::Get().MakeDirectory(*DumpDebugInfoPath, true), TEXT("Failed to create directory for shader debug info '%s'"), *DumpDebugInfoPath);
