Avoid unnecessary inner static variable for max path length. Appendix to CL 21109869.

#rnx

diff --git a/Engine/Source/Runtime/Engine/Private/ShaderCompiler/ShaderCompiler.cpp b/Engine/Source/Runtime/Engine/Private/ShaderCompiler/ShaderCompiler.cpp
--- a/Engine/Source/Runtime/Engine/Private/ShaderCompiler/ShaderCompiler.cpp
+++ b/Engine/Source/Runtime/Engine/Private/ShaderCompiler/ShaderCompiler.cpp
@@ -3800,8 +3800,7 @@ FString FShaderCompilingManager::CreateShaderDebugInfoPath(const FShaderCompiler
 	DumpDebugInfoPath.ReplaceInline(TEXT("\""), TEXT("\'"));
 
 	// Return empty string if we exceeded the platform path limitations and emit warning
-	static const int32 MaxLength = FPlatformMisc::GetMaxPathLength();
-	if (DumpDebugInfoPath.Len() > MaxLength)
+	if (DumpDebugInfoPath.Len() > FPlatformMisc::GetMaxPathLength())
 	{
 		UE_LOG(LogShaderCompilers, Warning, TEXT("Shader debug path exceeded platform limits: %s"), *DumpDebugInfoPath);
 		return TEXT("");
