Fixed rare RPC DoS scoping ensure, where 'PostReceivedRPC' could be fired, without the 'Pre' event (NotifyReceivedRPC/LightweightReceivedRPC) firing.

diff --git a/Engine/Source/Runtime/Engine/Private/DataReplication.cpp b/Engine/Source/Runtime/Engine/Private/DataReplication.cpp
--- a/Engine/Source/Runtime/Engine/Private/DataReplication.cpp
+++ b/Engine/Source/Runtime/Engine/Private/DataReplication.cpp
@@ -1191,10 +1191,9 @@ bool FObjectReplicator::ReceivedBunch(FNetBitReader& Bunch, const FReplicationFl
 bool GReceiveRPCTimingEnabled = false;
 struct FScopedRPCTimingTracker
 {
-	FScopedRPCTimingTracker(UFunction* InFunction, UNetConnection* InConnection, FRPCDoSDetection& InRPCDoS)
+	FScopedRPCTimingTracker(UFunction* InFunction, UNetConnection* InConnection)
 		: Connection(InConnection)
 		, Function(InFunction)
-		, RPCDoS(InRPCDoS)
 	{
 		if (GReceiveRPCTimingEnabled)
 		{
@@ -1206,9 +1205,9 @@ struct FScopedRPCTimingTracker
 
 	~FScopedRPCTimingTracker()
 	{
-		if (RPCDoS.IsRPCDoSDetectionEnabled())
+		if (RPCDoS != nullptr && RPCDoS->IsRPCDoSDetectionEnabled())
 		{
-			RPCDoS.PostReceivedRPC();
+			RPCDoS->PostReceivedRPC();
 		}
 
 		ActiveTrackers.RemoveSingleSwap(this);
@@ -1221,7 +1220,7 @@ struct FScopedRPCTimingTracker
 
 	UNetConnection* Connection;
 	UFunction* Function;
-	FRPCDoSDetection& RPCDoS;
+	FRPCDoSDetection* RPCDoS = nullptr;
 	double StartTime;
 
 	static TArray<FScopedRPCTimingTracker*> ActiveTrackers;
@@ -1252,7 +1251,7 @@ bool FObjectReplicator::ReceivedRPC(FNetBitReader& Reader, const FReplicationFla
 	FName FunctionName = FieldCache->Field.GetFName();
 	UFunction* Function = Object->FindFunction(FunctionName);
 	FRPCDoSDetection& RPCDoS = Connection->GetRPCDoS();
-	FScopedRPCTimingTracker ScopedTracker(Function, Connection, RPCDoS);
+	FScopedRPCTimingTracker ScopedTracker(Function, Connection);
 	SCOPE_CYCLE_COUNTER(STAT_NetReceiveRPC);
 	SCOPE_CYCLE_UOBJECT(Function, Function);
 
@@ -1277,6 +1276,8 @@ bool FObjectReplicator::ReceivedRPC(FNetBitReader& Reader, const FReplicationFla
 
 	if (bIsServer && RPCDoS.IsRPCDoSDetectionEnabled() && !Connection->IsReplay())
 	{
+		ScopedTracker.RPCDoS = &RPCDoS;
+
 		if (UNLIKELY(RPCDoS.ShouldMonitorReceivedRPC()))
 		{
 			ERPCNotifyResult Result = RPCDoS.NotifyReceivedRPC(Reader, UnmappedGuids, Object, Function, FunctionName);
