Level Instance loading optimization: don't call ResetLoaders when loading level instances, delay it if the user edit the level instance instead.

#rnx

diff --git a/Engine/Source/Runtime/Engine/Private/LevelInstance/LevelInstanceLevelStreaming.cpp b/Engine/Source/Runtime/Engine/Private/LevelInstance/LevelInstanceLevelStreaming.cpp
--- a/Engine/Source/Runtime/Engine/Private/LevelInstance/LevelInstanceLevelStreaming.cpp
+++ b/Engine/Source/Runtime/Engine/Private/LevelInstance/LevelInstanceLevelStreaming.cpp
@@ -8,6 +8,7 @@
 #include "Engine/World.h"
 #include "Misc/PackageName.h"
 #include "GameFramework/WorldSettings.h"
+#include "ProfilingDebugging/ScopedTimers.h"
 
 #if WITH_EDITOR
 #include "LevelInstance/LevelInstanceEditorInstanceActor.h"
@@ -21,6 +22,7 @@ ULevelStreamingLevelInstance::ULevelStreamingLevelInstance(const FObjectInitiali
 	: Super(ObjectInitializer)
 #if WITH_EDITOR
 	, CachedBounds(ForceInit)
+	, bResetLoadersCalled(false)
 #endif
 {
 #if WITH_EDITOR
@@ -75,17 +77,63 @@ FBox ULevelStreamingLevelInstance::GetBounds() const
 void ULevelStreamingLevelInstance::OnLoadedActorAddedToLevel(AActor& InActor)
 {
 	check(LevelInstanceEditorInstanceActor.IsValid());
-	PrepareLevelInstanceLoadedActor(InActor, GetLevelInstance());
+	PrepareLevelInstanceLoadedActor(InActor, GetLevelInstance(), true);
 }
 
-void ULevelStreamingLevelInstance::PrepareLevelInstanceLoadedActor(AActor& InActor, ILevelInstanceInterface* InLevelInstance)
+void ULevelStreamingLevelInstance::ResetLevelInstanceLoaders()
+{
+	// @todo_ow: Resetting the load will prevent any OFPA package to properly load since their import level will fail to resolve.
+	//           This is a temporary workaround. The downside of this is that the level can't be saved. 
+	//           Most of the changes will only affect OFPA packages. One problematic use case is when changing the pivot of a Level Instance.
+
+	if (bResetLoadersCalled)
+	{
+		return;
+	}
+
+	if (UWorld* OuterWorld = LoadedLevel ? LoadedLevel->GetTypedOuter<UWorld>() : nullptr)
+	{
+		UE_SCOPED_TIMER(*FString::Printf(TEXT("ULevelStreamingLevelInstance::ResetLevelInstanceLoaders(%s)"), *FPaths::GetBaseFilename(OuterWorld->GetPackage()->GetName())), LogLevelInstance, Log);
+
+		FName PackageName = OuterWorld->GetPackage()->GetLoadedPath().GetPackageFName();
+		if (!ULevel::GetIsLevelPartitionedFromPackage(PackageName))
+		{
+			ResetLoaders(OuterWorld->GetPackage());
+		}
+
+		for (AActor* Actor : LoadedLevel->Actors)
+		{
+			if (Actor && Actor->IsPackageExternal())
+			{
+				ResetLoaders(Actor->GetExternalPackage());
+			}
+		}
+
+		LoadedLevel->ForEachActorFolder([](UActorFolder* ActorFolder)
+		{
+			if (ActorFolder->IsPackageExternal())
+			{
+				ResetLoaders(ActorFolder->GetExternalPackage());
+			}
+			return true;
+		});
+
+		bResetLoadersCalled = true;
+	}
+}
+
+void ULevelStreamingLevelInstance::PrepareLevelInstanceLoadedActor(AActor& InActor, ILevelInstanceInterface* InLevelInstance, bool bResetLoaders)
 {
 	InActor.ClearFlags(RF_Transactional);
 	InActor.SetFlags(RF_Transient);
 
 	if (InActor.IsPackageExternal())
 	{
-		ResetLoaders(InActor.GetExternalPackage());
+		if (bResetLoaders)
+		{
+			ResetLoaders(InActor.GetExternalPackage());
+		}
+
 		InActor.GetPackage()->SetFlags(RF_Transient);
 	}
 
@@ -176,15 +224,6 @@ ULevelStreamingLevelInstance* ULevelStreamingLevelInstance::LoadInstance(ILevelI
 			OuterWorld->ClearFlags(RF_Transactional);
 			OuterWorld->SetFlags(RF_Transient);
 			
-			// @todo_ow: Resetting the load will prevent any OFPA package to properly load since their import level will fail to resolve.
-			//           This is a temporary workaround. The downside of this is that the level can't be saved. 
-			//           Most of the changes will only affect OFPA packages. One problematic use case is when changing the pivot of a Level Instance.
-			FName PackageName = OuterWorld->GetPackage()->GetLoadedPath().GetPackageFName();
-			if (!ULevel::GetIsLevelPartitionedFromPackage(PackageName))
-			{
-				ResetLoaders(OuterWorld->GetPackage());
-			}
-
 			OuterWorld->GetPackage()->ClearFlags(RF_Transactional);
 			OuterWorld->GetPackage()->SetFlags(RF_Transient);
 
@@ -198,14 +237,14 @@ ULevelStreamingLevelInstance* ULevelStreamingLevelInstance::LoadInstance(ILevelI
 			{
 				if (LevelActor)
 				{
-					LevelStreaming->PrepareLevelInstanceLoadedActor(*LevelActor, LevelInstance);
+					LevelStreaming->PrepareLevelInstanceLoadedActor(*LevelActor, LevelInstance, false);
 				}
 			}
+
 			Level->ForEachActorFolder([](UActorFolder* ActorFolder)
 			{
 				if (ActorFolder->IsPackageExternal())
 				{
-					ResetLoaders(ActorFolder->GetExternalPackage());
 					ActorFolder->GetPackage()->SetFlags(RF_Transient);
 				}
 				return true;
@@ -274,4 +313,4 @@ void ULevelStreamingLevelInstance::OnLevelLoadedChanged(ULevel* InLevel)
 			LevelInstanceSubsystem->RegisterLoadedLevelStreamingLevelInstance(this);
 		}
 	}
-}
\ No newline at end of file
+}
diff --git a/Engine/Source/Runtime/Engine/Private/LevelInstance/LevelInstanceSubsystem.cpp b/Engine/Source/Runtime/Engine/Private/LevelInstance/LevelInstanceSubsystem.cpp
--- a/Engine/Source/Runtime/Engine/Private/LevelInstance/LevelInstanceSubsystem.cpp
+++ b/Engine/Source/Runtime/Engine/Private/LevelInstance/LevelInstanceSubsystem.cpp
@@ -1936,6 +1936,14 @@ bool ULevelInstanceSubsystem::EditLevelInstanceInternal(ILevelInstanceInterface*
 	// When editing a Level Instance, push a new empty actor editor context
 	UActorEditorContextSubsystem::Get()->PushContext();
 
+	for (const auto& [LevelInstanceID, LoadedLevelInstance] : LoadedLevelInstances)
+	{
+		if (LoadedLevelInstance.LevelStreaming && LoadedLevelInstance.LevelStreaming->PackageNameToLoad.ToString() == LevelInstance->GetWorldAsset().GetLongPackageName())
+		{
+			LoadedLevelInstance.LevelStreaming->ResetLevelInstanceLoaders();
+		}
+	}
+
 	return true;
 }
 
@@ -2367,4 +2375,4 @@ bool ULevelInstanceSubsystem::CanUseWorldAsset(const ILevelInstanceInterface* Le
 
 #endif
 
-#undef LOCTEXT_NAMESPACE
\ No newline at end of file
+#undef LOCTEXT_NAMESPACE
diff --git a/Engine/Source/Runtime/Engine/Public/LevelInstance/LevelInstanceLevelStreaming.h b/Engine/Source/Runtime/Engine/Public/LevelInstance/LevelInstanceLevelStreaming.h
--- a/Engine/Source/Runtime/Engine/Public/LevelInstance/LevelInstanceLevelStreaming.h
+++ b/Engine/Source/Runtime/Engine/Public/LevelInstance/LevelInstanceLevelStreaming.h
@@ -37,7 +37,8 @@ protected:
 	const FLevelInstanceID& GetLevelInstanceID() const { return LevelInstanceID; }
 private:
 #if WITH_EDITOR
-	void PrepareLevelInstanceLoadedActor(AActor& InActor, ILevelInstanceInterface* InLevelInstance);
+	void ResetLevelInstanceLoaders();
+	void PrepareLevelInstanceLoadedActor(AActor& InActor, ILevelInstanceInterface* InLevelInstance, bool bResetLoaders);
 	void OnLoadedActorAddedToLevel(AActor& InActor);
 	void OnLoadedActorRemovedFromLevel(AActor& InActor);
 
@@ -45,6 +46,7 @@ private:
 
 	mutable FTransform CachedTransform;
 	mutable FBox CachedBounds;
+	bool bResetLoadersCalled;
 #endif
 	FLevelInstanceID LevelInstanceID;
 };
\ No newline at end of file
