DistanceFieldApproxAO material node.
- As an optimization CalculateDistanceFieldApproxAO(...) calculates which clipmap covers all the sampling points which reduces ALU however can lead to reduced quality when using a large radius.

diff --git a/Engine/Shaders/Private/DistanceField/GlobalDistanceFieldShared.ush b/Engine/Shaders/Private/DistanceField/GlobalDistanceFieldShared.ush
--- a/Engine/Shaders/Private/DistanceField/GlobalDistanceFieldShared.ush
+++ b/Engine/Shaders/Private/DistanceField/GlobalDistanceFieldShared.ush
@@ -316,6 +316,45 @@ float3 GetDistanceFieldGradientGlobal(float3 WorldPosition)
 	return Gradient;
 }
 
+float CalculateDistanceFieldApproxAO(float3 WorldPosition, float3 WorldNormal, uint NumSteps, float StepDistance, float StepScale, float DistanceBias, float MaxDistance)
+{
+	// determine the min clipmap that covers all the sampling points (sphere centered on WorldPosition with radius = MaxDistance)
+	int MinClipmapIndex = -1;
+
+	for (uint ClipmapIndex = 0; ClipmapIndex < NumGlobalSDFClipmaps; ++ClipmapIndex)
+	{
+		float DistanceFromClipmap = ComputeDistanceFromBoxToPointInside(GlobalVolumeCenterAndExtent[ClipmapIndex].xyz, GlobalVolumeCenterAndExtent[ClipmapIndex].www, WorldPosition);
+		if (DistanceFromClipmap > MaxDistance)
+		{
+			MinClipmapIndex = ClipmapIndex;
+			break;
+		}
+	}
+
+	if (MinClipmapIndex < 0)
+	{
+		return 1.0f;
+	}
+
+	float Occlusion = 0.0f;
+	float W = 0.5f;
+	float TotalW = 0.0f;
+
+	for (uint StepIndex = 0; StepIndex < NumSteps; ++StepIndex)
+	{
+		const float CurrentDistance = StepDistance + DistanceBias;
+		const float Dist = max(0.0f, SampleGlobalDistanceField(WorldPosition + WorldNormal * CurrentDistance, CurrentDistance, MinClipmapIndex));
+
+		Occlusion += W * Dist / CurrentDistance;
+		TotalW += W;
+		W *= 0.5f;
+
+		StepDistance *= StepScale;
+	}
+
+	return saturate(Occlusion / TotalW);
+}
+
 #else
 
 float GetDistanceToNearestSurfaceGlobal(float3 WorldPosition)
@@ -328,5 +367,10 @@ float3 GetDistanceFieldGradientGlobal(float3 WorldPosition)
 	return float3(0, 0, .001f);
 }
 
+float CalculateDistanceFieldApproxAO(float3 WorldPosition, float3 WorldNormal, uint NumSteps, float BaseStepDistance, float StepScale, float MaxDistance)
+{
+	return 1.0f;
+}
+
 #endif
 #endif // FEATURE_LEVEL >= FEATURE_LEVEL_SM5
diff --git a/Engine/Source/Runtime/Engine/Classes/Materials/MaterialExpressionDistanceFieldApproxAO.h b/Engine/Source/Runtime/Engine/Classes/Materials/MaterialExpressionDistanceFieldApproxAO.h
new file mode 100644
--- /dev/null
+++ b/Engine/Source/Runtime/Engine/Classes/Materials/MaterialExpressionDistanceFieldApproxAO.h
@@ -0,0 +1,42 @@
+// Copyright Epic Games, Inc. All Rights Reserved.
+
+
+#pragma once
+
+#include "CoreMinimal.h"
+#include "UObject/ObjectMacros.h"
+#include "MaterialExpressionIO.h"
+#include "Materials/MaterialExpression.h"
+#include "MaterialExpressionDistanceFieldApproxAO.generated.h"
+
+UCLASS()
+class UMaterialExpressionDistanceFieldApproxAO : public UMaterialExpression
+{
+	GENERATED_UCLASS_BODY()
+
+	UPROPERTY(meta = (RequiredInput = "false", ToolTip = "Defaults to current world position if not specified"))
+	FExpressionInput Position;
+
+	UPROPERTY(meta = (RequiredInput = "false", ToolTip = "Defaults to current world normal if not specified"))
+	FExpressionInput Normal;
+
+	UPROPERTY(meta = (RequiredInput = "false"))
+	FExpressionInput Radius;
+
+	/** only used if TextureSize is not hooked up */
+	UPROPERTY(EditAnywhere, Category = MaterialExpressionDistanceFieldApproxAO, meta=(OverridingInputProperty = "Radius"))
+	float RadiusDefault;
+
+	UPROPERTY(EditAnywhere, Category = MaterialExpressionDistanceFieldApproxAO, meta = (UIMin = "1", UIMax = "4", ShowAsInputPin = "Primary", ToolTip = "Number of samples used to calculate occlusion"))
+	uint32 NumSteps;
+
+	UPROPERTY(EditAnywhere, Category = MaterialExpressionDistanceFieldApproxAO, DisplayName = "Step Scale", meta = (UIMin = "1.0", UIMax = "8.0", ShowAsInputPin = "Advanced", ToolTip = "Used to control step distance distribution"))
+	float StepScaleDefault;
+
+	//~ Begin UMaterialExpression Interface
+#if WITH_EDITOR
+	virtual int32 Compile(class FMaterialCompiler* Compiler, int32 OutputIndex) override;
+	virtual void GetCaption(TArray<FString>& OutCaptions) const override;
+#endif
+	//~ End UMaterialExpression Interface
+};
diff --git a/Engine/Source/Runtime/Engine/Private/Materials/HLSLMaterialTranslator.cpp b/Engine/Source/Runtime/Engine/Private/Materials/HLSLMaterialTranslator.cpp
--- a/Engine/Source/Runtime/Engine/Private/Materials/HLSLMaterialTranslator.cpp
+++ b/Engine/Source/Runtime/Engine/Private/Materials/HLSLMaterialTranslator.cpp
@@ -9078,6 +9078,52 @@ int32 FHLSLMaterialTranslator::DistanceFieldGradient(int32 PositionArg)
 	return AddCodeChunk(MCT_Float3, TEXT("GetDistanceFieldGradientGlobal(%s)"), *CoerceParameter(PositionArg, MCT_Float3));
 }
 
+int32 FHLSLMaterialTranslator::DistanceFieldApproxAO(int32 PositionArg, int32 NormalArg, int32 RadiusArg, uint32 NumSteps, float StepScale)
+{
+	if (ErrorUnlessFeatureLevelSupported(ERHIFeatureLevel::SM5) == INDEX_NONE)
+	{
+		return INDEX_NONE;
+	}
+
+	if (PositionArg == INDEX_NONE)
+	{
+		return INDEX_NONE;
+	}
+
+	if (NormalArg == INDEX_NONE)
+	{
+		return INDEX_NONE;
+	}
+
+	if (RadiusArg == INDEX_NONE)
+	{
+		return INDEX_NONE;
+	}
+
+	MaterialCompilationOutput.bUsesGlobalDistanceField = true;
+
+	NumSteps = FMath::Clamp(NumSteps, 1, 4);
+	StepScale = FMath::Max(StepScale, 1.0f);
+
+	int32 NumStepsConst = Constant(NumSteps);
+	int32 NumStepsMinus1Const = Constant(NumSteps - 1);
+	int32 StepScaleConst = Constant(StepScale);
+
+	int32 StepDistance = Div(RadiusArg, Power(StepScaleConst, NumStepsMinus1Const));
+	int32 DistanceBiasConst = Constant(0);
+
+	// LWC_TODO: update for LWC position
+	return AddCodeChunk(MCT_Float,
+		TEXT("CalculateDistanceFieldApproxAO(%s, %s, %s, %s, %s, %s, %s)"),
+		*CoerceParameter(PositionArg, MCT_Float3),
+		*CoerceParameter(NormalArg, MCT_Float3),
+		*GetParameterCode(NumStepsConst),
+		*CoerceParameter(StepDistance, MCT_Float),
+		*GetParameterCode(StepScaleConst),
+		*GetParameterCode(DistanceBiasConst),
+		*CoerceParameter(RadiusArg, MCT_Float));
+}
+
 int32 FHLSLMaterialTranslator::SamplePhysicsField(int32 PositionArg, const int32 OutputType, const int32 TargetIndex)
 {
 	if (ErrorUnlessFeatureLevelSupported(ERHIFeatureLevel::SM5) == INDEX_NONE)
diff --git a/Engine/Source/Runtime/Engine/Private/Materials/HLSLMaterialTranslator.h b/Engine/Source/Runtime/Engine/Private/Materials/HLSLMaterialTranslator.h
--- a/Engine/Source/Runtime/Engine/Private/Materials/HLSLMaterialTranslator.h
+++ b/Engine/Source/Runtime/Engine/Private/Materials/HLSLMaterialTranslator.h
@@ -967,6 +967,7 @@ protected:
 	virtual int32 GetHairColorFromMelanin(int32 Melanin, int32 Redness, int32 DyeColor) override;
 	virtual int32 DistanceToNearestSurface(int32 PositionArg) override;
 	virtual int32 DistanceFieldGradient(int32 PositionArg) override;
+	virtual int32 DistanceFieldApproxAO(int32 PositionArg, int32 NormalArg, int32 RadiusArg, uint32 NumSteps, float StepScale) override;
 	virtual int32 SamplePhysicsField(int32 PositionArg, const int32 OutputType, const int32 TargetIndex) override;
 	virtual int32 AtmosphericFogColor(int32 WorldPosition) override;
 	virtual int32 AtmosphericLightVector() override;
diff --git a/Engine/Source/Runtime/Engine/Private/Materials/MaterialExpressions.cpp b/Engine/Source/Runtime/Engine/Private/Materials/MaterialExpressions.cpp
--- a/Engine/Source/Runtime/Engine/Private/Materials/MaterialExpressions.cpp
+++ b/Engine/Source/Runtime/Engine/Private/Materials/MaterialExpressions.cpp
@@ -245,6 +245,7 @@
 #include "Materials/MaterialExpressionWorldPosition.h"
 #include "Materials/MaterialExpressionDistanceToNearestSurface.h"
 #include "Materials/MaterialExpressionDistanceFieldGradient.h"
+#include "Materials/MaterialExpressionDistanceFieldApproxAO.h"
 #include "Materials/MaterialParameterCollection.h"
 #include "Materials/MaterialExpressionClearCoatNormalCustomOutput.h"
 #include "Materials/MaterialExpressionAtmosphericLightVector.h"
@@ -18361,6 +18362,68 @@ void UMaterialExpressionDistanceFieldGradient::GetCaption(TArray<FString>& OutCa
 }
 #endif // WITH_EDITOR
 
+///////////////////////////////////////////////////////////////////////////////
+// UMaterialExpressionDistanceFieldApproxAO
+///////////////////////////////////////////////////////////////////////////////
+UMaterialExpressionDistanceFieldApproxAO::UMaterialExpressionDistanceFieldApproxAO(const FObjectInitializer& ObjectInitializer)
+	: Super(ObjectInitializer)
+{
+	// Structure to hold one-time initialization
+	struct FConstructorStatics
+	{
+		FText NAME_Utility;
+		FConstructorStatics()
+			: NAME_Utility(LOCTEXT("Utility", "Utility"))
+		{
+		}
+	};
+	static FConstructorStatics ConstructorStatics;
+
+#if WITH_EDITORONLY_DATA
+	MenuCategories.Add(ConstructorStatics.NAME_Utility);
+#endif
+
+	NumSteps = 1;
+	RadiusDefault = 150;
+	StepScaleDefault = 3.0f;
+}
+
+#if WITH_EDITOR
+int32 UMaterialExpressionDistanceFieldApproxAO::Compile(class FMaterialCompiler* Compiler, int32 OutputIndex)
+{
+	int32 PositionArg = INDEX_NONE;
+
+	if (Position.GetTracedInput().Expression)
+	{
+		PositionArg = Position.Compile(Compiler);
+	}
+	else
+	{
+		PositionArg = Compiler->WorldPosition(WPT_Default);
+	}
+
+	int32 NormalArg = INDEX_NONE;
+
+	if (Normal.GetTracedInput().Expression)
+	{
+		NormalArg = Position.Compile(Compiler);
+	}
+	else
+	{
+		NormalArg = Compiler->VertexNormal();
+	}
+
+	int32 RadiusArg = Radius.GetTracedInput().Expression ? Radius.Compile(Compiler) : Compiler->Constant(RadiusDefault);
+
+	return Compiler->DistanceFieldApproxAO(PositionArg, NormalArg, RadiusArg, NumSteps, StepScaleDefault);
+}
+
+void UMaterialExpressionDistanceFieldApproxAO::GetCaption(TArray<FString>& OutCaptions) const
+{
+	OutCaptions.Add(TEXT("DistanceFieldApproxAO"));
+}
+#endif // WITH_EDITOR
+
 ///////////////////////////////////////////////////////////////////////////////
 // UMaterialExpressionSamplePhysicsVectorField
 ///////////////////////////////////////////////////////////////////////////////
diff --git a/Engine/Source/Runtime/Engine/Public/MaterialCompiler.h b/Engine/Source/Runtime/Engine/Public/MaterialCompiler.h
--- a/Engine/Source/Runtime/Engine/Public/MaterialCompiler.h
+++ b/Engine/Source/Runtime/Engine/Public/MaterialCompiler.h
@@ -502,6 +502,7 @@ public:
 	virtual int32 BlackBody( int32 Temp ) = 0;
 	virtual int32 DistanceToNearestSurface(int32 PositionArg) = 0;
 	virtual int32 DistanceFieldGradient(int32 PositionArg) = 0;
+	virtual int32 DistanceFieldApproxAO(int32 PositionArg, int32 NormalArg, int32 RadiusArg, uint32 NumSteps, float StepScale) = 0;
 	virtual int32 SamplePhysicsField(int32 PositionArg, const int32 OutputType, const int32 TargetIndex) = 0;
 	virtual int32 DepthOfFieldFunction(int32 Depth, int32 FunctionValueIndex) = 0;
 	virtual int32 AtmosphericFogColor(int32 WorldPosition) = 0;
@@ -903,6 +904,7 @@ public:
 	virtual int32 BlackBody(int32 Temp) override { return Compiler->BlackBody(Temp); }
 	virtual int32 DistanceToNearestSurface(int32 PositionArg) override { return Compiler->DistanceToNearestSurface(PositionArg); }
 	virtual int32 DistanceFieldGradient(int32 PositionArg) override { return Compiler->DistanceFieldGradient(PositionArg); }
+	virtual int32 DistanceFieldApproxAO(int32 PositionArg, int32 NormalArg, int32 RadiusArg, uint32 NumSteps, float StepScale) override { return Compiler->DistanceFieldApproxAO(PositionArg, NormalArg, RadiusArg, NumSteps, StepScale); }
 	virtual int32 SamplePhysicsField(int32 PositionArg, const int32 OutputType, const int32 TargetIndex)  override { return Compiler->SamplePhysicsField(PositionArg, OutputType, TargetIndex); }
 	virtual int32 PerInstanceRandom() override { return Compiler->PerInstanceRandom(); }
 	virtual int32 PerInstanceFadeAmount() override { return Compiler->PerInstanceFadeAmount(); }
