PR #6992: Add ability to check if a UInputComponent is in an APlayerController's input stack (Contributed by [REDACTED_PERSON])

diff --git a/Engine/Source/Runtime/Engine/Classes/GameFramework/PlayerController.h b/Engine/Source/Runtime/Engine/Classes/GameFramework/PlayerController.h
--- a/Engine/Source/Runtime/Engine/Classes/GameFramework/PlayerController.h
+++ b/Engine/Source/Runtime/Engine/Classes/GameFramework/PlayerController.h
@@ -1692,6 +1692,9 @@ public:
 	/** Removes given inputcomponent from the input stack (regardless of if it's the top, actually). */
 	virtual bool PopInputComponent(UInputComponent* Input);
 
+	/** Returns true if the given input component is in this PlayerController's CurrentInputStack */
+	virtual bool IsInputComponentInStack(const UInputComponent* Input) const;
+
 	/** Flushes the current key state. */
 	virtual void FlushPressedKeys();
 
diff --git a/Engine/Source/Runtime/Engine/Private/PlayerController.cpp b/Engine/Source/Runtime/Engine/Private/PlayerController.cpp
--- a/Engine/Source/Runtime/Engine/Private/PlayerController.cpp
+++ b/Engine/Source/Runtime/Engine/Private/PlayerController.cpp
@@ -5431,6 +5431,11 @@ bool APlayerController::PopInputComponent(UInputComponent* InInputComponent)
 	return false;
 }
 
+bool APlayerController::IsInputComponentInStack(const UInputComponent* InInputComponent) const
+{
+	return InInputComponent && CurrentInputStack.Contains(InInputComponent);
+}
+
 void APlayerController::AddPitchInput(float Val)
 {	
 	RotationInput.Pitch += !IsLookInputIgnored() ? Val * (GetDefault<UInputSettings>()->bEnableLegacyInputScales ? InputPitchScale_DEPRECATED : 1.0f) : 0.0f;
