Follow up to CL# 21088391
Remove properties and functions so that redirectors can do their job
Add redirector for function paramter

diff --git a/Engine/Config/BaseEngine.ini b/Engine/Config/BaseEngine.ini
--- a/Engine/Config/BaseEngine.ini
+++ b/Engine/Config/BaseEngine.ini
@@ -1103,6 +1103,7 @@ DefaultWorldPartitionHLODLayer=/Engine/Maps/Templates/HLODs/HLODLayer_Instanced.
 +PropertyRedirects=(OldName="SkinnedMeshComponent.bUseBoundsFromMasterPoseComponent", NewName="bUseBoundsFromLeaderPoseComponent")
 +PropertyRedirects=(OldName="SkinnedMeshComponent.bIgnoreMasterPoseComponentLOD", NewName="bIgnoreLeaderPoseComponentLOD")
 +FunctionRedirects=(OldName="SkinnedMeshComponent.SetMasterPoseComponent", NewName="SetLeaderPoseComponent")
++PropertyRedirects=(OldName="SkinnedMeshComponent.SetLeaderPoseComponent.NewMasterBoneComponent",NewName="SkinnedMeshComponent.SetLeaderPoseComponent.NewLeaderBoneComponent")
 
 +PropertyRedirects=(OldName="SkeletalMeshComponent.bPropagateCurvesToSlaves", NewName="bPropagateCurvesToFollowers")
 +FunctionRedirects=(OldName="SkeletalMeshComponent.BindClothToMasterPoseComponent", NewName="BindClothToLeaderPoseComponent")
diff --git a/Engine/Source/Runtime/AnimGraphRuntime/Public/BoneControllers/AnimNode_RigidBody.h b/Engine/Source/Runtime/AnimGraphRuntime/Public/BoneControllers/AnimNode_RigidBody.h
--- a/Engine/Source/Runtime/AnimGraphRuntime/Public/BoneControllers/AnimNode_RigidBody.h
+++ b/Engine/Source/Runtime/AnimGraphRuntime/Public/BoneControllers/AnimNode_RigidBody.h
@@ -67,7 +67,6 @@ struct ANIMGRAPHRUNTIME_API FSimSpaceSettings
 
 #if WITH_EDITORONLY_DATA
 	UE_DEPRECATED(5.1, "This property has been deprecated. Please, use WorldAlpha.")
-	UPROPERTY()
 	float MasterAlpha = 0.f;
 #endif // WITH_EDITORONLY_DATA
 
diff --git a/Engine/Source/Runtime/Engine/Classes/Components/SkeletalMeshComponent.h b/Engine/Source/Runtime/Engine/Classes/Components/SkeletalMeshComponent.h
--- a/Engine/Source/Runtime/Engine/Classes/Components/SkeletalMeshComponent.h
+++ b/Engine/Source/Runtime/Engine/Classes/Components/SkeletalMeshComponent.h
@@ -779,7 +779,6 @@ public:
 
 #if WITH_EDITORONLY_DATA
 	UE_DEPRECATED(5.1, "This property is deprecated. Please use bPropagateCurvesToFollowers instead")
-	UPROPERTY()
 	uint8 bPropagateCurvesToSlaves : 1;
 #endif // WITH_EDITORONLY_DATA
 
@@ -1232,7 +1231,6 @@ public:
 	void BindClothToLeaderPoseComponent();
 
 	UE_DEPRECATED(5.1, "BindClothToMasterPoseComponent has been deprecated. Please use BindClothToLeaderPoseComponent instead.")
-	UFUNCTION(BlueprintCallable, Category="Clothing", meta=(UnsafeDuringActorConstruction="true", DeprecatedFunction, DeprecationMessage = "BindClothToMasterPoseComponent has been deprecated. Please use BindClothToLeaderPoseComponent instead."))
 	void BindClothToMasterPoseComponent() { BindClothToLeaderPoseComponent(); }
 
 	/** 
@@ -1246,7 +1244,6 @@ public:
 	void UnbindClothFromLeaderPoseComponent(bool bRestoreSimulationSpace = true);
 
 	UE_DEPRECATED(5.1, "UnbindClothFromMasterPoseComponent has been deprecated. Please use UnbindClothFromLeaderPoseComponent instead.")
-	UFUNCTION(BlueprintCallable, Category = "Clothing", meta = (UnsafeDuringActorConstruction = "true", DeprecatedFunction, DeprecationMessage = "UnbindClothFromMasterPoseComponent has been deprecated. Please use UnbindClothFromLeaderPoseComponent instead."))
 	void UnbindClothFromMasterPoseComponent(bool bRestoreSimulationSpace = true) { UnbindClothFromLeaderPoseComponent(bRestoreSimulationSpace); }
 
 	/**
diff --git a/Engine/Source/Runtime/Engine/Classes/Components/SkinnedMeshComponent.h b/Engine/Source/Runtime/Engine/Classes/Components/SkinnedMeshComponent.h
--- a/Engine/Source/Runtime/Engine/Classes/Components/SkinnedMeshComponent.h
+++ b/Engine/Source/Runtime/Engine/Classes/Components/SkinnedMeshComponent.h
@@ -277,7 +277,6 @@ public:
 
 #if WITH_EDITORONLY_DATA
 	UE_DEPRECATED(5.1, "This property is deprecated. Please use LeaderPoseComponent instead")
-	UPROPERTY()
 	TWeakObjectPtr<USkinnedMeshComponent> MasterPoseComponent;
 #endif // WITH_EDITORONLY_DATA
 
@@ -576,7 +575,6 @@ public:
 
 #if WITH_EDITORONLY_DATA
 	UE_DEPRECATED(5.1, "This property is deprecated. Please use bUseBoundsFromLeaderPoseComponent instead")
-	UPROPERTY()
 	uint8 bUseBoundsFromMasterPoseComponent : 1;
 #endif // WITH_EDITORONLY_DATA
 
@@ -671,7 +669,6 @@ public:
 
 #if WITH_EDITORONLY_DATA
 	UE_DEPRECATED(5.1, "This property is deprecated. Please use bIgnoreLeaderPoseComponentLOD instead")
-	UPROPERTY()
 	uint8 bIgnoreMasterPoseComponentLOD : 1;
 #endif // WITH_EDITORONLY_DATA
 
@@ -1444,7 +1441,6 @@ public:
 	void SetLeaderPoseComponent(USkinnedMeshComponent* NewLeaderBoneComponent, bool bForceUpdate = false);
 
 	UE_DEPRECATED(5.1, "SetMasterPoseComponent has been deprecated. Please use SetLeaderPoseComponent instead.")
-	UFUNCTION(BlueprintCallable, Category = "Components|SkinnedMesh", meta = (DeprecatedFunction, DeprecationMessage = "SetMasterPoseComponent has been deprecated. Please use SetLeaderPoseComponent instead."))
 	void SetMasterPoseComponent(USkinnedMeshComponent* NewMasterBoneComponent, bool bForceUpdate = false) { SetLeaderPoseComponent(NewMasterBoneComponent, bForceUpdate); }
 
 	/** Return current active list of follower components */
