Remove blueprint deprecation of GetSkeletalMesh until all content is fixed up so that cooks do not fail
#rnx

diff --git a/Engine/Source/Runtime/Engine/Classes/Components/SkinnedMeshComponent.h b/Engine/Source/Runtime/Engine/Classes/Components/SkinnedMeshComponent.h
--- a/Engine/Source/Runtime/Engine/Classes/Components/SkinnedMeshComponent.h
+++ b/Engine/Source/Runtime/Engine/Classes/Components/SkinnedMeshComponent.h
@@ -874,7 +874,7 @@ public:
 	 * @return the SkeletalMesh set to this mesh.
 	 */
 	UE_DEPRECATED(5.1, "Use USkeletalMeshComponent::GetSkeletalMeshAsset() or GetSkinnedAsset() instead.")
-	UFUNCTION(BlueprintPure, Category = "Components|SkinnedMesh", DisplayName = "Get Skeletal Mesh", meta = (DeprecatedFunction, DeprecationMessage = "Use USkeletalMeshComponent::GetSkeletalMesh() or GetSkinnedAsset() instead."))
+	UFUNCTION(BlueprintPure, Category = "Components|SkinnedMesh", DisplayName = "Get Skeletal Mesh", meta = (DeprecationMessage = "Use USkeletalMeshComponent::GetSkeletalMesh() or GetSkinnedAsset() instead."))
 	class USkeletalMesh* GetSkeletalMesh_DEPRECATED() const;
 
 	/**
