More fully remove blueprint deprecation until content is fixed
#rnx

diff --git a/Engine/Source/Runtime/Engine/Classes/Components/SkinnedMeshComponent.h b/Engine/Source/Runtime/Engine/Classes/Components/SkinnedMeshComponent.h
--- a/Engine/Source/Runtime/Engine/Classes/Components/SkinnedMeshComponent.h
+++ b/Engine/Source/Runtime/Engine/Classes/Components/SkinnedMeshComponent.h
@@ -253,7 +253,7 @@ class ENGINE_API USkinnedMeshComponent : public UMeshComponent, public ILODSyncI
 
 	/** The skeletal mesh used by this component. */
 	UE_DEPRECATED(5.1, "Replaced by SkinnedAsset. Use GetSkinnedAsset()/SetSkinnedAsset() instead.")
-	UPROPERTY(EditAnywhere, BlueprintGetter = GetSkeletalMesh_DEPRECATED, Category = "Mesh|SkeletalAsset", meta = (DisallowedClasses = "/Script/ApexDestruction.DestructibleMesh", DeprecatedProperty, DeprecationMessage = "Use USkeletalMeshComponent::GetSkeletalMesh() or GetSkinnedAsset() instead."))
+	UPROPERTY(EditAnywhere, BlueprintGetter = GetSkeletalMesh_DEPRECATED, Category = "Mesh|SkeletalAsset", meta = (DisallowedClasses = "/Script/ApexDestruction.DestructibleMesh")
 	TObjectPtr<class USkeletalMesh> SkeletalMesh;
 
 private:
@@ -874,7 +874,7 @@ public:
 	 * @return the SkeletalMesh set to this mesh.
 	 */
 	UE_DEPRECATED(5.1, "Use USkeletalMeshComponent::GetSkeletalMeshAsset() or GetSkinnedAsset() instead.")
-	UFUNCTION(BlueprintPure, Category = "Components|SkinnedMesh", DisplayName = "Get Skeletal Mesh", meta = (DeprecationMessage = "Use USkeletalMeshComponent::GetSkeletalMesh() or GetSkinnedAsset() instead."))
+	UFUNCTION(BlueprintPure, Category = "Components|SkinnedMesh", DisplayName = "Get Skeletal Mesh")
 	class USkeletalMesh* GetSkeletalMesh_DEPRECATED() const;
 
 	/**
