Restore missing )
#rnx

diff --git a/Engine/Source/Runtime/Engine/Classes/Components/SkinnedMeshComponent.h b/Engine/Source/Runtime/Engine/Classes/Components/SkinnedMeshComponent.h
--- a/Engine/Source/Runtime/Engine/Classes/Components/SkinnedMeshComponent.h
+++ b/Engine/Source/Runtime/Engine/Classes/Components/SkinnedMeshComponent.h
@@ -253,7 +253,7 @@ class ENGINE_API USkinnedMeshComponent : public UMeshComponent, public ILODSyncI
 
 	/** The skeletal mesh used by this component. */
 	UE_DEPRECATED(5.1, "Replaced by SkinnedAsset. Use GetSkinnedAsset()/SetSkinnedAsset() instead.")
-	UPROPERTY(EditAnywhere, BlueprintGetter = GetSkeletalMesh_DEPRECATED, Category = "Mesh|SkeletalAsset", meta = (DisallowedClasses = "/Script/ApexDestruction.DestructibleMesh")
+	UPROPERTY(EditAnywhere, BlueprintGetter = GetSkeletalMesh_DEPRECATED, Category = "Mesh|SkeletalAsset", meta = (DisallowedClasses = "/Script/ApexDestruction.DestructibleMesh"))
 	TObjectPtr<class USkeletalMesh> SkeletalMesh;
 
 private:
