HISM: Allocate runs used by mesh batches with the scene rendering collector.
Previously we allocated on stack and relied on old memstack behavior to keep the data alive.

diff --git a/Engine/Source/Runtime/Engine/Private/HierarchicalInstancedStaticMesh.cpp b/Engine/Source/Runtime/Engine/Private/HierarchicalInstancedStaticMesh.cpp
--- a/Engine/Source/Runtime/Engine/Private/HierarchicalInstancedStaticMesh.cpp
+++ b/Engine/Source/Runtime/Engine/Private/HierarchicalInstancedStaticMesh.cpp
@@ -956,7 +956,7 @@ void FHierarchicalStaticMeshSceneProxy::ApplyWorldOffset(FVector InOffset)
 	}
 }
 
-struct FFoliageRenderInstanceParams
+struct FFoliageRenderInstanceParams : public FOneFrameResource
 {
 	bool bNeedsSingleLODRuns;
 	bool bNeedsMultipleLODRuns;
@@ -1533,7 +1533,7 @@ void FHierarchicalStaticMeshSceneProxy::GetDynamicMeshElements(const TArray<cons
 			// Render built instances
 			if (ClusterTree.Num())
 			{
-				FFoliageCullInstanceParams InstanceParams(bSingleSections, bMultipleSections, bOverestimate, ClusterTree);
+				FFoliageCullInstanceParams& InstanceParams = Collector.AllocateOneFrameResource<FFoliageCullInstanceParams>(bSingleSections, bMultipleSections, bOverestimate, ClusterTree);
 				InstanceParams.LODs = RenderData->LODResources.Num();
 
 				InstanceParams.View = View;
@@ -1793,7 +1793,7 @@ void FHierarchicalStaticMeshSceneProxy::GetDynamicMeshElements(const TArray<cons
 			// Render unbuilt instances
 			if (UnbuiltInstanceCount > 0)
 			{
-				FFoliageRenderInstanceParams InstanceParams(true, false, false);
+				FFoliageRenderInstanceParams& InstanceParams = Collector.AllocateOneFrameResource<FFoliageRenderInstanceParams>(true, false, false);
 
 				// disable LOD blending for unbuilt instances as we haven't calculated the correct LOD.
 				ElementParams.bBlendLODs = false;
