SoftObjectPath remapping cleanup:

- LinkerLoad uses InstancingContext to do SoftObjectPath remaps
- This fixes SoftObjectPath remaps for OFPA Actors
- Fix WorldPartitionLevelHelper now that a first remapping is done in the LinkerLoad we need to change how we remap to cells for LevelInstance embedded actors
- Remove now useless remap in UWorldPartition::Initialize

diff --git a/Engine/Source/Editor/UnrealEd/Private/EditorServer.cpp b/Engine/Source/Editor/UnrealEd/Private/EditorServer.cpp
--- a/Engine/Source/Editor/UnrealEd/Private/EditorServer.cpp
+++ b/Engine/Source/Editor/UnrealEd/Private/EditorServer.cpp
@@ -2678,6 +2678,12 @@ bool UEditorEngine::Map_Load(const TCHAR* Str, FOutputDevice& Ar)
 					//now load the map into the package created above
 					const FName WorldPackageFName = WorldPackage->GetFName();
 					UWorld::WorldTypePreLoadMap.FindOrAdd(WorldPackageFName) = EWorldType::Editor;
+					
+					// SoftObjectPaths: Specific case for new maps (/Temp/Untitled) where we need to remap the AssetPath and not just the Package name because the World gets renamed (See UWorld::PostLoad)
+					const FString ShortPackageName = FPackageName::GetShortName(LongTempFname);
+					const FString ShortWorldPackageName = FPackageName::GetShortName(WorldPackage);
+					WorldPackageInstancingContext.AddMapping(*WriteToString<256>(LongTempFname, TEXT("."), ShortPackageName), *WriteToString<256>(WorldPackage->GetName(), TEXT("."), ShortWorldPackageName ));
+
 					WorldPackage = LoadPackage( WorldPackage, *LongTempFname, LoadFlags, nullptr /* InReaderOverride */, &WorldPackageInstancingContext);
 					WorldPackage->SetPackageFlags(PKG_NewlyCreated);
 					UWorld::WorldTypePreLoadMap.Remove(WorldPackageFName);
diff --git a/Engine/Source/Runtime/CoreUObject/Private/UObject/LinkerLoad.cpp b/Engine/Source/Runtime/CoreUObject/Private/UObject/LinkerLoad.cpp
--- a/Engine/Source/Runtime/CoreUObject/Private/UObject/LinkerLoad.cpp
+++ b/Engine/Source/Runtime/CoreUObject/Private/UObject/LinkerLoad.cpp
@@ -958,12 +958,7 @@ FLinkerLoad::FLinkerLoad(UPackage* InParent, const FPackagePath& InPackagePath,
 	FName PackageNameToLoad = GetPackagePath().GetPackageFName();
 	if (LinkerRoot->GetFName() != PackageNameToLoad)
 	{
-		PackageNameToLoad.ToString(InstancedPackageSourceName);
-		LinkerRoot->GetFName().ToString(InstancedPackageInstanceName);
-
-#if WITH_EDITOR
 		InstancingContext.AddMapping(PackageNameToLoad, LinkerRoot->GetFName());
-#endif
 	}
 }
 
@@ -6196,36 +6191,39 @@ bool FLinkerLoad::FinishExternalReadDependencies(double InTimeLimit)
 
 bool FLinkerLoad::IsContextInstanced() const
 {
-#if WITH_EDITOR
 	return InstancingContext.IsInstanced();
-#else
-	return false;
-#endif
+}
+
+bool FLinkerLoad::IsSoftObjectRemappingEnabled() const
+{
+	return IsContextInstanced() && InstancingContext.GetSoftObjectPathRemappingEnabled();
 }
 
 FName FLinkerLoad::InstancingContextRemap(FName ObjectName) const
 {
-#if WITH_EDITOR
 	return InstancingContext.Remap(ObjectName);
-#else
-	return ObjectName;
-#endif
 }
 
 void FLinkerLoad::FixupSoftObjectPathForInstancedPackage(FSoftObjectPath& InOutSoftObjectPath)
 {
-	if (InstancedPackageSourceName.Len() > 0 && InstancedPackageInstanceName.Len() > 0)
+	if (IsSoftObjectRemappingEnabled())
 	{
-		FNameBuilder TmpSoftObjectPathBuilder;
-		InOutSoftObjectPath.ToString(TmpSoftObjectPathBuilder);
-
-		FStringView InstancedPackageSourceNameView = InstancedPackageSourceName.ToView();
-		FStringView TmpSoftObjectPathView = TmpSoftObjectPathBuilder.ToView();
-
-		if (TmpSoftObjectPathView.StartsWith(InstancedPackageSourceNameView) && (TmpSoftObjectPathView.Len() == InstancedPackageSourceNameView.Len() || TmpSoftObjectPathView[InstancedPackageSourceNameView.Len()] == TEXT('.')))
+		// Try remapping AssetPathName before remapping LongPackageName
+		if (FName RemappedAssetPath = InstancingContextRemap(InOutSoftObjectPath.GetAssetPathName()); RemappedAssetPath != InOutSoftObjectPath.GetAssetPathName())
 		{
-			TmpSoftObjectPathBuilder.ReplaceAt(0, InstancedPackageSourceNameView.Len(), InstancedPackageInstanceName.ToView());
-			InOutSoftObjectPath.SetPath(TmpSoftObjectPathBuilder.ToView());
+			InOutSoftObjectPath.SetAssetPathName(RemappedAssetPath);
+		}
+		else
+		{
+			FName LongPackageName = InOutSoftObjectPath.GetLongPackageFName();
+			if (FName RemappedPackage = InstancingContextRemap(LongPackageName); RemappedPackage != LongPackageName)
+			{
+				FNameBuilder TmpSoftObjectPathBuilder;
+				InOutSoftObjectPath.ToString(TmpSoftObjectPathBuilder);
+
+				TmpSoftObjectPathBuilder.ReplaceAt(0, LongPackageName.GetStringLength(), RemappedPackage.ToString());
+				InOutSoftObjectPath.SetPath(TmpSoftObjectPathBuilder.ToView());
+			}
 		}
 	}
 }
diff --git a/Engine/Source/Runtime/CoreUObject/Public/UObject/LinkerInstancingContext.h b/Engine/Source/Runtime/CoreUObject/Public/UObject/LinkerInstancingContext.h
--- a/Engine/Source/Runtime/CoreUObject/Public/UObject/LinkerInstancingContext.h
+++ b/Engine/Source/Runtime/CoreUObject/Public/UObject/LinkerInstancingContext.h
@@ -12,13 +12,17 @@ class FLinkerInstancingContext
 {
 public:
 	FLinkerInstancingContext() = default;
-	FLinkerInstancingContext(TMap<FName, FName> InInstanceMapping)
+	explicit FLinkerInstancingContext(TMap<FName, FName> InInstanceMapping)
 		: Mapping(MoveTemp(InInstanceMapping))
 	{}
-	FLinkerInstancingContext(TSet<FName> InTags)
+	explicit FLinkerInstancingContext(TSet<FName> InTags)
 		: Tags(MoveTemp(InTags))
 	{
 	}
+	explicit FLinkerInstancingContext(bool bInSoftObjectPathRemappingEnabled)
+		: bSoftObjectPathRemappingEnabled(bInSoftObjectPathRemappingEnabled)
+	{
+	}
 
 	bool IsInstanced() const
 	{
@@ -60,6 +64,13 @@ public:
 		return Tags.Contains(Tag);
 	}
 
+	void SetSoftObjectPathRemappingEnabled(bool bInSoftObjectPathRemappingEnabled)
+	{
+		bSoftObjectPathRemappingEnabled = bInSoftObjectPathRemappingEnabled;
+	}
+
+	bool GetSoftObjectPathRemappingEnabled() const { return bSoftObjectPathRemappingEnabled; }
+
 	/** Return the instanced package name for a given instanced outer package and an object package name */
 	static FString GetInstancedPackageName(const FString& InOuterPackageName, const FString& InPackageName)
 	{
@@ -73,4 +84,6 @@ private:
 	TMap<FName, FName> Mapping;
 	/** Tags can be used to determine some loading behavior. */
 	TSet<FName> Tags;
+	/** Remap soft object paths */
+	bool bSoftObjectPathRemappingEnabled = true;
 };
diff --git a/Engine/Source/Runtime/CoreUObject/Public/UObject/LinkerLoad.h b/Engine/Source/Runtime/CoreUObject/Public/UObject/LinkerLoad.h
--- a/Engine/Source/Runtime/CoreUObject/Public/UObject/LinkerLoad.h
+++ b/Engine/Source/Runtime/CoreUObject/Public/UObject/LinkerLoad.h
@@ -193,16 +193,15 @@ private:
 	TUniquePtr<UE::FPackageTrailer> PackageTrailer;
 
 	// Helper function to access the InstancingContext IsInstanced, 
-	// returns false if WITH_EDITOR isn't defined.
 	bool IsContextInstanced() const;
 
+	// Helper function to query if we should do any SoftObjectPath fixup
+	bool IsSoftObjectRemappingEnabled() const;
+
 	// Helper function to access the InstancingContext, 
-	// return ObjectName directly  if WITH_EDITOR isn't defined.
 	FName InstancingContextRemap(FName ObjectName) const;
 
-	/** Set when the package is being loaded as an instance; empty otherwise. */
-	FNameBuilder InstancedPackageSourceName;
-	FNameBuilder InstancedPackageInstanceName;
+	/** Remaps SoftObjectPaths using InstancingContextRemap. */
 	void FixupSoftObjectPathForInstancedPackage(FSoftObjectPath& InOutSoftObjectPath);
 
 protected:
diff --git a/Engine/Source/Runtime/Engine/Private/World.cpp b/Engine/Source/Runtime/Engine/Private/World.cpp
--- a/Engine/Source/Runtime/Engine/Private/World.cpp
+++ b/Engine/Source/Runtime/Engine/Private/World.cpp
@@ -1223,10 +1223,11 @@ void UWorld::PostLoad()
 #if WITH_EDITOR
 	if (GIsEditor)
 	{
-		if (!GetOutermost()->HasAnyPackageFlags(PKG_PlayInEditor))
+		// Avoid renaming PIE worlds and Instanced worlds except for the new maps (/Temp/Untitled) to preserve naming behavior
+		if (!GetPackage()->HasAnyPackageFlags(PKG_PlayInEditor) && (!IsInstanced() || GetPackage()->GetPathName().StartsWith(TEXT("/Temp/Untitled"))))
 		{
 			// Needed for VER_UE4_WORLD_NAMED_AFTER_PACKAGE. If this file was manually renamed outside of the editor, this is needed anyway
-			const FString ShortPackageName = FPackageName::GetLongPackageAssetName(GetOutermost()->GetName());
+			const FString ShortPackageName = FPackageName::GetLongPackageAssetName(GetPackage()->GetName());
 			OriginalWorldName = GetFName();
 			if (GetName() != ShortPackageName)
 			{
diff --git a/Engine/Source/Runtime/Engine/Private/WorldPartition/WorldPartition.cpp b/Engine/Source/Runtime/Engine/Private/WorldPartition/WorldPartition.cpp
--- a/Engine/Source/Runtime/Engine/Private/WorldPartition/WorldPartition.cpp
+++ b/Engine/Source/Runtime/Engine/Private/WorldPartition/WorldPartition.cpp
@@ -424,7 +424,9 @@ void UWorldPartition::Initialize(UWorld* InWorld, const FTransform& InTransform)
 		if (bIsInstanced)
 		{
 			InstancingContext.AddMapping(PackageName, LevelPackage->GetFName());
-			InstancingSoftObjectPathFixupArchive.Reset(new FSoftObjectPathFixupArchive(SourceWorldPath, RemappedWorldPath));
+
+			// SoftObjectPaths: Specific case for new maps (/Temp/Untitled) where we need to remap the AssetPath and not just the Package name because the World gets renamed (See UWorld::PostLoad)
+			InstancingContext.AddMapping(*FString::Format(TEXT("{0}.{1}"), { PackageName.ToString(), FPackageName::GetShortName(PackageName) }), *GetWorld()->GetPathName());
 		}
 
 		{
@@ -897,12 +899,11 @@ void UWorldPartition::OnActorDescUpdated(FWorldPartitionActorDesc* ActorDesc)
 	}
 }
 
-bool UWorldPartition::GetInstancingContext(const FLinkerInstancingContext*& OutInstancingContext, FSoftObjectPathFixupArchive*& OutSoftObjectPathFixupArchive) const
+bool UWorldPartition::GetInstancingContext(const FLinkerInstancingContext*& OutInstancingContext) const
 {
 	if (InstancingContext.IsInstanced())
 	{
 		OutInstancingContext = &InstancingContext;
-		OutSoftObjectPathFixupArchive = InstancingSoftObjectPathFixupArchive.Get();
 		return true;
 	}
 	return false;
@@ -1376,7 +1377,6 @@ void UWorldPartition::OnWorldRenamed()
 
 	// World was renamed so existing context is invalid.
 	InstancingContext = FLinkerInstancingContext();
-	InstancingSoftObjectPathFixupArchive.Reset();
 }
 
 void UWorldPartition::RemapSoftObjectPath(FSoftObjectPath& ObjectPath)
diff --git a/Engine/Source/Runtime/Engine/Private/WorldPartition/WorldPartitionActorDesc.cpp b/Engine/Source/Runtime/Engine/Private/WorldPartition/WorldPartitionActorDesc.cpp
--- a/Engine/Source/Runtime/Engine/Private/WorldPartition/WorldPartitionActorDesc.cpp
+++ b/Engine/Source/Runtime/Engine/Private/WorldPartition/WorldPartitionActorDesc.cpp
@@ -485,11 +485,9 @@ AActor* FWorldPartitionActorDesc::Load() const
 	if (ActorPtr.IsExplicitlyNull())
 	{
 		const FLinkerInstancingContext* InstancingContext = nullptr;
-		FSoftObjectPathFixupArchive* SoftObjectPathFixupArchive = nullptr;
-
 		if (Container)
 		{
-			Container->GetInstancingContext(InstancingContext, SoftObjectPathFixupArchive);
+			Container->GetInstancingContext(InstancingContext);
 		}
 
 		UPackage* Package = nullptr;
@@ -507,14 +505,7 @@ AActor* FWorldPartitionActorDesc::Load() const
 		if (Package)
 		{
 			ActorPtr = FindObject<AActor>(nullptr, *ActorPath.ToString());
-			if (AActor* Actor = ActorPtr.Get())
-			{
-				if (SoftObjectPathFixupArchive)
-				{
-					SoftObjectPathFixupArchive->Fixup(Actor);
-				}
-			}
-			else
+			if (!ActorPtr.IsValid())
 			{
 				UE_LOG(LogWorldPartition, Warning, TEXT("Can't load actor guid `%s` ('%s') from package '%s'"), *Guid.ToString(), *GetActorName().ToString(), *ActorPackage.ToString());
 			}
diff --git a/Engine/Source/Runtime/Engine/Private/WorldPartition/WorldPartitionLevelHelper.cpp b/Engine/Source/Runtime/Engine/Private/WorldPartition/WorldPartitionLevelHelper.cpp
--- a/Engine/Source/Runtime/Engine/Private/WorldPartition/WorldPartitionLevelHelper.cpp
+++ b/Engine/Source/Runtime/Engine/Private/WorldPartition/WorldPartitionLevelHelper.cpp
@@ -204,16 +204,36 @@ FString FWorldPartitionLevelHelper::AddActorContainerIDToSubPathString(const FAc
 	return InSubPathString;
 }
 
-FString FWorldPartitionLevelHelper::AddActorContainerIDToActorPath(const FActorContainerID& InContainerID, const FString& InActorPath)
+FString FWorldPartitionLevelHelper::GetContainerPackage(const FActorContainerID& InContainerID, const FString& InPackageName, int32 InPIEInstanceID)
+{
+	FString ContainerPackageName = InPackageName;
+	if (InPIEInstanceID != INDEX_NONE)
+	{
+		ContainerPackageName = UWorld::ConvertToPIEPackageName(ContainerPackageName, InPIEInstanceID);
+	}
+
+	return FString::Printf(TEXT("/Temp%s_%s"), *ContainerPackageName, *InContainerID.ToString());
+}
+
+bool FWorldPartitionLevelHelper::RemapActorPath(const FActorContainerID& InContainerID, const FString& InActorPath, FString& OutActorPath)
 {
 	if (!InContainerID.IsMainContainer())
 	{
 		const FSoftObjectPath SoftObjectPath(InActorPath);
+		const FString LongPackageName = SoftObjectPath.GetLongPackageName();
+			
+		const FString ContainerPackageString = GetContainerPackage(InContainerID, LongPackageName, INDEX_NONE);	
 		const FString NewSubPathString = FWorldPartitionLevelHelper::AddActorContainerIDToSubPathString(InContainerID, SoftObjectPath.GetSubPathString());
-		return FSoftObjectPath(SoftObjectPath.GetAssetPathName(), NewSubPathString).ToString();
+			
+		FNameBuilder NewAssetPathBuilder;
+		SoftObjectPath.GetAssetPathName().ToString(NewAssetPathBuilder);
+		NewAssetPathBuilder.ReplaceAt(0, LongPackageName.Len(), ContainerPackageString);
+
+		OutActorPath = FSoftObjectPath(*NewAssetPathBuilder, NewSubPathString).ToString();
+		return true;
 	}
 
-	return InActorPath;
+	return false;
 }
 
 /**
@@ -307,13 +327,8 @@ bool FWorldPartitionLevelHelper::LoadActors(UWorld* InOwningWorld, ULevel* InDes
 		{
 			check(!PackageObjectMapping.ContainerID.IsMainContainer());
 		
-			FString ContainerPackageName = PackageObjectMapping.ContainerPackage.ToString();
-			if (InDestLevel && InDestLevel->GetPackage()->GetPIEInstanceID() != INDEX_NONE)
-			{
-				ContainerPackageName = UWorld::ConvertToPIEPackageName(ContainerPackageName, InDestLevel->GetPackage()->GetPIEInstanceID());
-			}
-			
-			const FName ContainerPackageInstanceName(*FString::Printf(TEXT("/Temp%s_%s"), *ContainerPackageName, *PackageObjectMapping.ContainerID.ToString()));
+			const int32 PIEInstanceID = InDestLevel ? InDestLevel->GetPackage()->GetPIEInstanceID() : INDEX_NONE;
+			const FName ContainerPackageInstanceName(GetContainerPackage(PackageObjectMapping.ContainerID, PackageObjectMapping.ContainerPackage.ToString(), PIEInstanceID));
 
 			FLinkerInstancingContext& NewContext = LinkerInstancingContexts.Add(PackageObjectMapping.ContainerID);
 			NewContext.AddTag(ULevel::DontLoadExternalObjectsTag);
@@ -400,7 +415,7 @@ bool FWorldPartitionLevelHelper::LoadActors(UWorld* InOwningWorld, ULevel* InDes
 					// to a Cell in the StreamingPolicy (this relies on the fact that the _DUP package doesn't get fixed up)
 					FSoftObjectPathFixupArchive FixupArchive([&](FSoftObjectPath& Value)
 					{
-						if (!Value.IsNull() && Value.GetAssetPathString().Equals(SourceWorldPath, ESearchCase::IgnoreCase))
+						if (!Value.IsNull() && Value.GetAssetPathString().Equals(RemappedWorldPath, ESearchCase::IgnoreCase))
 						{
 							Value.SetSubPathString(AddActorContainerIDToSubPathString(PackageObjectMapping->ContainerID, Value.GetSubPathString()));
 						}
diff --git a/Engine/Source/Runtime/Engine/Private/WorldPartition/WorldPartitionLevelStreamingDynamic.cpp b/Engine/Source/Runtime/Engine/Private/WorldPartition/WorldPartitionLevelStreamingDynamic.cpp
--- a/Engine/Source/Runtime/Engine/Private/WorldPartition/WorldPartitionLevelStreamingDynamic.cpp
+++ b/Engine/Source/Runtime/Engine/Private/WorldPartition/WorldPartitionLevelStreamingDynamic.cpp
@@ -282,7 +282,9 @@ bool UWorldPartitionLevelStreamingDynamic::IssueLoadRequests()
 	check(!bLoadRequestInProgress);
 	bLoadSucceeded = false;
 	bLoadRequestInProgress = true;
-	FLinkerInstancingContext InstancingContext;
+	// Don't do SoftObjectPath remapping for PersistentLevel actors because references can end up in different cells
+	const bool bSoftObjectRemappingEnabled = false;
+	FLinkerInstancingContext InstancingContext(bSoftObjectRemappingEnabled);
 	UPackage* RuntimePackage = RuntimeLevel->GetPackage();
 	InstancingContext.AddMapping(OriginalLevelPackageName, RuntimePackage->GetFName());
 
diff --git a/Engine/Source/Runtime/Engine/Private/WorldPartition/WorldPartitionLevelStreamingPolicy.cpp b/Engine/Source/Runtime/Engine/Private/WorldPartition/WorldPartitionLevelStreamingPolicy.cpp
--- a/Engine/Source/Runtime/Engine/Private/WorldPartition/WorldPartitionLevelStreamingPolicy.cpp
+++ b/Engine/Source/Runtime/Engine/Private/WorldPartition/WorldPartitionLevelStreamingPolicy.cpp
@@ -77,24 +77,27 @@ void UWorldPartitionLevelStreamingPolicy::PrepareActorToCellRemapping()
 	// Build Actor-to-Cell remapping
 	for (const UWorldPartitionRuntimeCell* Cell : StreamingCells)
 	{
-		if (!Cell->NeedsActorToCellRemapping())
-		{
-			continue;
-		}
-
 		const UWorldPartitionRuntimeLevelStreamingCell* StreamingCell = Cast<const UWorldPartitionRuntimeLevelStreamingCell>(Cell);
 		check(StreamingCell);
 		for (const FWorldPartitionRuntimeCellObjectMapping& CellObjectMap : StreamingCell->GetPackages())
 		{
+			FString RemappedActorPath;
+			FString CellActorPath = CellObjectMap.Path.ToString();
+
 			// Add actor container id to actor path so that we can distinguish between actors of different Level Instances
-			const FString Path = FWorldPartitionLevelHelper::AddActorContainerIDToActorPath(CellObjectMap.ContainerID, CellObjectMap.Path.ToString());
-						
-			ActorToCellRemapping.Add(FName(*Path), StreamingCell->GetFName());
-
-			const int32 LastDotPos = Path.Find(TEXT("."), ESearchCase::IgnoreCase, ESearchDir::FromEnd);
-			check(LastDotPos != INDEX_NONE);
-			
-			SubObjectsToCellRemapping.Add(FName(*Path.Mid(LastDotPos+1)), StreamingCell->GetFName());
+			const bool ActorPathNeedsRemapping = FWorldPartitionLevelHelper::RemapActorPath(CellObjectMap.ContainerID, CellActorPath, RemappedActorPath);
+
+			if (ActorPathNeedsRemapping || Cell->NeedsActorToCellRemapping())
+			{
+				const FString& ActorPath = ActorPathNeedsRemapping ? RemappedActorPath : CellActorPath;
+
+				ActorToCellRemapping.Add(FName(*ActorPath), StreamingCell->GetFName());
+
+				const int32 LastDotPos = ActorPath.Find(TEXT("."), ESearchCase::IgnoreCase, ESearchDir::FromEnd);
+				check(LastDotPos != INDEX_NONE);
+
+				SubObjectsToCellRemapping.Add(FName(*ActorPath.Mid(LastDotPos + 1)), StreamingCell->GetFName());
+			}
 		}
 	}
 }
diff --git a/Engine/Source/Runtime/Engine/Private/WorldPartition/WorldPartitionRuntimeLevelStreamingCell.cpp b/Engine/Source/Runtime/Engine/Private/WorldPartition/WorldPartitionRuntimeLevelStreamingCell.cpp
--- a/Engine/Source/Runtime/Engine/Private/WorldPartition/WorldPartitionRuntimeLevelStreamingCell.cpp
+++ b/Engine/Source/Runtime/Engine/Private/WorldPartition/WorldPartitionRuntimeLevelStreamingCell.cpp
@@ -145,11 +145,17 @@ bool UWorldPartitionRuntimeLevelStreamingCell::PopulateGeneratorPackageForCook(T
 	{
 		FWorldPartitionLevelHelper::FPackageReferencer PackageReferencer;
 		const bool bLoadAsync = false;
-		UWorld* OuterWorld = GetOuterUWorldPartition()->GetTypedOuter<UWorld>();
-		verify(FWorldPartitionLevelHelper::LoadActors(OuterWorld, nullptr, Packages, PackageReferencer, [](bool) {}, bLoadAsync, FLinkerInstancingContext()));
+		UWorldPartition* WorldPartition = GetOuterUWorldPartition();
+		UWorld* OuterWorld = WorldPartition->GetTypedOuter<UWorld>();
+
+		// Don't do SoftObjectPath remapping for PersistentLevel actors because references can end up in different cells
+		const bool bSoftObjectRemappingEnabled = false;
+		verify(FWorldPartitionLevelHelper::LoadActors(OuterWorld, nullptr, Packages, PackageReferencer, [](bool) {}, bLoadAsync, FLinkerInstancingContext(bSoftObjectRemappingEnabled)));
 
 		FWorldPartitionLevelHelper::MoveExternalActorsToLevel(Packages, OuterWorld->PersistentLevel, OutModifiedPackages);
 
+		// Remap needed here for references to Actors that are inside a Container
+		FWorldPartitionLevelHelper::RemapLevelSoftObjectPaths(OuterWorld->PersistentLevel, WorldPartition);
 		// Empty cell's package list (this ensures that no one can rely on cell's content).
 		Packages.Empty();
 	}
@@ -195,7 +201,10 @@ bool UWorldPartitionRuntimeLevelStreamingCell::PopulateGeneratedPackageForCook(U
 		// Load cell Actors
 		FWorldPartitionLevelHelper::FPackageReferencer PackageReferencer;
 		const bool bLoadAsync = false;
-		verify(FWorldPartitionLevelHelper::LoadActors(OuterWorld, nullptr, Packages, PackageReferencer, [](bool) {}, bLoadAsync, FLinkerInstancingContext()));
+
+		// Don't do SoftObjectPath remapping for PersistentLevel actors because references can end up in different cells
+		const bool bSoftObjectRemappingEnabled = false;
+		verify(FWorldPartitionLevelHelper::LoadActors(OuterWorld, nullptr, Packages, PackageReferencer, [](bool) {}, bLoadAsync, FLinkerInstancingContext(bSoftObjectRemappingEnabled)));
 
 		// Create a level and move these actors in it
 		ULevel* NewLevel = FWorldPartitionLevelHelper::CreateEmptyLevelForRuntimeCell(this, OuterWorld, LevelStreaming->GetWorldAsset().ToString(), InPackage);
diff --git a/Engine/Source/Runtime/Engine/Public/WorldPartition/ActorDescContainer.h b/Engine/Source/Runtime/Engine/Public/WorldPartition/ActorDescContainer.h
--- a/Engine/Source/Runtime/Engine/Public/WorldPartition/ActorDescContainer.h
+++ b/Engine/Source/Runtime/Engine/Public/WorldPartition/ActorDescContainer.h
@@ -36,7 +36,7 @@ public:
 	void LoadAllActors(TArray<FWorldPartitionReference>& OutReferences);
 
 	/** Instancing support */
-	virtual bool GetInstancingContext(const FLinkerInstancingContext*& OutInstancingContext, FSoftObjectPathFixupArchive*& OutSoftObjectPathFixupArchive) const { return false; }
+	virtual bool GetInstancingContext(const FLinkerInstancingContext*& OutInstancingContext) const { return false; }
 #endif
 
 	virtual const FTransform& GetInstanceTransform() const { return FTransform::Identity; }
diff --git a/Engine/Source/Runtime/Engine/Public/WorldPartition/WorldPartition.h b/Engine/Source/Runtime/Engine/Public/WorldPartition/WorldPartition.h
--- a/Engine/Source/Runtime/Engine/Public/WorldPartition/WorldPartition.h
+++ b/Engine/Source/Runtime/Engine/Public/WorldPartition/WorldPartition.h
@@ -129,7 +129,7 @@ private:
 	virtual void OnActorDescUpdating(FWorldPartitionActorDesc* ActorDesc) override;
 	virtual void OnActorDescUpdated(FWorldPartitionActorDesc* ActorDesc) override;
 
-	virtual bool GetInstancingContext(const FLinkerInstancingContext*& OutInstancingContext, FSoftObjectPathFixupArchive*& OutSoftObjectPathFixupArchive) const override;
+	virtual bool GetInstancingContext(const FLinkerInstancingContext*& OutInstancingContext) const override;
 #endif
 
 public:
@@ -277,7 +277,6 @@ private:
 
 #if WITH_EDITORONLY_DATA
 	FLinkerInstancingContext InstancingContext;
-	TUniquePtr<FSoftObjectPathFixupArchive> InstancingSoftObjectPathFixupArchive;
 
 	FWorldPartitionReference WorldDataLayersActor;
 #endif
diff --git a/Engine/Source/Runtime/Engine/Public/WorldPartition/WorldPartitionLevelHelper.h b/Engine/Source/Runtime/Engine/Public/WorldPartition/WorldPartitionLevelHelper.h
--- a/Engine/Source/Runtime/Engine/Public/WorldPartition/WorldPartitionLevelHelper.h
+++ b/Engine/Source/Runtime/Engine/Public/WorldPartition/WorldPartitionLevelHelper.h
@@ -37,9 +37,9 @@ public:
 	
 	static bool LoadActors(UWorld* InOwningWorld, ULevel* InDestLevel, TArrayView<FWorldPartitionRuntimeCellObjectMapping> InActorPackages, FPackageReferencer& InPackageReferencer, TFunction<void(bool)> InCompletionCallback, bool bInLoadAsync, FLinkerInstancingContext InOutInstancingContext);
 	
-	static FString AddActorContainerIDToActorPath(const FActorContainerID& InContainerID, const FString& InActorPath);
+	static bool RemapActorPath(const FActorContainerID& InContainerID, const FString& InActorPath, FString& OutActorPath);
 	static FString AddActorContainerIDToSubPathString(const FActorContainerID& InContainerID, const FString& InSubPathString);
-
+	static FString GetContainerPackage(const FActorContainerID& InContainerID, const FString& InPackageName, int32 InPIEInstanceID);
 private:
 	FWorldPartitionLevelHelper();
 
