Cleaning up replay header usage

diff --git a/Engine/Source/Runtime/Engine/Private/Components/CharacterMovementComponent.cpp b/Engine/Source/Runtime/Engine/Private/Components/CharacterMovementComponent.cpp
--- a/Engine/Source/Runtime/Engine/Private/Components/CharacterMovementComponent.cpp
+++ b/Engine/Source/Runtime/Engine/Private/Components/CharacterMovementComponent.cpp
@@ -27,10 +27,7 @@
 #include "Misc/App.h"
 #include "CharacterMovementComponentAsync.h"
 #include "PBDRigidsSolver.h"
-
-#include "Engine/DemoNetDriver.h"
 #include "Engine/NetworkObjectList.h"
-
 #include "Net/PerfCountersHelpers.h"
 #include "ProfilingDebugging/CsvProfiler.h"
 #if UE_WITH_IRIS
diff --git a/Engine/Source/Runtime/Engine/Private/PackageMapClient.cpp b/Engine/Source/Runtime/Engine/Private/PackageMapClient.cpp
--- a/Engine/Source/Runtime/Engine/Private/PackageMapClient.cpp
+++ b/Engine/Source/Runtime/Engine/Private/PackageMapClient.cpp
@@ -24,7 +24,6 @@
 #include "HAL/LowLevelMemTracker.h"
 #include "HAL/LowLevelMemStats.h"
 #include "Net/Core/Trace/NetTrace.h"
-#include "Engine/DemoNetDriver.h"
 #include "Serialization/MemoryReader.h"
 #include "Serialization/MemoryWriter.h"
 #include "Components/ChildActorComponent.h"
@@ -687,11 +686,7 @@ bool UPackageMapClient::SerializeNewActor(FArchive& Ar, class UActorChannel *Cha
 	else if ( Ar.IsLoading() && Actor == NULL )
 	{
 		// Do not log a warning during replay, since this is a valid case
-		UDemoNetDriver* DemoNetDriver = Cast<UDemoNetDriver>(Connection->Driver);
-		if (DemoNetDriver == nullptr)
-		{
-			UE_LOG( LogNetPackageMap, Log, TEXT( "SerializeNewActor: Failed to find static actor: FullNetGuidPath: %s, Channel: %d" ), *GuidCache->FullNetGUIDPath( NetGUID ), Channel->ChIndex );
-		}
+		UE_CLOG(!Connection->IsReplay(), LogNetPackageMap, Log, TEXT("SerializeNewActor: Failed to find static actor: FullNetGuidPath: %s, Channel: %d"), *GuidCache->FullNetGUIDPath(NetGUID), Channel->ChIndex);
 
 		if (UE::Net::FilterGuidRemapping != 0)
 		{
diff --git a/Engine/Source/Runtime/Engine/Private/PhysicsEngine/PhysicsReplication.cpp b/Engine/Source/Runtime/Engine/Private/PhysicsEngine/PhysicsReplication.cpp
--- a/Engine/Source/Runtime/Engine/Private/PhysicsEngine/PhysicsReplication.cpp
+++ b/Engine/Source/Runtime/Engine/Private/PhysicsEngine/PhysicsReplication.cpp
@@ -18,8 +18,6 @@
 #include "Physics/PhysicsInterfaceCore.h"
 #include "Physics/PhysScene_PhysX.h"
 #include "Components/SkeletalMeshComponent.h"
-#include "Engine/DemoNetDriver.h"
-
 #include "Chaos/ChaosMarshallingManager.h"
 #include "PhysicsProxy/SingleParticlePhysicsProxy.h"
 #include "PBDRigidsSolver.h"
@@ -359,7 +357,7 @@ bool FPhysicsReplication::ApplyRigidBodyState(float DeltaSeconds, FBodyInstance*
 			{
 				UE_LOG(LogTemp, Warning, TEXT("Simulated HARD SNAP - \nCurrent Pos - %s, Target Pos - %s\n CurrentState.LinVel - %s, New Lin Vel - %s\nTarget Extrapolation Delta - %s, Is Replay? - %d, Is Asleep - %d, Prev Progress - %f, Prev Similarity - %f"),
 					*CurrentState.Position.ToString(), *TargetPos.ToString(), *CurrentState.LinVel.ToString(), *NewState.LinVel.ToString(),
-					*ExtrapolationDeltaPos.ToString(), (GetOwningWorld()->GetDemoNetDriver() && GetOwningWorld()->GetDemoNetDriver()->IsPlaying()), !BI->IsInstanceAwake(), PrevProgress, PrevSimilarity);
+					*ExtrapolationDeltaPos.ToString(), GetOwningWorld()->IsPlayingReplay(), !BI->IsInstanceAwake(), PrevProgress, PrevSimilarity);
 				if (bDidHardSnap)
 				{
 					*bDidHardSnap = true;
