fix missing include in Actor

diff --git a/Engine/Source/Runtime/Engine/Private/Actor.cpp b/Engine/Source/Runtime/Engine/Private/Actor.cpp
--- a/Engine/Source/Runtime/Engine/Private/Actor.cpp
+++ b/Engine/Source/Runtime/Engine/Private/Actor.cpp
@@ -56,6 +56,7 @@
 #include "GameFramework/InputSettings.h"
 #include "Algo/AnyOf.h"
 #include "PrimitiveSceneProxy.h"
+#include "UObject/MetaData.h"
 #include "Engine/DamageEvents.h"
 #include "Engine/OverlapInfo.h"
 
