CIS fix for non-unity builds

diff --git a/Engine/Source/Runtime/Engine/Public/Tests/AutomationCommon.h b/Engine/Source/Runtime/Engine/Public/Tests/AutomationCommon.h
--- a/Engine/Source/Runtime/Engine/Public/Tests/AutomationCommon.h
+++ b/Engine/Source/Runtime/Engine/Public/Tests/AutomationCommon.h
@@ -5,6 +5,7 @@
 #include "CoreMinimal.h"
 #include "Misc/AutomationTest.h"
 
+class AActor;
 class SWindow;
 class SWidget;
 
