Move FillComponentSpaceTransforms function and ISPC logic from SkeletalMeshComponent into SkinnedAsset so it can be generic and shared among different types of skinned assets.

diff --git a/Engine/Source/Runtime/Engine/Classes/Components/SkeletalMeshComponent.h b/Engine/Source/Runtime/Engine/Classes/Components/SkeletalMeshComponent.h
--- a/Engine/Source/Runtime/Engine/Classes/Components/SkeletalMeshComponent.h
+++ b/Engine/Source/Runtime/Engine/Classes/Components/SkeletalMeshComponent.h
@@ -2340,13 +2340,6 @@ private:
 	/** Duplicates cached transforms/curves and performs interpolation */
 	void ParallelDuplicateAndInterpolate(FAnimationEvaluationContext& InAnimEvaluationContext);
 
-	/**
-	* Take the BoneSpaceTransforms array (translation vector, rotation quaternion and scale vector) and update the array of component-space bone transformation matrices (ComponentSpaceTransforms).
-	* It will work down hierarchy multiplying the component-space transform of the parent by the relative transform of the child.
-	* This code also applies any per-bone rotators etc. as part of the composition process
-	*/
-	void FillComponentSpaceTransforms(const USkeletalMesh* InSkeletalMesh, const TArray<FTransform>& InBoneSpaceTransforms, TArray<FTransform>& OutComponentSpaceTransforms) const;
-
 	bool DoAnyPhysicsBodiesHaveWeight() const;
 
 	virtual void RefreshMorphTargets() override;
diff --git a/Engine/Source/Runtime/Engine/Classes/Engine/SkinnedAsset.h b/Engine/Source/Runtime/Engine/Classes/Engine/SkinnedAsset.h
--- a/Engine/Source/Runtime/Engine/Classes/Engine/SkinnedAsset.h
+++ b/Engine/Source/Runtime/Engine/Classes/Engine/SkinnedAsset.h
@@ -207,6 +207,13 @@ public:
 	virtual uint32 GetVertexBufferFlags() const
 	{ return GetHasVertexColors() ? ESkeletalMeshVertexFlags::HasVertexColors : ESkeletalMeshVertexFlags::None; }
 
+	/**
+	 * Take the BoneSpaceTransforms array (translation vector, rotation quaternion and scale vector) and update the array of component-space bone transformation matrices (ComponentSpaceTransforms).
+	 * It will work down hierarchy multiplying the component-space transform of the parent by the relative transform of the child.
+	 * This code also applies any per-bone rotators etc. as part of the composition process
+	 */
+	void FillComponentSpaceTransforms(const TArray<FTransform>& InBoneSpaceTransforms, const TArray<FBoneIndexType>& InFillComponentSpaceTransformsRequiredBones, TArray<FTransform>& OutComponentSpaceTransforms) const;
+
 	//~ Begin UObject Interface
 	/**
 	* This will return detail info about this specific object. (e.g. AudioComponent will return the name of the cue,
diff --git a/Engine/Source/Runtime/Engine/Private/Components/SkeletalMeshComponent.cpp b/Engine/Source/Runtime/Engine/Private/Components/SkeletalMeshComponent.cpp
--- a/Engine/Source/Runtime/Engine/Private/Components/SkeletalMeshComponent.cpp
+++ b/Engine/Source/Runtime/Engine/Private/Components/SkeletalMeshComponent.cpp
@@ -49,11 +49,6 @@
 
 LLM_DEFINE_TAG(SkeletalMesh_TransformData);
 
-#if INTEL_ISPC
-#include "SkeletalMeshComponent.ispc.generated.h"
-static_assert(sizeof(ispc::FTransform) == sizeof(FTransform), "sizeof(ispc::FTransform) != sizeof(FTransform)");
-#endif
-
 #define LOCTEXT_NAMESPACE "SkeletalMeshComponent"
 
 TAutoConsoleVariable<int32> CVarUseParallelAnimationEvaluation(TEXT("a.ParallelAnimEvaluation"), 1, TEXT("If 1, animation evaluation will be run across the task graph system. If 0, evaluation will run purely on the game thread"));
@@ -66,16 +61,17 @@ static TAutoConsoleVariable<float> CVarStallParallelAnimation(
 	0.0f,
 	TEXT("Sleep for the given time in each parallel animation task. Time is given in ms. This is a debug option used for critical path analysis and forcing a change in the critical path."));
 
+// Deprecated. Please switch to ANIM_SKINNED_ASSET_ISPC_ENABLED_DEFAULT and "a.SkinnedAsset.ISPC", see SkinnedAsset.cpp.
 #if !defined(ANIM_SKELETAL_MESH_ISPC_ENABLED_DEFAULT)
 #define ANIM_SKELETAL_MESH_ISPC_ENABLED_DEFAULT 1
 #endif
 
 // Support run-time toggling on supported platforms in non-shipping configurations
 #if !INTEL_ISPC || UE_BUILD_SHIPPING
-static constexpr bool bAnim_SkeletalMesh_ISPC_Enabled = INTEL_ISPC && ANIM_SKELETAL_MESH_ISPC_ENABLED_DEFAULT;
+bool bAnim_SkeletalMesh_ISPC_Enabled = INTEL_ISPC && ANIM_SKELETAL_MESH_ISPC_ENABLED_DEFAULT;
 #else
-static bool bAnim_SkeletalMesh_ISPC_Enabled = ANIM_SKELETAL_MESH_ISPC_ENABLED_DEFAULT;
-static FAutoConsoleVariableRef CVarAnimSkeletalMeshISPCEnabled(TEXT("a.SkeletalMesh.ISPC"), bAnim_SkeletalMesh_ISPC_Enabled, TEXT("Whether to use ISPC optimizations in animation skeletal mesh components"));
+bool bAnim_SkeletalMesh_ISPC_Enabled = ANIM_SKELETAL_MESH_ISPC_ENABLED_DEFAULT;
+static FAutoConsoleVariableRef CVarAnimSkeletalMeshISPCEnabled(TEXT("a.SkeletalMesh.ISPC"), bAnim_SkeletalMesh_ISPC_Enabled, TEXT("Whether to use ISPC optimizations in animation skeletal mesh components. Deprecated, please switch to a.SkinnedAsset.ISPC"));
 #endif
 
 static TAutoConsoleVariable<int32> CVarCacheLocalSpaceBounds(
@@ -808,7 +804,7 @@ void USkeletalMeshComponent::InitAnim(bool bForceReinit)
 					PRAGMA_DISABLE_DEPRECATION_WARNINGS
 					BoneSpaceTransforms = GetSkeletalMesh()->GetRefSkeleton().GetRefBonePose();
 					//Mini RefreshBoneTransforms (the bit we actually care about)
-					FillComponentSpaceTransforms(GetSkeletalMesh(), BoneSpaceTransforms, GetEditableComponentSpaceTransforms());
+					GetSkeletalMesh()->FillComponentSpaceTransforms(BoneSpaceTransforms, FillComponentSpaceTransformsRequiredBones, GetEditableComponentSpaceTransforms());
 					PRAGMA_ENABLE_DEPRECATION_WARNINGS
 					bNeedToFlipSpaceBaseBuffers = true; // Have updated space bases so need to flip
 					FlipEditableSpaceBases();
@@ -1604,87 +1600,6 @@ static void IntersectBoneIndexArrays(TArray<FBoneIndexType>& Output, const TArra
 	}
 }
 
-
-void USkeletalMeshComponent::FillComponentSpaceTransforms(const USkeletalMesh* InSkeletalMesh, const TArray<FTransform>& InBoneSpaceTransforms, TArray<FTransform>& OutComponentSpaceTransforms) const
-{
-	ANIM_MT_SCOPE_CYCLE_COUNTER(FillComponentSpaceTransforms, !IsInGameThread());
-
-	if( !InSkeletalMesh )
-	{
-		return;
-	}
-
-	// right now all this does is populate DestSpaceBases
-	check( InSkeletalMesh->GetRefSkeleton().GetNum() == InBoneSpaceTransforms.Num());
-	check( InSkeletalMesh->GetRefSkeleton().GetNum() == OutComponentSpaceTransforms.Num());
-
-	const int32 NumBones = InBoneSpaceTransforms.Num();
-
-#if DO_GUARD_SLOW
-	/** Keep track of which bones have been processed for fast look up */
-	TArray<uint8, TInlineAllocator<256>> BoneProcessed;
-	BoneProcessed.AddZeroed(NumBones);
-#endif
-
-	const FTransform* LocalTransformsData = InBoneSpaceTransforms.GetData();
-	FTransform* ComponentSpaceData = OutComponentSpaceTransforms.GetData();
-
-	// First bone (if we have one) is always root bone, and it doesn't have a parent.
-	{
-		check(FillComponentSpaceTransformsRequiredBones.Num() == 0 || FillComponentSpaceTransformsRequiredBones[0] == 0);
-		OutComponentSpaceTransforms[0] = InBoneSpaceTransforms[0];
-
-#if DO_GUARD_SLOW
-		// Mark bone as processed
-		BoneProcessed[0] = 1;
-#endif
-	}
-
-	if (bAnim_SkeletalMesh_ISPC_Enabled)
-	{
-#if INTEL_ISPC
-		ispc::FillComponentSpaceTransforms(
-			(ispc::FTransform*)&ComponentSpaceData[0],
-			(ispc::FTransform*)&LocalTransformsData[0],
-			FillComponentSpaceTransformsRequiredBones.GetData(),
-			(const uint8*)InSkeletalMesh->GetRefSkeleton().GetRefBoneInfo().GetData(),
-			sizeof(FMeshBoneInfo),
-			offsetof(FMeshBoneInfo, ParentIndex),
-			FillComponentSpaceTransformsRequiredBones.Num());
-#endif
-	}
-	else
-	{
-		for (int32 i = 1; i < FillComponentSpaceTransformsRequiredBones.Num(); i++)
-		{
-			const int32 BoneIndex = FillComponentSpaceTransformsRequiredBones[i];
-			FTransform* SpaceBase = ComponentSpaceData + BoneIndex;
-
-			FPlatformMisc::Prefetch(SpaceBase);
-
-#if DO_GUARD_SLOW
-			// Mark bone as processed
-			BoneProcessed[BoneIndex] = 1;
-#endif
-			// For all bones below the root, final component-space transform is relative transform * component-space transform of parent.
-			const int32 ParentIndex = InSkeletalMesh->GetRefSkeleton().GetParentIndex(BoneIndex);
-			FTransform* ParentSpaceBase = ComponentSpaceData + ParentIndex;
-			FPlatformMisc::Prefetch(ParentSpaceBase);
-
-#if DO_GUARD_SLOW
-			// Check the precondition that Parents occur before Children in the RequiredBones array.
-			checkSlow(BoneProcessed[ParentIndex] == 1);
-#endif
-			FTransform::Multiply(SpaceBase, LocalTransformsData + BoneIndex, ParentSpaceBase);
-
-			SpaceBase->NormalizeRotation();
-
-			checkSlow(SpaceBase->IsRotationNormalized());
-			checkSlow(!SpaceBase->ContainsNaN());
-		}
-	}
-}
-
 /** Takes sorted array Base and then adds any elements from sorted array Insert which is missing from it, preserving order.
  * this assumes both arrays are sorted and contain unique bone indices. */
 static void MergeInBoneIndexArrays(TArray<FBoneIndexType>& BaseArray, const TArray<FBoneIndexType>& InsertArray)
@@ -2116,7 +2031,7 @@ void USkeletalMeshComponent::PerformAnimationProcessing(const USkeletalMesh* InS
 		}
 
 		// Fill SpaceBases from LocalAtoms
-		FillComponentSpaceTransforms(InSkeletalMesh, OutBoneSpaceTransforms, OutSpaceBases);
+		InSkeletalMesh->FillComponentSpaceTransforms(OutBoneSpaceTransforms, FillComponentSpaceTransformsRequiredBones, OutSpaceBases);
 	}
 }
 
@@ -2710,7 +2625,7 @@ void USkeletalMeshComponent::PostAnimEvaluation(FAnimationEvaluationContext& Eva
 
 			PRAGMA_DISABLE_DEPRECATION_WARNINGS
 			FAnimationRuntime::LerpBoneTransforms(BoneSpaceTransforms, CachedBoneSpaceTransforms, Alpha, RequiredBones);
-			FillComponentSpaceTransforms(GetSkeletalMesh(), BoneSpaceTransforms, GetEditableComponentSpaceTransforms());
+			GetSkeletalMesh()->FillComponentSpaceTransforms(BoneSpaceTransforms, FillComponentSpaceTransformsRequiredBones, GetEditableComponentSpaceTransforms());
 			PRAGMA_ENABLE_DEPRECATION_WARNINGS
 			// interpolate curve
 			AnimCurves.LerpTo(CachedCurve, Alpha);
@@ -4037,7 +3952,7 @@ void USkeletalMeshComponent::ParallelDuplicateAndInterpolate(FAnimationEvaluatio
 			}
 
 			FAnimationRuntime::LerpBoneTransforms(InAnimEvaluationContext.BoneSpaceTransforms, InAnimEvaluationContext.CachedBoneSpaceTransforms, Alpha, RequiredBones);
-			FillComponentSpaceTransforms(InAnimEvaluationContext.SkeletalMesh, InAnimEvaluationContext.BoneSpaceTransforms, InAnimEvaluationContext.ComponentSpaceTransforms);
+			InAnimEvaluationContext.SkeletalMesh->FillComponentSpaceTransforms(InAnimEvaluationContext.BoneSpaceTransforms, FillComponentSpaceTransformsRequiredBones, InAnimEvaluationContext.ComponentSpaceTransforms);
 
 			// interpolate curve
 			InAnimEvaluationContext.Curve.LerpTo(InAnimEvaluationContext.CachedCurve, Alpha);
diff --git a/Engine/Source/Runtime/Engine/Private/Components/SkeletalMeshComponent.ispc b/Engine/Source/Runtime/Engine/Private/Components/SkeletalMeshComponent.ispc
deleted file mode 100644
diff --git a/Engine/Source/Runtime/Engine/Private/SkinnedAsset.cpp b/Engine/Source/Runtime/Engine/Private/SkinnedAsset.cpp
--- a/Engine/Source/Runtime/Engine/Private/SkinnedAsset.cpp
+++ b/Engine/Source/Runtime/Engine/Private/SkinnedAsset.cpp
@@ -2,6 +2,24 @@
 
 #include "Engine/SkinnedAsset.h"
 #include "SkinnedAssetCompiler.h"
+#include "Animation/AnimStats.h"
+
+#if INTEL_ISPC
+#include "SkinnedAsset.ispc.generated.h"
+static_assert(sizeof(ispc::FTransform) == sizeof(FTransform), "sizeof(ispc::FTransform) != sizeof(FTransform)");
+#endif
+
+#if !defined(ANIM_SKINNED_ASSET_ISPC_ENABLED_DEFAULT)
+#define ANIM_SKINNED_ASSET_ISPC_ENABLED_DEFAULT 1
+#endif
+
+// Support run-time toggling on supported platforms in non-shipping configurations
+#if !INTEL_ISPC || UE_BUILD_SHIPPING
+static constexpr bool bAnim_SkinnedAsset_ISPC_Enabled = INTEL_ISPC && ANIM_SKINNED_ASSET_ISPC_ENABLED_DEFAULT;
+#else
+static bool bAnim_SkinnedAsset_ISPC_Enabled = ANIM_SKINNED_ASSET_ISPC_ENABLED_DEFAULT;
+static FAutoConsoleVariableRef CVarAnimSkinnedAssetISPCEnabled(TEXT("a.SkinnedAsset.ISPC"), bAnim_SkinnedAsset_ISPC_Enabled, TEXT("Whether to use ISPC optimizations on skinned assets"));
+#endif
 
 DEFINE_LOG_CATEGORY_STATIC(LogSkinnedAsset, Log, All);
 
@@ -137,4 +155,101 @@ void USkinnedAsset::WaitUntilAsyncPropertyReleasedInternal(uint64 AsyncPropertie
 		}
 	}
 #endif
+}
+
+extern bool bAnim_SkeletalMesh_ISPC_Enabled;
+static bool IsISPCEnabled()
+{
+	// By default when ANIM_SKELETAL_MESH_ISPC_ENABLED_DEFAULT and CVarAnimSkeletalMeshISPCEnabled are not set, bAnim_SkeletalMesh_ISPC_Enabled should be true,
+	// so when it is false it means one of the two settings are changed. To be backwards compatible, its overridden value takes priority.
+	if (!bAnim_SkeletalMesh_ISPC_Enabled)
+	{
+		UE_LOG(
+			LogSkinnedAsset,
+			Log,
+			TEXT("ANIM_SKELETAL_MESH_ISPC_ENABLED_DEFAULT and CVar 'a.SkeletalMesh.ISPC' are deprecated and will be removed in future releases."
+				 "Please switch to ANIM_SKINNED_ASSET_ISPC_ENABLED_DEFAULT and 'a.SkinnedAsset.ISPC'")
+		);
+		return false;
+	}
+
+	// No old values detected, use the new setting
+	return bAnim_SkinnedAsset_ISPC_Enabled;
+}
+
+void USkinnedAsset::FillComponentSpaceTransforms(const TArray<FTransform>& InBoneSpaceTransforms,
+												 const TArray<FBoneIndexType>& InFillComponentSpaceTransformsRequiredBones, 
+												 TArray<FTransform>& OutComponentSpaceTransforms) const
+{
+	ANIM_MT_SCOPE_CYCLE_COUNTER(FillComponentSpaceTransforms, !IsInGameThread());
+
+	// right now all this does is populate DestSpaceBases
+	check(GetRefSkeleton().GetNum() == InBoneSpaceTransforms.Num());
+	check(GetRefSkeleton().GetNum() == OutComponentSpaceTransforms.Num());
+
+	const int32 NumBones = InBoneSpaceTransforms.Num();
+
+#if DO_GUARD_SLOW
+	/** Keep track of which bones have been processed for fast look up */
+	TArray<uint8, TInlineAllocator<256>> BoneProcessed;
+	BoneProcessed.AddZeroed(NumBones);
+#endif
+
+	const FTransform* LocalTransformsData = InBoneSpaceTransforms.GetData();
+	FTransform* ComponentSpaceData = OutComponentSpaceTransforms.GetData();
+
+	// First bone (if we have one) is always root bone, and it doesn't have a parent.
+	{
+		check(InFillComponentSpaceTransformsRequiredBones.Num() == 0 || InFillComponentSpaceTransformsRequiredBones[0] == 0);
+		OutComponentSpaceTransforms[0] = InBoneSpaceTransforms[0];
+
+#if DO_GUARD_SLOW
+		// Mark bone as processed
+		BoneProcessed[0] = 1;
+#endif
+	}
+
+	if (IsISPCEnabled())
+	{
+#if INTEL_ISPC
+		ispc::FillComponentSpaceTransforms(
+			(ispc::FTransform*)&ComponentSpaceData[0],
+			(ispc::FTransform*)&LocalTransformsData[0],
+			InFillComponentSpaceTransformsRequiredBones.GetData(),
+			(const uint8*)GetRefSkeleton().GetRefBoneInfo().GetData(),
+			sizeof(FMeshBoneInfo),
+			offsetof(FMeshBoneInfo, ParentIndex),
+			InFillComponentSpaceTransformsRequiredBones.Num());
+#endif
+	}
+	else
+	{
+		for (int32 i = 1; i < InFillComponentSpaceTransformsRequiredBones.Num(); i++)
+		{
+			const int32 BoneIndex = InFillComponentSpaceTransformsRequiredBones[i];
+			FTransform* SpaceBase = ComponentSpaceData + BoneIndex;
+
+			FPlatformMisc::Prefetch(SpaceBase);
+
+#if DO_GUARD_SLOW
+			// Mark bone as processed
+			BoneProcessed[BoneIndex] = 1;
+#endif
+			// For all bones below the root, final component-space transform is relative transform * component-space transform of parent.
+			const int32 ParentIndex = GetRefSkeleton().GetParentIndex(BoneIndex);
+			FTransform* ParentSpaceBase = ComponentSpaceData + ParentIndex;
+			FPlatformMisc::Prefetch(ParentSpaceBase);
+
+#if DO_GUARD_SLOW
+			// Check the precondition that Parents occur before Children in the RequiredBones array.
+			checkSlow(BoneProcessed[ParentIndex] == 1);
+#endif
+			FTransform::Multiply(SpaceBase, LocalTransformsData + BoneIndex, ParentSpaceBase);
+
+			SpaceBase->NormalizeRotation();
+
+			checkSlow(SpaceBase->IsRotationNormalized());
+			checkSlow(!SpaceBase->ContainsNaN());
+		}
+	}
 }
\ No newline at end of file
diff --git a/Engine/Source/Runtime/Engine/Private/SkinnedAsset.ispc b/Engine/Source/Runtime/Engine/Private/SkinnedAsset.ispc
new file mode 100644
--- /dev/null
+++ b/Engine/Source/Runtime/Engine/Private/SkinnedAsset.ispc
@@ -0,0 +1,36 @@
+// Copyright Epic Games, Inc. All Rights Reserved.
+
+#define EXPLICIT_VECTOR4 1
+#define EXPLICIT_MATRIX 1
+
+#include "Math/Vector.isph"
+#include "Math/Matrix.isph"
+#include "Math/Quat.isph"
+#include "Math/Transform.isph"
+
+export void FillComponentSpaceTransforms(uniform FTransform ComponentSpaceData[],
+										const uniform FTransform LocalTransformsData[],
+										const uniform unsigned int16 FillComponentSpaceTransformsRequiredBones[],
+										const uniform unsigned int8 * uniform ParentBones,
+										const uniform int ParentBoneSize,
+										const uniform int ParentIndexOffset,
+										const uniform int NumBones)
+{
+	for (uniform int i = 1; i < NumBones; i++)
+	{
+		const uniform int BoneIndex = FillComponentSpaceTransformsRequiredBones[i];
+		uniform FTransform SpaceBase = ComponentSpaceData[BoneIndex];
+
+		// For all bones below the root, final component-space transform is relative transform * component-space transform of parent.
+		const uniform unsigned int8 * uniform ParentIndexPtr = ParentBones + ((ParentBoneSize * BoneIndex) + ParentIndexOffset);
+		const uniform int ParentIndex = *((uniform int * uniform)ParentIndexPtr);
+		const uniform FTransform ParentSpaceBase = ComponentSpaceData[ParentIndex];
+
+		const uniform FTransform LocalTransform = LocalTransformsData[BoneIndex];
+
+		TransformMultiply(SpaceBase, LocalTransform, ParentSpaceBase);
+		SpaceBase.Rotation = VectorNormalizeQuaternion(SpaceBase.Rotation);
+
+		ComponentSpaceData[BoneIndex] = SpaceBase;
+	}
+}
