Clean up FStaticLightingSystemInterface

diff --git a/Engine/Plugins/Experimental/GPULightmass/Source/GPULightmass/Private/GPULightmass.cpp b/Engine/Plugins/Experimental/GPULightmass/Source/GPULightmass/Private/GPULightmass.cpp
--- a/Engine/Plugins/Experimental/GPULightmass/Source/GPULightmass/Private/GPULightmass.cpp
+++ b/Engine/Plugins/Experimental/GPULightmass/Source/GPULightmass/Private/GPULightmass.cpp
@@ -2,7 +2,6 @@
 
 #include "GPULightmass.h"
 #include "ComponentRecreateRenderStateContext.h"
-#include "EngineModule.h"
 #include "Misc/ScopedSlowTask.h"
 #include "Components/SkyLightComponent.h"
 #include "Framework/Notifications/NotificationManager.h"
@@ -13,8 +12,6 @@
 
 #define LOCTEXT_NAMESPACE "StaticLightingSystem"
 
-extern ENGINE_API void ToggleLightmapPreview_GameThread(UWorld* InWorld);
-
 FGPULightmass::FGPULightmass(UWorld* InWorld, FGPULightmassModule* InGPULightmassModule, UGPULightmassSettings* InSettings)
 	: World(InWorld)
 	, GPULightmassModule(InGPULightmassModule)
@@ -195,8 +192,6 @@ void FGPULightmass::OnLightComponentUnregistered(ULightComponentBase* InComponen
 {
 	if (InComponent->GetWorld() != World) return;
 
-	// UE_LOG(LogGPULightmass, Display, TEXT("LightComponent %s on actor %s is being unregistered with GPULightmass"), *InComponent->GetName(), *InComponent->GetOwner()->GetName());
-
 	if (UDirectionalLightComponent* DirectionalLight = Cast<UDirectionalLightComponent>(InComponent))
 	{
 		Scene.RemoveLight(DirectionalLight);
@@ -259,14 +254,10 @@ void FGPULightmass::OnStationaryLightChannelReassigned(ULightComponentBase* InCo
 
 void FGPULightmass::OnPreWorldFinishDestroy(UWorld* InWorld)
 {
-	UE_LOG(LogGPULightmass, Display, TEXT("World %s is being destroyed"), *World->GetName());
-
 	if (InWorld != World) return;
 
-	UE_LOG(LogGPULightmass, Display, TEXT("Removing all delegates (including this one)"), *World->GetName());
-
 	// This calls destructor of FGPULightmass
-	GPULightmassModule->RemoveStaticLightingSystemForWorld(World);
+	GPULightmassModule->RemoveGPULightmassFromWorld(World);
 }
 
 void FGPULightmass::EditorTick()
diff --git a/Engine/Plugins/Experimental/GPULightmass/Source/GPULightmass/Private/GPULightmassModule.cpp b/Engine/Plugins/Experimental/GPULightmass/Source/GPULightmass/Private/GPULightmassModule.cpp
--- a/Engine/Plugins/Experimental/GPULightmass/Source/GPULightmass/Private/GPULightmassModule.cpp
+++ b/Engine/Plugins/Experimental/GPULightmass/Source/GPULightmass/Private/GPULightmassModule.cpp
@@ -2,15 +2,11 @@
 
 #include "GPULightmassModule.h"
 #include "CoreMinimal.h"
-#include "Internationalization/Internationalization.h"
 #include "Modules/ModuleManager.h"
-#include "HAL/IConsoleManager.h"
 #include "RenderingThread.h"
 #include "Interfaces/IPluginManager.h"
 #include "ShaderCore.h"
 #include "GPULightmass.h"
-#include "SceneInterface.h"
-#include "CompactingObjectPool.h"
 
 #define LOCTEXT_NAMESPACE "StaticLightingSystem"
 
@@ -46,7 +42,7 @@ void FGPULightmassModule::ShutdownModule()
 	}
 }
 
-IStaticLightingSystem* FGPULightmassModule::AllocateStaticLightingSystemForWorldWithSettings(UWorld* InWorld, UGPULightmassSettings* Settings)
+FGPULightmass* FGPULightmassModule::CreateGPULightmassForWorld(UWorld* InWorld, UGPULightmassSettings* Settings)
 {
 	check(StaticLightingSystems.Find(InWorld) == nullptr);
 
@@ -61,16 +57,7 @@ IStaticLightingSystem* FGPULightmassModule::AllocateStaticLightingSystemForWorld
 	return GPULightmass;
 }
 
-IStaticLightingSystem* FGPULightmassModule::AllocateStaticLightingSystemForWorld(UWorld* InWorld)
-{
-	// Gather settings from CVars
-	UGPULightmassSettings* Settings = NewObject<UGPULightmassSettings>(GetTransientPackage(), MakeUniqueObjectName(GetTransientPackage(), UGPULightmassSettings::StaticClass()));
-	Settings->GatherSettingsFromCVars();
-
-	return AllocateStaticLightingSystemForWorldWithSettings(InWorld, Settings);
-}
-
-void FGPULightmassModule::RemoveStaticLightingSystemForWorld(UWorld* InWorld)
+void FGPULightmassModule::RemoveGPULightmassFromWorld(UWorld* InWorld)
 {
 	if (StaticLightingSystems.Find(InWorld) != nullptr)
 	{
@@ -107,11 +94,9 @@ void FGPULightmassModule::EditorTick()
 		}
 	}
 
-	for (auto& StaticLightingSystem : FinishedStaticLightingSystems)
+	for (FGPULightmass* StaticLightingSystem : FinishedStaticLightingSystems)
 	{
-		extern ENGINE_API void ToggleLightmapPreview_GameThread(UWorld * InWorld);
-
-		ToggleLightmapPreview_GameThread(StaticLightingSystem->World);
+		StaticLightingSystem->World->GetSubsystem<UGPULightmassSubsystem>()->Stop();
 	}
 }
 
diff --git a/Engine/Plugins/Experimental/GPULightmass/Source/GPULightmass/Private/GPULightmassSettings.cpp b/Engine/Plugins/Experimental/GPULightmass/Source/GPULightmass/Private/GPULightmassSettings.cpp
--- a/Engine/Plugins/Experimental/GPULightmass/Source/GPULightmass/Private/GPULightmassSettings.cpp
+++ b/Engine/Plugins/Experimental/GPULightmass/Source/GPULightmass/Private/GPULightmassSettings.cpp
@@ -13,114 +13,6 @@
 
 #define LOCTEXT_NAMESPACE "StaticLightingSystem"
 
-int32 GGPULightmassShowProgressBars = 1;
-static FAutoConsoleVariableRef CVarGPULightmassShowProgressBars(
-	TEXT("r.GPULightmass.ShowProgressBars"),
-	GGPULightmassShowProgressBars,
-	TEXT("\n"),
-	ECVF_Default
-);
-
-int32 GGPULightmassOnlyBakeWhatYouSee = 0;
-static FAutoConsoleVariableRef CVarGPULightmassOnlyBakeWhatYouSee(
-	TEXT("r.GPULightmass.OnlyBakeWhatYouSee"),
-	GGPULightmassOnlyBakeWhatYouSee,
-	TEXT("\n"),
-	ECVF_Default
-);
-
-int32 GGPULightmassSamplesPerTexel = 512;
-static FAutoConsoleVariableRef CVarGPULightmassSamplesPerTexel(
-	TEXT("r.GPULightmass.SamplesPerTexel"),
-	GGPULightmassSamplesPerTexel,
-	TEXT("\n"),
-	ECVF_Default
-);
-
-int32 GGPULightmassShadowSamplesPerTexel = 512; // 512 samples to reach good image plane stratification. Shadow samples are 100x faster than path samples
-static FAutoConsoleVariableRef CVarGPULightmassShadowSamplesPerTexel(
-	TEXT("r.GPULightmass.ShadowSamplesPerTexel"),
-	GGPULightmassShadowSamplesPerTexel,
-	TEXT("\n"),
-	ECVF_Default
-);
-
-int32 GGPULightmassUseIrradianceCaching = 0;
-static FAutoConsoleVariableRef CVarGPULightmassUseIrradianceCaching(
-	TEXT("r.GPULightmass.IrradianceCaching"),
-	GGPULightmassUseIrradianceCaching,
-	TEXT("\n"),
-	ECVF_Default
-);
-
-int32 GGPULightmassIrradianceCachingQuality = 128;
-static FAutoConsoleVariableRef CVarGPULightmassIrradianceCachingQuality(
-	TEXT("r.GPULightmass.IrradianceCaching.Quality"),
-	GGPULightmassIrradianceCachingQuality,
-	TEXT("\n"),
-	ECVF_Default
-);
-
-float GGPULightmassIrradianceCachingSpacing = 32.0f;
-static FAutoConsoleVariableRef CVarGPULightmassIrradianceCachingSpacing(
-	TEXT("r.GPULightmass.IrradianceCaching.Spacing"),
-	GGPULightmassIrradianceCachingSpacing,
-	TEXT("\n"),
-	ECVF_Default
-);
-
-int32 GGPULightmassVisualizeIrradianceCache = 0;
-static FAutoConsoleVariableRef CVarGPULightmassVisualizeIrradianceCache(
-	TEXT("r.GPULightmass.IrradianceCaching.Visualize"),
-	GGPULightmassVisualizeIrradianceCache,
-	TEXT("\n"),
-	ECVF_Default
-);
-
-int32 GGPULightmassUseFirstBounceRayGuiding = 0;
-static FAutoConsoleVariableRef CVarGPULightmassUseFirstBounceRayGuiding(
-	TEXT("r.GPULightmass.FirstBounceRayGuiding"),
-	GGPULightmassUseFirstBounceRayGuiding,
-	TEXT("\n"),
-	ECVF_Default
-);
-
-int32 GGPULightmassFirstBounceRayGuidingTrialSamples = 128;
-static FAutoConsoleVariableRef CVarGPULightmassFirstBounceRayGuidingTrialSamples(
-	TEXT("r.GPULightmass.FirstBounceRayGuiding.TrialSamples"),
-	GGPULightmassFirstBounceRayGuidingTrialSamples,
-	TEXT("\n"),
-	ECVF_Default
-);
-
-int32 GGPULightmassGPUTilePoolSize = 40;
-static FAutoConsoleVariableRef CVarGPULightmassGPUTilePoolSize(
-	TEXT("r.GPULightmass.System.GPUTilePoolSize"),
-	GGPULightmassGPUTilePoolSize,
-	TEXT("\n"),
-	ECVF_Default
-);
-
-int32 GGPULightmassDenoiseGIOnCompletion = WITH_INTELOIDN;
-#if WITH_INTELOIDN
-static FAutoConsoleVariableRef CVarGGPULightmassDenoiseOnCompletion(
-	TEXT("r.GPULightmass.DenoiseGIOnCompletion"),
-	GGPULightmassDenoiseGIOnCompletion,
-	TEXT("\n"),
-	ECVF_Default
-);
-#endif
-
-int32 GGPULightmassDenoiseGIDuringInteractiveBake = 0;
-#if WITH_INTELOIDN
-static FAutoConsoleVariableRef CVarGGPULightmassDenoiseGIDuringInteractiveBake(
-	TEXT("r.GPULightmass.DenoiseGIDuringInteractiveBake"),
-	GGPULightmassDenoiseGIDuringInteractiveBake,
-	TEXT("\n"),
-	ECVF_Default
-);
-#endif
-
 AGPULightmassSettingsActor::AGPULightmassSettingsActor(const FObjectInitializer& ObjectInitializer)
 	: Super(ObjectInitializer.DoNotCreateDefaultSubobject(TEXT("Sprite")))
 {
@@ -221,40 +113,6 @@ bool UGPULightmassSettings::CanEditChange(const FProperty* InProperty) const
 }
 #endif
 
-void UGPULightmassSettings::GatherSettingsFromCVars()
-{
-	bShowProgressBars = GGPULightmassShowProgressBars == 1;
-
-	if (GGPULightmassOnlyBakeWhatYouSee == 1)
-	{
-		Mode = EGPULightmassMode::BakeWhatYouSee;
-	}
-
-	GISamples = GGPULightmassSamplesPerTexel;
-	StationaryLightShadowSamples = GGPULightmassSamplesPerTexel;
-	
-	bUseIrradianceCaching = GGPULightmassUseIrradianceCaching == 1;
-	IrradianceCacheQuality = GGPULightmassIrradianceCachingQuality;
-	IrradianceCacheSpacing = GGPULightmassIrradianceCachingSpacing;
-	bVisualizeIrradianceCache = GGPULightmassVisualizeIrradianceCache == 1;
-
-	bUseFirstBounceRayGuiding = GGPULightmassUseFirstBounceRayGuiding == 1;
-	FirstBounceRayGuidingTrialSamples = GGPULightmassFirstBounceRayGuidingTrialSamples;
-
-	if (GGPULightmassDenoiseGIOnCompletion == 1)
-	{
-		DenoisingOptions = EGPULightmassDenoisingOptions::OnCompletion;
-	}
-
-	if (GGPULightmassDenoiseGIDuringInteractiveBake == 1)
-	{
-		// Override
-		DenoisingOptions = EGPULightmassDenoisingOptions::DuringInteractivePreview;
-	}
-
-	LightmapTilePoolSize = GGPULightmassGPUTilePoolSize;
-}
-
 void UGPULightmassSubsystem::Initialize(FSubsystemCollectionBase& Collection)
 {
 	UWorld* World = Cast<UWorld>(GetOuter());
@@ -302,7 +160,7 @@ void UGPULightmassSubsystem::Launch()
 
 			FlushRenderingCommands(); // Flush again to execute commands generated by DestroyRenderState_Concurrent()
 
-			IStaticLightingSystem* StaticLightingSystem = GPULightmassModule.AllocateStaticLightingSystemForWorldWithSettings(World, SettingsCopy);
+			FGPULightmass* StaticLightingSystem = GPULightmassModule.CreateGPULightmassForWorld(World, SettingsCopy);
 
 			if (StaticLightingSystem)
 			{
@@ -425,7 +283,7 @@ void UGPULightmassSubsystem::Stop()
 				SubSlowTask.EnterProgressFrame(1, LOCTEXT("UnregisteringComponentsWithStaticLightingSystem", "Unregistering components with static lighting system"));
 			}
 
-			GPULightmassModule.RemoveStaticLightingSystemForWorld(World);
+			GPULightmassModule.RemoveGPULightmassFromWorld(World);
 
 			UE_LOG(LogTemp, Log, TEXT("Static lighting system is removed for world %s."), *World->GetPathName(World->GetOuter()));
 		}
diff --git a/Engine/Plugins/Experimental/GPULightmass/Source/GPULightmass/Private/Scene/Scene.cpp b/Engine/Plugins/Experimental/GPULightmass/Source/GPULightmass/Private/Scene/Scene.cpp
--- a/Engine/Plugins/Experimental/GPULightmass/Source/GPULightmass/Private/Scene/Scene.cpp
+++ b/Engine/Plugins/Experimental/GPULightmass/Source/GPULightmass/Private/Scene/Scene.cpp
@@ -1541,6 +1541,8 @@ void FScene::BackgroundTick()
 		{
 			ApplyFinishedLightmapsToWorld();
 			GPULightmass->World->GetSubsystem<UGPULightmassSubsystem>()->OnLightBuildEnded().Broadcast();
+			// We can't destroy GPULM in its member func
+			// FGPULightmassModule::EditorTick() will handle the destruction
 		}
 	}
 }
diff --git a/Engine/Plugins/Experimental/GPULightmass/Source/GPULightmass/Public/GPULightmassModule.h b/Engine/Plugins/Experimental/GPULightmass/Source/GPULightmass/Public/GPULightmassModule.h
--- a/Engine/Plugins/Experimental/GPULightmass/Source/GPULightmass/Public/GPULightmassModule.h
+++ b/Engine/Plugins/Experimental/GPULightmass/Source/GPULightmass/Public/GPULightmassModule.h
@@ -7,39 +7,7 @@
 #include "Modules/ModuleManager.h"
 #include "Rendering/StaticLightingSystemInterface.h"
 
-
-/**
- * The public interface to this module.  In most cases, this interface is only public to sibling modules 
- * within this plugin.
- */
-class IGPULightmassModule : public IModuleInterface
-{
-
-public:
-
-	/**
-	 * Singleton-like access to this module's interface.  This is just for convenience!
-	 * Beware of calling this during the shutdown phase, though.  Your module might have been unloaded already.
-	 *
-	 * @return Returns singleton instance, loading the module on demand if needed
-	 */
-	static inline IGPULightmassModule& Get()
-	{
-		return FModuleManager::LoadModuleChecked< IGPULightmassModule >( "GPULightmass" );
-	}
-
-	/**
-	 * Checks to see if this module is loaded and ready.  It is only valid to call Get() if IsAvailable() returns true.
-	 *
-	 * @return True if the module is loaded and ready to use
-	 */
-	static inline bool IsAvailable()
-	{
-		return FModuleManager::Get().IsModuleLoaded( "GPULightmass" );
-	}
-};
-
-class FGPULightmassModule : public IGPULightmassModule, public IStaticLightingSystemImpl
+class FGPULightmassModule : public IModuleInterface, public IStaticLightingSystemImpl
 {
 public:
 	void RunSelfTests();
@@ -48,14 +16,12 @@ public:
 	virtual void StartupModule() override;
 	virtual void ShutdownModule() override;
 
-	virtual bool SupportsRealtimePreview() override { return true; }
-
-	virtual IStaticLightingSystem* AllocateStaticLightingSystemForWorld(UWorld* InWorld) override;
-	virtual IStaticLightingSystem* AllocateStaticLightingSystemForWorldWithSettings(UWorld* InWorld, class UGPULightmassSettings* Settings);
-	virtual void RemoveStaticLightingSystemForWorld(UWorld* InWorld) override;
 	virtual IStaticLightingSystem* GetStaticLightingSystemForWorld(UWorld* InWorld) override;
-	virtual void EditorTick() override;
 	virtual bool IsStaticLightingSystemRunning() override;
+	virtual void EditorTick() override;
+	
+	class FGPULightmass* CreateGPULightmassForWorld(UWorld* InWorld, class UGPULightmassSettings* Settings);
+	void RemoveGPULightmassFromWorld(UWorld* InWorld);
 
 	// Due to limitations in our TMap implementation I cannot use TUniquePtr here
 	// But the GPULightmassModule is the only owner of all static lighting systems, and all worlds weak refer to the systems
diff --git a/Engine/Plugins/Experimental/GPULightmass/Source/GPULightmass/Public/GPULightmassSettings.h b/Engine/Plugins/Experimental/GPULightmass/Source/GPULightmass/Public/GPULightmassSettings.h
--- a/Engine/Plugins/Experimental/GPULightmass/Source/GPULightmass/Public/GPULightmassSettings.h
+++ b/Engine/Plugins/Experimental/GPULightmass/Source/GPULightmass/Public/GPULightmassSettings.h
@@ -127,7 +127,6 @@ public:
 	int32 LightmapTilePoolSize = 55;
 
 public:
-	void GatherSettingsFromCVars();
 	void ApplyImmediateSettingsToRunningInstances();
 
 #if WITH_EDITOR
diff --git a/Engine/Source/Runtime/Engine/Private/Rendering/StaticLightingSystemInterface.cpp b/Engine/Source/Runtime/Engine/Private/Rendering/StaticLightingSystemInterface.cpp
--- a/Engine/Source/Runtime/Engine/Private/Rendering/StaticLightingSystemInterface.cpp
+++ b/Engine/Source/Runtime/Engine/Private/Rendering/StaticLightingSystemInterface.cpp
@@ -1,14 +1,10 @@
 // Copyright Epic Games, Inc. All Rights Reserved.
 
 #include "Rendering/StaticLightingSystemInterface.h"
-#include "RenderingThread.h"
 #include "ComponentRecreateRenderStateContext.h"
 #include "Misc/ScopedSlowTask.h"
 #include "Engine/World.h"
 #include "Components/LightComponent.h"
-#include "LandscapeComponent.h"
-#include "UObject/UObjectIterator.h"
-#include "MaterialShared.h"
 
 #define LOCTEXT_NAMESPACE "StaticLightingSystem"
 
@@ -131,168 +127,4 @@ IStaticLightingSystemImpl* FStaticLightingSystemInterface::GetPreferredImplement
 	return nullptr;
 }
 
-bool FStaticLightingSystemInterface::ShouldOperateOnWorld(UWorld* InWorld)
-{
-	check(IsInGameThread());
-
-	static const auto AllowStaticLightingVar = IConsoleManager::Get().FindTConsoleVariableDataInt(TEXT("r.AllowStaticLighting"));
-	const bool bAllowStaticLighting = (!AllowStaticLightingVar || AllowStaticLightingVar->GetValueOnGameThread() != 0);
-
-	// IsEditorWorld() is true also for PIE and EditorPreview, which is not what we want
-	return bAllowStaticLighting;
-}
-
-ENGINE_API void ToggleLightmapPreview_GameThread(UWorld* InWorld)
-{
-	if (FStaticLightingSystemInterface::Get()->ShouldOperateOnWorld(InWorld))
-	{
-		IStaticLightingSystemImpl* StaticLightingSystemImpl = FStaticLightingSystemInterface::Get()->GetPreferredImplementation();
-
-		// Extra check to avoid the cost of FGlobalComponentRecreateRenderStateContext
-		if (StaticLightingSystemImpl && StaticLightingSystemImpl->SupportsRealtimePreview())
-		{
-			// At this point the assumption is GT & RT are synchronized and we can access FScene from GT
-			if (!StaticLightingSystemImpl->GetStaticLightingSystemForWorld(InWorld))
-			{
-				FScopedSlowTask SlowTask(1);
-				SlowTask.MakeDialog();
-				SlowTask.EnterProgressFrame(1, LOCTEXT("StartingStaticLightingSystem", "Starting static lighting system"));
-
-				{
-					FGlobalComponentRecreateRenderStateContext RecreateRenderStateContext; // Implicit FlushRenderingCommands();
-
-					FlushRenderingCommands(); // Flush again to execute commands generated by DestroyRenderState_Concurrent()
-
-					IStaticLightingSystem* StaticLightingSystem = StaticLightingSystemImpl->AllocateStaticLightingSystemForWorld(InWorld);
-
-					if (StaticLightingSystem)
-					{
-						UE_LOG(LogTemp, Log, TEXT("Static lighting system is created for world %s."), *InWorld->GetPathName(InWorld->GetOuter()));
-
-						ULightComponent::ReassignStationaryLightChannels(InWorld, false, NULL);
-#if WITH_EDITOR
-						if (!GIsEditor)
-						{
-							if (GEngine)
-							{
-								GEngine->OnPostEditorTick().AddStatic(&FStaticLightingSystemInterface::GameTick);
-							}
-						}
-#endif // WITH_EDITOR
-						int32 NumPrimitiveComponents = 0;
-						int32 NumLightComponents = 0;
-
-						for (UPrimitiveComponent* Component : TObjectRange<UPrimitiveComponent>(RF_ClassDefaultObject | RF_ArchetypeObject, true, EInternalObjectFlags::Garbage))
-						{
-							if (Component->HasValidSettingsForStaticLighting(false))
-							{
-								NumPrimitiveComponents++;
-							}
-						}
-
-						for (ULightComponentBase* Component : TObjectRange<ULightComponentBase>(RF_ClassDefaultObject | RF_ArchetypeObject, true, EInternalObjectFlags::Garbage))
-						{
-							if (Component->bAffectsWorld && Component->HasStaticShadowing())
-							{
-								NumLightComponents++;
-							}
-						}
-
-						FScopedSlowTask SubSlowTask(NumPrimitiveComponents + NumLightComponents, LOCTEXT("RegisteringComponentsWithStaticLightingSystem", "Registering components with static lighting system"));
-						SubSlowTask.MakeDialog();
-
-						for (UPrimitiveComponent* Component : TObjectRange<UPrimitiveComponent>(RF_ClassDefaultObject | RF_ArchetypeObject, true, EInternalObjectFlags::Garbage))
-						{
-							if (Component->HasValidSettingsForStaticLighting(false))
-							{
-								FStaticLightingSystemInterface::OnPrimitiveComponentRegistered.Broadcast(Component);
-
-								SubSlowTask.EnterProgressFrame(1, LOCTEXT("RegisteringComponentsWithStaticLightingSystem", "Registering components with static lighting system"));
-							}
-						}
-
-						for (ULightComponentBase* Component : TObjectRange<ULightComponentBase>(RF_ClassDefaultObject | RF_ArchetypeObject, true, EInternalObjectFlags::Garbage))
-						{
-							if (Component->bAffectsWorld && Component->HasStaticShadowing())
-							{
-								FStaticLightingSystemInterface::OnLightComponentRegistered.Broadcast(Component);
-
-								SubSlowTask.EnterProgressFrame(1, LOCTEXT("RegisteringComponentsWithStaticLightingSystem", "Registering components with static lighting system"));
-							}
-						}
-					}
-					else
-					{
-						UE_LOG(LogTemp, Log, TEXT("Tried to create static lighting system for world %s, but failed"), *InWorld->GetPathName(InWorld->GetOuter()));
-					}
-				}
-
-				FlushRenderingCommands(); // Flush commands generated by ~FGlobalComponentRecreateRenderStateContext();
-			}
-			else
-			{
-				FScopedSlowTask SlowTask(1);
-				SlowTask.MakeDialog();
-				SlowTask.EnterProgressFrame(1, LOCTEXT("RemovingStaticLightingSystem", "Removing static lighting system"));
-
-				{
-					FGlobalComponentRecreateRenderStateContext RecreateRenderStateContext; // Implicit FlushRenderingCommands();
-
-					FlushRenderingCommands(); // Flush again to execute commands generated by DestroyRenderState_Concurrent()
-
-					int32 NumPrimitiveComponents = 0;
-					int32 NumLightComponents = 0;
-
-					for (UPrimitiveComponent* Component : TObjectRange<UPrimitiveComponent>(RF_ClassDefaultObject | RF_ArchetypeObject, true, EInternalObjectFlags::Garbage))
-					{
-						NumPrimitiveComponents++;
-					}
-
-					for (ULightComponentBase* Component : TObjectRange<ULightComponentBase>(RF_ClassDefaultObject | RF_ArchetypeObject, true, EInternalObjectFlags::Garbage))
-					{
-						NumLightComponents++;
-					}
-
-					FScopedSlowTask SubSlowTask(NumPrimitiveComponents + NumLightComponents, LOCTEXT("UnregisteringComponentsWithStaticLightingSystem", "Unregistering components with static lighting system"));
-
-					// Unregister all landscapes first to prevent grass picking up landscape lightmaps
-					for (ULandscapeComponent* Component : TObjectRange<ULandscapeComponent>(RF_ClassDefaultObject | RF_ArchetypeObject, true, EInternalObjectFlags::Garbage))
-					{
-						FStaticLightingSystemInterface::OnPrimitiveComponentUnregistered.Broadcast(Component);
-					}
-
-					for (UPrimitiveComponent* Component : TObjectRange<UPrimitiveComponent>(RF_ClassDefaultObject | RF_ArchetypeObject, true, EInternalObjectFlags::Garbage))
-					{
-						FStaticLightingSystemInterface::OnPrimitiveComponentUnregistered.Broadcast(Component);
-
-						SubSlowTask.EnterProgressFrame(1, LOCTEXT("UnregisteringComponentsWithStaticLightingSystem", "Unregistering components with static lighting system"));
-					}
-
-					for (ULightComponentBase* Component : TObjectRange<ULightComponentBase>(RF_ClassDefaultObject | RF_ArchetypeObject, true, EInternalObjectFlags::Garbage))
-					{
-						FStaticLightingSystemInterface::OnLightComponentUnregistered.Broadcast(Component);
-
-						SubSlowTask.EnterProgressFrame(1, LOCTEXT("UnregisteringComponentsWithStaticLightingSystem", "Unregistering components with static lighting system"));
-					}
-
-					StaticLightingSystemImpl->RemoveStaticLightingSystemForWorld(InWorld);
-
-					UE_LOG(LogTemp, Log, TEXT("Static lighting system is removed for world %s."), *InWorld->GetPathName(InWorld->GetOuter()));
-				}
-
-				FlushRenderingCommands(); // Flush commands generated by ~FGlobalComponentRecreateRenderStateContext();
-			}
-		}
-	}
-	else
-	{
-		UE_LOG(LogTemp, Log, TEXT("We should not operate on world %s."), *InWorld->GetPathName(InWorld->GetOuter()));
-	}
-}
-
-FAutoConsoleCommandWithWorld GToggleLightmapPreviewCmd(
-	TEXT("ToggleLightmapPreview"),
-	TEXT("Toggles lightmap preview in editor"),
-	FConsoleCommandWithWorldDelegate::CreateStatic(ToggleLightmapPreview_GameThread));
-
 #undef LOCTEXT_NAMESPACE
diff --git a/Engine/Source/Runtime/Engine/Public/Rendering/StaticLightingSystemInterface.h b/Engine/Source/Runtime/Engine/Public/Rendering/StaticLightingSystemInterface.h
--- a/Engine/Source/Runtime/Engine/Public/Rendering/StaticLightingSystemInterface.h
+++ b/Engine/Source/Runtime/Engine/Public/Rendering/StaticLightingSystemInterface.h
@@ -19,12 +19,7 @@ public:
 
 class ENGINE_API IStaticLightingSystemImpl
 {
-public:
-	virtual bool SupportsRealtimePreview() { return false; }
-	virtual IStaticLightingSystem* CreateStaticLightingSystemForWorld(UWorld* InWorld) { return nullptr; }
-	
-	virtual IStaticLightingSystem* AllocateStaticLightingSystemForWorld(UWorld* InWorld) { return nullptr; }
-	virtual void RemoveStaticLightingSystemForWorld(UWorld* InWorld) {}
+public:	
 	virtual IStaticLightingSystem* GetStaticLightingSystemForWorld(UWorld* InWorld) { return nullptr; }
 
 	virtual void EditorTick() {}
@@ -59,7 +54,6 @@ public:
 	void RegisterImplementation(FName Name, IStaticLightingSystemImpl* Impl);
 	void UnregisterImplementation(FName Name);
 	IStaticLightingSystemImpl* GetPreferredImplementation();
-	bool ShouldOperateOnWorld(UWorld* InWorld);
 
 private:
 	TMap<FName, IStaticLightingSystemImpl*> Implementations;
