Prevented ticking of anim instances during reinstancing

Component re-registration does not need to run a tick during re-instancing, and can sometimes force a tick for a REINST class that has not yet been reference-replaced. This prevents the tick from trying to use class data before it has properly been compiled.

diff --git a/Engine/Source/Runtime/Engine/Private/Animation/AnimInstance.cpp b/Engine/Source/Runtime/Engine/Private/Animation/AnimInstance.cpp
--- a/Engine/Source/Runtime/Engine/Private/Animation/AnimInstance.cpp
+++ b/Engine/Source/Runtime/Engine/Private/Animation/AnimInstance.cpp
@@ -424,6 +424,13 @@ void UAnimInstance::UpdateAnimation(float DeltaSeconds, bool bNeedsValidRootMoti
 {
 	LLM_SCOPE(ELLMTag::Animation);
 
+#if WITH_EDITOR
+	if(GIsReinstancing)
+	{
+		return;
+	}
+#endif
+
 #if DO_CHECK
 	checkf(!bUpdatingAnimation, TEXT("UpdateAnimation already in progress, circular detected for SkeletalMeshComponent [%s], AnimInstance [%s]"), *GetNameSafe(GetOwningComponent()),  *GetName());
 	TGuardValue<bool> CircularGuard(bUpdatingAnimation, true);
