Deprecate the IsPrimaryPlayer function on the APlayerState.
This function always returned "True" incorrectly. Instead of using this function on the player state, you can use the Player Controller or ULocalPlayer's Platform User ID and check if that is the platform's primary user in the IPlatformInputDeviceMapper.

diff --git a/Engine/Source/Runtime/Engine/Classes/GameFramework/PlayerState.h b/Engine/Source/Runtime/Engine/Classes/GameFramework/PlayerState.h
--- a/Engine/Source/Runtime/Engine/Classes/GameFramework/PlayerState.h
+++ b/Engine/Source/Runtime/Engine/Classes/GameFramework/PlayerState.h
@@ -300,6 +300,7 @@ public:
 	virtual void SeamlessTravelTo(class APlayerState* NewPlayerState);
 
 	/** return true if PlayerState is primary (ie. non-splitscreen) player */
+	UE_DEPRECATED(5.1, "This version of IsPrimaryPlayer has been deprecated, please use the Platform Device Mapper to check the owning PlatformUserId instead.")
 	virtual bool IsPrimaryPlayer() const;
 
 	virtual void GetLifetimeReplicatedProps(TArray< FLifetimeProperty > & OutLifetimeProps) const override;
