Fix an issue with multi-views rendering where skeletal mesh has no valid velocity on views other than the first view.

diff --git a/Engine/Plugins/Runtime/nDisplay/Source/DisplayCluster/Private/Game/EngineClasses/Basics/DisplayClusterViewportClient.cpp b/Engine/Plugins/Runtime/nDisplay/Source/DisplayCluster/Private/Game/EngineClasses/Basics/DisplayClusterViewportClient.cpp
--- a/Engine/Plugins/Runtime/nDisplay/Source/DisplayCluster/Private/Game/EngineClasses/Basics/DisplayClusterViewportClient.cpp
+++ b/Engine/Plugins/Runtime/nDisplay/Source/DisplayCluster/Private/Game/EngineClasses/Basics/DisplayClusterViewportClient.cpp
@@ -830,6 +830,7 @@ void UDisplayClusterViewportClient::Draw(FViewport* InViewport, FCanvas* SceneCa
 			ViewFamilies[0]->bAdditionalViewFamily = false;
 			ViewFamilies[0]->bIsFirstViewInMultipleViewFamily = true;
 			ViewFamilies[0]->bIsMultipleViewFamily = true;
+			ViewFamilies[0]->MultipleViewFamilyIndex = 0;
 
 			for (int32 FamilyIndex = 1; FamilyIndex < ViewFamilies.Num(); FamilyIndex++)
 			{
@@ -837,6 +838,7 @@ void UDisplayClusterViewportClient::Draw(FViewport* InViewport, FCanvas* SceneCa
 				ViewFamily.bAdditionalViewFamily = true;
 				ViewFamily.bIsFirstViewInMultipleViewFamily = false;
 				ViewFamily.bIsMultipleViewFamily = true;
+				ViewFamily.MultipleViewFamilyIndex = FamilyIndex;
 			}
 		}
 
diff --git a/Engine/Source/Runtime/Engine/Private/GPUSkinVertexFactory.cpp b/Engine/Source/Runtime/Engine/Private/GPUSkinVertexFactory.cpp
--- a/Engine/Source/Runtime/Engine/Private/GPUSkinVertexFactory.cpp
+++ b/Engine/Source/Runtime/Engine/Private/GPUSkinVertexFactory.cpp
@@ -718,7 +718,9 @@ public:
 
 				// Bone data is updated whenever animation triggers a dynamic update, animation can skip frames hence the frequency is not necessary every frame.
 				// So check if bone data is updated this frame, if not then the previous frame data is stale and not suitable for motion blur.
-				bool bBoneDataUpdatedThisFrame = (View->Family->FrameNumber == ShaderData.UpdatedFrameNumber);
+				// When using multiple view families in a frame, View->Family->FrameNumber is incremented for each view. The updated frame number for bone data 
+				// though is the frame number from the first view, so add its view index in the view families on top to get the right frame number to compare with.
+				bool bBoneDataUpdatedThisFrame = View->Family->FrameNumber == ShaderData.UpdatedFrameNumber + View->Family->MultipleViewFamilyIndex;
 				// If world is paused, use current frame bone matrices, so velocity is canceled and skeletal mesh isn't blurred from motion.
 				bool bPrevious = !View->Family->bWorldIsPaused && bBoneDataUpdatedThisFrame;
 				FRHIShaderResourceView* PreviousData = ShaderData.GetBoneBufferForReading(bPrevious).VertexBufferSRV;
@@ -810,7 +812,9 @@ public:
 
 		// Bone data is updated whenever animation triggers a dynamic update, animation can skip frames hence the frequency is not necessary every frame.
 		// So check if bone data is updated this frame, if not then the previous frame data is stale and not suitable for motion blur.
-		bool bBoneDataUpdatedThisFrame = (View->Family->FrameNumber == PassthroughVertexFactory->GetUpdatedFrameNumber());
+		// When using multiple view families in a frame, View->Family->FrameNumber is incremented for each view. The updated frame number for bone data 
+		// though is the frame number from the first view, so add its view index in the view families on top to get the right frame number to compare with.
+		bool bBoneDataUpdatedThisFrame = (View->Family->FrameNumber == PassthroughVertexFactory->GetUpdatedFrameNumber() + View->Family->MultipleViewFamilyIndex);
 		// If world is paused, use current frame bone matrices, so velocity is canceled and skeletal mesh isn't blurred from motion.
 		bool bVerticesInMotion = !View->Family->bWorldIsPaused && bBoneDataUpdatedThisFrame;
 
diff --git a/Engine/Source/Runtime/Engine/Public/SceneView.h b/Engine/Source/Runtime/Engine/Public/SceneView.h
--- a/Engine/Source/Runtime/Engine/Public/SceneView.h
+++ b/Engine/Source/Runtime/Engine/Public/SceneView.h
@@ -1768,6 +1768,9 @@ public:
 	*/
 	bool bIsFirstViewInMultipleViewFamily = true;
 
+	/** The index of this view inside multiple view families rendered in a single frame. */
+	uint8 MultipleViewFamilyIndex = 0;
+
 	bool bIsSceneTexturesInitialized = false;
 	bool bIsViewFamilyInfo = false;
 
