Physics - Adjust zero-component scale behaviour to alway pass a valid scale to avoid creating invalid body instances

diff --git a/Engine/Source/Runtime/Engine/Private/Components/PrimitiveComponent.cpp b/Engine/Source/Runtime/Engine/Private/Components/PrimitiveComponent.cpp
--- a/Engine/Source/Runtime/Engine/Private/Components/PrimitiveComponent.cpp
+++ b/Engine/Source/Runtime/Engine/Private/Components/PrimitiveComponent.cpp
@@ -823,6 +823,13 @@ void UPrimitiveComponent::OnCreatePhysicsState()
 			if ((BodyInstance.GetCollisionEnabled() != ECollisionEnabled::NoCollision) && (FMath::IsNearlyZero(BodyScale.X) || FMath::IsNearlyZero(BodyScale.Y) || FMath::IsNearlyZero(BodyScale.Z)))
 			{
 				UE_LOG(LogPhysics, Warning, TEXT("Scale for %s has a component set to zero, which will result in a bad body instance. Scale:%s"), *GetPathNameSafe(this), *BodyScale.ToString());
+				
+				// User warning has been output - fix up the scale to be valid for physics
+				BodyTransform.SetScale3D(FVector(
+					FMath::IsNearlyZero(BodyScale.X) ? UE_KINDA_SMALL_NUMBER : BodyScale.X,
+					FMath::IsNearlyZero(BodyScale.Y) ? UE_KINDA_SMALL_NUMBER : BodyScale.Y,
+					FMath::IsNearlyZero(BodyScale.Z) ? UE_KINDA_SMALL_NUMBER : BodyScale.Z
+				));
 			}
 #endif
 
