Minimise skincache header includes before some refactoring for getting skincache for groom.
#rnx

diff --git a/Engine/Plugins/FX/Niagara/Source/Niagara/Private/DataInterface/NiagaraDataInterfaceStaticMesh.cpp b/Engine/Plugins/FX/Niagara/Source/Niagara/Private/DataInterface/NiagaraDataInterfaceStaticMesh.cpp
--- a/Engine/Plugins/FX/Niagara/Source/Niagara/Private/DataInterface/NiagaraDataInterfaceStaticMesh.cpp
+++ b/Engine/Plugins/FX/Niagara/Source/Niagara/Private/DataInterface/NiagaraDataInterfaceStaticMesh.cpp
@@ -22,7 +22,6 @@
 #include "Engine/StaticMeshActor.h"
 #include "Engine/StaticMeshSocket.h"
 #include "Internationalization/Internationalization.h"
-#include "GPUSkinCache.h"
 #include "ShaderCompilerCore.h"
 
 #include "NiagaraGpuComputeDispatchInterface.h"
diff --git a/Engine/Plugins/Runtime/HairStrands/Source/HairStrandsEditor/Private/GroomMaterialDetails.cpp b/Engine/Plugins/Runtime/HairStrands/Source/HairStrandsEditor/Private/GroomMaterialDetails.cpp
--- a/Engine/Plugins/Runtime/HairStrands/Source/HairStrandsEditor/Private/GroomMaterialDetails.cpp
+++ b/Engine/Plugins/Runtime/HairStrands/Source/HairStrandsEditor/Private/GroomMaterialDetails.cpp
@@ -481,8 +481,6 @@ TSharedRef<SWidget> FGroomMaterialDetails::OnGenerateCustomNameWidgetsForSection
 
 TSharedRef<SWidget> FGroomMaterialDetails::OnGenerateCustomSectionWidgetsForSection(int32 LODIndex, int32 SectionIndex)
 {
-	extern ENGINE_API bool IsGPUSkinCacheAvailable(EShaderPlatform Platform);
-
 	TSharedRef<SVerticalBox> SectionWidget = SNew(SVerticalBox);
 	SectionWidget->AddSlot()
 	.AutoHeight()
diff --git a/Engine/Source/Runtime/Engine/Private/Components/SkinnedMeshComponent.cpp b/Engine/Source/Runtime/Engine/Private/Components/SkinnedMeshComponent.cpp
--- a/Engine/Source/Runtime/Engine/Private/Components/SkinnedMeshComponent.cpp
+++ b/Engine/Source/Runtime/Engine/Private/Components/SkinnedMeshComponent.cpp
@@ -29,6 +29,7 @@
 #include "Animation/MorphTarget.h"
 #include "AnimationRuntime.h"
 #include "Animation/SkinWeightProfileManager.h"
+#include "GPUSkinCache.h"
 
 DEFINE_LOG_CATEGORY_STATIC(LogSkinnedMeshComp, Log, All);
 
diff --git a/Engine/Source/Runtime/Engine/Private/SkeletalRenderGPUSkin.h b/Engine/Source/Runtime/Engine/Private/SkeletalRenderGPUSkin.h
--- a/Engine/Source/Runtime/Engine/Private/SkeletalRenderGPUSkin.h
+++ b/Engine/Source/Runtime/Engine/Private/SkeletalRenderGPUSkin.h
@@ -12,15 +12,17 @@
 #include "ShaderParameters.h"
 #include "Components/SkinnedMeshComponent.h"
 #include "GlobalShader.h"
-#include "GPUSkinCache.h"
 #include "GPUSkinVertexFactory.h"
 #include "SkeletalRenderPublic.h"
 #include "ClothingSystemRuntimeTypes.h"
 #include "Rendering/SkeletalMeshRenderData.h"
 #include "Rendering/SkeletalMeshLODRenderData.h"
 
+enum class EGPUSkinCacheEntryMode;
 class FGPUSkinCache;
 class FRayTracingSkinnedGeometryUpdateQueue;
+class FSkeletalMeshObjectGPUSkin;
+class FVertexOffsetBuffers;
 
 /** 
 * Stores the updated matrices needed to skin the verts.
diff --git a/Engine/Source/Runtime/Renderer/Private/DebugViewModeRendering.h b/Engine/Source/Runtime/Renderer/Private/DebugViewModeRendering.h
--- a/Engine/Source/Runtime/Renderer/Private/DebugViewModeRendering.h
+++ b/Engine/Source/Runtime/Renderer/Private/DebugViewModeRendering.h
@@ -16,7 +16,6 @@ DebugViewModeRendering.h: Contains definitions for rendering debug viewmodes.
 #include "ShaderBaseClasses.h"
 #include "MeshPassProcessor.h"
 #include "DebugViewModeInterface.h"
-#include "GPUSkinCache.h"
 #include "SceneRenderTargetParameters.h"
 
 class FPrimitiveSceneProxy;
diff --git a/Engine/Source/Runtime/Renderer/Private/HairStrands/HairStrandsInterface.cpp b/Engine/Source/Runtime/Renderer/Private/HairStrands/HairStrandsInterface.cpp
--- a/Engine/Source/Runtime/Renderer/Private/HairStrands/HairStrandsInterface.cpp
+++ b/Engine/Source/Runtime/Renderer/Private/HairStrands/HairStrandsInterface.cpp
@@ -9,7 +9,6 @@
 #include "HairStrandsMeshProjection.h"
 #include "HairStrandsData.h"
 
-#include "GPUSkinCache.h"
 #include "Rendering/SkeletalMeshRenderData.h"
 #include "Rendering/SkinWeightVertexBuffer.h"
 #include "CommonRenderResources.h"
