fix use of UnacceptableWords in texture code

diff --git a/Engine/Source/Runtime/Engine/Classes/Engine/Texture2DArray.cpp b/Engine/Source/Runtime/Engine/Classes/Engine/Texture2DArray.cpp
--- a/Engine/Source/Runtime/Engine/Classes/Engine/Texture2DArray.cpp
+++ b/Engine/Source/Runtime/Engine/Classes/Engine/Texture2DArray.cpp
@@ -26,7 +26,7 @@
 
 #define LOCTEXT_NAMESPACE "UTexture2DArray"
 
-// Master switch to control whether streaming is enabled for texture2darray. 
+// Externed global switch to control whether streaming is enabled for texture2darray. 
 bool GSupportsTexture2DArrayStreaming = true;
 
 static TAutoConsoleVariable<int32> CVarAllowTexture2DArrayAssetCreation(
diff --git a/Engine/Source/Runtime/Engine/Private/VolumeTexture.cpp b/Engine/Source/Runtime/Engine/Private/VolumeTexture.cpp
--- a/Engine/Source/Runtime/Engine/Private/VolumeTexture.cpp
+++ b/Engine/Source/Runtime/Engine/Private/VolumeTexture.cpp
@@ -25,7 +25,7 @@
 
 //*****************************************************************************
 
-// Master switch to control whether streaming is enabled for volume texture. 
+// Externed global switch to control whether streaming is enabled for volume texture. 
 bool GSupportsVolumeTextureStreaming = true;
 
 // Limit the possible depth of volume texture otherwise when the user converts 2D textures, they can crash the engine.
