Configuration support for AudioInfos / AudioFormats per platform

diff --git a/Engine/Config/Android/AndroidEngine.ini b/Engine/Config/Android/AndroidEngine.ini
--- a/Engine/Config/Android/AndroidEngine.ini
+++ b/Engine/Config/Android/AndroidEngine.ini
@@ -14,6 +14,8 @@ gc.MaxObjectsInGame=131072
 
 [Audio]
 AudioDeviceModuleName=AndroidAudio
+PlatformFormat=OGG
+PlatformStreamingFormat=OGG
 
 ; Defining below allows switching to audio mixer using -audiomixer commandline
 AudioMixerModuleName=AudioMixerAndroid
diff --git a/Engine/Config/BaseEngine.ini b/Engine/Config/BaseEngine.ini
--- a/Engine/Config/BaseEngine.ini
+++ b/Engine/Config/BaseEngine.ini
@@ -1323,6 +1323,19 @@ UnfocusedVolumeMultiplier=0.0
 UseAudioThread=true
 EnableAudioMixer=false
 DeferStartupPrecache=false
++AllWaveFormats=BINKA
++AllWaveFormats=ADPCM
++AllWaveFormats=PCM
++AllWaveFormats=OPUS
++AllWaveFormats=OGG
++FormatModuleHints=AudioFormatOPUS
++FormatModuleHints=AudioFormatOGG
++FormatModuleHints=AudioFormatADPCM
++FormatModuleHints=AudioFormatBINKA
++AudioInfoModules=BinkAudioDecoder
+FallbackFormat=BINKA
+PlatformFormat=BINKA
+PlatformStreamingFormat=BINKA
 
 [AudioChannelAzimuthMap]
 FrontLeft=330
diff --git a/Engine/Config/HoloLens/HoloLensEngine.ini b/Engine/Config/HoloLens/HoloLensEngine.ini
--- a/Engine/Config/HoloLens/HoloLensEngine.ini
+++ b/Engine/Config/HoloLens/HoloLensEngine.ini
@@ -5,6 +5,9 @@ AudioMixerModuleName=AudioMixerXAudio2
 UseAudioMixer=True
 ; Defines a platform-specific volume headroom (in dB) for audio to provide better platform consistency with respect to volume levels.
 PlatformHeadroomDB=0
+PlatformFormat=OGG
+PlatformStreamingFormat=OGG
+
 
 [TextureStreaming]
 ; PoolSizeVRAMPercentage is how much percentage of GPU Dedicated VRAM should be used as a TexturePool cache for streaming textures (0 - unlimited streaming)
diff --git a/Engine/Config/IOS/IOSEngine.ini b/Engine/Config/IOS/IOSEngine.ini
--- a/Engine/Config/IOS/IOSEngine.ini
+++ b/Engine/Config/IOS/IOSEngine.ini
@@ -24,6 +24,9 @@ AudioMixerModuleName=AudioMixerAudioUnit
 UseAudioMixer=true
 PlatformHeadroomDB=-6
 
+PlatformFormat=ADPCM
+PlatformStreamingFormat=ADPCM
+
 [PlatformMemoryBuckets]
 DefaultMemoryBucket_MinGB=4
 SmallerMemoryBucket_MinGB=3
diff --git a/Engine/Config/Linux/LinuxEngine.ini b/Engine/Config/Linux/LinuxEngine.ini
--- a/Engine/Config/Linux/LinuxEngine.ini
+++ b/Engine/Config/Linux/LinuxEngine.ini
@@ -1,6 +1,8 @@
 [Audio]
 ; Defines a platform-specific volume headroom (in dB) for audio to provide better platform consistency with respect to volume levels.
 PlatformHeadroomDB=0
+PlatformFormat=OGG
+PlatformStreamingFormat=OGG
 
 [SystemSettings]
 r.setres=1280x720
diff --git a/Engine/Config/Mac/MacEngine.ini b/Engine/Config/Mac/MacEngine.ini
--- a/Engine/Config/Mac/MacEngine.ini
+++ b/Engine/Config/Mac/MacEngine.ini
@@ -3,6 +3,8 @@ AudioDeviceModuleName=CoreAudio
 
 ; Defining below allows switching to audio mixer using -audiomixer commandline
 AudioMixerModuleName=AudioMixerCoreAudio
+PlatformFormat=OGG
+PlatformStreamingFormat=OGG
 
 ; Defines a platform-specific volume headroom (in dB) for audio to provide better platform consistency with respect to volume levels.
 PlatformHeadroomDB=-6
diff --git a/Engine/Config/Windows/WindowsEngine.ini b/Engine/Config/Windows/WindowsEngine.ini
--- a/Engine/Config/Windows/WindowsEngine.ini
+++ b/Engine/Config/Windows/WindowsEngine.ini
@@ -2,6 +2,8 @@
 AudioDeviceModuleName=XAudio2
 ; Defining below allows switching to audio mixer using -audiomixer commandline
 AudioMixerModuleName=AudioMixerXAudio2
+PlatformFormat=OGG
+PlatformStreamingFormat=OGG
 
 ; Defines a platform-specific volume headroom (in dB) for audio to provide better platform consistency with respect to volume levels.
 PlatformHeadroomDB=-3
diff --git a/Engine/Platforms/Hololens/Source/Developer/HoloLensTargetPlatform/Private/HoloLensTargetPlatform.h b/Engine/Platforms/Hololens/Source/Developer/HoloLensTargetPlatform/Private/HoloLensTargetPlatform.h
--- a/Engine/Platforms/Hololens/Source/Developer/HoloLensTargetPlatform/Private/HoloLensTargetPlatform.h
+++ b/Engine/Platforms/Hololens/Source/Developer/HoloLensTargetPlatform/Private/HoloLensTargetPlatform.h
@@ -82,22 +82,7 @@ public:
 	{
 		TextureLODSettings = InTextureLODSettings;
 	}
-
-	virtual FName GetWaveFormat( const class USoundWave* Wave ) const override
-	{
-		static FName NAME_OGG(TEXT("OGG"));
-
-		return NAME_OGG;
-	}
-
-	virtual void GetAllWaveFormats(TArray<FName>& OutFormats) const override
-	{
-		static FName NAME_OGG(TEXT("OGG"));
-		static FName NAME_OPUS(TEXT("OPUS"));
-		OutFormats.Add(NAME_OGG);
-		OutFormats.Add(NAME_OPUS);
-	}
-
+	
 #endif //WITH_ENGINE
 
 	virtual EPlatformAuthentication RequiresUserCredentials() const override
diff --git a/Engine/Source/Developer/Android/AndroidTargetPlatform/Private/AndroidTargetPlatform.cpp b/Engine/Source/Developer/Android/AndroidTargetPlatform/Private/AndroidTargetPlatform.cpp
--- a/Engine/Source/Developer/Android/AndroidTargetPlatform/Private/AndroidTargetPlatform.cpp
+++ b/Engine/Source/Developer/Android/AndroidTargetPlatform/Private/AndroidTargetPlatform.cpp
@@ -706,27 +706,6 @@ const UTextureLODSettings& FAndroidTargetPlatform::GetTextureLODSettings() const
 	return *TextureLODSettings;
 }
 
-
-FName FAndroidTargetPlatform::GetWaveFormat(const class USoundWave* Wave) const
-{
-	FName FormatName = Audio::ToName(Wave->GetSoundAssetCompressionType());
-	if (FormatName == Audio::NAME_PLATFORM_SPECIFIC)
-	{
-		FormatName = Audio::NAME_OGG;
-	}
-	return FormatName;
-}
-
-
-void FAndroidTargetPlatform::GetAllWaveFormats(TArray<FName>& OutFormats) const
-{
-
-	OutFormats.Add(Audio::NAME_BINKA);
-	OutFormats.Add(Audio::NAME_OGG);
-	OutFormats.Add(Audio::NAME_PCM);
-	OutFormats.Add(Audio::NAME_ADPCM);
-}
-
 #endif //WITH_ENGINE
 
 bool FAndroidTargetPlatform::SupportsVariants() const
diff --git a/Engine/Source/Developer/Android/AndroidTargetPlatform/Private/AndroidTargetPlatform.h b/Engine/Source/Developer/Android/AndroidTargetPlatform/Private/AndroidTargetPlatform.h
--- a/Engine/Source/Developer/Android/AndroidTargetPlatform/Private/AndroidTargetPlatform.h
+++ b/Engine/Source/Developer/Android/AndroidTargetPlatform/Private/AndroidTargetPlatform.h
@@ -141,9 +141,6 @@ public:
 		TextureLODSettings = InTextureLODSettings;
 	}
 
-	virtual FName GetWaveFormat( const class USoundWave* Wave ) const override;
-	virtual void GetAllWaveFormats( TArray<FName>& OutFormats) const override;
-
 #endif //WITH_ENGINE
 
 	virtual bool SupportsVariants() const override;
diff --git a/Engine/Source/Developer/AudioFormatADPCM/Private/AudioFormatADPCM.cpp b/Engine/Source/Developer/AudioFormatADPCM/Private/AudioFormatADPCM.cpp
--- a/Engine/Source/Developer/AudioFormatADPCM/Private/AudioFormatADPCM.cpp
+++ b/Engine/Source/Developer/AudioFormatADPCM/Private/AudioFormatADPCM.cpp
@@ -416,7 +416,7 @@ class FAudioFormatADPCM : public IAudioFormat
 	enum
 	{
 		/** Version for ADPCM format, this becomes part of the DDC key. */
-		UE_AUDIO_ADPCM_VER = 2,
+		UE_AUDIO_ADPCM_VER = 3,
 	};
 
 	void InterleaveBuffers(const TArray<TArray<uint8> >& SrcBuffers, TArray<uint8> & InterleavedBuffer) const
diff --git a/Engine/Source/Developer/AudioFormatBink/Private/AudioFormatBink.cpp b/Engine/Source/Developer/AudioFormatBink/Private/AudioFormatBink.cpp
--- a/Engine/Source/Developer/AudioFormatBink/Private/AudioFormatBink.cpp
+++ b/Engine/Source/Developer/AudioFormatBink/Private/AudioFormatBink.cpp
@@ -36,7 +36,7 @@ class FAudioFormatBink : public IAudioFormat
 	enum
 	{
 		/** Version for Bink Audio format, this becomes part of the DDC key. */
-		UE_AUDIO_BINK_VER = 3,
+		UE_AUDIO_BINK_VER = 4,
 	};
 
 public:
diff --git a/Engine/Source/Developer/AudioFormatOgg/Private/AudioFormatOgg.cpp b/Engine/Source/Developer/AudioFormatOgg/Private/AudioFormatOgg.cpp
--- a/Engine/Source/Developer/AudioFormatOgg/Private/AudioFormatOgg.cpp
+++ b/Engine/Source/Developer/AudioFormatOgg/Private/AudioFormatOgg.cpp
@@ -32,7 +32,7 @@ class FAudioFormatOgg : public IAudioFormat
 	enum
 	{
 		/** Version for OGG format, this becomes part of the DDC key. */
-		UE_AUDIO_OGG_VER = 4,
+		UE_AUDIO_OGG_VER = 5,
 	};
 
 public:
diff --git a/Engine/Source/Developer/AudioFormatOpus/Private/AudioFormatOpus.cpp b/Engine/Source/Developer/AudioFormatOpus/Private/AudioFormatOpus.cpp
--- a/Engine/Source/Developer/AudioFormatOpus/Private/AudioFormatOpus.cpp
+++ b/Engine/Source/Developer/AudioFormatOpus/Private/AudioFormatOpus.cpp
@@ -30,7 +30,7 @@ class FAudioFormatOpus : public IAudioFormat
 	enum
 	{
 		/** Version for OPUS format, this becomes part of the DDC key. */
-		UE_AUDIO_OPUS_VER = 6,
+		UE_AUDIO_OPUS_VER = 7,
 	};
 
 public:
diff --git a/Engine/Source/Developer/IOS/IOSTargetPlatform/Private/IOSTargetPlatform.cpp b/Engine/Source/Developer/IOS/IOSTargetPlatform/Private/IOSTargetPlatform.cpp
--- a/Engine/Source/Developer/IOS/IOSTargetPlatform/Private/IOSTargetPlatform.cpp
+++ b/Engine/Source/Developer/IOS/IOSTargetPlatform/Private/IOSTargetPlatform.cpp
@@ -721,22 +721,5 @@ const UTextureLODSettings& FIOSTargetPlatform::GetTextureLODSettings() const
 	return *TextureLODSettings;
 }
 
-FName FIOSTargetPlatform::GetWaveFormat( const USoundWave* Wave ) const
-{
-	FName FormatName = Audio::ToName(Wave->GetSoundAssetCompressionType());
-	if (FormatName == Audio::NAME_PLATFORM_SPECIFIC)
-	{
-			FormatName = Audio::NAME_ADPCM;
-	}
-	return FormatName;
-}
-
-void FIOSTargetPlatform::GetAllWaveFormats(TArray<FName>& OutFormat) const
-{
-	OutFormat.Add(Audio::NAME_ADPCM);
-	OutFormat.Add(Audio::NAME_PCM);
-	OutFormat.Add(Audio::NAME_BINKA);
-}
-
 #endif // WITH_ENGINE
 
diff --git a/Engine/Source/Developer/IOS/IOSTargetPlatform/Private/IOSTargetPlatform.h b/Engine/Source/Developer/IOS/IOSTargetPlatform/Private/IOSTargetPlatform.h
--- a/Engine/Source/Developer/IOS/IOSTargetPlatform/Private/IOSTargetPlatform.h
+++ b/Engine/Source/Developer/IOS/IOSTargetPlatform/Private/IOSTargetPlatform.h
@@ -86,9 +86,6 @@ public:
 		TextureLODSettings = InTextureLODSettings;
 	}
 
-	virtual FName GetWaveFormat( const class USoundWave* Wave ) const override;
-	virtual void GetAllWaveFormats(TArray<FName>& OutFormat) const override;
-
 #endif // WITH_ENGINE
 
 	virtual void GetBuildProjectSettingKeys(FString& OutSection, TArray<FString>& InBoolKeys, TArray<FString>& InIntKeys, TArray<FString>& InStringKeys) const override
diff --git a/Engine/Source/Developer/Linux/LinuxTargetPlatform/Private/LinuxTargetPlatform.h b/Engine/Source/Developer/Linux/LinuxTargetPlatform/Private/LinuxTargetPlatform.h
--- a/Engine/Source/Developer/Linux/LinuxTargetPlatform/Private/LinuxTargetPlatform.h
+++ b/Engine/Source/Developer/Linux/LinuxTargetPlatform/Private/LinuxTargetPlatform.h
@@ -337,34 +337,6 @@
 		TextureLODSettings = InTextureLODSettings;
 	}
 
-	virtual FName GetWaveFormat( const class USoundWave* Wave ) const override
-	{
-		FName FormatName = Audio::ToName(Wave->GetSoundAssetCompressionType());
-
-		if (FormatName == Audio::NAME_PLATFORM_SPECIFIC)
-		{
-			if (Wave->IsStreaming(*this->IniPlatformName()))
-			{
-				return Audio::NAME_OPUS;
-			}
-
-			return Audio::NAME_OGG;
-		}
-		else
-		{
-			return FormatName;
-		}
-	}
-
-	virtual void GetAllWaveFormats(TArray<FName>& OutFormats) const override
-	{
-		OutFormats.Add(Audio::NAME_BINKA);
-		OutFormats.Add(Audio::NAME_ADPCM);
-		OutFormats.Add(Audio::NAME_PCM);
-		OutFormats.Add(Audio::NAME_OGG);
-		OutFormats.Add(Audio::NAME_OPUS);
-	}
-
 #endif //WITH_ENGINE
 
 	virtual bool SupportsVariants() const override
diff --git a/Engine/Source/Developer/Mac/MacTargetPlatform/Private/GenericMacTargetPlatform.h b/Engine/Source/Developer/Mac/MacTargetPlatform/Private/GenericMacTargetPlatform.h
--- a/Engine/Source/Developer/Mac/MacTargetPlatform/Private/GenericMacTargetPlatform.h
+++ b/Engine/Source/Developer/Mac/MacTargetPlatform/Private/GenericMacTargetPlatform.h
@@ -193,33 +193,6 @@ public:
 		FTargetPlatformBase::AddDependencySCArrayHelper(OutDependencies, TEXT("Binaries/ThirdParty/ShaderConductor/Mac/libShaderConductor.dylib"));
 	}
 
-	virtual FName GetWaveFormat(const class USoundWave* Wave) const override
-	{
-		FName FormatName = Audio::ToName(Wave->GetSoundAssetCompressionType());
-		if (FormatName == Audio::NAME_PLATFORM_SPECIFIC)
-		{
-			if (Wave->IsStreaming(*this->IniPlatformName()))
-			{
-				return Audio::NAME_OPUS;
-			}
-
-			return Audio::NAME_OGG;
-		}
-		else
-		{
-			return FormatName;
-		}
-	}
-
-	virtual void GetAllWaveFormats(TArray<FName>& OutFormats) const override
-	{
-		OutFormats.Add(Audio::NAME_ADPCM);
-		OutFormats.Add(Audio::NAME_PCM);
-		OutFormats.Add(Audio::NAME_OGG);
-		OutFormats.Add(Audio::NAME_OPUS);
-		OutFormats.Add(Audio::NAME_BINKA);
-	}
-
 #endif //WITH_ENGINE
 
 
diff --git a/Engine/Source/Developer/TargetPlatform/Private/TargetPlatformBase.cpp b/Engine/Source/Developer/TargetPlatform/Private/TargetPlatformBase.cpp
--- a/Engine/Source/Developer/TargetPlatform/Private/TargetPlatformBase.cpp
+++ b/Engine/Source/Developer/TargetPlatform/Private/TargetPlatformBase.cpp
@@ -10,6 +10,7 @@
 #include "ProjectDescriptor.h"
 #include "Misc/App.h"
 #include "Misc/ConfigCacheIni.h"
+#include "Sound/AudioFormatSettings.h"
 
 #define LOCTEXT_NAMESPACE "TargetPlatform"
 
@@ -292,4 +293,44 @@ bool FTargetPlatformBase::DoProjectSettingsMatchDefault(const FString& InPlatfor
 	return true;
 }
 
+
+FTargetPlatformBase::FTargetPlatformBase(const PlatformInfo::FTargetPlatformInfo* const InPlatformInfo) : PlatformInfo(InPlatformInfo)
+{
+	checkf(PlatformInfo, TEXT("Null PlatformInfo was passed to FTargetPlatformBase. Check the static IsUsable function before creating this object. See FWindowsTargetPlatformModule::GetTargetPlatform()"));
+
+	PlatformOrdinal = AssignPlatformOrdinal(*this);
+
+#if WITH_ENGINE
+	// Build Audio Format Settings, Using long form equiv of GetConfigSysten to avoid calling a virtual
+	AudioFormatSettings = MakePimpl<Audio::FAudioFormatSettings>(
+		FConfigCacheIni::ForPlatform(InPlatformInfo->IniPlatformName), GEngineIni, InPlatformInfo->IniPlatformName.ToString());
+#endif //WITH_ENGINE
+
+}
+
+#if WITH_ENGINE
+
+const Audio::FAudioFormatSettings& FTargetPlatformBase::GetAudioFormatSettings() const
+{
+	check(AudioFormatSettings.IsValid())
+	return *AudioFormatSettings;
+}
+
+FName FTargetPlatformBase::GetWaveFormat(const class USoundWave* InWave) const
+{
+	return GetAudioFormatSettings().GetWaveFormat(InWave);
+}
+
+void FTargetPlatformBase::GetAllWaveFormats(TArray<FName>& OutFormats) const
+{
+	GetAudioFormatSettings().GetAllWaveFormats(OutFormats);
+}
+
+void FTargetPlatformBase::GetWaveFormatModuleHints(TArray<FName>& OutModuleNames) const
+{
+	GetAudioFormatSettings().GetWaveFormatModuleHints(OutModuleNames);
+}
+
+#endif // WITH_ENGINE
+
 #undef LOCTEXT_NAMESPACE
diff --git a/Engine/Source/Developer/TargetPlatform/Public/Common/TargetPlatformBase.h b/Engine/Source/Developer/TargetPlatform/Public/Common/TargetPlatformBase.h
--- a/Engine/Source/Developer/TargetPlatform/Public/Common/TargetPlatformBase.h
+++ b/Engine/Source/Developer/TargetPlatform/Public/Common/TargetPlatformBase.h
@@ -11,6 +11,12 @@
 #include "Misc/Paths.h"
 #include "PlatformInfo.h"
 
+// Forward declare.
+namespace Audio 
+{
+	class FAudioFormatSettings;
+}
+
 /**
  * Base class for target platforms.
  */
@@ -222,10 +228,12 @@ public:
 			OutModuleNames.Add(*TextureCompressionFormat);
 		}
 	}
+	
+	TARGETPLATFORM_API virtual FName GetWaveFormat(const class USoundWave* Wave) const override;
+	
+	TARGETPLATFORM_API virtual void GetAllWaveFormats(TArray<FName>& OutFormats) const override;
 
-	virtual void GetWaveFormatModuleHints(TArray<FName>& OutModuleNames) const override
-	{
-	}
+	TARGETPLATFORM_API virtual void GetWaveFormatModuleHints(TArray<FName>& OutModuleNames) const override;
 
 #endif
 
@@ -248,19 +256,16 @@ public:
 
 protected:
 
-	FTargetPlatformBase(const PlatformInfo::FTargetPlatformInfo *const InPlatformInfo)
-		: PlatformInfo(InPlatformInfo)
-	{
-		checkf(PlatformInfo, TEXT("Null PlatformInfo was passed to FTargetPlatformBase. Check the static IsUsable function before creating this object. See FWindowsTargetPlatformModule::GetTargetPlatform()"));
-
-		PlatformOrdinal = AssignPlatformOrdinal(*this);
-	}
+	TARGETPLATFORM_API FTargetPlatformBase(const PlatformInfo::FTargetPlatformInfo *const InPlatformInfo);
 
 	/** Information about this platform */
 	const PlatformInfo::FTargetPlatformInfo *PlatformInfo;
 	int32 PlatformOrdinal;
-
+	
+	TARGETPLATFORM_API const Audio::FAudioFormatSettings& GetAudioFormatSettings() const;
 private:
+	TPimplPtr<Audio::FAudioFormatSettings> AudioFormatSettings;
+	
 	bool HasDefaultBuildSettings() const;
 	static bool DoProjectSettingsMatchDefault(const FString& InPlatformName, const FString& InSection, const TArray<FString>* InBoolKeys, const TArray<FString>* InIntKeys, const TArray<FString>* InStringKeys);
 };
diff --git a/Engine/Source/Developer/Windows/WindowsTargetPlatform/Private/GenericWindowsTargetPlatform.h b/Engine/Source/Developer/Windows/WindowsTargetPlatform/Private/GenericWindowsTargetPlatform.h
--- a/Engine/Source/Developer/Windows/WindowsTargetPlatform/Private/GenericWindowsTargetPlatform.h
+++ b/Engine/Source/Developer/Windows/WindowsTargetPlatform/Private/GenericWindowsTargetPlatform.h
@@ -347,43 +347,6 @@ public:
 		TextureLODSettings = InTextureLODSettings;
 	}
 
-	virtual FName GetWaveFormat(const class USoundWave* Wave) const override
-	{
-		FName FormatName = Audio::ToName(Wave->GetSoundAssetCompressionType());
-
-		if (FormatName == Audio::NAME_PLATFORM_SPECIFIC)
-		{
-#if !USE_VORBIS_FOR_STREAMING
-			if (Wave->IsStreaming())
-			{
-				return Audio::NAME_OPUS;
-			}
-#endif
-
-			return Audio::NAME_OGG;
-		}
-		else
-		{
-			return FormatName;
-		}
-	}
-
-	virtual void GetAllWaveFormats(TArray<FName>& OutFormats) const override
-	{
-		OutFormats.Add(Audio::NAME_BINKA);
-		OutFormats.Add(Audio::NAME_ADPCM);
-		OutFormats.Add(Audio::NAME_PCM);
-		OutFormats.Add(Audio::NAME_OGG);
-		OutFormats.Add(Audio::NAME_OPUS);
-	}
-
-	virtual void GetWaveFormatModuleHints(TArray<FName>& OutModuleNames) const override
-	{
-		OutModuleNames.Add(TEXT("AudioFormatOPUS"));
-		OutModuleNames.Add(TEXT("AudioFormatOGG"));
-		OutModuleNames.Add(TEXT("AudioFormatADPCM"));
-	}
-
 #endif //WITH_ENGINE
 
 	virtual bool SupportsVariants() const override
diff --git a/Engine/Source/Runtime/AudioMixer/Private/AudioMixerDevice.cpp b/Engine/Source/Runtime/AudioMixer/Private/AudioMixerDevice.cpp
--- a/Engine/Source/Runtime/AudioMixer/Private/AudioMixerDevice.cpp
+++ b/Engine/Source/Runtime/AudioMixer/Private/AudioMixerDevice.cpp
@@ -26,6 +26,7 @@
 #include "AssetRegistry/IAssetRegistry.h"
 #include "Async/Async.h"
 #include "AudioDeviceNotificationSubsystem.h"
+#include "Sound/AudioFormatSettings.h"
 
 #if WITH_EDITOR
 #include "AudioEditorModule.h"
@@ -673,20 +674,8 @@ namespace Audio
 
 	FName FMixerDevice::GetRuntimeFormat(const USoundWave* InSoundWave) const
 	{
-		FName RuntimeFormat = Audio::ToName(InSoundWave->GetSoundAssetCompressionType());
-
-		// If not specified, the default platform codec is BINK.
-		if (RuntimeFormat == Audio::NAME_PLATFORM_SPECIFIC)
-		{
-			RuntimeFormat = AudioMixerPlatform->GetRuntimeFormat(InSoundWave);
-			if (RuntimeFormat.IsNone())
-			{
-				// TODO: make this an override per platform for "default" platform-specific codec
-				RuntimeFormat = Audio::NAME_BINKA;
-			}
-		}
-		return RuntimeFormat;
-	}
+		return InSoundWave->GetRuntimeFormat();
+    }
 
 	bool FMixerDevice::HasCompressedAudioInfoClass(USoundWave* InSoundWave)
 	{
@@ -706,15 +695,24 @@ namespace Audio
 	{
 		return AudioMixerPlatform->DisablePCMAudioCaching();
 	}
+	
+	ICompressedAudioInfo* FMixerDevice::CreateAudioInfo(FName InFormat) const
+	{
+		if (IAudioInfoFactory* Factory = IAudioInfoFactoryRegistry::Get().Find(InFormat))
+		{
+			return Factory->Create();
+		}
+		return nullptr;
+	}
 
 	ICompressedAudioInfo* FMixerDevice::CreateCompressedAudioInfo(const USoundWave* InSoundWave) const
 	{
-		return AudioMixerPlatform->CreateCompressedAudioInfo(GetRuntimeFormat(InSoundWave));
+		return CreateAudioInfo(GetRuntimeFormat(InSoundWave));
 	}
 
 	class ICompressedAudioInfo* FMixerDevice::CreateCompressedAudioInfo(const FSoundWaveProxyPtr& InSoundWaveProxy) const
 	{
-		return AudioMixerPlatform->CreateCompressedAudioInfo(InSoundWaveProxy->GetRuntimeFormat());
+		return CreateAudioInfo(InSoundWaveProxy->GetRuntimeFormat());
 	}
 
 	bool FMixerDevice::ValidateAPICall(const TCHAR* Function, uint32 ErrorCode)
diff --git a/Engine/Source/Runtime/AudioMixer/Public/AudioMixerDevice.h b/Engine/Source/Runtime/AudioMixer/Public/AudioMixerDevice.h
--- a/Engine/Source/Runtime/AudioMixer/Public/AudioMixerDevice.h
+++ b/Engine/Source/Runtime/AudioMixer/Public/AudioMixerDevice.h
@@ -24,6 +24,7 @@ namespace Audio
 	class FMixerSourceManager;
 	class FMixerSourceVoice;
 	class FMixerSubmix;
+	class FAudioFormatSettings;
 
 	typedef TSharedPtr<FMixerSubmix, ESPMode::ThreadSafe> FMixerSubmixPtr;
 	typedef TWeakPtr<FMixerSubmix, ESPMode::ThreadSafe> FMixerSubmixWeakPtr;
@@ -338,6 +339,8 @@ namespace Audio
 
 		void UnloadSoundSubmix(const USoundSubmixBase& SoundSubmix);
 
+		ICompressedAudioInfo* CreateAudioInfo(FName InFormat) const;
+
 		bool IsMasterSubmixType(const USoundSubmixBase* InSubmix) const;
 		FMixerSubmixPtr GetMasterSubmixInstance(uint32 InSubmixId);
 		FMixerSubmixPtr GetMasterSubmixInstance(const USoundSubmixBase* InSubmix);
diff --git a/Engine/Source/Runtime/BinkAudioDecoder/Module/Private/BinkAudioInfo.cpp b/Engine/Source/Runtime/BinkAudioDecoder/Module/Private/BinkAudioInfo.cpp
--- a/Engine/Source/Runtime/BinkAudioDecoder/Module/Private/BinkAudioInfo.cpp
+++ b/Engine/Source/Runtime/BinkAudioDecoder/Module/Private/BinkAudioInfo.cpp
@@ -16,12 +16,6 @@ DEFINE_LOG_CATEGORY_STATIC(LogBinkAudioDecoder, Log, All);
 
 //-----------------------------------------------------------------------------
 //-----------------------------------------------------------------------------
-class BINKAUDIODECODER_API FBinkAudioDecoderModule : public IModuleInterface
-{
-public:
-	virtual void StartupModule() override {}
-	virtual void ShutdownModule() override {}
-};
 
 #define PTR_ADD(ptr,off) ((void*)(((uint8*)(ptr))+(off)))
 #define Align32( val ) ( ( ( val ) + 31 ) & ~31 )
@@ -660,4 +654,17 @@ bool FBinkAudioInfo::HasError() const
 	return bErrorStateLatch;
 }
 
-IMPLEMENT_MODULE(FBinkAudioDecoderModule, BinkAudioDecoder)
\ No newline at end of file
+class BINKAUDIODECODER_API FBinkAudioDecoderModule : public IModuleInterface
+{
+public:	
+	TUniquePtr<IAudioInfoFactory> Factory;
+
+	virtual void StartupModule() override
+	{	
+		Factory = MakeUnique<FSimpleAudioInfoFactory>([] { return new FBinkAudioInfo(); }, Audio::NAME_BINKA);
+	}
+
+	virtual void ShutdownModule() override {}
+};
+
+IMPLEMENT_MODULE(FBinkAudioDecoderModule, BinkAudioDecoder)
diff --git a/Engine/Source/Runtime/CoreUObject/Public/Serialization/BulkData.h b/Engine/Source/Runtime/CoreUObject/Public/Serialization/BulkData.h
--- a/Engine/Source/Runtime/CoreUObject/Public/Serialization/BulkData.h
+++ b/Engine/Source/Runtime/CoreUObject/Public/Serialization/BulkData.h
@@ -1275,6 +1275,10 @@ public:
 	{
 		return Formats.Contains(Format);
 	}
+	void GetContainedFormats(TArray<FName>& OutFormats) const
+	{
+		Formats.GenerateKeyArray(OutFormats);
+	}
 	FByteBulkData& GetFormat(FName Format)
 	{
 		FByteBulkData* Result = Formats.FindRef(Format);
diff --git a/Engine/Source/Runtime/Engine/Classes/Sound/AudioFormatSettings.h b/Engine/Source/Runtime/Engine/Classes/Sound/AudioFormatSettings.h
new file mode 100644
--- /dev/null
+++ b/Engine/Source/Runtime/Engine/Classes/Sound/AudioFormatSettings.h
@@ -0,0 +1,27 @@
+// Copyright Epic Games, Inc. All Rights Reserved.
+#pragma once
+
+namespace Audio
+{
+	class ENGINE_API FAudioFormatSettings
+	{
+	public:
+		FAudioFormatSettings(FConfigCacheIni* InConfigSystem, const FString& InConfigFilename, const FString& InPlatformIdentifierForLogging);
+		~FAudioFormatSettings() = default;
+
+		FName GetWaveFormat(const class USoundWave* Wave) const;
+		void GetAllWaveFormats(TArray<FName>& OutFormats) const;
+		void GetWaveFormatModuleHints(TArray<FName>& OutHints) const;
+
+		FName GetFallbackFormat() const { return FallbackFormat; }
+
+	private:
+		void ReadConfiguration(FConfigCacheIni*, const FString& InConfigFilename, const FString& InPlatformIdentifierForLogging);
+
+		TArray<FName> AllWaveFormats;
+		TArray<FName> WaveFormatModuleHints;
+		FName PlatformFormat;
+		FName PlatformStreamingFormat;
+		FName FallbackFormat;
+	};
+}
diff --git a/Engine/Source/Runtime/Engine/Classes/Sound/SoundWave.h b/Engine/Source/Runtime/Engine/Classes/Sound/SoundWave.h
--- a/Engine/Source/Runtime/Engine/Classes/Sound/SoundWave.h
+++ b/Engine/Source/Runtime/Engine/Classes/Sound/SoundWave.h
@@ -143,6 +143,9 @@ struct ENGINE_API FStreamedAudioPlatformData
 	/** Get the number of chunks while making sure any async task are finished before returning. */
 	int32 GetNumChunks() const;
 
+	/** Get the audio format making sure any async task are finished before returning. */
+	FName GetAudioFormat() const;
+
 	/** Serialization. */
 	void Serialize(FArchive& Ar, class USoundWave* Owner);
 
@@ -459,6 +462,9 @@ public:
 	UFUNCTION(BlueprintCallable, Category = "Audio")
 	void SetSoundAssetCompressionType(ESoundAssetCompressionType InSoundAssetCompressionType);
 
+	/** Returns the Runtime format of the wave */
+	FName GetRuntimeFormat() const;
+		
 private:
 	// cached proxy
 	FSoundWaveProxyPtr InternalProxy{ nullptr };
@@ -1294,6 +1300,8 @@ public:
 private:
 	void DiscardZerothChunkData();
 
+	FName FindRuntimeFormat(const USoundWave&) const;
+
 	/** Zeroth Chunk of audio for sources that use Load On Demand. */
 	FBulkDataBuffer<uint8> ZerothChunkData;
 	mutable FCriticalSection LoadZerothChunkDataCriticalSection;
diff --git a/Engine/Source/Runtime/Engine/Private/AudioDecompress.cpp b/Engine/Source/Runtime/Engine/Private/AudioDecompress.cpp
--- a/Engine/Source/Runtime/Engine/Private/AudioDecompress.cpp
+++ b/Engine/Source/Runtime/Engine/Private/AudioDecompress.cpp
@@ -841,3 +841,36 @@ bool ICompressedAudioInfo::StreamCompressedInfo(const FSoundWaveProxyPtr& Wave,
 	StreamingSoundWave = Wave;
 	return StreamCompressedInfoInternal(StreamingSoundWave, QualityInfo);
 }
+
+IAudioInfoFactoryRegistry& IAudioInfoFactoryRegistry::Get()
+{
+	static struct FConcreteRegistry : IAudioInfoFactoryRegistry
+	{
+		mutable FRWLock FactoriesRWLock;
+		TMap<FName, IAudioInfoFactory*> Factories;
+
+		void Register(IAudioInfoFactory* InFactory, FName InFormatName) override
+		{
+			FRWScopeLock Lock(FactoriesRWLock, FRWScopeLockType::SLT_Write);
+			UE_LOG(LogAudio, Display, TEXT("AudioInfo: '%s' Registered"), *InFormatName.ToString());
+			check(!Factories.Contains(InFormatName));
+			Factories.Add(InFormatName) = InFactory;
+		}
+		void Unregister(IAudioInfoFactory* InFactory, FName InFormatName) override
+		{
+			FRWScopeLock Lock(FactoriesRWLock, FRWScopeLockType::SLT_Write);
+
+			Factories.Remove(InFormatName);
+		}
+		IAudioInfoFactory* Find(FName InFormat) const override
+		{
+			FRWScopeLock Lock(FactoriesRWLock, FRWScopeLockType::SLT_ReadOnly);
+			if( IAudioInfoFactory* const* Factory = Factories.Find(InFormat))
+			{
+				return *Factory;
+			}
+			return nullptr;
+		}
+	} sInstance;
+	return sInstance;
+}
diff --git a/Engine/Source/Runtime/Engine/Private/AudioDerivedData.cpp b/Engine/Source/Runtime/Engine/Private/AudioDerivedData.cpp
--- a/Engine/Source/Runtime/Engine/Private/AudioDerivedData.cpp
+++ b/Engine/Source/Runtime/Engine/Private/AudioDerivedData.cpp
@@ -114,6 +114,19 @@ int32 FStreamedAudioPlatformData::GetNumChunks() const
 	return NumChunks;
 }
 
+FName FStreamedAudioPlatformData::GetAudioFormat() const
+{
+#if WITH_EDITORONLY_DATA
+	if (FStreamedAudioBuildScope::ShouldWaitOnIncompleteProperties(this))
+	{
+		// AudioFormat is written by the caching process, any async task need to complete before we can read it.
+		const_cast<FStreamedAudioPlatformData*>(this)->FinishCache();
+	}
+#endif
+
+	return AudioFormat;
+}
+
 #if WITH_EDITORONLY_DATA
 
 /*------------------------------------------------------------------------------
diff --git a/Engine/Source/Runtime/Engine/Private/AudioDeviceManager.cpp b/Engine/Source/Runtime/Engine/Private/AudioDeviceManager.cpp
--- a/Engine/Source/Runtime/Engine/Private/AudioDeviceManager.cpp
+++ b/Engine/Source/Runtime/Engine/Private/AudioDeviceManager.cpp
@@ -15,7 +15,12 @@
 #include "Misc/ConfigCacheIni.h"
 #include "Misc/CoreDelegates.h"
 #include "UObject/UObjectIterator.h"
-#include "Audio/AudioDebug.h"
+
+#include "AudioDecompress.h"
+#include "Sound/AudioFormatSettings.h"
+#include "ADPCMAudioInfo.h"
+#include "VorbisAudioInfo.h"
+#include "OpusAudioInfo.h"
 
 #if INSTRUMENT_AUDIODEVICE_HANDLES
 #include "Containers/StringConv.h"
@@ -533,6 +538,42 @@ void FAudioDeviceManager::UnregisterWorld(UWorld* InWorld, Audio::FDeviceId Devi
 	}
 }
 
+void FAudioDeviceManager::RegisterAudioInfoFactories()
+{
+	// Load any Engine.ini defined modules necessary for registering format factories.
+	TArray<FString> AudioInfoModules;
+	if (GConfig->GetArray(TEXT("Audio"), TEXT("AudioInfoModules"), AudioInfoModules, GEngineIni))
+	{
+		for (const FString& i : AudioInfoModules)
+		{
+			FModuleManager::Get().LoadModuleChecked(*i);
+		}
+	}
+
+	// Register the engine formats.
+	EngineFormats.Add(MakePimpl<FSimpleAudioInfoFactory>([] { return new FADPCMAudioInfo(); }, Audio::NAME_PCM));
+	EngineFormats.Add(MakePimpl<FSimpleAudioInfoFactory>([] { return new FADPCMAudioInfo(); }, Audio::NAME_ADPCM));
+	EngineFormats.Add(MakePimpl<FSimpleAudioInfoFactory>([] { return new FVorbisAudioInfo(); }, Audio::NAME_OGG));
+	EngineFormats.Add(MakePimpl<FSimpleAudioInfoFactory>([] { return new FOpusAudioInfo(); }, Audio::NAME_OPUS));
+
+	// Sanity check we have all the Factories we need to run now by 
+	TArray<FName> AllFormats;
+	GetAudioFormatSettings().GetAllWaveFormats(AllFormats);
+
+	FString FailedFormatsString;
+	int32 NumFailedFormats = 0;
+	for (FName i : AllFormats)
+	{
+		if (!IAudioInfoFactoryRegistry::Get().Find(i))
+		{
+			FailedFormatsString += FString::Printf(TEXT("'%s' "), *i.ToString());
+			NumFailedFormats++;
+		}
+	}
+	checkf(NumFailedFormats == 0, TEXT("Failed to find these required AudioFormats: [ %s]"), *FailedFormatsString);
+}
+
+
 bool FAudioDeviceManager::InitializeManager()
 {
 	if (LoadDefaultAudioDeviceModule())
@@ -548,6 +589,10 @@ bool FAudioDeviceManager::InitializeManager()
 		const bool bIsAudioMixerEnabled = AudioDeviceModule->IsAudioMixerModule();
 		AudioSettings->SetAudioMixerEnabled(bIsAudioMixerEnabled);
 
+		// Register all formats
+		AudioFormatSettings = MakePimpl<Audio::FAudioFormatSettings>(GConfig, GEngineIni, FPlatformProperties::IniPlatformName());
+		RegisterAudioInfoFactories();
+
 #if WITH_EDITOR
 		if (bIsAudioMixerEnabled)
 		{
@@ -1232,6 +1277,12 @@ void FAudioDeviceManager::LogListOfAudioDevices()
 	UE_LOG(LogAudio, Display, TEXT("Active Audio Devices:\n%s"), *ListOfDevices);
 }
 
+Audio::FAudioFormatSettings& FAudioDeviceManager::GetAudioFormatSettings() const
+{
+	check(AudioFormatSettings.IsValid());
+	return *AudioFormatSettings;
+}
+
 uint32 FAudioDeviceManager::GetNewDeviceID()
 {
 	return ++DeviceIDCounter;
diff --git a/Engine/Source/Runtime/Engine/Private/AudioFormatSettings.cpp b/Engine/Source/Runtime/Engine/Private/AudioFormatSettings.cpp
new file mode 100644
--- /dev/null
+++ b/Engine/Source/Runtime/Engine/Private/AudioFormatSettings.cpp
@@ -0,0 +1,106 @@
+// Copyright Epic Games, Inc. All Rights Reserved.
+
+#include "Sound/AudioFormatSettings.h"
+#include "Algo/Transform.h"
+
+namespace Audio
+{
+	FAudioFormatSettings::FAudioFormatSettings(FConfigCacheIni* InConfigSystem, const FString& InConfigFilename, const FString& InPlatformIdentifierForLogging)
+	{
+		ReadConfiguration(InConfigSystem, InConfigFilename, InPlatformIdentifierForLogging);
+	}
+
+	FName FAudioFormatSettings::GetWaveFormat(const class USoundWave* Wave) const
+	{
+		FName FormatName = Audio::ToName(Wave->GetSoundAssetCompressionType());
+		if (FormatName == Audio::NAME_PLATFORM_SPECIFIC)
+		{
+			if (Wave->IsStreaming())
+			{
+				return PlatformStreamingFormat;
+			}
+			else
+			{
+				return PlatformFormat;
+			}
+		}
+		return FormatName;
+	}
+
+	void FAudioFormatSettings::GetAllWaveFormats(TArray<FName>& OutFormats) const
+	{
+		OutFormats = AllWaveFormats;
+	}
+
+	void FAudioFormatSettings::GetWaveFormatModuleHints(TArray<FName>& OutHints) const
+	{
+		OutHints = WaveFormatModuleHints;
+	}
+
+	void FAudioFormatSettings::ReadConfiguration(FConfigCacheIni* InConfigSystem, const FString& InConfigFilename, const FString& InPlatformIdentifierForLogging)
+	{
+		auto ToFName = [](const FString& InName) -> FName { return { *InName }; };
+
+		TArray<FString> FormatNames;
+		if (ensure(InConfigSystem->GetArray(TEXT("Audio"), TEXT("AllWaveFormats"), FormatNames, InConfigFilename)))
+		{
+			Algo::Transform(FormatNames, AllWaveFormats, ToFName);
+		}
+
+		TArray<FString> FormatModuleHints;
+		if (InConfigSystem->GetArray(TEXT("Audio"), TEXT("FormatModuleHints"), FormatModuleHints, InConfigFilename))
+		{
+			Algo::Transform(FormatNames, WaveFormatModuleHints, ToFName);
+		}
+
+		FString FallbackFormatString;
+		if (ensure(InConfigSystem->GetString(TEXT("Audio"), TEXT("FallbackFormat"), FallbackFormatString, InConfigFilename)))
+		{
+			FallbackFormat = *FallbackFormatString;
+
+			if (!AllWaveFormats.Contains(FallbackFormat) && AllWaveFormats.Num() > 0)
+			{
+				UE_LOG(LogAudio, Warning, TEXT("FallbackFormat '%s' not defined in 'AllWaveFormats'. Using first format listed '%s'"), *FallbackFormatString, *AllWaveFormats[0].ToString());
+				FallbackFormat = AllWaveFormats[0];
+			}
+		}
+
+		FString PlatformFormatString;
+		if (ensure(InConfigSystem->GetString(TEXT("Audio"), TEXT("PlatformFormat"), PlatformFormatString, InConfigFilename)))
+		{
+			PlatformFormat = *PlatformFormatString;
+
+			if (!AllWaveFormats.Contains(PlatformFormat))
+			{
+				UE_LOG(LogAudio, Warning, TEXT("PlatformStreamingFormat '%s' not defined in 'AllWaveFormats'. Using fallback format '%s'"), *PlatformFormatString, *FallbackFormat.ToString());
+				PlatformFormat = FallbackFormat;
+			}
+		}
+
+		FString PlatformStreamingFormatString;
+		if (ensure(InConfigSystem->GetString(TEXT("Audio"), TEXT("PlatformStreamingFormat"), PlatformStreamingFormatString, InConfigFilename)))
+		{
+			PlatformStreamingFormat = *PlatformStreamingFormatString;
+			if (!AllWaveFormats.Contains(PlatformStreamingFormat))
+			{
+				UE_LOG(LogAudio, Warning, TEXT("PlatformStreamingFormat '%s' not defined in 'AllWaveFormats'. Using fallback format '%s'"), *PlatformStreamingFormatString, *FallbackFormat.ToString());
+				PlatformStreamingFormat = FallbackFormat;
+			}
+		}
+
+	#if !NO_LOGGING
+
+		// Display log for sanity
+		FString AllFormatsConcat;
+		for (FName i : AllWaveFormats)
+		{
+			AllFormatsConcat += i.ToString() + TEXT(" ");
+		}
+
+		UE_LOG(LogAudio, Display, TEXT("AudioFormatSettings: TargetName='%s', AllWaveFormats=( %s), PlatformFormat='%s', PlatformStreamingFormat='%s', FallbackFormat='%s'"),
+			*InPlatformIdentifierForLogging, *AllFormatsConcat, *PlatformFormat.ToString(), *PlatformStreamingFormat.ToString(), *FallbackFormat.ToString());
+
+	#endif //!NO_LOGGING
+	}
+
+}// namespace Audio
diff --git a/Engine/Source/Runtime/Engine/Private/SoundWave.cpp b/Engine/Source/Runtime/Engine/Private/SoundWave.cpp
--- a/Engine/Source/Runtime/Engine/Private/SoundWave.cpp
+++ b/Engine/Source/Runtime/Engine/Private/SoundWave.cpp
@@ -37,6 +37,7 @@
 #include "Async/Async.h"
 #include "Misc/CommandLine.h"
 #include "UObject/UE5MainStreamObjectVersion.h"
+#include "Sound/AudioFormatSettings.h"
 
 #if WITH_EDITOR
 #include "Framework/Notifications/NotificationManager.h"
@@ -157,15 +158,8 @@ void FSoundWaveData::InitializeDataFromSoundWave(USoundWave& InWave)
 	}
 
 	// cache the runtime format for the wave
-	if (GEngine)
-	{
-		FAudioDevice* LocalAudioDevice = GEngine->GetMainAudioDeviceRaw();
-		if (LocalAudioDevice)
-		{
-			RuntimeFormat = LocalAudioDevice->GetRuntimeFormat(&InWave);
-		}
-	}
-
+	RuntimeFormat = FindRuntimeFormat(InWave);
+	
 	SoundWaveKeyCached = FObjectKey(&InWave);
 	CuePoints = InWave.CuePoints;
 	SampleRate = InWave.GetSampleRateForCurrentPlatform();
@@ -182,6 +176,58 @@ void FSoundWaveData::InitializeDataFromSoundWave(USoundWave& InWave)
 	bShouldUseStreamCaching = InWave.ShouldUseStreamCaching();
 }
 
+FName FSoundWaveData::FindRuntimeFormat(const USoundWave& InWave) const
+{		
+#if WITH_EDITOR	
+	// If this is an editor build, we can ask ITargetPlatform.
+	if (ITargetPlatformManagerModule* TPM = GetTargetPlatformManager())
+	{
+		if (ITargetPlatform* RunningTarget = TPM->GetRunningTargetPlatform())
+		{
+			return RunningTarget->GetWaveFormat(&InWave);
+		}
+	}
+#endif //WITH_EDITOR
+
+	// Force-inline does not use the stream cache chunks at all, so we have to work out the format
+	// from the name of the inner "Format" which has suffixes in the form "FORMAT123456789"
+	if (LoadingBehavior == ESoundWaveLoadingBehavior::ForceInline)
+	{
+		TArray<FName> AllValidFormats;
+		FAudioDeviceManager::Get()->GetAudioFormatSettings().GetAllWaveFormats(AllValidFormats);
+
+		TArray<FName> ContainedFormats;
+		CompressedFormatData.GetContainedFormats(ContainedFormats);
+		if (ContainedFormats.Num() > 0)
+		{
+			// We only support 1 cooked format per wave, so just look at the first one.
+			FString CompressedFormatString = ContainedFormats[0].GetPlainNameString();
+			for (const FName ValidFormat : AllValidFormats)
+			{
+				if (CompressedFormatString.StartsWith(*ValidFormat.GetPlainNameString()))
+				{
+					return ValidFormat;
+				}
+			}
+
+			checkf(false, TEXT("Failed to determine encoded (inlined) AudioFormat: '%s'"), *CompressedFormatString);
+			
+			// Fail.
+			return {};
+		}	
+
+		checkf(false, TEXT("Failed to find any encoded formats in container"));
+		
+		// Fail.
+		return {};
+	}
+	else
+	{
+		// If we're using force inline, the audio format is written in the running platform.
+		return RunningPlatformData.GetAudioFormat();
+	}
+}
+
 uint32 FSoundWaveData::GetNumChunks() const
 {
 	// This function shouldn't be called on audio marked "ForceInline."
@@ -1086,6 +1132,11 @@ void USoundWave::SetSoundAssetCompressionType(ESoundAssetCompressionType InSound
 #endif // #if WITH_EDITOR
 }
 
+FName USoundWave::GetRuntimeFormat() const
+{
+	return SoundWaveDataPtr->GetRuntimeFormat();
+}
+
 float USoundWave::GetSubtitlePriority() const
 {
 	return SubtitlePriority;
diff --git a/Engine/Source/Runtime/Engine/Public/AudioDecompress.h b/Engine/Source/Runtime/Engine/Public/AudioDecompress.h
--- a/Engine/Source/Runtime/Engine/Public/AudioDecompress.h
+++ b/Engine/Source/Runtime/Engine/Public/AudioDecompress.h
@@ -330,6 +330,37 @@ protected:
 	int32 StreamSeekBlockOffset;
 };
 
+struct IAudioInfoFactory 
+{	
+	virtual ~IAudioInfoFactory() = default;
+	virtual ICompressedAudioInfo* Create() = 0;	
+};
+struct ENGINE_API IAudioInfoFactoryRegistry
+{
+	static IAudioInfoFactoryRegistry& Get();
+	virtual ~IAudioInfoFactoryRegistry() = default;
+	virtual void Register(IAudioInfoFactory*, FName) = 0;
+	virtual void Unregister(IAudioInfoFactory*, FName) = 0;
+	virtual IAudioInfoFactory* Find(FName InName) const = 0;
+};
+
+class FSimpleAudioInfoFactory : public IAudioInfoFactory
+{
+public:
+	FSimpleAudioInfoFactory(TFunction<ICompressedAudioInfo*(void)> InLambda, FName InFormatName)
+	:  CreateLamda(InLambda), FormatName(InFormatName)
+	{
+		IAudioInfoFactoryRegistry::Get().Register(this, FormatName);
+	}
+	virtual ~FSimpleAudioInfoFactory() 
+	{
+		IAudioInfoFactoryRegistry::Get().Unregister(this, FormatName);
+	}
+	virtual ICompressedAudioInfo* Create() override { return CreateLamda(); }
+private:
+	TFunction<ICompressedAudioInfo* ()> CreateLamda;
+	FName FormatName;
+};
 
 /**
  * Asynchronous audio decompression
diff --git a/Engine/Source/Runtime/Engine/Public/AudioDeviceManager.h b/Engine/Source/Runtime/Engine/Public/AudioDeviceManager.h
--- a/Engine/Source/Runtime/Engine/Public/AudioDeviceManager.h
+++ b/Engine/Source/Runtime/Engine/Public/AudioDeviceManager.h
@@ -28,6 +28,7 @@ class USoundWave;
 class UWorld;
 class FAudioDevice;
 struct FSourceEffectChainEntry;
+class FSimpleAudioInfoFactory;
 
 namespace Audio
 {
@@ -39,6 +40,7 @@ namespace Audio
 	using FDeviceId = uint32;
 
 	class FAudioDebugger;
+	class FAudioFormatSettings;
 }
 
 enum class ESoundType : uint8
@@ -409,6 +411,8 @@ public:
 
 	void LogListOfAudioDevices();
 
+	Audio::FAudioFormatSettings& GetAudioFormatSettings() const;
+
 private:
 
 #if ENABLE_AUDIO_DEBUG
@@ -416,6 +420,9 @@ private:
 	TUniquePtr<Audio::FAudioDebugger> AudioDebugger;
 #endif // ENABLE_AUDIO_DEBUG
 
+	TPimplPtr<Audio::FAudioFormatSettings> AudioFormatSettings;
+	TArray<TPimplPtr<FSimpleAudioInfoFactory>> EngineFormats;
+
 	bool InitializeManager();
 
 	/** Creates a handle given the index and a generation value. */
@@ -437,6 +444,8 @@ private:
 
 	/** Application enters background handler */
 	void AppWillEnterBackground();
+	
+	void RegisterAudioInfoFactories();
 
 	/** Audio device module which creates (old backend) audio devices. */
 	IAudioDeviceModule* AudioDeviceModule;
