Submitting on behalf of Anna.Lantz: Inclusive language push, removing master language from Runtime Audio component
#rnx

diff --git a/Engine/Source/Editor/MovieSceneCaptureDialog/Private/MovieSceneCaptureDialogModule.cpp b/Engine/Source/Editor/MovieSceneCaptureDialog/Private/MovieSceneCaptureDialogModule.cpp
--- a/Engine/Source/Editor/MovieSceneCaptureDialog/Private/MovieSceneCaptureDialogModule.cpp
+++ b/Engine/Source/Editor/MovieSceneCaptureDialog/Private/MovieSceneCaptureDialogModule.cpp
@@ -440,8 +440,8 @@ void FInEditorCapture::Start()
 	FAudioDeviceHandle AudioDevice = GEngine->GetMainAudioDevice();
 	if (AudioDevice)
 	{
-		TransientMasterVolume = AudioDevice->GetTransientMasterVolume();
-		AudioDevice->SetTransientMasterVolume(0.0f);
+		TransientPrimaryVolume = AudioDevice->GetTransientPrimaryVolume();
+		AudioDevice->SetTransientPrimaryVolume(0.0f);
 	}
 
 	TSharedRef<SWindow> CustomWindow = SNew(SWindow)
@@ -604,7 +604,7 @@ void FInEditorCapture::Shutdown()
 	FAudioDeviceHandle AudioDevice = GEngine->GetMainAudioDevice();
 	if (AudioDevice)
 	{
-		AudioDevice->SetTransientMasterVolume(TransientMasterVolume);
+		AudioDevice->SetTransientPrimaryVolume(TransientPrimaryVolume);
 	}
 
 	CaptureObject->Close();
diff --git a/Engine/Source/Editor/MovieSceneCaptureDialog/Public/MovieSceneCaptureDialogModule.h b/Engine/Source/Editor/MovieSceneCaptureDialog/Public/MovieSceneCaptureDialogModule.h
--- a/Engine/Source/Editor/MovieSceneCaptureDialog/Public/MovieSceneCaptureDialogModule.h
+++ b/Engine/Source/Editor/MovieSceneCaptureDialog/Public/MovieSceneCaptureDialogModule.h
@@ -125,7 +125,7 @@ private:
 	UWorld* CapturingFromWorld;
 
 	bool bScreenMessagesWereEnabled;
-	float TransientMasterVolume;
+	float TransientPrimaryVolume;
 	int32 BackedUpStreamingPoolSize;
 	int32 BackedUpUseFixedPoolSize;
 	int32 BackedUpTextureStreaming;
diff --git a/Engine/Source/Editor/UnrealEd/Private/PlayLevel.cpp b/Engine/Source/Editor/UnrealEd/Private/PlayLevel.cpp
--- a/Engine/Source/Editor/UnrealEd/Private/PlayLevel.cpp
+++ b/Engine/Source/Editor/UnrealEd/Private/PlayLevel.cpp
@@ -730,7 +730,7 @@ void UEditorEngine::TeardownPlaySession(FWorldContext& PieWorldContext)
 		AudioDevice->Flush(PlayWorld);
 		AudioDevice->ResetInterpolation();
 		AudioDevice->OnEndPIE(false); // TODO: Should this have been bWasSimulatingInEditor?
-		AudioDevice->SetTransientMasterVolume(1.0f);
+		AudioDevice->SetTransientPrimaryVolume(1.0f);
 		// Reset solo audio
 		if (PlayInEditorSessionInfo.IsSet())
 		{
@@ -2516,7 +2516,7 @@ void UEditorEngine::ResetPIEAudioSetting(UWorld *CurrentPieWorld)
 	{
 		if (FAudioDevice* AudioDevice = CurrentPieWorld->GetAudioDeviceRaw())
 		{
-			AudioDevice->SetTransientMasterVolume(0.0f);
+			AudioDevice->SetTransientPrimaryVolume(0.0f);
 		}
 	}
 }
@@ -3017,7 +3017,7 @@ UGameInstance* UEditorEngine::CreateInnerProcessPIEGameInstance(FRequestPlaySess
 		{
 			if (FAudioDeviceHandle GameInstanceAudioDevice = GameInstance->GetWorld()->GetAudioDevice())
 			{
-				GameInstanceAudioDevice->SetTransientMasterVolume(0.0f);
+				GameInstanceAudioDevice->SetTransientPrimaryVolume(0.0f);
 			}
 		}
 		if (InParams.EditorPlaySettings->SoloAudioInFirstPIEClient)
diff --git a/Engine/Source/Runtime/AudioMixer/Private/AudioMixerDevice.cpp b/Engine/Source/Runtime/AudioMixer/Private/AudioMixerDevice.cpp
--- a/Engine/Source/Runtime/AudioMixer/Private/AudioMixerDevice.cpp
+++ b/Engine/Source/Runtime/AudioMixer/Private/AudioMixerDevice.cpp
@@ -161,7 +161,7 @@ namespace Audio
 		: QuantizedEventClockManager(this)
 		, AudioMixerPlatform(InAudioMixerPlatform)
 		, AudioClockDelta(0.0)
-		, PreviousMasterVolume((float)INDEX_NONE)
+		, PreviousPrimaryVolume((float)INDEX_NONE)
 		, GameOrAudioThreadId(INDEX_NONE)
 		, AudioPlatformThreadId(INDEX_NONE)
 		, bDebugOutputEnabled(false)
@@ -636,11 +636,11 @@ namespace Audio
 			}
 
 			// Check if the background mute changed state and update the submixes which are enabled to do background muting.
-			const float CurrentMasterVolume = GetMasterVolume();
-			if (!FMath::IsNearlyEqual(PreviousMasterVolume, CurrentMasterVolume))
+			const float CurrentPrimaryVolume = GetPrimaryVolume();
+			if (!FMath::IsNearlyEqual(PreviousPrimaryVolume, CurrentPrimaryVolume))
 			{
-				PreviousMasterVolume = CurrentMasterVolume;
-				bool IsMuted = FMath::IsNearlyZero(CurrentMasterVolume);
+				PreviousPrimaryVolume = CurrentPrimaryVolume;
+				bool IsMuted = FMath::IsNearlyZero(CurrentPrimaryVolume);
 
 				for (TObjectIterator<USoundSubmix> It; It; ++It)
 				{
diff --git a/Engine/Source/Runtime/AudioMixer/Private/AudioMixerSource.cpp b/Engine/Source/Runtime/AudioMixer/Private/AudioMixerSource.cpp
--- a/Engine/Source/Runtime/AudioMixer/Private/AudioMixerSource.cpp
+++ b/Engine/Source/Runtime/AudioMixer/Private/AudioMixerSource.cpp
@@ -1284,7 +1284,7 @@ namespace Audio
 		if (!AudioDevice->IsAudioDeviceMuted())
 		{
 			// 1. Apply device gain stage(s)
-			CurrentVolume = WaveInstance->ActiveSound->bIsPreviewSound ? 1.0f : AudioDevice->GetMasterVolume();
+			CurrentVolume = WaveInstance->ActiveSound->bIsPreviewSound ? 1.0f : AudioDevice->GetPrimaryVolume();
 			CurrentVolume *= AudioDevice->GetPlatformAudioHeadroom();
 
 			// 2. Apply instance gain stage(s)
diff --git a/Engine/Source/Runtime/AudioMixer/Public/AudioMixerDevice.h b/Engine/Source/Runtime/AudioMixer/Public/AudioMixerDevice.h
--- a/Engine/Source/Runtime/AudioMixer/Public/AudioMixerDevice.h
+++ b/Engine/Source/Runtime/AudioMixer/Public/AudioMixerDevice.h
@@ -391,7 +391,7 @@ namespace Audio
 		double AudioClockDelta;
 
 		/** What the previous master volume was. */
-		float PreviousMasterVolume;
+		float PreviousPrimaryVolume;
 
 		/** Timing data for audio thread. */
 		FAudioThreadTimingData AudioThreadTimingData;
diff --git a/Engine/Source/Runtime/Engine/Private/Audio.cpp b/Engine/Source/Runtime/Engine/Private/Audio.cpp
--- a/Engine/Source/Runtime/Engine/Private/Audio.cpp
+++ b/Engine/Source/Runtime/Engine/Private/Audio.cpp
@@ -976,7 +976,7 @@ float FWaveInstance::GetActualVolume() const
 		if (!ActiveSound->bIsPreviewSound)
 		{
 			check(ActiveSound->AudioDevice);
-			ActualVolume *= ActiveSound->AudioDevice->GetMasterVolume();
+			ActualVolume *= ActiveSound->AudioDevice->GetPrimaryVolume();
 		}
 	}
 
diff --git a/Engine/Source/Runtime/Engine/Private/AudioDevice.cpp b/Engine/Source/Runtime/Engine/Private/AudioDevice.cpp
--- a/Engine/Source/Runtime/Engine/Private/AudioDevice.cpp
+++ b/Engine/Source/Runtime/Engine/Private/AudioDevice.cpp
@@ -301,8 +301,8 @@ FAudioDevice::FAudioDevice()
 	, CurrentTick(0)
 	, TestAudioComponent(nullptr)
 	, DebugState(DEBUGSTATE_None)
-	, TransientMasterVolume(1.0f)
-	, MasterVolume(1.0f)
+	, TransientPrimaryVolume(1.0f)
+	, PrimaryVolume(1.0f)
 	, GlobalPitchScale(1.0f)
 	, LastUpdateTime(FPlatformTime::Seconds())
 	, NextResourceID(1)
@@ -4436,11 +4436,11 @@ void FAudioDevice::Update(bool bGameTicking)
 	bIsStoppingVoicesEnabled = !DisableStoppingVoicesCvar;
 
 	// Update the master volume
-	MasterVolume = GetTransientMasterVolume();
+	PrimaryVolume = GetTransientPrimaryVolume();
 	
 	if (!DisableAppVolumeCvar)
 	{
-		MasterVolume *= FApp::GetVolumeMultiplier();
+		PrimaryVolume *= FApp::GetVolumeMultiplier();
 	}
 
 	UpdateAudioPluginSettingsObjectCache();
@@ -6961,22 +6961,22 @@ bool FAudioDevice::CanUseVRAudioDevice()
 	}
 }
 
-void FAudioDevice::SetTransientMasterVolume(const float InTransientMasterVolume)
+void FAudioDevice::SetTransientPrimaryVolume(const float InTransientPrimaryVolume)
 {
 	if (!IsInAudioThread())
 	{
-		DECLARE_CYCLE_STAT(TEXT("FAudioThreadTask.SetTransientMasterVolume"), STAT_SetTransientMasterVolume, STATGROUP_AudioThreadCommands);
+		DECLARE_CYCLE_STAT(TEXT("FAudioThreadTask.SetTransientPrimaryVolume"), STAT_SetTransientPrimaryVolume, STATGROUP_AudioThreadCommands);
 
 		FAudioDevice* AudioDevice = this;
-		FAudioThread::RunCommandOnAudioThread([AudioDevice, InTransientMasterVolume]()
+		FAudioThread::RunCommandOnAudioThread([AudioDevice, InTransientPrimaryVolume]()
 		{
-			AudioDevice->SetTransientMasterVolume(InTransientMasterVolume);
-		}, GET_STATID(STAT_SetTransientMasterVolume));
+			AudioDevice->SetTransientPrimaryVolume(InTransientPrimaryVolume);
+		}, GET_STATID(STAT_SetTransientPrimaryVolume));
 
 		return;
 	}
 
-	TransientMasterVolume = InTransientMasterVolume;
+	TransientPrimaryVolume = InTransientPrimaryVolume;
 }
 
 FSoundSource* FAudioDevice::GetSoundSource(FWaveInstance* WaveInstance) const
diff --git a/Engine/Source/Runtime/Engine/Private/PlayerCameraManager.cpp b/Engine/Source/Runtime/Engine/Private/PlayerCameraManager.cpp
--- a/Engine/Source/Runtime/Engine/Private/PlayerCameraManager.cpp
+++ b/Engine/Source/Runtime/Engine/Private/PlayerCameraManager.cpp
@@ -474,7 +474,7 @@ void APlayerCameraManager::ApplyAudioFade()
 		{
 			if (FAudioDevice* AudioDevice = World->GetAudioDeviceRaw())
 			{
-				AudioDevice->SetTransientMasterVolume(1.0f - FadeAmount);
+				AudioDevice->SetTransientPrimaryVolume(1.0f - FadeAmount);
 			}
 		}
 	}
@@ -497,7 +497,7 @@ void APlayerCameraManager::StopAudioFade()
 		{
 			if (FAudioDevice* AudioDevice = World->GetAudioDeviceRaw())
 			{
-				AudioDevice->SetTransientMasterVolume(1.0f);
+				AudioDevice->SetTransientPrimaryVolume(1.0f);
 			}
 		}
 	}
diff --git a/Engine/Source/Runtime/Engine/Private/UnrealEngine.cpp b/Engine/Source/Runtime/Engine/Private/UnrealEngine.cpp
--- a/Engine/Source/Runtime/Engine/Private/UnrealEngine.cpp
+++ b/Engine/Source/Runtime/Engine/Private/UnrealEngine.cpp
@@ -14444,7 +14444,7 @@ bool UEngine::LoadMap( FWorldContext& WorldContext, FURL URL, class UPendingNetG
 		if (FAudioDevice* AudioDevice = WorldContext.World()->GetAudioDeviceRaw())
 		{
 			AudioDevice->Flush(WorldContext.World());
-			AudioDevice->SetTransientMasterVolume(1.0f);
+			AudioDevice->SetTransientPrimaryVolume(1.0f);
 		}
 
 		WorldContext.SetCurrentWorld(nullptr);
diff --git a/Engine/Source/Runtime/Engine/Public/AudioDevice.h b/Engine/Source/Runtime/Engine/Public/AudioDevice.h
--- a/Engine/Source/Runtime/Engine/Public/AudioDevice.h
+++ b/Engine/Source/Runtime/Engine/Public/AudioDevice.h
@@ -1885,11 +1885,20 @@ public:
 
 	bool AreStartupSoundsPreCached() const { return bStartupSoundsPreCached; }
 
-	float GetTransientMasterVolume() const { check(IsInAudioThread()); return TransientMasterVolume; }
-	void SetTransientMasterVolume(float TransientMasterVolume);
+	UE_DEPRECATED(5.1, "GetTransientMasterVolume has been deprecated. Please use GetTransientPrimaryVolume instead.")
+	float GetTransientMasterVolume() const { check(IsInAudioThread()); return TransientPrimaryVolume; }
+
+	UE_DEPRECATED(5.1, "SetTransientMasterVolume has been deprecated. Please use SetTransientPrimaryVolume instead.")
+	void SetTransientMasterVolume(float TransientPrimaryVolume);
+
+	UE_DEPRECATED(5.1, "GetMasterVolume has been deprecated. Please use GetPrimaryVolume instead.")
+	float GetMasterVolume() const { return PrimaryVolume; }
+
+	float GetTransientPrimaryVolume() const { check(IsInAudioThread()); return TransientPrimaryVolume; }
+	void SetTransientPrimaryVolume(float TransientPrimaryVolume);
 
 	/** Returns the volume that combines transient master volume and the FApp::GetVolumeMultiplier() value */
-	float GetMasterVolume() const { return MasterVolume; }
+	float GetPrimaryVolume() const { return PrimaryVolume; }
 
 	FSoundSource* GetSoundSource(FWaveInstance* WaveInstance) const;
 
@@ -2036,10 +2045,10 @@ private:
 	TEnumAsByte<enum EDebugState> DebugState;
 
 	/** transient master volume multiplier that can be modified at runtime without affecting user settings automatically reset to 1.0 on level change */
-	float TransientMasterVolume;
+	float TransientPrimaryVolume;
 
-	/** The master volume of the game combines the FApp::GetVolumeMultipler() value and the TransientMastervolume. */
-	float MasterVolume;
+	/** The master volume of the game combines the FApp::GetVolumeMultipler() value and the TransientPrimaryVolume. */
+	float PrimaryVolume;
 
 	/** Global dynamic pitch scale parameter */
 	FDynamicParameter GlobalPitchScale;
diff --git a/Engine/Source/Runtime/Launch/Private/IOS/LaunchIOS.cpp b/Engine/Source/Runtime/Launch/Private/IOS/LaunchIOS.cpp
--- a/Engine/Source/Runtime/Launch/Private/IOS/LaunchIOS.cpp
+++ b/Engine/Source/Runtime/Launch/Private/IOS/LaunchIOS.cpp
@@ -74,14 +74,14 @@ void FAppEntry::Suspend(bool bIsInterrupt)
 						{
 							if (GEngine && GEngine->GetMainAudioDevice())
 							{
-								GEngine->GetMainAudioDevice()->SetTransientMasterVolume(0.0f);
+								GEngine->GetMainAudioDevice()->SetTransientPrimaryVolume(0.0f);
 							}
 						}, TStatId());
 					}, TStatId(), NULL, ENamedThreads::GameThread);
 				}
 				else
 				{
-					AudioDevice->SetTransientMasterVolume(0.0f);
+					AudioDevice->SetTransientPrimaryVolume(0.0f);
 				}
 			}
 			else
@@ -169,14 +169,14 @@ void FAppEntry::Resume(bool bIsInterrupt)
 						{
 							if (GEngine && GEngine->GetMainAudioDevice())
 							{
-								GEngine->GetMainAudioDevice()->SetTransientMasterVolume(1.0f);
+								GEngine->GetMainAudioDevice()->SetTransientPrimaryVolume(1.0f);
 							}
 						}, TStatId());
 					}, TStatId(), NULL, ENamedThreads::GameThread);
 				}
 				else
 				{
-					AudioDevice->SetTransientMasterVolume(1.0f);
+					AudioDevice->SetTransientPrimaryVolume(1.0f);
 				}
 			}
 			else
diff --git a/Engine/Source/Runtime/Launch/Private/LaunchEngineLoop.cpp b/Engine/Source/Runtime/Launch/Private/LaunchEngineLoop.cpp
--- a/Engine/Source/Runtime/Launch/Private/LaunchEngineLoop.cpp
+++ b/Engine/Source/Runtime/Launch/Private/LaunchEngineLoop.cpp
@@ -1543,10 +1543,10 @@ int32 FEngineLoop::PreInitPreStartupScreen(const TCHAR* CmdLine)
 	FDelayedAutoRegisterHelper::RunAndClearDelayedAutoRegisterDelegates(EDelayedRegisterRunPhase::StartOfEnginePreInit);
 	SCOPED_BOOT_TIMING("FEngineLoop::PreInitPreStartupScreen");
 
-	// GLog is initialized lazily and its default master thread is the thread that initialized it.
+	// GLog is initialized lazily and its default primary thread is the thread that initialized it.
 	// This lazy initialization can happen during initialization of a DLL, which Windows does on a
-	// worker thread, which makes that worker thread the master thread. Make this the master thread
-	// until initialization is far enough along to try to start a dedicated master thread.
+	// worker thread, which makes that worker thread the primary thread. Make this the primary thread
+	// until initialization is far enough along to try to start a dedicated primary thread.
 	if (GLog)
 	{
 		GLog->SetCurrentThreadAsPrimaryThread();
@@ -1697,7 +1697,7 @@ int32 FEngineLoop::PreInitPreStartupScreen(const TCHAR* CmdLine)
 	// as soon as we determine whether GIsEditor == false
 	GLog->EnableBacklog(true);
 
-	// Try to start the dedicated master thread now that the command line is available.
+	// Try to start the dedicated primary thread now that the command line is available.
 	if (!FParse::Param(FCommandLine::Get(), TEXT("NoLogThread")))
 	{
 		GLog->TryStartDedicatedPrimaryThread();
