Fix for nounity CIS in AudioFormatSettings
#rnx

diff --git a/Engine/Source/Runtime/Engine/Classes/Sound/AudioFormatSettings.h b/Engine/Source/Runtime/Engine/Classes/Sound/AudioFormatSettings.h
--- a/Engine/Source/Runtime/Engine/Classes/Sound/AudioFormatSettings.h
+++ b/Engine/Source/Runtime/Engine/Classes/Sound/AudioFormatSettings.h
@@ -1,6 +1,12 @@
 // Copyright Epic Games, Inc. All Rights Reserved.
 #pragma once
 
+#include "Containers/Array.h"
+#include "UObject/NameTypes.h"
+
+class FString;
+class FConfigCacheIni;
+
 namespace Audio
 {
 	class ENGINE_API FAudioFormatSettings
diff --git a/Engine/Source/Runtime/Engine/Private/AudioFormatSettings.cpp b/Engine/Source/Runtime/Engine/Private/AudioFormatSettings.cpp
--- a/Engine/Source/Runtime/Engine/Private/AudioFormatSettings.cpp
+++ b/Engine/Source/Runtime/Engine/Private/AudioFormatSettings.cpp
@@ -2,6 +2,8 @@
 
 #include "Sound/AudioFormatSettings.h"
 #include "Algo/Transform.h"
+#include "Misc/ConfigCacheIni.h"
+#include "Containers/UnrealString.h"
 
 namespace Audio
 {
