WorldPartitionHandleTest: Add proper world cleanup to fix GC warnings (root cause for FSubsystemCollectionBase::Deinitialize crash)

diff --git a/Engine/Source/Runtime/Engine/Private/Tests/WorldPartition/WorldPartitionHandleTests.cpp b/Engine/Source/Runtime/Engine/Private/Tests/WorldPartition/WorldPartitionHandleTests.cpp
--- a/Engine/Source/Runtime/Engine/Private/Tests/WorldPartition/WorldPartitionHandleTests.cpp
+++ b/Engine/Source/Runtime/Engine/Private/Tests/WorldPartition/WorldPartitionHandleTests.cpp
@@ -4,6 +4,8 @@
 #include "WorldPartition/WorldPartitionHandle.h"
 #include "WorldPartition/WorldPartitionActorDesc.h"
 #include "WorldPartition/ActorDescContainer.h"
+#include "Engine/World.h"
+#include "PackageTools.h"
 
 #if WITH_DEV_AUTOMATION_TESTS
 
@@ -234,8 +236,18 @@ namespace WorldPartitionTests
 			TestTrue(TEXT("Invalid container test"), Handle.IsValid());
 			TestTrue(TEXT("Invalid container test"), Reference.IsValid());
 
+			// Make sure to cleanup world before collecting garbage so it gets uninitialized
+			UPackage* Package = FindPackage(NULL, *ActorDescContainer->GetContainerPackage().ToString());
+			check(Package);
+			UWorld* World = UWorld::FindWorldInPackage(Package);
+			check(World->IsInitializedAndNeedsCleanup());
+			World->CleanupWorld();
+			
 			CollectGarbage(GARBAGE_COLLECTION_KEEPFLAGS);
 
+			// Unload Package so test can run twice without issues.
+			UPackageTools::UnloadPackages({ Package });
+
 			TestFalse(TEXT("Invalid container test"), Handle.IsValid());
 			TestFalse(TEXT("Invalid container test"), Reference.IsValid());
 		}
