Turn down AudioFormatSettings startup message to "Verbose" from "Display"as there's lots platforms and its spammy.

diff --git a/Engine/Source/Runtime/Engine/Private/AudioFormatSettings.cpp b/Engine/Source/Runtime/Engine/Private/AudioFormatSettings.cpp
--- a/Engine/Source/Runtime/Engine/Private/AudioFormatSettings.cpp
+++ b/Engine/Source/Runtime/Engine/Private/AudioFormatSettings.cpp
@@ -99,7 +99,7 @@ namespace Audio
 			AllFormatsConcat += i.ToString() + TEXT(" ");
 		}
 
-		UE_LOG(LogAudio, Display, TEXT("AudioFormatSettings: TargetName='%s', AllWaveFormats=( %s), PlatformFormat='%s', PlatformStreamingFormat='%s', FallbackFormat='%s'"),
+		UE_LOG(LogAudio, Verbose, TEXT("AudioFormatSettings: TargetName='%s', AllWaveFormats=( %s), PlatformFormat='%s', PlatformStreamingFormat='%s', FallbackFormat='%s'"),
 			*InPlatformIdentifierForLogging, *AllFormatsConcat, *PlatformFormat.ToString(), *PlatformStreamingFormat.ToString(), *FallbackFormat.ToString());
 
 	#endif //!NO_LOGGING
