Fix for forceinline not working

diff --git a/Engine/Source/Runtime/Engine/Private/SoundWave.cpp b/Engine/Source/Runtime/Engine/Private/SoundWave.cpp
--- a/Engine/Source/Runtime/Engine/Private/SoundWave.cpp
+++ b/Engine/Source/Runtime/Engine/Private/SoundWave.cpp
@@ -1794,11 +1794,8 @@ void USoundWave::PostLoad()
 	INC_FLOAT_STAT_BY(STAT_AudioBufferTime, Duration);
 	INC_FLOAT_STAT_BY(STAT_AudioBufferTimeChannels, NumChannels * Duration);
 
-	// update shared state if loading behavior dictates we'll be using the streaming cache
-	if (ActualLoadingBehavior != ESoundWaveLoadingBehavior::ForceInline)
-	{
-		SoundWaveDataPtr->InitializeDataFromSoundWave(*this);
-	}
+	// update shared state
+	SoundWaveDataPtr->InitializeDataFromSoundWave(*this);
 }
 
 bool USoundWave::LoadZerothChunk()
