RepMovementNetSerializer fixup.

Members have been added to FRepMovement since the serializer was implemented.

diff --git a/Engine/Plugins/Runtime/ReplicationSystemTestPlugin/Source/Private/Tests/Serialization/TestRepMovementNetSerializer.cpp b/Engine/Plugins/Runtime/ReplicationSystemTestPlugin/Source/Private/Tests/Serialization/TestRepMovementNetSerializer.cpp
new file mode 100644
--- /dev/null
+++ b/Engine/Plugins/Runtime/ReplicationSystemTestPlugin/Source/Private/Tests/Serialization/TestRepMovementNetSerializer.cpp
@@ -0,0 +1,165 @@
+// Copyright Epic Games, Inc. All Rights Reserved.
+
+#include "TestNetSerializerFixture.h"
+#include "RepMovementNetSerializer.h"
+#include "Engine/ReplicatedState.h"
+
+UE::Net::FTestMessage& operator<<(UE::Net::FTestMessage& Message, const FRepMovement& Value)
+{
+	return Message << 
+		"LinearVelocity: " << Value.LinearVelocity << '\n' <<
+		"AngularVelocity: " << Value.AngularVelocity << '\n' <<
+		"Location: " << Value.Location << '\n' <<
+		"Rotation: " << Value.Rotation.ToString() << '\n' <<
+		"bSimulatedPhysicSleep: " << Value.bSimulatedPhysicSleep << '\n' <<
+		"bRepPhysics: " << Value.bRepPhysics << '\n' <<
+		"ServerFrame: " << Value.ServerFrame << '\n' <<
+		"ServerPhysicsHandle: " << Value.ServerPhysicsHandle << '\n' <<
+		"LocationQuantizationLevel: " << (uint32)Value.LocationQuantizationLevel << '\n' <<
+		"VelocityQuantizationLevel: " << (uint32)Value.VelocityQuantizationLevel << '\n' <<
+		"RotationQuantizationLevel: " << (uint32)Value.RotationQuantizationLevel;
+}
+
+namespace UE::Net::Private
+{
+
+static FTestMessage& PrintRepMovementNetSerializerConfig(FTestMessage& Message, const FNetSerializerConfig& InConfig)
+{
+	return Message;
+}
+
+class FTestRepMovementNetSerializer : public TTestNetSerializerFixture<PrintRepMovementNetSerializerConfig, FRepMovement>
+{
+	typedef TTestNetSerializerFixture<PrintRepMovementNetSerializerConfig, FRepMovement> Super;
+
+public:
+	FTestRepMovementNetSerializer() : Super(UE_NET_GET_SERIALIZER(FRepMovementNetSerializer)) {}
+
+	void TestIsEqual();
+	void TestSerialize();
+	void TestSerializeDelta();
+
+protected:
+	virtual void SetUp() override;
+
+	inline static bool bIsClassInitialized = false;
+	inline static TArray<FRepMovement> Values;
+
+private:
+	void OneTimeClassInitialization();
+};
+
+UE_NET_TEST_FIXTURE(FTestRepMovementNetSerializer, TestIsEqual)
+{
+	TestIsEqual();
+}
+
+UE_NET_TEST_FIXTURE(FTestRepMovementNetSerializer, TestSerialize)
+{
+	TestSerialize();
+}
+
+UE_NET_TEST_FIXTURE(FTestRepMovementNetSerializer, TestSerializeDelta)
+{
+	TestSerializeDelta();
+}
+
+void FTestRepMovementNetSerializer::SetUp()
+{
+	Super::SetUp();
+	OneTimeClassInitialization();
+}
+
+void FTestRepMovementNetSerializer::TestIsEqual()
+{
+	constexpr SIZE_T TestRoundCount = 2;
+	TArray<FRepMovement> CompareValues[TestRoundCount];
+	TArray<bool> ExpectedResults[TestRoundCount];
+
+	const SIZE_T ValueCount = Values.Num();
+	for (SIZE_T TestRoundIt = 0, TestRoundEndIt = TestRoundCount; TestRoundIt != TestRoundEndIt; ++TestRoundIt)
+	{
+		CompareValues[TestRoundIt].Reserve(ValueCount);
+		ExpectedResults[TestRoundIt].Reserve(ValueCount);
+
+		const SIZE_T SourceValueItOffset = TestRoundIt;
+		const bool ExpectedResult = (TestRoundIt == 0U);
+		for (SIZE_T ValueIt = 0, ValueEndIt = ValueCount; ValueIt < ValueEndIt; ++ValueIt)
+		{
+			CompareValues[TestRoundIt].Push(Values[(ValueIt + SourceValueItOffset) % ValueEndIt]);
+			ExpectedResults[TestRoundIt].Push(ExpectedResult);
+		}
+	}
+
+	const FRepMovementNetSerializerConfig Config;
+	for (SIZE_T TestRoundIt = 0, TestRoundEndIt = TestRoundCount; TestRoundIt != TestRoundEndIt; ++TestRoundIt)
+	{
+		constexpr bool bQuantizedCompare = true;
+		bool bSuccess = Super::TestIsEqual(Values.GetData(), CompareValues[TestRoundIt].GetData(), ExpectedResults[TestRoundIt].GetData(), CompareValues[TestRoundIt].Num(), Config, bQuantizedCompare);
+		if (!bSuccess)
+		{
+			return;
+		}
+	}
+}
+
+void FTestRepMovementNetSerializer::TestSerialize()
+{
+	const FRepMovementNetSerializerConfig Config;
+	constexpr bool bQuantizedCompare = true;
+	Super::TestSerialize(Values.GetData(), Values.GetData(), Values.Num(), Config, bQuantizedCompare);
+}
+
+void FTestRepMovementNetSerializer::TestSerializeDelta()
+{
+	const FRepMovementNetSerializerConfig Config;
+	Super::TestSerializeDelta(Values.GetData(), Values.Num(), Config);
+}
+
+void FTestRepMovementNetSerializer::OneTimeClassInitialization()
+{
+	if (bIsClassInitialized)
+	{
+		return;
+	}
+
+	// Default value
+	{
+		Values.Emplace();
+	}
+
+	// Values with location, velocities, rotation, ServerFrame and physics handle
+	{
+		FRepMovement& Value = Values.Emplace_GetRef();
+		Value.LinearVelocity = FVector(10.01, -20.23, 40.05);
+		Value.AngularVelocity = FVector(1000.23, -2000.05, 4000.01);
+		Value.Location = FVector(100000.05, -200000.01, 400000.23);
+		Value.Rotation = FRotator(213.16, 23.75, 109.12);
+
+		FRepMovement ValueWithServerFrame = Value;
+		ValueWithServerFrame.ServerFrame = 1908246313;
+		Values.Add(ValueWithServerFrame);
+
+		FRepMovement ValueWithPhysicsHandle = Value;
+		ValueWithPhysicsHandle.ServerPhysicsHandle = 4711;
+		Values.Add(ValueWithPhysicsHandle);
+	}
+
+	// Values with location, velocities, rotation, bRepPhysics and bSimulatedPhysicSleep
+	{
+		FRepMovement& Value = Values.Emplace_GetRef();
+		Value.LinearVelocity = FVector(10.01, 40.05, -20.23);
+		Value.AngularVelocity = FVector(-2000.05, 1000.23, 4000.01);
+		Value.Location = FVector(400000.23, 100000.05, -200000.01);
+		Value.Rotation = FRotator(213.16, 23.75, 109.12);
+		Value.bRepPhysics = true;
+
+		FRepMovement ValueWithSimulatedPhysicSleep = Value;
+		ValueWithSimulatedPhysicSleep.bSimulatedPhysicSleep = true;
+		Values.Add(ValueWithSimulatedPhysicSleep);
+	}
+
+	bIsClassInitialized = true;
+}
+
+}
diff --git a/Engine/Source/Runtime/Engine/Private/RepMovementNetSerializer.cpp b/Engine/Source/Runtime/Engine/Private/RepMovementNetSerializer.cpp
--- a/Engine/Source/Runtime/Engine/Private/RepMovementNetSerializer.cpp
+++ b/Engine/Source/Runtime/Engine/Private/RepMovementNetSerializer.cpp
@@ -4,7 +4,9 @@
 
 #if UE_WITH_IRIS
 #include "Engine/ReplicatedState.h"
+#include "Iris/Serialization/BitPacking.h"
 #include "Iris/Serialization/NetBitStreamReader.h"
+#include "Iris/Serialization/NetBitStreamUtil.h"
 #include "Iris/Serialization/NetBitStreamWriter.h"
 #include "Iris/Serialization/NetSerializerDelegates.h"
 #include "Iris/Serialization/PackedVectorNetSerializers.h"
@@ -27,8 +29,10 @@ struct FRepMovementNetSerializer
 		uint64 LinearVelocity[4];
 		uint64 Location[4];
 		uint16 Rotation[4];
+		int32 ServerFrame;
+		int32 ServerPhysicsHandle;
 
-		uint16 Flags : 2;
+		uint16 Flags : 4;
 		uint16 VelocityQuantizationLevel : 2;
 		uint16 LocationQuantizationLevel : 2;
 		uint16 RotationQuantizationLevel : 1;
@@ -60,6 +64,8 @@ private:
 	{
 		Flag_SimulatedPhysicSleep = 1U,
 		Flag_RepPhysics = 2U,
+		Flag_ServerFrameIsPresent = 4U,
+		Flag_ServerPhysicsHandleIsPresent = 8U,
 	};
 
 	// Dynamic precision support for FVector and FRotator
@@ -77,6 +83,10 @@ private:
 	inline static const FNetSerializer* RotatorNetSerializers[2] = {};
 
 private:
+	static constexpr SIZE_T DeltaBitCountTableEntryCount = 3;
+	// Bit counts aiming to have small value changes use few bits.
+	static constexpr uint8 DeltaBitCountTable[] = {0, 4, 14};
+
 	class FNetSerializerRegistryDelegates final : private UE::Net::FNetSerializerRegistryDelegates
 	{
 	public:
@@ -105,11 +115,13 @@ void FRepMovementNetSerializer::Serialize(FNetSerializationContext& Context, con
 
 	FNetBitStreamWriter* Writer = Context.GetBitStreamWriter();
 
-	const uint32 FlagsAndQuantizationLevels = (Value.RotationQuantizationLevel << 6U) | (Value.LocationQuantizationLevel << 4U) | (Value.VelocityQuantizationLevel << 2U) | (Value.Flags);
+	constexpr uint32 QuantizationLevelsBitOffset = 4;
+	const uint32 FlagsAndQuantizationLevels = (((Value.RotationQuantizationLevel << 4U) | (Value.LocationQuantizationLevel << 2U) | Value.VelocityQuantizationLevel) << QuantizationLevelsBitOffset) | Value.Flags;
 	if (Writer->WriteBool(FlagsAndQuantizationLevels != 0))
 	{
 		UE_NET_TRACE_SCOPE(FlagsAndQuantizationLevels, *Writer, Context.GetTraceCollector(), ENetTraceVerbosity::Verbose);
-		Writer->WriteBits(FlagsAndQuantizationLevels, 7U);
+		constexpr uint32 QuantizationLevelsBitCount = 5U;
+		Writer->WriteBits(FlagsAndQuantizationLevels, QuantizationLevelsBitCount + QuantizationLevelsBitOffset);
 	}
 
 	// Angular velocity
@@ -160,6 +172,18 @@ void FRepMovementNetSerializer::Serialize(FNetSerializationContext& Context, con
 		MemberArgs.Source = NetSerializerValuePointer(&Value.Rotation[0]);
 		Serializer->Serialize(Context, MemberArgs);
 	}
+
+	// ServerFrame
+	if (Value.Flags & Flag_ServerFrameIsPresent)
+	{
+		WritePackedInt32(Writer, Value.ServerFrame);
+	}
+
+	// ServerPhsyicsHandle
+	if (Value.Flags & Flag_ServerPhysicsHandleIsPresent)
+	{
+		WritePackedInt32(Writer, Value.ServerPhysicsHandle);
+	}
 }
 
 void FRepMovementNetSerializer::Deserialize(FNetSerializationContext& Context, const FNetDeserializeArgs& Args)
@@ -170,11 +194,12 @@ void FRepMovementNetSerializer::Deserialize(FNetSerializationContext& Context, c
 	if (Reader->ReadBool())
 	{
 		UE_NET_TRACE_SCOPE(FlagsAndQuantizationLevels, *Reader, Context.GetTraceCollector(), ENetTraceVerbosity::Verbose);
-		const uint32 FlagsAndQuantizationLevels = Reader->ReadBits(7U);
-		TempValue.Flags = (FlagsAndQuantizationLevels & 3U);
-		TempValue.VelocityQuantizationLevel = (FlagsAndQuantizationLevels >> 2U) & 3U;
-		TempValue.LocationQuantizationLevel = (FlagsAndQuantizationLevels >> 4U) & 3U;
-		TempValue.RotationQuantizationLevel = (FlagsAndQuantizationLevels >> 6U) & 1U;
+		const uint32 FlagsAndQuantizationLevels = Reader->ReadBits(9U);
+		constexpr uint32 QuantizationLevelsBitOffset = 4U;
+		TempValue.Flags = (FlagsAndQuantizationLevels & ((1U << QuantizationLevelsBitOffset) - 1U));
+		TempValue.VelocityQuantizationLevel = (FlagsAndQuantizationLevels >> (QuantizationLevelsBitOffset + 0U)) & 3U;
+		TempValue.LocationQuantizationLevel = (FlagsAndQuantizationLevels >> (QuantizationLevelsBitOffset + 2U)) & 3U;
+		TempValue.RotationQuantizationLevel = (FlagsAndQuantizationLevels >> (QuantizationLevelsBitOffset + 4U)) & 1U;
 
 		if (VectorNetQuantizeNetSerializers[TempValue.VelocityQuantizationLevel] == nullptr || RotatorNetSerializers[TempValue.RotationQuantizationLevel] == nullptr)
 		{
@@ -232,6 +257,26 @@ void FRepMovementNetSerializer::Deserialize(FNetSerializationContext& Context, c
 		Serializer->Deserialize(Context, MemberArgs);
 	}
 
+	// ServerFrame
+	if (TempValue.Flags & Flag_ServerFrameIsPresent)
+	{
+		TempValue.ServerFrame = ReadPackedInt32(Reader);
+	}
+	else
+	{
+		TempValue.ServerFrame = 0;
+	}
+
+	// ServerPhysicsHandle
+	if (TempValue.Flags & Flag_ServerPhysicsHandleIsPresent)
+	{
+		TempValue.ServerPhysicsHandle = ReadPackedInt32(Reader);
+	}
+	else
+	{
+		TempValue.ServerPhysicsHandle = INDEX_NONE;
+	}
+
 	QuantizedType& Target = *reinterpret_cast<QuantizedType*>(Args.Target);
 	Target = TempValue;
 }
@@ -252,10 +297,14 @@ void FRepMovementNetSerializer::SerializeDelta(FNetSerializationContext& Context
 	}
 
 	// We know the quantization levels are equal so we only need to care about the flags.
-	Writer->WriteBits(Value.Flags, 2U);
+	uint32 Flags = Value.Flags & (Flag_SimulatedPhysicSleep | Flag_RepPhysics);
+	Flags |= (Value.ServerFrame != PrevValue.ServerFrame ? Flag_ServerFrameIsPresent : 0U);
+	Flags |= (Value.ServerPhysicsHandle != PrevValue.ServerPhysicsHandle ? Flag_ServerPhysicsHandleIsPresent : 0U);
+
+	Writer->WriteBits(Flags, 4U);
 
 	// Angular velocity
-	if (Value.Flags & Flag_RepPhysics)
+	if (Flags & Flag_RepPhysics)
 	{
 		UE_NET_TRACE_SCOPE(AngularVelocity, *Writer, Context.GetTraceCollector(), ENetTraceVerbosity::Verbose);
 		const FNetSerializer* Serializer = VectorNetQuantizeNetSerializers[Value.VelocityQuantizationLevel];
@@ -315,6 +364,16 @@ void FRepMovementNetSerializer::SerializeDelta(FNetSerializationContext& Context
 		MemberArgs.Prev = NetSerializerValuePointer(&PrevValue.Rotation[0]);
 		Serializer->SerializeDelta(Context, MemberArgs);
 	}
+
+	if (Flags & Flag_ServerFrameIsPresent)
+	{
+		SerializeIntDelta(*Writer, Value.ServerFrame, PrevValue.ServerFrame, DeltaBitCountTable, DeltaBitCountTableEntryCount, 32U);
+	}
+
+	if (Flags & Flag_ServerPhysicsHandleIsPresent)
+	{
+		SerializeIntDelta(*Writer, Value.ServerPhysicsHandle, PrevValue.ServerPhysicsHandle, DeltaBitCountTable, DeltaBitCountTableEntryCount, 32U);
+	}
 }
 
 void FRepMovementNetSerializer::DeserializeDelta(FNetSerializationContext& Context, const FNetDeserializeDeltaArgs& Args)
@@ -334,10 +393,10 @@ void FRepMovementNetSerializer::DeserializeDelta(FNetSerializationContext& Conte
 	TempValue.VelocityQuantizationLevel = PrevValue.VelocityQuantizationLevel;
 	TempValue.LocationQuantizationLevel = PrevValue.LocationQuantizationLevel;
 	TempValue.RotationQuantizationLevel = PrevValue.RotationQuantizationLevel;
-	TempValue.Flags = Reader->ReadBits(2U);
+	const uint32 Flags = Reader->ReadBits(4U);
 
 	// Angular velocity
-	if (TempValue.Flags & Flag_RepPhysics)
+	if (Flags & Flag_RepPhysics)
 	{
 		UE_NET_TRACE_SCOPE(AngularVelocity, *Reader, Context.GetTraceCollector(), ENetTraceVerbosity::Verbose);
 		const FNetSerializer* Serializer = VectorNetQuantizeNetSerializers[TempValue.VelocityQuantizationLevel];
@@ -398,6 +457,30 @@ void FRepMovementNetSerializer::DeserializeDelta(FNetSerializationContext& Conte
 		Serializer->DeserializeDelta(Context, MemberArgs);
 	}
 
+	if (Flags & Flag_ServerFrameIsPresent)
+	{
+		DeserializeIntDelta(*Reader, TempValue.ServerFrame, PrevValue.ServerFrame, DeltaBitCountTable, DeltaBitCountTableEntryCount, 32U);
+	}
+	else
+	{
+		TempValue.ServerFrame = PrevValue.ServerFrame;
+	}
+
+	if (Flags & Flag_ServerPhysicsHandleIsPresent)
+	{
+		DeserializeIntDelta(*Reader, TempValue.ServerPhysicsHandle, PrevValue.ServerPhysicsHandle, DeltaBitCountTable, DeltaBitCountTableEntryCount, 32U);
+	}
+	else
+	{
+		TempValue.ServerPhysicsHandle = PrevValue.ServerPhysicsHandle;
+	}
+
+	// Fixup flags
+	uint32 FixedUpFlags = Flags & (Flag_RepPhysics | Flag_SimulatedPhysicSleep);
+	FixedUpFlags |= (TempValue.ServerFrame > 0 ? Flag_ServerFrameIsPresent : 0U);
+	FixedUpFlags |= (TempValue.ServerPhysicsHandle != INDEX_NONE ? Flag_ServerPhysicsHandleIsPresent : 0U);
+	TempValue.Flags = FixedUpFlags;
+
 	QuantizedType& Target = *reinterpret_cast<QuantizedType*>(Args.Target);
 	Target = TempValue;
 }
@@ -408,9 +491,13 @@ void FRepMovementNetSerializer::Quantize(FNetSerializationContext& Context, cons
 
 	QuantizedType TempValue = {};
 	TempValue.Flags = (Source.bSimulatedPhysicSleep ? Flag_SimulatedPhysicSleep : 0U) | (Source.bRepPhysics ? Flag_RepPhysics : 0);
+	TempValue.Flags |= (Source.ServerFrame > 0 ? Flag_ServerFrameIsPresent : 0U);
+	TempValue.Flags |= (Source.ServerPhysicsHandle != INDEX_NONE ? Flag_ServerPhysicsHandleIsPresent : 0U);
 	TempValue.LocationQuantizationLevel = uint8(Source.LocationQuantizationLevel);
 	TempValue.VelocityQuantizationLevel = uint8(Source.VelocityQuantizationLevel);
 	TempValue.RotationQuantizationLevel = uint8(Source.RotationQuantizationLevel);
+	TempValue.ServerFrame = Source.ServerFrame;
+	TempValue.ServerPhysicsHandle = Source.ServerPhysicsHandle;
 	
 	// Angular velocity. Do note that the quantized value is cleared if RepPhysics is false. Delta compression need to accommodate to this.
 	if (TempValue.Flags & Flag_RepPhysics)
@@ -472,6 +559,8 @@ void FRepMovementNetSerializer::Dequantize(FNetSerializationContext& Context, co
 
 	Target.bSimulatedPhysicSleep = (Source.Flags & Flag_SimulatedPhysicSleep);
 	Target.bRepPhysics = (Source.Flags & Flag_RepPhysics);
+	Target.ServerFrame = (Source.Flags & Flag_ServerFrameIsPresent ? Source.ServerFrame : 0);
+	Target.ServerPhysicsHandle = (Source.Flags & Flag_ServerPhysicsHandleIsPresent ? Source.ServerPhysicsHandle : INDEX_NONE);
 
 	// We do not overwrite AngularVelocity unless we're replicating it. This is consistent with the FRepMovement serialization method.
 	if (Source.Flags & Flag_RepPhysics)
@@ -662,7 +751,7 @@ bool FRepMovementNetSerializer::FNetSerializerRegistryDelegates::IsRepMovementLa
 {
 	const UStruct* RepMovementStruct = FRepMovement::StaticStruct();
 
-	constexpr int ExpectedPropertiesSize = 104;
+	constexpr int ExpectedPropertiesSize = 112;
 	if (!ensureMsgf(RepMovementStruct->GetPropertiesSize() == ExpectedPropertiesSize, TEXT("Unexpected FRepMovement properties size. %d != %d"), RepMovementStruct->GetPropertiesSize(), ExpectedPropertiesSize))
 	{
 		return false;
@@ -674,6 +763,7 @@ bool FRepMovementNetSerializer::FNetSerializerRegistryDelegates::IsRepMovementLa
 		{
 			"LinearVelocity", "AngularVelocity", "Location", "Rotation", "bSimulatedPhysicSleep", "bRepPhysics", 
 			"LocationQuantizationLevel", "VelocityQuantizationLevel", "RotationQuantizationLevel", 
+			"ServerFrame", "ServerPhysicsHandle", 
 		};
 		constexpr SIZE_T FieldCount = sizeof(Fields)/sizeof(Fields[0]);
 
