SkeletalMesh: fix serialisation code when both IsLoading() and IsSaving() are false.

diff --git a/Engine/Source/Runtime/Engine/Private/SkeletalMeshModel.cpp b/Engine/Source/Runtime/Engine/Private/SkeletalMeshModel.cpp
--- a/Engine/Source/Runtime/Engine/Private/SkeletalMeshModel.cpp
+++ b/Engine/Source/Runtime/Engine/Private/SkeletalMeshModel.cpp
@@ -18,7 +18,7 @@ void FSkeletalMeshModel::Serialize(FArchive& Ar, USkinnedAsset* Owner)
 {
 	DECLARE_SCOPE_CYCLE_COUNTER(TEXT("FSkeletalMeshModel::Serialize"), STAT_SkeletalMeshModel_Serialize, STATGROUP_LoadTime);
 	bool bIsEditorDataStripped = false;
-	if (Ar.IsSaving() || Ar.CustomVer(FFortniteMainBranchObjectVersion::GUID) >= FFortniteMainBranchObjectVersion::AllowSkeletalMeshToReduceTheBaseLOD)
+	if (Ar.IsSaving() || (Ar.IsLoading() && Ar.CustomVer(FFortniteMainBranchObjectVersion::GUID) >= FFortniteMainBranchObjectVersion::AllowSkeletalMeshToReduceTheBaseLOD))
 	{
 		FStripDataFlags StripFlags(Ar);
 		bIsEditorDataStripped = StripFlags.IsEditorDataStripped();
@@ -81,7 +81,7 @@ void FSkeletalMeshModel::Serialize(FArchive& Ar, USkinnedAsset* Owner)
 			}
 		}
 
-		if (Ar.IsSaving() || (Ar.CustomVer(FUE5MainStreamObjectVersion::GUID) >= FUE5MainStreamObjectVersion::ConvertReductionBaseSkeletalMeshBulkDataToInlineReductionCacheData))
+		if (Ar.IsSaving() || (Ar.IsLoading() && Ar.CustomVer(FUE5MainStreamObjectVersion::GUID) >= FUE5MainStreamObjectVersion::ConvertReductionBaseSkeletalMeshBulkDataToInlineReductionCacheData))
 		{
 			Ar << InlineReductionCacheDatas;
 		}
