Fix for Nounity CIS in AudioFormatSettings
#rnx

diff --git a/Engine/Source/Runtime/Engine/Private/AudioFormatSettings.cpp b/Engine/Source/Runtime/Engine/Private/AudioFormatSettings.cpp
--- a/Engine/Source/Runtime/Engine/Private/AudioFormatSettings.cpp
+++ b/Engine/Source/Runtime/Engine/Private/AudioFormatSettings.cpp
@@ -1,9 +1,12 @@
 // Copyright Epic Games, Inc. All Rights Reserved.
 
 #include "Sound/AudioFormatSettings.h"
+#include "Sound/SoundWave.h"
 #include "Algo/Transform.h"
 #include "Misc/ConfigCacheIni.h"
 #include "Containers/UnrealString.h"
+#include "Audio.h"
+#include "AudioMixerTypes.h"
 
 namespace Audio
 {
