Bugfix: Imported Audio Files No Longer Play (Project Defined Codec Not Working)

diff --git a/Engine/Source/Editor/AudioEditor/Private/Factories/SoundFactory.cpp b/Engine/Source/Editor/AudioEditor/Private/Factories/SoundFactory.cpp
--- a/Engine/Source/Editor/AudioEditor/Private/Factories/SoundFactory.cpp
+++ b/Engine/Source/Editor/AudioEditor/Private/Factories/SoundFactory.cpp
@@ -578,6 +578,8 @@ UObject* USoundFactory::CreateObject
 			Sound->LoadZerothChunk();
 		}
 
+		Sound->PostImport();
+
 		GEditor->GetEditorSubsystem<UImportSubsystem>()->BroadcastAssetPostImport(this, Sound);
 
 		if (ExistingSound && bUseExistingSettings)
diff --git a/Engine/Source/Runtime/Engine/Classes/Sound/SoundWave.h b/Engine/Source/Runtime/Engine/Classes/Sound/SoundWave.h
--- a/Engine/Source/Runtime/Engine/Classes/Sound/SoundWave.h
+++ b/Engine/Source/Runtime/Engine/Classes/Sound/SoundWave.h
@@ -860,6 +860,8 @@ public:
 
 	bool IsAsyncWorkComplete() const;
 
+	void PostImport();
+
 private:
 	friend class FSoundWaveCompilingManager;
 	/**  Utility function used internally to change task priority while maintaining thread-safety. */
diff --git a/Engine/Source/Runtime/Engine/Private/SoundWave.cpp b/Engine/Source/Runtime/Engine/Private/SoundWave.cpp
--- a/Engine/Source/Runtime/Engine/Private/SoundWave.cpp
+++ b/Engine/Source/Runtime/Engine/Private/SoundWave.cpp
@@ -533,6 +533,14 @@ bool USoundWave::IsAsyncWorkComplete() const
 	return SoundWaveDataPtr == nullptr || SoundWaveDataPtr->RunningPlatformData.IsAsyncWorkComplete();
 }
 
+void USoundWave::PostImport()
+{
+	if (SoundWaveDataPtr)
+	{
+		SoundWaveDataPtr->InitializeDataFromSoundWave(*this);
+	}
+}
+
 bool USoundWave::RescheduleAsyncTask(FQueuedThreadPool* InThreadPool, EQueuedWorkPriority InPriority)
 {
 	return SoundWaveDataPtr ? SoundWaveDataPtr->RunningPlatformData.RescheduleAsyncTask(InThreadPool, InPriority) : false;
