CIS fix for nounity AudioFormatSettings. (namespaced forward declare)
#rnx

diff --git a/Engine/Source/Runtime/Engine/Classes/Sound/AudioFormatSettings.h b/Engine/Source/Runtime/Engine/Classes/Sound/AudioFormatSettings.h
--- a/Engine/Source/Runtime/Engine/Classes/Sound/AudioFormatSettings.h
+++ b/Engine/Source/Runtime/Engine/Classes/Sound/AudioFormatSettings.h
@@ -6,6 +6,7 @@
 
 class FString;
 class FConfigCacheIni;
+class USoundWave;
 
 namespace Audio
 {
@@ -15,7 +16,7 @@ namespace Audio
 		FAudioFormatSettings(FConfigCacheIni* InConfigSystem, const FString& InConfigFilename, const FString& InPlatformIdentifierForLogging);
 		~FAudioFormatSettings() = default;
 
-		FName GetWaveFormat(const class USoundWave* Wave) const;
+		FName GetWaveFormat(const USoundWave* Wave) const;
 		void GetAllWaveFormats(TArray<FName>& OutFormats) const;
 		void GetWaveFormatModuleHints(TArray<FName>& OutHints) const;
 
diff --git a/Engine/Source/Runtime/Engine/Private/AudioFormatSettings.cpp b/Engine/Source/Runtime/Engine/Private/AudioFormatSettings.cpp
--- a/Engine/Source/Runtime/Engine/Private/AudioFormatSettings.cpp
+++ b/Engine/Source/Runtime/Engine/Private/AudioFormatSettings.cpp
@@ -15,7 +15,7 @@ namespace Audio
 		ReadConfiguration(InConfigSystem, InConfigFilename, InPlatformIdentifierForLogging);
 	}
 
-	FName FAudioFormatSettings::GetWaveFormat(const class USoundWave* Wave) const
+	FName FAudioFormatSettings::GetWaveFormat(const USoundWave* Wave) const
 	{
 		FName FormatName = Audio::ToName(Wave->GetSoundAssetCompressionType());
 		if (FormatName == Audio::NAME_PLATFORM_SPECIFIC)
