Console: Added EngineConsoleCommandExecutor to provide IConsoleCommandExecutor outside of the editor

diff --git a/Engine/Source/Runtime/Engine/Classes/Engine/GameEngine.h b/Engine/Source/Runtime/Engine/Classes/Engine/GameEngine.h
--- a/Engine/Source/Runtime/Engine/Classes/Engine/GameEngine.h
+++ b/Engine/Source/Runtime/Engine/Classes/Engine/GameEngine.h
@@ -8,9 +8,11 @@
 #include "Widgets/SViewport.h"
 #include "Engine/Engine.h"
 #include "Runtime/MovieSceneCapture/Public/MovieSceneCaptureHandle.h"
+#include "Templates/PimplPtr.h"
 #include "GameEngine.generated.h"
 
 class Error;
+class FEngineConsoleCommandExecutor;
 class FSceneViewport;
 class UGameViewportClient;
 class UNetDriver;
@@ -182,6 +184,8 @@ private:
 
 	/** Last time the logs have been flushed. */
 	double LastTimeLogsFlushed;
+
+	TPimplPtr<FEngineConsoleCommandExecutor> CmdExec;
 };
 
 PRAGMA_ENABLE_DEPRECATION_WARNINGS
diff --git a/Engine/Source/Runtime/Engine/Private/Engine/EngineConsoleCommandExecutor.cpp b/Engine/Source/Runtime/Engine/Private/Engine/EngineConsoleCommandExecutor.cpp
new file mode 100644
--- /dev/null
+++ b/Engine/Source/Runtime/Engine/Private/Engine/EngineConsoleCommandExecutor.cpp
@@ -0,0 +1,114 @@
+// Copyright Epic Games, Inc. All Rights Reserved.
+
+#include "Engine/EngineConsoleCommandExecutor.h"
+
+#include "Engine/GameEngine.h"
+#include "Engine/LocalPlayer.h"
+
+#define LOCTEXT_NAMESPACE "EngineConsoleCommandExecutor"
+
+static const FLazyName GEngineConsoleCommandExecutorName = TEXT("Cmd");
+
+FEngineConsoleCommandExecutor::FEngineConsoleCommandExecutor(UGameEngine* InGameEngine)
+	: GameEngine(InGameEngine)
+{
+	IModularFeatures::Get().RegisterModularFeature(IConsoleCommandExecutor::ModularFeatureName(), this);
+}
+
+FEngineConsoleCommandExecutor::~FEngineConsoleCommandExecutor()
+{
+	IModularFeatures::Get().UnregisterModularFeature(IConsoleCommandExecutor::ModularFeatureName(), this);
+}
+
+FName FEngineConsoleCommandExecutor::StaticName()
+{
+	return GEngineConsoleCommandExecutorName;
+}
+
+FName FEngineConsoleCommandExecutor::GetName() const
+{
+	return StaticName();
+}
+
+FText FEngineConsoleCommandExecutor::GetDisplayName() const
+{
+	return LOCTEXT("DisplayName", "Cmd");
+}
+
+FText FEngineConsoleCommandExecutor::GetDescription() const
+{
+	return LOCTEXT("Description", "Execute Unreal Console Commands");
+}
+
+FText FEngineConsoleCommandExecutor::GetHintText() const
+{
+	return LOCTEXT("HintText", "Enter Console Command");
+}
+
+void FEngineConsoleCommandExecutor::GetAutoCompleteSuggestions(const TCHAR* Input, TArray<FString>& Out)
+{
+	const auto OnConsoleVariable = [&Out](const TCHAR* Name, IConsoleObject* CVar)
+	{
+	#if UE_BUILD_SHIPPING || UE_BUILD_TEST
+		if (CVar->TestFlags(ECVF_Cheat))
+		{
+			return;
+		}
+	#endif // UE_BUILD_SHIPPING || UE_BUILD_TEST
+		if (CVar->TestFlags(ECVF_Unregistered))
+		{
+			return;
+		}
+		Out.Add(Name);
+	};
+
+	IConsoleManager::Get().ForEachConsoleObjectThatContains(FConsoleObjectVisitor::CreateLambda(OnConsoleVariable), Input);
+}
+
+void FEngineConsoleCommandExecutor::GetExecHistory(TArray<FString>& Out)
+{
+	IConsoleManager::Get().GetConsoleHistory(TEXT(""), Out);
+}
+
+bool FEngineConsoleCommandExecutor::Exec(const TCHAR* Input)
+{
+	IConsoleManager::Get().AddConsoleHistoryEntry(TEXT(""), Input);
+
+	int32 Len = FCString::Strlen(Input);
+	TArray<TCHAR> Buffer;
+	Buffer.AddZeroed(Len + 1);
+
+	bool bHandled = false;
+	while (FParse::Line(&Input, Buffer.GetData(), Buffer.Num()))
+	{
+		bHandled |= ExecInternal(Buffer.GetData());
+	}
+	return bHandled;
+}
+
+bool FEngineConsoleCommandExecutor::ExecInternal(const TCHAR* Input) const
+{
+	if (ULocalPlayer* Player = GameEngine->GetDebugLocalPlayer())
+	{
+		return Player->Exec(Player->GetWorld(), Input, *GLog);
+	}
+
+	return GameEngine->Exec(GameEngine->GetGameWorld(), Input, *GLog);
+}
+
+bool FEngineConsoleCommandExecutor::AllowHotKeyClose() const
+{
+	return true;
+}
+
+bool FEngineConsoleCommandExecutor::AllowMultiLine() const
+{
+	return true;
+}
+
+FInputChord FEngineConsoleCommandExecutor::GetHotKey() const
+{
+	return {};
+}
+
+#undef LOCTEXT_NAMESPACE
diff --git a/Engine/Source/Runtime/Engine/Private/Engine/EngineConsoleCommandExecutor.h b/Engine/Source/Runtime/Engine/Private/Engine/EngineConsoleCommandExecutor.h
new file mode 100644
--- /dev/null
+++ b/Engine/Source/Runtime/Engine/Private/Engine/EngineConsoleCommandExecutor.h
@@ -0,0 +1,32 @@
+// Copyright Epic Games, Inc. All Rights Reserved.
+
+#include "HAL/IConsoleManager.h"
+
+class UGameEngine;
+
+class FEngineConsoleCommandExecutor final : public IConsoleCommandExecutor
+{
+public:
+	explicit FEngineConsoleCommandExecutor(UGameEngine* GameEngine);
+	~FEngineConsoleCommandExecutor();
+
+	static FName StaticName();
+
+	FName GetName() const final;
+	FText GetDisplayName() const final;
+	FText GetDescription() const final;
+	FText GetHintText() const final;
+
+	void GetAutoCompleteSuggestions(const TCHAR* Input, TArray<FString>& Out) final;
+	void GetExecHistory(TArray<FString>& Out) final;
+	bool Exec(const TCHAR* Input) final;
+
+	bool AllowHotKeyClose() const final;
+	bool AllowMultiLine() const final;
+	FInputChord GetHotKey() const final;
+
+private:
+	bool ExecInternal(const TCHAR* Input) const;
+
+	UGameEngine* GameEngine;
+};
diff --git a/Engine/Source/Runtime/Engine/Private/GameEngine.cpp b/Engine/Source/Runtime/Engine/Private/GameEngine.cpp
--- a/Engine/Source/Runtime/Engine/Private/GameEngine.cpp
+++ b/Engine/Source/Runtime/Engine/Private/GameEngine.cpp
@@ -14,6 +14,7 @@
 #include "EngineStats.h"
 #include "EngineGlobals.h"
 #include "RenderingThread.h"
+#include "Engine/EngineConsoleCommandExecutor.h"
 #include "Engine/GameViewportClient.h"
 #include "Engine/LevelStreaming.h"
 #include "Engine/PlatformInterfaceBase.h"
@@ -1083,7 +1084,11 @@ public:
 void UGameEngine::Init(IEngineLoop* InEngineLoop)
 {
 	DECLARE_SCOPE_CYCLE_COUNTER(TEXT("UGameEngine Init"), STAT_GameEngineStartup, STATGROUP_LoadTime);
-	
+
+	if (!GIsEditor)
+	{
+		CmdExec = MakePimpl<FEngineConsoleCommandExecutor>(this);
+	}
 
 	// Call base.
 	UEngine::Init(InEngineLoop);
@@ -1235,6 +1240,8 @@ void UGameEngine::FinishDestroy()
 	}
 
 	Super::FinishDestroy();
+
+	CmdExec.Reset();
 }
 
 //@todo: unify this and the driver version
