Console: Fixed errors in EngineConsoleCommandExecutor

#rnx

diff --git a/Engine/Source/Runtime/Engine/Private/Engine/EngineConsoleCommandExecutor.cpp b/Engine/Source/Runtime/Engine/Private/Engine/EngineConsoleCommandExecutor.cpp
--- a/Engine/Source/Runtime/Engine/Private/Engine/EngineConsoleCommandExecutor.cpp
+++ b/Engine/Source/Runtime/Engine/Private/Engine/EngineConsoleCommandExecutor.cpp
@@ -4,6 +4,7 @@
 
 #include "Engine/GameEngine.h"
 #include "Engine/LocalPlayer.h"
+#include "Features/IModularFeatures.h"
 
 #define LOCTEXT_NAMESPACE "EngineConsoleCommandExecutor"
 
@@ -111,4 +112,9 @@ FInputChord FEngineConsoleCommandExecutor::GetHotKey() const
 	return {};
 }
 
+FInputChord FEngineConsoleCommandExecutor::GetIterateExecutorHotKey() const
+{
+	return {};
+}
+
 #undef LOCTEXT_NAMESPACE
diff --git a/Engine/Source/Runtime/Engine/Private/Engine/EngineConsoleCommandExecutor.h b/Engine/Source/Runtime/Engine/Private/Engine/EngineConsoleCommandExecutor.h
--- a/Engine/Source/Runtime/Engine/Private/Engine/EngineConsoleCommandExecutor.h
+++ b/Engine/Source/Runtime/Engine/Private/Engine/EngineConsoleCommandExecutor.h
@@ -1,5 +1,7 @@
 // Copyright Epic Games, Inc. All Rights Reserved.
 
+#pragma once
+
 #include "HAL/IConsoleManager.h"
 
 class UGameEngine;
@@ -24,6 +26,7 @@ public:
 	bool AllowHotKeyClose() const final;
 	bool AllowMultiLine() const final;
 	FInputChord GetHotKey() const final;
+	FInputChord GetIterateExecutorHotKey() const final;
 
 private:
 	bool ExecInternal(const TCHAR* Input) const;
