Simplify transaction diffing by tracking non-property areas during serialization

#rnx

diff --git a/Engine/Source/Editor/UnrealEd/Classes/Editor/Transactor.h b/Engine/Source/Editor/UnrealEd/Classes/Editor/Transactor.h
--- a/Engine/Source/Editor/UnrealEd/Classes/Editor/Transactor.h
+++ b/Engine/Source/Editor/UnrealEd/Classes/Editor/Transactor.h
@@ -43,8 +43,6 @@ protected:
 	class UNREALED_API FObjectRecord
 	{
 	public:
-		using FSerializedProperty = FTransactionSerializedProperty;
-
 		struct FSerializedObject : public FTransactionSerializedObject
 		{
 			void SetObject(const UObject* InObject)
diff --git a/Engine/Source/Runtime/Engine/Private/TransactionCommon.cpp b/Engine/Source/Runtime/Engine/Private/TransactionCommon.cpp
--- a/Engine/Source/Runtime/Engine/Private/TransactionCommon.cpp
+++ b/Engine/Source/Runtime/Engine/Private/TransactionCommon.cpp
@@ -136,14 +136,17 @@ void FTransactionPersistentObjectRef::AddStructReferencedObjects(FReferenceColle
 }
 
 
-FName FTransactionSerializedProperty::BuildSerializedPropertyKey(const FArchiveSerializedPropertyChain& InPropertyChain)
+FTransactionSerializedTaggedData FTransactionSerializedTaggedData::FromOffsetAndSize(const int64 InOffset, const int64 InSize)
 {
-	const int32 NumProperties = InPropertyChain.GetNumProperties();
-	check(NumProperties > 0);
-	return InPropertyChain.GetPropertyFromRoot(0)->GetFName();
+	return FTransactionSerializedTaggedData{ InOffset, InSize };
 }
 
-void FTransactionSerializedProperty::AppendSerializedData(const int64 InOffset, const int64 InSize)
+FTransactionSerializedTaggedData FTransactionSerializedTaggedData::FromStartAndEnd(const int64 InStart, const int64 InEnd)
+{
+	return FTransactionSerializedTaggedData{ InStart, InEnd - InStart };
+}
+
+void FTransactionSerializedTaggedData::AppendSerializedData(const int64 InOffset, const int64 InSize)
 {
 	if (DataOffset == INDEX_NONE)
 	{
@@ -157,6 +160,19 @@ void FTransactionSerializedProperty::AppendSerializedData(const int64 InOffset,
 	}
 }
 
+void FTransactionSerializedTaggedData::AppendSerializedData(const FTransactionSerializedTaggedData& InData)
+{
+	if (InData.HasSerializedData())
+	{
+		AppendSerializedData(InData.DataOffset, InData.DataSize);
+	}
+}
+
+bool FTransactionSerializedTaggedData::HasSerializedData() const
+{
+	return DataOffset != INDEX_NONE && DataSize != 0;
+}
+
 
 void FTransactionSerializedObjectData::Read(void* Dest, int64 Offset, int64 Num) const
 {
@@ -177,9 +193,11 @@ void FTransactionSerializedObjectData::Write(const void* Src, int64 Offset, int6
 namespace UE::Transaction
 {
 
+const FName TaggedDataKey_UnknownData = ".UnknownData";
+const FName TaggedDataKey_ScriptData = ".ScriptData";
+
 FSerializedObjectDataReader::FSerializedObjectDataReader(const FTransactionSerializedObject& InSerializedObject)
 	: SerializedObject(InSerializedObject)
-	, Offset(0)
 {
 	SetIsLoading(true);
 }
@@ -219,7 +237,6 @@ FArchive& FSerializedObjectDataReader::operator<<(UObject*& Res)
 
 FSerializedObjectDataWriter::FSerializedObjectDataWriter(FTransactionSerializedObject& InSerializedObject)
 	: SerializedObject(InSerializedObject)
-	, Offset(0)
 {
 	SetIsSaving(true);
 
@@ -234,6 +251,74 @@ FSerializedObjectDataWriter::FSerializedObjectDataWriter(FTransactionSerializedO
 	}
 }
 
+FName FSerializedObjectDataWriter::GetTaggedDataKey() const
+{
+	FName TaggedDataKey;
+
+	// Is this known property data?
+	if (TaggedDataKey.IsNone())
+	{
+		const FArchiveSerializedPropertyChain* PropertyChain = GetSerializedPropertyChain();
+		if (PropertyChain && PropertyChain->GetNumProperties() > 0)
+		{
+			TaggedDataKey = CachedSerializedTaggedPropertyKey.SyncCache(PropertyChain);
+		}
+	}
+
+	// Is this script data?
+	if (TaggedDataKey.IsNone() && bIsPerformingScriptSerialization)
+	{
+		bWasUsingTaggedDataKey_ScriptData = true;
+		TaggedDataKey = FName(TaggedDataKey_ScriptData, TaggedDataKeyIndex_ScriptData);
+	}
+	else if (bWasUsingTaggedDataKey_ScriptData)
+	{
+		++TaggedDataKeyIndex_ScriptData;
+		bWasUsingTaggedDataKey_ScriptData = false;
+	}
+
+	// Is this unknown data?
+	if (TaggedDataKey.IsNone())
+	{
+		bWasUsingTaggedDataKey_UnknownData = true;
+		TaggedDataKey = FName(TaggedDataKey_UnknownData, TaggedDataKeyIndex_UnknownData);
+	}
+	else if (bWasUsingTaggedDataKey_UnknownData)
+	{
+		++TaggedDataKeyIndex_UnknownData;
+		bWasUsingTaggedDataKey_UnknownData = false;
+	}
+
+	return TaggedDataKey;
+}
+
+bool FSerializedObjectDataWriter::DoesObjectMatchSerializedObject(const UObject* Obj) const
+{
+	FNameBuilder ObjPathName;
+	Obj->GetPathName(nullptr, ObjPathName);
+	return FName(ObjPathName) == SerializedObject.ObjectPathName;
+}
+
+void FSerializedObjectDataWriter::MarkScriptSerializationStart(const UObject* Obj)
+{
+	FArchiveUObject::MarkScriptSerializationStart(Obj);
+
+	if (DoesObjectMatchSerializedObject(Obj))
+	{
+		bIsPerformingScriptSerialization = true;
+	}
+}
+
+void FSerializedObjectDataWriter::MarkScriptSerializationEnd(const UObject* Obj)
+{
+	FArchiveUObject::MarkScriptSerializationEnd(Obj);
+
+	if (DoesObjectMatchSerializedObject(Obj))
+	{
+		bIsPerformingScriptSerialization = false;
+	}
+}
+
 void FSerializedObjectDataWriter::Serialize(void* SerData, int64 Num)
 {
 	if (Num)
@@ -242,13 +327,12 @@ void FSerializedObjectDataWriter::Serialize(void* SerData, int64 Num)
 		SerializedObject.SerializedData.Write(SerData, Offset, Num);
 		Offset += Num;
 
-		// Track this property offset in the serialized data
-		const FArchiveSerializedPropertyChain* PropertyChain = GetSerializedPropertyChain();
-		if (PropertyChain && PropertyChain->GetNumProperties() > 0)
+		// Track this offset index in the serialized data
+		const FName SerializedTaggedDataKey = GetTaggedDataKey();
+		if (!SerializedTaggedDataKey.IsNone())
 		{
-			const FName SerializedTaggedPropertyKey = CachedSerializedTaggedPropertyKey.SyncCache(PropertyChain);
-			FTransactionSerializedProperty& SerializedTaggedProperty = SerializedObject.SerializedProperties.FindOrAdd(SerializedTaggedPropertyKey);
-			SerializedTaggedProperty.AppendSerializedData(DataIndex, Num);
+			FTransactionSerializedTaggedData& SerializedTaggedData = SerializedObject.SerializedTaggedData.FindOrAdd(SerializedTaggedDataKey);
+			SerializedTaggedData.AppendSerializedData(DataIndex, Num);
 		}
 	}
 }
@@ -268,10 +352,10 @@ FArchive& FSerializedObjectDataWriter::operator<<(FName& N)
 	}
 
 	// Track this name index in the serialized data
+	const FName SerializedTaggedDataKey = GetTaggedDataKey();
+	if (!SerializedTaggedDataKey.IsNone())
 	{
-		const FArchiveSerializedPropertyChain* PropertyChain = GetSerializedPropertyChain();
-		const FName SerializedTaggedPropertyKey = CachedSerializedTaggedPropertyKey.SyncCache(PropertyChain);
-		SerializedObject.SerializedNameIndices.Add(SerializedTaggedPropertyKey).Indices.AddUnique(NameIndex);
+		SerializedObject.SerializedNameIndices.Add(SerializedTaggedDataKey).Indices.AddUnique(NameIndex);
 	}
 
 	return *this << NameIndex;
@@ -292,10 +376,10 @@ FArchive& FSerializedObjectDataWriter::operator<<(UObject*& Res)
 	}
 
 	// Track this object offset in the serialized data
+	const FName SerializedTaggedDataKey = GetTaggedDataKey();
+	if (!SerializedTaggedDataKey.IsNone())
 	{
-		const FArchiveSerializedPropertyChain* PropertyChain = GetSerializedPropertyChain();
-		const FName SerializedTaggedPropertyKey = CachedSerializedTaggedPropertyKey.SyncCache(PropertyChain);
-		SerializedObject.SerializedObjectIndices.Add(SerializedTaggedPropertyKey).Indices.AddUnique(ObjectIndex);
+		SerializedObject.SerializedObjectIndices.Add(SerializedTaggedDataKey).Indices.AddUnique(ObjectIndex);
 	}
 
 	return *this << ObjectIndex;
@@ -308,7 +392,7 @@ FName FSerializedObjectDataWriter::FCachedPropertyKey::SyncCache(const FArchiveS
 		const uint32 CurrentUpdateCount = InPropertyChain->GetUpdateCount();
 		if (CurrentUpdateCount != LastUpdateCount)
 		{
-			CachedKey = InPropertyChain->GetNumProperties() > 0 ? FTransactionSerializedProperty::BuildSerializedPropertyKey(*InPropertyChain) : FName();
+			CachedKey = InPropertyChain->GetNumProperties() > 0 ? InPropertyChain->GetPropertyFromRoot(0)->GetFName() : FName();
 			LastUpdateCount = CurrentUpdateCount;
 		}
 	}
@@ -327,148 +411,108 @@ namespace DiffUtil
 
 void GenerateObjectDiff(const FTransactionSerializedObject& OldSerializedObject, const FTransactionSerializedObject& NewSerializedObject, FTransactionObjectDeltaChange& OutDeltaChange, const bool bFullDiff)
 {
-	if (bFullDiff)
+	auto IsTaggedDataBlockIdentical = [&OldSerializedObject, &NewSerializedObject, &OutDeltaChange](const FName TaggedDataKey, const FTransactionSerializedTaggedData& OldSerializedTaggedData, const FTransactionSerializedTaggedData& NewSerializedTaggedData) -> bool
 	{
-		OutDeltaChange.bHasNameChange |= OldSerializedObject.ObjectName != NewSerializedObject.ObjectName;
-		OutDeltaChange.bHasOuterChange |= OldSerializedObject.ObjectOuterPathName != NewSerializedObject.ObjectOuterPathName;
-		OutDeltaChange.bHasExternalPackageChange |= OldSerializedObject.ObjectExternalPackageName != NewSerializedObject.ObjectExternalPackageName;
-		OutDeltaChange.bHasPendingKillChange |= OldSerializedObject.bIsPendingKill != NewSerializedObject.bIsPendingKill;
-
-		if (!Internal::AreObjectPointersIdentical(OldSerializedObject, NewSerializedObject, NAME_None))
+		// Binary compare the serialized data to see if something has changed for this property
+		bool bIsTaggedDataIdentical = OldSerializedTaggedData.DataSize == NewSerializedTaggedData.DataSize;
+		if (bIsTaggedDataIdentical && OldSerializedTaggedData.HasSerializedData())
 		{
-			OutDeltaChange.bHasNonPropertyChanges = true;
+			bIsTaggedDataIdentical = FMemory::Memcmp(OldSerializedObject.SerializedData.GetPtr(OldSerializedTaggedData.DataOffset), NewSerializedObject.SerializedData.GetPtr(NewSerializedTaggedData.DataOffset), NewSerializedTaggedData.DataSize) == 0;
 		}
+		if (bIsTaggedDataIdentical)
+		{
+			bIsTaggedDataIdentical = Internal::AreObjectPointersIdentical(OldSerializedObject, NewSerializedObject, TaggedDataKey);
+		}
+		if (bIsTaggedDataIdentical)
+		{
+			bIsTaggedDataIdentical = Internal::AreNamesIdentical(OldSerializedObject, NewSerializedObject, TaggedDataKey);
+		}
+		return bIsTaggedDataIdentical;
+	};
 
-		if (!Internal::AreNamesIdentical(OldSerializedObject, NewSerializedObject, NAME_None))
+	auto ShouldCompareTaggedData = [](const FName TaggedDataKey) -> bool
+	{
+		if (TaggedDataKey.GetComparisonIndex() == TaggedDataKey_ScriptData.GetComparisonIndex())
 		{
-			OutDeltaChange.bHasNonPropertyChanges = true;
+			// Never compare script data, as it's assumed to be overhead around the tagged property serialization
+			return false;
 		}
-	}
 
-	if (OldSerializedObject.SerializedProperties.Num() > 0 || NewSerializedObject.SerializedProperties.Num() > 0)
+		return true;
+	};
+
+	auto ShouldCompareAsNonPropertyData = [](const FName TaggedDataKey)
 	{
-		int64 StartOfOldPropertyBlock = INT64_MAX;
-		int64 StartOfNewPropertyBlock = INT64_MAX;
-		int64 EndOfOldPropertyBlock = -1;
-		int64 EndOfNewPropertyBlock = -1;
+		return TaggedDataKey.GetComparisonIndex() == TaggedDataKey_UnknownData.GetComparisonIndex();
+	};
 
-		for (const TPair<FName, FTransactionSerializedProperty>& NewNamePropertyPair : NewSerializedObject.SerializedProperties)
-		{
-			const FTransactionSerializedProperty* OldSerializedProperty = OldSerializedObject.SerializedProperties.Find(NewNamePropertyPair.Key);
-			if (!OldSerializedProperty)
-			{
-				if (bFullDiff)
-				{
-					// Missing property, assume that the property changed
-					OutDeltaChange.ChangedProperties.AddUnique(NewNamePropertyPair.Key);
-				}
-				continue;
-			}
+	if (bFullDiff)
+	{
+		OutDeltaChange.bHasNameChange |= OldSerializedObject.ObjectName != NewSerializedObject.ObjectName;
+		OutDeltaChange.bHasOuterChange |= OldSerializedObject.ObjectOuterPathName != NewSerializedObject.ObjectOuterPathName;
+		OutDeltaChange.bHasExternalPackageChange |= OldSerializedObject.ObjectExternalPackageName != NewSerializedObject.ObjectExternalPackageName;
+		OutDeltaChange.bHasPendingKillChange |= OldSerializedObject.bIsPendingKill != NewSerializedObject.bIsPendingKill;
+	}
 
-			// Update the tracking for the start/end of the property block within the serialized data
-			StartOfOldPropertyBlock = FMath::Min(StartOfOldPropertyBlock, OldSerializedProperty->DataOffset);
-			StartOfNewPropertyBlock = FMath::Min(StartOfNewPropertyBlock, NewNamePropertyPair.Value.DataOffset);
-			EndOfOldPropertyBlock = FMath::Max(EndOfOldPropertyBlock, OldSerializedProperty->DataOffset + OldSerializedProperty->DataSize);
-			EndOfNewPropertyBlock = FMath::Max(EndOfNewPropertyBlock, NewNamePropertyPair.Value.DataOffset + NewNamePropertyPair.Value.DataSize);
+	for (const TPair<FName, FTransactionSerializedTaggedData>& NewNamePropertyPair : NewSerializedObject.SerializedTaggedData)
+	{
+		if (!ShouldCompareTaggedData(NewNamePropertyPair.Key))
+		{
+			continue;
+		}
 
-			// Binary compare the serialized data to see if something has changed for this property
-			bool bIsPropertyIdentical = OldSerializedProperty->DataSize == NewNamePropertyPair.Value.DataSize;
-			if (bIsPropertyIdentical && NewNamePropertyPair.Value.DataSize > 0)
-			{
-				bIsPropertyIdentical = FMemory::Memcmp(OldSerializedObject.SerializedData.GetPtr(OldSerializedProperty->DataOffset), NewSerializedObject.SerializedData.GetPtr(NewNamePropertyPair.Value.DataOffset), NewNamePropertyPair.Value.DataSize) == 0;
-			}
-			if (bIsPropertyIdentical)
-			{
-				bIsPropertyIdentical = Internal::AreObjectPointersIdentical(OldSerializedObject, NewSerializedObject, NewNamePropertyPair.Key);
-			}
-			if (bIsPropertyIdentical)
+		const FTransactionSerializedTaggedData* OldSerializedTaggedData = OldSerializedObject.SerializedTaggedData.Find(NewNamePropertyPair.Key);
+		if (ShouldCompareAsNonPropertyData(NewNamePropertyPair.Key))
+		{
+			if (bFullDiff && !OutDeltaChange.bHasNonPropertyChanges && (!OldSerializedTaggedData || !IsTaggedDataBlockIdentical(NewNamePropertyPair.Key, *OldSerializedTaggedData, NewNamePropertyPair.Value)))
 			{
-				bIsPropertyIdentical = Internal::AreNamesIdentical(OldSerializedObject, NewSerializedObject, NewNamePropertyPair.Key);
+				OutDeltaChange.bHasNonPropertyChanges = true;
 			}
-
-			if (!bIsPropertyIdentical)
+		}
+		else if (OldSerializedTaggedData)
+		{
+			if (!IsTaggedDataBlockIdentical(NewNamePropertyPair.Key, *OldSerializedTaggedData, NewNamePropertyPair.Value))
 			{
 				OutDeltaChange.ChangedProperties.AddUnique(NewNamePropertyPair.Key);
 			}
 		}
+		else if (bFullDiff)
+		{
+			// Missing property, assume that the property changed
+			OutDeltaChange.ChangedProperties.AddUnique(NewNamePropertyPair.Key);
+		}
+	}
 
-		for (const TPair<FName, FTransactionSerializedProperty>& OldNamePropertyPair : OldSerializedObject.SerializedProperties)
+	if (bFullDiff)
+	{
+		for (const TPair<FName, FTransactionSerializedTaggedData>& OldNamePropertyPair : OldSerializedObject.SerializedTaggedData)
 		{
-			const FTransactionSerializedProperty* NewSerializedProperty = NewSerializedObject.SerializedProperties.Find(OldNamePropertyPair.Key);
-			if (!NewSerializedProperty)
+			if (!ShouldCompareTaggedData(OldNamePropertyPair.Key) || NewSerializedObject.SerializedTaggedData.Contains(OldNamePropertyPair.Key))
 			{
-				if (bFullDiff)
-				{
-					// Missing property, assume that the property changed
-					OutDeltaChange.ChangedProperties.AddUnique(OldNamePropertyPair.Key);
-				}
 				continue;
 			}
-		}
 
-		if (bFullDiff)
-		{
-			// Compare the data before the property block to see if something else in the object has changed
-			if (!OutDeltaChange.bHasNonPropertyChanges)
+			if (ShouldCompareAsNonPropertyData(OldNamePropertyPair.Key))
 			{
-				const int64 OldHeaderSize = FMath::Min(StartOfOldPropertyBlock, OldSerializedObject.SerializedData.Num());
-				const int64 CurrentHeaderSize = FMath::Min(StartOfNewPropertyBlock, NewSerializedObject.SerializedData.Num());
-
-				bool bIsHeaderIdentical = OldHeaderSize == CurrentHeaderSize;
-				if (bIsHeaderIdentical && CurrentHeaderSize > 0)
-				{
-					bIsHeaderIdentical = FMemory::Memcmp(OldSerializedObject.SerializedData.GetPtr(0), NewSerializedObject.SerializedData.GetPtr(0), CurrentHeaderSize) == 0;
-				}
-
-				if (!bIsHeaderIdentical)
-				{
-					OutDeltaChange.bHasNonPropertyChanges = true;
-				}
+				OutDeltaChange.bHasNonPropertyChanges = true;
 			}
-
-			// Compare the data after the property block to see if something else in the object has changed
-			if (!OutDeltaChange.bHasNonPropertyChanges)
+			else
 			{
-				const int64 OldFooterSize = OldSerializedObject.SerializedData.Num() - FMath::Max<int64>(EndOfOldPropertyBlock, 0);
-				const int64 CurrentFooterSize = NewSerializedObject.SerializedData.Num() - FMath::Max<int64>(EndOfNewPropertyBlock, 0);
-
-				bool bIsFooterIdentical = OldFooterSize == CurrentFooterSize;
-				if (bIsFooterIdentical && CurrentFooterSize > 0)
-				{
-					bIsFooterIdentical = FMemory::Memcmp(OldSerializedObject.SerializedData.GetPtr(EndOfOldPropertyBlock), NewSerializedObject.SerializedData.GetPtr(EndOfNewPropertyBlock), CurrentFooterSize) == 0;
-				}
-
-				if (!bIsFooterIdentical)
-				{
-					OutDeltaChange.bHasNonPropertyChanges = true;
-				}
+				// Missing property, assume that the property changed
+				OutDeltaChange.ChangedProperties.AddUnique(OldNamePropertyPair.Key);
 			}
 		}
 	}
-	else if (bFullDiff)
-	{
-		// No properties, so just compare the whole blob
-		bool bIsBlobIdentical = OldSerializedObject.SerializedData.Num() == NewSerializedObject.SerializedData.Num();
-		if (bIsBlobIdentical && NewSerializedObject.SerializedData.Num() > 0)
-		{
-			bIsBlobIdentical = FMemory::Memcmp(OldSerializedObject.SerializedData.GetPtr(0), NewSerializedObject.SerializedData.GetPtr(0), NewSerializedObject.SerializedData.Num()) == 0;
-		}
-
-		if (!bIsBlobIdentical)
-		{
-			OutDeltaChange.bHasNonPropertyChanges = true;
-		}
-	}
 }
 
 namespace Internal
 {
 
-bool AreObjectPointersIdentical(const FTransactionSerializedObject& OldSerializedObject, const FTransactionSerializedObject& NewSerializedObject, const FName PropertyName)
+bool AreObjectPointersIdentical(const FTransactionSerializedObject& OldSerializedObject, const FTransactionSerializedObject& NewSerializedObject, const FName TaggedDataKey)
 {
-	auto GetIndicesArray = [&PropertyName](const FTransactionSerializedObject& SerializedObject) -> const TArray<int32>&
+	auto GetIndicesArray = [&TaggedDataKey](const FTransactionSerializedObject& SerializedObject) -> const TArray<int32>&
 	{
-		if (const FTransactionSerializedIndices* SerializedObjectIndices = SerializedObject.SerializedObjectIndices.Find(PropertyName))
+		if (const FTransactionSerializedIndices* SerializedObjectIndices = SerializedObject.SerializedObjectIndices.Find(TaggedDataKey))
 		{
 			return SerializedObjectIndices->Indices;
 		}
@@ -490,11 +534,11 @@ bool AreObjectPointersIdentical(const FTransactionSerializedObject& OldSerialize
 	return bAreObjectPointersIdentical;
 }
 
-bool AreNamesIdentical(const FTransactionSerializedObject& OldSerializedObject, const FTransactionSerializedObject& NewSerializedObject, const FName PropertyName)
+bool AreNamesIdentical(const FTransactionSerializedObject& OldSerializedObject, const FTransactionSerializedObject& NewSerializedObject, const FName TaggedDataKey)
 {
-	auto GetIndicesArray = [&PropertyName](const FTransactionSerializedObject& SerializedObject) -> const TArray<int32>&
+	auto GetIndicesArray = [&TaggedDataKey](const FTransactionSerializedObject& SerializedObject) -> const TArray<int32>&
 	{
-		if (const FTransactionSerializedIndices* SerializedNameIndices = SerializedObject.SerializedNameIndices.Find(PropertyName))
+		if (const FTransactionSerializedIndices* SerializedNameIndices = SerializedObject.SerializedNameIndices.Find(TaggedDataKey))
 		{
 			return SerializedNameIndices->Indices;
 		}
diff --git a/Engine/Source/Runtime/Engine/Public/TransactionCommon.h b/Engine/Source/Runtime/Engine/Public/TransactionCommon.h
--- a/Engine/Source/Runtime/Engine/Public/TransactionCommon.h
+++ b/Engine/Source/Runtime/Engine/Public/TransactionCommon.h
@@ -112,20 +112,45 @@ struct TStructOpsTypeTraits<FTransactionPersistentObjectRef> : public TStructOps
 };
 
 USTRUCT()
-struct ENGINE_API FTransactionSerializedProperty
+struct ENGINE_API FTransactionSerializedTaggedData
 {
 	GENERATED_BODY()
 
 public:
-	static FName BuildSerializedPropertyKey(const FArchiveSerializedPropertyChain& InPropertyChain);
+	static FTransactionSerializedTaggedData FromOffsetAndSize(const int64 InOffset, const int64 InSize);
+	static FTransactionSerializedTaggedData FromStartAndEnd(const int64 InStart, const int64 InEnd);
 
 	void AppendSerializedData(const int64 InOffset, const int64 InSize);
+	void AppendSerializedData(const FTransactionSerializedTaggedData& InData);
 
-	/** Offset to the start of this property within the serialized object */
+	bool HasSerializedData() const;
+
+	int64 GetStart() const
+	{
+		return DataOffset;
+	}
+
+	int64 GetEnd() const
+	{
+		return DataOffset + DataSize;
+	}
+
+	friend bool operator==(const FTransactionSerializedTaggedData& LHS, const FTransactionSerializedTaggedData& RHS)
+	{
+		return LHS.DataOffset == RHS.DataOffset
+			&& LHS.DataSize == RHS.DataSize;
+	}
+
+	friend bool operator!=(const FTransactionSerializedTaggedData& LHS, const FTransactionSerializedTaggedData& RHS)
+	{
+		return !(LHS == RHS);
+	}
+
+	/** Offset to the start of the tagged data within the serialized object */
 	UPROPERTY()
 	int64 DataOffset = INDEX_NONE;
 
-	/** Size (in bytes) of this property within the serialized object */
+	/** Size (in bytes) of the tagged data within the serialized object */
 	UPROPERTY()
 	int64 DataSize = 0;
 };
@@ -229,7 +254,7 @@ public:
 		SerializedData.Reset();
 		ReferencedObjects.Reset();
 		ReferencedNames.Reset();
-		SerializedProperties.Reset();
+		SerializedTaggedData.Reset();
 		SerializedObjectIndices.Reset();
 		SerializedNameIndices.Reset();
 	}
@@ -246,7 +271,7 @@ public:
 		Exchange(SerializedData, Other.SerializedData);
 		Exchange(ReferencedObjects, Other.ReferencedObjects);
 		Exchange(ReferencedNames, Other.ReferencedNames);
-		Exchange(SerializedProperties, Other.SerializedProperties);
+		Exchange(SerializedTaggedData, Other.SerializedTaggedData);
 		Exchange(SerializedObjectIndices, Other.SerializedObjectIndices);
 		Exchange(SerializedNameIndices, Other.SerializedNameIndices);
 	}
@@ -291,15 +316,15 @@ public:
 	UPROPERTY()
 	TArray<FName> ReferencedNames;
 	
-	/** Information about the properties that were serialized within this object */
+	/** Information about tagged data (mainly properties) that were serialized within this object */
 	UPROPERTY()
-	TMap<FName, FTransactionSerializedProperty> SerializedProperties;
+	TMap<FName, FTransactionSerializedTaggedData> SerializedTaggedData;
 	
-	/** Information about the object pointer offsets that were serialized within this object (this maps the property name (or None if there was no property) to the ReferencedObjects indices of the property) */
+	/** Information about the object pointer offsets that were serialized within this object (this maps the tagged data key (or None if there was no tagged data key) to the ReferencedObjects indices of the tagged data) */
 	UPROPERTY()
 	TMap<FName, FTransactionSerializedIndices> SerializedObjectIndices;
 	
-	/** Information about the name offsets that were serialized within this object (this maps the property name to the ReferencedNames index of the property) */
+	/** Information about the name offsets that were serialized within this object (this maps the tagged data name to the ReferencedNames index of the tagged data) */
 	UPROPERTY()
 	TMap<FName, FTransactionSerializedIndices> SerializedNameIndices;
 };
@@ -307,6 +332,9 @@ public:
 namespace UE::Transaction
 {
 
+ENGINE_API extern const FName TaggedDataKey_UnknownData;
+ENGINE_API extern const FName TaggedDataKey_ScriptData;
+
 /** Transfers data from an array. */
 class ENGINE_API FSerializedObjectDataReader : public FArchiveUObject
 {
@@ -325,7 +353,7 @@ protected:
 	virtual FArchive& operator<<(class UObject*& Res) override;
 
 	const FTransactionSerializedObject& SerializedObject;
-	int64 Offset;
+	int64 Offset = 0;
 };
 
 /**
@@ -351,6 +379,13 @@ protected:
 		uint32 LastUpdateCount = 0;
 	};
 
+	FName GetTaggedDataKey() const;
+
+	bool DoesObjectMatchSerializedObject(const UObject* Obj) const;
+
+	virtual void MarkScriptSerializationStart(const UObject* Obj) override;
+	virtual void MarkScriptSerializationEnd(const UObject* Obj) override;
+
 	virtual void Serialize(void* SerData, int64 Num) override;
 
 	using FArchiveUObject::operator<<;
@@ -358,10 +393,19 @@ protected:
 	virtual FArchive& operator<<(class UObject*& Res) override;
 
 	FTransactionSerializedObject& SerializedObject;
+	int64 Offset = 0;
+
 	TMap<UObject*, int32> ObjectMap;
 	TMap<FName, int32> NameMap;
-	FCachedPropertyKey CachedSerializedTaggedPropertyKey;
-	int64 Offset;
+
+private:
+	bool bIsPerformingScriptSerialization = false;
+
+	mutable bool bWasUsingTaggedDataKey_UnknownData = false;
+	mutable bool bWasUsingTaggedDataKey_ScriptData = false;
+	mutable int32 TaggedDataKeyIndex_UnknownData = 0;
+	mutable int32 TaggedDataKeyIndex_ScriptData = 0;
+	mutable FCachedPropertyKey CachedSerializedTaggedPropertyKey;
 };
 
 namespace DiffUtil
@@ -372,8 +416,8 @@ ENGINE_API void GenerateObjectDiff(const FTransactionSerializedObject& OldSerial
 namespace Internal
 {
 
-ENGINE_API bool AreObjectPointersIdentical(const FTransactionSerializedObject& OldSerializedObject, const FTransactionSerializedObject& NewSerializedObject, const FName PropertyName);
-ENGINE_API bool AreNamesIdentical(const FTransactionSerializedObject& OldSerializedObject, const FTransactionSerializedObject& NewSerializedObject, const FName PropertyName);
+ENGINE_API bool AreObjectPointersIdentical(const FTransactionSerializedObject& OldSerializedObject, const FTransactionSerializedObject& NewSerializedObject, const FName TaggedDataKey);
+ENGINE_API bool AreNamesIdentical(const FTransactionSerializedObject& OldSerializedObject, const FTransactionSerializedObject& NewSerializedObject, const FName TaggedDataKey);
 
 } // namespace Internal
 
