Replay
* Actors flagged for serialization will now be serialized before we call destroy on them.  This allows them to be serialized in a state where they still have most components and properties intact.

diff --git a/Engine/Source/Runtime/Engine/Classes/Engine/DemoNetDriver.h b/Engine/Source/Runtime/Engine/Classes/Engine/DemoNetDriver.h
--- a/Engine/Source/Runtime/Engine/Classes/Engine/DemoNetDriver.h
+++ b/Engine/Source/Runtime/Engine/Classes/Engine/DemoNetDriver.h
@@ -679,6 +679,11 @@ private:
 	// Diff Actor plus it's Components and Subobjects from given ActorChannel
 	void DiffActorProperties(UActorChannel* const ActorChannel);
 
+	// Callback sent just before an actor has destroy called on itself.
+	void OnActorPreDestroy(AActor* DestroyedActor);
+
+	FDelegateHandle DelegateHandleActorPreDestroy;
+
 protected:
 
 	void ProcessSeamlessTravel(int32 LevelIndex);
diff --git a/Engine/Source/Runtime/Engine/Private/DemoNetDriver.cpp b/Engine/Source/Runtime/Engine/Private/DemoNetDriver.cpp
--- a/Engine/Source/Runtime/Engine/Private/DemoNetDriver.cpp
+++ b/Engine/Source/Runtime/Engine/Private/DemoNetDriver.cpp
@@ -5478,9 +5478,19 @@ void UDemoNetDriver::SetAnalyticsProvider(TSharedPtr<IAnalyticsProvider> InProvi
 
 void UDemoNetDriver::SetWorld(UWorld* InWorld)
 {
+	if (World)
+	{
+		World->RemoveOnActorDestroyededHandler(DelegateHandleActorPreDestroy);
+	}
+
 	Super::SetWorld(InWorld);
 
 	ReplayHelper.World = InWorld;
+
+	if (InWorld)
+	{
+		DelegateHandleActorPreDestroy = InWorld->AddOnActorDestroyedHandler(FOnActorDestroyed::FDelegate::CreateUObject(this, &UDemoNetDriver::OnActorPreDestroy));
+	}
 }
 
 bool UDemoNetDriver::ShouldForwardFunction(AActor* Actor, UFunction* Function, void* Parms) const
diff --git a/Engine/Source/Runtime/Engine/Private/ReplayHelper.cpp b/Engine/Source/Runtime/Engine/Private/ReplayHelper.cpp
--- a/Engine/Source/Runtime/Engine/Private/ReplayHelper.cpp
+++ b/Engine/Source/Runtime/Engine/Private/ReplayHelper.cpp
@@ -2296,9 +2296,8 @@
 	bPendingCheckpointRequest = true;
 }
 
-void FReplayHelper::NotifyActorDestroyed(UNetConnection* Connection, AActor* Actor)
+void FReplayHelper::OnActorPreDestroy(UNetConnection* Connection, AActor* Actor)
 {
-	check(Actor);
 	check(Connection);
 
 	// if we're recording a checkpoint, and we have not yet passed the actor recording phase
@@ -2312,7 +2311,7 @@
 				// if this actor is in the pending checkpoint actor list and we have not already recorded it
 				if (int32* PendingIndex = CheckpointSaveContext.PendingActorToIndex.Find(Actor))
 				{
-					if (CheckpointSaveContext.PendingCheckpointActors.IsValidIndex(*PendingIndex) 
+					if (CheckpointSaveContext.PendingCheckpointActors.IsValidIndex(*PendingIndex)
 						&& (*PendingIndex >= CheckpointSaveContext.NextAmortizedItem)
 						&& (CheckpointSaveContext.PendingCheckpointActors[*PendingIndex].Actor.Get() == Actor))
 					{
@@ -2328,6 +2327,7 @@
 						TArrayView<FPendingCheckPointActor> PendingView(CheckpointSaveContext.PendingCheckpointActors);
 						int32 ActorIndex = 0;
 
+						// Serialize the actor one last time before it gets destroyed
 						ProcessCheckpointActors(Connection, PendingView.Slice(*PendingIndex, 1), ActorIndex, Params);
 
 						// don't allow the channel index to be reused until we're done with the checkpoint
@@ -2337,7 +2337,13 @@
 			}
 		}
 	}
+}
 
+void FReplayHelper::NotifyActorDestroyed(UNetConnection* Connection, AActor* Actor)
+{
+	check(Actor);
+	check(Connection);
+
 	const bool bNetStartup = Actor->IsNetStartupActor();
 	const bool bActorRewindable = Actor->bReplayRewindable;
 	const bool bDeltaCheckpoint = HasDeltaCheckpoints();
diff --git a/Engine/Source/Runtime/Engine/Public/ReplayHelper.h b/Engine/Source/Runtime/Engine/Public/ReplayHelper.h
--- a/Engine/Source/Runtime/Engine/Public/ReplayHelper.h
+++ b/Engine/Source/Runtime/Engine/Public/ReplayHelper.h
@@ -106,6 +106,8 @@ private:
 
 	void RequestCheckpoint();
 
+	void OnActorPreDestroy(UNetConnection* Connection, AActor* Actor);
+
 	void NotifyActorDestroyed(UNetConnection* Connection, AActor* Actor);
 
 private:
