- GameEngine flag to disable game window setting overrides
- Fix UGameUserSettings::ResetToCurrentSettings so that it sets the window mode and resolution to the currently active settings (as stated by the method's documentation) and not the last confirmed settings

diff --git a/Engine/Source/Runtime/Engine/Classes/Engine/GameEngine.h b/Engine/Source/Runtime/Engine/Classes/Engine/GameEngine.h
--- a/Engine/Source/Runtime/Engine/Classes/Engine/GameEngine.h
+++ b/Engine/Source/Runtime/Engine/Classes/Engine/GameEngine.h
@@ -72,8 +72,16 @@ public:
 
 	static void SafeFrameChanged();
 
+	/** 
+	 * Enables/disables game window resolution setting overrides specified on the command line (true by default) 
+	 * @note Does not trigger a refresh of the game window based on the newly effective settings
+	 * @note Not thread-safe; must be called from the main thread
+	 */
+	static void EnableGameWindowSettingsOverride(bool bEnabled);
+
 	/**
 	 * Modifies the game window resolution settings if any overrides have been specified on the command line
+	 * and overrides are enabled (see EnableGameWindowSettingsOverride)
 	 *
 	 * @param ResolutionX	[in/out] Width of the game window, in pixels
 	 * @param ResolutionY	[in/out] Height of the game window, in pixels
diff --git a/Engine/Source/Runtime/Engine/Private/GameEngine.cpp b/Engine/Source/Runtime/Engine/Private/GameEngine.cpp
--- a/Engine/Source/Runtime/Engine/Private/GameEngine.cpp
+++ b/Engine/Source/Runtime/Engine/Private/GameEngine.cpp
@@ -85,6 +85,8 @@ ENGINE_API bool GDisallowNetworkTravel = false;
 // How slow must a frame be (in seconds) to be logged out (<= 0 to disable)
 ENGINE_API float GSlowFrameLoggingThreshold = 0.0f;
 
+static bool bGameWindowSettingsOverrideEnabled = true;
+
 static FAutoConsoleVariableRef CvarSlowFrameLoggingThreshold(
 	TEXT("t.SlowFrameLoggingThreshold"),
 	GSlowFrameLoggingThreshold,
@@ -288,28 +290,37 @@ FSceneViewport* UGameEngine::GetGameSceneViewport(UGameViewportClient* ViewportC
 	return ViewportClient->GetGameViewport();
 }
 
+void UGameEngine::EnableGameWindowSettingsOverride(bool bEnabled)
+{
+	check(IsInGameThread());
+	bGameWindowSettingsOverrideEnabled = bEnabled;
+}
+
 void UGameEngine::ConditionallyOverrideSettings(int32& ResolutionX, int32& ResolutionY, EWindowMode::Type& WindowMode)
 {
-	if (FParse::Param(FCommandLine::Get(), TEXT("Windowed")) || FParse::Param(FCommandLine::Get(), TEXT("SimMobile")))
-	{
-		// -Windowed or -SimMobile
-		WindowMode = EWindowMode::Windowed;
-	}
-	else if (FParse::Param(FCommandLine::Get(), TEXT("FullScreen")))
+	if (bGameWindowSettingsOverrideEnabled)
 	{
-		// -FullScreen
-		static auto CVar = IConsoleManager::Get().FindTConsoleVariableDataInt(TEXT("r.FullScreenMode"));
-		check(CVar);
-		WindowMode = CVar->GetValueOnGameThread() == 0 ? EWindowMode::Fullscreen : EWindowMode::WindowedFullscreen;
-
-		if (PLATFORM_WINDOWS && WindowMode == EWindowMode::Fullscreen)
+		if (FParse::Param(FCommandLine::Get(), TEXT("Windowed")) || FParse::Param(FCommandLine::Get(), TEXT("SimMobile")))
+		{
+			// -Windowed or -SimMobile
+			WindowMode = EWindowMode::Windowed;
+		}
+		else if (FParse::Param(FCommandLine::Get(), TEXT("FullScreen")))
 		{
-			// Handle fullscreen mode differently for D3D11/D3D12
-			static const bool bD3D12 = FParse::Param(FCommandLine::Get(), TEXT("d3d12")) || FParse::Param(FCommandLine::Get(), TEXT("dx12"));
-			if (bD3D12)
+			// -FullScreen
+			static auto CVar = IConsoleManager::Get().FindTConsoleVariableDataInt(TEXT("r.FullScreenMode"));
+			check(CVar);
+			WindowMode = CVar->GetValueOnGameThread() == 0 ? EWindowMode::Fullscreen : EWindowMode::WindowedFullscreen;
+
+			if (PLATFORM_WINDOWS && WindowMode == EWindowMode::Fullscreen)
 			{
-				// Force D3D12 RHI to use windowed fullscreen mode
-				WindowMode = EWindowMode::WindowedFullscreen;
+				// Handle fullscreen mode differently for D3D11/D3D12
+				static const bool bD3D12 = FParse::Param(FCommandLine::Get(), TEXT("d3d12")) || FParse::Param(FCommandLine::Get(), TEXT("dx12"));
+				if (bD3D12)
+				{
+					// Force D3D12 RHI to use windowed fullscreen mode
+					WindowMode = EWindowMode::WindowedFullscreen;
+				}
 			}
 		}
 	}
@@ -320,7 +331,7 @@ void UGameEngine::ConditionallyOverrideSettings(int32& ResolutionX, int32& Resol
 void UGameEngine::DetermineGameWindowResolution( int32& ResolutionX, int32& ResolutionY, EWindowMode::Type& WindowMode, bool bUseWorkAreaForWindowed )
 {
 	FString ResolutionStr;;
-	if (FParse::Value(FCommandLine::Get(), TEXT("Res="), ResolutionStr))
+	if (bGameWindowSettingsOverrideEnabled && FParse::Value(FCommandLine::Get(), TEXT("Res="), ResolutionStr))
 	{
 		uint32 ResX = 0;
 		uint32 ResY = 0;
@@ -335,8 +346,8 @@ void UGameEngine::DetermineGameWindowResolution( int32& ResolutionX, int32& Reso
 	}
 	else
 	{
-		bool UserSpecifiedWidth = FParse::Value(FCommandLine::Get(), TEXT("ResX="), ResolutionX);
-		bool UserSpecifiedHeight = FParse::Value(FCommandLine::Get(), TEXT("ResY="), ResolutionY);
+		bool UserSpecifiedWidth = bGameWindowSettingsOverrideEnabled && FParse::Value(FCommandLine::Get(), TEXT("ResX="), ResolutionX);
+		bool UserSpecifiedHeight = bGameWindowSettingsOverrideEnabled && FParse::Value(FCommandLine::Get(), TEXT("ResY="), ResolutionY);
 
 		const float AspectRatio = 16.0 / 9.0;
 
@@ -397,7 +408,7 @@ void UGameEngine::DetermineGameWindowResolution( int32& ResolutionX, int32& Reso
 	}
 
 	// Optionally force the resolution by passing -ForceRes
-	const bool bForceRes = FParse::Param(FCommandLine::Get(), TEXT("ForceRes"));
+	const bool bForceRes = bGameWindowSettingsOverrideEnabled && FParse::Param(FCommandLine::Get(), TEXT("ForceRes"));
 
 	//Don't allow a resolution bigger then the desktop found a convenient one
 	if (!bForceRes && !IsRunningDedicatedServer() && ((ResolutionX <= 0 || ResolutionX > MaxResolutionX) || (ResolutionY <= 0 || ResolutionY > MaxResolutionY)))
@@ -443,7 +454,7 @@ void UGameEngine::DetermineGameWindowResolution( int32& ResolutionX, int32& Reso
 	}
 
 
-	if (FParse::Param(FCommandLine::Get(), TEXT("Portrait")))
+	if (bGameWindowSettingsOverrideEnabled && FParse::Param(FCommandLine::Get(), TEXT("Portrait")))
 	{
 		Swap(ResolutionX, ResolutionY);
 	}
diff --git a/Engine/Source/Runtime/Engine/Private/GameUserSettings.cpp b/Engine/Source/Runtime/Engine/Private/GameUserSettings.cpp
--- a/Engine/Source/Runtime/Engine/Private/GameUserSettings.cpp
+++ b/Engine/Source/Runtime/Engine/Private/GameUserSettings.cpp
@@ -700,13 +700,10 @@ void UGameUserSettings::ResetToCurrentSettings()
 		// Set the current dynamic resolution state
 		SetDynamicResolutionEnabled(GEngine->GetDynamicResolutionUserSetting());
 
-		// Reset to confirmed settings
-		FullscreenMode = LastConfirmedFullscreenMode;
-		ResolutionSizeX = LastUserConfirmedResolutionSizeX;
-		ResolutionSizeY = LastUserConfirmedResolutionSizeY;
-
-		DesiredScreenWidth = LastUserConfirmedDesiredScreenWidth;
-		DesiredScreenHeight = LastUserConfirmedDesiredScreenHeight;
+		// Reset confirmed settings
+		ConfirmVideoMode();
+		LastUserConfirmedDesiredScreenWidth = DesiredScreenWidth;
+		LastUserConfirmedDesiredScreenHeight = DesiredScreenHeight;
 
 		// Reset the quality settings to the current levels
 		ScalabilityQuality = Scalability::GetQualityLevels();
