#UE Merge the rest of 21150980.

Replay
* Actors flagged for serialization will now be serialized before we call destroy on them.  This allows them to be serialized in a state where they still have most components and properties intact.

diff --git a/Engine/Source/Runtime/Engine/Private/DemoNetDriver.cpp b/Engine/Source/Runtime/Engine/Private/DemoNetDriver.cpp
--- a/Engine/Source/Runtime/Engine/Private/DemoNetDriver.cpp
+++ b/Engine/Source/Runtime/Engine/Private/DemoNetDriver.cpp
@@ -874,6 +874,11 @@ void UDemoNetDriver::FinishDestroy()
 
 	ReplayHelper.OnReplayRecordError.RemoveAll(this);
 	ReplayHelper.OnReplayPlaybackError.RemoveAll(this);
+	
+	if (World)
+	{
+		World->RemoveOnActorDestroyededHandler(DelegateHandleActorPreDestroy);
+	}
 
 	Super::FinishDestroy();
 }
@@ -4975,6 +4980,16 @@ void UDemoNetDriver::InitDestroyedStartupActors()
 	}
 }
 
+void UDemoNetDriver::OnActorPreDestroy(AActor* DestroyedActor)
+{
+	check(DestroyedActor);
+
+	if (IsRecording())
+	{
+		ReplayHelper.OnActorPreDestroy(ClientConnections[0], DestroyedActor);
+	}
+}
+
 void UDemoNetDriver::NotifyActorDestroyed(AActor* Actor, bool IsSeamlessTravel)
 {
 	check(Actor != nullptr);
