World Partition: fixed superfluous world partition object initialization when streaming partitioned worlds.

#rnx

diff --git a/Engine/Source/Runtime/Engine/Private/Level.cpp b/Engine/Source/Runtime/Engine/Private/Level.cpp
--- a/Engine/Source/Runtime/Engine/Private/Level.cpp
+++ b/Engine/Source/Runtime/Engine/Private/Level.cpp
@@ -2782,11 +2782,28 @@ void ULevel::OnLevelLoaded()
 
 	// 1. Cook commandlet does it's own UWorldPartition::Initialize call in FWorldPartitionCookPackageSplitter::GetGenerateList
 	// 2. Do not Initialize if World doesn't have a UWorldPartitionSubsystem (Known case is when WorldType == EWorldType::Inactive)
-	if (!IsRunningCookCommandlet() && GetWorld()->HasSubsystem<UWorldPartitionSubsystem>())
+	if (!IsRunningCookCommandlet() && OwningWorld->HasSubsystem<UWorldPartitionSubsystem>())
 	{
 		if (UWorldPartition* WorldPartition = GetWorldPartition())
 		{
-			if (IsPersistentLevel() || GetWorld()->IsGameWorld() || IsValidLevelInstanceWorldPartition(WorldPartition))
+			//
+			// When do we need to initialize the associated world partition object?
+			//
+			//	- When the level is the main world persistent level
+			//	- When the sublevel is streamed in the editor (mainly for data layers)
+			//	- When the sublevel is streamed in game and the main world is not partitioned
+			//
+			const bool bIsOwningWorldGameWorld = OwningWorld->IsGameWorld();
+			const bool bIsOwningWorldPartitioned = OwningWorld->IsPartitionedWorld();
+			const bool bIsValidLevelInstance = IsValidLevelInstanceWorldPartition(WorldPartition);
+			const bool bIsMainWorldLevel = OwningWorld->PersistentLevel == this;
+			const bool bInitializeForEditor = !bIsOwningWorldGameWorld && bIsValidLevelInstance;
+			const bool bInitializeForGame = bIsOwningWorldGameWorld && !bIsOwningWorldPartitioned;
+
+			UE_LOG(LogWorldPartition, Log, TEXT("ULevel::OnLevelLoaded(%s)(bIsOwningWorldGameWorld=%d, bIsOwningWorldPartitioned=%d, bIsValidLevelInstance=%d, InitializeForMainWorld=%d, InitializeForEditor=%d, InitializeForGame=%d)"), 
+				*GetTypedOuter<UWorld>()->GetName(), bIsOwningWorldGameWorld ? 1 : 0, bIsOwningWorldPartitioned ? 1 : 0, bIsValidLevelInstance ? 1 : 0, bIsMainWorldLevel ? 1 : 0, bInitializeForEditor ? 1 : 0, bInitializeForGame ? 1 : 0);
+
+			if (bIsMainWorldLevel || bInitializeForEditor || bInitializeForGame)
 			{
 				FTransform Transform = FTransform::Identity;
 				if (ULevelStreaming* LevelStreaming = FLevelUtils::FindStreamingLevel(this))
@@ -2798,11 +2815,11 @@ void ULevel::OnLevelLoaded()
 				// was streamed-out and is streamed-in again, without a CleanupLevel (GC).
 				if (!WorldPartition->IsInitialized())
 				{
-					WorldPartition->Initialize(GetWorld(), Transform);
+					WorldPartition->Initialize(OwningWorld, Transform);
 				}
 				else
 				{
-					check(GetWorld()->IsGameWorld());
+					check(OwningWorld->IsGameWorld());
 					check(WorldPartition->GetInstanceTransform().Equals(Transform));
 				}
 			}
