LevelInstance: Handle "Cancel" properly on MessageBox close button

diff --git a/Engine/Source/Runtime/Engine/Private/LevelInstance/LevelInstanceSubsystem.cpp b/Engine/Source/Runtime/Engine/Private/LevelInstance/LevelInstanceSubsystem.cpp
--- a/Engine/Source/Runtime/Engine/Private/LevelInstance/LevelInstanceSubsystem.cpp
+++ b/Engine/Source/Runtime/Engine/Private/LevelInstance/LevelInstanceSubsystem.cpp
@@ -514,12 +514,12 @@ bool ULevelInstanceSubsystem::OnExitEditorModeInternal(bool bForceExit)
 			FText Title = LOCTEXT("CommitOrDiscardChangesTitle", "Save changes?");
 			// if bForceExit we can't cancel the exiting of the mode so the user needs to decide between saving or discarding
 			EAppReturnType::Type Ret = FMessageDialog::Open(bForceExit ? EAppMsgType::YesNo : EAppMsgType::YesNoCancel, LOCTEXT("CommitOrDiscardChangesMsg", "Unsaved Level changes will get discarded. Do you want to save them now?"), &Title);
-			if (Ret == EAppReturnType::Cancel)
+			if (Ret == EAppReturnType::Cancel && !bForceExit)
 			{
 				return false;
 			}
 
-			bDiscard = (Ret == EAppReturnType::No);
+			bDiscard = (Ret != EAppReturnType::Yes);
 		}
 
 		return CommitLevelInstanceInternal(LevelInstanceEdit, bDiscard, /*bDiscardOnFailure=*/bForceExit);
