Unshelved from pending changelist '20736424':

World Partition Streaming Generation:
- Create actor descriptors map and actors clustering once per container instead of per container instance.
- Created IGenerateStreamingContext interface between streaming generation and runtime hashes to completly separate the two processes.
- Deprecated FActorClusterContext which was replaced by generic function GenerateObjectsClusters which was added in Engine recently.

Tests performed:
- Generated HLODs for AncientWorld and Small_City and tested in PIE, -game and cooked builds.
- Several tests usingf internal maps and projects.

diff --git a/Engine/Plugins/Editor/WorldPartitionHLODUtilities/Source/Private/WorldPartitionHLODUtilities.cpp b/Engine/Plugins/Editor/WorldPartitionHLODUtilities/Source/Private/WorldPartitionHLODUtilities.cpp
--- a/Engine/Plugins/Editor/WorldPartitionHLODUtilities/Source/Private/WorldPartitionHLODUtilities.cpp
+++ b/Engine/Plugins/Editor/WorldPartitionHLODUtilities/Source/Private/WorldPartitionHLODUtilities.cpp
@@ -99,7 +99,7 @@ static uint32 ComputeHLODHash(AWorldPartitionHLOD* InHLODActor, const TArray<AAc
 	return Ar.GetCrc();
 }
 
-TArray<AWorldPartitionHLOD*> FWorldPartitionHLODUtilities::CreateHLODActors(FHLODCreationContext& InCreationContext, const FHLODCreationParams& InCreationParams, const TSet<FActorInstance>& InActors, const TArray<const UDataLayerInstance*>& InDataLayersInstances)
+TArray<AWorldPartitionHLOD*> FWorldPartitionHLODUtilities::CreateHLODActors(FHLODCreationContext& InCreationContext, const FHLODCreationParams& InCreationParams, const TArray<IStreamingGenerationContext::FActorInstance>& InActors, const TArray<const UDataLayerInstance*>& InDataLayersInstances)
 {
 	struct FSubActorsInfo
 	{
@@ -108,7 +108,7 @@ TArray<AWorldPartitionHLOD*> FWorldPartitionHLODUtilities::CreateHLODActors(FHLO
 	};
 	TMap<UHLODLayer*, FSubActorsInfo> SubActorsInfos;
 
-	for (const FActorInstance& ActorInstance : InActors)
+	for (const IStreamingGenerationContext::FActorInstance& ActorInstance : InActors)
 	{
 		const FWorldPartitionActorDescView& ActorDescView = ActorInstance.GetActorDescView();
 		if (ActorDescView.GetActorIsHLODRelevant())
@@ -118,7 +118,7 @@ TArray<AWorldPartitionHLOD*> FWorldPartitionHLODUtilities::CreateHLODActors(FHLO
 			{
 				FSubActorsInfo& SubActorsInfo = SubActorsInfos.FindOrAdd(HLODLayer);
 
-				SubActorsInfo.SubActors.Emplace(ActorDescView.GetGuid(), ActorDescView.GetActorPackage(), ActorDescView.GetActorPath(), ActorInstance.ContainerInstance->ID, ActorInstance.ContainerInstance->Container->GetContainerPackage(), ActorInstance.ContainerInstance->Transform);
+				SubActorsInfo.SubActors.Emplace(ActorDescView.GetGuid(), ActorDescView.GetActorPackage(), ActorDescView.GetActorPath(), ActorInstance.GetContainerID(), ActorInstance.GetActorDescContainer()->GetContainerPackage(), ActorInstance.GetTransform());
 				if (ActorDescView.GetIsSpatiallyLoaded())
 				{
 					SubActorsInfo.bIsSpatiallyLoaded = true;
diff --git a/Engine/Plugins/Editor/WorldPartitionHLODUtilities/Source/Public/WorldPartitionHLODUtilities.h b/Engine/Plugins/Editor/WorldPartitionHLODUtilities/Source/Public/WorldPartitionHLODUtilities.h
--- a/Engine/Plugins/Editor/WorldPartitionHLODUtilities/Source/Public/WorldPartitionHLODUtilities.h
+++ b/Engine/Plugins/Editor/WorldPartitionHLODUtilities/Source/Public/WorldPartitionHLODUtilities.h
@@ -10,7 +10,7 @@
 class FWorldPartitionHLODUtilities : public IWorldPartitionHLODUtilities
 {
 public:
-	virtual TArray<AWorldPartitionHLOD*> CreateHLODActors(FHLODCreationContext& InCreationContext, const FHLODCreationParams& InCreationParams, const TSet<FActorInstance>& InActors, const TArray<const UDataLayerInstance*>& InDataLayerInstances) override;
+	virtual TArray<AWorldPartitionHLOD*> CreateHLODActors(FHLODCreationContext& InCreationContext, const FHLODCreationParams& InCreationParams, const TArray<IStreamingGenerationContext::FActorInstance>& InActors, const TArray<const UDataLayerInstance*>& InDataLayerInstances) override;
 	virtual uint32 BuildHLOD(AWorldPartitionHLOD* InHLODActor) override;
 	virtual TSubclassOf<UHLODBuilder> GetHLODBuilderClass(const UHLODLayer* InHLODLayer) override;
 	virtual UHLODBuilderSettings* CreateHLODBuilderSettings(UHLODLayer* InHLODLayer) override;
diff --git a/Engine/Source/Runtime/Core/Public/Misc/HierarchicalLogArchive.h b/Engine/Source/Runtime/Core/Public/Misc/HierarchicalLogArchive.h
--- a/Engine/Source/Runtime/Core/Public/Misc/HierarchicalLogArchive.h
+++ b/Engine/Source/Runtime/Core/Public/Misc/HierarchicalLogArchive.h
@@ -50,4 +50,11 @@ private:
 	void WriteLine(const FString& InLine, bool bIndent = false);
 
 	int32 Indentation;
-};
\ No newline at end of file
+};
+
+#if NO_LOGGING
+#define UE_SCOPED_INDENT_LOG_ARCHIVE(Code)
+#else
+#define UE_SCOPED_INDENT_LOG_ARCHIVE(Code) \
+	FHierarchicalLogArchive::FIndentScope BODY_MACRO_COMBINE(Scoped,Indent,_,__LINE__) = Code
+#endif
diff --git a/Engine/Source/Runtime/Engine/Private/WorldPartition/DataLayer/DataLayerSubsystem.cpp b/Engine/Source/Runtime/Engine/Private/WorldPartition/DataLayer/DataLayerSubsystem.cpp
--- a/Engine/Source/Runtime/Engine/Private/WorldPartition/DataLayer/DataLayerSubsystem.cpp
+++ b/Engine/Source/Runtime/Engine/Private/WorldPartition/DataLayer/DataLayerSubsystem.cpp
@@ -180,6 +180,26 @@ void UDataLayerSubsystem::GetUserLoadedInEditorStates(TArray<FName>& OutDataLaye
 		return true;
 	});
 }
+
+TArray<const UDataLayerInstance*> UDataLayerSubsystem::GetRuntimeDataLayerInstances(UWorld* InWorld, const TArray<FName>& DataLayerInstanceNames)
+{
+	TArray<const UDataLayerInstance*> DataLayerInstances;
+	if (const UDataLayerSubsystem* DataLayerSubsystem = UWorld::GetSubsystem<UDataLayerSubsystem>(InWorld))
+	{
+		for (const FName& DataLayerInstanceName : DataLayerInstanceNames)
+		{
+			if (const UDataLayerInstance* DataLayerInstance = DataLayerSubsystem->GetDataLayerInstance(DataLayerInstanceName))
+			{
+				if (ensure(DataLayerInstance->IsRuntime()))
+				{
+					DataLayerInstances.Add(DataLayerInstance);
+				}
+			}
+		}
+	}
+	return DataLayerInstances;
+}
+
 #endif
 
 TSet<FName> UDataLayerSubsystem::GetEffectiveActiveDataLayerNames() const
diff --git a/Engine/Source/Runtime/Engine/Private/WorldPartition/DataLayer/WorldDataLayersActorDesc.cpp b/Engine/Source/Runtime/Engine/Private/WorldPartition/DataLayer/WorldDataLayersActorDesc.cpp
--- a/Engine/Source/Runtime/Engine/Private/WorldPartition/DataLayer/WorldDataLayersActorDesc.cpp
+++ b/Engine/Source/Runtime/Engine/Private/WorldPartition/DataLayer/WorldDataLayersActorDesc.cpp
@@ -7,7 +7,7 @@
 #include "WorldPartition/DataLayer/WorldDataLayers.h"
 #include "WorldPartition/DataLayer/DataLayerInstanceWithAsset.h"
 #include "WorldPartition/DataLayer/DataLayerAsset.h"
-#include "WorldPartition/WorldPartitionActorCluster.h"
+#include "WorldPartition/WorldPartitionActorContainerID.h"
 #include "UObject/FortniteNCBranchObjectVersion.h"
 
 FArchive& operator<<(FArchive& Ar, FDataLayerInstanceDesc& Desc)
diff --git a/Engine/Source/Runtime/Engine/Private/WorldPartition/RuntimeSpatialHash/RuntimeSpatialHashGridHelper.cpp b/Engine/Source/Runtime/Engine/Private/WorldPartition/RuntimeSpatialHash/RuntimeSpatialHashGridHelper.cpp
--- a/Engine/Source/Runtime/Engine/Private/WorldPartition/RuntimeSpatialHash/RuntimeSpatialHashGridHelper.cpp
+++ b/Engine/Source/Runtime/Engine/Private/WorldPartition/RuntimeSpatialHash/RuntimeSpatialHashGridHelper.cpp
@@ -2,6 +2,7 @@
 
 #include "WorldPartition/RuntimeSpatialHash/RuntimeSpatialHashGridHelper.h"
 #include "ActorPartition/PartitionActor.h"
+#include "WorldPartition/WorldPartitionStreamingGenerationContext.h"
 #include "Algo/Transform.h"
 #include "ProfilingDebugging/ScopedTimers.h"
 #include "WorldPartition/DataLayer/DataLayerInstance.h"
@@ -136,31 +137,6 @@ int32 FSquare2DGridHelper::ForEachIntersectingCells(const FSphericalSector& InSh
 }
 
 #if WITH_EDITOR
-void FSquare2DGridHelper::ValidateSingleActorReferer()
-{
-	UE_SCOPED_TIMER(TEXT("ValidateSingleActorReferer"), LogWorldPartition, Display);
-
-	TSet<FActorInstance> ActorUsage;
-	for (int32 Level = 0; Level < Levels.Num() - 1; Level++)
-	{
-		for (const FSquare2DGridHelper::FGridLevel::FGridCell& ThisCell : Levels[Level].Cells)
-		{
-			for (const FGridLevel::FGridCellDataChunk& DataChunk : ThisCell.GetDataChunks())
-			{
-				for (const FActorInstance& ActorInstance : DataChunk.GetActors())
-				{
-					bool bWasAlreadyInSet;
-					ActorUsage.Add(ActorInstance, &bWasAlreadyInSet);
-					check(!bWasAlreadyInSet);
-				}
-			}
-		}
-	}
-}
-#endif // #if WITH_EDITOR
-
-#if WITH_EDITOR
-
 FSquare2DGridHelper GetGridHelper(const FBox& WorldBounds, int64 GridCellSize)
 {
 	// Default grid to a minimum of 1 level and 1 cell, for always loaded actors
@@ -168,7 +144,7 @@ FSquare2DGridHelper GetGridHelper(const FBox& WorldBounds, int64 GridCellSize)
 	return FSquare2DGridHelper(WorldBounds, GridOrigin, GridCellSize);
 }
 
-FSquare2DGridHelper GetPartitionedActors(const UWorldPartition* WorldPartition, const FBox& WorldBounds, const FSpatialHashRuntimeGrid& Grid, const TArray<const FActorClusterInstance*>& GridActors)
+FSquare2DGridHelper GetPartitionedActors(const UWorldPartition* WorldPartition, const FBox& WorldBounds, const FSpatialHashRuntimeGrid& Grid, const TArray<const IStreamingGenerationContext::FActorSetInstance*>& ActorSetInstances)
 {
 	TRACE_CPUPROFILER_EVENT_SCOPE(GetPartitionedActors);
 	UE_SCOPED_TIMER(TEXT("GetPartitionedActors"), LogWorldPartition, Display);
@@ -190,20 +166,20 @@ FSquare2DGridHelper GetPartitionedActors(const UWorldPartition* WorldPartition,
 
 	const float CellArea = PartitionedActors.GetLowestLevel().CellSize * PartitionedActors.GetLowestLevel().CellSize;
 
-	auto ShouldActorUseLocationPlacement = [&CellArea](const FActorClusterInstance* ClusterInstance)
+	auto ShouldActorUseLocationPlacement = [&CellArea](const IStreamingGenerationContext::FActorSetInstance* ActorSetInstance)
 	{
 		if (GRuntimeSpatialHashPlaceSmallActorsUsingLocation)
 		{
-			const FBox2D ClusterBounds(FVector2D(ClusterInstance->Bounds.Min), FVector2D(ClusterInstance->Bounds.Max));
+			const FBox2D ClusterBounds(FVector2D(ActorSetInstance->Bounds.Min), FVector2D(ActorSetInstance->Bounds.Max));
 			return ClusterBounds.GetArea() <= CellArea;
 		}
 
 		if (GRuntimeSpatialHashPlacePartitionActorsUsingLocation)
 		{
 			bool bUseLocation = true;
-			for (const FGuid& ActorGuid : ClusterInstance->Cluster->Actors)
+			for (const FGuid& ActorGuid : ActorSetInstance->ActorSet->Actors)
 			{
-				const FWorldPartitionActorDescView& ActorDescView = ClusterInstance->ContainerInstance->GetActorDescView(ActorGuid);
+				const FWorldPartitionActorDescView& ActorDescView = ActorSetInstance->ContainerInstance->ActorDescViewMap->FindByGuidChecked(ActorGuid);
 
 				if (!ActorDescView.GetActorNativeClass()->IsChildOf<APartitionActor>())
 				{
@@ -217,20 +193,19 @@ FSquare2DGridHelper GetPartitionedActors(const UWorldPartition* WorldPartition,
 		return false;
 	};
 
-	for (const FActorClusterInstance* ClusterInstance : GridActors)
+	for (const IStreamingGenerationContext::FActorSetInstance* ActorSetInstance : ActorSetInstances)
 	{
-		const FActorCluster* ActorCluster = ClusterInstance->Cluster;
-		const FBox2D ClusterBounds(FVector2D(ClusterInstance->Bounds.Min), FVector2D(ClusterInstance->Bounds.Max));
-		check(ActorCluster && ActorCluster->Actors.Num() > 0);
+		const FBox2D ActorSetInstanceBounds(FVector2D(ActorSetInstance->Bounds.Min), FVector2D(ActorSetInstance->Bounds.Max));
+		check(ActorSetInstance->ActorSet->Actors.Num() > 0);
 		FSquare2DGridHelper::FGridLevel::FGridCell* GridCell = nullptr;
 
-		if (ActorCluster->bIsSpatiallyLoaded)
+		if (ActorSetInstance->bIsSpatiallyLoaded)
 		{
-			if (ShouldActorUseLocationPlacement(ClusterInstance))
+			if (ShouldActorUseLocationPlacement(ActorSetInstance))
 			{
 				// Find grid level cell that contains the actor cluster pivot and put actors in it.
 				FGridCellCoord2 CellCoords;
-				if (PartitionedActors.GetLowestLevel().GetCellCoords(ClusterBounds.GetCenter(), CellCoords))
+				if (PartitionedActors.GetLowestLevel().GetCellCoords(ActorSetInstanceBounds.GetCenter(), CellCoords))
 				{
 					GridCell = &PartitionedActors.GetLowestLevel().GetCell(CellCoords);
 				}
@@ -238,7 +213,7 @@ FSquare2DGridHelper GetPartitionedActors(const UWorldPartition* WorldPartition,
 			else
 			{
 				// Find grid level cell that encompasses the actor cluster bounding box and put actors in it.
-				const FVector2D ClusterSize = ClusterBounds.GetSize();
+				const FVector2D ClusterSize = ActorSetInstanceBounds.GetSize();
 				const double MinRequiredCellExtent = FMath::Max(ClusterSize.X, ClusterSize.Y);
 				const int32 FirstPotentialGridLevel = FMath::Max(FMath::CeilToDouble(FMath::Log2(MinRequiredCellExtent / (double)PartitionedActors.CellSize)), 0);
 
@@ -246,9 +221,9 @@ FSquare2DGridHelper GetPartitionedActors(const UWorldPartition* WorldPartition,
 				{
 					FSquare2DGridHelper::FGridLevel& GridLevel = PartitionedActors.Levels[GridLevelIndex];
 
-					if (GridLevel.GetNumIntersectingCells(ClusterInstance->Bounds) == 1)
+					if (GridLevel.GetNumIntersectingCells(ActorSetInstance->Bounds) == 1)
 					{
-						GridLevel.ForEachIntersectingCells(ClusterInstance->Bounds, [&GridLevel, ActorCluster, ClusterInstance, &GridCell](const FGridCellCoord2& Coords)
+						GridLevel.ForEachIntersectingCells(ActorSetInstance->Bounds, [&GridLevel, &GridCell](const FGridCellCoord2& Coords)
 						{
 							check(!GridCell);
 							GridCell = &GridLevel.GetCell(Coords);
@@ -259,61 +234,15 @@ FSquare2DGridHelper GetPartitionedActors(const UWorldPartition* WorldPartition,
 				}
 			}
 		}
-		else
-		{
-			GridCell = &PartitionedActors.GetAlwaysLoadedCell();
-		}
-
+		
 		if (!GridCell)
 		{
-			UE_LOG(LogWorldPartition, Warning, TEXT("Cluster of %d actors was promoted always loaded, BV of [%lld x %lld] (meters)"), 
-				ActorCluster->Actors.Num(),
-				(int64)(0.01f * ClusterInstance->Bounds.GetSize().X),
-				(int64)(0.01f * ClusterInstance->Bounds.GetSize().Y));
-			for (const FGuid& ActorGuid : ActorCluster->Actors)
-			{
-				const FWorldPartitionActorDescView& ActorDescView = ClusterInstance->ContainerInstance->GetActorDescView(ActorGuid);
-				UE_LOG(LogWorldPartition, Warning, TEXT("   - Actor: %s (%s)"), *ActorDescView.GetActorPath().ToString(), *ActorGuid.ToString(EGuidFormats::UniqueObjectGuid));
-				UE_LOG(LogWorldPartition, Warning, TEXT("         Package: %s"), *ActorDescView.GetActorPackage().ToString());
-				UE_LOG(LogWorldPartition, Warning, TEXT("         Container (%s): %s"), *ClusterInstance->ContainerInstance->ID.ToString(), *ClusterInstance->ContainerInstance->Container->GetContainerPackage().ToString())
-			}
-
 			GridCell = &PartitionedActors.GetAlwaysLoadedCell();
 		}
 
-		check(GridCell);
-		GridCell->AddActors(ActorCluster->Actors, ClusterInstance->ContainerInstance, ClusterInstance->DataLayers);
-
-		if (UE_LOG_ACTIVE(LogWorldPartition, Verbose))
-		{
-			if (ActorCluster->Actors.Num() > 1)
-			{
-				UE_LOG(LogWorldPartition, Verbose, TEXT("Clustered %d actors, generated shared BV of [%lld x %lld] (meters)"),
-					ActorCluster->Actors.Num(),
-					(int64)(0.01f * ClusterInstance->Bounds.GetSize().X),
-					(int64)(0.01f * ClusterInstance->Bounds.GetSize().Y));
-
-				if (ClusterInstance->DataLayers.Num())
-				{
-					TArray<FString> DataLayers;
-					Algo::Transform(ClusterInstance->DataLayers, DataLayers, [](const UDataLayerInstance* DataLayer) { return DataLayer->GetDataLayerFullName(); });
-					UE_LOG(LogWorldPartition, Verbose, TEXT("   - DataLayers: %s"), *FString::Join(DataLayers, TEXT(", ")));
-				}
-
-				for (const FGuid& ActorGuid : ActorCluster->Actors)
-				{
-					const FWorldPartitionActorDescView& ActorDescView = ClusterInstance->ContainerInstance->GetActorDescView(ActorGuid);
-					UE_LOG(LogWorldPartition, Verbose, TEXT("   - Actor: %s (%s)"), *ActorDescView.GetActorPath().ToString(), *ActorGuid.ToString(EGuidFormats::UniqueObjectGuid));
-					UE_LOG(LogWorldPartition, Verbose, TEXT("         Package: %s"), *ActorDescView.GetActorPackage().ToString());
-					UE_LOG(LogWorldPartition, Verbose, TEXT("         Container (%s): %s"), *ClusterInstance->ContainerInstance->ID.ToString(), *ClusterInstance->ContainerInstance->Container->GetContainerPackage().ToString())
-				}
-			}
-		}
+		GridCell->AddActorSetInstance(ActorSetInstance);
 	}
 
-	// Perform validation
-	PartitionedActors.ValidateSingleActorReferer();
-
 	return PartitionedActors;
 }
 
diff --git a/Engine/Source/Runtime/Engine/Private/WorldPartition/RuntimeSpatialHash/RuntimeSpatialHashGridHelper.h b/Engine/Source/Runtime/Engine/Private/WorldPartition/RuntimeSpatialHash/RuntimeSpatialHashGridHelper.h
--- a/Engine/Source/Runtime/Engine/Private/WorldPartition/RuntimeSpatialHash/RuntimeSpatialHashGridHelper.h
+++ b/Engine/Source/Runtime/Engine/Private/WorldPartition/RuntimeSpatialHash/RuntimeSpatialHashGridHelper.h
@@ -7,7 +7,7 @@
 #include "ProfilingDebugging/ScopedTimers.h"
 
 #include "WorldPartition/WorldPartitionRuntimeSpatialHash.h"
-#include "WorldPartition/WorldPartitionActorCluster.h"
+#include "WorldPartition/WorldPartitionStreamingGenerationContext.h"
 #include "WorldPartition/DataLayer/DataLayersID.h"
 #include "WorldPartition/DataLayer/DataLayerInstance.h"
 
@@ -267,14 +267,16 @@ struct FSquare2DGridHelper
 				DataLayersID = FDataLayersID(DataLayers);
 			}
 
-			void AddActor(FActorInstance ActorInstance) { Actors.Add(MoveTemp(ActorInstance)); }
-			const TSet<FActorInstance>& GetActors() const { return Actors; }
+			void AddActorSetInstance(const IStreamingGenerationContext::FActorSetInstance* ActorSetInstance) { ActorSetInstances.Add(ActorSetInstance); }
+			const TArray<const IStreamingGenerationContext::FActorSetInstance*>& GetActorSetInstances() const { return ActorSetInstances; }
 			bool HasDataLayers() const { return !DataLayers.IsEmpty(); }
 			const TArray<const UDataLayerInstance*>& GetDataLayers() const { return DataLayers; }
 			const FDataLayersID& GetDataLayersID() const { return DataLayersID; }
+			bool operator==(const FGridCellDataChunk& InGridCellDataChunk) const { return DataLayersID == InGridCellDataChunk.DataLayersID;}
+			friend uint32 GetTypeHash(const FGridCellDataChunk& InGridCellDataChunk) { return GetTypeHash(InGridCellDataChunk.DataLayersID);}
 
 		private:
-			TSet<FActorInstance> Actors;
+			TArray<const IStreamingGenerationContext::FActorSetInstance*> ActorSetInstances;
 			TArray<const UDataLayerInstance*> DataLayers;
 			FDataLayersID DataLayersID;
 		};
@@ -285,23 +287,17 @@ struct FSquare2DGridHelper
 				: Coords(InCoords)
 			{}
 
-			void AddActors(const TSet<FGuid>& InActors, const FActorContainerInstance* ContainerInstance, const TArray<const UDataLayerInstance*>& InDataLayers)
+			void AddActorSetInstance(const IStreamingGenerationContext::FActorSetInstance* ActorSetInstance)
 			{
-				for (const FGuid& Actor : InActors)
-				{
-					FActorInstance ActorInstance(Actor, ContainerInstance);
-
-					FDataLayersID DataLayersID = FDataLayersID(InDataLayers);
-					FGridCellDataChunk* ActorDataChunk = Algo::FindByPredicate(DataChunks, [&](FGridCellDataChunk& InDataChunk) { return InDataChunk.GetDataLayersID() == DataLayersID; });
-					if (!ActorDataChunk)
-					{
-						ActorDataChunk = &DataChunks.Emplace_GetRef(InDataLayers);
-					}
-					ActorDataChunk->AddActor(MoveTemp(ActorInstance));
-				}
+				const FDataLayersID DataLayersID = FDataLayersID(ActorSetInstance->DataLayers);
+				FGridCellDataChunk& ActorDataChunk = DataChunks.FindOrAddByHash(DataLayersID.GetHash(), FGridCellDataChunk(ActorSetInstance->DataLayers));
+				ActorDataChunk.AddActorSetInstance(ActorSetInstance);
 			}
 
-			const TArray<FGridCellDataChunk>& GetDataChunks() const { return DataChunks; }
+			const TSet<FGridCellDataChunk>& GetDataChunks() const
+			{
+				return DataChunks;
+			}
 
 			const FGridCellDataChunk* GetNoDataLayersDataChunk() const
 			{
@@ -322,7 +318,7 @@ struct FSquare2DGridHelper
 
 		private:
 			FGridCellCoord Coords;
-			TArray<FGridCellDataChunk> DataChunks;
+			TSet<FGridCellDataChunk> DataChunks;
 		};
 #endif
 
@@ -464,11 +460,6 @@ struct FSquare2DGridHelper
 	 */
 	int32 ForEachIntersectingCells(const FSphericalSector& InShape, TFunctionRef<void(const FGridCellCoord&)> InOperation, int32 InStartLevel = 0) const;
 
-#if WITH_EDITOR
-	// Validates that actor is not referenced by multiple cells
-	void ValidateSingleActorReferer();
-#endif
-
 public:
 	FBox WorldBounds;
 	FVector Origin;
@@ -478,5 +469,5 @@ public:
 
 #if WITH_EDITOR
 FSquare2DGridHelper GetGridHelper(const FBox& WorldBounds, int64 GridCellSize);
-FSquare2DGridHelper GetPartitionedActors(const UWorldPartition* WorldPartition, const FBox& WorldBounds, const FSpatialHashRuntimeGrid& Grid, const TArray<const FActorClusterInstance*>& GridActors);
+FSquare2DGridHelper GetPartitionedActors(const UWorldPartition* WorldPartition, const FBox& WorldBounds, const FSpatialHashRuntimeGrid& Grid, const TArray<const IStreamingGenerationContext::FActorSetInstance*>& ActorSetInstances);
 #endif // #if WITH_EDITOR
\ No newline at end of file
diff --git a/Engine/Source/Runtime/Engine/Private/WorldPartition/RuntimeSpatialHash/RuntimeSpatialHashHLOD.cpp b/Engine/Source/Runtime/Engine/Private/WorldPartition/RuntimeSpatialHash/RuntimeSpatialHashHLOD.cpp
--- a/Engine/Source/Runtime/Engine/Private/WorldPartition/RuntimeSpatialHash/RuntimeSpatialHashHLOD.cpp
+++ b/Engine/Source/Runtime/Engine/Private/WorldPartition/RuntimeSpatialHash/RuntimeSpatialHashHLOD.cpp
@@ -10,6 +10,7 @@
 #include "WorldPartition/HLOD/HLODActor.h"
 #include "WorldPartition/HLOD/HLODActorDesc.h"
 #include "WorldPartition/HLOD/IWorldPartitionHLODUtilitiesModule.h"
+#include "WorldPartition/DataLayer/DataLayerSubsystem.h"
 
 #include "WorldPartition/WorldPartitionHelpers.h"
 
@@ -113,18 +114,18 @@ static void GameTick(UWorld* InWorld)
 	}
 }
 
-static TArray<FGuid> GenerateHLODsForGrid(UWorldPartition* WorldPartition, const FActorClusterContext& ActorClusterContext, const FSpatialHashRuntimeGrid& RuntimeGrid, uint32 HLODLevel, FHLODCreationContext& Context, ISourceControlHelper* SourceControlHelper, bool bCreateActorsOnly, const TArray<const FActorClusterInstance*>& ClusterInstances)
+static TArray<FGuid> GenerateHLODsForGrid(UWorldPartition* WorldPartition, const IStreamingGenerationContext* StreamingGenerationContext, const FSpatialHashRuntimeGrid& RuntimeGrid, uint32 HLODLevel, FHLODCreationContext& Context, ISourceControlHelper* SourceControlHelper, bool bCreateActorsOnly, const TArray<const IStreamingGenerationContext::FActorSetInstance*>& ActorSetInstances)
 {
-	const FActorContainerInstance& MainContainerInstance = *ActorClusterContext.GetClusterInstance(WorldPartition);
-	const FBox WorldBounds = MainContainerInstance.Bounds;
+	const IStreamingGenerationContext::FActorSetContainer* MainActorSetContainer = StreamingGenerationContext->GetMainWorldContainer();
+	const FBox WorldBounds = StreamingGenerationContext->GetWorldBounds();
 
-	const FSquare2DGridHelper PartitionedActors = GetPartitionedActors(WorldPartition, WorldBounds, RuntimeGrid, ClusterInstances);
+	const FSquare2DGridHelper PartitionedActors = GetPartitionedActors(WorldPartition, WorldBounds, RuntimeGrid, ActorSetInstances);
 	const FSquare2DGridHelper::FGridLevel::FGridCell& AlwaysLoadedCell = PartitionedActors.GetAlwaysLoadedCell();
 
 	auto ShouldGenerateHLODs = [&AlwaysLoadedCell](const FSquare2DGridHelper::FGridLevel::FGridCell& GridCell, const FSquare2DGridHelper::FGridLevel::FGridCellDataChunk& GridCellDataChunk)
 	{
 		const bool bIsCellAlwaysLoaded = &GridCell == &AlwaysLoadedCell;
-		const bool bChunkHasActors = !GridCellDataChunk.GetActors().IsEmpty();
+		const bool bChunkHasActors = !GridCellDataChunk.GetActorSetInstances().IsEmpty();
 
 		const bool bShouldGenerateHLODs = !bIsCellAlwaysLoaded && bChunkHasActors;
 		return bShouldGenerateHLODs;
@@ -191,7 +192,16 @@ static TArray<FGuid> GenerateHLODsForGrid(UWorldPartition* WorldPartition, const
 
 				IWorldPartitionHLODUtilities* WPHLODUtilities = FModuleManager::Get().LoadModuleChecked<IWorldPartitionHLODUtilitiesModule>("WorldPartitionHLODUtilities").GetUtilities();
 
-				TArray<AWorldPartitionHLOD*> CellHLODActors = WPHLODUtilities->CreateHLODActors(Context, CreationParams, GridCellDataChunk.GetActors(), GridCellDataChunk.GetDataLayers());
+				TArray<IStreamingGenerationContext::FActorInstance> ActorInstances;
+				for (const IStreamingGenerationContext::FActorSetInstance* ActorSetInstance : GridCellDataChunk.GetActorSetInstances())
+				{
+					for (const FGuid& ActorGuid : ActorSetInstance->ActorSet->Actors)
+					{
+						ActorInstances.Emplace(ActorGuid, ActorSetInstance);
+					}
+				}
+
+				TArray<AWorldPartitionHLOD*> CellHLODActors = WPHLODUtilities->CreateHLODActors(Context, CreationParams, ActorInstances, GridCellDataChunk.GetDataLayers());
 				if (!CellHLODActors.IsEmpty())
 				{
 					TArray<AWorldPartitionHLOD*> NewCellHLODActors;
@@ -221,19 +231,6 @@ static TArray<FGuid> GenerateHLODsForGrid(UWorldPartition* WorldPartition, const
 						}
 					}
 
-					// Make sure we registered actor descriptors for newly created actors
-					for (AWorldPartitionHLOD* CellHLODActor : CellHLODActors)
-					{
-						check(WorldPartition->GetActorDesc(CellHLODActor->GetActorGuid()));
-
-						UE_LOG(LogWorldPartitionRuntimeSpatialHashHLOD, Verbose, TEXT("Created HLOD actor %s - %s, for cell %s, represented actors:"), *CellHLODActor->GetName(), *CellHLODActor->GetActorGuid().ToString(), *CellName.ToString());
-						for (const auto& SubActor : CellHLODActor->GetSubActors())
-						{
-							const FActorContainerInstance& ContainerInstance = *ActorClusterContext.GetClusterInstance(SubActor.ContainerID);
-							UE_LOG(LogWorldPartitionRuntimeSpatialHashHLOD, Verbose, TEXT("\t\t%s - %s - %s"), *SubActor.ContainerID.ToString(), *ContainerInstance.GetActorDescView(SubActor.ActorGuid).GetActorLabel().ToString(), *SubActor.ActorGuid.ToString());
-						}
-					}
-
 					// Update newly created HLOD actors
 					for (AWorldPartitionHLOD* CellHLODActor : NewCellHLODActors)
 					{
@@ -384,7 +381,7 @@ static void UpdateHLODGridsActors(UWorld* World, const TMap<FName, FSpatialHashR
 	}
 }
 
-bool UWorldPartitionRuntimeSpatialHash::GenerateHLOD(ISourceControlHelper* SourceControlHelper, FActorClusterContext& ActorClusterContext, bool bCreateActorsOnly)
+bool UWorldPartitionRuntimeSpatialHash::GenerateHLOD(ISourceControlHelper* SourceControlHelper, const IStreamingGenerationContext* StreamingGenerationContext, bool bCreateActorsOnly)
 {
 	TRACE_CPUPROFILER_EVENT_SCOPE(UWorldPartitionRuntimeSpatialHash::GenerateHLOD);
 
@@ -448,34 +445,33 @@ bool UWorldPartitionRuntimeSpatialHash::GenerateHLOD(ISourceControlHelper* Sourc
 		}
 	}
 
-	FActorContainerInstance& MainContainerInstance = *ActorClusterContext.GetClusterInstance(WorldPartition);
+	IStreamingGenerationContext::FActorSetContainer* MainActorSetContainer = const_cast<IStreamingGenerationContext::FActorSetContainer*>(StreamingGenerationContext->GetMainWorldContainer());
 
-	TArray<TArray<const FActorClusterInstance*>> GridsClusters;
-	GridsClusters.InsertDefaulted(0, Grids.Num());
+	TArray<TArray<const IStreamingGenerationContext::FActorSetInstance*>> GridActorSetInstances;
+	GridActorSetInstances.InsertDefaulted(0, Grids.Num());
 
-	for (const FActorClusterInstance& ClusterInstance : ActorClusterContext.GetClusterInstances())
+	StreamingGenerationContext->ForEachActorSetInstance([&GridsMapping, &GridActorSetInstances](const IStreamingGenerationContext::FActorSetInstance& ActorSetInstance)
 	{
-		const FActorCluster* ActorCluster = ClusterInstance.Cluster;
-		check(ActorCluster);
-		int32* FoundIndex = GridsMapping.Find(ActorCluster->RuntimeGrid);
+		int32* FoundIndex = GridsMapping.Find(ActorSetInstance.RuntimeGrid);
 		if (!FoundIndex)
 		{
-			UE_LOG(LogWorldPartitionRuntimeSpatialHashHLOD, Error, TEXT("Invalid partition grid '%s' referenced by actor cluster"), *ActorCluster->RuntimeGrid.ToString());
+			//@todo_ow should this be done upstream?
+			UE_LOG(LogWorldPartition, Error, TEXT("Invalid partition grid '%s' referenced by actor cluster"), *ActorSetInstance.RuntimeGrid.ToString());
 		}
 
 		int32 GridIndex = FoundIndex ? *FoundIndex : 0;
-		GridsClusters[GridIndex].Add(&ClusterInstance);
-	}
+		GridActorSetInstances[GridIndex].Add(&ActorSetInstance);
+	});
 
 	// Keep track of all valid HLOD actors, along with which runtime grid they live in
 	TMap<FName, TArray<FGuid>> GridsHLODActors;
 	TArray<TUniquePtr<FWorldPartitionActorDescView>> ActorDescViews;
 
-	auto GenerateHLODs = [&GridsHLODActors, &ActorDescViews, &MainContainerInstance, &ActorClusterContext, WorldPartition, &Context, SourceControlHelper, bCreateActorsOnly]
-	(const FSpatialHashRuntimeGrid& RuntimeGrid, uint32 HLODLevel, const TArray<const FActorClusterInstance*>& ActorClusterInstances)
+	auto GenerateHLODs = [&GridsHLODActors, &ActorDescViews, &MainActorSetContainer, StreamingGenerationContext, WorldPartition, &Context, SourceControlHelper, bCreateActorsOnly]
+	(const FSpatialHashRuntimeGrid& RuntimeGrid, uint32 HLODLevel, const TArray<const IStreamingGenerationContext::FActorSetInstance*>& ActorSetInstances)
 	{
 		// Generate HLODs for this grid
-		TArray<FGuid> HLODActors = GenerateHLODsForGrid(WorldPartition, ActorClusterContext, RuntimeGrid, HLODLevel, Context, SourceControlHelper, bCreateActorsOnly, ActorClusterInstances);
+		TArray<FGuid> HLODActors = GenerateHLODsForGrid(WorldPartition, StreamingGenerationContext, RuntimeGrid, HLODLevel, Context, SourceControlHelper, bCreateActorsOnly, ActorSetInstances);
 
 		for (const FGuid& HLODActorGuid : HLODActors)
 		{
@@ -483,9 +479,10 @@ bool UWorldPartitionRuntimeSpatialHash::GenerateHLOD(ISourceControlHelper* Sourc
 			check(HLODActorDesc);
 
 			FWorldPartitionActorDescView* ActorDescView = ActorDescViews.Emplace_GetRef(MakeUnique<FWorldPartitionActorDescView>(HLODActorDesc)).Get();
-			MainContainerInstance.ActorDescViewMap.Emplace(HLODActorGuid, ActorDescView);
+			FActorDescViewMap* NonConstActorDescViewMap = const_cast<FActorDescViewMap*>(MainActorSetContainer->ActorDescViewMap);			
+			NonConstActorDescViewMap->Emplace(HLODActorGuid, *ActorDescView);
 
-			const FWorldPartitionActorDescView& HLODActorDescView = MainContainerInstance.GetActorDescView(HLODActorGuid);
+			const FWorldPartitionActorDescView& HLODActorDescView = MainActorSetContainer->ActorDescViewMap->FindByGuidChecked(HLODActorGuid);
 			GridsHLODActors.FindOrAdd(HLODActorDescView.GetRuntimeGrid()).Add(HLODActorGuid);
 		}
 	};
@@ -494,7 +491,7 @@ bool UWorldPartitionRuntimeSpatialHash::GenerateHLOD(ISourceControlHelper* Sourc
 	for (int32 GridIndex = 0; GridIndex < Grids.Num(); GridIndex++)
 	{
 		// Generate HLODs for this grid - retrieve actors from the world partition
-		GenerateHLODs(Grids[GridIndex], 0, GridsClusters[GridIndex]);
+		GenerateHLODs(Grids[GridIndex], 0, GridActorSetInstances[GridIndex]);
 	}
 
 	// Now, go on and create HLOD actors from HLOD grids (HLOD 1-N)
@@ -510,24 +507,31 @@ bool UWorldPartitionRuntimeSpatialHash::GenerateHLOD(ISourceControlHelper* Sourc
 		}
 
 		// Create one actor cluster per HLOD actor
-		TArray<FActorCluster> HLODActorClusters;
-		TArray<FActorClusterInstance> HLODActorClusterInstances;
-		TArray<const FActorClusterInstance*> HLODActorClusterInstancePtrs;
-		HLODActorClusters.Reserve(GridsHLODActors[HLODGridName].Num());
-		HLODActorClusterInstances.Reserve(GridsHLODActors[HLODGridName].Num());
-		HLODActorClusterInstancePtrs.Reserve(GridsHLODActors[HLODGridName].Num());
+		TArray<IStreamingGenerationContext::FActorSetInstance> HLODActorSetInstances;
+		TArray<const IStreamingGenerationContext::FActorSetInstance*> HLODActorSetInstancePtrs;
+		HLODActorSetInstances.Reserve(GridsHLODActors[HLODGridName].Num());
+		HLODActorSetInstancePtrs.Reserve(GridsHLODActors[HLODGridName].Num());
 
 		for (const FGuid& HLODActorGuid : GridsHLODActors[HLODGridName])
 		{
-			const FWorldPartitionActorDescView& HLODActorDescView = MainContainerInstance.GetActorDescView(HLODActorGuid);
-			FActorCluster& NewHLODActorCluster = HLODActorClusters.Emplace_GetRef(WorldPartition->GetWorld(), HLODActorDescView);
-			FActorClusterInstance& NewHLODActorClusterInstance = HLODActorClusterInstances.Emplace_GetRef(&NewHLODActorCluster, &MainContainerInstance);
-			HLODActorClusterInstancePtrs.Add(&NewHLODActorClusterInstance);
+			const FWorldPartitionActorDescView& HLODActorDescView = MainActorSetContainer->ActorDescViewMap->FindByGuidChecked(HLODActorGuid);
+			
+			IStreamingGenerationContext::FActorSetInstance& NewHLODActorSetInstance = HLODActorSetInstances.Emplace_GetRef();
+			
+			NewHLODActorSetInstance.Bounds = HLODActorDescView.GetBounds();
+			NewHLODActorSetInstance.RuntimeGrid = HLODActorDescView.GetRuntimeGrid();
+			NewHLODActorSetInstance.bIsSpatiallyLoaded = HLODActorDescView.GetIsSpatiallyLoaded();
+			NewHLODActorSetInstance.DataLayers = UDataLayerSubsystem::GetRuntimeDataLayerInstances(MainActorSetContainer->ActorDescContainer->GetWorld(), HLODActorDescView.GetRuntimeDataLayers());
+			NewHLODActorSetInstance.ContainerInstance = MainActorSetContainer;
+			NewHLODActorSetInstance.ActorSet = &MainActorSetContainer->ActorSets.AddDefaulted_GetRef();
+			const_cast<IStreamingGenerationContext::FActorSet*>(NewHLODActorSetInstance.ActorSet)->Actors.Add(HLODActorDescView.GetGuid());
+
+			HLODActorSetInstancePtrs.Add(&NewHLODActorSetInstance);
 		};
 
 		// Generate HLODs for this grid - retrieve actors from our ValidHLODActors map
 		// We can't rely on actor descs for newly created HLOD actors
-		GenerateHLODs(HLODGrids[HLODGridName], HLODLayersLevels[HLODGrids[HLODGridName].HLODLayer] + 1, HLODActorClusterInstancePtrs);
+		GenerateHLODs(HLODGrids[HLODGridName], HLODLayersLevels[HLODGrids[HLODGridName].HLODLayer] + 1, HLODActorSetInstancePtrs);
 	}
 
 	auto DeleteHLODActor = [&SourceControlHelper, WorldPartition](FWorldPartitionHandle ActorHandle)
diff --git a/Engine/Source/Runtime/Engine/Private/WorldPartition/WorldPartitionActorCluster.cpp b/Engine/Source/Runtime/Engine/Private/WorldPartition/WorldPartitionActorCluster.cpp
deleted file mode 100644
diff --git a/Engine/Source/Runtime/Engine/Private/WorldPartition/WorldPartitionActorDescView.cpp b/Engine/Source/Runtime/Engine/Private/WorldPartition/WorldPartitionActorDescView.cpp
--- a/Engine/Source/Runtime/Engine/Private/WorldPartition/WorldPartitionActorDescView.cpp
+++ b/Engine/Source/Runtime/Engine/Private/WorldPartition/WorldPartitionActorDescView.cpp
@@ -192,4 +192,9 @@ bool FWorldPartitionActorDescView::IsResaveNeeded() const
 {
 	return ActorDesc->IsResaveNeeded();
 }
+
+AActor* FWorldPartitionActorDescView::GetActor() const
+{
+	return ActorDesc->GetActor();
+}
 #endif
\ No newline at end of file
diff --git a/Engine/Source/Runtime/Engine/Private/WorldPartition/WorldPartitionRuntimeHash.cpp b/Engine/Source/Runtime/Engine/Private/WorldPartition/WorldPartitionRuntimeHash.cpp
--- a/Engine/Source/Runtime/Engine/Private/WorldPartition/WorldPartitionRuntimeHash.cpp
+++ b/Engine/Source/Runtime/Engine/Private/WorldPartition/WorldPartitionRuntimeHash.cpp
@@ -10,6 +10,7 @@
 #if WITH_EDITOR
 #include "WorldPartition/WorldPartitionHandle.h"
 #include "WorldPartition/DataLayer/WorldDataLayers.h"
+#include "WorldPartition/WorldPartitionStreamingGenerationContext.h"
 #endif
 
 #define LOCTEXT_NAMESPACE "WorldPartition"
diff --git a/Engine/Source/Runtime/Engine/Private/WorldPartition/WorldPartitionRuntimeSpatialHash.cpp b/Engine/Source/Runtime/Engine/Private/WorldPartition/WorldPartitionRuntimeSpatialHash.cpp
--- a/Engine/Source/Runtime/Engine/Private/WorldPartition/WorldPartitionRuntimeSpatialHash.cpp
+++ b/Engine/Source/Runtime/Engine/Private/WorldPartition/WorldPartitionRuntimeSpatialHash.cpp
@@ -5,6 +5,7 @@
 =============================================================================*/
 #include "WorldPartition/WorldPartitionRuntimeSpatialHash.h"
 
+#include "WorldPartition/WorldPartitionStreamingGenerationContext.h"
 #include "WorldPartition/WorldPartitionRuntimeSpatialHashCell.h"
 #include "WorldPartition/WorldPartitionActorDesc.h"
 #include "WorldPartition/WorldPartitionActorDescView.h"
@@ -777,7 +778,7 @@ void UWorldPartitionRuntimeSpatialHash::SetDefaultValues()
 	MainGrid.DebugColor = FLinearColor::Gray;
 }
 
-bool UWorldPartitionRuntimeSpatialHash::GenerateStreaming(UWorldPartitionStreamingPolicy* StreamingPolicy, const FActorClusterContext& ActorClusterContext, TArray<FString>* OutPackagesToGenerate)
+bool UWorldPartitionRuntimeSpatialHash::GenerateStreaming(UWorldPartitionStreamingPolicy* StreamingPolicy, const IStreamingGenerationContext* StreamingGenerationContext, TArray<FString>* OutPackagesToGenerate)
 {
 	TRACE_CPUPROFILER_EVENT_SCOPE(UWorldPartitionRuntimeSpatialHash::GenerateStreaming);
 	UWorldPartition* WorldPartition = GetOuterUWorldPartition();
@@ -797,20 +798,12 @@ bool UWorldPartitionRuntimeSpatialHash::GenerateStreaming(UWorldPartitionStreami
 	TArray<FSpatialHashRuntimeGrid> AllGrids;
 	AllGrids.Append(Grids);
 
-	const FActorContainerInstance* ContainerInstance = ActorClusterContext.GetClusterInstance(WorldPartition);
-	check(ContainerInstance);
-
-	for (auto& ActorDescViewPair : ContainerInstance->ActorDescViewMap)
+	TArray<const FWorldPartitionActorDescView*> SpatialHashRuntimeGridInfos = StreamingGenerationContext->GetMainWorldContainer()->ActorDescViewMap->FindByExactNativeClass<ASpatialHashRuntimeGridInfo>();
+	for (const FWorldPartitionActorDescView* SpatialHashRuntimeGridInfo : SpatialHashRuntimeGridInfos)
 	{
-		const FWorldPartitionActorDescView& ActorDescView = *ActorDescViewPair.Value;
-		if (ActorDescView.GetActorNativeClass()->IsChildOf<ASpatialHashRuntimeGridInfo>())
-		{
-			FWorldPartitionReference Ref(WorldPartition, ActorDescView.GetGuid());
-			if (ASpatialHashRuntimeGridInfo* RuntimeGridActor = Cast<ASpatialHashRuntimeGridInfo>(Ref->GetActor()))
-			{
-				AllGrids.Add(RuntimeGridActor->GridSettings);
-			}
-		}
+		FWorldPartitionReference Ref(WorldPartition, SpatialHashRuntimeGridInfo->GetGuid());
+		ASpatialHashRuntimeGridInfo* RuntimeGridActor = CastChecked<ASpatialHashRuntimeGridInfo>(SpatialHashRuntimeGridInfo->GetActor());
+		AllGrids.Add(RuntimeGridActor->GridSettings);
 	}
 
 	TMap<FName, int32> GridsMapping;
@@ -822,27 +815,27 @@ bool UWorldPartitionRuntimeSpatialHash::GenerateStreaming(UWorldPartitionStreami
 		GridsMapping.Add(Grid.GridName, i);
 	}
 
-	TArray<TArray<const FActorClusterInstance*>> GridActors;
-	GridActors.InsertDefaulted(0, AllGrids.Num());
+	TArray<TArray<const IStreamingGenerationContext::FActorSetInstance*>> GridActorSetInstances;
+	GridActorSetInstances.InsertDefaulted(0, AllGrids.Num());
 
-	for (const FActorClusterInstance& ClusterInstance : ActorClusterContext.GetClusterInstances())
+	StreamingGenerationContext->ForEachActorSetInstance([&GridsMapping, &GridActorSetInstances](const IStreamingGenerationContext::FActorSetInstance& ActorSetInstance)
 	{
-		check(ClusterInstance.Cluster);
-		int32* FoundIndex = GridsMapping.Find(ClusterInstance.Cluster->RuntimeGrid);
+		int32* FoundIndex = GridsMapping.Find(ActorSetInstance.RuntimeGrid);
 		if (!FoundIndex)
 		{
-			UE_LOG(LogWorldPartition, Error, TEXT("Invalid partition grid '%s' referenced by actor cluster"), *ClusterInstance.Cluster->RuntimeGrid.ToString());
+			//@todo_ow should this be done upstream?
+			UE_LOG(LogWorldPartition, Error, TEXT("Invalid partition grid '%s' referenced by actor cluster"), *ActorSetInstance.RuntimeGrid.ToString());
 		}
 
 		int32 GridIndex = FoundIndex ? *FoundIndex : 0;
-		GridActors[GridIndex].Add(&ClusterInstance);
-	}
+		GridActorSetInstances[GridIndex].Add(&ActorSetInstance);
+	});
 	
-	const FBox WorldBounds = ActorClusterContext.GetClusterInstance(WorldPartition)->Bounds;
+	const FBox WorldBounds = StreamingGenerationContext->GetWorldBounds();
 	for (int32 GridIndex=0; GridIndex < AllGrids.Num(); GridIndex++)
 	{
 		const FSpatialHashRuntimeGrid& Grid = AllGrids[GridIndex];
-		const FSquare2DGridHelper PartionedActors = GetPartitionedActors(WorldPartition, WorldBounds, Grid, GridActors[GridIndex]);
+		const FSquare2DGridHelper PartionedActors = GetPartitionedActors(WorldPartition, WorldBounds, Grid, GridActorSetInstances[GridIndex]);
 		if (!CreateStreamingGrid(Grid, PartionedActors, StreamingPolicy, OutPackagesToGenerate))
 		{
 			return false;
@@ -980,7 +973,6 @@ bool UWorldPartitionRuntimeSpatialHash::CreateStreamingGrid(const FSpatialHashRu
 	// Move actors into the final streaming grids
 	CurrentStreamingGrid.GridLevels.Reserve(PartionedActors.Levels.Num());
 
-	TArray<FActorInstance> FilteredActors;
 	int32 Level = INDEX_NONE;
 	for (const FSquare2DGridHelper::FGridLevel& TempLevel : PartionedActors.Levels)
 	{
@@ -1001,16 +993,16 @@ bool UWorldPartitionRuntimeSpatialHash::CreateStreamingGrid(const FSpatialHashRu
 				// Cell cannot be treated as always loaded if it has data layers
 				const bool bIsCellAlwaysLoaded = (&TempCell == &PartionedActors.GetAlwaysLoadedCell()) && !GridCellDataChunk.HasDataLayers();
 				
-				FilteredActors.SetNum(0, false);
-				FilteredActors.Reset(GridCellDataChunk.GetActors().Num());
-				if (GridCellDataChunk.GetActors().Num())
+				TArray<IStreamingGenerationContext::FActorInstance> FilteredActors;
+				FWorldPartitionLoadingContext::FDeferred LoadingContext;
+
+				for (const IStreamingGenerationContext::FActorSetInstance* ActorSetInstance : GridCellDataChunk.GetActorSetInstances())
 				{
-					FWorldPartitionLoadingContext::FDeferred LoadingContext;
-					for (const FActorInstance& ActorInstance : GridCellDataChunk.GetActors())
+					for (const FGuid& ActorGuid : ActorSetInstance->ActorSet->Actors)
 					{
 						if (bIsMainWorldPartition && !IsRunningCookCommandlet())
 						{
-							const FWorldPartitionActorDescView& ActorDescView = ActorInstance.GetActorDescView();
+							const FWorldPartitionActorDescView& ActorDescView = ActorSetInstance->ContainerInstance->ActorDescViewMap->FindByGuidChecked(ActorGuid);
 
 							// In PIE, Always loaded cell is not generated. Instead, always loaded actors will be added to AlwaysLoadedActorsForPIE.
 							// This will trigger loading/registration of these actors in the PersistentLevel (if not already loaded).
@@ -1019,10 +1011,11 @@ bool UWorldPartitionRuntimeSpatialHash::CreateStreamingGrid(const FSpatialHashRu
 							// will unload actors that were not already loaded in the non PIE world.
 							if (bIsCellAlwaysLoaded)
 							{
-								if (ActorInstance.ContainerInstance->Container == WorldPartition)
+								if (ActorSetInstance->ContainerID.IsMainContainer())
 								{
 									// This will load the actor if it isn't already loaded
-									FWorldPartitionReference Reference(WorldPartition, ActorInstance.Actor);									
+									FWorldPartitionReference Reference(WorldPartition, ActorGuid);				
+
 									if (AActor* AlwaysLoadedActor = FindObject<AActor>(nullptr, *ActorDescView.GetActorPath().ToString()))
 									{
 										AlwaysLoadedActorsForPIE.Emplace(Reference, AlwaysLoadedActor);
@@ -1041,8 +1034,8 @@ bool UWorldPartitionRuntimeSpatialHash::CreateStreamingGrid(const FSpatialHashRu
 							}
 						}
 
-						FilteredActors.Add(ActorInstance);
-					}
+						FilteredActors.Emplace(ActorGuid, ActorSetInstance);
+					}					
 				}
 
 				if (!FilteredActors.Num())
@@ -1084,10 +1077,8 @@ bool UWorldPartitionRuntimeSpatialHash::CreateStreamingGrid(const FSpatialHashRu
 				StreamingCell->SetBlockOnSlowLoading(CurrentStreamingGrid.bBlockOnSlowStreaming);
 				StreamingCell->SetIsHLOD(RuntimeGrid.HLODLayer ? true : false);
 
-				UE_LOG(LogWorldPartition, Verbose, TEXT("Cell%s %s Actors = %d Bounds (%s)"), bIsCellAlwaysLoaded ? TEXT(" (AlwaysLoaded)") : TEXT(""), *StreamingCell->GetName(), FilteredActors.Num(), *Bounds.ToString());
-
 				check(!StreamingCell->UnsavedActorsContainer);
-				for (const FActorInstance& ActorInstance : FilteredActors)
+				for (const IStreamingGenerationContext::FActorInstance& ActorInstance : FilteredActors)
 				{
 					const FWorldPartitionActorDescView& ActorDescView = ActorInstance.GetActorDescView();
 					if (AActor* Actor = FindObject<AActor>(nullptr, *ActorDescView.GetActorPath().ToString()))
@@ -1102,15 +1093,15 @@ bool UWorldPartitionRuntimeSpatialHash::CreateStreamingGrid(const FSpatialHashRu
 				}
 
 				FVector2D CellMinMaxZ(UE_BIG_NUMBER, -UE_BIG_NUMBER);
-				for (const FActorInstance& ActorInstance : FilteredActors)
+				for (const IStreamingGenerationContext::FActorInstance& ActorInstance : FilteredActors)
 				{
 					const FWorldPartitionActorDescView& ActorDescView = ActorInstance.GetActorDescView();
-					StreamingCell->AddActorToCell(ActorDescView, ActorInstance.ContainerInstance->ID, ActorInstance.ContainerInstance->Transform, ActorInstance.ContainerInstance->Container);
+					StreamingCell->AddActorToCell(ActorDescView, ActorInstance.GetContainerID(), ActorInstance.GetTransform(), ActorInstance.GetActorDescContainer());
 					
 					CellMinMaxZ.X = FMath::Min(CellMinMaxZ.X, ActorDescView.GetBounds().Min.Z);
 					CellMinMaxZ.Y = FMath::Max(CellMinMaxZ.Y, ActorDescView.GetBounds().Max.Z);
 
-					if (ActorInstance.ContainerInstance->ID.IsMainContainer() && StreamingCell->UnsavedActorsContainer)
+					if (ActorInstance.GetContainerID().IsMainContainer() && StreamingCell->UnsavedActorsContainer)
 					{
 						if (AActor* Actor = FindObject<AActor>(nullptr, *ActorDescView.GetActorPath().ToString()))
 						{
@@ -1126,7 +1117,6 @@ bool UWorldPartitionRuntimeSpatialHash::CreateStreamingGrid(const FSpatialHashRu
 							});
 						}
 					}
-					UE_LOG(LogWorldPartition, Verbose, TEXT("  Actor : %s (%s) (Container %s)"), *(ActorDescView.GetActorPath().ToString()), *ActorDescView.GetGuid().ToString(EGuidFormats::UniqueObjectGuid), *ActorInstance.ContainerInstance->ID.ToString());
 				}
 				StreamingCell->SetMinMaxZ(CellMinMaxZ);
 
diff --git a/Engine/Source/Runtime/Engine/Private/WorldPartition/WorldPartitionStreamingGeneration.cpp b/Engine/Source/Runtime/Engine/Private/WorldPartition/WorldPartitionStreamingGeneration.cpp
--- a/Engine/Source/Runtime/Engine/Private/WorldPartition/WorldPartitionStreamingGeneration.cpp
+++ b/Engine/Source/Runtime/Engine/Private/WorldPartition/WorldPartitionStreamingGeneration.cpp
@@ -2,15 +2,17 @@
 
 #if WITH_EDITOR
 #include "WorldPartition/WorldPartitionStreamingGeneration.h"
+#include "WorldPartition/WorldPartitionStreamingGenerationContext.h"
 
 #include "Editor.h"
 #include "Engine/World.h"
 #include "Engine/LevelScriptBlueprint.h"
 #include "ActorReferencesUtils.h"
+#include "ReferenceCluster.h"
 #include "WorldPartition/WorldPartition.h"
 #include "WorldPartition/WorldPartitionRuntimeHash.h"
 #include "WorldPartition/WorldPartitionStreamingPolicy.h"
-#include "WorldPartition/WorldPartitionActorCluster.h"
+#include "WorldPartition/WorldPartitionActorContainerID.h"
 #include "WorldPartition/DataLayer/DataLayerSubsystem.h"
 #include "WorldPartition/DataLayer/DataLayerUtils.h"
 #include "WorldPartition/DataLayer/WorldDataLayers.h"
@@ -63,24 +65,149 @@ FWorldPartitionActorDescView* FActorDescViewMap::Emplace(const FGuid& InGuid, co
 	return NewActorDescView;
 }
 
-/*
-	Preparation Phase
-		Actor Descriptor Views Creation
-		Actor Descriptor Views Validation
-		Actor Clusters Creation
+class FWorldPartitionStreamingGenerator
+{
+	class FStreamingGenerationContext : public IStreamingGenerationContext
+	{
+	public:
+		FStreamingGenerationContext(const FWorldPartitionStreamingGenerator* StreamingGenerator, const UActorDescContainer* MainWorldContainer)
+		{
+			// Create the dataset required for IStreamingGenerationContext interface
+			MainWorldActorSetContainerIndex = INDEX_NONE;
+			ActorSetContainers.Empty(StreamingGenerator->ContainerDescriptorsMap.Num());
 
-	Generation Phase
-		Streaming Grids Generation
-		Data Layers Split Pass
+			TMap<const UActorDescContainer*, int32> ActorSetContainerMap;
+			for (const auto& [ActorDescContainer, ContainerDescriptor] : StreamingGenerator->ContainerDescriptorsMap)
+			{
+				const int32 ContainerIndex = ActorSetContainers.AddDefaulted();
+				ActorSetContainerMap.Add(ContainerDescriptor.Container, ContainerIndex);
 
-	Output Phase
-		Report Generation
-		SubLevels Generation
-		HLOD Generation
-*/
+				FActorSetContainer& ActorSetContainer = ActorSetContainers[ContainerIndex];
+				ActorSetContainer.ActorDescViewMap = &ContainerDescriptor.ActorDescViewMap;
+				ActorSetContainer.ActorDescContainer = ContainerDescriptor.Container;
+
+				ActorSetContainer.ActorSets.Empty(ContainerDescriptor.Clusters.Num());
+				for (const TArray<FGuid>& Cluster : ContainerDescriptor.Clusters)
+				{
+					FActorSet& ActorSet = ActorSetContainer.ActorSets.Emplace_GetRef();
+					ActorSet.Actors = Cluster;
+				}
+
+				if (ContainerDescriptor.Container == MainWorldContainer)
+				{
+					check(MainWorldActorSetContainerIndex == INDEX_NONE);
+					MainWorldActorSetContainerIndex = ContainerIndex;
+				}
+			}
+
+			ActorSetInstances.Empty();
+			for (const auto& [ContainerID, ContainerInstanceDescriptor] : StreamingGenerator->ContainerInstanceDescriptorsMap)
+			{
+				const FContainerDescriptor& ContainerDescriptor = StreamingGenerator->ContainerDescriptorsMap.FindChecked(ContainerInstanceDescriptor.Container);
+				const FActorSetContainer& ActorSetContainer = ActorSetContainers[ActorSetContainerMap.FindChecked(ContainerInstanceDescriptor.Container)];
+
+				for (const FActorSet& ActorSet : ActorSetContainer.ActorSets)
+				{
+					const FWorldPartitionActorDescView& ReferenceActorDescView = ActorSetContainer.ActorDescViewMap->FindByGuidChecked(ActorSet.Actors[0]);
+
+					FActorSetInstance& ActorSetInstance = ActorSetInstances.Emplace_GetRef();
+				
+					ActorSetInstance.ContainerInstance = &ActorSetContainer;
+					ActorSetInstance.ActorSet = &ActorSet;
+					ActorSetInstance.ContainerID = ContainerInstanceDescriptor.ID;
+					ActorSetInstance.Transform = ContainerInstanceDescriptor.Transform;
+
+					if (ContainerInstanceDescriptor.ID.IsMainContainer())
+					{
+						// Main container will get inherited properties from the actor descriptors
+						ActorSetInstance.RuntimeGrid = ReferenceActorDescView.GetRuntimeGrid();
+						ActorSetInstance.bIsSpatiallyLoaded = ReferenceActorDescView.GetIsSpatiallyLoaded();
+						ActorSetInstance.DataLayers = UDataLayerSubsystem::GetRuntimeDataLayerInstances(ContainerDescriptor.Container->GetWorld(), ReferenceActorDescView.GetRuntimeDataLayers());
+					}
+					else
+					{
+						// Sub containers will inherit some properties from the parent actor descriptors
+						ActorSetInstance.RuntimeGrid = ContainerInstanceDescriptor.RuntimeGrid;
+						ActorSetInstance.bIsSpatiallyLoaded = ContainerInstanceDescriptor.bIsSpatiallyLoaded;
+
+						TSet<FName> CombinedDataLayers = ContainerInstanceDescriptor.RuntimeDataLayers;
+						CombinedDataLayers.Append(ReferenceActorDescView.GetRuntimeDataLayers());
+
+						ActorSetInstance.DataLayers = UDataLayerSubsystem::GetRuntimeDataLayerInstances(ContainerDescriptor.Container->GetWorld(), CombinedDataLayers.Array());
+					}
+
+					ActorSetInstance.Bounds.Init();
+					for (const FGuid& ActorGuid : ActorSetInstance.ActorSet->Actors)
+					{
+						const FWorldPartitionActorDescView& ActorDescView = ActorSetContainer.ActorDescViewMap->FindByGuidChecked(ActorGuid);
+						ActorSetInstance.Bounds += ActorDescView.GetBounds().TransformBy(ContainerInstanceDescriptor.Transform);
+					}
+				}
+			}
+
+			WorldBounds = StreamingGenerator->ContainerInstanceDescriptorsMap.FindChecked(FActorContainerID::GetMainContainerID()).Bounds;
+		}
+
+		virtual ~FStreamingGenerationContext()
+		{}
+
+		//~Begin IStreamingGenerationContext interface
+		virtual FBox GetWorldBounds() const override
+		{
+			return WorldBounds;
+		}
+
+		virtual const FActorSetContainer* GetMainWorldContainer() const override
+		{
+			return ActorSetContainers.IsValidIndex(MainWorldActorSetContainerIndex) ? &ActorSetContainers[MainWorldActorSetContainerIndex] : nullptr;
+		}
+
+		virtual void ForEachActorSetInstance(TFunctionRef<void(const FActorSetInstance&)> Func) const
+		{
+			for (const FActorSetInstance& ActorSetInstance : ActorSetInstances)
+			{
+				Func(ActorSetInstance);
+			}
+		}
+		//~End IStreamingGenerationContext interface};
+
+	private:
+		FBox WorldBounds;
+		int32 MainWorldActorSetContainerIndex;
+		TArray<FActorSetContainer> ActorSetContainers;
+		TArray<FActorSetInstance> ActorSetInstances;
+	};
+
+	struct FContainerDescriptor
+	{
+		FContainerDescriptor()
+			: Container(nullptr)
+		{}
+
+		const UActorDescContainer* Container;
+		FActorDescViewMap ActorDescViewMap;
+		TArray<FWorldPartitionActorDescView> ContainerInstanceViews;
+		TArray<TArray<FGuid>> Clusters;
+	};
+
+	struct FContainerInstanceDescriptor
+	{
+		FContainerInstanceDescriptor()
+			: Bounds(ForceInit)
+			, Container(nullptr)
+		{}
+
+		FBox Bounds;
+		FTransform Transform;
+		const UActorDescContainer* Container;
+		EContainerClusterMode ClusterMode;
+		TSet<FName> RuntimeDataLayers;
+		FName RuntimeGrid;
+		bool bIsSpatiallyLoaded;
+		FString OwnerName;
+		FActorContainerID ID;
+	};
 
-class FWorldPartitionStreamingGenerator
-{
 	void ResolveRuntimeDataLayers(FWorldPartitionActorDescView& ActorDescView, const FActorDescViewMap& ActorDescViewMap)
 	{
 		TArray<FName> RuntimeDataLayerInstanceNames;
@@ -114,7 +241,7 @@ class FWorldPartitionStreamingGenerator
 
 	void CreateActorDescViewMap(const UActorDescContainer* InContainer, FActorDescViewMap& OutActorDescViewMap, const FActorContainerID& InContainerID, TArray<FWorldPartitionActorDescView>& OutContainerInstances)
 	{
-		// should we handle unsaved or newly created actors?
+		// Should we handle unsaved or newly created actors?
 		const bool bHandleUnsavedActors = ModifiedActorsDescList && InContainerID.IsMainContainer();
 
 		// Consider all actors of a /Temp/ container package as Unsaved because loading them from disk will fail (Outer world name mismatch)
@@ -137,6 +264,18 @@ class FWorldPartitionStreamingGenerator
 			return true;
 		};
 
+		auto IsFilteredActorClass = [this](const FWorldPartitionActorDesc* ActorDesc)
+		{
+			for (UClass* FilteredClass : FilteredClasses)
+			{
+				if (ActorDesc->GetActorNativeClass()->IsChildOf(FilteredClass))
+				{
+					return true;
+				}
+			}
+			return false;
+		};
+
 		// Create an actor descriptor view for the specified actor (modified or unsaved actors)
 		auto GetModifiedActorDesc = [this](AActor* InActor, const UActorDescContainer* InContainer) -> FWorldPartitionActorDesc*
 		{
@@ -150,22 +289,22 @@ class FWorldPartitionStreamingGenerator
 		};
 
 		// Register the actor descriptor view
-		auto RegisterActorDescView = [this, InContainer, &OutActorDescViewMap, &OutContainerInstances](const FGuid& ActorGuid, FWorldPartitionActorDescView& ActorDescView)
+		auto RegisterActorDescView = [this, InContainer, &OutActorDescViewMap, &OutContainerInstances](const FGuid& ActorGuid, FWorldPartitionActorDescView& InActorDescView)
 		{
-			if (ActorDescView.IsContainerInstance())
+			if (InActorDescView.IsContainerInstance())
 			{
-				OutContainerInstances.Add(ActorDescView);
+				OutContainerInstances.Add(InActorDescView);
 			}
 			else
 			{
-				OutActorDescViewMap.Emplace(ActorGuid, ActorDescView);
+				OutActorDescViewMap.Emplace(ActorGuid, InActorDescView);
 			}
 		};
 		
 		TMap<FGuid, FGuid> ContainerGuidsRemap;
 		for (FActorDescList::TConstIterator<> ActorDescIt(InContainer); ActorDescIt; ++ActorDescIt)
 		{
-			if (!IsActorEditorOnly(*ActorDescIt, InContainerID))
+			if (!IsActorEditorOnly(*ActorDescIt, InContainerID) && !IsFilteredActorClass(*ActorDescIt))
 			{
 				// Handle unsaved actors
 				if (AActor* Actor = ActorDescIt->GetActor())
@@ -206,15 +345,44 @@ class FWorldPartitionStreamingGenerator
 		}
 	}
 
-	void CreateActorDescriptorViewsRecursive(const UActorDescContainer* InContainer, const FTransform& InTransform, const TSet<FName>& InRuntimeDataLayers, const FActorContainerID& InContainerID, const FActorContainerID& InParentContainerID, EContainerClusterMode InClusterMode, const TCHAR* OwnerName)
+	void CreateActorDescriptorViewsRecursive(const UActorDescContainer* InContainer, const FTransform& InTransform, FName InRuntimeGrid, bool bInIsSpatiallyLoaded, const TSet<FName>& InRuntimeDataLayers, const FActorContainerID& InContainerID, const FActorContainerID& InParentContainerID, EContainerClusterMode InClusterMode, const TCHAR* InOwnerName)
 	{
-		FActorDescViewMap ActorDescViewMap;
-		TArray<FWorldPartitionActorDescView> ContainerInstanceViews;
+		// Create or resolve container descriptor
+		FContainerDescriptor& ContainerDescriptor = ContainerDescriptorsMap.FindOrAdd(InContainer);
+
+		if (!ContainerDescriptor.Container)
+		{
+			ContainerDescriptor.Container = InContainer;
 		
-		// Gather actor descriptor views for this container
-		CreateActorDescViewMap(InContainer, ActorDescViewMap, InContainerID, ContainerInstanceViews);
+			// Gather actor descriptor views for this container
+			CreateActorDescViewMap(InContainer, ContainerDescriptor.ActorDescViewMap, InContainerID, ContainerDescriptor.ContainerInstanceViews);
+
+			// Resolve actor descriptor views before validation
+			ResolveContainerDescriptor(ContainerDescriptor);
+
+			// Validate container, fixing anything illegal, etc.
+			ValidateContainerDescriptor(ContainerDescriptor, InContainerID.IsMainContainer());
+
+			// Update container, computing cluster, bounds, etc.
+			UpdateContainerDescriptor(ContainerDescriptor);
+		}
+
+		// Inherited parent properties logic
+		auto InheritParentContainerProperties = [](const FActorContainerID& InParentContainerID, const FWorldPartitionActorDescView& InParentActorDescView, TSet<FName>& InOutRuntimeDataLayers, FName& InOutRuntimeGrid, bool& bInOutIsSpatiallyLoaded)
+		{
+			// Runtime grid and spatially loaded flag are only inherited from the main world, since level instance don't support setting these values on actors
+			if (InParentContainerID.IsMainContainer())
+			{
+				InOutRuntimeGrid = InParentActorDescView.GetRuntimeGrid();
+				bInOutIsSpatiallyLoaded = InParentActorDescView.GetIsSpatiallyLoaded();
+			}
+
+			// Data layers are accumulated down the hierarchy chain, since level instances supports data layers assignation on actors
+			InOutRuntimeDataLayers.Append(InParentActorDescView.GetRuntimeDataLayers());
+		};
 
 		// Parse actor containers
+		TArray<FWorldPartitionActorDescView> ContainerInstanceViews = ContainerDescriptor.ContainerInstanceViews;
 		for (const FWorldPartitionActorDescView& ContainerInstanceView : ContainerInstanceViews)
 		{
 			const UActorDescContainer* SubContainer;
@@ -233,23 +401,27 @@ class FWorldPartitionStreamingGenerator
 			const FGuid ActorGuid = ContainerInstanceView.GetGuid();
 			const FActorContainerID SubContainerID(InContainerID, ActorGuid);
 
-			// Combine actor runtime Data Layers with parent container runtime Data Layers
-			TSet<FName> CombinedRuntimeDataLayers = InRuntimeDataLayers;
-			CombinedRuntimeDataLayers.Append(ContainerInstanceView.GetRuntimeDataLayers());
+			// Inherit parent container properties
+			FName InheritedRuntimeGrid = InRuntimeGrid;
+			bool bInheritedbIsSpatiallyLoaded = bInIsSpatiallyLoaded;
+			TSet<FName> InheritedRuntimeDataLayers = InRuntimeDataLayers;
+			InheritParentContainerProperties(InContainerID, ContainerInstanceView, InheritedRuntimeDataLayers, InheritedRuntimeGrid, bInheritedbIsSpatiallyLoaded);
 
-			CreateActorDescriptorViewsRecursive(SubContainer, SubTransform * InTransform, CombinedRuntimeDataLayers, SubContainerID, InContainerID, SubClusterMode, *ContainerInstanceView.GetActorLabelOrName().ToString());
+			CreateActorDescriptorViewsRecursive(SubContainer, SubTransform * InTransform, InheritedRuntimeGrid, bInheritedbIsSpatiallyLoaded, InheritedRuntimeDataLayers, SubContainerID, InContainerID, SubClusterMode, *ContainerInstanceView.GetActorLabelOrName().ToString());
 		}
 
-		// Create container descriptor
-		check(!ContainerDescriptorsMap.Contains(InContainerID));
+		// Create container instance descriptor
+		check(!ContainerInstanceDescriptorsMap.Contains(InContainerID));
 
-		FContainerDescriptor& ContainerDescriptor = ContainerDescriptorsMap.Add(InContainerID);
-		ContainerDescriptor.Container = InContainer;
-		ContainerDescriptor.Transform = InTransform;
-		ContainerDescriptor.ClusterMode = InClusterMode;
-		ContainerDescriptor.ActorDescViewMap = MoveTemp(ActorDescViewMap);
-		ContainerDescriptor.RuntimeDataLayers = InRuntimeDataLayers;
-		ContainerDescriptor.OwnerName = OwnerName;
+		FContainerInstanceDescriptor& ContainerInstanceDescriptor = ContainerInstanceDescriptorsMap.Add(InContainerID);
+		ContainerInstanceDescriptor.Container = InContainer;
+		ContainerInstanceDescriptor.Transform = InTransform;
+		ContainerInstanceDescriptor.ClusterMode = InClusterMode;
+		ContainerInstanceDescriptor.RuntimeDataLayers = InRuntimeDataLayers;
+		ContainerInstanceDescriptor.RuntimeGrid = InRuntimeGrid;
+		ContainerInstanceDescriptor.bIsSpatiallyLoaded = bInIsSpatiallyLoaded;
+		ContainerInstanceDescriptor.OwnerName = InOwnerName;
+		ContainerInstanceDescriptor.ID = InContainerID;
 
 		// Maintain containers hierarchy, bottom up
 		if (InContainerID != InParentContainerID)
@@ -261,328 +433,311 @@ class FWorldPartitionStreamingGenerator
 	/** 
 	 * Creates the actor descriptor views for the specified container.
 	 */
-	void CreateActorDescriptorViews(const UActorDescContainer* InContainer)
+	void CreateActorContainers(const UActorDescContainer* InContainer)
 	{
-		CreateActorDescriptorViewsRecursive(InContainer, FTransform::Identity, TSet<FName>(), FActorContainerID(), FActorContainerID(), EContainerClusterMode::Partitioned, TEXT("MainContainer"));
+		CreateActorDescriptorViewsRecursive(InContainer, FTransform::Identity, NAME_None, false, TSet<FName>(), FActorContainerID(), FActorContainerID(), EContainerClusterMode::Partitioned, TEXT("MainContainer"));
 
-		// Resolve actor descriptor views once all views are created
-		for (auto It = ContainerDescriptorsMap.CreateIterator(); It; ++It)
+		// Update container instances bounds
+		for (auto& [ContainerID, ContainerInstanceDescriptor] : ContainerInstanceDescriptorsMap)
 		{
-			const FActorContainerID& ContainerID = It.Key();
-			FContainerDescriptor& ContainerDescriptor = It.Value();
+			const FContainerDescriptor& ContainerDescriptor = ContainerDescriptorsMap.FindChecked(ContainerInstanceDescriptor.Container);
 
-			ContainerDescriptor.ActorDescViewMap.ForEachActorDescView([this, &ContainerDescriptor](FWorldPartitionActorDescView& ActorDescView)
+			// Lambdas can't capture structured bindings
+			const FTransform& ContainerInstanceDescriptorTransform = ContainerInstanceDescriptor.Transform;
+			FBox& ContainerInstanceDescriptorBounds = ContainerInstanceDescriptor.Bounds;			
+			check(!ContainerInstanceDescriptorBounds.IsValid);
+
+			ContainerDescriptor.ActorDescViewMap.ForEachActorDescView([&ContainerInstanceDescriptorBounds, &ContainerInstanceDescriptorTransform](const FWorldPartitionActorDescView& ActorDescView)
 			{
-				if (!bEnableStreaming)
+				if (ActorDescView.GetIsSpatiallyLoaded())
 				{
-					ActorDescView.SetForcedNonSpatiallyLoaded();
+					ContainerInstanceDescriptorBounds += ActorDescView.GetBounds().TransformBy(ContainerInstanceDescriptorTransform);
 				}
-
-				ResolveRuntimeDataLayers(ActorDescView, ContainerDescriptor.ActorDescViewMap);
-				ResolveRuntimeReferences(ActorDescView, ContainerDescriptor.ActorDescViewMap);
 			});
 		}
+
+		// Update parent containers bounds, this relies on the fact that ContainersHierarchy is built bottom up
+		for (auto& [ChildContainerID, ParentContainerID] : ContainersHierarchy)
+		{
+			const FContainerInstanceDescriptor& CurrentContainer = ContainerInstanceDescriptorsMap.FindChecked(ChildContainerID);
+			FContainerInstanceDescriptor& ParentContainer = ContainerInstanceDescriptorsMap.FindChecked(ParentContainerID);
+			ParentContainer.Bounds += CurrentContainer.Bounds;
+		}
 	}
 
 	/** 
-	 * Perform various validations on actor descriptor views, and adjust them based on different requirements. This needs to happen before updating containers bounds because
-	 * Some actor descriptor views might change grid placement, etc.
+	 * Perform various validations on the container descriptor, and adjust it based on different requirements. This needs to happen before updating
+	 * containers bounds because some actor descriptor views might change grid placement, etc.
 	 */
-	void ValidateActorDescriptorViews()
+	void ResolveContainerDescriptor(FContainerDescriptor& ContainerDescriptor)
 	{
-		for (auto It = ContainerDescriptorsMap.CreateIterator(); It; ++It)
+		ContainerDescriptor.ActorDescViewMap.ForEachActorDescView([this, &ContainerDescriptor](FWorldPartitionActorDescView& ActorDescView)
 		{
-			const FActorContainerID& ContainerID = It.Key();
-			FContainerDescriptor& ContainerDescriptor = It.Value();
+			if (!bEnableStreaming)
+			{
+				ActorDescView.SetForcedNonSpatiallyLoaded();
+			}
+
+			ResolveRuntimeDataLayers(ActorDescView, ContainerDescriptor.ActorDescViewMap);
+			ResolveRuntimeReferences(ActorDescView, ContainerDescriptor.ActorDescViewMap);
+		});
+	}
 
-			if (ContainerID.IsMainContainer() && ContainerDescriptor.Container->GetWorld())
+	/** 
+	 * Perform various validations on the container descriptor, and adjust it based on different requirements. This needs to happen before updating
+	 * containers bounds because some actor descriptor views might change grid placement, etc.
+	 */
+	void ValidateContainerDescriptor(FContainerDescriptor& ContainerDescriptor, bool bIsMainContainer)
+	{
+		if (bIsMainContainer)
+		{
+			// Gather all references to external actors from the level script and make them always loaded
+			if (ULevelScriptBlueprint* LevelScriptBlueprint = ContainerDescriptor.Container->GetWorld()->PersistentLevel->GetLevelScriptBlueprint(true))
 			{
-				// Gather all references to external actors from the level script and make them always loaded
-				if (ULevelScriptBlueprint* LevelScriptBlueprint = ContainerDescriptor.Container->GetWorld()->PersistentLevel->GetLevelScriptBlueprint(true))
-				{
-					TArray<AActor*> LevelScriptExternalActorReferences = ActorsReferencesUtils::GetExternalActorReferences(LevelScriptBlueprint);
+				TArray<AActor*> LevelScriptExternalActorReferences = ActorsReferencesUtils::GetExternalActorReferences(LevelScriptBlueprint);
 
-					for (AActor* Actor : LevelScriptExternalActorReferences)
+				for (AActor* Actor : LevelScriptExternalActorReferences)
+				{
+					if (FWorldPartitionActorDescView* ActorDescView = ContainerDescriptor.ActorDescViewMap.FindByGuid(Actor->GetActorGuid()))
 					{
-						if (FWorldPartitionActorDescView* ActorDescView = ContainerDescriptor.ActorDescViewMap.FindByGuid(Actor->GetActorGuid()))
+						if (ActorDescView->GetIsSpatiallyLoaded())
 						{
-							if (ActorDescView->GetIsSpatiallyLoaded())
-							{
-								ErrorHandler->OnInvalidReferenceLevelScriptStreamed(*ActorDescView);
-								ActorDescView->SetForcedNonSpatiallyLoaded();
-							}
+							ErrorHandler->OnInvalidReferenceLevelScriptStreamed(*ActorDescView);
+							ActorDescView->SetForcedNonSpatiallyLoaded();
+						}
 
-							if (ActorDescView->GetRuntimeDataLayers().Num())
-							{
-								ErrorHandler->OnInvalidReferenceLevelScriptDataLayers(*ActorDescView);
-								ActorDescView->SetInvalidDataLayers();
-							}
+						if (ActorDescView->GetRuntimeDataLayers().Num())
+						{
+							ErrorHandler->OnInvalidReferenceLevelScriptDataLayers(*ActorDescView);
+							ActorDescView->SetInvalidDataLayers();
 						}
 					}
 				}
 			}
+		}
 
-			// Perform various adjustements based on validations and report errors
-			//
-			// The first validation pass is used to report errors, subsequent passes are used to make corrections to the FWorldPartitionActorDescView
-			// Since the references can form cycles/long chains in the data fixes might need to be propagated in multiple passes.
-			// 
-			// This works because fixes are deterministic and always apply the same way to both Actors being modified, so there's no ordering issues possible
-			int32 NbErrorsDetected = INDEX_NONE;
-			for(uint32 NbValidationPasses = 0; NbErrorsDetected; NbValidationPasses++)
-			{
-				NbErrorsDetected = 0;
+		// Perform various adjustements based on validations and report errors
+		//
+		// The first validation pass is used to report errors, subsequent passes are used to make corrections to the FWorldPartitionActorDescView
+		// Since the references can form cycles/long chains in the data fixes might need to be propagated in multiple passes.
+		// 
+		// This works because fixes are deterministic and always apply the same way to both Actors being modified, so there's no ordering issues possible
+		int32 NbErrorsDetected = INDEX_NONE;
+		for(uint32 NbValidationPasses = 0; NbErrorsDetected; NbValidationPasses++)
+		{
+			NbErrorsDetected = 0;
 
-				ContainerDescriptor.ActorDescViewMap.ForEachActorDescView([this, &ContainerDescriptor, &NbErrorsDetected, &NbValidationPasses](FWorldPartitionActorDescView& ActorDescView)
+			ContainerDescriptor.ActorDescViewMap.ForEachActorDescView([this, &ContainerDescriptor, &NbErrorsDetected, &NbValidationPasses](FWorldPartitionActorDescView& ActorDescView)
+			{
+				// Validate data layers
+				auto IsReferenceGridPlacementValid = [](const FWorldPartitionActorDescView& RefererActorDescView, const FWorldPartitionActorDescView& ReferenceActorDescView)
 				{
-					// Validate data layers
-					auto IsReferenceGridPlacementValid = [](const FWorldPartitionActorDescView& RefererActorDescView, const FWorldPartitionActorDescView& ReferenceActorDescView)
-					{
-						const bool bIsActorDescSpatiallyLoaded = RefererActorDescView.GetIsSpatiallyLoaded();
-						const bool bIsActorDescRefSpatiallyLoaded = ReferenceActorDescView.GetIsSpatiallyLoaded();
+					const bool bIsActorDescSpatiallyLoaded = RefererActorDescView.GetIsSpatiallyLoaded();
+					const bool bIsActorDescRefSpatiallyLoaded = ReferenceActorDescView.GetIsSpatiallyLoaded();
 
-						// The only case we support right now is spatially loaded actors referencing non-spatially loaded actors, when target is not in data layers.
-						// For this to work with data layers, we need to implement dependency logic support in the content cooker splitter.
-						if (bIsActorDescSpatiallyLoaded && !bIsActorDescRefSpatiallyLoaded && ReferenceActorDescView.GetDataLayers().IsEmpty())
-						{
-							return true;
-						}
+					// The only case we support right now is spatially loaded actors referencing non-spatially loaded actors, when target is not in data layers.
+					// For this to work with data layers, we need to implement dependency logic support in the content cooker splitter.
+					if (bIsActorDescSpatiallyLoaded && !bIsActorDescRefSpatiallyLoaded && ReferenceActorDescView.GetDataLayers().IsEmpty())
+					{
+						return true;
+					}
 
-						return bIsActorDescSpatiallyLoaded == bIsActorDescRefSpatiallyLoaded;
-					};
+					return bIsActorDescSpatiallyLoaded == bIsActorDescRefSpatiallyLoaded;
+				};
 
-					// Validate grid placement
-					auto IsReferenceDataLayersValid = [](const FWorldPartitionActorDescView& RefererActorDescView, const FWorldPartitionActorDescView& ReferenceActorDescView)
+				// Validate grid placement
+				auto IsReferenceDataLayersValid = [](const FWorldPartitionActorDescView& RefererActorDescView, const FWorldPartitionActorDescView& ReferenceActorDescView)
+				{
+					if (RefererActorDescView.GetRuntimeDataLayers().Num() == ReferenceActorDescView.GetRuntimeDataLayers().Num())
 					{
-						if (RefererActorDescView.GetRuntimeDataLayers().Num() == ReferenceActorDescView.GetRuntimeDataLayers().Num())
-						{
-							const TSet<FName> RefererActorDescDataLayers(RefererActorDescView.GetRuntimeDataLayers());
-							const TSet<FName> ReferenceActorDescDataLayers(ReferenceActorDescView.GetRuntimeDataLayers());
+						const TSet<FName> RefererActorDescDataLayers(RefererActorDescView.GetRuntimeDataLayers());
+						const TSet<FName> ReferenceActorDescDataLayers(ReferenceActorDescView.GetRuntimeDataLayers());
 
-							return RefererActorDescDataLayers.Includes(ReferenceActorDescDataLayers);
-						}
+						return RefererActorDescDataLayers.Includes(ReferenceActorDescDataLayers);
+					}
 
-						return false;
-					};
+					return false;
+				};
 
-					// Validate runtime grid
-					auto IsReferenceRuntimeGridValid = [](const FWorldPartitionActorDescView& RefererActorDescView, const FWorldPartitionActorDescView& ReferenceActorDescView)
-					{
-						return RefererActorDescView.GetRuntimeGrid() == ReferenceActorDescView.GetRuntimeGrid();
-					};
+				// Validate runtime grid
+				auto IsReferenceRuntimeGridValid = [](const FWorldPartitionActorDescView& RefererActorDescView, const FWorldPartitionActorDescView& ReferenceActorDescView)
+				{
+					return RefererActorDescView.GetRuntimeGrid() == ReferenceActorDescView.GetRuntimeGrid();
+				};
 
-					struct FActorReferenceInfo
-					{
-						FGuid ActorGuid;
-						FWorldPartitionActorDescView* ActorDesc;
-						FGuid ReferenceGuid;
-						FWorldPartitionActorDescView* ReferenceActorDesc;
-					};
+				struct FActorReferenceInfo
+				{
+					FGuid ActorGuid;
+					FWorldPartitionActorDescView* ActorDesc;
+					FGuid ReferenceGuid;
+					FWorldPartitionActorDescView* ReferenceActorDesc;
+				};
 
-					// Build references List
-					TArray<FActorReferenceInfo> References;
+				// Build references List
+				TArray<FActorReferenceInfo> References;
 
-					// Add normal actor references
-					for (const FGuid& ReferenceGuid : ActorDescView.GetReferences())
+				// Add normal actor references
+				for (const FGuid& ReferenceGuid : ActorDescView.GetReferences())
+				{
+					if (ReferenceGuid != ActorDescView.GetParentActor()) // References to the parent are inversed in their handling 
 					{
-						if (ReferenceGuid != ActorDescView.GetParentActor()) // References to the parent are inversed in their handling 
+						// Filter out parent back references
+						FWorldPartitionActorDescView* ReferenceActorDesc = ContainerDescriptor.ActorDescViewMap.FindByGuid(ReferenceGuid);
+						if (ReferenceActorDesc && ReferenceActorDesc->GetParentActor() == ActorDescView.GetGuid())
 						{
-							// Filter out parent back references
-							FWorldPartitionActorDescView* ReferenceActorDesc = ContainerDescriptor.ActorDescViewMap.FindByGuid(ReferenceGuid);
-							if (ReferenceActorDesc && ReferenceActorDesc->GetParentActor() == ActorDescView.GetGuid())
-							{
-								continue;
-							}
-
-							References.Emplace(FActorReferenceInfo{ ActorDescView.GetGuid(), &ActorDescView, ReferenceGuid, ReferenceActorDesc });
+							continue;
 						}
+
+						References.Emplace(FActorReferenceInfo{ ActorDescView.GetGuid(), &ActorDescView, ReferenceGuid, ReferenceActorDesc });
 					}
+				}
 
-					// Add attach reference for the topmost parent, this reference is inverted since we consider the top most existing 
-					// parent to be refering to us, not the child to be referering the parent
-					FGuid ParentGuid = ActorDescView.GetParentActor();
-					FWorldPartitionActorDescView* TopParentDescView = nullptr;
+				// Add attach reference for the topmost parent, this reference is inverted since we consider the top most existing 
+				// parent to be refering to us, not the child to be referering the parent
+				FGuid ParentGuid = ActorDescView.GetParentActor();
+				FWorldPartitionActorDescView* TopParentDescView = nullptr;
 
-					while (ParentGuid.IsValid())
-					{
-						FWorldPartitionActorDescView* ParentDescView = ContainerDescriptor.ActorDescViewMap.FindByGuid(ParentGuid);
+				while (ParentGuid.IsValid())
+				{
+					FWorldPartitionActorDescView* ParentDescView = ContainerDescriptor.ActorDescViewMap.FindByGuid(ParentGuid);
 					
-						if (ParentDescView)
-						{
-							TopParentDescView = ParentDescView;
-							ParentGuid = ParentDescView->GetParentActor();
-						}
-						else
-						{
-							// we had a guid but parent cannot be found, this will be a missing reference
-							break; 
-						}
-					}
-
-					if (TopParentDescView)
+					if (ParentDescView)
 					{
-						References.Emplace(FActorReferenceInfo{ TopParentDescView->GetGuid(), TopParentDescView, ActorDescView.GetGuid(), &ActorDescView });
+						TopParentDescView = ParentDescView;
+						ParentGuid = ParentDescView->GetParentActor();
 					}
-
-					if (ParentGuid.IsValid())
+					else
 					{
-						// In case of missing parent add a missing reference 
-						References.Emplace(FActorReferenceInfo{ ActorDescView.GetGuid(), &ActorDescView, ParentGuid, nullptr });
+						// we had a guid but parent cannot be found, this will be a missing reference
+						break; 
 					}
+				}
 
-					for (FActorReferenceInfo& Info : References)
-					{
-						FWorldPartitionActorDescView* RefererActorDescView = Info.ActorDesc;
-						FWorldPartitionActorDescView* ReferenceActorDescView = Info.ReferenceActorDesc;
+				if (TopParentDescView)
+				{
+					References.Emplace(FActorReferenceInfo{ TopParentDescView->GetGuid(), TopParentDescView, ActorDescView.GetGuid(), &ActorDescView });
+				}
+
+				if (ParentGuid.IsValid())
+				{
+					// In case of missing parent add a missing reference 
+					References.Emplace(FActorReferenceInfo{ ActorDescView.GetGuid(), &ActorDescView, ParentGuid, nullptr });
+				}
+
+				for (FActorReferenceInfo& Info : References)
+				{
+					FWorldPartitionActorDescView* RefererActorDescView = Info.ActorDesc;
+					FWorldPartitionActorDescView* ReferenceActorDescView = Info.ReferenceActorDesc;
 
-						if (ReferenceActorDescView)
+					if (ReferenceActorDescView)
+					{
+						// Validate grid placement
+						if (!IsReferenceGridPlacementValid(*RefererActorDescView, *ReferenceActorDescView))
 						{
-							// Validate grid placement
-							if (!IsReferenceGridPlacementValid(*RefererActorDescView, *ReferenceActorDescView))
+							if (!NbValidationPasses)
 							{
-								if (!NbValidationPasses)
-								{
-									ErrorHandler->OnInvalidReferenceGridPlacement(*RefererActorDescView, *ReferenceActorDescView);									
-								}
-								else
-								{
-									RefererActorDescView->SetForcedNonSpatiallyLoaded();
-									ReferenceActorDescView->SetForcedNonSpatiallyLoaded();
-								}
-
-								NbErrorsDetected++;
+								ErrorHandler->OnInvalidReferenceGridPlacement(*RefererActorDescView, *ReferenceActorDescView);									
 							}
-
-							if (!IsReferenceDataLayersValid(*RefererActorDescView, *ReferenceActorDescView))
+							else
 							{
-								if (!NbValidationPasses)
-								{
-									ErrorHandler->OnInvalidReferenceDataLayers(*RefererActorDescView, *ReferenceActorDescView);									
-								}
-								else
-								{
-									RefererActorDescView->SetInvalidDataLayers();
-									ReferenceActorDescView->SetInvalidDataLayers();
-								}
-
-								NbErrorsDetected++;
+								RefererActorDescView->SetForcedNonSpatiallyLoaded();
+								ReferenceActorDescView->SetForcedNonSpatiallyLoaded();
 							}
 
-							if (!IsReferenceRuntimeGridValid(*RefererActorDescView, *ReferenceActorDescView))
+							NbErrorsDetected++;
+						}
+
+						if (!IsReferenceDataLayersValid(*RefererActorDescView, *ReferenceActorDescView))
+						{
+							if (!NbValidationPasses)
+							{
+								ErrorHandler->OnInvalidReferenceDataLayers(*RefererActorDescView, *ReferenceActorDescView);									
+							}
+							else
 							{
-								if (!NbValidationPasses)
-								{
-									ErrorHandler->OnInvalidReferenceRuntimeGrid(*RefererActorDescView, *ReferenceActorDescView);
-								}
-								else
-								{
-									RefererActorDescView->SetInvalidRuntimeGrid();
-									ReferenceActorDescView->SetInvalidRuntimeGrid();
-								}
-
-								NbErrorsDetected++;
+								RefererActorDescView->SetInvalidDataLayers();
+								ReferenceActorDescView->SetInvalidDataLayers();
 							}
 
+							NbErrorsDetected++;
 						}
-						else
+
+						if (!IsReferenceRuntimeGridValid(*RefererActorDescView, *ReferenceActorDescView))
 						{
 							if (!NbValidationPasses)
 							{
-								ErrorHandler->OnInvalidReference(*RefererActorDescView, Info.ReferenceGuid);
+								ErrorHandler->OnInvalidReferenceRuntimeGrid(*RefererActorDescView, *ReferenceActorDescView);
+							}
+							else
+							{
+								RefererActorDescView->SetInvalidRuntimeGrid();
+								ReferenceActorDescView->SetInvalidRuntimeGrid();
 							}
-							// Do not increment NbErrorsDetected since it won't be fixed and thus will always occur
+
+							NbErrorsDetected++;
 						}
+
 					}
-				});		
-			}
+					else
+					{
+						if (!NbValidationPasses)
+						{
+							ErrorHandler->OnInvalidReference(*RefererActorDescView, Info.ReferenceGuid);
+						}
+						// Do not increment NbErrorsDetected since it won't be fixed and thus will always occur
+					}
+				}
+			});		
+		}
 
-			// Report actors that need to be resaved
-			ContainerDescriptor.ActorDescViewMap.ForEachActorDescView([this](FWorldPartitionActorDescView& ActorDescView)
+		// Report actors that need to be resaved
+		ContainerDescriptor.ActorDescViewMap.ForEachActorDescView([this](FWorldPartitionActorDescView& ActorDescView)
+		{
+			if (ActorDescView.IsResaveNeeded())
 			{
-				if (ActorDescView.IsResaveNeeded())
-				{
-					ErrorHandler->OnActorNeedsResave(ActorDescView);
-				}
-			});
+				ErrorHandler->OnActorNeedsResave(ActorDescView);
+			}
+		});
 
-			// Validate data layers
-			if (ContainerID.IsMainContainer() && ContainerDescriptor.Container->GetWorld())
+		// Validate data layers
+		if (bIsMainContainer)
+		{
+			if (UDataLayerSubsystem* DataLayerSubsystem = UWorld::GetSubsystem<UDataLayerSubsystem>(ContainerDescriptor.Container->GetWorld()))
 			{
-				if (UDataLayerSubsystem* DataLayerSubsystem = UWorld::GetSubsystem<UDataLayerSubsystem>(ContainerDescriptor.Container->GetWorld()))
+				DataLayerSubsystem->ForEachDataLayer([this](const UDataLayerInstance* DataLayerInstance)
 				{
-					DataLayerSubsystem->ForEachDataLayer([this](const UDataLayerInstance* DataLayerInstance)
-					{
-						DataLayerInstance->Validate(ErrorHandler);
-						return true;
-					});
-					break;
-				}
+					DataLayerInstance->Validate(ErrorHandler);
+					return true;
+				});
 			}
 		}
 	}
 
 	/** 
-	 * Update the actor descriptor containers to adjust their bounds from actor descriptor views.
+	 * Update the container descriptor containers to adjust their bounds from actor descriptor views.
 	 */
-	void UpdateContainerDescriptors()
+	void UpdateContainerDescriptor(FContainerDescriptor& ContainerDescriptor)
 	{
-		// Update containers bounds
-		for (auto ContainerIt = ContainerDescriptorsMap.CreateIterator(); ContainerIt; ++ContainerIt)
-		{
-			const FActorContainerID& ContainerID = ContainerIt.Key();
-			FContainerDescriptor& ContainerDescriptor = ContainerIt.Value();
-
-			ContainerDescriptor.ActorDescViewMap.ForEachActorDescView([&ContainerDescriptor](const FWorldPartitionActorDescView& ActorDescView)
-			{
-				if (ActorDescView.GetIsSpatiallyLoaded())
-				{
-					ContainerDescriptor.Bounds += ActorDescView.GetBounds().TransformBy(ContainerDescriptor.Transform);
-				}
-			});
-		}
-
-		// Update parent containers bounds, this relies on the fact that ContainersHierarchy is built bottom up
-		for (auto ContainerPairIt = ContainersHierarchy.CreateIterator(); ContainerPairIt; ++ContainerPairIt)
-		{
-			const FContainerDescriptor& CurrentContainer = ContainerDescriptorsMap.FindChecked(ContainerPairIt.Key());
-			FContainerDescriptor& ParentContainer = ContainerDescriptorsMap.FindChecked(ContainerPairIt.Value());
-			ParentContainer.Bounds += CurrentContainer.Bounds;
-		}
+		// Build clusters for this container - at this point, all actors references should be in the same data layers, grid, etc because of actor descriptors validation.
+		TArray<TPair<FGuid, TArray<FGuid>>> ActorsWithRefs;
+		ContainerDescriptor.ActorDescViewMap.ForEachActorDescView([&ActorsWithRefs](const FWorldPartitionActorDescView& ActorDescView) { ActorsWithRefs.Emplace(ActorDescView.GetGuid(), ActorDescView.GetReferences()); });
+		ContainerDescriptor.Clusters = GenerateObjectsClusters(ActorsWithRefs);
 	}
 
 public:
-	FWorldPartitionStreamingGenerator(FActorDescList* InModifiedActorsDescList, IStreamingGenerationErrorHandler* InErrorHandler, bool bInEnableStreaming)
-	: bEnableStreaming(bInEnableStreaming)
-	, ModifiedActorsDescList(InModifiedActorsDescList)
-	, ErrorHandler(InErrorHandler ? InErrorHandler : &NullErrorHandler)	
+	FWorldPartitionStreamingGenerator(FActorDescList* InModifiedActorsDescList, IStreamingGenerationErrorHandler* InErrorHandler, bool bInEnableStreaming, TArray<TSubclassOf<AActor>> InFilteredClasses = {})
+		: bEnableStreaming(bInEnableStreaming)
+		, ModifiedActorsDescList(InModifiedActorsDescList)
+		, FilteredClasses(InFilteredClasses)
+		, ErrorHandler(InErrorHandler ? InErrorHandler : &NullErrorHandler)	
 	{}
 
 	void PreparationPhase(const UActorDescContainer* Container)
 	{
-		// Preparation Phase :: Actor Descriptor Views Creation
-		CreateActorDescriptorViews(Container);
-
-		// Preparation Phase :: Actor Descriptor Views Validation
-		ValidateActorDescriptorViews();
+		CreateActorContainers(Container);
 
-		// Update container descriptors
-		UpdateContainerDescriptors();
-	}
-
-	FActorClusterContext CreateActorClusters(FActorClusterContext::FFilterActorDescViewFunc InFilterActorDescViewFunc = nullptr)
-	{
-		TArray<FActorContainerInstance> ContainerInstances;
-		ContainerInstances.Reserve(ContainerDescriptorsMap.Num());
-
-		for (auto It = ContainerDescriptorsMap.CreateConstIterator(); It; ++It)
-		{
-			const FActorContainerID& ContainerID = It.Key();
-			const FContainerDescriptor& ContainerDescriptor = It.Value();
-
-			ContainerInstances.Emplace(ContainerID, ContainerDescriptor.Transform, ContainerDescriptor.Bounds, ContainerDescriptor.RuntimeDataLayers, ContainerDescriptor.ClusterMode, ContainerDescriptor.Container, ContainerDescriptor.ActorDescViewMap.ActorDescViewsByGuid);
-		}
-
-		return FActorClusterContext(MoveTemp(ContainerInstances), InFilterActorDescViewFunc);
+		// Construct the streaming generation context
+		StreamingGenerationContext = MakeUnique<FStreamingGenerationContext>(this, Container);
 	}
 
 	static TUniquePtr<FArchive> CreateDumpStateLogArchive(const TCHAR* Suffix)
@@ -595,43 +750,61 @@ public:
 	{
 		// Build the containers tree representation
 		TMultiMap<FActorContainerID, FActorContainerID> InvertedContainersHierarchy;
-		for (auto ContainerPairIt = ContainersHierarchy.CreateIterator(); ContainerPairIt; ++ContainerPairIt)
+		for (auto& [ChildContainerID, ParentContainerID] : ContainersHierarchy)
 		{
-			const FActorContainerID& ChildContainerID = ContainerPairIt.Key();
-			const FActorContainerID& ParentContainerID = ContainerPairIt.Value();				
 			InvertedContainersHierarchy.Add(ParentContainerID, ChildContainerID);
 		}
 
-		Ar.Printf(TEXT("Containers:"));
-
-		auto DumpContainers = [this, &InvertedContainersHierarchy, &Ar](const FActorContainerID& ContainerID)
+		UE_SCOPED_INDENT_LOG_ARCHIVE(Ar.PrintfIndent(TEXT("Containers:")));
+		for (auto& [ActorDescContainer, ContainerDescriptor] : ContainerDescriptorsMap)
 		{
-			auto DumpContainersRecursive = [this, &InvertedContainersHierarchy, &Ar](const FActorContainerID& ContainerID, auto& RecursiveFunc) -> void
+			UE_SCOPED_INDENT_LOG_ARCHIVE(Ar.PrintfIndent(TEXT("Container: %s"), *ContainerDescriptor.Container->GetContainerPackage().ToString()));
+
+			if (ContainerDescriptor.ActorDescViewMap.Num())
 			{
-				const FContainerDescriptor& ContainerDescriptor = ContainerDescriptorsMap.FindChecked(ContainerID);
-				
-				{
-					FHierarchicalLogArchive::FIndentScope IndentScope = Ar.PrintfIndent(TEXT("%s:"), *ContainerDescriptor.OwnerName);
+				UE_SCOPED_INDENT_LOG_ARCHIVE(Ar.PrintfIndent(TEXT("ActorDescs:")));
 
-					Ar.Printf(TEXT("       ID: 0x%016llx"), ContainerID.ID);
-					Ar.Printf(TEXT("   Bounds: %s"), *ContainerDescriptor.Bounds.ToString());
-					Ar.Printf(TEXT("Transform: %s"), *ContainerDescriptor.Transform.ToString());
-					Ar.Printf(TEXT("Container: %s"), *ContainerDescriptor.Container->GetContainerPackage().ToString());
-				}
+				TMap<FGuid, FWorldPartitionActorDescView*> SortedActorDescViewMap = ContainerDescriptor.ActorDescViewMap.ActorDescViewsByGuid;
+				SortedActorDescViewMap.KeySort([](const FGuid& GuidA, const FGuid& GuidB) { return GuidA < GuidB; });
 
-				if (ContainerDescriptor.ActorDescViewMap.ActorDescViewsByGuid.Num())
+				for (auto& [ActorGuid, ActorDescView] : SortedActorDescViewMap)
 				{
-					FHierarchicalLogArchive::FIndentScope IndentScope = Ar.PrintfIndent(TEXT("ActorDescs:"));
+					Ar.Print(*ActorDescView->ToString());
+				}
+			}
 
-					TMap<FGuid, FWorldPartitionActorDescView*> SortedActorDescViewMap = ContainerDescriptor.ActorDescViewMap.ActorDescViewsByGuid;
-					SortedActorDescViewMap.KeySort([](const FGuid& GuidA, const FGuid& GuidB) { return GuidA < GuidB; });
+			if (ContainerDescriptor.Clusters.Num())
+			{
+				UE_SCOPED_INDENT_LOG_ARCHIVE(Ar.PrintfIndent(TEXT("Clusters:")));
 
-					for (auto ActorDescIt = SortedActorDescViewMap.CreateConstIterator(); ActorDescIt; ++ActorDescIt)
+				int ClusterIndex = 0;
+				for (TArray<FGuid>& ActorGuids : ContainerDescriptor.Clusters)
+				{
+					UE_SCOPED_INDENT_LOG_ARCHIVE(Ar.PrintfIndent(TEXT("[%3d]"), ClusterIndex++));
+					for (const FGuid& ActorGuid : ActorGuids)
 					{
-						const FWorldPartitionActorDescView& ActorDescView = *ActorDescIt.Value();
+						const FWorldPartitionActorDescView& ActorDescView = ContainerDescriptor.ActorDescViewMap.FindByGuidChecked(ActorGuid);
 						Ar.Print(*ActorDescView.ToString());
 					}
 				}
+			}
+		}
+
+		Ar.Printf(TEXT("ContainerInstances:"));
+		auto DumpContainerInstances = [this, &InvertedContainersHierarchy, &Ar](const FActorContainerID& ContainerID)
+		{
+			auto DumpContainerInstancesRecursive = [this, &InvertedContainersHierarchy, &Ar](const FActorContainerID& ContainerID, auto& RecursiveFunc) -> void
+			{
+				const FContainerInstanceDescriptor& ContainerInstanceDescriptor = ContainerInstanceDescriptorsMap.FindChecked(ContainerID);
+				
+				{
+					UE_SCOPED_INDENT_LOG_ARCHIVE(Ar.PrintfIndent(TEXT("%s:"), *ContainerInstanceDescriptor.OwnerName));
+
+					Ar.Printf(TEXT("       ID: 0x%016llx"), ContainerID.ID);
+					Ar.Printf(TEXT("   Bounds: %s"), *ContainerInstanceDescriptor.Bounds.ToString());
+					Ar.Printf(TEXT("Transform: %s"), *ContainerInstanceDescriptor.Transform.ToString());
+					Ar.Printf(TEXT("Container: %s"), *ContainerInstanceDescriptor.Container->GetContainerPackage().ToString());
+				}
 
 				TArray<FActorContainerID> ChildContainersIDs;
 				InvertedContainersHierarchy.MultiFind(ContainerID, ChildContainersIDs);
@@ -639,7 +812,7 @@ public:
 
 				if (ChildContainersIDs.Num())
 				{
-					FHierarchicalLogArchive::FIndentScope IndentScope = Ar.PrintfIndent(TEXT("SubContainers:"));
+					UE_SCOPED_INDENT_LOG_ARCHIVE(Ar.PrintfIndent(TEXT("SubContainers:")));
 						
 					for (const FActorContainerID& ChildContainerID : ChildContainersIDs)
 					{
@@ -648,39 +821,35 @@ public:
 				}
 			};
 
-			DumpContainersRecursive(ContainerID, DumpContainersRecursive);
+			DumpContainerInstancesRecursive(ContainerID, DumpContainerInstancesRecursive);
 		};
 
-		DumpContainers(FActorContainerID());
+		DumpContainerInstances(FActorContainerID());
+	}
+
+	const FStreamingGenerationContext* GetStreamingGenerationContext()
+	{
+		return StreamingGenerationContext.Get();
 	}
 
 private:
 	bool bEnableStreaming;
 	FActorDescList* ModifiedActorsDescList;
+	TArray<TSubclassOf<AActor>> FilteredClasses;
 	IStreamingGenerationErrorHandler* ErrorHandler;
 	FStreamingGenerationNullErrorHandler NullErrorHandler;
 
-	struct FContainerDescriptor
-	{
-		FContainerDescriptor()
-		: Bounds(ForceInit)
-		, Container(nullptr)
-		{}
-
-		FBox Bounds;
-		FTransform Transform;
-		const UActorDescContainer* Container;
-		EContainerClusterMode ClusterMode;
-		FActorDescViewMap ActorDescViewMap;
-		TSet<FName> RuntimeDataLayers;
-		FString OwnerName;
-	};
-
-	/** Maps containers IDs to their container descriptor */
-	TMap<FActorContainerID, FContainerDescriptor> ContainerDescriptorsMap;
+	/** Maps containers to their container descriptor */
+	TMap<const UActorDescContainer*, FContainerDescriptor> ContainerDescriptorsMap;
+	
+	/** Maps containers IDs to their container instance descriptor */
+	TMap<FActorContainerID, FContainerInstanceDescriptor> ContainerInstanceDescriptorsMap;
 
 	/** Maps containers IDs to their parent ID */
 	TMap<FActorContainerID, FActorContainerID> ContainersHierarchy;
+
+	/** Data required for streaming generation interface */
+	TUniquePtr<FStreamingGenerationContext> StreamingGenerationContext;
 };
 
 bool UWorldPartition::GenerateStreaming(TArray<FString>* OutPackagesToGenerate)
@@ -704,7 +873,6 @@ bool UWorldPartition::GenerateStreaming(TArray<FString>* OutPackagesToGenerate)
 	TUniquePtr<FArchive> LogFileAr = FWorldPartitionStreamingGenerator::CreateDumpStateLogArchive(StateLogSuffix);
 	FHierarchicalLogArchive HierarchicalLogAr(*LogFileAr);
 
-	FActorClusterContext ActorClusterContext;
 	FWorldPartitionStreamingGenerator StreamingGenerator(ModifiedActorsDescList, ErrorHandler, IsStreamingEnabled());
 
 	// Preparation Phase
@@ -712,17 +880,14 @@ bool UWorldPartition::GenerateStreaming(TArray<FString>* OutPackagesToGenerate)
 
 	StreamingGenerator.DumpStateLog(HierarchicalLogAr);
 
-	// Preparation Phase :: Actor Clusters Creation
-	ActorClusterContext = StreamingGenerator.CreateActorClusters();
-
 	// Generate streaming
 	check(!StreamingPolicy);
 	StreamingPolicy = NewObject<UWorldPartitionStreamingPolicy>(const_cast<UWorldPartition*>(this), WorldPartitionStreamingPolicyClass.Get(), NAME_None, bIsPIE ? RF_Transient : RF_NoFlags);
 
 	check(RuntimeHash);
-	if (RuntimeHash->GenerateStreaming(StreamingPolicy, ActorClusterContext, OutPackagesToGenerate))
+	if (RuntimeHash->GenerateStreaming(StreamingPolicy, StreamingGenerator.GetStreamingGenerationContext(), OutPackagesToGenerate))
 	{
-		if (IsRunningCookCommandlet())
+		//if (IsRunningCookCommandlet())
 		{
 			RuntimeHash->DumpStateLog(HierarchicalLogAr);
 		}
@@ -743,20 +908,14 @@ void UWorldPartition::FlushStreaming()
 void UWorldPartition::GenerateHLOD(ISourceControlHelper* SourceControlHelper, bool bCreateActorsOnly)
 {
 	FStreamingGenerationLogErrorHandler LogErrorHandler;
-	FWorldPartitionStreamingGenerator StreamingGenerator(nullptr, &LogErrorHandler, IsStreamingEnabled());
+	FWorldPartitionStreamingGenerator StreamingGenerator(nullptr, &LogErrorHandler, IsStreamingEnabled(), { AWorldPartitionHLOD::StaticClass() });
 	StreamingGenerator.PreparationPhase(this);
 
 	TUniquePtr<FArchive> LogFileAr = FWorldPartitionStreamingGenerator::CreateDumpStateLogArchive(TEXT("HLOD"));
 	FHierarchicalLogArchive HierarchicalLogAr(*LogFileAr);
 	StreamingGenerator.DumpStateLog(HierarchicalLogAr);
 
-	// Preparation Phase :: Actor Clusters Creation
-	FActorClusterContext ActorClusterContext = StreamingGenerator.CreateActorClusters([](const FWorldPartitionActorDescView& ActorDescView)
-	{
-		return !ActorDescView.GetActorNativeClass()->IsChildOf<AWorldPartitionHLOD>();
-	});
-
-	RuntimeHash->GenerateHLOD(SourceControlHelper, ActorClusterContext, bCreateActorsOnly);
+	RuntimeHash->GenerateHLOD(SourceControlHelper, StreamingGenerator.GetStreamingGenerationContext(), bCreateActorsOnly);
 }
 
 void UWorldPartition::CheckForErrors(IStreamingGenerationErrorHandler* ErrorHandler) const
@@ -772,4 +931,4 @@ void UWorldPartition::CheckForErrors(IStreamingGenerationErrorHandler* ErrorHand
 }
 #endif // WITH_EDITOR
 
-#undef LOCTEXT_NAMESPACE
\ No newline at end of file
+#undef LOCTEXT_NAMESPACE
diff --git a/Engine/Source/Runtime/Engine/Private/WorldPartition/WorldPartitionStreamingGenerationContext.cpp b/Engine/Source/Runtime/Engine/Private/WorldPartition/WorldPartitionStreamingGenerationContext.cpp
new file mode 100644
--- /dev/null
+++ b/Engine/Source/Runtime/Engine/Private/WorldPartition/WorldPartitionStreamingGenerationContext.cpp
@@ -0,0 +1,25 @@
+// Copyright Epic Games, Inc. All Rights Reserved.
+
+#include "WorldPartition/WorldPartitionStreamingGenerationContext.h"
+
+#if WITH_EDITOR
+const FWorldPartitionActorDescView& IStreamingGenerationContext::FActorInstance::GetActorDescView() const
+{
+	return ActorSetInstance->ContainerInstance->ActorDescViewMap->FindByGuidChecked(ActorGuid);
+}
+
+const FActorContainerID& IStreamingGenerationContext::FActorInstance::GetContainerID() const
+{
+	return ActorSetInstance->ContainerID;
+}
+
+const FTransform& IStreamingGenerationContext::FActorInstance::GetTransform() const
+{
+	return ActorSetInstance->Transform;
+}
+
+const UActorDescContainer* IStreamingGenerationContext::FActorInstance::GetActorDescContainer() const
+{
+	return ActorSetInstance->ContainerInstance->ActorDescContainer;
+}
+#endif
\ No newline at end of file
diff --git a/Engine/Source/Runtime/Engine/Public/WorldPartition/DataLayer/DataLayerSubsystem.h b/Engine/Source/Runtime/Engine/Public/WorldPartition/DataLayer/DataLayerSubsystem.h
--- a/Engine/Source/Runtime/Engine/Public/WorldPartition/DataLayer/DataLayerSubsystem.h
+++ b/Engine/Source/Runtime/Engine/Public/WorldPartition/DataLayer/DataLayerSubsystem.h
@@ -89,6 +89,8 @@ public:
 
 	void UpdateDataLayerEditorPerProjectUserSettings() const;
 	void GetUserLoadedInEditorStates(TArray<FName>& OutDataLayersLoadedInEditor, TArray<FName>& OutDataLayersNotLoadedInEditor) const;
+	
+	static TArray<const UDataLayerInstance*> GetRuntimeDataLayerInstances(UWorld* InWorld, const TArray<FName>& DataLayerInstanceNames);
 
 	void PushActorEditorContext() const;
 	void PopActorEditorContext() const;
diff --git a/Engine/Source/Runtime/Engine/Public/WorldPartition/DataLayer/DataLayersID.h b/Engine/Source/Runtime/Engine/Public/WorldPartition/DataLayer/DataLayersID.h
--- a/Engine/Source/Runtime/Engine/Public/WorldPartition/DataLayer/DataLayersID.h
+++ b/Engine/Source/Runtime/Engine/Public/WorldPartition/DataLayer/DataLayersID.h
@@ -23,7 +23,15 @@ public:
 		return !(*this == Other);
 	}
 
-	uint32 GetHash() const { return Hash; }
+	uint32 GetHash() const
+	{
+		return Hash;
+	}
+
+	friend uint32 GetTypeHash(const FDataLayersID& InDataLayersID)
+	{
+		return InDataLayersID.Hash;
+	}
 
 private:
 	uint32 Hash;
diff --git a/Engine/Source/Runtime/Engine/Public/WorldPartition/HLOD/HLODSubActor.h b/Engine/Source/Runtime/Engine/Public/WorldPartition/HLOD/HLODSubActor.h
--- a/Engine/Source/Runtime/Engine/Public/WorldPartition/HLOD/HLODSubActor.h
+++ b/Engine/Source/Runtime/Engine/Public/WorldPartition/HLOD/HLODSubActor.h
@@ -2,7 +2,7 @@
 #pragma once
 
 #include "Misc/Guid.h"
-#include "WorldPartition/WorldPartitionActorCluster.h"
+#include "WorldPartition/WorldPartitionActorContainerID.h"
 #include "HLODSubActor.generated.h"
 
 USTRUCT()
diff --git a/Engine/Source/Runtime/Engine/Public/WorldPartition/HLOD/IWorldPartitionHLODUtilities.h b/Engine/Source/Runtime/Engine/Public/WorldPartition/HLOD/IWorldPartitionHLODUtilities.h
--- a/Engine/Source/Runtime/Engine/Public/WorldPartition/HLOD/IWorldPartitionHLODUtilities.h
+++ b/Engine/Source/Runtime/Engine/Public/WorldPartition/HLOD/IWorldPartitionHLODUtilities.h
@@ -5,8 +5,8 @@
 #include "CoreMinimal.h"
 #include "Containers/Array.h"
 #include "WorldPartition/WorldPartitionHandle.h"
+#include "WorldPartition/WorldPartitionStreamingGenerationContext.h"
 #include "WorldPartition/DataLayer/DataLayersID.h"
-#include "WorldPartition/WorldPartitionActorCluster.h"
 
 class AActor;
 class UWorldPartition;
@@ -55,7 +55,7 @@ public:
 	 * @param	InActors			The actors for which we'll build an HLOD representation
 	 * @param	InDataLayers		The data layers to assign to the newly created HLOD actors
 	 */
-	virtual TArray<AWorldPartitionHLOD*> CreateHLODActors(FHLODCreationContext& InCreationContext, const FHLODCreationParams& InCreationParams, const TSet<FActorInstance>& InActors, const TArray<const UDataLayerInstance*>& InDataLayerInstances) = 0;
+	virtual TArray<AWorldPartitionHLOD*> CreateHLODActors(FHLODCreationContext& InCreationContext, const FHLODCreationParams& InCreationParams, const TArray<IStreamingGenerationContext::FActorInstance>& InActors, const TArray<const UDataLayerInstance*>& InDataLayerInstances) = 0;
 
 	/**
 	 * Build HLOD for the specified AWorldPartitionHLOD actor.
diff --git a/Engine/Source/Runtime/Engine/Public/WorldPartition/WorldPartitionActorCluster.h b/Engine/Source/Runtime/Engine/Public/WorldPartition/WorldPartitionActorCluster.h
deleted file mode 100644
diff --git a/Engine/Source/Runtime/Engine/Public/WorldPartition/WorldPartitionActorContainerID.h b/Engine/Source/Runtime/Engine/Public/WorldPartition/WorldPartitionActorContainerID.h
new file mode 100644
--- /dev/null
+++ b/Engine/Source/Runtime/Engine/Public/WorldPartition/WorldPartitionActorContainerID.h
@@ -0,0 +1,61 @@
+// Copyright Epic Games, Inc. All Rights Reserved.
+#pragma once
+
+#include "Hash/CityHash.h"
+#include "WorldPartitionActorContainerID.generated.h"
+
+USTRUCT()
+struct FActorContainerID
+{
+	GENERATED_USTRUCT_BODY()
+
+	FActorContainerID()
+	: ID(0)
+	{}
+
+	FActorContainerID(const FActorContainerID& InOther)
+	: ID(InOther.ID)
+	{}
+
+	FActorContainerID(const FActorContainerID& InParent, FGuid InActorGuid)
+	: ID(CityHash64WithSeed((const char*)&InActorGuid, sizeof(InActorGuid), InParent.ID))
+	{}
+
+	void operator=(const FActorContainerID& InOther)
+	{
+		ID = InOther.ID;
+	}
+
+	bool operator==(const FActorContainerID& InOther) const
+	{
+		return ID == InOther.ID;
+	}
+
+	bool operator!=(const FActorContainerID& InOther) const
+	{
+		return ID != InOther.ID;
+	}
+
+	bool IsMainContainer() const
+	{
+		return !ID;
+	}
+
+	FString ToString() const
+	{
+		return FString::Printf(TEXT("%016llx"), ID);
+	}
+
+	friend FORCEINLINE uint32 GetTypeHash(const FActorContainerID& InContainerID)
+	{
+		return GetTypeHash(InContainerID.ID);
+	}
+
+	static FActorContainerID GetMainContainerID()
+	{
+		return FActorContainerID();
+	}
+
+	UPROPERTY()
+	uint64 ID;
+};
diff --git a/Engine/Source/Runtime/Engine/Public/WorldPartition/WorldPartitionActorDescView.h b/Engine/Source/Runtime/Engine/Public/WorldPartition/WorldPartitionActorDescView.h
--- a/Engine/Source/Runtime/Engine/Public/WorldPartition/WorldPartitionActorDescView.h
+++ b/Engine/Source/Runtime/Engine/Public/WorldPartition/WorldPartitionActorDescView.h
@@ -5,6 +5,7 @@
 #include "Misc/Optional.h"
 
 #if WITH_EDITOR
+class AActor;
 class FWorldPartitionActorDesc;
 class UActorDescContainer;
 class UWorldPartition;
@@ -59,6 +60,8 @@ public:
 
 	bool IsResaveNeeded() const;
 
+	AActor* GetActor() const;
+
 	bool operator==(const FWorldPartitionActorDescView& Other) const
 	{
 		return GetGuid() == Other.GetGuid();
diff --git a/Engine/Source/Runtime/Engine/Public/WorldPartition/WorldPartitionRuntimeCell.h b/Engine/Source/Runtime/Engine/Public/WorldPartition/WorldPartitionRuntimeCell.h
--- a/Engine/Source/Runtime/Engine/Public/WorldPartition/WorldPartitionRuntimeCell.h
+++ b/Engine/Source/Runtime/Engine/Public/WorldPartition/WorldPartitionRuntimeCell.h
@@ -6,7 +6,7 @@
 #include "WorldPartition/DataLayer/DataLayerInstance.h"
 #include "WorldPartition/WorldPartitionActorDescView.h"
 #include "WorldPartition/WorldPartitionStreamingSource.h"
-#include "WorldPartition/WorldPartitionActorCluster.h"
+#include "WorldPartition/WorldPartitionActorContainerID.h"
 #include "WorldPartition/WorldPartitionRuntimeCellInterface.h"
 #include "ProfilingDebugging/ProfilingHelpers.h"
 #include "Misc/HierarchicalLogArchive.h"
diff --git a/Engine/Source/Runtime/Engine/Public/WorldPartition/WorldPartitionRuntimeHash.h b/Engine/Source/Runtime/Engine/Public/WorldPartition/WorldPartitionRuntimeHash.h
--- a/Engine/Source/Runtime/Engine/Public/WorldPartition/WorldPartitionRuntimeHash.h
+++ b/Engine/Source/Runtime/Engine/Public/WorldPartition/WorldPartitionRuntimeHash.h
@@ -10,13 +10,15 @@
 #include "WorldPartition/WorldPartitionHandle.h"
 #include "WorldPartition/WorldPartitionRuntimeCell.h"
 #include "WorldPartition/WorldPartitionActorDescViewProxy.h"
+#include "WorldPartition/WorldPartitionStreamingGeneration.h"
 #if WITH_EDITOR
 #include "CookPackageSplitter.h"
 #include "Misc/HierarchicalLogArchive.h"
 #endif
 #include "WorldPartitionRuntimeHash.generated.h"
 
-class FActorClusterContext;
+class FActorDescViewMap;
+class IStreamingGenerationContext;
 struct FHierarchicalLogArchive;
 
 UENUM()
@@ -36,9 +38,9 @@ class ENGINE_API UWorldPartitionRuntimeHash : public UObject
 	virtual void SetDefaultValues() {}
 	virtual bool PopulateGeneratorPackageForCook(const TArray<ICookPackageSplitter::FGeneratedPackageForPreSave>& InGeneratedPackages, TArray<UPackage*>& OutModifiedPackages) { return false; }
 	virtual bool PopulateGeneratedPackageForCook(UPackage* InPackage, const FString& InPackageRelativePath, TArray<UPackage*>& OutModifiedPackages) { return false; }
-	virtual bool GenerateStreaming(class UWorldPartitionStreamingPolicy* StreamingPolicy, const FActorClusterContext& ActorClusterContext, TArray<FString>* OutPackagesToGenerate) { return false; }
+	virtual bool GenerateStreaming(class UWorldPartitionStreamingPolicy* StreamingPolicy, const IStreamingGenerationContext* StreamingGenerationContext, TArray<FString>* OutPackagesToGenerate) { return false; }
 	virtual void FlushStreaming() {}
-	virtual bool GenerateHLOD(ISourceControlHelper* SourceControlHelper, FActorClusterContext& ActorClusterContext, bool bCreateActorsOnly) { return false; }
+	virtual bool GenerateHLOD(ISourceControlHelper* SourceControlHelper, const IStreamingGenerationContext* StreamingGenerationContext, bool bCreateActorsOnly) { return false; }
 	virtual void DrawPreview() const {}
 
 	virtual void DumpStateLog(FHierarchicalLogArchive& Ar);
diff --git a/Engine/Source/Runtime/Engine/Public/WorldPartition/WorldPartitionRuntimeSpatialHash.h b/Engine/Source/Runtime/Engine/Public/WorldPartition/WorldPartitionRuntimeSpatialHash.h
--- a/Engine/Source/Runtime/Engine/Public/WorldPartition/WorldPartitionRuntimeSpatialHash.h
+++ b/Engine/Source/Runtime/Engine/Public/WorldPartition/WorldPartitionRuntimeSpatialHash.h
@@ -235,7 +235,7 @@ public:
 	virtual bool PopulateGeneratorPackageForCook(const TArray<ICookPackageSplitter::FGeneratedPackageForPreSave>& InGeneratedPackages, TArray<UPackage*>& OutModifiedPackages) override;
 	virtual bool PopulateGeneratedPackageForCook(UPackage* InPackage, const FString& InPackageRelativePath, TArray<UPackage*>& OutModifiedPackages) override;
 	virtual void FlushStreaming() override;
-	virtual bool GenerateHLOD(ISourceControlHelper* SourceControlHelper, FActorClusterContext& ActorClusterContext, bool bCreateActorsOnly) override;
+	virtual bool GenerateHLOD(ISourceControlHelper* SourceControlHelper, const IStreamingGenerationContext* StreamingGenerationContext, bool bCreateActorsOnly) override;
 	virtual void DrawPreview() const override;
 
 	FName GetCellName(FName InGridName, const FGridCellCoord& InCellGlobalCoord, const FDataLayersID& InDataLayerID) const;
@@ -257,7 +257,7 @@ protected:
 	virtual EWorldPartitionStreamingPerformance GetStreamingPerformanceForCell(const UWorldPartitionRuntimeCell* Cell) const override;
 
 #if WITH_EDITOR
-	virtual bool GenerateStreaming(class UWorldPartitionStreamingPolicy* StreamingPolicy, const FActorClusterContext& ActorClusterContext, TArray<FString>* OutPackagesToGenerate = nullptr) override;
+	virtual bool GenerateStreaming(class UWorldPartitionStreamingPolicy* StreamingPolicy, const IStreamingGenerationContext* StreamingGenerationContext, TArray<FString>* OutPackagesToGenerate = nullptr) override;
 	virtual void DumpStateLog(FHierarchicalLogArchive& Ar) override;
 #endif
 
diff --git a/Engine/Source/Runtime/Engine/Public/WorldPartition/WorldPartitionStreamingGeneration.h b/Engine/Source/Runtime/Engine/Public/WorldPartition/WorldPartitionStreamingGeneration.h
--- a/Engine/Source/Runtime/Engine/Public/WorldPartition/WorldPartitionStreamingGeneration.h
+++ b/Engine/Source/Runtime/Engine/Public/WorldPartition/WorldPartitionStreamingGeneration.h
@@ -78,6 +78,11 @@ public:
 		return nullptr;
 	}
 
+	const FWorldPartitionActorDescView& FindByGuidChecked(const FGuid& InGuid) const
+	{
+		return *ActorDescViewsByGuid.FindChecked(InGuid);
+	}
+
 	template <class ClassType>
 	TArray<const FWorldPartitionActorDescView*> FindByExactNativeClass() const
 	{
diff --git a/Engine/Source/Runtime/Engine/Public/WorldPartition/WorldPartitionStreamingGenerationContext.h b/Engine/Source/Runtime/Engine/Public/WorldPartition/WorldPartitionStreamingGenerationContext.h
new file mode 100644
--- /dev/null
+++ b/Engine/Source/Runtime/Engine/Public/WorldPartition/WorldPartitionStreamingGenerationContext.h
@@ -0,0 +1,58 @@
+// Copyright Epic Games, Inc. All Rights Reserved.
+#pragma once
+
+#include "CoreMinimal.h"
+#include "WorldPartition/WorldPartitionStreamingGeneration.h"
+#include "WorldPartition/WorldPartitionActorContainerID.h"
+
+class UDataLayerInstance;
+
+#if WITH_EDITOR
+class ENGINE_API IStreamingGenerationContext
+{
+public:
+	struct ENGINE_API FActorSet
+	{
+		TArray<FGuid> Actors;
+	};
+
+	struct ENGINE_API FActorSetContainer
+	{
+		const FActorDescViewMap* ActorDescViewMap;
+		const UActorDescContainer* ActorDescContainer;
+		TArray<FActorSet> ActorSets;
+	};
+
+	struct ENGINE_API FActorSetInstance
+	{
+		FBox Bounds;
+		FName RuntimeGrid;
+		bool bIsSpatiallyLoaded;
+		TArray<const UDataLayerInstance*> DataLayers;
+		const FActorSetContainer* ContainerInstance;
+		FActorContainerID ContainerID;
+		FTransform Transform;
+		const FActorSet* ActorSet;
+	};
+
+	struct ENGINE_API FActorInstance
+	{
+		FActorInstance(const FGuid& InActorGuid, const FActorSetInstance* InActorSetInstance)
+			: ActorGuid(InActorGuid)
+			, ActorSetInstance(InActorSetInstance)
+		{}
+
+		FGuid ActorGuid;
+		const FActorSetInstance* ActorSetInstance;
+
+		const FWorldPartitionActorDescView& GetActorDescView() const;
+		const FActorContainerID& GetContainerID() const;
+		const FTransform& GetTransform() const;
+		const UActorDescContainer* GetActorDescContainer() const;
+	};
+
+	virtual FBox GetWorldBounds() const = 0;
+	virtual const FActorSetContainer* GetMainWorldContainer() const = 0;
+	virtual void ForEachActorSetInstance(TFunctionRef<void(const FActorSetInstance&)> Func) const = 0;
+};
+#endif
\ No newline at end of file
