UAudioComponent crash fix

diff --git a/Engine/Source/Runtime/Engine/Private/Components/AudioComponent.cpp b/Engine/Source/Runtime/Engine/Private/Components/AudioComponent.cpp
--- a/Engine/Source/Runtime/Engine/Private/Components/AudioComponent.cpp
+++ b/Engine/Source/Runtime/Engine/Private/Components/AudioComponent.cpp
@@ -666,7 +666,7 @@ void UAudioComponent::PlayInternal(const PlayInternalRequestData& InPlayRequestD
 	// Auto attach if requested
 	const bool bWasAutoAttached = bDidAutoAttach;
 	bDidAutoAttach = false;
-	if (bAutoManageAttachment && World->IsGameWorld())
+	if (bAutoManageAttachment && World && World->IsGameWorld())
 	{
 		USceneComponent* NewParent = AutoAttachParent.Get();
 		if (NewParent)
@@ -701,7 +701,7 @@ void UAudioComponent::PlayInternal(const PlayInternalRequestData& InPlayRequestD
 
 	FActiveSound NewActiveSound;
 	NewActiveSound.SetAudioComponent(*this);
-	NewActiveSound.SetWorld(GetWorld());
+	NewActiveSound.SetWorld(World);
 	NewActiveSound.SetSound(SoundToPlay);
 	NewActiveSound.SetSourceEffectChain(SourceEffectChain);
 	NewActiveSound.SetSoundClass(SoundClassOverride);
