World Partition: don't try to validate the level script during changelist validations when the world is not loaded.

#rnx

diff --git a/Engine/Source/Runtime/Engine/Private/WorldPartition/WorldPartitionStreamingGeneration.cpp b/Engine/Source/Runtime/Engine/Private/WorldPartition/WorldPartitionStreamingGeneration.cpp
--- a/Engine/Source/Runtime/Engine/Private/WorldPartition/WorldPartitionStreamingGeneration.cpp
+++ b/Engine/Source/Runtime/Engine/Private/WorldPartition/WorldPartitionStreamingGeneration.cpp
@@ -491,28 +491,41 @@ class FWorldPartitionStreamingGenerator
 	{
 		if (bIsMainContainer)
 		{
-			// Gather all references to external actors from the level script and make them always loaded
-			if (ULevelScriptBlueprint* LevelScriptBlueprint = ContainerDescriptor.Container->GetWorld()->PersistentLevel->GetLevelScriptBlueprint(true))
+			if (UWorld* World = ContainerDescriptor.Container->GetWorld())
 			{
-				TArray<AActor*> LevelScriptExternalActorReferences = ActorsReferencesUtils::GetExternalActorReferences(LevelScriptBlueprint);
-
-				for (AActor* Actor : LevelScriptExternalActorReferences)
+				// Gather all references to external actors from the level script and make them always loaded
+				if (ULevelScriptBlueprint* LevelScriptBlueprint = World->PersistentLevel->GetLevelScriptBlueprint(true))
 				{
-					if (FWorldPartitionActorDescView* ActorDescView = ContainerDescriptor.ActorDescViewMap.FindByGuid(Actor->GetActorGuid()))
+					TArray<AActor*> LevelScriptExternalActorReferences = ActorsReferencesUtils::GetExternalActorReferences(LevelScriptBlueprint);
+
+					for (AActor* Actor : LevelScriptExternalActorReferences)
 					{
-						if (ActorDescView->GetIsSpatiallyLoaded())
+						if (FWorldPartitionActorDescView* ActorDescView = ContainerDescriptor.ActorDescViewMap.FindByGuid(Actor->GetActorGuid()))
 						{
-							ErrorHandler->OnInvalidReferenceLevelScriptStreamed(*ActorDescView);
-							ActorDescView->SetForcedNonSpatiallyLoaded();
-						}
+							if (ActorDescView->GetIsSpatiallyLoaded())
+							{
+								ErrorHandler->OnInvalidReferenceLevelScriptStreamed(*ActorDescView);
+								ActorDescView->SetForcedNonSpatiallyLoaded();
+							}
 
-						if (ActorDescView->GetRuntimeDataLayers().Num())
-						{
-							ErrorHandler->OnInvalidReferenceLevelScriptDataLayers(*ActorDescView);
-							ActorDescView->SetInvalidDataLayers();
+							if (ActorDescView->GetRuntimeDataLayers().Num())
+							{
+								ErrorHandler->OnInvalidReferenceLevelScriptDataLayers(*ActorDescView);
+								ActorDescView->SetInvalidDataLayers();
+							}
 						}
 					}
 				}
+
+				// Validate data layers
+				if (UDataLayerSubsystem* DataLayerSubsystem = UWorld::GetSubsystem<UDataLayerSubsystem>(World))
+				{
+					DataLayerSubsystem->ForEachDataLayer([this](const UDataLayerInstance* DataLayerInstance)
+					{
+						DataLayerInstance->Validate(ErrorHandler);
+						return true;
+					});
+				}
 			}
 		}
 
@@ -698,19 +711,6 @@ class FWorldPartitionStreamingGenerator
 				ErrorHandler->OnActorNeedsResave(ActorDescView);
 			}
 		});
-
-		// Validate data layers
-		if (bIsMainContainer)
-		{
-			if (UDataLayerSubsystem* DataLayerSubsystem = UWorld::GetSubsystem<UDataLayerSubsystem>(ContainerDescriptor.Container->GetWorld()))
-			{
-				DataLayerSubsystem->ForEachDataLayer([this](const UDataLayerInstance* DataLayerInstance)
-				{
-					DataLayerInstance->Validate(ErrorHandler);
-					return true;
-				});
-			}
-		}
 	}
 
 	/** 
