Optimus: Add component provider support.

diff --git a/Engine/Plugins/Experimental/MLDeformer/LegacyVertexDeltaModel/Source/LegacyVertexDeltaModel/Private/LegacyVertexDeltaModelDataInterface.cpp b/Engine/Plugins/Experimental/MLDeformer/LegacyVertexDeltaModel/Source/LegacyVertexDeltaModel/Private/LegacyVertexDeltaModelDataInterface.cpp
--- a/Engine/Plugins/Experimental/MLDeformer/LegacyVertexDeltaModel/Source/LegacyVertexDeltaModel/Private/LegacyVertexDeltaModelDataInterface.cpp
+++ b/Engine/Plugins/Experimental/MLDeformer/LegacyVertexDeltaModel/Source/LegacyVertexDeltaModel/Private/LegacyVertexDeltaModelDataInterface.cpp
@@ -31,8 +31,8 @@ MLDEFORMER_GRAPH_IMPLEMENT_BASICS(
 
 namespace UE::LegacyVertexDeltaModel
 {
-	FLegacyVertexDeltaModelDataProviderProxy::FLegacyVertexDeltaModelDataProviderProxy(USkeletalMeshComponent* SkeletalMeshComponent, UMLDeformerComponent* DeformerComponent)
-		: UE::MLDeformer::FMLDeformerGraphDataProviderProxy(SkeletalMeshComponent, DeformerComponent)
+	FLegacyVertexDeltaModelDataProviderProxy::FLegacyVertexDeltaModelDataProviderProxy(UMLDeformerComponent* DeformerComponent)
+		: UE::MLDeformer::FMLDeformerGraphDataProviderProxy(DeformerComponent)
 	{
 		const UMLDeformerModel* Model = DeformerComponent->GetDeformerAsset()->GetModel();
 		const ULegacyVertexDeltaModel* VertexDeltaModel = Cast<ULegacyVertexDeltaModel>(Model);
diff --git a/Engine/Plugins/Experimental/MLDeformer/LegacyVertexDeltaModel/Source/LegacyVertexDeltaModel/Public/LegacyVertexDeltaModelDataInterface.h b/Engine/Plugins/Experimental/MLDeformer/LegacyVertexDeltaModel/Source/LegacyVertexDeltaModel/Public/LegacyVertexDeltaModelDataInterface.h
--- a/Engine/Plugins/Experimental/MLDeformer/LegacyVertexDeltaModel/Source/LegacyVertexDeltaModel/Public/LegacyVertexDeltaModelDataInterface.h
+++ b/Engine/Plugins/Experimental/MLDeformer/LegacyVertexDeltaModel/Source/LegacyVertexDeltaModel/Public/LegacyVertexDeltaModelDataInterface.h
@@ -56,7 +56,7 @@ namespace UE::LegacyVertexDeltaModel
 		: public UE::MLDeformer::FMLDeformerGraphDataProviderProxy
 	{
 	public:
-		FLegacyVertexDeltaModelDataProviderProxy(USkeletalMeshComponent* SkeletalMeshComponent, UMLDeformerComponent* DeformerComponent);
+		FLegacyVertexDeltaModelDataProviderProxy(UMLDeformerComponent* DeformerComponent);
 
 		// FComputeDataProviderRenderProxy overrides.
 		virtual void GatherDispatchData(FDispatchSetup const& InDispatchSetup, FCollectedDispatchData& InOutDispatchData) override;
diff --git a/Engine/Plugins/Experimental/MLDeformer/MLDeformerFramework/Source/MLDeformerFramework/Private/MLDeformerComponentSource.cpp b/Engine/Plugins/Experimental/MLDeformer/MLDeformerFramework/Source/MLDeformerFramework/Private/MLDeformerComponentSource.cpp
new file mode 100644
--- /dev/null
+++ b/Engine/Plugins/Experimental/MLDeformer/MLDeformerFramework/Source/MLDeformerFramework/Private/MLDeformerComponentSource.cpp
@@ -0,0 +1,32 @@
+﻿// Copyright Epic Games, Inc. All Rights Reserved.
+
+#include "MLDeformerComponentSource.h"
+
+#include "MLDeformerComponent.h"
+
+
+#define LOCTEXT_NAMESPACE "MLDeformerComponentSource"
+
+
+FName UMLDeformerComponentSource::Contexts::Vertex("Vertex");
+
+
+FText UMLDeformerComponentSource::GetDisplayName() const
+{
+	return LOCTEXT("MLDeformerComponent", "ML Deformer Component");
+}
+
+
+TSubclassOf<UActorComponent> UMLDeformerComponentSource::GetComponentClass() const
+{
+	return UMLDeformerComponent::StaticClass();
+}
+
+
+TArray<FName> UMLDeformerComponentSource::GetExecutionContexts() const
+{
+	return {Contexts::Vertex};
+}
+
+
+#undef LOCTEXT_NAMESPACE
diff --git a/Engine/Plugins/Experimental/MLDeformer/MLDeformerFramework/Source/MLDeformerFramework/Private/MLDeformerGraphDataInterface.cpp b/Engine/Plugins/Experimental/MLDeformer/MLDeformerFramework/Source/MLDeformerFramework/Private/MLDeformerGraphDataInterface.cpp
--- a/Engine/Plugins/Experimental/MLDeformer/MLDeformerFramework/Source/MLDeformerFramework/Private/MLDeformerGraphDataInterface.cpp
+++ b/Engine/Plugins/Experimental/MLDeformer/MLDeformerFramework/Source/MLDeformerFramework/Private/MLDeformerGraphDataInterface.cpp
@@ -23,6 +23,11 @@ TArray<FOptimusCDIPinDefinition> UMLDeformerGraphDataInterface::GetPinDefinition
 	return Defs;
 }
 
+TSubclassOf<UActorComponent> UMLDeformerGraphDataInterface::GetRequiredComponentClass() const
+{
+	return UMLDeformerComponent::StaticClass();
+}
+
 void UMLDeformerGraphDataInterface::GetSupportedInputs(TArray<FShaderFunctionDefinition>& OutFunctions) const
 {
 	OutFunctions.AddDefaulted_GetRef()
@@ -47,33 +52,9 @@ MLDEFORMER_GRAPH_IMPLEMENT_BASICS(
 	TEXT("#include \"/Plugin/MLDeformerFramework/Private/MLDeformerGraphDataInterface.ush\"\n"),
 	TEXT("ML Deformer"))
 
-UMLDeformerComponent* UMLDeformerGraphDataProvider::GetDeformerComponent() const
-{
-	UMLDeformerComponent* DeformerComponent = nullptr;
-	if (ComponentTag != NAME_None)
-	{
-		TArray<UActorComponent*> Components = SkeletalMeshComponent->GetOwner()->GetComponentsByTag(UMLDeformerComponent::StaticClass(), ComponentTag);
-		if (Components.Num())
-		{
-			DeformerComponent = Cast<UMLDeformerComponent>(Components[0]);
-		}
-	}
-	else
-	{
-		DeformerComponent = SkeletalMeshComponent->GetOwner()->FindComponentByClass<UMLDeformerComponent>();
-	}
-	
-	return DeformerComponent;
-}
 
 bool UMLDeformerGraphDataProvider::IsValid() const
 {
-	if (SkeletalMeshComponent == nullptr || SkeletalMeshComponent->MeshObject == nullptr)
-	{
-		return false;
-	}
-	
-	UMLDeformerComponent* DeformerComponent = GetDeformerComponent();
 	if (DeformerComponent == nullptr || DeformerComponent->GetDeformerAsset() == nullptr)
 	{
 		return false;
@@ -84,16 +65,15 @@ bool UMLDeformerGraphDataProvider::IsValid() const
 
 namespace UE::MLDeformer
 {
-	FMLDeformerGraphDataProviderProxy::FMLDeformerGraphDataProviderProxy(USkeletalMeshComponent* SkeletalMeshComponent, UMLDeformerComponent* DeformerComponent)
+	FMLDeformerGraphDataProviderProxy::FMLDeformerGraphDataProviderProxy(UMLDeformerComponent* DeformerComponent)
 	{
 		using namespace UE::MLDeformer;
 
-		SkeletalMeshObject = SkeletalMeshComponent->MeshObject;
-
 		const UMLDeformerAsset* DeformerAsset = DeformerComponent->GetDeformerAsset();
 		const UMLDeformerModel* Model = DeformerAsset->GetModel();
 		const UMLDeformerModelInstance* ModelInstance = DeformerComponent->GetModelInstance();
-
+		
+		SkeletalMeshObject = ModelInstance->GetSkeletalMeshComponent()->MeshObject;
 		NeuralNetwork = Model->GetNeuralNetwork();
 		NeuralNetworkInferenceHandle = ModelInstance->GetNeuralNetworkInferenceHandle();
 		bCanRunNeuralNet = ModelInstance->IsCompatible();
diff --git a/Engine/Plugins/Experimental/MLDeformer/MLDeformerFramework/Source/MLDeformerFramework/Private/MLDeformerGraphDebugDataInterface.cpp b/Engine/Plugins/Experimental/MLDeformer/MLDeformerFramework/Source/MLDeformerFramework/Private/MLDeformerGraphDebugDataInterface.cpp
--- a/Engine/Plugins/Experimental/MLDeformer/MLDeformerFramework/Source/MLDeformerFramework/Private/MLDeformerGraphDebugDataInterface.cpp
+++ b/Engine/Plugins/Experimental/MLDeformer/MLDeformerFramework/Source/MLDeformerFramework/Private/MLDeformerGraphDebugDataInterface.cpp
@@ -32,6 +32,11 @@ TArray<FOptimusCDIPinDefinition> UMLDeformerGraphDebugDataInterface::GetPinDefin
 	return Defs;
 }
 
+TSubclassOf<UActorComponent> UMLDeformerGraphDebugDataInterface::GetRequiredComponentClass() const
+{
+	return UMLDeformerComponent::StaticClass();
+}
+
 void UMLDeformerGraphDebugDataInterface::GetSupportedInputs(TArray<FShaderFunctionDefinition>& OutFunctions) const
 {
 	OutFunctions.AddDefaulted_GetRef()
@@ -71,12 +76,6 @@ MLDEFORMER_GRAPH_IMPLEMENT_DEBUG_BASICS_WITH_PROXY(
 bool UMLDeformerGraphDebugDataProvider::IsValid() const
 {
 #if WITH_EDITORONLY_DATA
-	if (SkeletalMeshComponent == nullptr || SkeletalMeshComponent->MeshObject == nullptr)
-	{
-		return false;
-	}
-	
-	UMLDeformerComponent* DeformerComponent = SkeletalMeshComponent->GetOwner()->FindComponentByClass<UMLDeformerComponent>();
 	if (DeformerComponent == nullptr || DeformerComponent->GetDeformerAsset() == nullptr)
 	{
 		return false;
@@ -91,15 +90,16 @@ bool UMLDeformerGraphDebugDataProvider::IsValid() const
 #if WITH_EDITORONLY_DATA
 namespace UE::MLDeformer
 {
-	FMLDeformerGraphDebugDataProviderProxy::FMLDeformerGraphDebugDataProviderProxy(USkeletalMeshComponent* SkeletalMeshComponent, UMLDeformerAsset* DeformerAsset, UMLDeformerGraphDebugDataProvider* InProvider)
+	FMLDeformerGraphDebugDataProviderProxy::FMLDeformerGraphDebugDataProviderProxy(UMLDeformerComponent* DeformerComponent, UMLDeformerAsset* DeformerAsset, UMLDeformerGraphDebugDataProvider* InProvider)
 		: FComputeDataProviderRenderProxy()
 	{
-		SkeletalMeshObject = SkeletalMeshComponent->MeshObject;
 		Provider = InProvider;
 
 		UMLDeformerModel* Model = DeformerAsset->GetModel();	
 		UMLDeformerVizSettings* VizSettings = Model->GetVizSettings();
+		const UMLDeformerModelInstance* ModelInstance = DeformerComponent->GetModelInstance();
 
+		SkeletalMeshObject = ModelInstance->GetSkeletalMeshComponent()->MeshObject;
 		VertexMapBufferSRV = Model->GetVertexMapBuffer().ShaderResourceViewRHI;
 		HeatMapMode = (int32)VizSettings->GetHeatMapMode();
 		HeatMapMax = 1.0f / FMath::Max(VizSettings->GetHeatMapMax(), 0.00001f);
diff --git a/Engine/Plugins/Experimental/MLDeformer/MLDeformerFramework/Source/MLDeformerFramework/Public/MLDeformerGraphDataInterface.h b/Engine/Plugins/Experimental/MLDeformer/MLDeformerFramework/Source/MLDeformerFramework/Public/MLDeformerGraphDataInterface.h
--- a/Engine/Plugins/Experimental/MLDeformer/MLDeformerFramework/Source/MLDeformerFramework/Public/MLDeformerGraphDataInterface.h
+++ b/Engine/Plugins/Experimental/MLDeformer/MLDeformerFramework/Source/MLDeformerFramework/Public/MLDeformerGraphDataInterface.h
@@ -63,13 +63,12 @@ class UMLDeformerModel;
 	UComputeDataProvider* InterfaceClassName::CreateDataProvider(TObjectPtr<UObject> InBinding, uint64 InInputMask, uint64 InOutputMask) const \
 	{ \
 		DataProviderClassName* Provider = NewObject<DataProviderClassName>(); \
-		Provider->SkeletalMeshComponent = Cast<USkeletalMeshComponent>(InBinding); \
-		Provider->ComponentTag = ComponentTag; \
+		Provider->DeformerComponent = Cast<UMLDeformerComponent>(InBinding); \
 		return Provider; \
 	} \
 	FComputeDataProviderRenderProxy* DataProviderClassName::GetRenderProxy() \
 	{ \
-		return new DataProviderProxyClassName(SkeletalMeshComponent, GetDeformerComponent()); \
+		return new DataProviderProxyClassName(DeformerComponent); \
 	}
 
 #if WITH_EDITORONLY_DATA
@@ -89,6 +88,7 @@ public:
 	// UOptimusComputeDataInterface overrides.
 	virtual FString GetDisplayName() const override;
 	virtual TArray<FOptimusCDIPinDefinition> GetPinDefinitions() const override;
+	virtual TSubclassOf<UActorComponent> GetRequiredComponentClass() const override;
 	// ~END UOptimusComputeDataInterface overrides.
 
 	// UComputeDataInterface overrides.
@@ -98,10 +98,6 @@ public:
 	virtual void GetHLSL(FString& OutHLSL) const override;
 	virtual UComputeDataProvider* CreateDataProvider(TObjectPtr<UObject> InBinding, uint64 InInputMask, uint64 InOutputMask) const override;
 	// ~END UComputeDataInterface overrides.
-
-	/** When set we bind only with to a MLDeformer Component with the same tag. */
-	UPROPERTY(EditAnywhere, Category = Binding)
-	FName ComponentTag;
 };
 
 /** Compute Framework Data Provider for MLDeformer data. */
@@ -113,17 +109,12 @@ class MLDEFORMERFRAMEWORK_API UMLDeformerGraphDataProvider
 
 public:
 	UPROPERTY(EditAnywhere, BlueprintReadWrite, Category = Binding)
-	TObjectPtr<USkeletalMeshComponent> SkeletalMeshComponent = nullptr;
-
-	FName ComponentTag;
+	TObjectPtr<UMLDeformerComponent> DeformerComponent = nullptr;
 
 	// UComputeDataProvider overrides.
 	virtual bool IsValid() const override;
 	virtual FComputeDataProviderRenderProxy* GetRenderProxy() override;
 	// ~END UComputeDataProvider overrides.
-
-protected:
-	UMLDeformerComponent* GetDeformerComponent() const;
 };
 
 namespace UE::MLDeformer
@@ -133,7 +124,7 @@ namespace UE::MLDeformer
 		: public FComputeDataProviderRenderProxy
 	{
 	public:
-		FMLDeformerGraphDataProviderProxy(USkeletalMeshComponent* SkeletalMeshComponent, UMLDeformerComponent* DeformerComponent);
+		FMLDeformerGraphDataProviderProxy(UMLDeformerComponent* DeformerComponent);
 
 		// FComputeDataProviderRenderProxy overrides.
 		virtual void AllocateResources(FRDGBuilder& GraphBuilder) override;
diff --git a/Engine/Plugins/Experimental/MLDeformer/MLDeformerFramework/Source/MLDeformerFramework/Public/MLDeformerGraphDebugDataInterface.h b/Engine/Plugins/Experimental/MLDeformer/MLDeformerFramework/Source/MLDeformerFramework/Public/MLDeformerGraphDebugDataInterface.h
--- a/Engine/Plugins/Experimental/MLDeformer/MLDeformerFramework/Source/MLDeformerFramework/Public/MLDeformerGraphDebugDataInterface.h
+++ b/Engine/Plugins/Experimental/MLDeformer/MLDeformerFramework/Source/MLDeformerFramework/Public/MLDeformerGraphDebugDataInterface.h
@@ -52,12 +52,10 @@ class USkeletalMeshComponent;
 	UComputeDataProvider* InterfaceClassName::CreateDataProvider(TObjectPtr<UObject> InBinding, uint64 InInputMask, uint64 InOutputMask) const \
 	{ \
 		DataProviderClassName* Provider = NewObject<DataProviderClassName>(); \
-		Provider->SkeletalMeshComponent = Cast<USkeletalMeshComponent>(InBinding); \
-		if (Provider->SkeletalMeshComponent) \
+		Provider->DeformerComponent = Cast<UMLDeformerComponent>(InBinding); \
+		if (Provider->DeformerComponent) \
 		{ \
-			AActor* Actor = Provider->SkeletalMeshComponent->GetOwner(); \
-			UMLDeformerComponent* DeformerComponent = (Actor != nullptr) ? Cast<UMLDeformerComponent>(Actor->GetComponentByClass(UMLDeformerComponent::StaticClass())) : nullptr; \
-			Provider->DeformerAsset = (DeformerComponent != nullptr) ? DeformerComponent->GetDeformerAsset() : nullptr; \
+			Provider->DeformerAsset = Provider->DeformerComponent->GetDeformerAsset(); \
 		} \
 		MLDEFORMER_EDITORDATA_ONLY( \
 			if (Provider->DeformerAsset != nullptr) \
@@ -74,8 +72,8 @@ class USkeletalMeshComponent;
 	FComputeDataProviderRenderProxy* DataProviderClassName::GetRenderProxy() \
 	{ \
 		MLDEFORMER_EDITORDATA_ONLY( \
-			DataProviderProxyClassName* Proxy = new DataProviderProxyClassName(SkeletalMeshComponent, DeformerAsset, this); \
-			const float SampleTime = SkeletalMeshComponent->GetPosition(); \
+			DataProviderProxyClassName* Proxy = new DataProviderProxyClassName(DeformerComponent, DeformerAsset, this); \
+			const float SampleTime = DeformerComponent->GetModelInstance()->GetSkeletalMeshComponent()->GetPosition(); \
 			UMLDeformerModel* Model = DeformerAsset->GetModel(); \
 			Model->SampleGroundTruthPositions(SampleTime, Proxy->GetGroundTruthPositions()); \
 			Proxy->HandleZeroGroundTruthPositions(); \
@@ -98,6 +96,7 @@ public:
 	// UOptimusComputeDataInterface overrides.
 	virtual FString GetDisplayName() const override;
 	virtual TArray<FOptimusCDIPinDefinition> GetPinDefinitions() const override;
+	virtual TSubclassOf<UActorComponent> GetRequiredComponentClass() const override;
 	// ~END UOptimusComputeDataInterface overrides.
 
 	// UComputeDataInterface overrides.
@@ -126,7 +125,7 @@ public:
 
 public:
 	UPROPERTY(EditAnywhere, BlueprintReadWrite, Category = Binding)
-	TObjectPtr<USkeletalMeshComponent> SkeletalMeshComponent = nullptr;
+	TObjectPtr<UMLDeformerComponent> DeformerComponent = nullptr;
 
 	UPROPERTY(EditAnywhere, BlueprintReadWrite, Category = Binding)
 	TObjectPtr<UMLDeformerAsset> DeformerAsset = nullptr;
@@ -140,7 +139,7 @@ namespace UE::MLDeformer
 		: public FComputeDataProviderRenderProxy
 	{
 	public:
-		FMLDeformerGraphDebugDataProviderProxy(USkeletalMeshComponent* SkeletalMeshComponent, UMLDeformerAsset* DeformerAsset, UMLDeformerGraphDebugDataProvider* InProvider);
+		FMLDeformerGraphDebugDataProviderProxy(UMLDeformerComponent* DeformerComponent, UMLDeformerAsset* DeformerAsset, UMLDeformerGraphDebugDataProvider* InProvider);
 
 		virtual void HandleZeroGroundTruthPositions();
 
diff --git a/Engine/Source/Editor/UnrealEd/Private/Text/HLSLSyntaxHighlighterMarshaller.cpp b/Engine/Source/Editor/UnrealEd/Private/Text/HLSLSyntaxHighlighterMarshaller.cpp
--- a/Engine/Source/Editor/UnrealEd/Private/Text/HLSLSyntaxHighlighterMarshaller.cpp
+++ b/Engine/Source/Editor/UnrealEd/Private/Text/HLSLSyntaxHighlighterMarshaller.cpp
@@ -55,6 +55,7 @@ const TCHAR* HlslKeywords[] =
 	TEXT("packoffset"),
 	TEXT("numthreads"),
 	TEXT("nointerpolation"),
+	TEXT("namespace"),
 	TEXT("matrix"),
 	TEXT("int4x4"),
 	TEXT("int4x3"),
diff --git a/Engine/Source/Runtime/Engine/Classes/Animation/MeshDeformer.h b/Engine/Source/Runtime/Engine/Classes/Animation/MeshDeformer.h
--- a/Engine/Source/Runtime/Engine/Classes/Animation/MeshDeformer.h
+++ b/Engine/Source/Runtime/Engine/Classes/Animation/MeshDeformer.h
@@ -7,6 +7,7 @@
 
 class UMeshComponent;
 class UMeshDeformerInstance;
+class UMeshDeformerInstanceSettings;
 
 /** 
  * Base class for mesh deformer assets.
@@ -19,6 +20,17 @@ class ENGINE_API UMeshDeformer : public UObject
 	GENERATED_BODY()
 
 public:
+	/** Create persistent settings for the mesh deformer instance */ 
+	virtual UMeshDeformerInstanceSettings* CreateSettingsInstance(
+		UMeshComponent* InMeshComponent
+		)
+	{
+		return nullptr;
+	}
+	
 	/** Create an instance of the mesh deformer. */
-	virtual UMeshDeformerInstance* CreateInstance(UMeshComponent* InMeshComponent) PURE_VIRTUAL(, return nullptr;);
+	virtual UMeshDeformerInstance* CreateInstance(
+		UMeshComponent* InMeshComponent,
+		UMeshDeformerInstanceSettings* InSettings
+		) PURE_VIRTUAL(, return nullptr;);
 };
diff --git a/Engine/Source/Runtime/Engine/Classes/Animation/MeshDeformerInstance.h b/Engine/Source/Runtime/Engine/Classes/Animation/MeshDeformerInstance.h
--- a/Engine/Source/Runtime/Engine/Classes/Animation/MeshDeformerInstance.h
+++ b/Engine/Source/Runtime/Engine/Classes/Animation/MeshDeformerInstance.h
@@ -7,6 +7,12 @@
 
 class FSceneInterface;
 
+UCLASS(Abstract, EditInlineNew)
+class ENGINE_API UMeshDeformerInstanceSettings : public UObject
+{
+	GENERATED_BODY()
+};
+
 /** 
  * Base class for mesh deformers instances.
  * This contains the per instance state for a UMeshDeformer.
diff --git a/Engine/Source/Runtime/Engine/Classes/Components/SkinnedMeshComponent.h b/Engine/Source/Runtime/Engine/Classes/Components/SkinnedMeshComponent.h
--- a/Engine/Source/Runtime/Engine/Classes/Components/SkinnedMeshComponent.h
+++ b/Engine/Source/Runtime/Engine/Classes/Components/SkinnedMeshComponent.h
@@ -28,6 +28,7 @@ struct FSkelMeshRenderSection;
 class FPositionVertexBuffer;
 class UMeshDeformer;
 class UMeshDeformerInstance;
+class UMeshDeformerInstanceSettings;
 class USkinnedAsset;
 
 DECLARE_DELEGATE_OneParam(FOnAnimUpdateRateParamsCreated, FAnimUpdateRateParameters*)
@@ -294,6 +295,9 @@ public:
 	UPROPERTY(Transient, BlueprintReadOnly, Category = "Deformer")
 	TObjectPtr<UMeshDeformerInstance> MeshDeformerInstance;
 
+	UPROPERTY(VisibleAnywhere, BlueprintReadOnly, DisplayName = "Settings", Category = "Deformer", meta=(EditInline, EditCondition="MeshDeformerInstanceSettings!=nullptr", HideEditConditionToggle, EditConditionHides))
+	TObjectPtr<UMeshDeformerInstanceSettings> MeshDeformerInstanceSettings;
+
 	/** const getters for previous transform idea */
 	const TArray<uint8>& GetPreviousBoneVisibilityStates() const
 	{
@@ -964,6 +968,7 @@ public:
 	virtual FString GetDetailedInfoInternal() const override;
 #if WITH_EDITOR
 	virtual bool CanEditChange(const FProperty* InProperty) const override;
+	virtual void PostEditChangeProperty(FPropertyChangedEvent& PropertyChangedEvent) override;
 #endif // WITH_EDITOR
 
 protected:
diff --git a/Engine/Source/Runtime/Engine/Private/Components/SkinnedMeshComponent.cpp b/Engine/Source/Runtime/Engine/Private/Components/SkinnedMeshComponent.cpp
--- a/Engine/Source/Runtime/Engine/Private/Components/SkinnedMeshComponent.cpp
+++ b/Engine/Source/Runtime/Engine/Private/Components/SkinnedMeshComponent.cpp
@@ -142,7 +142,7 @@ namespace FAnimUpdateRateManager
 	{
 		if (!SkinnedComponent)
 		{
-			return NULL;
+			return nullptr;
 		}
 		UObject* TrackerIndex = GetMapIndexForComponent(SkinnedComponent);
 
@@ -561,7 +561,7 @@ void USkinnedMeshComponent::OnRegister()
 	UpdateLODStatus();
 	InvalidateCachedBounds();
 
-	MeshDeformerInstance = (MeshDeformer != nullptr) ? MeshDeformer->CreateInstance(this) : nullptr;
+	MeshDeformerInstance = (MeshDeformer != nullptr) ? MeshDeformer->CreateInstance(this, MeshDeformerInstanceSettings) : nullptr;
 
 	RefreshExternalMorphTargetWeights();
 }
@@ -768,7 +768,7 @@ void USkinnedMeshComponent::DestroyRenderState_Concurrent()
 		// Begin a deferred delete of MeshObject.  BeginCleanup will call MeshObject->FinishDestroy after the above release resource
 		// commands execute in the rendering thread.
 		BeginCleanup(MeshObject);
-		MeshObject = NULL;
+		MeshObject = nullptr;
 	}
 }
 
@@ -790,7 +790,7 @@ FString USkinnedMeshComponent::GetDetailedInfoInternal() const
 {
 	FString Result;  
 
-	if(GetSkinnedAsset() != NULL )
+	if(GetSkinnedAsset() != nullptr )
 	{
 		Result = GetSkinnedAsset()->GetDetailedInfoInternal();
 	}
@@ -932,11 +932,31 @@ bool USkinnedMeshComponent::CanEditChange(const FProperty* InProperty) const
 	return Super::CanEditChange(InProperty);
 }
 
+void USkinnedMeshComponent::PostEditChangeProperty(FPropertyChangedEvent& PropertyChangedEvent)
+{
+	Super::PostEditChangeProperty(PropertyChangedEvent);
+
+	if (const FProperty* Property = PropertyChangedEvent.Property)
+	{
+		if ( Property->GetFName() == GET_MEMBER_NAME_CHECKED( USkinnedMeshComponent, MeshDeformer ) )
+		{
+			if (MeshDeformer)
+			{
+				MeshDeformerInstanceSettings = MeshDeformer->CreateSettingsInstance(this);
+			}
+			else
+			{
+				MeshDeformerInstanceSettings = nullptr;
+			}
+		}
+	}
+}
+
 #endif // WITH_EDITOR
 
 void USkinnedMeshComponent::InitLODInfos()
 {
-	if (GetSkinnedAsset() != NULL)
+	if (GetSkinnedAsset() != nullptr)
 	{
 		if (GetSkinnedAsset()->GetLODNum() != LODInfo.Num())
 		{
@@ -1217,7 +1237,7 @@ bool USkinnedMeshComponent::ShouldUpdateBoneVisibility() const
 }
 void USkinnedMeshComponent::RebuildVisibilityArray()
 {
-	// BoneVisibility needs update if LeaderComponent == NULL
+	// BoneVisibility needs update if LeaderComponent == nullptr
 	// if MaterComponent, it should follow MaterPoseComponent
 	if ( ShouldUpdateBoneVisibility())
 	{
@@ -1315,7 +1335,7 @@ FBoxSphereBounds USkinnedMeshComponent::CalcMeshBound(const FVector3f& RootOffse
 
 	// Can only use the PhysicsAsset to calculate the bounding box if we are not non-uniformly scaling the mesh.
 	const USkinnedAsset* SkinnedAssetConst = GetSkinnedAsset();
-	const bool bCanUsePhysicsAsset = DrawScale.IsUniform() && (GetSkinnedAsset() != NULL)
+	const bool bCanUsePhysicsAsset = DrawScale.IsUniform() && (GetSkinnedAsset() != nullptr)
 		// either space base exists or child component
 		&& ( (GetNumComponentSpaceTransforms() == SkinnedAssetConst->GetRefSkeleton().GetNum()) || (LeaderPhysicsAsset) );
 
@@ -1554,7 +1574,7 @@ int32 USkinnedMeshComponent::GetBoneIndex( FName BoneName) const
 
 FName USkinnedMeshComponent::GetBoneName(int32 BoneIndex) const
 {
-	return (GetSkinnedAsset() != NULL && GetSkinnedAsset()->GetRefSkeleton().IsValidIndex(BoneIndex)) ? GetSkinnedAsset()->GetRefSkeleton().GetBoneName(BoneIndex) : NAME_None;
+	return (GetSkinnedAsset() != nullptr && GetSkinnedAsset()->GetRefSkeleton().IsValidIndex(BoneIndex)) ? GetSkinnedAsset()->GetRefSkeleton().GetBoneName(BoneIndex) : NAME_None;
 }
 
 
@@ -1684,7 +1704,7 @@ bool USkinnedMeshComponent::IsSkinCacheAllowed(int32 LodIdx) const
 
 void USkinnedMeshComponent::GetBoneNames(TArray<FName>& BoneNames)
 {
-	if (GetSkinnedAsset() == NULL)
+	if (GetSkinnedAsset() == nullptr)
 	{
 		// no mesh, so no bones
 		BoneNames.Empty();
@@ -1765,7 +1785,7 @@ USkeletalMesh* USkinnedMeshComponent::GetSkeletalMesh_DEPRECATED() const
 
 void USkinnedMeshComponent::SetSkinnedAsset(USkinnedAsset* InSkinnedAsset, bool /*bReinitPose*/)
 {
-	// NOTE: InSkinnedAsset may be NULL (useful in the editor for removing the skeletal mesh associated with
+	// NOTE: InSkinnedAsset may be nullptr (useful in the editor for removing the skeletal mesh associated with
 	//   this component on-the-fly)
 
 	if (InSkinnedAsset == GetSkinnedAsset())
@@ -1838,7 +1858,8 @@ USkinnedAsset* USkinnedMeshComponent::GetSkinnedAsset() const
 void USkinnedMeshComponent::SetMeshDeformer(UMeshDeformer* InMeshDeformer)
 {
 	MeshDeformer = InMeshDeformer;
-	MeshDeformerInstance = (MeshDeformer != nullptr) ? MeshDeformer->CreateInstance(this) : nullptr;
+	MeshDeformerInstanceSettings = MeshDeformer->CreateSettingsInstance(this);
+	MeshDeformerInstance = (MeshDeformer != nullptr) ? MeshDeformer->CreateInstance(this, MeshDeformerInstanceSettings) : nullptr;
 	MarkRenderDynamicDataDirty();
 }
 
@@ -1854,7 +1875,7 @@ FSkeletalMeshRenderData* USkinnedMeshComponent::GetSkeletalMeshRenderData() cons
 	}
 	else
 	{
-		return NULL;
+		return nullptr;
 	}
 }
 
@@ -1863,7 +1884,7 @@ bool USkinnedMeshComponent::AllocateTransformData()
 	LLM_SCOPE_BYNAME(TEXT("SkeletalMesh/TransformData"));
 
 	// Allocate transforms if not present.
-	if (GetSkinnedAsset() != NULL && LeaderPoseComponent == NULL )
+	if (GetSkinnedAsset() != nullptr && LeaderPoseComponent == nullptr )
 	{
 		const int32 NumBones = GetSkinnedAsset()->GetRefSkeleton().GetNum();
 
@@ -2155,12 +2176,12 @@ void USkinnedMeshComponent::SetSelectedEditorMaterial(int32 InSelectedEditorMate
 
 UMorphTarget* USkinnedMeshComponent::FindMorphTarget( FName MorphTargetName ) const
 {
-	if(GetSkinnedAsset() != NULL )
+	if(GetSkinnedAsset() != nullptr )
 	{
 		return GetSkinnedAsset()->FindMorphTarget(MorphTargetName);
 	}
 
-	return NULL;
+	return nullptr;
 }
 
 bool USkinnedMeshComponent::GetMissingLeaderBoneRelativeTransform(int32 InBoneIndex, FMissingLeaderBoneCacheEntry& OutInfo) const
@@ -2555,7 +2576,7 @@ FVector USkinnedMeshComponent::GetBoneAxis( FName BoneName, EAxis::Type Axis ) c
 
 bool USkinnedMeshComponent::HasAnySockets() const
 {
-	return (GetSkinnedAsset() != NULL) && (
+	return (GetSkinnedAsset() != nullptr) && (
 #if WITH_EDITOR
 		(GetSkinnedAsset()->GetActiveSocketList().Num() > 0) ||
 #endif
@@ -2564,7 +2585,7 @@ bool USkinnedMeshComponent::HasAnySockets() const
 
 void USkinnedMeshComponent::QuerySupportedSockets(TArray<FComponentSocketDescription>& OutSockets) const
 {
-	if (GetSkinnedAsset() != NULL)
+	if (GetSkinnedAsset() != nullptr)
 	{
 		// Grab all the mesh and skeleton sockets
 		const TArray<USkeletalMeshSocket*> AllSockets = GetSkinnedAsset()->GetActiveSocketList();
@@ -2652,9 +2673,9 @@ void USkinnedMeshComponent::TransformFromBoneSpace(FName BoneName, FVector InPos
 
 FName USkinnedMeshComponent::FindClosestBone(FVector TestLocation, FVector* BoneLocation, float IgnoreScale, bool bRequirePhysicsAsset) const
 {
-	if (GetSkinnedAsset() == NULL)
+	if (GetSkinnedAsset() == nullptr)
 	{
-		if (BoneLocation != NULL)
+		if (BoneLocation != nullptr)
 		{
 			*BoneLocation = FVector::ZeroVector;
 		}
@@ -2666,7 +2687,7 @@ FName USkinnedMeshComponent::FindClosestBone(FVector TestLocation, FVector* Bone
 		const UPhysicsAsset* PhysAsset = GetPhysicsAsset();
 		if (bRequirePhysicsAsset && !PhysAsset)
 		{
-			if (BoneLocation != NULL)
+			if (BoneLocation != nullptr)
 			{
 				*BoneLocation = FVector::ZeroVector;
 			}
@@ -2706,7 +2727,7 @@ FName USkinnedMeshComponent::FindClosestBone(FVector TestLocation, FVector* Bone
 
 		if (BestIndex == -1)
 		{
-			if (BoneLocation != NULL)
+			if (BoneLocation != nullptr)
 			{
 				*BoneLocation = FVector::ZeroVector;
 			}
@@ -2715,7 +2736,7 @@ FName USkinnedMeshComponent::FindClosestBone(FVector TestLocation, FVector* Bone
 		else
 		{
 			// transform the bone location into world space
-			if (BoneLocation != NULL)
+			if (BoneLocation != nullptr)
 			{
 				*BoneLocation = (CompSpaceTransforms[BestIndex] * GetComponentTransform()).GetLocation();
 			}
