Add a "ShouldFlushKeysWhenViewportFocusChanges" function to the player controller, which allows the GameViewportClient to check if the PlayerController keys need to be flushed when the viewport changes depending on the Input Mode you are in. For example, if you are set to use the "Game And UI" input mode then you don't want to flush the player controller's keys because then it would stop any active movement of the player.

One example of this in a game would be a third person game where you can walk around and press "Tab" to bring up some piece of UI that you can use while you walk around. As it was, when you pressed tab your player would stop moving if you were holding down "W" to walk forward because the viewport client would change focus, calling FlushPressedKeys. Now that wouldn't happen! Woot

diff --git a/Engine/Source/Runtime/Engine/Classes/GameFramework/InputSettings.h b/Engine/Source/Runtime/Engine/Classes/GameFramework/InputSettings.h
--- a/Engine/Source/Runtime/Engine/Classes/GameFramework/InputSettings.h
+++ b/Engine/Source/Runtime/Engine/Classes/GameFramework/InputSettings.h
@@ -69,6 +69,16 @@ class ENGINE_API UInputSettings
 	 */
 	UPROPERTY(config, EditAnywhere, Category = "Input")
 	uint8 bFilterInputByPlatformUser:1;
+
+	/**
+	 * If true, then the Player Controller will have it's Pressed Keys flushed when the input mode is changed
+	 * to Game and UI mode or the game viewport loses focus. The default behavior is true.
+	 * 
+	 * @see UGameViewportClient::LostFocus
+	 * @see APlayerController::ShouldFlushKeysWhenViewportFocusChanges
+	 */
+	UPROPERTY(config, EditAnywhere, Category = "Input")
+	uint8 bShouldFlushPressedKeysOnViewportFocusLost:1;
 	
 	/** Should the touch input interface be shown always, or only when the platform has a touch screen? */
 	UPROPERTY(config, EditAnywhere, Category="Mobile")
diff --git a/Engine/Source/Runtime/Engine/Classes/GameFramework/PlayerController.h b/Engine/Source/Runtime/Engine/Classes/GameFramework/PlayerController.h
--- a/Engine/Source/Runtime/Engine/Classes/GameFramework/PlayerController.h
+++ b/Engine/Source/Runtime/Engine/Classes/GameFramework/PlayerController.h
@@ -111,6 +111,8 @@ protected:
 	/** Derived classes override this function to apply the necessary settings for the desired input mode */
 	virtual void ApplyInputMode(class FReply& SlateOperations, class UGameViewportClient& GameViewportClient) const = 0;
 
+	virtual bool ShouldFlushInputOnViewportFocus() const { return true; };
+
 	/** Utility functions for derived classes. */
 	void SetFocusAndLocking(FReply& SlateOperations, TSharedPtr<class SWidget> InWidgetToFocus, bool bLockMouseToViewport, TSharedRef<class SViewport> InViewportWidget) const;
 
@@ -177,6 +179,8 @@ struct ENGINE_API FInputModeGameAndUI : public FInputModeDataBase
 	/** Whether to hide the cursor during temporary mouse capture caused by a mouse down */
 	FInputModeGameAndUI& SetHideCursorDuringCapture(bool InHideCursorDuringCapture) { bHideCursorDuringCapture = InHideCursorDuringCapture; return *this; }
 
+	virtual bool ShouldFlushInputOnViewportFocus() const override { return false; }
+
 	FInputModeGameAndUI()
 		: WidgetToFocus()
 		, MouseLockMode(EMouseLockMode::DoNotLock)
@@ -1675,6 +1679,12 @@ public:
 	/** Flushes the current key state. */
 	virtual void FlushPressedKeys();
 
+	/**
+	 * If true, then the GameViewportClient should call FlushPressedKeys on this controller when it loses focus.
+	 * The default behavior here is to return true if the PlayerController is in any input mode other than GameAndUI
+	 */
+	virtual bool ShouldFlushKeysWhenViewportFocusChanges() const { return bShouldFlushInputWhenViewportFocusChanges; }
+
 	UFUNCTION(BlueprintCallable, Category = Input)
 	TSubclassOf<UPlayerInput> GetOverridePlayerInputClass() const;
 
@@ -2139,6 +2149,14 @@ private:
 	/** If true, prevent any haptic effects from playing */
 	bool bDisableHaptics : 1;
 
+	/**
+	 * If true, then the GameViewportCliet will call FlushPressedKeys when it loses focus ( UGameViewportClient::LostFocus ).
+	 * By default this is true for all Input Modes except for GameAndUI. When you are using the "GameAndUI" input mode 
+	 * then you don't normally want the input to be flushed because it would reset the player controller's inputs if you bring
+	 * up any in-game UI.
+	 */
+	bool bShouldFlushInputWhenViewportFocusChanges : 1;
+
 public: 
 
 	/**
diff --git a/Engine/Source/Runtime/Engine/Private/GameViewportClient.cpp b/Engine/Source/Runtime/Engine/Private/GameViewportClient.cpp
--- a/Engine/Source/Runtime/Engine/Private/GameViewportClient.cpp
+++ b/Engine/Source/Runtime/Engine/Private/GameViewportClient.cpp
@@ -2084,12 +2084,14 @@ void UGameViewportClient::LostFocus(FViewport* InViewport)
 	// We need to reset some key inputs, since keyup events will sometimes not be processed (such as going into immersive/maximized mode).
 	// Resetting them will prevent them from "sticking"
 	UWorld* const ViewportWorld = GetWorld();
+	const bool bShouldFlush = GetDefault<UInputSettings>()->bShouldFlushPressedKeysOnViewportFocusLost;
+	
 	if (ViewportWorld && !ViewportWorld->bIsTearingDown)
 	{
 		for (FConstPlayerControllerIterator Iterator = ViewportWorld->GetPlayerControllerIterator(); Iterator; ++Iterator)
 		{
 			APlayerController* const PlayerController = Iterator->Get();
-			if (PlayerController)
+			if (PlayerController && (bShouldFlush || PlayerController->ShouldFlushKeysWhenViewportFocusChanges()))
 			{
 				PlayerController->FlushPressedKeys();
 			}
diff --git a/Engine/Source/Runtime/Engine/Private/PlayerController.cpp b/Engine/Source/Runtime/Engine/Private/PlayerController.cpp
--- a/Engine/Source/Runtime/Engine/Private/PlayerController.cpp
+++ b/Engine/Source/Runtime/Engine/Private/PlayerController.cpp
@@ -154,7 +154,8 @@ APlayerController::APlayerController(const FObjectInitializer& ObjectInitializer
 
 	bIsPlayerController = true;
 	bIsLocalPlayerController = false;
-	bDisableHaptics = false;	
+	bDisableHaptics = false;
+	bShouldFlushInputWhenViewportFocusChanges = true;
 
 	ClickEventKeys.Add(EKeys::LeftMouseButton);
 
@@ -5739,6 +5740,7 @@ void APlayerController::SetInputMode(const FInputModeDataBase& InData)
 	if ( GameViewportClient && LocalPlayer )
 	{
 		InData.ApplyInputMode( LocalPlayer->GetSlateOperations(), *GameViewportClient );
+		bShouldFlushInputWhenViewportFocusChanges = InData.ShouldFlushInputOnViewportFocus();
 	}
 }
 
diff --git a/Engine/Source/Runtime/Engine/Private/UserInterface/InputSettings.cpp b/Engine/Source/Runtime/Engine/Private/UserInterface/InputSettings.cpp
--- a/Engine/Source/Runtime/Engine/Private/UserInterface/InputSettings.cpp
+++ b/Engine/Source/Runtime/Engine/Private/UserInterface/InputSettings.cpp
@@ -24,6 +24,7 @@ UInputSettings::UInputSettings(const FObjectInitializer& ObjectInitializer)
 	, bEnableLegacyInputScales(true)
 	, bEnableMotionControls(true)
 	, bFilterInputByPlatformUser(false)
+	, bShouldFlushPressedKeysOnViewportFocusLost(true)
 	, DefaultViewportMouseCaptureMode(EMouseCaptureMode::CapturePermanently_IncludingInitialMouseDown)
 	, DefaultViewportMouseLockMode(EMouseLockMode::LockOnCapture)
 	, DefaultPlayerInputClass(UPlayerInput::StaticClass())
