Pose Asset name change from review feedback meant to be in previous change.

diff --git a/Engine/Source/Runtime/Engine/Classes/Animation/PoseAsset.h b/Engine/Source/Runtime/Engine/Classes/Animation/PoseAsset.h
--- a/Engine/Source/Runtime/Engine/Classes/Animation/PoseAsset.h
+++ b/Engine/Source/Runtime/Engine/Classes/Animation/PoseAsset.h
@@ -286,7 +286,7 @@ public:
 	ENGINE_API bool GetFullPose(int32 PoseIndex, TArray<FTransform>& OutTransforms) const;
 	
 	// util to return transform of a bone from the pose asset in component space, by walking up tracks in pose asset */
-	ENGINE_API FTransform GetComponentSpaceTransform(FName TrackName, const TArray<FTransform>& LocalTransforms) const;
+	ENGINE_API FTransform GetComponentSpaceTransform(FName BoneName, const TArray<FTransform>& LocalTransforms) const;
 
 	ENGINE_API int32 DeletePoses(TArray<FName> PoseNamesToDelete);
 	ENGINE_API int32 DeleteCurves(TArray<FName> CurveNamesToDelete);
diff --git a/Engine/Source/Runtime/Engine/Private/Animation/PoseAsset.cpp b/Engine/Source/Runtime/Engine/Private/Animation/PoseAsset.cpp
--- a/Engine/Source/Runtime/Engine/Private/Animation/PoseAsset.cpp
+++ b/Engine/Source/Runtime/Engine/Private/Animation/PoseAsset.cpp
@@ -1625,7 +1625,7 @@ bool UPoseAsset::GetFullPose(int32 PoseIndex, TArray<FTransform>& OutTransforms)
 	return true;
 }
 
-FTransform UPoseAsset::GetComponentSpaceTransform(FName TrackName, const TArray<FTransform>& LocalTransforms) const
+FTransform UPoseAsset::GetComponentSpaceTransform(FName BoneName, const TArray<FTransform>& LocalTransforms) const
 {
 	const FReferenceSkeleton& RefSkel = GetSkeleton()->GetReferenceSkeleton();
 
@@ -1633,11 +1633,11 @@ FTransform UPoseAsset::GetComponentSpaceTransform(FName TrackName, const TArray<
 	FTransform ComponentSpaceTransform = FTransform::Identity;
 
 	// Start to walk up parent chain until we reach root (ParentIndex == INDEX_NONE)
-	int32 BoneIndex = RefSkel.FindBoneIndex(TrackName);
+	int32 BoneIndex = RefSkel.FindBoneIndex(BoneName);
 	while (BoneIndex != INDEX_NONE)
 	{
-		TrackName = RefSkel.GetBoneName(BoneIndex);
-		const int32 TrackIndex = GetTrackIndexByName(TrackName);
+		BoneName = RefSkel.GetBoneName(BoneIndex);
+		int32 TrackIndex = GetTrackIndexByName(BoneName);
 
 		// If a track for parent, get local space transform from that
 		// If not, get from ref pose
