Make GetPrimaryPlayerController more resilient to crash

diff --git a/Engine/Source/Runtime/Engine/Private/GameInstance.cpp b/Engine/Source/Runtime/Engine/Private/GameInstance.cpp
--- a/Engine/Source/Runtime/Engine/Private/GameInstance.cpp
+++ b/Engine/Source/Runtime/Engine/Private/GameInstance.cpp
@@ -999,16 +999,18 @@ APlayerController* UGameInstance::GetFirstLocalPlayerController(const UWorld* Wo
 APlayerController* UGameInstance::GetPrimaryPlayerController(bool bRequiresValidUniqueId) const
 {
 	UWorld* World = GetWorld();
-	check(World);
 
 	APlayerController* PrimaryController = nullptr;
-	for (FConstPlayerControllerIterator Iterator = World->GetPlayerControllerIterator(); Iterator; ++Iterator)
+	if (ensure(World))
 	{
-		APlayerController* NextPlayer = Iterator->Get();
-		if (NextPlayer && NextPlayer->PlayerState && NextPlayer->IsPrimaryPlayer() && (!bRequiresValidUniqueId || NextPlayer->PlayerState->GetUniqueId().IsValid()))
+		for (FConstPlayerControllerIterator Iterator = World->GetPlayerControllerIterator(); Iterator; ++Iterator)
 		{
-			PrimaryController = NextPlayer;
-			break;
+			APlayerController* NextPlayer = Iterator->Get();
+			if (NextPlayer && NextPlayer->PlayerState && NextPlayer->IsPrimaryPlayer() && (!bRequiresValidUniqueId || NextPlayer->PlayerState->GetUniqueId().IsValid()))
+			{
+				PrimaryController = NextPlayer;
+				break;
+			}
 		}
 	}
 
