Move CSV stats updating to separate function. Declare strings used for said stats outside of the function to avoid unneeded thread-safe initialization.

diff --git a/Engine/Source/Runtime/Engine/Classes/Engine/NetDriver.h b/Engine/Source/Runtime/Engine/Classes/Engine/NetDriver.h
--- a/Engine/Source/Runtime/Engine/Classes/Engine/NetDriver.h
+++ b/Engine/Source/Runtime/Engine/Classes/Engine/NetDriver.h
@@ -1770,6 +1770,9 @@ protected:
 	/** Subclasses may override this to customize channel creation. Called by GetOrCreateChannel if the pool is exhausted and a new channel must be allocated. */
 	ENGINE_API virtual UChannel* InternalCreateChannelByName(const FName& ChName);
 
+	/** Update stats related to networking. */
+	void UpdateNetworkStats();
+
 #if WITH_SERVER_CODE
 	/**
 	* Helper functions for ServerReplicateActors
diff --git a/Engine/Source/Runtime/Engine/Private/NetDriver.cpp b/Engine/Source/Runtime/Engine/Private/NetDriver.cpp
--- a/Engine/Source/Runtime/Engine/Private/NetDriver.cpp
+++ b/Engine/Source/Runtime/Engine/Private/NetDriver.cpp
@@ -178,6 +178,38 @@
 
 extern int32 GNetRPCDebug;
 
+namespace UE::Net::Private
+{
+	static const FString PerfCounter_NumClients(TEXT("NumClients"));
+	static const FString PerfCounter_MaxPacketOverhead(TEXT("MaxPacketOverhead"));
+	static const FString PerfCounter_InRateClientMax(TEXT("InRateClientMax"));
+	static const FString PerfCounter_InRateClientMin(TEXT("InRateClientMin"));
+	static const FString PerfCounter_InRateClientAvg(TEXT("InRateClientAvg"));
+	static const FString PerfCounter_InPacketsClientMax(TEXT("InPacketsClientMax"));
+	static const FString PerfCounter_InPacketsClientMin(TEXT("InPacketsClientMin"));
+	static const FString PerfCounter_InPacketsClientAvg(TEXT("InPacketsClientAvg"));
+	static const FString PerfCounter_OutRateClientMax(TEXT("OutRateClientMax"));
+	static const FString PerfCounter_OutRateClientMin(TEXT("OutRateClientMin"));
+	static const FString PerfCounter_OutRateClientAvg(TEXT("OutRateClientAvg"));
+	static const FString PerfCounter_OutPacketsClientMax(TEXT("OutPacketsClientMax"));
+	static const FString PerfCounter_OutPacketsClientMin(TEXT("OutPacketsClientMin"));
+	static const FString PerfCounter_OutPacketsClientAvg(TEXT("OutPacketsClientAvg"));
+
+	static const FString PerfCounter_InRate(TEXT("InRate"));
+	static const FString PerfCounter_OutRate(TEXT("OutRate"));
+	static const FString PerfCounter_InPacketsLost(TEXT("InPacketsLost"));
+	static const FString PerfCounter_OutPacketsLost(TEXT("OutPacketsLost"));
+	static const FString PerfCounter_InPackets(TEXT("InPackets"));
+	static const FString PerfCounter_OutPackets(TEXT("OutPackets"));
+	static const FString PerfCounter_InBunches(TEXT("InBunches"));
+	static const FString PerfCounter_OutBunches(TEXT("OutBunches"));
+
+	static const FString PerfCounter_AvgPing(TEXT("AvgPing"));
+	static const FString PerfCounter_MaxPing(TEXT("MaxPing"));
+	static const FString PerfCounter_MinPing(TEXT("MinPing"));
+	static const FString PerfCounter_NumConnections(TEXT("NumConnections"));
+}
+
 namespace NetEmulationHelper
 {
 #if DO_ENABLE_NET_TEST
@@ -816,425 +848,6 @@
 	// Reset queued bunch amortization timer
 	ProcessQueuedBunchesCurrentFrameMilliseconds = 0.0f;
 
-	const double CurrentRealtimeSeconds = FPlatformTime::Seconds();
-
-	bool bCollectServerStats = false;
-#if USE_SERVER_PERF_COUNTERS || STATS
-	bCollectServerStats = true;
-#endif
-
-	if (bCollectNetStats || bCollectServerStats)
-	{
-		SCOPE_CYCLE_COUNTER(STAT_NetTickFlushGatherStats);
-		// Update network stats (only main game net driver for now) if stats or perf counters are used
-		if (NetDriverName == NAME_GameNetDriver &&
-			CurrentRealtimeSeconds - StatUpdateTime > StatPeriod)
-		{
-			int32 ClientInBytesMax = 0;
-			int32 ClientInBytesMin = 0;
-			int32 ClientInBytesAvg = 0;
-			int32 ClientInPacketsMax = 0;
-			int32 ClientInPacketsMin = 0;
-			int32 ClientInPacketsAvg = 0;
-			int32 ClientOutBytesMax = 0;
-			int32 ClientOutBytesMin = 0;
-			int32 ClientOutBytesAvg = 0;
-			int32 ClientOutPacketsMax = 0;
-			int32 ClientOutPacketsMin = 0;
-			int32 ClientOutPacketsAvg = 0;
-			int NumClients = 0;
-			int32 MaxPacketOverhead = 0;
-
-			// these need to be updated even if we are not collecting stats, since they get reported to analytics/QoS
-			for (UNetConnection * Client : ClientConnections)
-			{
-				if (Client)
-				{
-#define UpdatePerClientMinMaxAvg(VariableName) \
-				Client##VariableName##Max = FMath::Max(Client##VariableName##Max, Client->VariableName##PerSecond); \
-				if (Client##VariableName##Min == 0 || Client->VariableName##PerSecond < Client##VariableName##Min) \
-				{ \
-					Client##VariableName##Min = Client->VariableName##PerSecond; \
-				} \
-				Client##VariableName##Avg += Client->VariableName##PerSecond; \
-
-					UpdatePerClientMinMaxAvg(InBytes);
-					UpdatePerClientMinMaxAvg(OutBytes);
-					UpdatePerClientMinMaxAvg(InPackets);
-					UpdatePerClientMinMaxAvg(OutPackets);
-
-					MaxPacketOverhead = FMath::Max(Client->PacketOverhead, MaxPacketOverhead);
-
-#undef UpdatePerClientMinMaxAvg
-
-					++NumClients;
-				}
-			}
-
-			if (NumClients > 1)
-			{
-				ClientInBytesAvg /= NumClients;
-				ClientInPacketsAvg /= NumClients;
-				ClientOutBytesAvg /= NumClients;
-				ClientOutPacketsAvg /= NumClients;
-			}
-
-			int32 Ping = 0;
-			int32 NumOpenChannels = 0;
-			int32 NumActorChannels = 0;
-			int32 NumDormantActors = 0;
-			int32 NumActors = 0;
-			int32 AckCount = 0;
-			int32 UnAckCount = 0;
-			int32 PendingCount = 0;
-			int32 NetSaturated = 0;
-			int32 NumActiveNetActors = GetNetworkObjectList().GetActiveObjects().Num();
-
-			if (
-#if STATS
-				FThreadStats::IsCollectingData() ||
-#endif
-				bCollectNetStats)
-			{
-				const float RealTime = CurrentRealtimeSeconds - StatUpdateTime;
-
-				// Use the elapsed time to keep things scaled to one measured unit
-				InBytes = FMath::TruncToInt(InBytes / RealTime);
-				OutBytes = FMath::TruncToInt(OutBytes / RealTime);
-
-				NetGUIDOutBytes = FMath::TruncToInt(NetGUIDOutBytes / RealTime);
-				NetGUIDInBytes = FMath::TruncToInt(NetGUIDInBytes / RealTime);
-
-				// Save off for stats later
-
-				InBytesPerSecond = InBytes;
-				OutBytesPerSecond = OutBytes;
-
-				InPackets = FMath::TruncToInt(InPackets / RealTime);
-				OutPackets = FMath::TruncToInt(OutPackets / RealTime);
-				InBunches = FMath::TruncToInt(InBunches / RealTime);
-				OutBunches = FMath::TruncToInt(OutBunches / RealTime);
-				OutPacketsLost = FMath::TruncToInt(100.f * OutPacketsLost / FMath::Max((float)OutPackets, 1.f));
-				InPacketsLost = FMath::TruncToInt(100.f * InPacketsLost / FMath::Max((float)InPackets + InPacketsLost, 1.f));
-
-				if (ServerConnection != NULL && ServerConnection->PlayerController != NULL && ServerConnection->PlayerController->PlayerState != NULL)
-				{
-					Ping = FMath::TruncToInt(ServerConnection->PlayerController->PlayerState->ExactPing);
-				}
-
-				if (ServerConnection != NULL)
-				{
-					NumOpenChannels = ServerConnection->OpenChannels.Num();
-				}
-
-				for (int32 i = 0; i < ClientConnections.Num(); i++)
-				{
-					NumOpenChannels += ClientConnections[i]->OpenChannels.Num();
-				}
-
-				// Use the elapsed time to keep things scaled to one measured unit
-				VoicePacketsSent = FMath::TruncToInt(VoicePacketsSent / RealTime);
-				VoicePacketsRecv = FMath::TruncToInt(VoicePacketsRecv / RealTime);
-				VoiceBytesSent = FMath::TruncToInt(VoiceBytesSent / RealTime);
-				VoiceBytesRecv = FMath::TruncToInt(VoiceBytesRecv / RealTime);
-
-				// Determine voice percentages
-				VoiceInPercent = (InBytes > 0) ? FMath::TruncToInt(100.f * (float)VoiceBytesRecv / (float)InBytes) : 0;
-				VoiceOutPercent = (OutBytes > 0) ? FMath::TruncToInt(100.f * (float)VoiceBytesSent / (float)OutBytes) : 0;
-
-				if (World)
-				{
-					NumActors = World->GetActorCount();
-				}
-				
-				// Handle these stats to do this differently for client & server
-				if (ServerConnection != nullptr)
-				{
-					NumActorChannels = ServerConnection->ActorChannelsNum();
-					NumDormantActors = GetNetworkObjectList().GetNumDormantActorsForConnection(ServerConnection);
-#if STATS
-					ServerConnection->PackageMap->GetNetGUIDStats(AckCount, UnAckCount, PendingCount);
-#endif
-					NetSaturated = ServerConnection->IsNetReady(false) ? 0 : 1;
-				}
-				else
-				{
-					int32 SharedDormantActors = GetNetworkObjectList().GetDormantObjectsOnAllConnections().Num();
-					for (int32 i = 0; i < ClientConnections.Num(); i++)
-					{
-						NumActorChannels += ClientConnections[i]->ActorChannelsNum();
-						NumDormantActors += GetNetworkObjectList().GetNumDormantActorsForConnection(ClientConnections[i]);
-#if STATS
-						int32 ClientAckCount = 0;
-						int32 ClientUnAckCount = 0;
-						int32 ClientPendingCount = 0;
-						ClientConnections[i]->PackageMap->GetNetGUIDStats(AckCount, UnAckCount, PendingCount);
-						AckCount += ClientAckCount;
-						UnAckCount += ClientUnAckCount;
-						PendingCount += ClientPendingCount;
-#endif
-						NetSaturated |= ClientConnections[i]->IsNetReady(false) ? 0 : 1;
-					}
-					// Subtract out the duplicate counted dormant actors
-					NumDormantActors -= FMath::Max(0, ClientConnections.Num() - 1) * SharedDormantActors;
-				}
-			}
-
-#if STATS
-			if (!bSkipLocalStats)
-			{
-				// Copy the net status values over
-				SET_DWORD_STAT(STAT_Ping, Ping);
-				SET_DWORD_STAT(STAT_Channels, NumOpenChannels);
-				SET_DWORD_STAT(STAT_MaxPacketOverhead, MaxPacketOverhead);
-
-				SET_DWORD_STAT(STAT_OutLoss, OutPacketsLost);
-				SET_DWORD_STAT(STAT_InLoss, InPacketsLost);
-				SET_DWORD_STAT(STAT_InRate, InBytes);
-				SET_DWORD_STAT(STAT_OutRate, OutBytes);
-				SET_DWORD_STAT(STAT_InRateClientMax, ClientInBytesMax);
-				SET_DWORD_STAT(STAT_InRateClientMin, ClientInBytesMin);
-				SET_DWORD_STAT(STAT_InRateClientAvg, ClientInBytesAvg);
-				SET_DWORD_STAT(STAT_InPacketsClientMax, ClientInPacketsMax);
-				SET_DWORD_STAT(STAT_InPacketsClientMin, ClientInPacketsMin);
-				SET_DWORD_STAT(STAT_InPacketsClientAvg, ClientInPacketsAvg);
-				SET_DWORD_STAT(STAT_OutRateClientMax, ClientOutBytesMax);
-				SET_DWORD_STAT(STAT_OutRateClientMin, ClientOutBytesMin);
-				SET_DWORD_STAT(STAT_OutRateClientAvg, ClientOutBytesAvg);
-				SET_DWORD_STAT(STAT_OutPacketsClientMax, ClientOutPacketsMax);
-				SET_DWORD_STAT(STAT_OutPacketsClientMin, ClientOutPacketsMin);
-				SET_DWORD_STAT(STAT_OutPacketsClientAvg, ClientOutPacketsAvg);
-
-				SET_DWORD_STAT(STAT_NetNumClients, NumClients);
-				SET_DWORD_STAT(STAT_InPackets, InPackets);
-				SET_DWORD_STAT(STAT_OutPackets, OutPackets);
-				SET_DWORD_STAT(STAT_InBunches, InBunches);
-				SET_DWORD_STAT(STAT_OutBunches, OutBunches);
-
-				SET_DWORD_STAT(STAT_NetGUIDInRate, NetGUIDInBytes);
-				SET_DWORD_STAT(STAT_NetGUIDOutRate, NetGUIDOutBytes);
-
-				SET_DWORD_STAT(STAT_VoicePacketsSent, VoicePacketsSent);
-				SET_DWORD_STAT(STAT_VoicePacketsRecv, VoicePacketsRecv);
-				SET_DWORD_STAT(STAT_VoiceBytesSent, VoiceBytesSent);
-				SET_DWORD_STAT(STAT_VoiceBytesRecv, VoiceBytesRecv);
-
-				SET_DWORD_STAT(STAT_PercentInVoice, VoiceInPercent);
-				SET_DWORD_STAT(STAT_PercentOutVoice, VoiceOutPercent);
-
-				SET_DWORD_STAT(STAT_NumActorChannels, NumActorChannels);
-				SET_DWORD_STAT(STAT_NumDormantActors, NumDormantActors);
-				SET_DWORD_STAT(STAT_NumActors, NumActors);
-				SET_DWORD_STAT(STAT_NumNetActors, NumActiveNetActors);
-				SET_DWORD_STAT(STAT_NumNetGUIDsAckd, AckCount);
-				SET_DWORD_STAT(STAT_NumNetGUIDsPending, UnAckCount);
-				SET_DWORD_STAT(STAT_NumNetGUIDsUnAckd, PendingCount);
-				SET_DWORD_STAT(STAT_NetSaturated, NetSaturated);
-			}
-
-			// If we are to replicate server stats out to an observer, then set those values
-			if (AGameModeBase* GMBase = World->GetAuthGameMode())
-			{
-				if (GMBase->ServerStatReplicator != nullptr)
-				{
-					GMBase->ServerStatReplicator->Channels = NumOpenChannels;
-					GMBase->ServerStatReplicator->MaxPacketOverhead = MaxPacketOverhead;
-					GMBase->ServerStatReplicator->OutLoss = OutPacketsLost;
-					GMBase->ServerStatReplicator->InLoss = InPacketsLost;
-					GMBase->ServerStatReplicator->InRate = InBytes;
-					GMBase->ServerStatReplicator->OutRate = OutBytes;
-					GMBase->ServerStatReplicator->InRateClientMax = ClientInBytesMax;
-					GMBase->ServerStatReplicator->InRateClientMin = ClientInBytesMin;
-					GMBase->ServerStatReplicator->InRateClientAvg = ClientInBytesAvg;
-					GMBase->ServerStatReplicator->InPacketsClientMax = ClientInPacketsMax;
-					GMBase->ServerStatReplicator->InPacketsClientMin = ClientInPacketsMin;
-					GMBase->ServerStatReplicator->InPacketsClientAvg = ClientInPacketsAvg;
-					GMBase->ServerStatReplicator->OutRateClientMax = ClientOutBytesMax;
-					GMBase->ServerStatReplicator->OutRateClientMin = ClientOutBytesMin;
-					GMBase->ServerStatReplicator->OutRateClientAvg = ClientOutBytesAvg;
-					GMBase->ServerStatReplicator->OutPacketsClientMax = ClientOutPacketsMax;
-					GMBase->ServerStatReplicator->OutPacketsClientMin = ClientOutPacketsMin;
-					GMBase->ServerStatReplicator->OutPacketsClientAvg = ClientOutPacketsAvg;
-					GMBase->ServerStatReplicator->NetNumClients = NumClients;
-					GMBase->ServerStatReplicator->InPackets = InPackets;
-					GMBase->ServerStatReplicator->OutPackets = OutPackets;
-					GMBase->ServerStatReplicator->InBunches = InBunches;
-					GMBase->ServerStatReplicator->OutBunches = OutBunches;
-					GMBase->ServerStatReplicator->NetGUIDInRate = NetGUIDInBytes;
-					GMBase->ServerStatReplicator->NetGUIDOutRate = NetGUIDOutBytes;
-					GMBase->ServerStatReplicator->VoicePacketsSent = VoicePacketsSent;
-					GMBase->ServerStatReplicator->VoicePacketsRecv = VoicePacketsRecv;
-					GMBase->ServerStatReplicator->VoiceBytesSent = VoiceBytesSent;
-					GMBase->ServerStatReplicator->VoiceBytesRecv = VoiceBytesRecv;
-					GMBase->ServerStatReplicator->PercentInVoice = VoiceInPercent;
-					GMBase->ServerStatReplicator->PercentOutVoice = VoiceOutPercent;
-					GMBase->ServerStatReplicator->NumActorChannels = NumActorChannels;
-					GMBase->ServerStatReplicator->NumDormantActors = NumDormantActors;
-					GMBase->ServerStatReplicator->NumActors = NumActors;
-					GMBase->ServerStatReplicator->NumNetActors = NumActiveNetActors;
-					GMBase->ServerStatReplicator->NumNetGUIDsAckd = AckCount;
-					GMBase->ServerStatReplicator->NumNetGUIDsPending = UnAckCount;
-					GMBase->ServerStatReplicator->NumNetGUIDsUnAckd = PendingCount;
-					GMBase->ServerStatReplicator->NetSaturated = NetSaturated;
-				}
-			}
-#endif // STATS
-
-#if USE_SERVER_PERF_COUNTERS
-			IPerfCounters* PerfCounters = IPerfCountersModule::Get().GetPerformanceCounters();
-			if (PerfCounters)
-			{
-				SCOPE_CYCLE_COUNTER(STAT_NetTickFlushGatherStatsPerfCounters);
-
-				static const FString PerfCounter_AvgPing = TEXT("AvgPing");
-				static const FString PerfCounter_MaxPing = TEXT("MaxPing");
-				static const FString PerfCounter_MinPing = TEXT("MinPing");
-				static const FString PerfCounter_NumConnections = TEXT("NumConnections");
-
-				// Update total connections
-				PerfCounters->Set(PerfCounter_NumConnections, ClientConnections.Num());
-
-				const int kNumBuckets = 8;	// evenly spaced with increment of 30 ms; last bucket collects all off-scale pings as well
-				if (ClientConnections.Num() > 0)
-				{
-					// Update per connection statistics
-					float MinPing = UE_MAX_FLT;
-					float AvgPing = 0;
-					float MaxPing = -UE_MAX_FLT;
-					float PingCount = 0;
-
-					int32 Buckets[kNumBuckets] = { 0 };
-
-					for (int32 i = 0; i < ClientConnections.Num(); i++)
-					{
-						UNetConnection* Connection = ClientConnections[i];
-
-						if (Connection != nullptr)
-						{
-							if (Connection->PlayerController != nullptr && Connection->PlayerController->PlayerState != nullptr)
-							{
-								// Ping value calculated per client
-								float ConnPing = Connection->PlayerController->PlayerState->ExactPing;
-
-								int Bucket = FMath::Max(0, FMath::Min(kNumBuckets - 1, (static_cast<int>(ConnPing) / 30)));
-								++Buckets[Bucket];
-
-								if (ConnPing < MinPing)
-								{
-									MinPing = ConnPing;
-								}
-
-								if (ConnPing > MaxPing)
-								{
-									MaxPing = ConnPing;
-								}
-
-								AvgPing += ConnPing;
-								PingCount++;
-							}
-						}
-					}
-
-					if (PingCount > 0)
-					{
-						AvgPing /= static_cast<float>(PingCount);
-					}
-
-					PerfCounters->Set(PerfCounter_AvgPing, AvgPing, IPerfCounters::Flags::Transient);
-					PerfCounters->Set(PerfCounter_MaxPing, MaxPing, IPerfCounters::Flags::Transient);
-					PerfCounters->Set(PerfCounter_MinPing, MinPing, IPerfCounters::Flags::Transient);
-
-					// update buckets
-					for (int BucketIdx = 0; BucketIdx < UE_ARRAY_COUNT(Buckets); ++BucketIdx)
-					{
-						PerfCountersIncrement(FString::Printf(TEXT("PingBucketInt%d"), BucketIdx), Buckets[BucketIdx], 0, IPerfCounters::Flags::Transient);
-					}
-				}
-				else
-				{
-					PerfCounters->Set(PerfCounter_AvgPing, 0.0f, IPerfCounters::Flags::Transient);
-					PerfCounters->Set(PerfCounter_MaxPing, -FLT_MAX, IPerfCounters::Flags::Transient);
-					PerfCounters->Set(PerfCounter_MinPing, FLT_MAX, IPerfCounters::Flags::Transient);
-
-					for (int BucketIdx = 0; BucketIdx < kNumBuckets; ++BucketIdx)
-					{
-						PerfCounters->Set(FString::Printf(TEXT("PingBucketInt%d"), BucketIdx), 0, IPerfCounters::Flags::Transient);
-					}
-				}
-
-				static const FString PerfCounter_NumClients = TEXT("NumClients");
-				static const FString PerfCounter_MaxPacketOverhead = TEXT("MaxPacketOverhead");
-				static const FString PerfCounter_InRateClientMax = TEXT("InRateClientMax");
-				static const FString PerfCounter_InRateClientMin = TEXT("InRateClientMin");
-				static const FString PerfCounter_InRateClientAvg = TEXT("InRateClientAvg");
-				static const FString PerfCounter_InPacketsClientMax = TEXT("InPacketsClientMax");
-				static const FString PerfCounter_InPacketsClientMin = TEXT("InPacketsClientMin");
-				static const FString PerfCounter_InPacketsClientAvg = TEXT("InPacketsClientAvg");
-				static const FString PerfCounter_OutRateClientMax = TEXT("OutRateClientMax");
-				static const FString PerfCounter_OutRateClientMin = TEXT("OutRateClientMin");
-				static const FString PerfCounter_OutRateClientAvg = TEXT("OutRateClientAvg");
-				static const FString PerfCounter_OutPacketsClientMax = TEXT("OutPacketsClientMax");
-				static const FString PerfCounter_OutPacketsClientMin = TEXT("OutPacketsClientMin");
-				static const FString PerfCounter_OutPacketsClientAvg = TEXT("OutPacketsClientAvg");
-
-				static const FString PerfCounter_InRate = TEXT("InRate");
-				static const FString PerfCounter_OutRate = TEXT("OutRate");
-				static const FString PerfCounter_InPacketsLost = TEXT("InPacketsLost");
-				static const FString PerfCounter_OutPacketsLost = TEXT("OutPacketsLost");
-				static const FString PerfCounter_InPackets = TEXT("InPackets");
-				static const FString PerfCounter_OutPackets = TEXT("OutPackets");
-				static const FString PerfCounter_InBunches = TEXT("InBunches");
-				static const FString PerfCounter_OutBunches = TEXT("OutBunches");
-
-				// set the per connection stats (these are calculated earlier).
-				// Note that NumClients may be != NumConnections. Also, if NumClients is 0, the rest of counters should be 0 as well
-				PerfCounters->Set(PerfCounter_NumClients, NumClients);
-				PerfCounters->Set(PerfCounter_MaxPacketOverhead, MaxPacketOverhead);
-				PerfCounters->Set(PerfCounter_InRateClientMax, ClientInBytesMax);
-				PerfCounters->Set(PerfCounter_InRateClientMin, ClientInBytesMin);
-				PerfCounters->Set(PerfCounter_InRateClientAvg, ClientInBytesAvg);
-				PerfCounters->Set(PerfCounter_InPacketsClientMax, ClientInPacketsMax);
-				PerfCounters->Set(PerfCounter_InPacketsClientMin, ClientInPacketsMin);
-				PerfCounters->Set(PerfCounter_InPacketsClientAvg, ClientInPacketsAvg);
-				PerfCounters->Set(PerfCounter_OutRateClientMax, ClientOutBytesMax);
-				PerfCounters->Set(PerfCounter_OutRateClientMin, ClientOutBytesMin);
-				PerfCounters->Set(PerfCounter_OutRateClientAvg, ClientOutBytesAvg);
-				PerfCounters->Set(PerfCounter_OutPacketsClientMax, ClientOutPacketsMax);
-				PerfCounters->Set(PerfCounter_OutPacketsClientMin, ClientOutPacketsMin);
-				PerfCounters->Set(PerfCounter_OutPacketsClientAvg, ClientOutPacketsAvg);
-
-				PerfCounters->Set(PerfCounter_InRate, InBytes);
-				PerfCounters->Set(PerfCounter_OutRate, OutBytes);
-				PerfCounters->Set(PerfCounter_InPacketsLost, InPacketsLost);
-				PerfCounters->Set(PerfCounter_OutPacketsLost, OutPacketsLost);
-				PerfCounters->Set(PerfCounter_InPackets, InPackets);
-				PerfCounters->Set(PerfCounter_OutPackets, OutPackets);
-				PerfCounters->Set(PerfCounter_InBunches, InBunches);
-				PerfCounters->Set(PerfCounter_OutBunches, OutBunches);
-			}
-#endif // USE_SERVER_PERF_COUNTERS
-
-			// Reset everything
-			InBytes = 0;
-			OutBytes = 0;
-			NetGUIDOutBytes = 0;
-			NetGUIDInBytes = 0;
-			InPackets = 0;
-			OutPackets = 0;
-			InBunches = 0;
-			OutBunches = 0;
-			OutPacketsLost = 0;
-			InPacketsLost = 0;
-			VoicePacketsSent = 0;
-			VoiceBytesSent = 0;
-			VoicePacketsRecv = 0;
-			VoiceBytesRecv = 0;
-			VoiceInPercent = 0;
-			VoiceOutPercent = 0;
-			StatUpdateTime = CurrentRealtimeSeconds;
-		}
-	} // bCollectNetStats ||(USE_SERVER_PERF_COUNTERS) || STATS
-
 	// Poll all sockets.
 	if( ServerConnection )
 	{
@@ -1426,6 +1039,7 @@
 	// Unfortunately if you mark an object as pending kill, it will no longer find itself in this map,
 	// so we do this as a fail safe to make sure we never leak memory from this map
 	const double CleanupTimeSeconds = 10.0;
+	const double CurrentRealtimeSeconds = FPlatformTime::Seconds();
 
 	if ( CurrentRealtimeSeconds - LastCleanupTime > CleanupTimeSeconds )
 	{
@@ -1453,6 +1067,8 @@
 		LastCleanupTime = CurrentRealtimeSeconds;
 	}
 
+	UpdateNetworkStats();
+
 	// Update the lag state
 	UpdateNetworkLagState();
 }
@@ -6980,4 +6596,399 @@
 	}
 }
 
+void UNetDriver::UpdateNetworkStats()
+{
+	using namespace UE::Net::Private;
+
+	const double CurrentRealtimeSeconds = FPlatformTime::Seconds();
+
+	bool bCollectServerStats = false;
+#if USE_SERVER_PERF_COUNTERS || STATS
+	bCollectServerStats = true;
+#endif
+
+	if (bCollectNetStats || bCollectServerStats)
+	{
+		SCOPE_CYCLE_COUNTER(STAT_NetTickFlushGatherStats);
+		// Update network stats (only main game net driver for now) if stats or perf counters are used
+		if (NetDriverName == NAME_GameNetDriver &&
+			CurrentRealtimeSeconds - StatUpdateTime > StatPeriod)
+		{
+			int32 ClientInBytesMax = 0;
+			int32 ClientInBytesMin = 0;
+			int32 ClientInBytesAvg = 0;
+			int32 ClientInPacketsMax = 0;
+			int32 ClientInPacketsMin = 0;
+			int32 ClientInPacketsAvg = 0;
+			int32 ClientOutBytesMax = 0;
+			int32 ClientOutBytesMin = 0;
+			int32 ClientOutBytesAvg = 0;
+			int32 ClientOutPacketsMax = 0;
+			int32 ClientOutPacketsMin = 0;
+			int32 ClientOutPacketsAvg = 0;
+			int NumClients = 0;
+			int32 MaxPacketOverhead = 0;
+
+			// these need to be updated even if we are not collecting stats, since they get reported to analytics/QoS
+			for (UNetConnection* Client : ClientConnections)
+			{
+				if (Client)
+				{
+#define UpdatePerClientMinMaxAvg(VariableName) \
+				Client##VariableName##Max = FMath::Max(Client##VariableName##Max, Client->VariableName##PerSecond); \
+				if (Client##VariableName##Min == 0 || Client->VariableName##PerSecond < Client##VariableName##Min) \
+				{ \
+					Client##VariableName##Min = Client->VariableName##PerSecond; \
+				} \
+				Client##VariableName##Avg += Client->VariableName##PerSecond; \
+
+					UpdatePerClientMinMaxAvg(InBytes);
+					UpdatePerClientMinMaxAvg(OutBytes);
+					UpdatePerClientMinMaxAvg(InPackets);
+					UpdatePerClientMinMaxAvg(OutPackets);
+
+					MaxPacketOverhead = FMath::Max(Client->PacketOverhead, MaxPacketOverhead);
+
+#undef UpdatePerClientMinMaxAvg
+
+					++NumClients;
+				}
+			}
+
+			if (NumClients > 1)
+			{
+				ClientInBytesAvg /= NumClients;
+				ClientInPacketsAvg /= NumClients;
+				ClientOutBytesAvg /= NumClients;
+				ClientOutPacketsAvg /= NumClients;
+			}
+
+			int32 Ping = 0;
+			int32 NumOpenChannels = 0;
+			int32 NumActorChannels = 0;
+			int32 NumDormantActors = 0;
+			int32 NumActors = 0;
+			int32 AckCount = 0;
+			int32 UnAckCount = 0;
+			int32 PendingCount = 0;
+			int32 NetSaturated = 0;
+			int32 NumActiveNetActors = GetNetworkObjectList().GetActiveObjects().Num();
+
+			if (
+#if STATS
+				FThreadStats::IsCollectingData() ||
+#endif
+				bCollectNetStats)
+			{
+				const float RealTime = CurrentRealtimeSeconds - StatUpdateTime;
+
+				// Use the elapsed time to keep things scaled to one measured unit
+				InBytes = FMath::TruncToInt(InBytes / RealTime);
+				OutBytes = FMath::TruncToInt(OutBytes / RealTime);
+
+				NetGUIDOutBytes = FMath::TruncToInt(NetGUIDOutBytes / RealTime);
+				NetGUIDInBytes = FMath::TruncToInt(NetGUIDInBytes / RealTime);
+
+				// Save off for stats later
+
+				InBytesPerSecond = InBytes;
+				OutBytesPerSecond = OutBytes;
+
+				InPackets = FMath::TruncToInt(InPackets / RealTime);
+				OutPackets = FMath::TruncToInt(OutPackets / RealTime);
+				InBunches = FMath::TruncToInt(InBunches / RealTime);
+				OutBunches = FMath::TruncToInt(OutBunches / RealTime);
+				OutPacketsLost = FMath::TruncToInt(100.f * OutPacketsLost / FMath::Max((float)OutPackets, 1.f));
+				InPacketsLost = FMath::TruncToInt(100.f * InPacketsLost / FMath::Max((float)InPackets + InPacketsLost, 1.f));
+
+				if (ServerConnection != nullptr && ServerConnection->PlayerController != nullptr && ServerConnection->PlayerController->PlayerState != nullptr)
+				{
+					Ping = FMath::TruncToInt(ServerConnection->PlayerController->PlayerState->ExactPing);
+				}
+
+				if (ServerConnection != nullptr)
+				{
+					NumOpenChannels = ServerConnection->OpenChannels.Num();
+				}
+
+				for (int32 i = 0; i < ClientConnections.Num(); i++)
+				{
+					NumOpenChannels += ClientConnections[i]->OpenChannels.Num();
+				}
+
+				// Use the elapsed time to keep things scaled to one measured unit
+				VoicePacketsSent = FMath::TruncToInt(VoicePacketsSent / RealTime);
+				VoicePacketsRecv = FMath::TruncToInt(VoicePacketsRecv / RealTime);
+				VoiceBytesSent = FMath::TruncToInt(VoiceBytesSent / RealTime);
+				VoiceBytesRecv = FMath::TruncToInt(VoiceBytesRecv / RealTime);
+
+				// Determine voice percentages
+				VoiceInPercent = (InBytes > 0) ? FMath::TruncToInt(100.f * (float)VoiceBytesRecv / (float)InBytes) : 0;
+				VoiceOutPercent = (OutBytes > 0) ? FMath::TruncToInt(100.f * (float)VoiceBytesSent / (float)OutBytes) : 0;
+
+				if (World)
+				{
+					NumActors = World->GetActorCount();
+				}
+				
+				// Handle these stats to do this differently for client & server
+				if (ServerConnection != nullptr)
+				{
+					NumActorChannels = ServerConnection->ActorChannelsNum();
+					NumDormantActors = GetNetworkObjectList().GetNumDormantActorsForConnection(ServerConnection);
+#if STATS
+					ServerConnection->PackageMap->GetNetGUIDStats(AckCount, UnAckCount, PendingCount);
+#endif
+					NetSaturated = ServerConnection->IsNetReady(false) ? 0 : 1;
+				}
+				else
+				{
+					int32 SharedDormantActors = GetNetworkObjectList().GetDormantObjectsOnAllConnections().Num();
+					for (int32 i = 0; i < ClientConnections.Num(); i++)
+					{
+						NumActorChannels += ClientConnections[i]->ActorChannelsNum();
+						NumDormantActors += GetNetworkObjectList().GetNumDormantActorsForConnection(ClientConnections[i]);
+#if STATS
+						int32 ClientAckCount = 0;
+						int32 ClientUnAckCount = 0;
+						int32 ClientPendingCount = 0;
+						ClientConnections[i]->PackageMap->GetNetGUIDStats(AckCount, UnAckCount, PendingCount);
+						AckCount += ClientAckCount;
+						UnAckCount += ClientUnAckCount;
+						PendingCount += ClientPendingCount;
+#endif
+						NetSaturated |= ClientConnections[i]->IsNetReady(false) ? 0 : 1;
+					}
+					// Subtract out the duplicate counted dormant actors
+					NumDormantActors -= FMath::Max(0, ClientConnections.Num() - 1) * SharedDormantActors;
+				}
+			}
+
+#if STATS
+			if (!bSkipLocalStats)
+			{
+				// Copy the net status values over
+				SET_DWORD_STAT(STAT_Ping, Ping);
+				SET_DWORD_STAT(STAT_Channels, NumOpenChannels);
+				SET_DWORD_STAT(STAT_MaxPacketOverhead, MaxPacketOverhead);
+
+				SET_DWORD_STAT(STAT_OutLoss, OutPacketsLost);
+				SET_DWORD_STAT(STAT_InLoss, InPacketsLost);
+				SET_DWORD_STAT(STAT_InRate, InBytes);
+				SET_DWORD_STAT(STAT_OutRate, OutBytes);
+				SET_DWORD_STAT(STAT_InRateClientMax, ClientInBytesMax);
+				SET_DWORD_STAT(STAT_InRateClientMin, ClientInBytesMin);
+				SET_DWORD_STAT(STAT_InRateClientAvg, ClientInBytesAvg);
+				SET_DWORD_STAT(STAT_InPacketsClientMax, ClientInPacketsMax);
+				SET_DWORD_STAT(STAT_InPacketsClientMin, ClientInPacketsMin);
+				SET_DWORD_STAT(STAT_InPacketsClientAvg, ClientInPacketsAvg);
+				SET_DWORD_STAT(STAT_OutRateClientMax, ClientOutBytesMax);
+				SET_DWORD_STAT(STAT_OutRateClientMin, ClientOutBytesMin);
+				SET_DWORD_STAT(STAT_OutRateClientAvg, ClientOutBytesAvg);
+				SET_DWORD_STAT(STAT_OutPacketsClientMax, ClientOutPacketsMax);
+				SET_DWORD_STAT(STAT_OutPacketsClientMin, ClientOutPacketsMin);
+				SET_DWORD_STAT(STAT_OutPacketsClientAvg, ClientOutPacketsAvg);
+
+				SET_DWORD_STAT(STAT_NetNumClients, NumClients);
+				SET_DWORD_STAT(STAT_InPackets, InPackets);
+				SET_DWORD_STAT(STAT_OutPackets, OutPackets);
+				SET_DWORD_STAT(STAT_InBunches, InBunches);
+				SET_DWORD_STAT(STAT_OutBunches, OutBunches);
+
+				SET_DWORD_STAT(STAT_NetGUIDInRate, NetGUIDInBytes);
+				SET_DWORD_STAT(STAT_NetGUIDOutRate, NetGUIDOutBytes);
+
+				SET_DWORD_STAT(STAT_VoicePacketsSent, VoicePacketsSent);
+				SET_DWORD_STAT(STAT_VoicePacketsRecv, VoicePacketsRecv);
+				SET_DWORD_STAT(STAT_VoiceBytesSent, VoiceBytesSent);
+				SET_DWORD_STAT(STAT_VoiceBytesRecv, VoiceBytesRecv);
+
+				SET_DWORD_STAT(STAT_PercentInVoice, VoiceInPercent);
+				SET_DWORD_STAT(STAT_PercentOutVoice, VoiceOutPercent);
+
+				SET_DWORD_STAT(STAT_NumActorChannels, NumActorChannels);
+				SET_DWORD_STAT(STAT_NumDormantActors, NumDormantActors);
+				SET_DWORD_STAT(STAT_NumActors, NumActors);
+				SET_DWORD_STAT(STAT_NumNetActors, NumActiveNetActors);
+				SET_DWORD_STAT(STAT_NumNetGUIDsAckd, AckCount);
+				SET_DWORD_STAT(STAT_NumNetGUIDsPending, UnAckCount);
+				SET_DWORD_STAT(STAT_NumNetGUIDsUnAckd, PendingCount);
+				SET_DWORD_STAT(STAT_NetSaturated, NetSaturated);
+			}
+
+			// If we are to replicate server stats out to an observer, then set those values
+			if (AGameModeBase* GMBase = World->GetAuthGameMode())
+			{
+				if (GMBase->ServerStatReplicator != nullptr)
+				{
+					GMBase->ServerStatReplicator->Channels = NumOpenChannels;
+					GMBase->ServerStatReplicator->MaxPacketOverhead = MaxPacketOverhead;
+					GMBase->ServerStatReplicator->OutLoss = OutPacketsLost;
+					GMBase->ServerStatReplicator->InLoss = InPacketsLost;
+					GMBase->ServerStatReplicator->InRate = InBytes;
+					GMBase->ServerStatReplicator->OutRate = OutBytes;
+					GMBase->ServerStatReplicator->InRateClientMax = ClientInBytesMax;
+					GMBase->ServerStatReplicator->InRateClientMin = ClientInBytesMin;
+					GMBase->ServerStatReplicator->InRateClientAvg = ClientInBytesAvg;
+					GMBase->ServerStatReplicator->InPacketsClientMax = ClientInPacketsMax;
+					GMBase->ServerStatReplicator->InPacketsClientMin = ClientInPacketsMin;
+					GMBase->ServerStatReplicator->InPacketsClientAvg = ClientInPacketsAvg;
+					GMBase->ServerStatReplicator->OutRateClientMax = ClientOutBytesMax;
+					GMBase->ServerStatReplicator->OutRateClientMin = ClientOutBytesMin;
+					GMBase->ServerStatReplicator->OutRateClientAvg = ClientOutBytesAvg;
+					GMBase->ServerStatReplicator->OutPacketsClientMax = ClientOutPacketsMax;
+					GMBase->ServerStatReplicator->OutPacketsClientMin = ClientOutPacketsMin;
+					GMBase->ServerStatReplicator->OutPacketsClientAvg = ClientOutPacketsAvg;
+					GMBase->ServerStatReplicator->NetNumClients = NumClients;
+					GMBase->ServerStatReplicator->InPackets = InPackets;
+					GMBase->ServerStatReplicator->OutPackets = OutPackets;
+					GMBase->ServerStatReplicator->InBunches = InBunches;
+					GMBase->ServerStatReplicator->OutBunches = OutBunches;
+					GMBase->ServerStatReplicator->NetGUIDInRate = NetGUIDInBytes;
+					GMBase->ServerStatReplicator->NetGUIDOutRate = NetGUIDOutBytes;
+					GMBase->ServerStatReplicator->VoicePacketsSent = VoicePacketsSent;
+					GMBase->ServerStatReplicator->VoicePacketsRecv = VoicePacketsRecv;
+					GMBase->ServerStatReplicator->VoiceBytesSent = VoiceBytesSent;
+					GMBase->ServerStatReplicator->VoiceBytesRecv = VoiceBytesRecv;
+					GMBase->ServerStatReplicator->PercentInVoice = VoiceInPercent;
+					GMBase->ServerStatReplicator->PercentOutVoice = VoiceOutPercent;
+					GMBase->ServerStatReplicator->NumActorChannels = NumActorChannels;
+					GMBase->ServerStatReplicator->NumDormantActors = NumDormantActors;
+					GMBase->ServerStatReplicator->NumActors = NumActors;
+					GMBase->ServerStatReplicator->NumNetActors = NumActiveNetActors;
+					GMBase->ServerStatReplicator->NumNetGUIDsAckd = AckCount;
+					GMBase->ServerStatReplicator->NumNetGUIDsPending = UnAckCount;
+					GMBase->ServerStatReplicator->NumNetGUIDsUnAckd = PendingCount;
+					GMBase->ServerStatReplicator->NetSaturated = NetSaturated;
+				}
+			}
+#endif // STATS
+
+#if USE_SERVER_PERF_COUNTERS
+			IPerfCounters* PerfCounters = IPerfCountersModule::Get().GetPerformanceCounters();
+			if (PerfCounters)
+			{
+				SCOPE_CYCLE_COUNTER(STAT_NetTickFlushGatherStatsPerfCounters);
+
+				// Update total connections
+				PerfCounters->Set(PerfCounter_NumConnections, ClientConnections.Num());
+
+				const int kNumBuckets = 8;	// evenly spaced with increment of 30 ms; last bucket collects all off-scale pings as well
+				if (ClientConnections.Num() > 0)
+				{
+					// Update per connection statistics
+					float MinPing = UE_MAX_FLT;
+					float AvgPing = 0;
+					float MaxPing = -UE_MAX_FLT;
+					float PingCount = 0;
+
+					int32 Buckets[kNumBuckets] = { 0 };
+
+					for (int32 i = 0; i < ClientConnections.Num(); i++)
+					{
+						UNetConnection* Connection = ClientConnections[i];
+
+						if (Connection != nullptr)
+						{
+							if (Connection->PlayerController != nullptr && Connection->PlayerController->PlayerState != nullptr)
+							{
+								// Ping value calculated per client
+								float ConnPing = Connection->PlayerController->PlayerState->ExactPing;
+
+								int Bucket = FMath::Max(0, FMath::Min(kNumBuckets - 1, (static_cast<int>(ConnPing) / 30)));
+								++Buckets[Bucket];
+
+								if (ConnPing < MinPing)
+								{
+									MinPing = ConnPing;
+								}
+
+								if (ConnPing > MaxPing)
+								{
+									MaxPing = ConnPing;
+								}
+
+								AvgPing += ConnPing;
+								PingCount++;
+							}
+						}
+					}
+
+					if (PingCount > 0)
+					{
+						AvgPing /= static_cast<float>(PingCount);
+					}
+
+					PerfCounters->Set(PerfCounter_AvgPing, AvgPing, IPerfCounters::Flags::Transient);
+					PerfCounters->Set(PerfCounter_MaxPing, MaxPing, IPerfCounters::Flags::Transient);
+					PerfCounters->Set(PerfCounter_MinPing, MinPing, IPerfCounters::Flags::Transient);
+
+					// update buckets
+					for (int BucketIdx = 0; BucketIdx < UE_ARRAY_COUNT(Buckets); ++BucketIdx)
+					{
+						PerfCountersIncrement(FString::Printf(TEXT("PingBucketInt%d"), BucketIdx), Buckets[BucketIdx], 0, IPerfCounters::Flags::Transient);
+					}
+				}
+				else
+				{
+					PerfCounters->Set(PerfCounter_AvgPing, 0.0f, IPerfCounters::Flags::Transient);
+					PerfCounters->Set(PerfCounter_MaxPing, -FLT_MAX, IPerfCounters::Flags::Transient);
+					PerfCounters->Set(PerfCounter_MinPing, FLT_MAX, IPerfCounters::Flags::Transient);
+
+					for (int BucketIdx = 0; BucketIdx < kNumBuckets; ++BucketIdx)
+					{
+						PerfCounters->Set(FString::Printf(TEXT("PingBucketInt%d"), BucketIdx), 0, IPerfCounters::Flags::Transient);
+					}
+				}
+
+				// set the per connection stats (these are calculated earlier).
+				// Note that NumClients may be != NumConnections. Also, if NumClients is 0, the rest of counters should be 0 as well
+				PerfCounters->Set(PerfCounter_NumClients, NumClients);
+				PerfCounters->Set(PerfCounter_MaxPacketOverhead, MaxPacketOverhead);
+				PerfCounters->Set(PerfCounter_InRateClientMax, ClientInBytesMax);
+				PerfCounters->Set(PerfCounter_InRateClientMin, ClientInBytesMin);
+				PerfCounters->Set(PerfCounter_InRateClientAvg, ClientInBytesAvg);
+				PerfCounters->Set(PerfCounter_InPacketsClientMax, ClientInPacketsMax);
+				PerfCounters->Set(PerfCounter_InPacketsClientMin, ClientInPacketsMin);
+				PerfCounters->Set(PerfCounter_InPacketsClientAvg, ClientInPacketsAvg);
+				PerfCounters->Set(PerfCounter_OutRateClientMax, ClientOutBytesMax);
+				PerfCounters->Set(PerfCounter_OutRateClientMin, ClientOutBytesMin);
+				PerfCounters->Set(PerfCounter_OutRateClientAvg, ClientOutBytesAvg);
+				PerfCounters->Set(PerfCounter_OutPacketsClientMax, ClientOutPacketsMax);
+				PerfCounters->Set(PerfCounter_OutPacketsClientMin, ClientOutPacketsMin);
+				PerfCounters->Set(PerfCounter_OutPacketsClientAvg, ClientOutPacketsAvg);
+
+				PerfCounters->Set(PerfCounter_InRate, InBytes);
+				PerfCounters->Set(PerfCounter_OutRate, OutBytes);
+				PerfCounters->Set(PerfCounter_InPacketsLost, InPacketsLost);
+				PerfCounters->Set(PerfCounter_OutPacketsLost, OutPacketsLost);
+				PerfCounters->Set(PerfCounter_InPackets, InPackets);
+				PerfCounters->Set(PerfCounter_OutPackets, OutPackets);
+				PerfCounters->Set(PerfCounter_InBunches, InBunches);
+				PerfCounters->Set(PerfCounter_OutBunches, OutBunches);
+			}
+#endif // USE_SERVER_PERF_COUNTERS
+
+			// Reset everything
+			InBytes = 0;
+			OutBytes = 0;
+			NetGUIDOutBytes = 0;
+			NetGUIDInBytes = 0;
+			InPackets = 0;
+			OutPackets = 0;
+			InBunches = 0;
+			OutBunches = 0;
+			OutPacketsLost = 0;
+			InPacketsLost = 0;
+			VoicePacketsSent = 0;
+			VoiceBytesSent = 0;
+			VoicePacketsRecv = 0;
+			VoiceBytesRecv = 0;
+			VoiceInPercent = 0;
+			VoiceOutPercent = 0;
+			StatUpdateTime = CurrentRealtimeSeconds;
+		}
+	} // bCollectNetStats ||(USE_SERVER_PERF_COUNTERS) || STATS
+}
+
 #if NET_DEBUG_RELEVANT_ACTORS
