[UE][FEATURE] Extend Pose Watches to support control of other aspects of animation node debug draw.
- Pose watches can now contain elements which represent different asspects of an anim nodes debug draw (i.e. bodies and constraints in RBAN)
- Elements can be used to enable anim node debug draw even when a node is not selected, acting in the same way as the previous pose watches
- The previous Pose watch behavior is now provided by an element within a pose watch
- RBAN nodes now make use of Pose Watch Elements
- Moved skeleton draw parameter UPROPERTIES from PoseWatch to specialised PoseWatchPoseElement class

diff --git a/Engine/Source/Editor/AnimGraph/Private/AnimGraphNode_Base.cpp b/Engine/Source/Editor/AnimGraph/Private/AnimGraphNode_Base.cpp
--- a/Engine/Source/Editor/AnimGraph/Private/AnimGraphNode_Base.cpp
+++ b/Engine/Source/Editor/AnimGraph/Private/AnimGraphNode_Base.cpp
@@ -656,11 +656,43 @@ void UAnimGraphNode_Base::OnNodeSelected(bool bInIsSelected, FEditorModeTools& I
 	}
 }
 
+void UAnimGraphNode_Base::OnPoseWatchChanged(const bool IsPoseWatchEnabled, TObjectPtr<class UPoseWatch> InPoseWatch, FEditorModeTools& InModeTools, FAnimNode_Base* InRuntimeNode)
+{
+	const FEditorModeID ModeID = GetEditorMode();
+	if (ModeID != NAME_None)
+	{
+		if (IsPoseWatchEnabled)
+		{
+			InModeTools.ActivateMode(ModeID);
+			if (FEdMode* EdMode = InModeTools.GetActiveMode(ModeID))
+			{
+				static_cast<IAnimNodeEditMode*>(EdMode)->RegisterPoseWatchedNode(this, InRuntimeNode);
+			}
+		}
+		else
+		{
+			if (FEdMode* EdMode = InModeTools.GetActiveMode(ModeID))
+			{
+				static_cast<IAnimNodeEditMode*>(EdMode)->ExitMode();
+			}
+			InModeTools.DeactivateMode(ModeID);
+		}
+	}
+}
+
 FEditorModeID UAnimGraphNode_Base::GetEditorMode() const
 {
 	return AnimNodeEditModes::AnimNode;
 }
 
+void UAnimGraphNode_Base::Draw(FPrimitiveDrawInterface* PDI, USkeletalMeshComponent* PreviewSkelMeshComp, const bool IsSelected, const bool IsPoseWatchEnabled) const
+{
+	if (IsSelected)
+	{
+		Draw(PDI, PreviewSkelMeshComp); 
+	}
+}
+
 FAnimNode_Base* UAnimGraphNode_Base::FindDebugAnimNode(USkeletalMeshComponent* InPreviewSkelMeshComp) const
 {
 	FAnimNode_Base* DebugNode = nullptr;
diff --git a/Engine/Source/Editor/AnimGraph/Private/AnimGraphNode_RigidBody.cpp b/Engine/Source/Editor/AnimGraph/Private/AnimGraphNode_RigidBody.cpp
--- a/Engine/Source/Editor/AnimGraph/Private/AnimGraphNode_RigidBody.cpp
+++ b/Engine/Source/Editor/AnimGraph/Private/AnimGraphNode_RigidBody.cpp
@@ -4,7 +4,9 @@
 #include "AnimNodeEditModes.h"
 #include "Kismet2/CompilerResultsLog.h"
 #include "BoneControllers/AnimNode_RigidBody.h"
+#include "EditorModeManager.h"
 #include "IPhysicsAssetRenderInterface.h"
+#include "IAnimNodeEditMode.h"
 
 // Details includes
 #include "PropertyHandle.h"
@@ -48,14 +50,54 @@ void UAnimGraphNode_RigidBody::ValidateAnimNodeDuringCompilation(USkeleton* ForS
 	Super::ValidateAnimNodeDuringCompilation(ForSkeleton, MessageLog);
 }
 
-void UAnimGraphNode_RigidBody::Draw(FPrimitiveDrawInterface* PDI, USkeletalMeshComponent* PreviewSkelMeshComp) const
+void UAnimGraphNode_RigidBody::Draw(FPrimitiveDrawInterface* PDI, USkeletalMeshComponent* PreviewSkelMeshComp, const bool bIsSelected, const bool bIsPoseWatchEnabled) const
 {
 	if (const FAnimNode_RigidBody* const RuntimeRigidBodyNode = GetDebuggedAnimNode<FAnimNode_RigidBody>())
 	{
 		if (UPhysicsAsset* const PhysicsAsset = RuntimeRigidBodyNode->GetPhysicsAsset())
 		{
 			IPhysicsAssetRenderInterface& PhysicsAssetRenderInterface = IModularFeatures::Get().GetModularFeature<IPhysicsAssetRenderInterface>("PhysicsAssetRenderInterface");
-			PhysicsAssetRenderInterface.DebugDraw(PreviewSkelMeshComp, PhysicsAsset, PDI);
+
+			// Draw Bodies.
+			if (bIsSelected || (bIsPoseWatchEnabled && PoseWatchElementBodies.IsValid() && PoseWatchElementBodies->GetIsVisible()))
+			{
+				FColor PrimitiveColorOverride = FColor::Transparent;
+
+				// Get primitive color from pose watch component.
+				if (!bIsSelected)
+				{
+					PrimitiveColorOverride = PoseWatchElementBodies->GetColor();
+					PrimitiveColorOverride.A = 255;
+				}
+
+				PhysicsAssetRenderInterface.DebugDrawBodies(PreviewSkelMeshComp, PhysicsAsset, PDI, PrimitiveColorOverride);
+			}
+
+			// Draw Constraints.
+			if (bIsSelected || (bIsPoseWatchEnabled && PoseWatchElementConstraints.IsValid() && PoseWatchElementConstraints->GetIsVisible()))
+			{
+				PhysicsAssetRenderInterface.DebugDrawConstraints(PreviewSkelMeshComp, PhysicsAsset, PDI);
+			}
+		}
+	}
+}
+
+void UAnimGraphNode_RigidBody::OnPoseWatchChanged(const bool IsPoseWatchEnabled, TObjectPtr<UPoseWatch> InPoseWatch, FEditorModeTools& InModeTools, FAnimNode_Base* InRuntimeNode)
+{
+	Super::OnPoseWatchChanged(IsPoseWatchEnabled, InPoseWatch, InModeTools, InRuntimeNode);
+
+	UPoseWatch* const PoseWatch = InPoseWatch.Get();
+
+	if (PoseWatch && !(PoseWatch->Contains(PoseWatchElementBodies.Get()) && PoseWatch->Contains(PoseWatchElementConstraints.Get())))
+	{
+		// A new pose watch has been created for this node - add node specific pose watch components.
+		PoseWatchElementBodies = InPoseWatch.Get()->AddElement(FText(LOCTEXT("PoseWatchElementLabel_RigidBody_PhysicsBodies", "Physics Bodies")), TEXT("PhysicsAssetEditor.Tree.Body"));
+		PoseWatchElementConstraints = InPoseWatch.Get()->AddElement(FText(LOCTEXT("PoseWatchElementLabel__RigidBody_PhysicsConstraints", "Physics Constraints")), TEXT("PhysicsAssetEditor.Tree.Constraint"));
+
+		check(PoseWatchElementConstraints.IsValid()); // Expect to find a valid component;
+		if (PoseWatchElementConstraints.IsValid())
+		{
+			PoseWatchElementConstraints->SetHasColor(false);
 		}
 	}
 }
diff --git a/Engine/Source/Editor/AnimGraph/Private/EditModes/AnimDynamicsEditMode.cpp b/Engine/Source/Editor/AnimGraph/Private/EditModes/AnimDynamicsEditMode.cpp
--- a/Engine/Source/Editor/AnimGraph/Private/EditModes/AnimDynamicsEditMode.cpp
+++ b/Engine/Source/Editor/AnimGraph/Private/EditModes/AnimDynamicsEditMode.cpp
@@ -121,7 +121,7 @@ void FAnimDynamicsEditMode::ExitMode()
 {
 	SelectedViewportObjects.Empty();
 
-	for (const EditorRuntimeNodePair& CurrentNodePair : AnimNodes)
+	for (const EditorRuntimeNodePair& CurrentNodePair : SelectedAnimNodes)
 	{
 		if (const UAnimGraphNode_AnimDynamics* const EditorAnimDynamicsNode = CastChecked<UAnimGraphNode_AnimDynamics>(CurrentNodePair.EditorAnimNode))
 		{
@@ -144,7 +144,7 @@ void FAnimDynamicsEditMode::Render(const FSceneView* View, FViewport* Viewport,
 	check(PDI);
 	check(PreviewSkelMeshComp);
 
-	for (const EditorRuntimeNodePair& CurrentNodePair : AnimNodes)
+	for (const EditorRuntimeNodePair& CurrentNodePair : SelectedAnimNodes)
 	{
 		const UAnimGraphNode_AnimDynamics* const EditorAnimDynamicsNode = CastChecked<UAnimGraphNode_AnimDynamics>(CurrentNodePair.EditorAnimNode);
 
@@ -531,7 +531,7 @@ FName FAnimDynamicsEditMode::GetSelectedBone() const
 
 const UAnimGraphNode_AnimDynamics* const FAnimDynamicsEditMode::FindSelectedEditorAnimNode(const int32 InEditorNodeId) const
 {
-	const FAnimNodeEditMode::EditorRuntimeNodePair* const FoundNodePair = AnimNodes.FindByPredicate([InEditorNodeId](const FAnimNodeEditMode::EditorRuntimeNodePair& Element) { return Element.EditorAnimNode && Element.EditorAnimNode->GetUniqueID() == InEditorNodeId; });
+	const FAnimNodeEditMode::EditorRuntimeNodePair* const FoundNodePair = SelectedAnimNodes.FindByPredicate([InEditorNodeId](const FAnimNodeEditMode::EditorRuntimeNodePair& Element) { return Element.EditorAnimNode && Element.EditorAnimNode->GetUniqueID() == InEditorNodeId; });
 
 	if (FoundNodePair)
 	{
@@ -543,7 +543,7 @@ const UAnimGraphNode_AnimDynamics* const FAnimDynamicsEditMode::FindSelectedEdit
 
 UAnimGraphNode_AnimDynamics* const FAnimDynamicsEditMode::FindSelectedEditorAnimNode(const int32 InEditorNodeId)
 {
-	const FAnimNodeEditMode::EditorRuntimeNodePair* const FoundNodePair = AnimNodes.FindByPredicate([InEditorNodeId](const FAnimNodeEditMode::EditorRuntimeNodePair& Element) { return Element.EditorAnimNode && Element.EditorAnimNode->GetUniqueID() == InEditorNodeId; });
+	const FAnimNodeEditMode::EditorRuntimeNodePair* const FoundNodePair = SelectedAnimNodes.FindByPredicate([InEditorNodeId](const FAnimNodeEditMode::EditorRuntimeNodePair& Element) { return Element.EditorAnimNode && Element.EditorAnimNode->GetUniqueID() == InEditorNodeId; });
 
 	if (FoundNodePair)
 	{
@@ -788,7 +788,7 @@ void FAnimDynamicsEditMode::Tick(FEditorViewportClient* ViewportClient, float De
 
 	USkeletalMeshComponent* const PreviewSkelMeshComp = GetAnimPreviewScene().GetPreviewMeshComponent();
 
-	for (const EditorRuntimeNodePair& CurrentNodePair : AnimNodes)
+	for (const EditorRuntimeNodePair& CurrentNodePair : SelectedAnimNodes)
 	{
 		const UAnimGraphNode_AnimDynamics* const EditorAnimDynamicsNode = CastChecked<UAnimGraphNode_AnimDynamics>(CurrentNodePair.EditorAnimNode);
 
diff --git a/Engine/Source/Editor/AnimGraph/Private/EditModes/AnimNodeEditMode.cpp b/Engine/Source/Editor/AnimGraph/Private/EditModes/AnimNodeEditMode.cpp
--- a/Engine/Source/Editor/AnimGraph/Private/EditModes/AnimNodeEditMode.cpp
+++ b/Engine/Source/Editor/AnimGraph/Private/EditModes/AnimNodeEditMode.cpp
@@ -8,9 +8,15 @@
 #include "EngineUtils.h"
 #include "AnimGraphNode_SkeletalControlBase.h"
 #include "AssetEditorModeManager.h"
+#include "Templates/Tuple.h"
 
 #define LOCTEXT_NAMESPACE "AnimNodeEditMode"
 
+const bool operator==(const FAnimNodeEditMode::EditorRuntimeNodePair& Lhs, const FAnimNodeEditMode::EditorRuntimeNodePair& Rhs)
+{
+	return (Lhs.EditorAnimNode == Rhs.EditorAnimNode) && (Lhs.RuntimeAnimNode == Rhs.RuntimeAnimNode);
+}
+
 FAnimNodeEditMode::FAnimNodeEditMode()
 	: bManipulating(false)
 	, bInTransaction(false)
@@ -35,7 +41,7 @@ IPersonaPreviewScene& FAnimNodeEditMode::GetAnimPreviewScene() const
 
 void FAnimNodeEditMode::GetOnScreenDebugInfo(TArray<FText>& OutDebugInfo) const
 {
-	for (EditorRuntimeNodePair CurrentNodePair : AnimNodes)
+	for (EditorRuntimeNodePair CurrentNodePair : SelectedAnimNodes)
 	{
 		if ((CurrentNodePair.EditorAnimNode != nullptr) && (CurrentNodePair.RuntimeAnimNode != nullptr))
 		{
@@ -115,7 +121,7 @@ void FAnimNodeEditMode::EnterMode(UAnimGraphNode_Base* InEditorNode, FAnimNode_B
 
 	if (InEditorNode && InRuntimeNode)
 	{
-		AnimNodes.Add(EditorRuntimeNodePair(InEditorNode, InRuntimeNode));
+		SelectedAnimNodes.Add(EditorRuntimeNodePair(InEditorNode, InRuntimeNode));
 
 		UAnimGraphNode_SkeletalControlBase* SkelControl = Cast<UAnimGraphNode_SkeletalControlBase>(InEditorNode);
 		if (SkelControl != nullptr)
@@ -133,7 +139,7 @@ void FAnimNodeEditMode::EnterMode(UAnimGraphNode_Base* InEditorNode, FAnimNode_B
 
 void FAnimNodeEditMode::ExitMode()
 {
-	for (EditorRuntimeNodePair CurrentNodePair : AnimNodes)
+	for (EditorRuntimeNodePair CurrentNodePair : SelectedAnimNodes)
 	{
 		if (CurrentNodePair.EditorAnimNode != nullptr)
 		{
@@ -147,23 +153,60 @@ void FAnimNodeEditMode::ExitMode()
 		}
 	}
 
-	AnimNodes.Empty();
+	SelectedAnimNodes.Empty();
+	PoseWatchedAnimNodes.Empty();
 }
 
 void FAnimNodeEditMode::Render(const FSceneView* View, FViewport* Viewport, FPrimitiveDrawInterface* PDI)
 {
-	for (EditorRuntimeNodePair CurrentNodePair : AnimNodes)
+	USkeletalMeshComponent* const PreviewSkelMeshComp = GetAnimPreviewScene().GetPreviewMeshComponent();
+
+	check(View);
+	check(PDI);
+	check(PreviewSkelMeshComp);
+
+	// Build a unique list of all selected or pose watched nodes, 0 = Node, 1 = IsSelected, 2 = IsPoseWatchEnabled. 
+	using DrawParameters = TTuple<UAnimGraphNode_Base*, bool, bool>;
+
+	TArray< DrawParameters > DrawParameterList;
+
+	// Collect selected nodes.	
+	for (const FAnimNodeEditMode::EditorRuntimeNodePair& CurrentNodePair : SelectedAnimNodes)
+	{
+		if (CurrentNodePair.EditorAnimNode != nullptr)
+		{
+			DrawParameterList.Add(DrawParameters(CurrentNodePair.EditorAnimNode, true, false));
+		}
+	}
+
+	// Collect pose watched nodes.
+	for (const FAnimNodeEditMode::EditorRuntimeNodePair& CurrentNodePair : PoseWatchedAnimNodes)
 	{
 		if (CurrentNodePair.EditorAnimNode != nullptr)
 		{
-			CurrentNodePair.EditorAnimNode->Draw(PDI, GetAnimPreviewScene().GetPreviewMeshComponent());
+			if (DrawParameters* Parameter = DrawParameterList.FindByPredicate([CurrentNodePair](DrawParameters& Other) { return Other.Get<0>() == CurrentNodePair.EditorAnimNode; }))
+			{
+				// Add pose watch flag to existing, selected node's draw parameters.
+				Parameter->Get<2>() = true;
+			}
+			else
+			{
+				// Create a new draw parameter for this pose watched but not selected node.
+				DrawParameterList.Add(DrawParameters(CurrentNodePair.EditorAnimNode, false, true));
+			}
 		}
 	}
+
+	// Draw all collected nodes.
+	for (DrawParameters CurrentParameters : DrawParameterList)
+	{
+		CurrentParameters.Get<0>()->Draw(PDI, PreviewSkelMeshComp, CurrentParameters.Get<1>(), CurrentParameters.Get<2>());
+	}
 }
 
 void FAnimNodeEditMode::DrawHUD(FEditorViewportClient* ViewportClient, FViewport* Viewport, const FSceneView* View, FCanvas* Canvas)
 {
-	for (EditorRuntimeNodePair CurrentNodePair : AnimNodes)
+	for (EditorRuntimeNodePair CurrentNodePair : SelectedAnimNodes)
 	{
 		if (CurrentNodePair.EditorAnimNode != nullptr)
 		{
@@ -179,7 +222,7 @@ bool FAnimNodeEditMode::HandleClick(FEditorViewportClient* InViewportClient, HHi
 		HActor* ActorHitProxy = static_cast<HActor*>(HitProxy);
 		GetAnimPreviewScene().SetSelectedActor(ActorHitProxy->Actor);
 
-		for (EditorRuntimeNodePair CurrentNodePair : AnimNodes)
+		for (EditorRuntimeNodePair CurrentNodePair : SelectedAnimNodes)
 		{
 			UAnimGraphNode_SkeletalControlBase* SkelControl = Cast<UAnimGraphNode_SkeletalControlBase>(CurrentNodePair.EditorAnimNode);
 			if (SkelControl != nullptr)
@@ -216,7 +259,7 @@ bool FAnimNodeEditMode::StartTracking(FEditorViewportClient* InViewportClient, F
 	{
 		GEditor->BeginTransaction(LOCTEXT("EditSkelControlNodeTransaction", "Edit Skeletal Control Node"));
 
-		for (EditorRuntimeNodePair CurrentNodePair : AnimNodes)
+		for (EditorRuntimeNodePair CurrentNodePair : SelectedAnimNodes)
 		{
 			if (CurrentNodePair.EditorAnimNode != nullptr)
 			{
@@ -304,7 +347,7 @@ bool FAnimNodeEditMode::InputDelta(FEditorViewportClient* InViewportClient, FVie
 			DoScale(InScale);
 		}
 
-		for (EditorRuntimeNodePair CurrentNodePair : AnimNodes)
+		for (EditorRuntimeNodePair CurrentNodePair : SelectedAnimNodes)
 		{
 			if (CurrentNodePair.EditorAnimNode)
 			{
@@ -398,14 +441,19 @@ void FAnimNodeEditMode::Exit()
 {
 	IAnimNodeEditMode::Exit();
 
-	AnimNodes.Empty();
+	SelectedAnimNodes.Empty();
+}
+
+void FAnimNodeEditMode::RegisterPoseWatchedNode(UAnimGraphNode_Base* InEditorNode, FAnimNode_Base* InRuntimeNode)
+{
+	PoseWatchedAnimNodes.Add(FAnimNodeEditMode::EditorRuntimeNodePair(InEditorNode, InRuntimeNode));
 }
 
 UAnimGraphNode_Base* FAnimNodeEditMode::GetActiveWidgetAnimNode() const
 {
-	if (AnimNodes.Num() > 0)
+	if (SelectedAnimNodes.Num() > 0)
 	{
-		return AnimNodes.Last().EditorAnimNode;
+		return SelectedAnimNodes.Last().EditorAnimNode;
 	}
 
 	return nullptr;
@@ -413,9 +461,9 @@ UAnimGraphNode_Base* FAnimNodeEditMode::GetActiveWidgetAnimNode() const
 
 FAnimNode_Base* FAnimNodeEditMode::GetActiveWidgetRuntimeAnimNode() const
 {
-	if (AnimNodes.Num() > 0)
+	if (SelectedAnimNodes.Num() > 0)
 	{
-		return AnimNodes.Last().RuntimeAnimNode;
+		return SelectedAnimNodes.Last().RuntimeAnimNode;
 	}
 
 	return nullptr;
diff --git a/Engine/Source/Editor/AnimGraph/Public/AnimGraphNode_Base.h b/Engine/Source/Editor/AnimGraph/Public/AnimGraphNode_Base.h
--- a/Engine/Source/Editor/AnimGraph/Public/AnimGraphNode_Base.h
+++ b/Engine/Source/Editor/AnimGraph/Public/AnimGraphNode_Base.h
@@ -332,6 +332,9 @@ public:
 	 */
 	virtual void OnNodeSelected(bool bInIsSelected, class FEditorModeTools& InModeTools, struct FAnimNode_Base* InRuntimeNode);
 
+	/** Pose Watch change notification callback. Should be called every time a pose watch on this node is created or destroyed. */
+	virtual void OnPoseWatchChanged(const bool IsPoseWatchEnabled, TObjectPtr<class UPoseWatch> InPoseWatch, FEditorModeTools& InModeTools, FAnimNode_Base* InRuntimeNode);
+
 	/**
 	 * Override this function to push an editor mode when this node is selected
 	 * @return the editor mode to use when this node is selected
@@ -340,8 +343,16 @@ public:
 
 	// Draw function for supporting visualization
 	virtual void Draw(FPrimitiveDrawInterface* PDI, USkeletalMeshComponent * PreviewSkelMeshComp) const {}
+	/**
+	 *	Draw function called on nodes that are selected and / or have a pose watch enabled.
+	 *	Default implementation calls the basic draw function for selected nodes and does nothing for pose watched nodes. Nodes
+	 *	that should render something when a pose watch is enabled but they are not selected should override this function.
+	 */
+	virtual void Draw(FPrimitiveDrawInterface* PDI, USkeletalMeshComponent* PreviewSkelMeshComp, const bool bIsSelected, const bool bIsPoseWatchEnabled) const;
+
 	// Canvas draw function to draw to viewport
 	virtual void DrawCanvas(FViewport& InViewport, FSceneView& View, FCanvas& Canvas, USkeletalMeshComponent * PreviewSkelMeshComp) const {}
+
 	// Function to collect strings from nodes to display in the viewport.
 	// Use this rather than DrawCanvas when adding general text to the viewport.
 	virtual void GetOnScreenDebugInfo(TArray<FText>& DebugInfo, FAnimNode_Base* RuntimeAnimNode, USkeletalMeshComponent* PreviewSkelMeshComp) const {}
diff --git a/Engine/Source/Editor/AnimGraph/Public/AnimGraphNode_RigidBody.h b/Engine/Source/Editor/AnimGraph/Public/AnimGraphNode_RigidBody.h
--- a/Engine/Source/Editor/AnimGraph/Public/AnimGraphNode_RigidBody.h
+++ b/Engine/Source/Editor/AnimGraph/Public/AnimGraphNode_RigidBody.h
@@ -23,7 +23,8 @@ public:
 	// End of UEdGraphNode interface
 
 	// UAnimGraphNode_Base
-	virtual void Draw(FPrimitiveDrawInterface* PDI, USkeletalMeshComponent* PreviewSkelMeshComp) const override;
+	virtual void Draw(FPrimitiveDrawInterface* PDI, USkeletalMeshComponent* PreviewSkelMeshComp, const bool bIsSelected, const bool bIsPoseWatchEnabled) const override;
+	virtual void OnPoseWatchChanged(const bool IsPoseWatchActive, TObjectPtr<UPoseWatch> InPoseWatch, FEditorModeTools& InModeTools, FAnimNode_Base* InRuntimeNode) override;
 
 	virtual void CustomizeDetails(IDetailLayoutBuilder& DetailBuilder) override;
 	virtual void PostEditChangeProperty(struct FPropertyChangedEvent& PropertyChangedEvent) override;
@@ -39,4 +40,8 @@ protected:
 
 	bool AreAnyBodiesHidden() const;
 	bool AreAnyConstraintsHidden() const;
+
+private:
+	TWeakObjectPtr<UPoseWatchElement> PoseWatchElementBodies;
+	TWeakObjectPtr<UPoseWatchElement> PoseWatchElementConstraints;
 };
diff --git a/Engine/Source/Editor/AnimGraph/Public/AnimNodeEditMode.h b/Engine/Source/Editor/AnimGraph/Public/AnimNodeEditMode.h
--- a/Engine/Source/Editor/AnimGraph/Public/AnimNodeEditMode.h
+++ b/Engine/Source/Editor/AnimGraph/Public/AnimNodeEditMode.h
@@ -53,6 +53,22 @@ public:
 	virtual void Tick(FEditorViewportClient* ViewportClient, float DeltaTime) override;
 	virtual void Exit() override;
 
+	virtual void RegisterPoseWatchedNode(UAnimGraphNode_Base* InEditorNode, FAnimNode_Base* InRuntimeNode);
+
+	struct EditorRuntimeNodePair
+	{
+		EditorRuntimeNodePair(UAnimGraphNode_Base* InEditorAnimNode, FAnimNode_Base* InRuntimeAnimNode)
+			: EditorAnimNode(InEditorAnimNode)
+			, RuntimeAnimNode(InRuntimeAnimNode)
+		{}
+
+		/** The node we are operating on */
+		UAnimGraphNode_Base* EditorAnimNode;
+
+		/** The runtime node in the preview scene */
+		FAnimNode_Base* RuntimeAnimNode;
+	};
+
 protected:
 	// local conversion functions for drawing
 	static void ConvertToComponentSpaceTransform(const USkeletalMeshComponent* SkelComp, const FTransform & InTransform, FTransform & OutCSTransform, int32 BoneIndex, EBoneControlSpace Space);
@@ -71,23 +87,8 @@ protected:
 
 	const bool IsManipulatingWidget() const { return bManipulating; }
 
-protected:
-
-	struct EditorRuntimeNodePair
-	{
-		EditorRuntimeNodePair(UAnimGraphNode_Base* InEditorAnimNode, FAnimNode_Base* InRuntimeAnimNode)
-			: EditorAnimNode(InEditorAnimNode)
-			, RuntimeAnimNode(InRuntimeAnimNode)
-		{}
-
-		/** The node we are operating on */
-		UAnimGraphNode_Base* EditorAnimNode;
-
-		/** The runtime node in the preview scene */
-		FAnimNode_Base* RuntimeAnimNode;
-	};
-
-	TArray< EditorRuntimeNodePair > AnimNodes;
+	TArray< EditorRuntimeNodePair > SelectedAnimNodes;	// Selected Anim Graph Nodes
+	TArray< EditorRuntimeNodePair > PoseWatchedAnimNodes; 	// Pose Watched Anim Graph Nodes. 
 
 private:
 	bool bManipulating;
diff --git a/Engine/Source/Editor/AnimGraph/Public/IAnimNodeEditMode.h b/Engine/Source/Editor/AnimGraph/Public/IAnimNodeEditMode.h
--- a/Engine/Source/Editor/AnimGraph/Public/IAnimNodeEditMode.h
+++ b/Engine/Source/Editor/AnimGraph/Public/IAnimNodeEditMode.h
@@ -6,6 +6,9 @@
 #include "UnrealWidgetFwd.h"
 #include "AnimationEditMode.h"
 
+class UAnimGraphNode_Base;
+struct FAnimNode_Base;
+
 /** Base interface for skeletal control edit modes */
 class IAnimNodeEditMode : public FAnimationEditMode
 {
@@ -35,8 +38,11 @@ public:
 	virtual void DoScale(FVector& InScale) = 0;
 
 	/** Called when entering this edit mode */
-	virtual void EnterMode(class UAnimGraphNode_Base* InEditorNode, struct FAnimNode_Base* InRuntimeNode) = 0;
+	virtual void EnterMode(UAnimGraphNode_Base* InEditorNode, FAnimNode_Base* InRuntimeNode) = 0;
 
 	/** Called when exiting this edit mode */
 	virtual void ExitMode() = 0;
+
+	/** Called when a Pose Watch is created on a node edited by this edit mode */
+	virtual void RegisterPoseWatchedNode(UAnimGraphNode_Base* InEditorNode, FAnimNode_Base* InRuntimeNode) = 0;
 };
diff --git a/Engine/Source/Editor/AnimationBlueprintEditor/Private/AnimationBlueprintEditor.cpp b/Engine/Source/Editor/AnimationBlueprintEditor/Private/AnimationBlueprintEditor.cpp
--- a/Engine/Source/Editor/AnimationBlueprintEditor/Private/AnimationBlueprintEditor.cpp
+++ b/Engine/Source/Editor/AnimationBlueprintEditor/Private/AnimationBlueprintEditor.cpp
@@ -156,6 +156,40 @@ void FAnimationBlueprintEditor::NotifyAllNodesOnSelection(const bool bInIsSelect
 	}
 }
 
+void FAnimationBlueprintEditor::NotifyAllNodesOnPoseWatchChanged(const bool IsPoseWatchEnabled)
+{
+	UAnimBlueprint* const AnimBP = GetAnimBlueprint();
+
+	if (AnimBP)
+	{
+		FEditorModeTools& ModeTools = GetEditorModeManager();
+
+		for (TObjectPtr<UPoseWatch> CurrentPoseWatch : AnimBP->PoseWatches)
+		{
+			UAnimGraphNode_Base* const CurrentAnimGraphNodePtr = Cast<UAnimGraphNode_Base>(CurrentPoseWatch->Node.Get());
+			FAnimNode_Base* const PreviewNode = FindAnimNode(CurrentAnimGraphNodePtr);
+
+			// Note: Potentially passing a null PreviewNode ptr when IsPoseWatchEnabled is false is required to un-watch nodes that no longer exist.
+			if (CurrentAnimGraphNodePtr && (!IsPoseWatchEnabled || PreviewNode))
+			{
+				CurrentAnimGraphNodePtr->OnPoseWatchChanged(IsPoseWatchEnabled, CurrentPoseWatch, ModeTools, PreviewNode);
+			}
+		}
+	}
+}
+
+void FAnimationBlueprintEditor::ReleaseAllManagedNodes()
+{
+	NotifyAllNodesOnSelection(false);
+	NotifyAllNodesOnPoseWatchChanged(false);
+}
+
+void FAnimationBlueprintEditor::AcquireAllManagedNodes()
+{
+	NotifyAllNodesOnSelection(true);
+	NotifyAllNodesOnPoseWatchChanged(true);
+}
+
 /////////////////////////////////////////////////////
 // SAnimBlueprintPreviewPropertyEditor
 
@@ -741,6 +775,8 @@ bool FAnimationBlueprintEditor::CanTogglePoseWatch()
 
 void FAnimationBlueprintEditor::OnTogglePoseWatch()
 {
+	ReleaseAllManagedNodes();
+
 	const FGraphPanelSelectionSet SelectedNodes = GetSelectedNodes();
 	UAnimBlueprint* AnimBP = GetAnimBlueprint();
 
@@ -773,6 +809,8 @@ void FAnimationBlueprintEditor::OnTogglePoseWatch()
 			}
 		}
 	}
+
+	AcquireAllManagedNodes();
 }
 
 // Helper function for node conversions
@@ -1522,7 +1560,7 @@ void FAnimationBlueprintEditor::ClearSelectedActor()
 
 void FAnimationBlueprintEditor::ClearSelectedAnimGraphNodes()
 {
-	NotifyAllNodesOnSelection(false);
+	ReleaseAllManagedNodes();
 
 	SelectedAnimGraphNodes.Empty();
 }
@@ -1882,7 +1920,7 @@ void FAnimationBlueprintEditor::OnSelectedNodesChangedImpl(const TSet<class UObj
 			SortedContainerDifference(OldSelectionSorted, NewSelectionSorted, AddSelection, RemSelection, SortPredicate);
 		}
 		
-		NotifyAllNodesOnSelection(false); // Register de-selection with all the previously selected nodes.
+		ReleaseAllManagedNodes(); // Register de-selection with all the previously selected nodes.
 
 
 		// Remove all the nodes that are no longer selected.
@@ -1897,7 +1935,7 @@ void FAnimationBlueprintEditor::OnSelectedNodesChangedImpl(const TSet<class UObj
 			SelectedAnimGraphNodes.Add(CurrentAnimGraphNode);
 		}
 
-		NotifyAllNodesOnSelection(true); // Register re-selection with all the currently selected nodes.
+		AcquireAllManagedNodes(); // Register re-selection with all the currently selected nodes.
 	}
 
 	bSelectRegularNode = false;
@@ -1934,6 +1972,8 @@ void FAnimationBlueprintEditor::HandlePoseWatchSelectedNodes()
 	TSharedPtr<SGraphEditor> FocusedGraphEd = FocusedGraphEdPtr.Pin();
 	if (FocusedGraphEd.IsValid())
 	{
+		ReleaseAllManagedNodes(); // Register de-selection with all the previously selected nodes.
+
 		UAnimBlueprint* AnimBP = GetAnimBlueprint();
 		TArray<UEdGraphNode*> AllNodes = FocusedGraphEd->GetCurrentGraph()->Nodes;
 
@@ -1962,6 +2002,8 @@ void FAnimationBlueprintEditor::HandlePoseWatchSelectedNodes()
 				}
 			}
 		}
+
+		AcquireAllManagedNodes(); // Register re-selection with all the currently selected nodes.
 	}
 }
 
@@ -1970,6 +2012,8 @@ void FAnimationBlueprintEditor::RemoveAllSelectionPoseWatches()
 	TSharedPtr<SGraphEditor> FocusedGraphEd = FocusedGraphEdPtr.Pin();
 	if (FocusedGraphEd.IsValid())
 	{
+		ReleaseAllManagedNodes(); // Register de-selection with all the previously selected nodes.
+
 		UAnimBlueprint* AnimBP = GetAnimBlueprint();
 		TArray<UEdGraphNode*> AllNodes = FocusedGraphEd->GetCurrentGraph()->Nodes;
 
@@ -1985,6 +2029,8 @@ void FAnimationBlueprintEditor::RemoveAllSelectionPoseWatches()
 				}
 			}
 		}
+
+		AcquireAllManagedNodes(); // Register re-selection with all the currently selected nodes.
 	}
 }
 void FAnimationBlueprintEditor::OnPostCompile()
diff --git a/Engine/Source/Editor/AnimationBlueprintEditor/Private/AnimationBlueprintEditor.h b/Engine/Source/Editor/AnimationBlueprintEditor/Private/AnimationBlueprintEditor.h
--- a/Engine/Source/Editor/AnimationBlueprintEditor/Private/AnimationBlueprintEditor.h
+++ b/Engine/Source/Editor/AnimationBlueprintEditor/Private/AnimationBlueprintEditor.h
@@ -279,6 +279,15 @@ private:
 	/** Call OnNodeSelected for each selected node **/
 	void NotifyAllNodesOnSelection(const bool bInIsSelected);
 
+	/** Call OnPoseWatchChanged for each pose watched node **/
+	void NotifyAllNodesOnPoseWatchChanged(const bool IsPoseWatchActive);
+
+	/** Called to notify all Nodes before any change to node selection or pose watch status **/
+	void ReleaseAllManagedNodes();
+
+	/** Called to notify all Nodes after any change to node selection or pose watch status **/
+	void AcquireAllManagedNodes();
+
 	/** Helper function used to keep skeletal controls in preview & instance in sync */
 	struct FAnimNode_Base* FindAnimNode(class UAnimGraphNode_Base* AnimGraphNode) const;
 
@@ -297,7 +306,7 @@ private:
 	/** Enable/disable pose watch on selected nodes */
 	void HandlePoseWatchSelectedNodes();
 
-	/** Removes all pose watches created by selectionfrom the current view */
+	/** Removes all pose watches created by selection from the current view */
 	void RemoveAllSelectionPoseWatches();
 
     /**
@@ -314,7 +323,7 @@ private:
 
 	void HandleUpdateSettings(const UAnimationBlueprintEditorSettings* AnimationBlueprintEditorSettings, EPropertyChangeType::Type ChangeType);
 
-	/** Chooses a suitable pose watch colour automatically - i.e. one that isn't already in use (if possible) */
+	/** Chooses a suitable pose watch color automatically - i.e. one that isn't already in use (if possible) */
 	FColor ChoosePoseWatchColor() const;
 
 	// Pose pin UI handlers
diff --git a/Engine/Source/Editor/Persona/Private/AnimationEditorViewportClient.cpp b/Engine/Source/Editor/Persona/Private/AnimationEditorViewportClient.cpp
--- a/Engine/Source/Editor/Persona/Private/AnimationEditorViewportClient.cpp
+++ b/Engine/Source/Editor/Persona/Private/AnimationEditorViewportClient.cpp
@@ -1587,9 +1587,9 @@ void FAnimationViewportClient::DrawWatchedPoses(UDebugSkelMeshComponent * MeshCo
 			if (Blueprint->GetObjectBeingDebugged())
 			{
 				const FAnimBlueprintDebugData& DebugData = AnimBlueprintGeneratedClass->GetAnimBlueprintDebugData();
-				DebugData.ForEachActiveVisiblePoseWatch([MeshComponent, PDI](const FPoseWatchDebugData& PoseWatchDebugData)
+				DebugData.ForEachActiveVisiblePoseWatchPoseElement([MeshComponent, PDI](const FPoseWatchDebugData& PoseWatchDebugData)
 				{
-					SkeletalDebugRendering::DrawBonesFromPoseWatch(*PoseWatchDebugData.PoseInfo.Get(), MeshComponent, PDI, PoseWatchDebugData.PoseWatch);
+					SkeletalDebugRendering::DrawBonesFromPoseWatch(*PoseWatchDebugData.PoseInfo.Get(), MeshComponent, PDI, PoseWatchDebugData.PoseWatchPoseElement);
 				});
 			}
 		}
diff --git a/Engine/Source/Editor/Persona/Private/IPoseWatchManagerTreeItem.h b/Engine/Source/Editor/Persona/Private/IPoseWatchManagerTreeItem.h
--- a/Engine/Source/Editor/Persona/Private/IPoseWatchManagerTreeItem.h
+++ b/Engine/Source/Editor/Persona/Private/IPoseWatchManagerTreeItem.h
@@ -37,6 +37,9 @@ public:
 	/** Sets this tree item's visibility */
 	virtual void SetIsVisible(const bool bVisible) = 0;
 
+	/** Sets whether or not this tree item is expanded  */
+	virtual void SetIsExpanded(const bool bIsExpanded) {}
+
 	/** Add a child to this tree item */
 	void AddChild(FPoseWatchManagerTreeItemRef Child)
 	{
@@ -98,6 +101,15 @@ public:
 
 	virtual void OnRemoved() = 0;
 
+	/** Returns the view port render color */
+	virtual FColor GetColor() const { return FColor(); }
+
+	/** Sets the view port render color */
+	virtual void SetColor(const FColor& InColor) {}
+
+	/** Returns true if the view port render color for this item can be changed */
+	virtual bool ShouldDisplayColorPicker() const { return false; }
+
 protected:
 	/** This item's parent tree item */
 	TWeakPtr<IPoseWatchManagerTreeItem> Parent;
@@ -111,3 +123,13 @@ protected:
 	/** Tree item type identifier */
 	EPoseWatchTreeItemType TreeType;
 };
+
+template <> inline const IPoseWatchManagerTreeItem* IPoseWatchManagerTreeItem::CastTo<IPoseWatchManagerTreeItem>() const
+{
+	return this;
+}
+
+template <> inline IPoseWatchManagerTreeItem* IPoseWatchManagerTreeItem::CastTo<IPoseWatchManagerTreeItem>()
+{
+	return this;
+}
\ No newline at end of file
diff --git a/Engine/Source/Editor/Persona/Private/PhysicsAssetRenderUtils.cpp b/Engine/Source/Editor/Persona/Private/PhysicsAssetRenderUtils.cpp
--- a/Engine/Source/Editor/Persona/Private/PhysicsAssetRenderUtils.cpp
+++ b/Engine/Source/Editor/Persona/Private/PhysicsAssetRenderUtils.cpp
@@ -5,6 +5,7 @@
 #include "Materials/MaterialInstanceDynamic.h"
 #include "Components/SkeletalMeshComponent.h"
 #include "Materials/Material.h"
+#include "Math/UnrealMathVectorCommon.h"
 #include "Preferences/PhysicsAssetEditorOptions.h"
 #include "PhysicsEngine/PhysicsAsset.h"
 #include "PhysicsEngine/PhysicsConstraintTemplate.h"
@@ -639,11 +640,42 @@ uint32 UPhysicsAssetRenderUtilities::GetPathNameHash(const FString& InPhysicsAss
 }
 
 // class FPhysicsAssetRenderInterface
-void FPhysicsAssetRenderInterface::DebugDraw(class USkeletalMeshComponent* const SkeletalMeshComponent, class UPhysicsAsset* const PhysicsAsset, FPrimitiveDrawInterface* PDI)
+void FPhysicsAssetRenderInterface::DebugDraw(USkeletalMeshComponent* const SkeletalMeshComponent, UPhysicsAsset* const PhysicsAsset, FPrimitiveDrawInterface* PDI)
 {
 	PhysicsAssetRender::DebugDraw(SkeletalMeshComponent, PhysicsAsset, PDI);
 }
 
+void FPhysicsAssetRenderInterface::DebugDrawBodies(USkeletalMeshComponent* const SkeletalMeshComponent, UPhysicsAsset* const PhysicsAsset, FPrimitiveDrawInterface* PDI, const FColor& PrimitiveColorOverride)
+{
+	auto TransformFn = [](const UPhysicsAsset* PhysicsAsset, const FTransform& BoneTM, const int32 BodyIndex, const EAggCollisionShape::Type PrimType, const int32 PrimIndex, const float Scale) { return PhysicsAssetRender::GetPrimitiveTransform(PhysicsAsset, BoneTM, BodyIndex, PrimType, PrimIndex, Scale);  };
+	auto HitProxyFn  = [](const int32 BodyIndex, const EAggCollisionShape::Type PrimitiveType, const int32 PrimitiveIndex) { return nullptr; };
+
+	auto MaterialFn = [](const int32 BodyIndex, const EAggCollisionShape::Type PrimitiveType, const int32 PrimitiveIndex, const FPhysicsAssetRenderSettings& Settings)
+	{ 
+		return PhysicsAssetRender::GetPrimitiveMaterial(BodyIndex, PrimitiveType, PrimitiveIndex, Settings); 
+	};
+
+	auto ColorFn = [&PrimitiveColorOverride](const int32 BodyIndex, const EAggCollisionShape::Type PrimitiveType, const int32 PrimitiveIndex, const FPhysicsAssetRenderSettings& Settings)
+	{ 
+		if (PrimitiveColorOverride.A > (1.0f - KINDA_SMALL_NUMBER))
+		{
+			return PrimitiveColorOverride;
+		}
+
+		const FLinearColor LinearOverride = PrimitiveColorOverride.ReinterpretAsLinear();
+		return FMath::Lerp(PhysicsAssetRender::GetPrimitiveColor(BodyIndex, PrimitiveType, PrimitiveIndex, Settings).ReinterpretAsLinear(), LinearOverride, LinearOverride.A).QuantizeRound();
+	};
+
+	PhysicsAssetRender::DebugDrawBodies(SkeletalMeshComponent, PhysicsAsset, PDI, ColorFn, MaterialFn, TransformFn, HitProxyFn);
+}
+
+void FPhysicsAssetRenderInterface::DebugDrawConstraints(USkeletalMeshComponent* const SkeletalMeshComponent, UPhysicsAsset* const PhysicsAsset, FPrimitiveDrawInterface* PDI)
+{
+	auto HitProxyFn   = [](const int32) { return nullptr; };
+	auto IsSelectedFn = [](const uint32) { return false; };
+	PhysicsAssetRender::DebugDrawConstraints(SkeletalMeshComponent, PhysicsAsset, PDI, IsSelectedFn, false, HitProxyFn);
+}
+
 void FPhysicsAssetRenderInterface::SaveConfig()
 {
 	if (UPhysicsAssetRenderUtilities* PhysicsAssetRenderUtilities = GetMutableDefault<UPhysicsAssetRenderUtilities>())
diff --git a/Engine/Source/Editor/Persona/Private/PoseWatchManagerColumnColor.cpp b/Engine/Source/Editor/Persona/Private/PoseWatchManagerColumnColor.cpp
--- a/Engine/Source/Editor/Persona/Private/PoseWatchManagerColumnColor.cpp
+++ b/Engine/Source/Editor/Persona/Private/PoseWatchManagerColumnColor.cpp
@@ -10,6 +10,7 @@
 #include "Widgets/Views/STreeView.h"
 #include "IPoseWatchManager.h"
 #include "Widgets/Colors/SColorPicker.h"
+#include "PoseWatchManagerElementTreeItem.h"
 #include "PoseWatchManagerPoseWatchTreeItem.h"
 #include "Engine/PoseWatch.h"
 
@@ -21,9 +22,9 @@ public:
 	SLATE_BEGIN_ARGS(SColorBoxWidget) {}
 	SLATE_END_ARGS()
 
-	void Construct(const FArguments& InArgs, FPoseWatchManagerPoseWatchTreeItem* InPoseWatchTreeItem)
+	void Construct(const FArguments& InArgs, IPoseWatchManagerTreeItem* InTreeItem)
 	{
-		PoseWatchTreeItem = InPoseWatchTreeItem;
+		TreeItem = InTreeItem;
 
 		ChildSlot
 		[
@@ -57,17 +58,17 @@ public:
 
 	FLinearColor OnGetColorForColorBlock() const
 	{
-		return PoseWatchTreeItem->PoseWatch->GetColor();
+		return TreeItem->GetColor();
 	}
 
 	void OnSetColorFromColorPicker(FLinearColor NewColor)
 	{
-		PoseWatchTreeItem->PoseWatch->SetColor(NewColor.ToFColorSRGB());
+		TreeItem->SetColor(NewColor.ToFColorSRGB());
 	}
 
 	void OnColorPickerCancelled(FLinearColor OriginalColor)
 	{
-		PoseWatchTreeItem->PoseWatch->SetColor(OriginalColor.ToFColorSRGB());
+		TreeItem->SetColor(OriginalColor.ToFColorSRGB());
 	}
 
 	FReply OnMouseButtonDownColorBlock(const FGeometry& MyGeometry, const FPointerEvent& MouseEvent)
@@ -78,7 +79,7 @@ public:
 			PickerArgs.DisplayGamma = TAttribute<float>::Create(TAttribute<float>::FGetter::CreateUObject(GEngine, &UEngine::GetDisplayGamma));
 			PickerArgs.OnColorCommitted = FOnLinearColorValueChanged::CreateSP(this, &SColorBoxWidget::OnSetColorFromColorPicker);
 			PickerArgs.OnColorPickerCancelled = FOnColorPickerCancelled::CreateSP(this, &SColorBoxWidget::OnColorPickerCancelled);
-			PickerArgs.InitialColorOverride = PoseWatchTreeItem->PoseWatch->GetColor();
+			PickerArgs.InitialColorOverride = TreeItem->GetColor();
 			PickerArgs.ParentWidget = ColorPickerParentWidget;
 			PickerArgs.OptionalOwningDetailsView = ColorPickerParentWidget;
 			FWidgetPath ParentWidgetPath;
@@ -101,7 +102,7 @@ public:
 	}
 
 private:
-	FPoseWatchManagerPoseWatchTreeItem* PoseWatchTreeItem;
+	IPoseWatchManagerTreeItem* TreeItem;
 	TSharedPtr<SWidget> ColorPickerParentWidget;
 	TSharedPtr<SWidget > ColorWidgetBackgroundBorder;
 };
@@ -126,9 +127,9 @@ SHeaderRow::FColumn::FArguments FPoseWatchManagerColumnColor::ConstructHeaderRow
 
 const TSharedRef<SWidget> FPoseWatchManagerColumnColor::ConstructRowWidget(FPoseWatchManagerTreeItemRef TreeItem, const STableRow<FPoseWatchManagerTreeItemPtr>& Row)
 {
-	if (FPoseWatchManagerPoseWatchTreeItem* PoseWatchTreeItem = TreeItem->CastTo<FPoseWatchManagerPoseWatchTreeItem>())
+	if (TreeItem->ShouldDisplayColorPicker())
 	{
-		return SNew(SColorBoxWidget, PoseWatchTreeItem);
+		return SNew(SColorBoxWidget, TreeItem->CastTo<IPoseWatchManagerTreeItem>());
 	}
 	return SNullWidget::NullWidget;
 }
diff --git a/Engine/Source/Editor/Persona/Private/PoseWatchManagerDefaultHierarchy.cpp b/Engine/Source/Editor/Persona/Private/PoseWatchManagerDefaultHierarchy.cpp
--- a/Engine/Source/Editor/Persona/Private/PoseWatchManagerDefaultHierarchy.cpp
+++ b/Engine/Source/Editor/Persona/Private/PoseWatchManagerDefaultHierarchy.cpp
@@ -1,6 +1,7 @@
 // Copyright Epic Games, Inc. All Rights Reserved.
 
 #include "PoseWatchManagerDefaultHierarchy.h"
+#include "PoseWatchManagerElementTreeItem.h"
 #include "PoseWatchManagerFolderTreeItem.h"
 #include "PoseWatchManagerPoseWatchTreeItem.h"
 #include "PoseWatchManagerDefaultMode.h"
@@ -13,33 +14,34 @@ FPoseWatchManagerDefaultHierarchy::FPoseWatchManagerDefaultHierarchy(FPoseWatchM
 
 FPoseWatchManagerTreeItemPtr FPoseWatchManagerDefaultHierarchy::FindParent(const IPoseWatchManagerTreeItem& Item, const TMap<FObjectKey, FPoseWatchManagerTreeItemPtr>& Items) const
 {
-	UPoseWatchFolder* ParentFolder = nullptr;
+	FObjectKey ParentIdentifier = nullptr;
+
 	if (const FPoseWatchManagerFolderTreeItem* FolderTreeItem = Item.CastTo<FPoseWatchManagerFolderTreeItem>())
 	{
-		ParentFolder = FolderTreeItem->PoseWatchFolder->GetParent();
-
+		ParentIdentifier = FolderTreeItem->PoseWatchFolder->GetParent();
 	}
 	else if (const FPoseWatchManagerPoseWatchTreeItem* PoseWatchTreeItem = Item.CastTo<FPoseWatchManagerPoseWatchTreeItem>())
 	{
-		ParentFolder = PoseWatchTreeItem->PoseWatch->GetParent();
+		ParentIdentifier = PoseWatchTreeItem->PoseWatch->GetParent();
+	}
+	else if (const FPoseWatchManagerElementTreeItem* ElementTreeItem = Item.CastTo<FPoseWatchManagerElementTreeItem>())
+	{
+		ParentIdentifier = ElementTreeItem->PoseWatchElement->GetParent();
 	}
 
-	if (ParentFolder)
+	const FPoseWatchManagerTreeItemPtr* ParentTreeItem = Items.Find(ParentIdentifier);
+	if (ParentTreeItem)
 	{
-		const FPoseWatchManagerTreeItemPtr* ParentTreeItem = Items.Find(ParentFolder);
-		if (ParentTreeItem)
-		{
-			return *ParentTreeItem;
-		}
+		return *ParentTreeItem;
 	}
-	
+
 	return nullptr;
 }
 
 void FPoseWatchManagerDefaultHierarchy::CreateItems(TArray<FPoseWatchManagerTreeItemPtr>& OutItems) const
 {
 	UAnimBlueprint* AnimBlueprint = Mode->PoseWatchManager->AnimBlueprint;
-	for (int32 Index = 0; Index <  AnimBlueprint->PoseWatchFolders.Num(); ++Index)
+	for (int32 Index = 0; Index < AnimBlueprint->PoseWatchFolders.Num(); ++Index)
 	{
 		TObjectPtr<UPoseWatchFolder>& PoseWatchFolder = AnimBlueprint->PoseWatchFolders[Index];
 		if (PoseWatchFolder)
@@ -60,7 +62,14 @@ void FPoseWatchManagerDefaultHierarchy::CreateItems(TArray<FPoseWatchManagerTree
 			if (!PoseWatch->GetShouldDeleteOnDeselect())
 			{
 				FPoseWatchManagerTreeItemPtr PoseWatchItem = Mode->PoseWatchManager->CreateItemFor<FPoseWatchManagerPoseWatchTreeItem>(PoseWatch);
+
 				OutItems.Add(PoseWatchItem);
+
+				for (TObjectPtr<UPoseWatchElement> CurrentElement : PoseWatch->GetElements())
+				{
+					FPoseWatchManagerTreeItemPtr ElementItem = Mode->PoseWatchManager->CreateItemFor<FPoseWatchManagerElementTreeItem>(CurrentElement);
+					OutItems.Add(ElementItem);
+				}
 			}
 		}
 		else
@@ -68,4 +77,4 @@ void FPoseWatchManagerDefaultHierarchy::CreateItems(TArray<FPoseWatchManagerTree
 			AnimBlueprint->PoseWatches.RemoveAtSwap(Index);
 		}
 	}
-}
+}
\ No newline at end of file
diff --git a/Engine/Source/Editor/Persona/Private/PoseWatchManagerDefaultMode.cpp b/Engine/Source/Editor/Persona/Private/PoseWatchManagerDefaultMode.cpp
--- a/Engine/Source/Editor/Persona/Private/PoseWatchManagerDefaultMode.cpp
+++ b/Engine/Source/Editor/Persona/Private/PoseWatchManagerDefaultMode.cpp
@@ -6,6 +6,7 @@
 #include "SGraphNode.h"
 #include "SGraphPanel.h"
 #include "IAnimationBlueprintEditor.h"
+#include "PoseWatchManagerElementTreeItem.h"
 #include "PoseWatchManagerFolderTreeItem.h"
 #include "PoseWatchManagerPoseWatchTreeItem.h"
 #include "SPoseWatchManager.h"
@@ -50,83 +51,90 @@ bool FPoseWatchManagerDefaultMode::ParseDragDrop(FPoseWatchManagerDragDropPayloa
 
 FPoseWatchManagerDragValidationInfo FPoseWatchManagerDefaultMode::ValidateDrop(const IPoseWatchManagerTreeItem& DropTarget, const FPoseWatchManagerDragDropPayload& Payload) const
 {
-	check(Payload.DraggedItem.IsValid())
-	TSharedPtr<IPoseWatchManagerTreeItem> PayloadItem = Payload.DraggedItem.Pin();
-	
-	// Pose watches cannot be a parent
-	if (const FPoseWatchManagerPoseWatchTreeItem* PoseWatchDropTarget = DropTarget.CastTo<FPoseWatchManagerPoseWatchTreeItem>())
+	if (Payload.DraggedItem.IsValid())
 	{
-		return FPoseWatchManagerDragValidationInfo(EPoseWatchManagerDropCompatibility::Incompatible, LOCTEXT("PoseWatchBadParent", "A pose watch cannot be a parent"));
-	}
+		TSharedPtr<IPoseWatchManagerTreeItem> PayloadItem = Payload.DraggedItem.Pin();
 
-	// Drop target must either be a valid UPoseWatchFolder or nullptr denoting the root of the tree
-	UPoseWatchFolder* TargetFolder = nullptr;
-	FText TargetFolderLabel = LOCTEXT("Root", "root");
-	if (DropTarget.IsValid())
-	{
-		if (const FPoseWatchManagerFolderTreeItem* DropTargetFolder = DropTarget.CastTo<FPoseWatchManagerFolderTreeItem>())
+		// Pose watches cannot be a parent
+		if (const FPoseWatchManagerPoseWatchTreeItem* PoseWatchDropTarget = DropTarget.CastTo<FPoseWatchManagerPoseWatchTreeItem>())
 		{
-			TargetFolder = DropTargetFolder->PoseWatchFolder.Get();
-			if (TargetFolder)
-			{
-				TargetFolderLabel = TargetFolder->GetLabel();
-			}
+			return FPoseWatchManagerDragValidationInfo(EPoseWatchManagerDropCompatibility::Incompatible, LOCTEXT("PoseWatchBadParent", "A pose watch cannot be a parent"));
 		}
-	}
 
-	if (FPoseWatchManagerPoseWatchTreeItem* PayloadPoseWatchItem = PayloadItem->CastTo<FPoseWatchManagerPoseWatchTreeItem>())
-	{
-		// Dropping a pose watch into a folder
-		UPoseWatch* PayloadPoseWatch = PayloadPoseWatchItem->PoseWatch.Get();
-		check(PayloadPoseWatch);
-
-		if (PayloadPoseWatch->IsIn(TargetFolder))
+		if (const FPoseWatchManagerElementTreeItem* ElementDropTarget = DropTarget.CastTo<FPoseWatchManagerElementTreeItem>())
 		{
-			FText ValidationText = FText::Format(LOCTEXT("PoseWatchAlreadyInDirectory", "This pose watch is already inside {0}"), TargetFolderLabel);
-			return FPoseWatchManagerDragValidationInfo(EPoseWatchManagerDropCompatibility::Incompatible, ValidationText);
+			return FPoseWatchManagerDragValidationInfo(EPoseWatchManagerDropCompatibility::Incompatible, LOCTEXT("ElementBadParent", "A pose watch element cannot be a parent"));
 		}
 
-		if (!PayloadPoseWatch->IsPoseWatchLabelUniqueInFolder(PayloadPoseWatch->GetLabel(), TargetFolder))
+		// Drop target must either be a valid UPoseWatchFolder or nullptr denoting the root of the tree
+		UPoseWatchFolder* TargetFolder = nullptr;
+		FText TargetFolderLabel = LOCTEXT("Root", "root");
+		if (DropTarget.IsValid())
 		{
-			FText ValidationText = FText::Format(LOCTEXT("PoseWatchAlreadyExistsInFolder", "A pose watch with that name already exists inside {0}"), TargetFolderLabel);
-			return FPoseWatchManagerDragValidationInfo(EPoseWatchManagerDropCompatibility::Incompatible, ValidationText);
+			if (const FPoseWatchManagerFolderTreeItem* DropTargetFolder = DropTarget.CastTo<FPoseWatchManagerFolderTreeItem>())
+			{
+				TargetFolder = DropTargetFolder->PoseWatchFolder.Get();
+				if (TargetFolder)
+				{
+					TargetFolderLabel = TargetFolder->GetLabel();
+				}
+			}
 		}
 
-		FText ValidationText = FText::Format(LOCTEXT("MovePoseWatchIntoFolder", "Move {0} into {1}"), PayloadPoseWatch->GetLabel(), TargetFolderLabel);
-		return FPoseWatchManagerDragValidationInfo(EPoseWatchManagerDropCompatibility::Compatible, ValidationText);
-	}
-	else if (FPoseWatchManagerFolderTreeItem* PayloadFolderItem = PayloadItem->CastTo<FPoseWatchManagerFolderTreeItem>())
-	{
-		// Dropping a folder into a folder
-		UPoseWatchFolder* PayloadFolder = PayloadFolderItem->PoseWatchFolder.Get();
-		check(PayloadFolder);
-
-		if (PayloadFolder == TargetFolder)
+		if (FPoseWatchManagerPoseWatchTreeItem* PayloadPoseWatchItem = PayloadItem->CastTo<FPoseWatchManagerPoseWatchTreeItem>())
 		{
-			return FPoseWatchManagerDragValidationInfo(EPoseWatchManagerDropCompatibility::Incompatible, LOCTEXT("FolderInItself", "A folder cannot contain itself"));
-		}
+			// Dropping a pose watch into a folder
+			UPoseWatch* PayloadPoseWatch = PayloadPoseWatchItem->PoseWatch.Get();
+			check(PayloadPoseWatch);
 
-		if (PayloadFolder->IsIn(TargetFolder))
-		{
-			FText ValidationText = FText::Format(LOCTEXT("FolderAlreadyInDirectory", "This folder is already inside {0}"), TargetFolderLabel);
-			return FPoseWatchManagerDragValidationInfo(EPoseWatchManagerDropCompatibility::Incompatible, ValidationText);
-		}
+			if (PayloadPoseWatch->IsIn(TargetFolder))
+			{
+				FText ValidationText = FText::Format(LOCTEXT("PoseWatchAlreadyInDirectory", "This pose watch is already inside {0}"), TargetFolderLabel);
+				return FPoseWatchManagerDragValidationInfo(EPoseWatchManagerDropCompatibility::Incompatible, ValidationText);
+			}
 
-		if (TargetFolder && TargetFolder->IsDescendantOf(PayloadFolder))
-		{
-			return FPoseWatchManagerDragValidationInfo(EPoseWatchManagerDropCompatibility::Incompatible, LOCTEXT("ParentFolderNotChildren", "Parent folders cannot be children"));
-		}
+			if (!PayloadPoseWatch->IsLabelUniqueInParent(PayloadPoseWatch->GetLabel(), TargetFolder))
+			{
+				FText ValidationText = FText::Format(LOCTEXT("PoseWatchAlreadyExistsInFolder", "A pose watch with that name already exists inside {0}"), TargetFolderLabel);
+				return FPoseWatchManagerDragValidationInfo(EPoseWatchManagerDropCompatibility::Incompatible, ValidationText);
+			}
 
-		if (!PayloadFolder->IsFolderLabelUniqueInFolder(PayloadFolder->GetLabel(), TargetFolder))
-		{
-			FText ValidationText = FText::Format(LOCTEXT("FolderNameConflictInFolder", "A folder with that name already exists inside {0}"), TargetFolderLabel);
-			return FPoseWatchManagerDragValidationInfo(EPoseWatchManagerDropCompatibility::Incompatible, ValidationText);
+			FText ValidationText = FText::Format(LOCTEXT("MovePoseWatchIntoFolder", "Move {0} into {1}"), PayloadPoseWatch->GetLabel(), TargetFolderLabel);
+			return FPoseWatchManagerDragValidationInfo(EPoseWatchManagerDropCompatibility::Compatible, ValidationText);
 		}
+		else if (FPoseWatchManagerFolderTreeItem* PayloadFolderItem = PayloadItem->CastTo<FPoseWatchManagerFolderTreeItem>())
+		{
+			// Dropping a folder into a folder
+			UPoseWatchFolder* PayloadFolder = PayloadFolderItem->PoseWatchFolder.Get();
+			check(PayloadFolder);
+
+			if (PayloadFolder == TargetFolder)
+			{
+				return FPoseWatchManagerDragValidationInfo(EPoseWatchManagerDropCompatibility::Incompatible, LOCTEXT("FolderInItself", "A folder cannot contain itself"));
+			}
+
+			if (PayloadFolder->IsIn(TargetFolder))
+			{
+				FText ValidationText = FText::Format(LOCTEXT("FolderAlreadyInDirectory", "This folder is already inside {0}"), TargetFolderLabel);
+				return FPoseWatchManagerDragValidationInfo(EPoseWatchManagerDropCompatibility::Incompatible, ValidationText);
+			}
 
-		FText ValidationText = FText::Format(LOCTEXT("MoveFolderIntoFolder", "Move {0} into {1}"), PayloadFolder->GetLabel(), TargetFolderLabel);
-		return FPoseWatchManagerDragValidationInfo(EPoseWatchManagerDropCompatibility::Compatible, ValidationText);
+			if (TargetFolder && TargetFolder->IsDescendantOf(PayloadFolder))
+			{
+				return FPoseWatchManagerDragValidationInfo(EPoseWatchManagerDropCompatibility::Incompatible, LOCTEXT("ParentFolderNotChildren", "Parent folders cannot be children"));
+			}
+
+			if (!PayloadFolder->IsLabelUniqueInParent(PayloadFolder->GetLabel(), TargetFolder))
+			{
+				FText ValidationText = FText::Format(LOCTEXT("FolderNameConflictInFolder", "A folder with that name already exists inside {0}"), TargetFolderLabel);
+				return FPoseWatchManagerDragValidationInfo(EPoseWatchManagerDropCompatibility::Incompatible, ValidationText);
+			}
+
+			FText ValidationText = FText::Format(LOCTEXT("MoveFolderIntoFolder", "Move {0} into {1}"), PayloadFolder->GetLabel(), TargetFolderLabel);
+			return FPoseWatchManagerDragValidationInfo(EPoseWatchManagerDropCompatibility::Compatible, ValidationText);
+		}
 	}
-	
+
 	return FPoseWatchManagerDragValidationInfo(EPoseWatchManagerDropCompatibility::Incompatible, FText());
 }
 
diff --git a/Engine/Source/Editor/Persona/Private/PoseWatchManagerElementTreeItem.cpp b/Engine/Source/Editor/Persona/Private/PoseWatchManagerElementTreeItem.cpp
new file mode 100644
--- /dev/null
+++ b/Engine/Source/Editor/Persona/Private/PoseWatchManagerElementTreeItem.cpp
@@ -0,0 +1,223 @@
+// Copyright Epic Games, Inc. All Rights Reserved.
+
+#include "PoseWatchManagerElementTreeItem.h"
+#include "Textures/SlateIcon.h"
+#include "Widgets/Text/SInlineEditableTextBlock.h"
+#include "Framework/Commands/UIAction.h"
+#include "ToolMenus.h"
+#include "Styling/AppStyle.h"
+#include "PoseWatchManagerDragDrop.h"
+#include "SPoseWatchManager.h"
+#include "Framework/MultiBox/MultiBoxBuilder.h"
+#include "Engine/PoseWatch.h"
+
+#define LOCTEXT_NAMESPACE "PoseWatchManagerElementTreeItem"
+
+struct PERSONA_API SPoseWatchManagerElementTreeLabel : FPoseWatchManagerCommonLabelData, public SCompoundWidget
+{
+	SLATE_BEGIN_ARGS(SPoseWatchManagerElementTreeLabel) {}
+	SLATE_END_ARGS()
+
+	void Construct(const FArguments& InArgs, FPoseWatchManagerElementTreeItem& NodeWatchTreeItem, IPoseWatchManager& PoseWatchManager, const STableRow<FPoseWatchManagerTreeItemPtr>& InRow)
+	{
+		WeakPoseWatchManager = StaticCastSharedRef<IPoseWatchManager>(PoseWatchManager.AsShared());
+
+		TreeItemPtr = StaticCastSharedRef<FPoseWatchManagerElementTreeItem>(NodeWatchTreeItem.AsShared());
+		WeakPoseWatchElementPtr = NodeWatchTreeItem.PoseWatchElement;
+
+		HighlightText = PoseWatchManager.GetFilterHighlightText();
+
+		TSharedPtr<SInlineEditableTextBlock> InlineTextBlock;
+
+		TSharedRef<SHorizontalBox> MainContent = SNew(SHorizontalBox)
+		+ SHorizontalBox::Slot()
+		.VAlign(VAlign_Center)
+		[
+			SAssignNew(InlineTextBlock, SInlineEditableTextBlock)
+			.Text(this, &SPoseWatchManagerElementTreeLabel::GetDisplayText)
+			.ToolTipText(this, &SPoseWatchManagerElementTreeLabel::GetTooltipText)
+			.HighlightText(HighlightText)
+			.ColorAndOpacity(this, &SPoseWatchManagerElementTreeLabel::GetForegroundColor)
+			.OnTextCommitted(this, &SPoseWatchManagerElementTreeLabel::OnLabelCommitted)
+			.OnVerifyTextChanged(this, &SPoseWatchManagerElementTreeLabel::OnVerifyItemLabelChanged)
+			.OnEnterEditingMode(this, &SPoseWatchManagerElementTreeLabel::OnEnterEditingMode)
+			.OnExitEditingMode(this, &SPoseWatchManagerElementTreeLabel::OnExitEditingMode)
+			.IsSelected(FIsSelected::CreateSP(&InRow, &STableRow<FPoseWatchManagerTreeItemPtr>::IsSelectedExclusively))
+		];
+
+		if (WeakPoseWatchManager.Pin())
+		{
+			NodeWatchTreeItem.RenameRequestEvent.BindSP(InlineTextBlock.Get(), &SInlineEditableTextBlock::EnterEditingMode);
+		}
+
+		ChildSlot
+		[
+			SNew(SHorizontalBox)
+
+			+ SHorizontalBox::Slot()
+			.AutoWidth()
+			.VAlign(VAlign_Center)
+			.Padding(FPoseWatchManagerDefaultTreeItemMetrics::IconPadding())
+			[
+				SNew(SBox)
+				.WidthOverride(FPoseWatchManagerDefaultTreeItemMetrics::IconSize())
+				.HeightOverride(FPoseWatchManagerDefaultTreeItemMetrics::IconSize())
+				[
+					SNew(SImage)
+					.Image(this, &SPoseWatchManagerElementTreeLabel::GetIcon)
+					.ToolTipText(this, &SPoseWatchManagerElementTreeLabel::GetIconTooltip)
+					.ColorAndOpacity(FSlateColor::UseForeground())
+				]
+			]
+
+			+ SHorizontalBox::Slot()
+			.FillWidth(1.0f)
+			.VAlign(VAlign_Center)
+			.Padding(0.0f, 0.0f)
+			[
+				MainContent
+			]
+		];
+	}
+
+private:
+	TWeakPtr<FPoseWatchManagerElementTreeItem> TreeItemPtr;
+	TWeakObjectPtr<UPoseWatchElement> WeakPoseWatchElementPtr;
+	TAttribute<FText> HighlightText;
+	bool bInEditingMode = false;
+
+	FText GetDisplayText() const
+	{
+		return WeakPoseWatchElementPtr.Get()->GetLabel();
+	}
+
+	FText GetTooltipText() const
+	{
+		if (!TreeItemPtr.Pin()->IsEnabled())
+		{
+			return LOCTEXT("NodeWatchDisabled", "This node watch element is disabled because it is not connected to the 'Output Pose' node");
+		}
+		return GetDisplayText();
+	}
+
+	const FSlateBrush* GetIcon() const
+	{
+		FName IconBrushName = TEXT("ClassIcon.Default");
+		if (const UPoseWatchElement* PoseWatchElement = WeakPoseWatchElementPtr.Get())
+		{
+			IconBrushName = PoseWatchElement->GetIconName();
+		}
+
+		return FAppStyle::Get().GetBrush(IconBrushName);
+	}
+
+	const FSlateBrush* GetIconOverlay() const
+	{
+		return nullptr;
+	}
+
+	FText GetIconTooltip() const
+	{
+		return LOCTEXT("PoseWatchElement", "Component");
+	}
+
+	FSlateColor GetForegroundColor() const
+	{
+		if (const TOptional<FLinearColor> BaseColor = FPoseWatchManagerCommonLabelData::GetForegroundColor(*TreeItemPtr.Pin()))
+		{
+			return BaseColor.GetValue();
+		}
+		return FSlateColor::UseForeground();
+	}
+
+	bool OnVerifyItemLabelChanged(const FText& InLabel, FText& OutErrorMessage)
+	{
+		return WeakPoseWatchElementPtr->ValidateLabelRename(InLabel, OutErrorMessage);
+	}
+
+	void OnLabelCommitted(const FText& InLabel, ETextCommit::Type InCommitInfo)
+	{
+		check(WeakPoseWatchElementPtr->SetLabel(InLabel));
+		WeakPoseWatchManager.Pin()->FullRefresh();
+		WeakPoseWatchManager.Pin()->SetKeyboardFocus();
+	}
+
+	void OnEnterEditingMode()
+	{
+		bInEditingMode = true;
+	}
+
+	void OnExitEditingMode()
+	{
+		bInEditingMode = false;
+	}
+};
+
+const EPoseWatchTreeItemType FPoseWatchManagerElementTreeItem::Type(EPoseWatchTreeItemType::Element);
+
+FPoseWatchManagerElementTreeItem::FPoseWatchManagerElementTreeItem(TWeakObjectPtr<UPoseWatchElement> InPoseWatchElement)
+	: IPoseWatchManagerTreeItem(Type)
+	, ID(InPoseWatchElement.Get())
+	, PoseWatchElement(InPoseWatchElement)
+{
+	check(InPoseWatchElement.Get());
+}
+
+FObjectKey FPoseWatchManagerElementTreeItem::GetID() const
+{
+	return ID;
+}
+
+FString FPoseWatchManagerElementTreeItem::GetDisplayString() const
+{
+	return PoseWatchElement->GetLabel().ToString();
+}
+
+bool FPoseWatchManagerElementTreeItem::IsAssignedFolder() const
+{
+	return false;
+}
+
+TSharedRef<SWidget> FPoseWatchManagerElementTreeItem::GenerateLabelWidget(IPoseWatchManager& PoseWatchManager, const STableRow<FPoseWatchManagerTreeItemPtr>& InRow)
+{
+	return SNew(SPoseWatchManagerElementTreeLabel, *this, PoseWatchManager, InRow);
+}
+
+bool FPoseWatchManagerElementTreeItem::GetVisibility() const
+{
+	return PoseWatchElement->GetIsVisible();
+}
+
+void FPoseWatchManagerElementTreeItem::SetIsVisible(const bool bVisible)
+{
+	PoseWatchElement->SetIsVisible(bVisible);
+}
+
+TSharedPtr<SWidget> FPoseWatchManagerElementTreeItem::CreateContextMenu()
+{
+	return SNullWidget::NullWidget;
+}
+
+void FPoseWatchManagerElementTreeItem::OnRemoved() {}
+
+bool FPoseWatchManagerElementTreeItem::IsEnabled() const
+{
+	return PoseWatchElement->GetIsEnabled();
+}
+
+FColor FPoseWatchManagerElementTreeItem::GetColor() const
+{
+	return PoseWatchElement->GetColor();
+}
+
+void FPoseWatchManagerElementTreeItem::SetColor(const FColor& InColor)
+{
+	PoseWatchElement->SetColor(InColor);
+}
+
+bool FPoseWatchManagerElementTreeItem::ShouldDisplayColorPicker() const
+{
+	return PoseWatchElement->HasColor();
+}
+
+#undef LOCTEXT_NAMESPACE
diff --git a/Engine/Source/Editor/Persona/Private/PoseWatchManagerElementTreeItem.h b/Engine/Source/Editor/Persona/Private/PoseWatchManagerElementTreeItem.h
new file mode 100644
--- /dev/null
+++ b/Engine/Source/Editor/Persona/Private/PoseWatchManagerElementTreeItem.h
@@ -0,0 +1,41 @@
+// Copyright Epic Games, Inc. All Rights Reserved.
+
+#pragma once
+
+#include "CoreMinimal.h"
+#include "Widgets/SWidget.h"
+#include "IPoseWatchManagerTreeItem.h"
+
+class UPoseWatchElement;
+
+struct FPoseWatchManagerElementTreeItem : IPoseWatchManagerTreeItem
+{
+public:
+	FPoseWatchManagerElementTreeItem(TWeakObjectPtr<UPoseWatchElement> InPoseWatchElement);
+
+	/** Constant identifier for this tree item */
+	const FObjectKey ID;
+
+	/** The node watch component this item is describing */
+	TWeakObjectPtr<UPoseWatchElement> PoseWatchElement;
+
+	/** Static type identifier for this tree item class */
+	static const EPoseWatchTreeItemType Type;
+
+	bool IsValid() const override { return PoseWatchElement.IsValid(); }
+
+	/* Begin IPoseWatchManagerTreeItem Implementation */
+	virtual FObjectKey GetID() const override;
+	virtual FString GetDisplayString() const override;
+	virtual TSharedRef<SWidget> GenerateLabelWidget(IPoseWatchManager& Outliner, const STableRow<FPoseWatchManagerTreeItemPtr>& InRow) override;
+	virtual bool GetVisibility() const override;
+	virtual bool IsAssignedFolder() const;
+	virtual void SetIsVisible(const bool bVisible) override;
+	virtual TSharedPtr<SWidget> CreateContextMenu() override;
+	virtual void OnRemoved() override;
+	virtual bool IsEnabled() const override;
+	virtual FColor GetColor() const override;
+	virtual void SetColor(const FColor& InColor) override;
+	virtual bool ShouldDisplayColorPicker() const override;
+	/* End IPoseWatchManagerTreeItem Implementation */
+};
diff --git a/Engine/Source/Editor/Persona/Private/PoseWatchManagerFolderTreeItem.cpp b/Engine/Source/Editor/Persona/Private/PoseWatchManagerFolderTreeItem.cpp
--- a/Engine/Source/Editor/Persona/Private/PoseWatchManagerFolderTreeItem.cpp
+++ b/Engine/Source/Editor/Persona/Private/PoseWatchManagerFolderTreeItem.cpp
@@ -81,10 +81,7 @@ struct SPoseWatchManagerFolderTreeItemLabel : FPoseWatchManagerCommonLabelData,
 			.OnTextCommitted(this, &SPoseWatchManagerFolderTreeItemLabel::OnLabelCommitted)
 			.OnVerifyTextChanged(this, &SPoseWatchManagerFolderTreeItemLabel::OnVerifyItemLabelChanged)
 			.IsSelected(FIsSelected::CreateSP(&InRow, &STableRow<FPoseWatchManagerTreeItemPtr>::IsSelectedExclusively))
-			.IsReadOnly_Lambda([Item = FolderItem.AsShared(), this]()
-		{
-			return !CanExecuteRenameRequest(Item.Get());
-		});
+			.IsReadOnly_Lambda([this]() { return !CanExecuteRenameRequest(*TreeItemPtr.Pin()); });
 
 		FolderItem.RenameRequestEvent.BindSP(InlineTextBlock.Get(), &SInlineEditableTextBlock::EnterEditingMode);
 
@@ -185,9 +182,17 @@ bool FPoseWatchManagerFolderTreeItem::HasChildren() const
 	return PoseWatchFolder.IsValid() ? PoseWatchFolder.Get()->HasChildren() : false;
 }
 
+void FPoseWatchManagerFolderTreeItem::SetIsExpanded(const bool bIsExpanded)
+{
+	if (PoseWatchFolder.IsValid())
+	{
+		PoseWatchFolder.Get()->SetIsExpanded(bIsExpanded);
+	}
+}
+
 bool FPoseWatchManagerFolderTreeItem::IsExpanded() const
 {
-	return PoseWatchFolder.IsValid() ? PoseWatchFolder.Get()->GetIsExpanded() : false;
+	return PoseWatchFolder.IsValid() && PoseWatchFolder.Get()->GetIsExpanded();
 }
 
 #undef LOCTEXT_NAMESPACE
diff --git a/Engine/Source/Editor/Persona/Private/PoseWatchManagerFolderTreeItem.h b/Engine/Source/Editor/Persona/Private/PoseWatchManagerFolderTreeItem.h
--- a/Engine/Source/Editor/Persona/Private/PoseWatchManagerFolderTreeItem.h
+++ b/Engine/Source/Editor/Persona/Private/PoseWatchManagerFolderTreeItem.h
@@ -35,6 +35,7 @@ public:
 	virtual void SetIsVisible(const bool bVisible) override;
 	virtual bool HasChildren() const override;
 	virtual void OnRemoved() override;
+	virtual void SetIsExpanded(const bool bIsExpanded) override;
 	virtual bool IsExpanded() const override;
 	/* End IPoseWatchManagerTreeItem Implementation */
 };
diff --git a/Engine/Source/Editor/Persona/Private/PoseWatchManagerPoseWatchTreeItem.cpp b/Engine/Source/Editor/Persona/Private/PoseWatchManagerPoseWatchTreeItem.cpp
--- a/Engine/Source/Editor/Persona/Private/PoseWatchManagerPoseWatchTreeItem.cpp
+++ b/Engine/Source/Editor/Persona/Private/PoseWatchManagerPoseWatchTreeItem.cpp
@@ -27,35 +27,14 @@ struct PERSONA_API SPoseWatchManagerPoseWatchTreeLabel : FPoseWatchManagerCommon
 
 		HighlightText = PoseWatchManager.GetFilterHighlightText();
 
-		TSharedPtr<SInlineEditableTextBlock> InlineTextBlock;
-
-		auto MainContent = SNew(SHorizontalBox)
-
-			+ SHorizontalBox::Slot()
-			.VAlign(VAlign_Center)
-			[
-				SAssignNew(InlineTextBlock, SInlineEditableTextBlock)
-				.Text(this, &SPoseWatchManagerPoseWatchTreeLabel::GetDisplayText)
-				.ToolTipText(this, &SPoseWatchManagerPoseWatchTreeLabel::GetTooltipText)
-				.HighlightText(HighlightText)
-				.ColorAndOpacity(this, &SPoseWatchManagerPoseWatchTreeLabel::GetForegroundColor)
-				.OnTextCommitted(this, &SPoseWatchManagerPoseWatchTreeLabel::OnLabelCommitted)
-				.OnVerifyTextChanged(this, &SPoseWatchManagerPoseWatchTreeLabel::OnVerifyItemLabelChanged)
-				.OnEnterEditingMode(this, &SPoseWatchManagerPoseWatchTreeLabel::OnEnterEditingMode)
-				.OnExitEditingMode(this, &SPoseWatchManagerPoseWatchTreeLabel::OnExitEditingMode)
-				.IsSelected(FIsSelected::CreateSP(&InRow, &STableRow<FPoseWatchManagerTreeItemPtr>::IsSelectedExclusively))
-			]
-
-			+ SHorizontalBox::Slot()
-			.VAlign(VAlign_Center)
-			.AutoWidth()
-			.Padding(0.0f, 0.f, 3.0f, 0.0f)
-			[
-				SNew(STextBlock)
-				.Text(this, &SPoseWatchManagerPoseWatchTreeLabel::GetTypeText)
-				.Visibility(this, &SPoseWatchManagerPoseWatchTreeLabel::GetTypeTextVisibility)
-				.HighlightText(HighlightText)
-			];
+		TSharedPtr<SInlineEditableTextBlock> InlineTextBlock = SNew(SInlineEditableTextBlock)
+			.Text(this, &SPoseWatchManagerPoseWatchTreeLabel::GetDisplayText)
+			.HighlightText(PoseWatchManager.GetFilterHighlightText())
+			.ColorAndOpacity(this, &SPoseWatchManagerPoseWatchTreeLabel::GetForegroundColor)
+			.OnTextCommitted(this, &SPoseWatchManagerPoseWatchTreeLabel::OnLabelCommitted)
+			.OnVerifyTextChanged(this, &SPoseWatchManagerPoseWatchTreeLabel::OnVerifyItemLabelChanged)
+			.IsSelected(FIsSelected::CreateSP(&InRow, &STableRow<FPoseWatchManagerTreeItemPtr>::IsSelectedExclusively))
+			.IsReadOnly_Lambda([this]() { return !CanExecuteRenameRequest(*TreeItemPtr.Pin()); });
 
 		if (WeakPoseWatchManager.Pin())
 		{
@@ -66,28 +45,12 @@ struct PERSONA_API SPoseWatchManagerPoseWatchTreeLabel : FPoseWatchManagerCommon
 		[
 			SNew(SHorizontalBox)
 
-			+ SHorizontalBox::Slot()
-			.AutoWidth()
-			.VAlign(VAlign_Center)
-			.Padding(FPoseWatchManagerDefaultTreeItemMetrics::IconPadding())
-			[
-				SNew(SBox)
-				.WidthOverride(FPoseWatchManagerDefaultTreeItemMetrics::IconSize())
-				.HeightOverride(FPoseWatchManagerDefaultTreeItemMetrics::IconSize())
-				[
-					SNew(SImage)
-					.Image(this, &SPoseWatchManagerPoseWatchTreeLabel::GetIcon)
-					.ToolTipText(this, &SPoseWatchManagerPoseWatchTreeLabel::GetIconTooltip)
-					.ColorAndOpacity(FSlateColor::UseForeground())
-				]
-			]
-
 			+ SHorizontalBox::Slot()
 			.FillWidth(1.0f)
 			.VAlign(VAlign_Center)
-			.Padding(0.0f, 0.0f)
+			.Padding(0.0f, 2.0f)
 			[
-				MainContent
+				InlineTextBlock.ToSharedRef()
 			]
 		];
 	}
@@ -100,7 +63,11 @@ private:
 
 	FText GetDisplayText() const
 	{
-		return WeakPoseWatchPtr.Get()->GetLabel();
+		if (const UPoseWatch* PoseWatch = WeakPoseWatchPtr.Get())
+		{
+			return PoseWatch->GetLabel();
+		}
+		return FText();
 	}
 
 	FText GetTooltipText() const
@@ -245,4 +212,24 @@ bool FPoseWatchManagerPoseWatchTreeItem::IsEnabled() const
 	return PoseWatch->GetIsEnabled();
 }
 
+void FPoseWatchManagerPoseWatchTreeItem::SetIsExpanded(const bool bIsExpanded)
+{
+	PoseWatch->SetIsExpanded(bIsExpanded);
+}
+
+bool FPoseWatchManagerPoseWatchTreeItem::IsExpanded() const
+{
+	return PoseWatch->GetIsExpanded();
+}
+
+FColor FPoseWatchManagerPoseWatchTreeItem::GetColor() const
+{
+	return PoseWatch->GetColor();
+}
+
+void FPoseWatchManagerPoseWatchTreeItem::SetColor(const FColor& InColor)
+{
+	PoseWatch->SetColor(InColor);
+}
+
 #undef LOCTEXT_NAMESPACE
diff --git a/Engine/Source/Editor/Persona/Private/PoseWatchManagerPoseWatchTreeItem.h b/Engine/Source/Editor/Persona/Private/PoseWatchManagerPoseWatchTreeItem.h
--- a/Engine/Source/Editor/Persona/Private/PoseWatchManagerPoseWatchTreeItem.h
+++ b/Engine/Source/Editor/Persona/Private/PoseWatchManagerPoseWatchTreeItem.h
@@ -34,5 +34,9 @@ public:
 	virtual TSharedPtr<SWidget> CreateContextMenu() override;
 	virtual void OnRemoved() override;
 	virtual bool IsEnabled() const override;
+	virtual void SetIsExpanded(const bool bIsExpanded) override;
+	virtual bool IsExpanded() const override;
+	virtual FColor GetColor() const override;
+	virtual void SetColor(const FColor& InColor) override;
 	/* End IPoseWatchManagerTreeItem Implementation */
 };
diff --git a/Engine/Source/Editor/Persona/Private/PoseWatchManagerStandaloneTypes.h b/Engine/Source/Editor/Persona/Private/PoseWatchManagerStandaloneTypes.h
--- a/Engine/Source/Editor/Persona/Private/PoseWatchManagerStandaloneTypes.h
+++ b/Engine/Source/Editor/Persona/Private/PoseWatchManagerStandaloneTypes.h
@@ -18,7 +18,8 @@ typedef TMap<FObjectKey, FPoseWatchManagerTreeItemPtr> FPoseWatchManagerTreeItem
 enum EPoseWatchTreeItemType
 {
 	PoseWatch,
-	Folder
+	Folder,
+	Element
 };
 
 struct PERSONA_API FPoseWatchManagerCommonLabelData
diff --git a/Engine/Source/Editor/Persona/Private/SPoseWatchManager.cpp b/Engine/Source/Editor/Persona/Private/SPoseWatchManager.cpp
--- a/Engine/Source/Editor/Persona/Private/SPoseWatchManager.cpp
+++ b/Engine/Source/Editor/Persona/Private/SPoseWatchManager.cpp
@@ -37,9 +37,8 @@
 #include "AnimationEditorUtils.h"
 #include "PoseWatchManagerDefaultHierarchy.h"
 #include "PoseWatchManagerDefaultMode.h"
-
+#include "PoseWatchManagerElementTreeItem.h"
 #include "PoseWatchManagerPoseWatchTreeItem.h"
-
 #include "PoseWatchManagerColumnVisibility.h"
 #include "PoseWatchManagerColumnColor.h"
 #include "PoseWatchManagerColumnLabel.h"
@@ -391,7 +390,12 @@ void SPoseWatchManager::RepopulateEntireTree()
 		if (Item && Item->IsValid())
 		{
 			AddPendingItem(Item);
-			if (FPoseWatchManagerFolderTreeItem* FolderItem = Item->CastTo<FPoseWatchManagerFolderTreeItem>())
+
+			if (FPoseWatchManagerPoseWatchTreeItem* PoseWatchItem = Item->CastTo<FPoseWatchManagerPoseWatchTreeItem>())
+			{
+				PoseWatchManagerTreeView->SetItemExpansion(Item, PoseWatchItem->PoseWatch->GetIsExpanded());
+			}
+			else if (FPoseWatchManagerFolderTreeItem* FolderItem = Item->CastTo<FPoseWatchManagerFolderTreeItem>())
 			{
 				PoseWatchManagerTreeView->SetItemExpansion(Item, FolderItem->PoseWatchFolder->GetIsExpanded());
 			}
@@ -420,7 +424,7 @@ FPoseWatchManagerTreeItemPtr SPoseWatchManager::EnsureParentForItem(FPoseWatchMa
 	{
 		return Parent;
 	}
-		
+
 	// Try to find the parent in the pending items
 	Parent = Mode->Hierarchy->FindParent(*Item, PendingTreeItemMap);
 	if (Parent.IsValid())
@@ -726,9 +730,8 @@ void SPoseWatchManager::OnManagerTreeItemScrolledIntoView(FPoseWatchManagerTreeI
 
 void SPoseWatchManager::OnItemExpansionChanged(FPoseWatchManagerTreeItemPtr TreeItem, bool bIsExpanded) const
 {
-	check(TreeItem.Get()->IsA<FPoseWatchManagerFolderTreeItem>());
-	FPoseWatchManagerFolderTreeItem* FolderTreeItem = TreeItem.Get()->CastTo<FPoseWatchManagerFolderTreeItem>();
-	FolderTreeItem->PoseWatchFolder->SetIsExpanded(bIsExpanded);
+	check(TreeItem.Get()->IsA<FPoseWatchManagerFolderTreeItem>() || TreeItem.Get()->IsA<FPoseWatchManagerPoseWatchTreeItem>());
+	TreeItem->SetIsExpanded(bIsExpanded);
 }
 
 void SPoseWatchManager::OnHierarchyChangedEvent(FPoseWatchManagerHierarchyChangedData Event)
@@ -962,9 +965,19 @@ uint32 SPoseWatchManager::GetTypeSortPriority(const IPoseWatchManagerTreeItem& I
 
 void SPoseWatchManager::OnManagerTreeDoubleClick(FPoseWatchManagerTreeItemPtr TreeItem)
 {
+	UPoseWatch* PoseWatch = nullptr;
+
 	if (FPoseWatchManagerPoseWatchTreeItem* PoseWatchTreeItem = TreeItem->CastTo<FPoseWatchManagerPoseWatchTreeItem>())
 	{
-		UPoseWatch* PoseWatch = PoseWatchTreeItem->PoseWatch.Get();
+		PoseWatch = PoseWatchTreeItem->PoseWatch.Get();
+	}
+	else if (FPoseWatchManagerElementTreeItem* ElementTreeItem = TreeItem->CastTo<FPoseWatchManagerElementTreeItem>())
+	{
+		PoseWatch = ElementTreeItem->PoseWatchElement->GetParent();
+	}
+
+	if (PoseWatch)
+	{
 		BlueprintEditor->JumpToHyperlink(PoseWatch->Node.Get(), false);
 	}
 }
diff --git a/Engine/Source/Editor/Persona/Public/IPhysicsAssetRenderInterface.h b/Engine/Source/Editor/Persona/Public/IPhysicsAssetRenderInterface.h
--- a/Engine/Source/Editor/Persona/Public/IPhysicsAssetRenderInterface.h
+++ b/Engine/Source/Editor/Persona/Public/IPhysicsAssetRenderInterface.h
@@ -13,6 +13,8 @@ class IPhysicsAssetRenderInterface : public IModuleInterface, public IModularFea
 public:
 
 	virtual void DebugDraw(class USkeletalMeshComponent* const SkeletalMeshComponent, class UPhysicsAsset* const PhysicsAsset, class FPrimitiveDrawInterface* PDI) = 0;
+	virtual void DebugDrawBodies(class USkeletalMeshComponent* const SkeletalMeshComponent, class UPhysicsAsset* const PhysicsAsset, FPrimitiveDrawInterface* PDI, const FColor& PrimitiveColorOverride) = 0;
+	virtual void DebugDrawConstraints(class USkeletalMeshComponent* const SkeletalMeshComponent, class UPhysicsAsset* const PhysicsAsset, FPrimitiveDrawInterface* PDI) = 0;
 
 	virtual void SaveConfig() = 0;
 
diff --git a/Engine/Source/Editor/Persona/Public/PhysicsAssetRenderUtils.h b/Engine/Source/Editor/Persona/Public/PhysicsAssetRenderUtils.h
--- a/Engine/Source/Editor/Persona/Public/PhysicsAssetRenderUtils.h
+++ b/Engine/Source/Editor/Persona/Public/PhysicsAssetRenderUtils.h
@@ -13,8 +13,9 @@
 enum class EPhysicsAssetEditorConstraintViewMode : uint8;
 enum class EPhysicsAssetEditorCollisionViewMode : uint8;
 
-class UPhysicsAsset;
 class HHitProxy;
+class UPhysicsAsset;
+class USkeletalMeshComponent;
 
 //////////////////////////////////////////////////////////////////////////
 // FPhysicsAssetRenderSettings
@@ -196,12 +197,14 @@ namespace PhysicsAssetRender
 class FPhysicsAssetRenderInterface : public IPhysicsAssetRenderInterface
 {
 public:
-	virtual void DebugDraw(class USkeletalMeshComponent* const SkeletalMeshComponent, class UPhysicsAsset* const PhysicsAsset, FPrimitiveDrawInterface* PDI) override;
+	virtual void DebugDraw(USkeletalMeshComponent* const SkeletalMeshComponent, UPhysicsAsset* const PhysicsAsset, FPrimitiveDrawInterface* PDI) override;
+	virtual void DebugDrawBodies(USkeletalMeshComponent* const SkeletalMeshComponent, UPhysicsAsset* const PhysicsAsset, FPrimitiveDrawInterface* PDI, const FColor& PrimitiveColorOverride) override;
+	virtual void DebugDrawConstraints(USkeletalMeshComponent* const SkeletalMeshComponent, UPhysicsAsset* const PhysicsAsset, FPrimitiveDrawInterface* PDI) override;
 
 	virtual void SaveConfig() override;
 
-	virtual void ToggleShowAllBodies(class UPhysicsAsset* const PhysicsAsset) override;
-	virtual void ToggleShowAllConstraints(class UPhysicsAsset* const PhysicsAsset) override;
-	virtual bool AreAnyBodiesHidden(class UPhysicsAsset* const PhysicsAsset) override;
-	virtual bool AreAnyConstraintsHidden(class UPhysicsAsset* const PhysicsAsset) override;
+	virtual void ToggleShowAllBodies(UPhysicsAsset* const PhysicsAsset) override;
+	virtual void ToggleShowAllConstraints(UPhysicsAsset* const PhysicsAsset) override;
+	virtual bool AreAnyBodiesHidden(UPhysicsAsset* const PhysicsAsset) override;
+	virtual bool AreAnyConstraintsHidden(UPhysicsAsset* const PhysicsAsset) override;
 };
\ No newline at end of file
diff --git a/Engine/Source/Editor/UnrealEd/Private/AnimationEditorUtils.cpp b/Engine/Source/Editor/UnrealEd/Private/AnimationEditorUtils.cpp
--- a/Engine/Source/Editor/UnrealEd/Private/AnimationEditorUtils.cpp
+++ b/Engine/Source/Editor/UnrealEd/Private/AnimationEditorUtils.cpp
@@ -876,8 +876,16 @@ namespace AnimationEditorUtils
 					if (UAnimBlueprintGeneratedClass* AnimBPGenClass = Cast<UAnimBlueprintGeneratedClass>(*AnimBlueprint->GeneratedClass))
 					{
 						// Find the insertion point from the debugging data
-						int32 LinkID = AnimBPGenClass->GetLinkIDForNode<FAnimNode_Base>(TargetNode);
-						AnimBPGenClass->GetAnimBlueprintDebugData().AddPoseWatch(LinkID, PoseWatch);
+						const int32 LinkID = AnimBPGenClass->GetLinkIDForNode<FAnimNode_Base>(TargetNode);
+
+						for (const TObjectPtr<UPoseWatchElement>& PoseWatchElement : PoseWatch->GetElements())
+						{
+							if (UPoseWatchPoseElement* PoseWatchPoseElement = Cast<UPoseWatchPoseElement>(PoseWatchElement.Get()))
+							{
+								AnimBPGenClass->GetAnimBlueprintDebugData().AddPoseWatch(LinkID, PoseWatchPoseElement);
+							}
+						}
+
 						OnPoseWatchesChangedDelegate.Broadcast(AnimBlueprint, TargetNode);
 					}
 				}
@@ -931,6 +939,7 @@ namespace AnimationEditorUtils
 		UPoseWatch* NewPoseWatch = NewObject<UPoseWatch>(AnimBlueprint);
 		NewPoseWatch->Node = Node;
 		NewPoseWatch->SetUniqueDefaultLabel();
+		NewPoseWatch->AddElement<UPoseWatchPoseElement>(LOCTEXT("PoseWatchElementLabel_PoseWatch", "Pose Watch"), TEXT("ClassIcon.PoseAsset"));
 		AnimBlueprint->PoseWatches.Add(NewPoseWatch);
 		SetPoseWatch(NewPoseWatch, AnimBlueprint);
 		return NewPoseWatch;
diff --git a/Engine/Source/Runtime/Core/Public/UObject/FortniteMainBranchObjectVersion.h b/Engine/Source/Runtime/Core/Public/UObject/FortniteMainBranchObjectVersion.h
--- a/Engine/Source/Runtime/Core/Public/UObject/FortniteMainBranchObjectVersion.h
+++ b/Engine/Source/Runtime/Core/Public/UObject/FortniteMainBranchObjectVersion.h
@@ -229,6 +229,9 @@ struct CORE_API FFortniteMainBranchObjectVersion
 		// Added water HLOD material
 		WaterHLODSupportAdded,
 
+		// Moved parameters affecting Skeleton pose rendering from the PoseWatch class to the PoseWatchPoseElement class.
+		PoseWatchMigrateSkeletonDrawParametersToPoseElement,
+
 		// -----<new versions can be added above this line>-------------------------------------------------
 		VersionPlusOne,
 		LatestVersion = VersionPlusOne - 1
diff --git a/Engine/Source/Runtime/Engine/Classes/Animation/AnimBlueprint.h b/Engine/Source/Runtime/Engine/Classes/Animation/AnimBlueprint.h
--- a/Engine/Source/Runtime/Engine/Classes/Animation/AnimBlueprint.h
+++ b/Engine/Source/Runtime/Engine/Classes/Animation/AnimBlueprint.h
@@ -234,8 +234,8 @@ public:
 	UPROPERTY()
 	TArray<FAnimParentNodeAssetOverride> ParentAssetOverrides;
 
-	// Array of active pose watches (pose watch allows us to see the bone pose at a 
-	// particular point of the anim graph) 
+	// Array of active pose watches (pose watches allows us to see the bone pose at a 
+	// particular point of the anim graph and control debug draw for unselected anim nodes).
 	UPROPERTY()
 	TArray<TObjectPtr<UPoseWatchFolder>> PoseWatchFolders;
 	
diff --git a/Engine/Source/Runtime/Engine/Classes/Animation/AnimBlueprintGeneratedClass.h b/Engine/Source/Runtime/Engine/Classes/Animation/AnimBlueprintGeneratedClass.h
--- a/Engine/Source/Runtime/Engine/Classes/Animation/AnimBlueprintGeneratedClass.h
+++ b/Engine/Source/Runtime/Engine/Classes/Animation/AnimBlueprintGeneratedClass.h
@@ -145,7 +145,7 @@ struct FAnimBlueprintDebugData_NodeVisit
 
 struct FPoseWatchDebugData
 {
-	const UPoseWatch* PoseWatch;
+	const UPoseWatchPoseElement* PoseWatchPoseElement;
 	TSharedPtr<FCompactHeapPose> PoseInfo;
 };
 
@@ -317,9 +317,9 @@ public:
 	void RecordSequencePlayer(int32 InNodeID, float InPosition, float InLength, int32 InFrameCount);
 	void RecordBlendSpacePlayer(int32 InNodeID, const UBlendSpace* InBlendSpace, const FVector& InPosition, const FVector& InFilteredPosition);
 
-	void AddPoseWatch(int32 NodeID, UPoseWatch*);
+	void AddPoseWatch(int32 NodeID, UPoseWatchPoseElement* const InPoseWatchPoseElement);
 	void RemovePoseWatch(int32 NodeID);
-	void ForEachActiveVisiblePoseWatch(const TFunctionRef<void(FPoseWatchDebugData&)>& InFunction) const;
+	void ForEachActiveVisiblePoseWatchPoseElement(const TFunctionRef<void(FPoseWatchDebugData&)>& InFunction) const;
 
 	TArrayView<const FName> GetNodeAttributes(TWeakObjectPtr<UAnimGraphNode_Base> InAnimGraphNode) const;
 #endif
diff --git a/Engine/Source/Runtime/Engine/Classes/Engine/PoseWatch.h b/Engine/Source/Runtime/Engine/Classes/Engine/PoseWatch.h
--- a/Engine/Source/Runtime/Engine/Classes/Engine/PoseWatch.h
+++ b/Engine/Source/Runtime/Engine/Classes/Engine/PoseWatch.h
@@ -15,10 +15,11 @@ class UBlendProfile;
 struct FAnimNodePoseWatch
 {
 	// Object (anim instance) that this pose came from
-	TWeakObjectPtr<const UObject>	Object;
-	TWeakObjectPtr<UPoseWatch>		PoseWatch;
-	TSharedPtr<FCompactHeapPose>	PoseInfo;
-	int32							NodeID;
+	TWeakObjectPtr<const UObject>			Object;
+	TWeakObjectPtr<UPoseWatch>				PoseWatch;
+	TWeakObjectPtr<UPoseWatchPoseElement>	PoseWatchPoseElement;
+	TSharedPtr<FCompactHeapPose>			PoseInfo;
+	int32									NodeID;
 };
 
 namespace PoseWatchUtil
@@ -29,16 +30,10 @@ namespace PoseWatchUtil
 	/** Gets all pose watches folders that are parented to Folder, if Folder is nullptr then gets orphans */
 	TSet<UPoseWatchFolder*> GetChildrenPoseWatchFoldersOf(const UPoseWatchFolder* Folder, const UAnimBlueprint* AnimBlueprint);
 
-	/** Returns a UPoseWatch/UPoseWatchFolder from PoseWatchCollection that is inside InFolder has has label Label, if one exists */
-	template <typename T>
-	T* FindInFolderInCollection(const FName& Label, UPoseWatchFolder* InFolder, const TArray<TObjectPtr<T>>& PoseWatchCollection);
-
-	/** Finds unique name for a pose watch or folder inside of InParent */
-	template <typename T>
-	FText FindUniqueNameInFolder(UPoseWatchFolder* InParent, const T* Item, const TArray<TObjectPtr<T>>& Collection);
-
 	/** Returns a new random color */
 	FColor ChoosePoseWatchColor();
+
+	template <typename TParent, typename TItem> FText FindUniqueNameInParent(TParent* InParent, const TItem* InItem);
 }
 
 UCLASS()
@@ -94,7 +89,7 @@ public:
 	bool ValidateLabelRename(const FText& InLabel, FText& OutErrorMessage);
 
 	/** Returns true if InLabel is a unique folder label among the children of InFolder,  excluding this */
-	bool IsFolderLabelUniqueInFolder(const FText& InLabel, UPoseWatchFolder* InFolder) const;
+	bool IsLabelUniqueInParent(const FText& InLabel, UPoseWatchFolder* InFolder) const;
 
 	/** Returns true if at least one UPoseWatch/UPoseWatchFolder has this as it's parent */
 	bool HasChildren() const;
@@ -108,13 +103,15 @@ public:
 	/** Returns the anim blueprint this pose watch folder is stored inside */
 	UAnimBlueprint* GetAnimBlueprint() const;
 
+	/** Set whether this should display its children in the pose watch manager window */
 	void SetIsExpanded(bool bInIsExpanded);
 
+	/** Returns true if this should display its children in the pose watch manager window */
 	bool GetIsExpanded() const;
 
 private:
 	/** Returns a unique name for a new folder placed within InParent */
-	FText FindUniqueNameInFolder(UPoseWatchFolder* InParent) const;
+	FText FindUniqueNameInParent(UPoseWatchFolder* InParent) const;
 
 	/** Returns true if there is at least one direct pose watch child */
 	bool HasPoseWatchChildren() const;
@@ -139,6 +136,137 @@ protected:
 #endif // WITH_EDITORONLY_DATA
 };
 
+UCLASS()
+class ENGINE_API UPoseWatchElement
+	: public UObject
+{
+	GENERATED_BODY()
+
+public:
+#if WITH_EDITOR
+
+	virtual ~UPoseWatchElement() = default;
+
+	/** The default name given to all new pose watch elements */
+	FText GetDefaultLabel() const;
+
+	/** Returns the display label assigned to this pose watch */
+	FText GetLabel() const;
+
+	/** Returns the visibility of this pose watch */
+	bool GetIsVisible() const;
+
+	/** Returns true if the pose watch is able to draw anything in the viewport */
+	virtual bool GetIsEnabled() const;
+
+	/** Returns the parent pose watch this element belongs to. */
+	UPoseWatch* GetParent() const;
+
+	/**  Sets this element's parent. */
+	void SetParent(UPoseWatch* InParent);
+
+	/** Attempts to set the label, returns false if unsuccessful (if there's a name clash with another pose watch in the current directory). */
+	bool SetLabel(const FText& InLabel);
+
+	/** Sets whether or not to render this pose watch to the view port. */
+	virtual void SetIsVisible(bool bInIsVisible);
+
+	/** Returns true if there would be no name clashes when assigning the label InLabel. */
+	bool ValidateLabelRename(const FText& InLabel, FText& OutErrorMessage);
+
+	/** Returns true if InLabel is a unique among the children of this elements parent, excluding this. */
+	bool IsLabelUniqueInParent(const FText& InLabel, UPoseWatch* InParent) const;
+
+	/** Takes GetDefaultLabel() and generates unique labels to avoid name clashes (e.g. PoseWatch1, PoseWatch2, ...). */
+	void SetUniqueDefaultLabel();
+
+	/** Takes InLabel and generates unique labels to avoid name clashes (e.g. PoseWatch1, PoseWatch2, ...). */
+	void SetUniqueLabel(const FText& InLabel);
+
+	/** Returns the color to display the pose watch using. */
+	FColor GetColor() const;
+
+	/** Sets the display color of this pose watch in the UI and view port. */
+	virtual void SetColor(const FColor& InColor);
+
+	/** Returns true if it is possible to set the color of this element. */
+	bool HasColor() const;
+
+	/** Set the can set the color flag. */
+	void SetHasColor(const bool bInHasColor);
+
+	/** Get the name of the icon used to represent this element. */
+	FName GetIconName() const;
+
+	/** Set the name of the icon used to represent this element. */
+	void SetIconName(const FName InIconName);
+
+	/** Toggle's the pose watch's visibility */
+	void ToggleIsVisible();
+
+private:
+	FText FindUniqueNameInParent(class UPoseWatch* InParent) const;
+
+#endif // WITH_EDITOR
+
+#if WITH_EDITORONLY_DATA
+protected:
+
+	// If true will draw the node to the view port.
+	UPROPERTY()
+	bool bIsVisible = true;
+
+	UPROPERTY()
+	bool bHasColor = true;
+
+	UPROPERTY()
+	FColor Color;
+
+	UPROPERTY()
+	FText Label;
+
+	UPROPERTY()
+	FName IconName;
+
+	UPROPERTY()
+	TWeakObjectPtr<UPoseWatch> Parent;
+
+#endif // WITH_EDITORONLY_DATA
+};
+
+UCLASS()
+class ENGINE_API UPoseWatchPoseElement : public UPoseWatchElement
+{
+	GENERATED_BODY()
+
+public:
+
+	UPoseWatchPoseElement(const class FObjectInitializer& PCIP);
+
+#if WITH_EDITOR
+
+	virtual bool GetIsEnabled() const;
+
+#endif // WITH_EDITOR
+
+#if WITH_EDITORONLY_DATA
+	/** Optionally select a Blend Mask to control which bones on the skeleton are rendered. Any non-zero entries are rendered. */
+	UPROPERTY(EditAnywhere, editfixedsize, Category = Default, meta = (UseAsBlendMask = true))
+	TObjectPtr<UBlendProfile> ViewportMask;
+
+	/** Invert which bones are rendered when using a viewport mask */
+	UPROPERTY(EditAnywhere, Category = Default, meta = (EditCondition = "ViewportMask != nullptr"))
+	bool bInvertViewportMask;
+
+	/** The threshold which each bone's blend scale much surpass to be rendered using the viewport mask */
+	UPROPERTY(EditAnywhere, Category = Default, meta = (ClampMin = 0.f, ClampMax = 1.f, EditCondition = "ViewportMask != nullptr"))
+	float BlendScaleThreshold;
+
+	/** Offset the rendering of the bones in the viewport. */
+	UPROPERTY(EditAnywhere, Category = Default)
+	FVector3d ViewportOffset;
+#endif // WITH_EDITORONLY_DATA
+};
 
 UCLASS()
 class ENGINE_API UPoseWatch
@@ -160,16 +288,25 @@ public:
 	/** Returns the visibility of this pose watch */
 	bool GetIsVisible() const;
 
-	/** Returns true if the pose watch is connected to the output pose */
+	/** Returns true if any child element is enabled */
 	bool GetIsEnabled() const;
 
+	/** Returns true if the this pose watch's node is active in its anim graph */
+	bool GetIsNodeEnabled() const;
+
 	/** Returns the color to display the pose watch using */
 	FColor GetColor() const;
 
+	/** Set whether this should display its children in the pose watch manager window */
+	void SetIsExpanded(bool bInIsExpanded);
+
+	/** Returns true if this should display its children in the pose watch manager window */
+	bool GetIsExpanded() const;
+
 	/** Returns true if this pose watch should be deleted after the user has deselected its assigned node (Editor preference) */
 	bool GetShouldDeleteOnDeselect() const;
 
-	/** Returns the parent folder this folder belongs to, if any */
+	/** Returns the parent folder this pose watch belongs to, if any */
 	UPoseWatchFolder* GetParent() const;
 
 	/**  Attempts to set this pose watch's parent, returns false if unsuccessful (if moving causes a name clash)
@@ -178,7 +315,7 @@ public:
 	bool SetParent(UPoseWatchFolder* InParent, bool bForce=false);
 
 	/** If set, denotes the pose watch is able to be drawn to the viewport */
-	void SetIsEnabled(bool bInIsEnabled);
+	void SetIsNodeEnabled(const bool bInIsEnabled);
 
 	/** Alias of SetParent */
 	void MoveTo(UPoseWatchFolder* InFolder);
@@ -211,41 +348,56 @@ public:
 	bool ValidateLabelRename(const FText& InLabel, FText& OutErrorMessage);
 
 	/** Returns true if InLabel is a unique pose watch label among the children of InFolder, excluding this */
-	bool IsPoseWatchLabelUniqueInFolder(const FText& InLabel, UPoseWatchFolder* InFolder) const;
+	bool IsLabelUniqueInParent(const FText& InLabel, UPoseWatchFolder* InFolder) const;
 
 	/** Takes GetDefaultLabel() and generates unique labels to avoid name clashes (e.g. PoseWatch1, PoseWatch2, ...) */
 	void SetUniqueDefaultLabel();
 
+	/** Called when a child element's visibility is changed */
+	void UpdateVisibility();
+
 	/** Returns the anim blueprint this pose watch is stored inside */
 	UAnimBlueprint* GetAnimBlueprint() const;
 
+	/** Create a new element and add it to this pose watch */
+	TObjectPtr<UPoseWatchElement> AddElement(const FText InLabel, const FName IconName);
+
+	/** Create a new element of a specified type and add it to this pose watch */
+	template< class TElementType > TObjectPtr<TElementType> AddElement(const FText InLabel, const FName IconName);
+	
+	/** Returns a reference to the array of elements stored in this */
+	TArray<TObjectPtr<UPoseWatchElement>>& GetElements() { return Elements; }
+
+	/** Returns true if this pose watch contains the specified element */
+	bool Contains(const TObjectPtr<UPoseWatchElement> InElement) { return Elements.Contains(InElement); }
+	bool Contains(const UPoseWatchElement* const InElement) { return Elements.ContainsByPredicate([InElement](const TObjectPtr<UPoseWatchElement>& ContainedElement) { return InElement == ContainedElement.Get(); }); }
+	
 private:
-	/** Returns a unique name for a new pose watch  placed within InParent */
-	FText FindUniqueNameInFolder(UPoseWatchFolder* InParent) const;
+	/** Returns a unique name for a new pose watch placed within InParent */
+	FText FindUniqueNameInParent(UPoseWatchFolder* InParent) const;
+
 #endif // WITH_EDITOR
 
 #if WITH_EDITORONLY_DATA
 public:
+	virtual void Serialize(FArchive& Ar) override;
+
 	UPROPERTY()
 	TWeakObjectPtr<class UEdGraphNode> Node;
 
-	/** Optionally select a Blend Mask to control which bones on the skeleton are rendered. Any non-zero entries are rendered. */
-	UPROPERTY(EditAnywhere, editfixedsize, Category = Default, meta = (UseAsBlendMask = true))
-	TObjectPtr<UBlendProfile> ViewportMask;
-
-	/** Invert which bones are rendered when using a viewport mask */
-	UPROPERTY(EditAnywhere, Category = Default, meta = (EditCondition = "ViewportMask != nullptr"))
-	bool bInvertViewportMask;
-
-	/** The threshold which each bone's blend scale much surpass to be rendered using the viewport mask */
-	UPROPERTY(EditAnywhere, Category = Default, meta = (ClampMin = 0.f, ClampMax = 1.f, EditCondition="ViewportMask != nullptr"))
-	float BlendScaleThreshold;
-
-	/** Offset the rendering of the bones in the viewport. */
-	UPROPERTY(EditAnywhere, Category = Default)
-	FVector3d ViewportOffset;
+	UPROPERTY()
+	TObjectPtr<UBlendProfile> ViewportMask_DEPRECATED;
+	UPROPERTY()
+	bool bInvertViewportMask_DEPRECATED;
+	UPROPERTY()
+	float BlendScaleThreshold_DEPRECATED;
+	UPROPERTY()
+	FVector3d ViewportOffset_DEPRECATED;
 
 protected:
+	UPROPERTY()
+	TArray<TObjectPtr<UPoseWatchElement>> Elements;
+
 	UPROPERTY()
 	bool bDeleteOnDeselection = false;
 
@@ -253,9 +405,12 @@ protected:
 	UPROPERTY()
 	bool bIsVisible = true;
 
-	// If true, the pose is able to be drawn to the viewport
+	// If true, the node is able to be drawn to the view port.
 	UPROPERTY(Transient)
-	bool bIsEnabled = false;
+	bool bIsNodeEnabled = false;
+
+	UPROPERTY()
+	bool bIsExpanded = true;
 
 	UPROPERTY()
 	FColor Color;
@@ -263,7 +418,27 @@ protected:
 	UPROPERTY()
 	FText Label;
 
+	UPROPERTY()
+	FName IconName;
+
 	UPROPERTY()
 	TWeakObjectPtr<UPoseWatchFolder> Parent;
 #endif // WITH_EDITORONLY_DATA
-};
\ No newline at end of file
+};
+
+
+#if WITH_EDITOR
+
+template< class TElementType > TObjectPtr<TElementType> UPoseWatch::AddElement(const FText InLabel, const FName InIconName)
+{
+	TObjectPtr<TElementType> PoseWatchElement(NewObject<TElementType>());
+	PoseWatchElement->SetParent(this);
+	PoseWatchElement->SetUniqueLabel(InLabel);
+	PoseWatchElement->SetIconName(InIconName);
+	PoseWatchElement->SetColor(GetColor());
+	Elements.Add(PoseWatchElement);
+
+	return PoseWatchElement;
+}
+
+#endif // WITH_EDITOR
\ No newline at end of file
diff --git a/Engine/Source/Runtime/Engine/Private/Animation/AnimBlueprintGeneratedClass.cpp b/Engine/Source/Runtime/Engine/Private/Animation/AnimBlueprintGeneratedClass.cpp
--- a/Engine/Source/Runtime/Engine/Private/Animation/AnimBlueprintGeneratedClass.cpp
+++ b/Engine/Source/Runtime/Engine/Private/Animation/AnimBlueprintGeneratedClass.cpp
@@ -229,7 +229,7 @@ void FAnimBlueprintDebugData::RecordBlendSpacePlayer(int32 InNodeID, const UBlen
 	BlendSpacePlayerRecordsThisFrame.Emplace(InNodeID, InBlendSpace, InPosition, InFilteredPosition);
 }
 
-void FAnimBlueprintDebugData::AddPoseWatch(int32 NodeID, UPoseWatch* InPoseWatch)
+void FAnimBlueprintDebugData::AddPoseWatch(int32 NodeID, UPoseWatchPoseElement* const InPoseWatchPoseElement)
 {
 	for (FAnimNodePoseWatch& PoseWatch : AnimNodePoseWatch)
 	{
@@ -239,11 +239,14 @@ void FAnimBlueprintDebugData::AddPoseWatch(int32 NodeID, UPoseWatch* InPoseWatch
 		}
 	}
 
+	check(InPoseWatchPoseElement != nullptr); // Expect a valid pointer.
+
 	//Not found so make new one
 	AnimNodePoseWatch.Add(FAnimNodePoseWatch());
 	FAnimNodePoseWatch& NewAnimNodePoseWatch = AnimNodePoseWatch.Last();
 	NewAnimNodePoseWatch.NodeID = NodeID;
-	NewAnimNodePoseWatch.PoseWatch = InPoseWatch;
+	NewAnimNodePoseWatch.PoseWatch = InPoseWatchPoseElement->GetParent();
+	NewAnimNodePoseWatch.PoseWatchPoseElement = InPoseWatchPoseElement;
 	NewAnimNodePoseWatch.PoseInfo = MakeShareable(new FCompactHeapPose());
 	NewAnimNodePoseWatch.Object = nullptr;
 }
@@ -260,16 +263,16 @@ void FAnimBlueprintDebugData::RemovePoseWatch(int32 NodeID)
 	}
 }
 
-void FAnimBlueprintDebugData::ForEachActiveVisiblePoseWatch(const TFunctionRef<void(FPoseWatchDebugData&)>& InFunction) const
+void FAnimBlueprintDebugData::ForEachActiveVisiblePoseWatchPoseElement(const TFunctionRef<void(FPoseWatchDebugData&)>& InFunction) const
 {
 	for (const FAnimNodePoseWatch& PoseWatchNode : AnimNodePoseWatch)
 	{
-		if (const UPoseWatch* PoseWatch = PoseWatchNode.PoseWatch.Get())
+		if (const UPoseWatchPoseElement* PoseWatchPoseElement = PoseWatchNode.PoseWatchPoseElement.Get())
 		{
-			if (PoseWatch->GetIsEnabled() && PoseWatch->GetIsVisible())
+			if (PoseWatchPoseElement->GetIsEnabled() && PoseWatchPoseElement->GetIsVisible())
 			{
 				FPoseWatchDebugData DrawData;
-				DrawData.PoseWatch = PoseWatch;
+				DrawData.PoseWatchPoseElement = PoseWatchPoseElement;
 				DrawData.PoseInfo = PoseWatchNode.PoseInfo;
 
 				InFunction(DrawData);
diff --git a/Engine/Source/Runtime/Engine/Private/Animation/AnimInstanceProxy.cpp b/Engine/Source/Runtime/Engine/Private/Animation/AnimInstanceProxy.cpp
--- a/Engine/Source/Runtime/Engine/Private/Animation/AnimInstanceProxy.cpp
+++ b/Engine/Source/Runtime/Engine/Private/Animation/AnimInstanceProxy.cpp
@@ -498,7 +498,7 @@ void FAnimInstanceProxy::PreUpdate(UAnimInstance* InAnimInstance, float DeltaSec
 			PoseWatchEntriesForThisFrame = DebugData->AnimNodePoseWatch;
 			for (FAnimNodePoseWatch& PoseWatch : PoseWatchEntriesForThisFrame)
 			{
-				PoseWatch.PoseWatch.Get()->SetIsEnabled(false);
+				PoseWatch.PoseWatch.Get()->SetIsNodeEnabled(false);
 			}
 		}
 	}
@@ -3150,7 +3150,7 @@ void FAnimInstanceProxy::RegisterWatchedPose(const FCompactPose& Pose, int32 Lin
 			{
 				PoseWatch.PoseInfo->CopyBonesFrom(Pose);
 				PoseWatch.Object = GetAnimInstanceObject();
-				PoseWatch.PoseWatch->SetIsEnabled(true);
+				PoseWatch.PoseWatch->SetIsNodeEnabled(true);
 				break;
 			}
 		}
@@ -3169,7 +3169,7 @@ void FAnimInstanceProxy::RegisterWatchedPose(const FCSPose<FCompactPose>& Pose,
 				FCSPose<FCompactPose>::ConvertComponentPosesToLocalPoses(Pose, TempPose);
 				PoseWatch.PoseInfo->CopyBonesFrom(TempPose);
 				PoseWatch.Object = GetAnimInstanceObject();
-				PoseWatch.PoseWatch->SetIsEnabled(true);
+				PoseWatch.PoseWatch->SetIsNodeEnabled(true);
 				break;
 			}
 		}
diff --git a/Engine/Source/Runtime/Engine/Private/Engine/PoseWatch.cpp b/Engine/Source/Runtime/Engine/Private/Engine/PoseWatch.cpp
--- a/Engine/Source/Runtime/Engine/Private/Engine/PoseWatch.cpp
+++ b/Engine/Source/Runtime/Engine/Private/Engine/PoseWatch.cpp
@@ -9,29 +9,7 @@
 
 #define LOCTEXT_NAMESPACE "PoseWatch"
 
-UPoseWatchFolder::UPoseWatchFolder(const FObjectInitializer& ObjectInitializer)
-	: Super(ObjectInitializer)
-{
 #if WITH_EDITOR
-	Label = GetDefaultLabel();
-#endif
-}
-
-UPoseWatch::UPoseWatch(const FObjectInitializer& ObjectInitializer)
-	: Super(ObjectInitializer)
-{
-#if WITH_EDITOR
-	Label = GetDefaultLabel();
-	SetColor(PoseWatchUtil::ChoosePoseWatchColor());
-	ViewportMask = nullptr;
-	bInvertViewportMask = false;
-	BlendScaleThreshold = 0.f;
-	ViewportOffset = FVector3d::ZeroVector;
-#endif
-}
-
-#if WITH_EDITOR
-
 
 // PoseWatchUtil
 
@@ -65,34 +43,18 @@ TSet<UPoseWatchFolder*> PoseWatchUtil::GetChildrenPoseWatchFoldersOf(const UPose
 	return Children;
 }
 
-template <typename T>
-T* PoseWatchUtil::FindInFolderInCollection(const FName& Label, UPoseWatchFolder* InFolder, const TArray<TObjectPtr<T>>& PoseWatchCollection)
+template <typename TParent, typename TItem> FText PoseWatchUtil::FindUniqueNameInParent(TParent* InParent, const TItem* InItem)
 {
-	for (const TObjectPtr<T>& PoseWatchItem : PoseWatchCollection)
-	{
-		if (PoseWatchItem->IsIn(InFolder) && PoseWatchItem->GetLabel().ToString().Equals(Label.ToString()))
-		{
-			return PoseWatchItem;
-		}
-	}
-	return nullptr;
-}
+	static const FTextFormat LabelFormat = FTextFormat::FromString("{0}{1}");
 
-template <typename T>
-FText PoseWatchUtil::FindUniqueNameInFolder(UPoseWatchFolder* InParent, const T* Item, const TArray<TObjectPtr<T>>& Collection)
-{
-	FName NewName(Item->GetLabel().ToString());
-	FString BaseLabel = Item->GetLabel().ToString();
-	int32 Index = 0;
-	T* ConflictingItem = nullptr;
-	do
+	FText NewLabel = InItem->GetLabel();
+
+	for (uint32 Counter = 1; !InItem->IsLabelUniqueInParent(NewLabel, InParent); ++Counter)
 	{
-		++Index;
-		NewName = *FString::Printf(TEXT("%s%d"), *BaseLabel, Index);
-		ConflictingItem = FindInFolderInCollection(NewName, InParent, Collection);
-	} while (ConflictingItem != nullptr && ConflictingItem != Item);
+		NewLabel = FText::Format(LabelFormat, InItem->GetLabel(), FText::AsNumber(Counter));
+	}
 
-	return FText::FromName(NewName);
+	return NewLabel;
 }
 
 FColor PoseWatchUtil::ChoosePoseWatchColor()
@@ -100,9 +62,20 @@ FColor PoseWatchUtil::ChoosePoseWatchColor()
 	return FColor::MakeRandomColor();
 }
 
+#endif // WITH_EDITOR
 
 // UPoseWatchFolder
 
+UPoseWatchFolder::UPoseWatchFolder(const FObjectInitializer& ObjectInitializer)
+	: Super(ObjectInitializer)
+{
+#if WITH_EDITOR
+	Label = GetDefaultLabel();
+#endif
+}
+
+#if WITH_EDITOR
+
 const FText UPoseWatchFolder::GetPath() const
 {
 	if (Parent.Get() != nullptr)
@@ -134,22 +107,22 @@ UPoseWatchFolder* UPoseWatchFolder::GetParent() const
 
 bool UPoseWatchFolder::SetParent(UPoseWatchFolder* InParent, bool bForce)
 {
-	if (IsFolderLabelUniqueInFolder(Label, InParent))
+	if (IsLabelUniqueInParent(Label, InParent))
 	{
 		Parent = InParent;
 		return true;
 	}
 	else if (bForce)
 	{
-		Label = FindUniqueNameInFolder(InParent);
-		check(IsFolderLabelUniqueInFolder(Label, InParent))
+		Label = FindUniqueNameInParent(InParent);
+		check(IsLabelUniqueInParent(Label, InParent))
 		Parent = InParent;
 		return true;
 	}
 	return false;
 }
 
-bool UPoseWatchFolder::IsFolderLabelUniqueInFolder(const FText& InLabel, UPoseWatchFolder* InFolder) const
+bool UPoseWatchFolder::IsLabelUniqueInParent(const FText& InLabel, UPoseWatchFolder* InFolder) const
 {
 	for (UPoseWatchFolder* SomeChildFolder : PoseWatchUtil::GetChildrenPoseWatchFoldersOf(InFolder, GetAnimBlueprint()))
 	{
@@ -162,6 +135,8 @@ bool UPoseWatchFolder::IsFolderLabelUniqueInFolder(const FText& InLabel, UPoseWa
 		}
 	}
 	return true;
+
+	// TODO - folders can contain folders AND node watches ... don't we need to check against node watch names as well ?
 }
 
 void UPoseWatchFolder::MoveTo(UPoseWatchFolder* InFolder)
@@ -172,7 +147,7 @@ void UPoseWatchFolder::MoveTo(UPoseWatchFolder* InFolder)
 bool UPoseWatchFolder::SetLabel(const FText& InLabel)
 {
 	FText NewLabel = FText::TrimPrecedingAndTrailing(InLabel);
-	if (IsFolderLabelUniqueInFolder(NewLabel, Parent.Get()))
+	if (IsLabelUniqueInParent(NewLabel, Parent.Get()))
 	{
 		Label = NewLabel;
 		return true;
@@ -299,7 +274,8 @@ bool UPoseWatchFolder::ValidateLabelRename(const FText& InLabel, FText& OutError
 		OutErrorMessage = LOCTEXT("PoseWatchFolderNameEmpty", "A pose watch folder must have a label");
 		return false;
 	}
-	if (!IsFolderLabelUniqueInFolder(UseLabel, Parent.Get()))
+
+	if (!IsLabelUniqueInParent(UseLabel, Parent.Get()))
 	{
 		OutErrorMessage = LOCTEXT("PoseWatchFolderNameTaken", "A folder already has this name at this level");
 		return false;
@@ -323,12 +299,12 @@ bool UPoseWatchFolder::HasChildren() const
 void UPoseWatchFolder::SetUniqueDefaultLabel()
 {
 	Label = GetDefaultLabel();
-	Label = FindUniqueNameInFolder(Parent.Get());
+	Label = FindUniqueNameInParent(Parent.Get());
 }
 
-FText UPoseWatchFolder::FindUniqueNameInFolder(UPoseWatchFolder* InParent) const
+FText UPoseWatchFolder::FindUniqueNameInParent(UPoseWatchFolder* InFolder) const
 {
-	return PoseWatchUtil::FindUniqueNameInFolder(InParent, this, GetAnimBlueprint()->PoseWatchFolders);
+	return PoseWatchUtil::FindUniqueNameInParent(InFolder, this);
 }
 
 bool UPoseWatchFolder::HasPoseWatchChildren() const
@@ -352,8 +328,188 @@ bool UPoseWatchFolder::HasPoseWatchDescendents() const
 	return false;
 }
 
+// UPoseWatchElement
+
+FText UPoseWatchElement::GetDefaultLabel() const
+{
+	return LOCTEXT("NewPoseWatchElement", "NewPoseWatchElement");
+}
+
+FText UPoseWatchElement::GetLabel() const
+{
+	return Label;
+}
+
+bool UPoseWatchElement::GetIsVisible() const
+{
+	return bIsVisible;
+}
+
+bool UPoseWatchElement::GetIsEnabled() const
+{
+	return true;
+}
+
+UPoseWatch* UPoseWatchElement::GetParent() const
+{
+	return Parent.Get();
+}
+
+void UPoseWatchElement::SetParent(UPoseWatch* InParent)
+{
+	Parent = InParent;
+}
+
+bool UPoseWatchElement::SetLabel(const FText& InLabel)
+{
+	if (Parent.IsValid() && IsLabelUniqueInParent(InLabel, Parent.Get()))
+	{
+		Label = InLabel;
+		return true;
+	}
+
+	return false;
+}
+
+void UPoseWatchElement::SetIsVisible(bool bInIsVisible)
+{
+	bIsVisible = bInIsVisible;
+
+	if (Parent.IsValid())
+	{
+		Parent->UpdateVisibility();
+	}
+}
+
+bool UPoseWatchElement::ValidateLabelRename(const FText& InLabel, FText& OutErrorMessage)
+{
+	if (!IsLabelUniqueInParent(InLabel, Parent.Get()))
+	{
+		OutErrorMessage = LOCTEXT("PoseWatchElementNameTaken", "A element already has this name in this node watch");
+		return false;
+	}
+
+	return true;
+}
+
+bool UPoseWatchElement::IsLabelUniqueInParent(const FText& InLabel, UPoseWatch* InParent) const
+{
+	for (TObjectPtr<UPoseWatchElement>& Element : InParent->GetElements())
+	{
+		if (Element->GetLabel().ToString().Equals(InLabel.ToString()))
+		{
+			if (Element != this)
+			{
+				return false;
+			}
+		}
+	}
+	return true;
+}
+
+void UPoseWatchElement::SetUniqueDefaultLabel()
+{
+	Label = GetDefaultLabel();
+	Label = FindUniqueNameInParent(Parent.Get());
+}
+
+void UPoseWatchElement::SetUniqueLabel(const FText& InLabel)
+{
+	Label = InLabel;
+	Label = FindUniqueNameInParent(Parent.Get());
+}
+
+FColor UPoseWatchElement::GetColor() const
+{
+	return Color;
+}
+
+void UPoseWatchElement::SetColor(const FColor& InColor)
+{
+	Color = InColor;
+}
+
+bool UPoseWatchElement::HasColor() const
+{
+	return bHasColor;
+}
+
+void UPoseWatchElement::SetHasColor(const bool bInHasColor)
+{
+	bHasColor = bInHasColor;
+}
+
+FName UPoseWatchElement::GetIconName() const
+{
+	return IconName;
+}
+
+void UPoseWatchElement::SetIconName(const FName InIconName)
+{
+	IconName = InIconName;
+}
+
+void UPoseWatchElement::ToggleIsVisible()
+{
+	SetIsVisible(!bIsVisible);
+}
+
+FText UPoseWatchElement::FindUniqueNameInParent(UPoseWatch* InParent) const
+{
+	if (InParent)
+	{
+		return PoseWatchUtil::FindUniqueNameInParent(InParent, this);
+	}
+	return FText();
+}
+
+#endif // WITH_EDITOR
+
+// UPoseWatchPoseElement
+
+UPoseWatchPoseElement::UPoseWatchPoseElement(const FObjectInitializer& ObjectInitializer)
+	: Super(ObjectInitializer)
+{
+#if WITH_EDITOR
+	ViewportMask = nullptr;
+	bInvertViewportMask = false;
+	BlendScaleThreshold = 0.f;
+	ViewportOffset = FVector3d::ZeroVector;
+#endif
+}
+
+#if WITH_EDITOR
+
+bool UPoseWatchPoseElement::GetIsEnabled() const
+{
+	return Parent.Get() && Parent->GetIsNodeEnabled();
+}
+
+#endif // WITH_EDITOR
+
 // UPoseWatch
 
+UPoseWatch::UPoseWatch(const FObjectInitializer& ObjectInitializer)
+	: Super(ObjectInitializer)
+{
+#if WITH_EDITOR
+	Label = GetDefaultLabel();
+	SetColor(PoseWatchUtil::ChoosePoseWatchColor());
+#endif
+}
+
+#if WITH_EDITOR
+
+void UPoseWatch::SetIsExpanded(bool bInIsExpanded)
+{
+	bIsExpanded = bInIsExpanded;
+}
+
+bool UPoseWatch::GetIsExpanded() const
+{
+	return bIsExpanded;
+}
+
 const FText UPoseWatch::GetPath() const
 {
 	check(!Label.IsEmpty())
@@ -383,6 +539,24 @@ bool UPoseWatch::GetIsVisible() const
 	return bIsVisible;
 }
 
+bool UPoseWatch::GetIsNodeEnabled() const
+{
+	return bIsNodeEnabled;
+}
+
+bool UPoseWatch::GetIsEnabled() const
+{
+ 	bool bIsEnabled = false;
+  
+	// Set visible if any child component is visible.
+	for (TObjectPtr<UPoseWatchElement> PoseWatchElement : Elements)
+	{
+  		bIsEnabled |= PoseWatchElement->GetIsEnabled();
+	}
+  
+	return bIsEnabled;
+}
+
 FColor UPoseWatch::GetColor() const
 {
 	return Color;
@@ -400,7 +574,6 @@ void UPoseWatch::OnRemoved()
 	
 	AnimationEditorUtils::RemovePoseWatch(this, AnimBlueprint);
 
-
 	if (Parent.IsValid())
 	{
 		Parent->UpdateVisibility();
@@ -416,13 +589,13 @@ UPoseWatchFolder* UPoseWatch::GetParent() const
 
 bool UPoseWatch::SetParent(UPoseWatchFolder* InParent, bool bForce)
 {
-	if (!IsPoseWatchLabelUniqueInFolder(Label, InParent))
+	if (!IsLabelUniqueInParent(Label, InParent))
 	{
 		if (!bForce)
 		{
 			return false;
 		}
-		Label = FindUniqueNameInFolder(InParent);
+		Label = FindUniqueNameInParent(InParent);
 	}
 
 	UPoseWatchFolder* OldParent = Parent.Get();
@@ -442,14 +615,9 @@ bool UPoseWatch::SetParent(UPoseWatchFolder* InParent, bool bForce)
 	return true;
 }
 
-bool UPoseWatch::GetIsEnabled() const
+void UPoseWatch::SetIsNodeEnabled(const bool bInIsNodeEnabled)
 {
-	return bIsEnabled;
-}
-
-void UPoseWatch::SetIsEnabled(bool bInIsEnabled)
-{
-	bIsEnabled = bInIsEnabled;
+	bIsNodeEnabled = bInIsNodeEnabled;
 }
 
 void UPoseWatch::MoveTo(UPoseWatchFolder* InFolder)
@@ -460,16 +628,23 @@ void UPoseWatch::MoveTo(UPoseWatchFolder* InFolder)
 bool UPoseWatch::SetLabel(const FText& InLabel)
 {
 	FText NewLabel = FText::TrimPrecedingAndTrailing(InLabel);
-	if (IsPoseWatchLabelUniqueInFolder(NewLabel, Parent.Get()))
+
+	if (IsLabelUniqueInParent(NewLabel, Parent.Get()))
 	{
 		Label = NewLabel;
 		return true;
 	}
+
 	return false;
 }
 
 void UPoseWatch::SetIsVisible(bool bInIsVisible)
 {
+	for (TObjectPtr<UPoseWatchElement> PoseWatchElement : Elements)
+	{
+		PoseWatchElement->SetIsVisible(bInIsVisible);
+	}
+
 	bIsVisible = bInIsVisible;
 
 	if (Parent.IsValid())
@@ -481,6 +656,11 @@ void UPoseWatch::SetIsVisible(bool bInIsVisible)
 void UPoseWatch::SetColor(const FColor& InColor)
 {
 	Color = InColor;
+
+	for (TObjectPtr<UPoseWatchElement> PoseWatchElement : Elements)
+	{
+		PoseWatchElement->SetColor(InColor);
+	}
 }
 
 void UPoseWatch::SetShouldDeleteOnDeselect(const bool bInDeleteOnDeselection)
@@ -511,7 +691,8 @@ bool UPoseWatch::ValidateLabelRename(const FText& InLabel, FText& OutErrorMessag
 		OutErrorMessage = LOCTEXT("PoseWatchNameEmpty", "A pose watch must have a label");
 		return false;
 	}
-	if (!IsPoseWatchLabelUniqueInFolder(UseLabel, Parent.Get()))
+
+	if (!IsLabelUniqueInParent(UseLabel, Parent.Get()))
 	{
 		OutErrorMessage = LOCTEXT("PoseWatchNameTaken", "A pose watch already has this name at this level");
 		return false;
@@ -519,7 +700,7 @@ bool UPoseWatch::ValidateLabelRename(const FText& InLabel, FText& OutErrorMessag
 	return true;
 }
 
-bool UPoseWatch::IsPoseWatchLabelUniqueInFolder(const FText& InLabel, UPoseWatchFolder* InFolder) const
+bool UPoseWatch::IsLabelUniqueInParent(const FText& InLabel, UPoseWatchFolder* InFolder) const
 {
 	for (UPoseWatch* SomePoseWatch : PoseWatchUtil::GetChildrenPoseWatchOf(InFolder, GetAnimBlueprint()))
 	{
@@ -537,7 +718,18 @@ bool UPoseWatch::IsPoseWatchLabelUniqueInFolder(const FText& InLabel, UPoseWatch
 void UPoseWatch::SetUniqueDefaultLabel()
 {
 	Label = GetDefaultLabel();
-	Label = FindUniqueNameInFolder(Parent.Get());
+	Label = FindUniqueNameInParent(Parent.Get());
+}
+
+void UPoseWatch::UpdateVisibility()
+{
+	bIsVisible = false;
+
+	// Set visible if any child component is visible.
+	for (TObjectPtr<UPoseWatchElement> PoseWatchElement : Elements)
+	{
+		bIsVisible |= PoseWatchElement->GetIsVisible();
+	}
 }
 
 UAnimBlueprint* UPoseWatch::GetAnimBlueprint() const
@@ -545,11 +737,42 @@ UAnimBlueprint* UPoseWatch::GetAnimBlueprint() const
 	return  CastChecked<UAnimBlueprint>(GetOuter());
 }
 
-FText UPoseWatch::FindUniqueNameInFolder(UPoseWatchFolder* InParent) const
+FText UPoseWatch::FindUniqueNameInParent(UPoseWatchFolder* InParent) const
+{
+	return PoseWatchUtil::FindUniqueNameInParent(InParent, this);
+}
+
+TObjectPtr<UPoseWatchElement> UPoseWatch::AddElement(const FText InLabel, const FName InIconName)
 {
-	return PoseWatchUtil::FindUniqueNameInFolder(InParent, this, GetAnimBlueprint()->PoseWatches);
+	return AddElement<UPoseWatchElement>(InLabel, InIconName);
 }
 
 #endif // WITH_EDITOR
 
+#if WITH_EDITORONLY_DATA
+void UPoseWatch::Serialize(FArchive& Ar)
+{
+	Super::Serialize(Ar);
+
+	Ar.UsingCustomVersion(FFortniteMainBranchObjectVersion::GUID);
+
+	if (Ar.CustomVer(FFortniteMainBranchObjectVersion::GUID) < FFortniteMainBranchObjectVersion::PoseWatchMigrateSkeletonDrawParametersToPoseElement)
+	{
+		// Initialize all PoseWatchElement's Skeleton Draw Parameters from those defined in the parent Pose Watch.
+		for (TObjectPtr<UPoseWatchElement> PoseWatchElement : Elements)
+		{
+			UPoseWatchPoseElement* const PoseWatchPoseElement = Cast<UPoseWatchPoseElement>(PoseWatchElement.Get());
+
+			if (PoseWatchPoseElement)
+			{
+				PoseWatchPoseElement->ViewportMask = ViewportMask_DEPRECATED;
+				PoseWatchPoseElement->bInvertViewportMask = bInvertViewportMask_DEPRECATED;
+				PoseWatchPoseElement->BlendScaleThreshold = BlendScaleThreshold_DEPRECATED;
+				PoseWatchPoseElement->ViewportOffset = ViewportOffset_DEPRECATED;
+			}
+		}
+	}
+}
+#endif // WITH_EDITORONLY_DATA
+
 #undef LOCTEXT_NAMESPACE
\ No newline at end of file
diff --git a/Engine/Source/Runtime/Engine/Private/SkeletalDebugRendering.cpp b/Engine/Source/Runtime/Engine/Private/SkeletalDebugRendering.cpp
--- a/Engine/Source/Runtime/Engine/Private/SkeletalDebugRendering.cpp
+++ b/Engine/Source/Runtime/Engine/Private/SkeletalDebugRendering.cpp
@@ -175,29 +175,29 @@ void DrawBonesFromPoseWatch(
 	const FCompactHeapPose& Pose,
 	USkeletalMeshComponent* MeshComponent,
 	FPrimitiveDrawInterface* PDI,
-	const UPoseWatch* PoseWatch)
+	const UPoseWatchPoseElement* PoseWatchPoseElement)
 {
 #if WITH_EDITOR
 	if (Pose.GetNumBones() == 0 ||
 		Pose.GetBoneContainer().GetCompactPoseNumBones() == 0 ||
 		!MeshComponent ||
 		!MeshComponent->GetSkeletalMesh() ||
-		!PoseWatch
+		!PoseWatchPoseElement
 		)
 	{
 		return;
 	}
 
 	// optionally override draw color
-	const FLinearColor BoneColor = FLinearColor(PoseWatch->GetColor());
+	const FLinearColor BoneColor = FLinearColor(PoseWatchPoseElement->GetColor());
 
 	TArray<FTransform> WorldTransforms;
 	WorldTransforms.AddUninitialized(Pose.GetBoneContainer().GetNumBones());
 
 	TArray<uint16> RequiredBones;
-	const UBlendProfile* ViewportMask = PoseWatch->ViewportMask;
+	const UBlendProfile* ViewportMask = PoseWatchPoseElement->ViewportMask;
 
-	const FVector RelativeOffset = MeshComponent->GetComponentRotation().RotateVector(PoseWatch->ViewportOffset);
+	const FVector RelativeOffset = MeshComponent->GetComponentRotation().RotateVector(PoseWatchPoseElement->ViewportOffset);
 
 	// we could cache parent bones as we calculate, but right now I'm not worried about perf issue of this
 	for (FCompactPoseBoneIndex BoneIndex : Pose.ForEachBoneIndex())
@@ -223,8 +223,8 @@ void DrawBonesFromPoseWatch(
 		else
 		{
 			const FName& BoneName = MeshComponent->GetBoneName(MeshBoneIndex.GetInt());
-			const bool bBoneBlendScaleMeetsThreshold = ViewportMask->GetBoneBlendScale(BoneName) > PoseWatch->BlendScaleThreshold;
-			if ((!PoseWatch->bInvertViewportMask && bBoneBlendScaleMeetsThreshold) || (PoseWatch->bInvertViewportMask && !bBoneBlendScaleMeetsThreshold))
+			const bool bBoneBlendScaleMeetsThreshold = ViewportMask->GetBoneBlendScale(BoneName) > PoseWatchPoseElement->BlendScaleThreshold;
+			if ((!PoseWatchPoseElement->bInvertViewportMask && bBoneBlendScaleMeetsThreshold) || (PoseWatchPoseElement->bInvertViewportMask && !bBoneBlendScaleMeetsThreshold))
 			{
 				RequiredBones.Add(MeshBoneIndex.GetInt());
 			}
diff --git a/Engine/Source/Runtime/Engine/Public/SkeletalDebugRendering.h b/Engine/Source/Runtime/Engine/Public/SkeletalDebugRendering.h
--- a/Engine/Source/Runtime/Engine/Public/SkeletalDebugRendering.h
+++ b/Engine/Source/Runtime/Engine/Public/SkeletalDebugRendering.h
@@ -7,6 +7,7 @@
 #include "BoneIndices.h"
 
 class UPoseWatch;
+class UPoseWatchPoseElement;
 class FPrimitiveDrawInterface;
 struct FCompactHeapPose;
 struct FReferenceSkeleton;
@@ -116,7 +117,7 @@ ENGINE_API void DrawBonesFromPoseWatch(
 	const FCompactHeapPose& Pose,
 	USkeletalMeshComponent* MeshComponent,
 	FPrimitiveDrawInterface* PDI,
-	const UPoseWatch* PoseWatch);
+	const UPoseWatchPoseElement* PoseWatch);
 
 
 /**
