When post-loading, prevent marking animation asset dirty when replacing its skeleton, animation sequences when handling any model notifications and when calling UAnimSequence::PostEditChangeProperty

diff --git a/Engine/Source/Runtime/Engine/Private/Animation/AnimSequence.cpp b/Engine/Source/Runtime/Engine/Private/Animation/AnimSequence.cpp
--- a/Engine/Source/Runtime/Engine/Private/Animation/AnimSequence.cpp
+++ b/Engine/Source/Runtime/Engine/Private/Animation/AnimSequence.cpp
@@ -1091,7 +1091,7 @@ void UAnimSequence::PostEditChangeProperty(FPropertyChangedEvent& PropertyChange
 {
 	Super::PostEditChangeProperty(PropertyChangedEvent);
 
-	if(!IsTemplate())
+	if(!IsTemplate() && !FUObjectThreadContext::Get().IsRoutingPostLoad)
 	{
 		// Make sure package is marked dirty when doing stuff like adding/removing notifies
 		MarkPackageDirty();
@@ -5357,7 +5357,7 @@ void UAnimSequence::OnModelModified(const EAnimDataModelNotifyType& NotifyType,
 		}
 	};
 
-	bool bShouldMarkPackageDirty = true;
+	bool bShouldMarkPackageDirty = !FUObjectThreadContext::Get().IsRoutingPostLoad;
 
 	PRAGMA_DISABLE_DEPRECATION_WARNINGS
 	switch (NotifyType)
diff --git a/Engine/Source/Runtime/Engine/Private/AnimationAsset.cpp b/Engine/Source/Runtime/Engine/Private/AnimationAsset.cpp
--- a/Engine/Source/Runtime/Engine/Private/AnimationAsset.cpp
+++ b/Engine/Source/Runtime/Engine/Private/AnimationAsset.cpp
@@ -13,6 +13,7 @@
 #include "Animation/PoseAsset.h"
 #include "Animation/AnimSequenceHelpers.h"
 #include "Animation/AnimNodeBase.h"
+#include "UObject/UObjectThreadContext.h"
 
 #define LOCTEXT_NAMESPACE "AnimationAsset"
 
@@ -476,7 +477,10 @@ bool UAnimationAsset::ReplaceSkeleton(USkeleton* NewSkeleton, bool bConvertSpace
 		}
 
 		PostEditChange();
-		MarkPackageDirty();
+		if (!FUObjectThreadContext::Get().IsRoutingPostLoad)
+		{
+		    MarkPackageDirty();
+		}
 		return true;
 	}
 
