Avoid using SetWorldAssetByPackageName for Instanced Packages

diff --git a/Engine/Source/Runtime/Engine/Private/LevelStreaming.cpp b/Engine/Source/Runtime/Engine/Private/LevelStreaming.cpp
--- a/Engine/Source/Runtime/Engine/Private/LevelStreaming.cpp
+++ b/Engine/Source/Runtime/Engine/Private/LevelStreaming.cpp
@@ -1617,7 +1617,9 @@ ULevelStreaming* ULevelStreaming::CreateInstance(const FString& InstanceUniqueNa
 			// new level streaming instance will load the same map package as this object
 			StreamingLevelInstance->PackageNameToLoad = ((PackageNameToLoad == NAME_None) ? GetWorldAssetPackageFName() : PackageNameToLoad);
 			// under a provided unique name
-			StreamingLevelInstance->SetWorldAssetByPackageName(InstanceUniquePackageName);
+
+			FSoftObjectPath WorldAssetPath(*WriteToString<512>(InstanceUniquePackageName, TEXT("."), FPackageName::GetShortName(StreamingLevelInstance->PackageNameToLoad)));
+			StreamingLevelInstance->SetWorldAsset(TSoftObjectPtr<UWorld>(WorldAssetPath));
 			StreamingLevelInstance->SetShouldBeLoaded(false);
 			StreamingLevelInstance->SetShouldBeVisible(false);
 			StreamingLevelInstance->LevelTransform = LevelTransform;
@@ -2255,7 +2257,9 @@ ULevelStreamingDynamic* ULevelStreamingDynamic::LoadLevelInstance_Internal(UWorl
     
 	// Setup streaming level object that will load specified map
 	ULevelStreamingDynamic* StreamingLevel = NewObject<ULevelStreamingDynamic>(World, LevelStreamingClass, NAME_None, RF_Transient, NULL);
-    StreamingLevel->SetWorldAssetByPackageName(UnmodifiedLevelPackageName);
+
+	FSoftObjectPath WorldAssetPath(*WriteToString<512>(UnmodifiedLevelPackageName, TEXT("."), ShortPackageName));
+    StreamingLevel->SetWorldAsset(TSoftObjectPtr<UWorld>(WorldAssetPath));
 #if WITH_EDITOR
 	if (bIsPlayInEditor)
 	{
