World Partition Changelist Validation: store level script external actors references in the Asset Registry to be able to validate changelists without loading the associated world.

#rnx

diff --git a/Engine/Source/Runtime/Engine/Classes/Engine/Level.h b/Engine/Source/Runtime/Engine/Classes/Engine/Level.h
--- a/Engine/Source/Runtime/Engine/Classes/Engine/Level.h
+++ b/Engine/Source/Runtime/Engine/Classes/Engine/Level.h
@@ -793,20 +793,29 @@ public:
 
 #if WITH_EDITOR
 	ENGINE_API static bool GetLevelBoundsFromAsset(const FAssetData& Asset, FBox& OutLevelBounds);
-	ENGINE_API static bool GetIsLevelPartitionedFromAsset(const FAssetData& Asset);
-	ENGINE_API static bool GetIsLevelUsingExternalActorsFromAsset(const FAssetData& Asset);
-	ENGINE_API static bool GetIsUsingActorFoldersFromAsset(const FAssetData& Asset);
-	ENGINE_API static bool GetIsStreamingDisabledFromAsset(const FAssetData& Asset);
-
 	ENGINE_API static bool GetLevelBoundsFromPackage(FName LevelPackage, FBox& OutLevelBounds);
+
+	ENGINE_API static bool GetLevelScriptExternalActorsReferencesFromAsset(const FAssetData& Asset, TArray<FGuid>& OutLevelScriptExternalActorsReferences);
+	ENGINE_API static bool GetLevelScriptExternalActorsReferencesFromPackage(FName LevelPackage, TArray<FGuid>& OutLevelScriptExternalActorsReferences);
+
+	ENGINE_API static bool GetIsLevelPartitionedFromAsset(const FAssetData& Asset);
 	ENGINE_API static bool GetIsLevelPartitionedFromPackage(FName LevelPackage);
+
+	ENGINE_API static bool GetIsLevelUsingExternalActorsFromAsset(const FAssetData& Asset);
 	ENGINE_API static bool GetIsLevelUsingExternalActorsFromPackage(FName LevelPackage);
+
+	ENGINE_API static bool GetIsUsingActorFoldersFromAsset(const FAssetData& Asset);
 	ENGINE_API static bool GetIsUsingActorFoldersFromPackage(FName LevelPackage);
+
+	ENGINE_API static bool GetIsStreamingDisabledFromAsset(const FAssetData& Asset);	
 	ENGINE_API static bool GetIsStreamingDisabledFromPackage(FName LevelPackage);
+
 	ENGINE_API static bool GetPartitionedLevelCanBeUsedByLevelInstanceFromAsset(const FAssetData& Asset);
 	ENGINE_API static bool GetPartitionedLevelCanBeUsedByLevelInstanceFromPackage(FName LevelPackage);
+
 	ENGINE_API static FVector GetLevelInstancePivotOffsetFromAsset(const FAssetData& Asset);
 	ENGINE_API static FVector GetLevelInstancePivotOffsetFromPackage(FName LevelPackage);
+
 	ENGINE_API static const FName LoadAllExternalObjectsTag;
 	ENGINE_API static const FName DontLoadExternalObjectsTag;
 
diff --git a/Engine/Source/Runtime/Engine/Private/Level.cpp b/Engine/Source/Runtime/Engine/Private/Level.cpp
--- a/Engine/Source/Runtime/Engine/Private/Level.cpp
+++ b/Engine/Source/Runtime/Engine/Private/Level.cpp
@@ -2385,6 +2385,22 @@ bool ULevel::GetIsStreamingDisabledFromAsset(const FAssetData& Asset)
 	return false;
 }
 
+bool ULevel::GetLevelScriptExternalActorsReferencesFromAsset(const FAssetData& Asset, TArray<FGuid>& OutLevelScriptExternalActorsReferences)
+{
+	FString LevelScriptExternalActorsReferencesStr;
+	static const FName NAME_LevelScriptExternalActorsReferences(TEXT("LevelScriptExternalActorsReferences"));
+	if (Asset.GetTagValue(NAME_LevelScriptExternalActorsReferences, LevelScriptExternalActorsReferencesStr))
+	{
+		TArray<FString> LevelScriptReferencesStr;
+		if (LevelScriptExternalActorsReferencesStr.ParseIntoArray(LevelScriptReferencesStr, TEXT(",")))
+		{
+			Algo::Transform(LevelScriptReferencesStr, OutLevelScriptExternalActorsReferences, [](const FString& GuidStr) { return FGuid(GuidStr); });
+			return true;
+		}
+	}
+	return false;
+}
+
 bool ULevel::GetIsStreamingDisabledFromPackage(FName LevelPackage)
 {
 	return LevelAssetRegistryHelper::GetLevelInfoFromAssetRegistry(LevelPackage, [](const FAssetData& Asset)
@@ -2393,6 +2409,14 @@ bool ULevel::GetIsStreamingDisabledFromPackage(FName LevelPackage)
 	});
 }
 
+bool ULevel::GetLevelScriptExternalActorsReferencesFromPackage(FName LevelPackage, TArray<FGuid>& OutLevelScriptExternalActorsReferences)
+{
+	return LevelAssetRegistryHelper::GetLevelInfoFromAssetRegistry(LevelPackage, [&OutLevelScriptExternalActorsReferences](const FAssetData& Asset)
+	{
+		return GetLevelScriptExternalActorsReferencesFromAsset(Asset, OutLevelScriptExternalActorsReferences);
+	});
+}
+
 bool ULevel::GetPartitionedLevelCanBeUsedByLevelInstanceFromAsset(const FAssetData& Asset)
 {
 	FString PartitionedLevelCanBeUsedByLevelInstanceStr;
diff --git a/Engine/Source/Runtime/Engine/Private/WorldPartition/WorldPartition.cpp b/Engine/Source/Runtime/Engine/Private/WorldPartition/WorldPartition.cpp
--- a/Engine/Source/Runtime/Engine/Private/WorldPartition/WorldPartition.cpp
+++ b/Engine/Source/Runtime/Engine/Private/WorldPartition/WorldPartition.cpp
@@ -34,6 +34,8 @@
 #include "ScopedTransaction.h"
 #include "UnrealEdMisc.h"
 #include "LocationVolume.h"
+#include "Engine/LevelScriptBlueprint.h"
+#include "Engine/Public/ActorReferencesUtils.h"
 #include "WorldPartition/IWorldPartitionEditorModule.h"
 #include "WorldPartition/WorldPartitionLevelHelper.h"
 #include "WorldPartition/WorldPartitionLevelStreamingDynamic.h"
@@ -1315,6 +1317,26 @@ void UWorldPartition::AppendAssetRegistryTags(TArray<FAssetRegistryTag>& OutTags
 		static const FName NAME_PartitionedLevelCanBeUsedByLevelInstance(TEXT("PartitionedLevelCanBeUsedByLevelInstance"));
 		OutTags.Add(FAssetRegistryTag(NAME_PartitionedLevelCanBeUsedByLevelInstance, TEXT("1"), FAssetRegistryTag::TT_Hidden));
 	}
+
+	// Append level script references so we can perform changelists validations without loading the world
+	if (const ULevelScriptBlueprint* LevelScriptBlueprint = GetWorld()->PersistentLevel->GetLevelScriptBlueprint(true))
+	{
+		TArray<AActor*> LevelScriptExternalActorReferences = ActorsReferencesUtils::GetExternalActorReferences((UObject*)LevelScriptBlueprint);
+		
+		if (LevelScriptExternalActorReferences.Num())
+		{
+			FStringBuilderBase StringBuilder;
+			for (const AActor* Actor : LevelScriptExternalActorReferences)
+			{
+				StringBuilder.Append(Actor->GetActorGuid().ToString(EGuidFormats::Short));
+				StringBuilder.AppendChar(TEXT(','));
+			}
+			StringBuilder.RemoveSuffix(1);
+
+			static const FName NAME_LevelScriptExternalActorsReferences(TEXT("LevelScriptExternalActorsReferences"));
+			OutTags.Add(FAssetRegistryTag(NAME_LevelScriptExternalActorsReferences, StringBuilder.ToString(), FAssetRegistryTag::TT_Hidden));
+		}
+	}
 }
 
 void UWorldPartition::PinActors(const TArray<FGuid>& ActorGuids)
diff --git a/Engine/Source/Runtime/Engine/Private/WorldPartition/WorldPartitionStreamingGeneration.cpp b/Engine/Source/Runtime/Engine/Private/WorldPartition/WorldPartitionStreamingGeneration.cpp
--- a/Engine/Source/Runtime/Engine/Private/WorldPartition/WorldPartitionStreamingGeneration.cpp
+++ b/Engine/Source/Runtime/Engine/Private/WorldPartition/WorldPartitionStreamingGeneration.cpp
@@ -21,6 +21,7 @@
 #include "WorldPartition/ErrorHandling/WorldPartitionStreamingGenerationMapCheckErrorHandler.h"
 #include "WorldPartition/HLOD/HLODActor.h"
 #include "HAL/FileManager.h"
+#include "Algo/Transform.h"
 
 #define LOCTEXT_NAMESPACE "WorldPartition"
 
@@ -491,30 +492,14 @@ class FWorldPartitionStreamingGenerator
 	{
 		if (bIsMainContainer)
 		{
+			TArray<FGuid> LevelScriptReferences;
 			if (UWorld* World = ContainerDescriptor.Container->GetWorld())
 			{
 				// Gather all references to external actors from the level script and make them always loaded
 				if (ULevelScriptBlueprint* LevelScriptBlueprint = World->PersistentLevel->GetLevelScriptBlueprint(true))
 				{
 					TArray<AActor*> LevelScriptExternalActorReferences = ActorsReferencesUtils::GetExternalActorReferences(LevelScriptBlueprint);
-
-					for (AActor* Actor : LevelScriptExternalActorReferences)
-					{
-						if (FWorldPartitionActorDescView* ActorDescView = ContainerDescriptor.ActorDescViewMap.FindByGuid(Actor->GetActorGuid()))
-						{
-							if (ActorDescView->GetIsSpatiallyLoaded())
-							{
-								ErrorHandler->OnInvalidReferenceLevelScriptStreamed(*ActorDescView);
-								ActorDescView->SetForcedNonSpatiallyLoaded();
-							}
-
-							if (ActorDescView->GetRuntimeDataLayers().Num())
-							{
-								ErrorHandler->OnInvalidReferenceLevelScriptDataLayers(*ActorDescView);
-								ActorDescView->SetInvalidDataLayers();
-							}
-						}
-					}
+					Algo::Transform(LevelScriptExternalActorReferences, LevelScriptReferences, [](const AActor* Actor) { return Actor->GetActorGuid(); });
 				}
 
 				// Validate data layers
@@ -527,6 +512,28 @@ class FWorldPartitionStreamingGenerator
 					});
 				}
 			}
+			else
+			{
+				ULevel::GetLevelScriptExternalActorsReferencesFromPackage(ContainerDescriptor.Container->GetContainerPackage(), LevelScriptReferences);
+			}
+
+			for (const FGuid& LevelScriptReferenceActorGuid : LevelScriptReferences)
+			{
+				if (FWorldPartitionActorDescView* ActorDescView = ContainerDescriptor.ActorDescViewMap.FindByGuid(LevelScriptReferenceActorGuid))
+				{
+					if (ActorDescView->GetIsSpatiallyLoaded())
+					{
+						ErrorHandler->OnInvalidReferenceLevelScriptStreamed(*ActorDescView);
+						ActorDescView->SetForcedNonSpatiallyLoaded();
+					}
+
+					if (ActorDescView->GetRuntimeDataLayers().Num())
+					{
+						ErrorHandler->OnInvalidReferenceLevelScriptDataLayers(*ActorDescView);
+						ActorDescView->SetInvalidDataLayers();
+					}
+				}
+			}
 		}
 
 		// Perform various adjustements based on validations and report errors
