UWorld::BeginDestroy: Call CleanupWorld on EWorldType::Inactive (Editor only)

Takes care of cases where UEditorEngine::OnAssetLoaded/OnAssetCreated would initialize a EWorldType::Inactive world and then some code paths would remove it's RF_Standalone flag + GC. (ex: Redirector fixup)

diff --git a/Engine/Source/Runtime/Engine/Private/World.cpp b/Engine/Source/Runtime/Engine/Private/World.cpp
--- a/Engine/Source/Runtime/Engine/Private/World.cpp
+++ b/Engine/Source/Runtime/Engine/Private/World.cpp
@@ -1011,6 +1011,15 @@ void UWorld::BeginDestroy()
 {
 	Super::BeginDestroy();
 
+#if WITH_EDITOR
+	// Make sure we catch worlds that where initialized through UEditorEngine::OnAssetLoaded/OnAssetCreated
+	// (This can happen if World was loaded, then its RF_Standalone flag was removed and a GC happened) 
+	if (WorldType == EWorldType::Inactive && IsInitializedAndNeedsCleanup())
+	{
+		CleanupWorld();
+	}
+#endif
+
 	for (FLevelCollection& LevelCollection : LevelCollections)
 	{
 		TSet<TObjectPtr<ULevel>> CollectionLevels = LevelCollection.GetLevels();
