Fix for incorrect package name in ClientHasInitializedLevelFor

diff --git a/Engine/Source/Runtime/Engine/Private/NetConnection.cpp b/Engine/Source/Runtime/Engine/Private/NetConnection.cpp
--- a/Engine/Source/Runtime/Engine/Private/NetConnection.cpp
+++ b/Engine/Source/Runtime/Engine/Private/NetConnection.cpp
@@ -1389,9 +1389,8 @@ bool UNetConnection::ClientHasInitializedLevelFor(const AActor* TestActor) const
 
 	// This function is called a lot, basically for every replicated actor every time it replicates, on every client connection
 	// Each client connection has a different visibility state (what levels are currently loaded for them).
-	// Actor's outer is what we need
 
-	const FName PackageName = TestActor->GetOutermost()->GetFName();
+	const FName PackageName = TestActor->GetLevel()->GetPackage()->GetFName();
 
 	if (const bool* bIsVisible = ClientVisibleActorOuters.Find(PackageName))
 	{
