Fix compile errors as a result of CL# 21181135
#rnx

diff --git a/Engine/Source/Runtime/Engine/Private/Components/SkeletalMeshComponent.cpp b/Engine/Source/Runtime/Engine/Private/Components/SkeletalMeshComponent.cpp
--- a/Engine/Source/Runtime/Engine/Private/Components/SkeletalMeshComponent.cpp
+++ b/Engine/Source/Runtime/Engine/Private/Components/SkeletalMeshComponent.cpp
@@ -4567,9 +4567,9 @@ void USkeletalMeshComponent::DebugDrawPoseWatches(FPrimitiveDrawInterface* PDI)
 			if (AnimBPGenClass)
 			{
 				const FAnimBlueprintDebugData& DebugData = AnimBPGenClass->GetAnimBlueprintDebugData();
-				DebugData.ForEachActiveVisiblePoseWatch([this, PDI](const FPoseWatchDebugData& PoseWatchDebugData)
+				DebugData.ForEachActiveVisiblePoseWatchPoseElement([this, PDI](const FPoseWatchDebugData& PoseWatchDebugData)
 				{
-					SkeletalDebugRendering::DrawBonesFromPoseWatch(*PoseWatchDebugData.PoseInfo.Get(), this, PDI, PoseWatchDebugData.PoseWatch);
+					SkeletalDebugRendering::DrawBonesFromPoseWatch(*PoseWatchDebugData.PoseInfo.Get(), this, PDI, PoseWatchDebugData.PoseWatchPoseElement);
 				});
 			}
 		}
