Split the data used for the general transaction object away from the data used to diff a transacted object

This allows us more control over what goes into the diff object (eg, skipping all the internal state that Ar.IsTransacting adds), as well as allowing us to always use non-tagged property serialization (as we never need to load the diff object data back into a BP that has changed layout).

The downside of this approach is that each transacted object now needs to be serialized twice during the start and end of a transaction; once to get the object state for the transaction buffer, and once to get the object state to diff. The second set of diff data is discarded once the transaction has finalized, and any transaction snapshots still only need to serialize once (as snapshots only perform a diff).

API changes:
    * Split the object info out of FTransactionSerializedObject and into FTransactionSerializedObjectInfo so that it can be re-used. This data is only really needed for a diff, but the transaction buffer also need to persist a copy for notifications.
    * Split the diff only data out of FTransactionSerializedObject and into FTransactionDiffableObject.
    * Split the diff serialization logic out of FSerializedObjectDataWriter and into FDiffableObjectDataWriter.
    * Updated various diff APIs to work on FTransactionDiffableObject.
    * Updated FObjectRecord to store FTransactionDiffableObject instances for the initial and snapshot diffable object state.

#rnx

diff --git a/Engine/Source/Editor/UnrealEd/Classes/Editor/Transactor.h b/Engine/Source/Editor/UnrealEd/Classes/Editor/Transactor.h
--- a/Engine/Source/Editor/UnrealEd/Classes/Editor/Transactor.h
+++ b/Engine/Source/Editor/UnrealEd/Classes/Editor/Transactor.h
@@ -47,7 +47,7 @@ protected:
 		{
 			void SetObject(const UObject* InObject)
 			{
-				FTransactionSerializedObject::SetObject(InObject);
+				ObjectInfo.SetObject(InObject);
 				ObjectAnnotation = InObject->FindOrCreateTransactionAnnotation();
 			}
 
@@ -63,6 +63,9 @@ protected:
 				Exchange(ObjectAnnotation, Other.ObjectAnnotation);
 			}
 
+			/** Information about the object when it was serialized */
+			FTransactionSerializedObjectInfo ObjectInfo;
+
 			/** Annotation data for the object stored externally */
 			TSharedPtr<ITransactionObjectAnnotation> ObjectAnnotation;
 		};
@@ -95,16 +98,16 @@ protected:
 		bool				bRestored = false;
 		/** True if object has been finalized and generated diff data */
 		bool				bFinalized = false;
-		/** True if object has been snapshot before */
-		bool				bSnapshot = false;
 		/** True if record should serialize data as binary blob (more compact). False to use tagged serialization (more robust) */
 		bool				bWantsBinarySerialization = true;
 		/** The serialized object data */
 		FSerializedObject	SerializedObject;
 		/** The serialized object data that will be used when the transaction is flipped */
 		FSerializedObject	SerializedObjectFlip;
-		/** The serialized object data when it was last snapshot (if bSnapshot is true) */
-		FSerializedObject	SerializedObjectSnapshot;
+		/** The diffable object data (always null once finalized) */
+		TUniquePtr<FTransactionDiffableObject> DiffableObject;
+		/** The diffable object data when it was last snapshot (always null once finalized) */
+		TUniquePtr<FTransactionDiffableObject> DiffableObjectSnapshot;
 		/** Delta change information between the state of the object when the transaction started, and the state of the object when the transaction ended */
 		FTransactionObjectDeltaChange DeltaChange;
 
@@ -126,7 +129,6 @@ protected:
 		void Load( FTransaction* Owner );
 		void Finalize( FTransaction* Owner, TSharedPtr<ITransactionObjectAnnotation>& OutFinalizedObjectAnnotation );
 		void Snapshot( FTransaction* Owner, TArrayView<const FProperty*> Properties );
-		static void Diff( const FTransaction* Owner, const FSerializedObject& OldSerializedObect, const FSerializedObject& NewSerializedObject, FTransactionObjectDeltaChange& OutDeltaChange, const bool bFullDiff = true );
 
 		/** Used by GC to collect referenced objects. */
 		void AddReferencedObjects( FReferenceCollector& Collector );
@@ -185,24 +187,13 @@ protected:
 		public:
 			FWriter(
 				FSerializedObject& InSerializedObject,
-				bool bWantBinarySerialization,
-				TArrayView<const FProperty*> InPropertiesToSerialize = TArrayView<const FProperty*>()
+				bool bWantBinarySerialization
 				)
 				: UE::Transaction::FSerializedObjectDataWriter(InSerializedObject)
-				, PropertiesToSerialize(InPropertiesToSerialize)
 			{
 				this->SetWantBinaryPropertySerialization(bWantBinarySerialization);
 				this->SetIsTransacting(true);
 			}
-
-			virtual bool ShouldSkipProperty(const FProperty* InProperty) const override
-			{
-				return (PropertiesToSerialize.Num() > 0 && !PropertiesToSerialize.Contains(InProperty))
-					|| FArchiveUObject::ShouldSkipProperty(InProperty);
-			}
-
-		private:
-			TArrayView<const FProperty*> PropertiesToSerialize;
 		};
 	};
 
diff --git a/Engine/Source/Editor/UnrealEd/Private/EditorTransaction.cpp b/Engine/Source/Editor/UnrealEd/Private/EditorTransaction.cpp
--- a/Engine/Source/Editor/UnrealEd/Private/EditorTransaction.cpp
+++ b/Engine/Source/Editor/UnrealEd/Private/EditorTransaction.cpp
@@ -61,9 +61,19 @@ FTransaction::FObjectRecord::FObjectRecord(FTransaction* Owner, UObject* InObjec
 	}
 	else
 	{
-		SerializedObject.SetObject(CurrentObject);
-		FWriter Writer( SerializedObject, bWantsBinarySerialization );
-		SerializeContents( Writer, Oper );
+		{
+			SerializedObject.SetObject(CurrentObject);
+			FWriter Writer(SerializedObject, bWantsBinarySerialization);
+			SerializeContents(Writer, Oper);
+		}
+
+		if (!Array)
+		{
+			DiffableObject = MakeUnique<FTransactionDiffableObject>();
+			DiffableObject->SetObject(CurrentObject);
+			UE::Transaction::FDiffableObjectDataWriter Writer(*DiffableObject);
+			SerializeObject(Writer);
+		}
 	}
 }
 
@@ -253,34 +263,46 @@ void FTransaction::FObjectRecord::Finalize( FTransaction* Owner, TSharedPtr<ITra
 	{
 		bFinalized = true;
 
+		if (CustomChange)
+		{
+			// Cannot diff custom changes
+			return;
+		}
+
 		UObject* CurrentObject = Object.Get();
 		if (CurrentObject)
 		{
+			// Serialize the finalized object for the redo operation
+			SerializedObjectFlip.Reset();
+			{
+				SerializedObjectFlip.SetObject(CurrentObject);
+				OutFinalizedObjectAnnotation = SerializedObjectFlip.ObjectAnnotation;
+				FWriter Writer(SerializedObjectFlip, bWantsBinarySerialization);
+				SerializeObject(Writer);
+			}
+
 			// Serialize the object so we can diff it
-			FSerializedObject CurrentSerializedObject;
+			FTransactionDiffableObject CurrentDiffableObject;
 			{
-				CurrentSerializedObject.SetObject(CurrentObject);
-				OutFinalizedObjectAnnotation = CurrentSerializedObject.ObjectAnnotation;
-				FWriter Writer(CurrentSerializedObject, bWantsBinarySerialization);
+				CurrentDiffableObject.SetObject(CurrentObject);
+				UE::Transaction::FDiffableObjectDataWriter Writer(CurrentDiffableObject);
 				SerializeObject(Writer);
 			}
 			
 			// Diff against the object state when the transaction started
-			Diff(Owner, SerializedObject, CurrentSerializedObject, DeltaChange);
+			UE::Transaction::DiffUtil::GenerateObjectDiff(*DiffableObject, CurrentDiffableObject, DeltaChange);
 
 			// If we have a previous snapshot then we need to consider that part of the diff for the finalized object, as systems may 
 			// have been tracking delta-changes between snapshots and this finalization will need to account for those changes too
-			if (bSnapshot)
+			if (DiffableObjectSnapshot)
 			{
-				Diff(Owner, SerializedObjectSnapshot, CurrentSerializedObject, DeltaChange, /*bFullDiff*/false);
+				UE::Transaction::DiffUtil::GenerateObjectDiff(*DiffableObjectSnapshot, CurrentDiffableObject, DeltaChange, /*bFullDiff*/false);
 			}
-
-			SerializedObjectFlip.Swap(CurrentSerializedObject);
 		}
 
-		// Clear out any snapshot data now as we won't be getting any more snapshot requests once finalized
-		bSnapshot = false;
-		SerializedObjectSnapshot.Reset();
+		// Clear out any diff data now as we won't be getting any more diff requests once finalized
+		DiffableObject.Reset();
+		DiffableObjectSnapshot.Reset();
 	}
 }
 
@@ -298,49 +320,53 @@ void FTransaction::FObjectRecord::Snapshot( FTransaction* Owner, TArrayView<cons
 		return;
 	}
 
+	if (CustomChange)
+	{
+		// Cannot snapshot custom changes
+		return;
+	}
+
 	UObject* CurrentObject = Object.Get();
 	if (CurrentObject)
 	{
 		// Serialize the object so we can diff it
-		FSerializedObject CurrentSerializedObject;
+		FTransactionDiffableObject CurrentDiffableObject;
 		{
-			CurrentSerializedObject.SetObject(CurrentObject);
-			FWriter Writer(CurrentSerializedObject, bWantsBinarySerialization, Properties);
+			CurrentDiffableObject.SetObject(CurrentObject);
+			UE::Transaction::FDiffableObjectDataWriter Writer(CurrentDiffableObject, Properties);
 			// although it would be preferable to use SerializeScriptProperties, this cause a false diff between the first snapshot and the base object
 			// since they were serialized with different algo and we don't record enough context to make the comparison appropriately
 			SerializeObject(Writer);
 		}
 
 		// Diff against the correct serialized data depending on whether we already had a snapshot
-		const FSerializedObject& InitialSerializedObject = bSnapshot ? SerializedObjectSnapshot : SerializedObject;
+		const FTransactionDiffableObject& InitialDiffableObject = DiffableObjectSnapshot ? *DiffableObjectSnapshot : *DiffableObject;
 		FTransactionObjectDeltaChange SnapshotDeltaChange;
-		Diff(Owner, InitialSerializedObject, CurrentSerializedObject, SnapshotDeltaChange, /*bFullDiff*/false);
+		UE::Transaction::DiffUtil::GenerateObjectDiff(InitialDiffableObject, CurrentDiffableObject, SnapshotDeltaChange, /*bFullDiff*/false);
 
 		// Update the snapshot data for next time
-		bSnapshot = true;
-		SerializedObjectSnapshot.Swap(CurrentSerializedObject);
+		if (!DiffableObjectSnapshot)
+		{
+			DiffableObjectSnapshot = MakeUnique<FTransactionDiffableObject>();
+		}
+		DiffableObjectSnapshot->Swap(CurrentDiffableObject);
 
-		TSharedPtr<ITransactionObjectAnnotation> ChangedObjectTransactionAnnotation = SerializedObjectSnapshot.ObjectAnnotation;
+		TSharedPtr<ITransactionObjectAnnotation> ChangedObjectTransactionAnnotation = CurrentObject->FindOrCreateTransactionAnnotation();
 
 		// Notify any listeners of this change
 		if (SnapshotDeltaChange.HasChanged() || ChangedObjectTransactionAnnotation.IsValid())
 		{
 			CurrentObject->PostTransacted(FTransactionObjectEvent(Owner->GetId(), Owner->GetOperationId(), ETransactionObjectEventType::Snapshot, SnapshotDeltaChange, ChangedObjectTransactionAnnotation
-				, InitialSerializedObject.ObjectPackageName
-				, InitialSerializedObject.ObjectName
-				, InitialSerializedObject.ObjectPathName
-				, InitialSerializedObject.ObjectOuterPathName
-				, InitialSerializedObject.ObjectExternalPackageName
-				, InitialSerializedObject.ObjectClassPathName));
+				, InitialDiffableObject.ObjectInfo.ObjectPackageName
+				, InitialDiffableObject.ObjectInfo.ObjectName
+				, InitialDiffableObject.ObjectInfo.ObjectPathName
+				, InitialDiffableObject.ObjectInfo.ObjectOuterPathName
+				, InitialDiffableObject.ObjectInfo.ObjectExternalPackageName
+				, InitialDiffableObject.ObjectInfo.ObjectClassPathName));
 		}
 	}
 }
 
-void FTransaction::FObjectRecord::Diff( const FTransaction* Owner, const FSerializedObject& OldSerializedObject, const FSerializedObject& NewSerializedObject, FTransactionObjectDeltaChange& OutDeltaChange, const bool bFullDiff )
-{
-	UE::Transaction::DiffUtil::GenerateObjectDiff(OldSerializedObject, NewSerializedObject, OutDeltaChange, bFullDiff);
-}
-
 int32 FTransaction::GetRecordCount() const
 {
 	return Records.Num();
@@ -451,12 +477,16 @@ void FTransaction::FObjectRecord::AddReferencedObjects( FReferenceCollector& Col
 {
 	Object.AddStructReferencedObjects(Collector);
 
-	auto AddSerializedObjectReferences = [&Collector](FSerializedObject& InSerializedObject)
+	auto AddTransactionSerializedObjectReferences = [&Collector](FTransactionSerializedObject& InSerializedObject)
 	{
 		for (FPersistentObjectRef& ReferencedObject : InSerializedObject.ReferencedObjects)
 		{
 			ReferencedObject.AddStructReferencedObjects(Collector);
 		}
+	};
+	auto AddSerializedObjectReferences = [&Collector, &AddTransactionSerializedObjectReferences](FSerializedObject& InSerializedObject)
+	{
+		AddTransactionSerializedObjectReferences(InSerializedObject);
 
 		if (InSerializedObject.ObjectAnnotation.IsValid())
 		{
@@ -465,7 +495,14 @@ void FTransaction::FObjectRecord::AddReferencedObjects( FReferenceCollector& Col
 	};
 	AddSerializedObjectReferences(SerializedObject);
 	AddSerializedObjectReferences(SerializedObjectFlip);
-	AddSerializedObjectReferences(SerializedObjectSnapshot);
+	if (DiffableObject)
+	{
+		AddTransactionSerializedObjectReferences(DiffableObject->SerializedObject);
+	}
+	if (DiffableObjectSnapshot)
+	{
+		AddTransactionSerializedObjectReferences(DiffableObjectSnapshot->SerializedObject);
+	}
 
 	if (CustomChange.IsValid())
 	{
@@ -504,10 +541,6 @@ bool FTransaction::FObjectRecord::ContainsPieObject() const
 	{
 		return true;
 	}
-	if (SerializedObjectContainPieObjects(SerializedObjectSnapshot))
-	{
-		return true;
-	}
 	return false;
 }
 
@@ -736,12 +769,12 @@ void FTransaction::Apply()
 		{
 			const FObjectRecord::FSerializedObject& InitialSerializedObject = ChangedObjectRecord.SerializedObject;
 			ChangedObject->PostTransacted(FTransactionObjectEvent(Id, OperationId, ETransactionObjectEventType::UndoRedo, DeltaChange, ChangedObjectTransactionAnnotation
-				, InitialSerializedObject.ObjectPackageName
-				, InitialSerializedObject.ObjectName
-				, InitialSerializedObject.ObjectPathName
-				, InitialSerializedObject.ObjectOuterPathName
-				, InitialSerializedObject.ObjectExternalPackageName
-				, InitialSerializedObject.ObjectClassPathName));
+				, InitialSerializedObject.ObjectInfo.ObjectPackageName
+				, InitialSerializedObject.ObjectInfo.ObjectName
+				, InitialSerializedObject.ObjectInfo.ObjectPathName
+				, InitialSerializedObject.ObjectInfo.ObjectOuterPathName
+				, InitialSerializedObject.ObjectInfo.ObjectExternalPackageName
+				, InitialSerializedObject.ObjectInfo.ObjectClassPathName));
 		}
 	}
 
@@ -803,12 +836,12 @@ void FTransaction::Finalize()
 			
 			const FObjectRecord::FSerializedObject& InitialSerializedObject = ChangedObjectRecord.SerializedObject;
 			ChangedObject->PostTransacted(FTransactionObjectEvent(Id, OperationId, ETransactionObjectEventType::Finalized, DeltaChange, ChangedObjectTransactionAnnotation
-				, InitialSerializedObject.ObjectPackageName
-				, InitialSerializedObject.ObjectName
-				, InitialSerializedObject.ObjectPathName
-				, InitialSerializedObject.ObjectOuterPathName
-				, InitialSerializedObject.ObjectExternalPackageName
-				, InitialSerializedObject.ObjectClassPathName));
+				, InitialSerializedObject.ObjectInfo.ObjectPackageName
+				, InitialSerializedObject.ObjectInfo.ObjectName
+				, InitialSerializedObject.ObjectInfo.ObjectPathName
+				, InitialSerializedObject.ObjectInfo.ObjectOuterPathName
+				, InitialSerializedObject.ObjectInfo.ObjectExternalPackageName
+				, InitialSerializedObject.ObjectInfo.ObjectClassPathName));
 		}
 	}
 	ChangedObjects.Reset();
@@ -855,22 +888,17 @@ FTransactionDiff FTransaction::GenerateDiff() const
 			const FObjectRecord& ObjectRecord = Records[i];
 			if (UObject* TransactedObject = ObjectRecord.Object.Get())
 			{
-				// The last snapshot object is reset so we can only diff against the initial object for the moment.
-				FTransactionObjectDeltaChange RecordDeltaChange;
-				FObjectRecord::Diff(this, ObjectRecord.SerializedObject, ObjectRecord.SerializedObjectFlip, RecordDeltaChange);
-
-				// TODO: Add annotation
-				if (RecordDeltaChange.HasChanged())
+				if (ObjectRecord.DeltaChange.HasChanged() || ObjectRecord.SerializedObject.ObjectAnnotation)
 				{
 					// Since this transaction is not currently in an undo operation, generate a valid Guid.
 					FGuid Guid = FGuid::NewGuid();
-					TransactionDiff.DiffMap.Emplace(FName(*TransactedObject->GetPathName()), MakeShared<FTransactionObjectEvent>(this->GetId(), Guid, ETransactionObjectEventType::Finalized, RecordDeltaChange, ObjectRecord.SerializedObject.ObjectAnnotation
-						, ObjectRecord.SerializedObject.ObjectPackageName
-						, ObjectRecord.SerializedObject.ObjectName
-						, ObjectRecord.SerializedObject.ObjectPathName
-						, ObjectRecord.SerializedObject.ObjectOuterPathName
-						, ObjectRecord.SerializedObject.ObjectExternalPackageName
-						, ObjectRecord.SerializedObject.ObjectClassPathName));
+					TransactionDiff.DiffMap.Emplace(FName(*TransactedObject->GetPathName()), MakeShared<FTransactionObjectEvent>(this->GetId(), Guid, ETransactionObjectEventType::Finalized, ObjectRecord.DeltaChange, ObjectRecord.SerializedObject.ObjectAnnotation
+						, ObjectRecord.SerializedObject.ObjectInfo.ObjectPackageName
+						, ObjectRecord.SerializedObject.ObjectInfo.ObjectName
+						, ObjectRecord.SerializedObject.ObjectInfo.ObjectPathName
+						, ObjectRecord.SerializedObject.ObjectInfo.ObjectOuterPathName
+						, ObjectRecord.SerializedObject.ObjectInfo.ObjectExternalPackageName
+						, ObjectRecord.SerializedObject.ObjectInfo.ObjectClassPathName));
 				}
 			}
 		}
diff --git a/Engine/Source/Runtime/Engine/Private/TransactionCommon.cpp b/Engine/Source/Runtime/Engine/Private/TransactionCommon.cpp
--- a/Engine/Source/Runtime/Engine/Private/TransactionCommon.cpp
+++ b/Engine/Source/Runtime/Engine/Private/TransactionCommon.cpp
@@ -251,7 +251,76 @@ FSerializedObjectDataWriter::FSerializedObjectDataWriter(FTransactionSerializedO
 	}
 }
 
-FName FSerializedObjectDataWriter::GetTaggedDataKey() const
+void FSerializedObjectDataWriter::Serialize(void* SerData, int64 Num)
+{
+	if (Num)
+	{
+		int64 DataIndex = Offset;
+		SerializedObject.SerializedData.Write(SerData, Offset, Num);
+		Offset += Num;
+
+		OnDataSerialized(DataIndex, Num);
+	}
+}
+
+FArchive& FSerializedObjectDataWriter::operator<<(FName& N)
+{
+	int32 NameIndex = INDEX_NONE;
+	const int32* NameIndexPtr = NameMap.Find(N);
+	if (NameIndexPtr)
+	{
+		NameIndex = *NameIndexPtr;
+	}
+	else
+	{
+		NameIndex = SerializedObject.ReferencedNames.Add(N);
+		NameMap.Add(N, NameIndex);
+	}
+
+	OnNameSerialized(NameIndex);
+
+	return *this << NameIndex;
+}
+
+FArchive& FSerializedObjectDataWriter::operator<<(UObject*& Res)
+{
+	int32 ObjectIndex = INDEX_NONE;
+	const int32* ObjIndexPtr = ObjectMap.Find(Res);
+	if (ObjIndexPtr)
+	{
+		ObjectIndex = *ObjIndexPtr;
+	}
+	else if (Res)
+	{
+		ObjectIndex = SerializedObject.ReferencedObjects.Add(FTransactionPersistentObjectRef(Res));
+		ObjectMap.Add(Res, ObjectIndex);
+	}
+
+	OnObjectSerialized(ObjectIndex);
+
+	return *this << ObjectIndex;
+}
+
+
+FDiffableObjectDataWriter::FDiffableObjectDataWriter(FTransactionDiffableObject& InDiffableObject, TArrayView<const FProperty*> InPropertiesToSerialize)
+	: FSerializedObjectDataWriter(InDiffableObject.SerializedObject)
+	, DiffableObject(InDiffableObject)
+	, PropertiesToSerialize(InPropertiesToSerialize)
+{
+	SetWantBinaryPropertySerialization(true);
+
+	for (int32 ObjIndex = 0; ObjIndex < SerializedObject.ReferencedObjects.Num(); ++ObjIndex)
+	{
+		ObjectMap.Add(SerializedObject.ReferencedObjects[ObjIndex].Get(), ObjIndex);
+	}
+
+	for (int32 NameIndex = 0; NameIndex < SerializedObject.ReferencedNames.Num(); ++NameIndex)
+	{
+		NameMap.Add(SerializedObject.ReferencedNames[NameIndex], NameIndex);
+	}
+}
+
+FName FDiffableObjectDataWriter::GetTaggedDataKey() const
 {
 	FName TaggedDataKey;
 
@@ -292,100 +361,72 @@ FName FSerializedObjectDataWriter::GetTaggedDataKey() const
 	return TaggedDataKey;
 }
 
-bool FSerializedObjectDataWriter::DoesObjectMatchSerializedObject(const UObject* Obj) const
+bool FDiffableObjectDataWriter::DoesObjectMatchDiffableObject(const UObject* Obj) const
 {
 	FNameBuilder ObjPathName;
 	Obj->GetPathName(nullptr, ObjPathName);
-	return FName(ObjPathName) == SerializedObject.ObjectPathName;
+	return FName(ObjPathName) == DiffableObject.ObjectInfo.ObjectPathName;
+}
+
+bool FDiffableObjectDataWriter::ShouldSkipProperty(const FProperty* InProperty) const
+{
+	return (PropertiesToSerialize.Num() > 0 && !PropertiesToSerialize.Contains(InProperty))
+		|| InProperty->HasAnyPropertyFlags(CPF_Transient | CPF_NonTransactional | CPF_Deprecated)
+		|| FSerializedObjectDataWriter::ShouldSkipProperty(InProperty);
 }
 
-void FSerializedObjectDataWriter::MarkScriptSerializationStart(const UObject* Obj)
+void FDiffableObjectDataWriter::MarkScriptSerializationStart(const UObject* Obj)
 {
 	FArchiveUObject::MarkScriptSerializationStart(Obj);
 
-	if (DoesObjectMatchSerializedObject(Obj))
+	if (DoesObjectMatchDiffableObject(Obj))
 	{
 		bIsPerformingScriptSerialization = true;
 	}
 }
 
-void FSerializedObjectDataWriter::MarkScriptSerializationEnd(const UObject* Obj)
+void FDiffableObjectDataWriter::MarkScriptSerializationEnd(const UObject* Obj)
 {
 	FArchiveUObject::MarkScriptSerializationEnd(Obj);
 
-	if (DoesObjectMatchSerializedObject(Obj))
+	if (DoesObjectMatchDiffableObject(Obj))
 	{
 		bIsPerformingScriptSerialization = false;
 	}
 }
 
-void FSerializedObjectDataWriter::Serialize(void* SerData, int64 Num)
+void FDiffableObjectDataWriter::OnDataSerialized(int64 InOffset, int64 InNum)
 {
-	if (Num)
+	// Track this offset index in the serialized data
+	const FName SerializedTaggedDataKey = GetTaggedDataKey();
+	if (!SerializedTaggedDataKey.IsNone())
 	{
-		int64 DataIndex = Offset;
-		SerializedObject.SerializedData.Write(SerData, Offset, Num);
-		Offset += Num;
-
-		// Track this offset index in the serialized data
-		const FName SerializedTaggedDataKey = GetTaggedDataKey();
-		if (!SerializedTaggedDataKey.IsNone())
-		{
-			FTransactionSerializedTaggedData& SerializedTaggedData = SerializedObject.SerializedTaggedData.FindOrAdd(SerializedTaggedDataKey);
-			SerializedTaggedData.AppendSerializedData(DataIndex, Num);
-		}
+		FTransactionSerializedTaggedData& SerializedTaggedData = DiffableObject.SerializedTaggedData.FindOrAdd(SerializedTaggedDataKey);
+		SerializedTaggedData.AppendSerializedData(InOffset, InNum);
 	}
 }
 
-FArchive& FSerializedObjectDataWriter::operator<<(FName& N)
+void FDiffableObjectDataWriter::OnNameSerialized(int32 InNameIndex)
 {
-	int32 NameIndex = INDEX_NONE;
-	const int32* NameIndexPtr = NameMap.Find(N);
-	if (NameIndexPtr)
-	{
-		NameIndex = *NameIndexPtr;
-	}
-	else
-	{
-		NameIndex = SerializedObject.ReferencedNames.Add(N);
-		NameMap.Add(N, NameIndex);
-	}
-
 	// Track this name index in the serialized data
 	const FName SerializedTaggedDataKey = GetTaggedDataKey();
 	if (!SerializedTaggedDataKey.IsNone())
 	{
-		SerializedObject.SerializedNameIndices.Add(SerializedTaggedDataKey).Indices.AddUnique(NameIndex);
+		DiffableObject.SerializedNameIndices.Add(SerializedTaggedDataKey).Indices.AddUnique(InNameIndex);
 	}
-
-	return *this << NameIndex;
 }
 
-FArchive& FSerializedObjectDataWriter::operator<<(UObject*& Res)
+void FDiffableObjectDataWriter::OnObjectSerialized(int32 InObjectIndex)
 {
-	int32 ObjectIndex = INDEX_NONE;
-	const int32* ObjIndexPtr = ObjectMap.Find(Res);
-	if (ObjIndexPtr)
-	{
-		ObjectIndex = *ObjIndexPtr;
-	}
-	else if (Res)
-	{
-		ObjectIndex = SerializedObject.ReferencedObjects.Add(FTransactionPersistentObjectRef(Res));
-		ObjectMap.Add(Res, ObjectIndex);
-	}
-
 	// Track this object offset in the serialized data
 	const FName SerializedTaggedDataKey = GetTaggedDataKey();
 	if (!SerializedTaggedDataKey.IsNone())
 	{
-		SerializedObject.SerializedObjectIndices.Add(SerializedTaggedDataKey).Indices.AddUnique(ObjectIndex);
+		DiffableObject.SerializedObjectIndices.Add(SerializedTaggedDataKey).Indices.AddUnique(InObjectIndex);
 	}
-
-	return *this << ObjectIndex;
 }
 
-FName FSerializedObjectDataWriter::FCachedPropertyKey::SyncCache(const FArchiveSerializedPropertyChain* InPropertyChain)
+FName FDiffableObjectDataWriter::FCachedPropertyKey::SyncCache(const FArchiveSerializedPropertyChain* InPropertyChain)
 {
 	if (InPropertyChain)
 	{
@@ -409,23 +450,23 @@ FName FSerializedObjectDataWriter::FCachedPropertyKey::SyncCache(const FArchiveS
 namespace DiffUtil
 {
 
-void GenerateObjectDiff(const FTransactionSerializedObject& OldSerializedObject, const FTransactionSerializedObject& NewSerializedObject, FTransactionObjectDeltaChange& OutDeltaChange, const bool bFullDiff)
+void GenerateObjectDiff(const FTransactionDiffableObject& OldDiffableObject, const FTransactionDiffableObject& NewDiffableObject, FTransactionObjectDeltaChange& OutDeltaChange, const bool bFullDiff)
 {
-	auto IsTaggedDataBlockIdentical = [&OldSerializedObject, &NewSerializedObject, &OutDeltaChange](const FName TaggedDataKey, const FTransactionSerializedTaggedData& OldSerializedTaggedData, const FTransactionSerializedTaggedData& NewSerializedTaggedData) -> bool
+	auto IsTaggedDataBlockIdentical = [&OldDiffableObject, &NewDiffableObject, &OutDeltaChange](const FName TaggedDataKey, const FTransactionSerializedTaggedData& OldSerializedTaggedData, const FTransactionSerializedTaggedData& NewSerializedTaggedData) -> bool
 	{
 		// Binary compare the serialized data to see if something has changed for this property
 		bool bIsTaggedDataIdentical = OldSerializedTaggedData.DataSize == NewSerializedTaggedData.DataSize;
 		if (bIsTaggedDataIdentical && OldSerializedTaggedData.HasSerializedData())
 		{
-			bIsTaggedDataIdentical = FMemory::Memcmp(OldSerializedObject.SerializedData.GetPtr(OldSerializedTaggedData.DataOffset), NewSerializedObject.SerializedData.GetPtr(NewSerializedTaggedData.DataOffset), NewSerializedTaggedData.DataSize) == 0;
+			bIsTaggedDataIdentical = FMemory::Memcmp(OldDiffableObject.SerializedObject.SerializedData.GetPtr(OldSerializedTaggedData.DataOffset), NewDiffableObject.SerializedObject.SerializedData.GetPtr(NewSerializedTaggedData.DataOffset), NewSerializedTaggedData.DataSize) == 0;
 		}
 		if (bIsTaggedDataIdentical)
 		{
-			bIsTaggedDataIdentical = Internal::AreObjectPointersIdentical(OldSerializedObject, NewSerializedObject, TaggedDataKey);
+			bIsTaggedDataIdentical = Internal::AreObjectPointersIdentical(OldDiffableObject, NewDiffableObject, TaggedDataKey);
 		}
 		if (bIsTaggedDataIdentical)
 		{
-			bIsTaggedDataIdentical = Internal::AreNamesIdentical(OldSerializedObject, NewSerializedObject, TaggedDataKey);
+			bIsTaggedDataIdentical = Internal::AreNamesIdentical(OldDiffableObject, NewDiffableObject, TaggedDataKey);
 		}
 		return bIsTaggedDataIdentical;
 	};
@@ -448,20 +489,20 @@ void GenerateObjectDiff(const FTransactionSerializedObject& OldSerializedObject,
 
 	if (bFullDiff)
 	{
-		OutDeltaChange.bHasNameChange |= OldSerializedObject.ObjectName != NewSerializedObject.ObjectName;
-		OutDeltaChange.bHasOuterChange |= OldSerializedObject.ObjectOuterPathName != NewSerializedObject.ObjectOuterPathName;
-		OutDeltaChange.bHasExternalPackageChange |= OldSerializedObject.ObjectExternalPackageName != NewSerializedObject.ObjectExternalPackageName;
-		OutDeltaChange.bHasPendingKillChange |= OldSerializedObject.bIsPendingKill != NewSerializedObject.bIsPendingKill;
+		OutDeltaChange.bHasNameChange |= OldDiffableObject.ObjectInfo.ObjectName != NewDiffableObject.ObjectInfo.ObjectName;
+		OutDeltaChange.bHasOuterChange |= OldDiffableObject.ObjectInfo.ObjectOuterPathName != NewDiffableObject.ObjectInfo.ObjectOuterPathName;
+		OutDeltaChange.bHasExternalPackageChange |= OldDiffableObject.ObjectInfo.ObjectExternalPackageName != NewDiffableObject.ObjectInfo.ObjectExternalPackageName;
+		OutDeltaChange.bHasPendingKillChange |= OldDiffableObject.ObjectInfo.bIsPendingKill != NewDiffableObject.ObjectInfo.bIsPendingKill;
 	}
 
-	for (const TPair<FName, FTransactionSerializedTaggedData>& NewNamePropertyPair : NewSerializedObject.SerializedTaggedData)
+	for (const TPair<FName, FTransactionSerializedTaggedData>& NewNamePropertyPair : NewDiffableObject.SerializedTaggedData)
 	{
 		if (!ShouldCompareTaggedData(NewNamePropertyPair.Key))
 		{
 			continue;
 		}
 
-		const FTransactionSerializedTaggedData* OldSerializedTaggedData = OldSerializedObject.SerializedTaggedData.Find(NewNamePropertyPair.Key);
+		const FTransactionSerializedTaggedData* OldSerializedTaggedData = OldDiffableObject.SerializedTaggedData.Find(NewNamePropertyPair.Key);
 		if (ShouldCompareAsNonPropertyData(NewNamePropertyPair.Key))
 		{
 			if (bFullDiff && !OutDeltaChange.bHasNonPropertyChanges && (!OldSerializedTaggedData || !IsTaggedDataBlockIdentical(NewNamePropertyPair.Key, *OldSerializedTaggedData, NewNamePropertyPair.Value)))
@@ -485,9 +526,9 @@ void GenerateObjectDiff(const FTransactionSerializedObject& OldSerializedObject,
 
 	if (bFullDiff)
 	{
-		for (const TPair<FName, FTransactionSerializedTaggedData>& OldNamePropertyPair : OldSerializedObject.SerializedTaggedData)
+		for (const TPair<FName, FTransactionSerializedTaggedData>& OldNamePropertyPair : OldDiffableObject.SerializedTaggedData)
 		{
-			if (!ShouldCompareTaggedData(OldNamePropertyPair.Key) || NewSerializedObject.SerializedTaggedData.Contains(OldNamePropertyPair.Key))
+			if (!ShouldCompareTaggedData(OldNamePropertyPair.Key) || NewDiffableObject.SerializedTaggedData.Contains(OldNamePropertyPair.Key))
 			{
 				continue;
 			}
@@ -508,11 +549,11 @@ void GenerateObjectDiff(const FTransactionSerializedObject& OldSerializedObject,
 namespace Internal
 {
 
-bool AreObjectPointersIdentical(const FTransactionSerializedObject& OldSerializedObject, const FTransactionSerializedObject& NewSerializedObject, const FName TaggedDataKey)
+bool AreObjectPointersIdentical(const FTransactionDiffableObject& OldDiffableObject, const FTransactionDiffableObject& NewDiffableObject, const FName TaggedDataKey)
 {
-	auto GetIndicesArray = [&TaggedDataKey](const FTransactionSerializedObject& SerializedObject) -> const TArray<int32>&
+	auto GetIndicesArray = [&TaggedDataKey](const FTransactionDiffableObject& DiffableObject) -> const TArray<int32>&
 	{
-		if (const FTransactionSerializedIndices* SerializedObjectIndices = SerializedObject.SerializedObjectIndices.Find(TaggedDataKey))
+		if (const FTransactionSerializedIndices* SerializedObjectIndices = DiffableObject.SerializedObjectIndices.Find(TaggedDataKey))
 		{
 			return SerializedObjectIndices->Indices;
 		}
@@ -521,24 +562,24 @@ bool AreObjectPointersIdentical(const FTransactionSerializedObject& OldSerialize
 		return EmptyIndicesArray;
 	};
 
-	const TArray<int32>& OldSerializedObjectIndices = GetIndicesArray(OldSerializedObject);
-	const TArray<int32>& NewSerializedObjectIndices = GetIndicesArray(NewSerializedObject);
+	const TArray<int32>& OldSerializedObjectIndices = GetIndicesArray(OldDiffableObject);
+	const TArray<int32>& NewSerializedObjectIndices = GetIndicesArray(NewDiffableObject);
 
 	bool bAreObjectPointersIdentical = OldSerializedObjectIndices.Num() == NewSerializedObjectIndices.Num();
 	for (int32 ObjIndex = 0; ObjIndex < OldSerializedObjectIndices.Num() && bAreObjectPointersIdentical; ++ObjIndex)
 	{
-		const FTransactionPersistentObjectRef& OldObjectRef = OldSerializedObject.ReferencedObjects.IsValidIndex(OldSerializedObjectIndices[ObjIndex]) ? OldSerializedObject.ReferencedObjects[OldSerializedObjectIndices[ObjIndex]] : FTransactionPersistentObjectRef();
-		const FTransactionPersistentObjectRef& NewObjectRef = NewSerializedObject.ReferencedObjects.IsValidIndex(NewSerializedObjectIndices[ObjIndex]) ? NewSerializedObject.ReferencedObjects[NewSerializedObjectIndices[ObjIndex]] : FTransactionPersistentObjectRef();
+		const FTransactionPersistentObjectRef& OldObjectRef = OldDiffableObject.SerializedObject.ReferencedObjects.IsValidIndex(OldSerializedObjectIndices[ObjIndex]) ? OldDiffableObject.SerializedObject.ReferencedObjects[OldSerializedObjectIndices[ObjIndex]] : FTransactionPersistentObjectRef();
+		const FTransactionPersistentObjectRef& NewObjectRef = NewDiffableObject.SerializedObject.ReferencedObjects.IsValidIndex(NewSerializedObjectIndices[ObjIndex]) ? NewDiffableObject.SerializedObject.ReferencedObjects[NewSerializedObjectIndices[ObjIndex]] : FTransactionPersistentObjectRef();
 		bAreObjectPointersIdentical = OldObjectRef == NewObjectRef;
 	}
 	return bAreObjectPointersIdentical;
 }
 
-bool AreNamesIdentical(const FTransactionSerializedObject& OldSerializedObject, const FTransactionSerializedObject& NewSerializedObject, const FName TaggedDataKey)
+bool AreNamesIdentical(const FTransactionDiffableObject& OldDiffableObject, const FTransactionDiffableObject& NewDiffableObject, const FName TaggedDataKey)
 {
-	auto GetIndicesArray = [&TaggedDataKey](const FTransactionSerializedObject& SerializedObject) -> const TArray<int32>&
+	auto GetIndicesArray = [&TaggedDataKey](const FTransactionDiffableObject& DiffableObject) -> const TArray<int32>&
 	{
-		if (const FTransactionSerializedIndices* SerializedNameIndices = SerializedObject.SerializedNameIndices.Find(TaggedDataKey))
+		if (const FTransactionSerializedIndices* SerializedNameIndices = DiffableObject.SerializedNameIndices.Find(TaggedDataKey))
 		{
 			return SerializedNameIndices->Indices;
 		}
@@ -547,14 +588,14 @@ bool AreNamesIdentical(const FTransactionSerializedObject& OldSerializedObject,
 		return EmptyIndicesArray;
 	};
 
-	const TArray<int32>& OldSerializedNameIndices = GetIndicesArray(OldSerializedObject);
-	const TArray<int32>& NewSerializedNameIndices = GetIndicesArray(NewSerializedObject);
+	const TArray<int32>& OldSerializedNameIndices = GetIndicesArray(OldDiffableObject);
+	const TArray<int32>& NewSerializedNameIndices = GetIndicesArray(NewDiffableObject);
 
 	bool bAreNamesIdentical = OldSerializedNameIndices.Num() == NewSerializedNameIndices.Num();
 	for (int32 ObjIndex = 0; ObjIndex < OldSerializedNameIndices.Num() && bAreNamesIdentical; ++ObjIndex)
 	{
-		const FName& OldName = OldSerializedObject.ReferencedNames.IsValidIndex(OldSerializedNameIndices[ObjIndex]) ? OldSerializedObject.ReferencedNames[OldSerializedNameIndices[ObjIndex]] : FName();
-		const FName& NewName = NewSerializedObject.ReferencedNames.IsValidIndex(NewSerializedNameIndices[ObjIndex]) ? NewSerializedObject.ReferencedNames[NewSerializedNameIndices[ObjIndex]] : FName();
+		const FName& OldName = OldDiffableObject.SerializedObject.ReferencedNames.IsValidIndex(OldSerializedNameIndices[ObjIndex]) ? OldDiffableObject.SerializedObject.ReferencedNames[OldSerializedNameIndices[ObjIndex]] : FName();
+		const FName& NewName = NewDiffableObject.SerializedObject.ReferencedNames.IsValidIndex(NewSerializedNameIndices[ObjIndex]) ? NewDiffableObject.SerializedObject.ReferencedNames[NewSerializedNameIndices[ObjIndex]] : FName();
 		bAreNamesIdentical = OldName == NewName;
 	}
 	return bAreNamesIdentical;
diff --git a/Engine/Source/Runtime/Engine/Public/TransactionCommon.h b/Engine/Source/Runtime/Engine/Public/TransactionCommon.h
--- a/Engine/Source/Runtime/Engine/Public/TransactionCommon.h
+++ b/Engine/Source/Runtime/Engine/Public/TransactionCommon.h
@@ -226,7 +226,7 @@
 };
 
 USTRUCT()
-struct FTransactionSerializedObject
+struct FTransactionSerializedObjectInfo
 {
 	GENERATED_BODY()
 
@@ -251,15 +251,9 @@
 		ObjectExternalPackageName = FName();
 		ObjectClassPathName = FName();
 		bIsPendingKill = false;
-		SerializedData.Reset();
-		ReferencedObjects.Reset();
-		ReferencedNames.Reset();
-		SerializedTaggedData.Reset();
-		SerializedObjectIndices.Reset();
-		SerializedNameIndices.Reset();
 	}
 
-	void Swap(FTransactionSerializedObject& Other)
+	void Swap(FTransactionSerializedObjectInfo& Other)
 	{
 		Exchange(ObjectPackageName, Other.ObjectPackageName);
 		Exchange(ObjectName, Other.ObjectName);
@@ -268,41 +262,56 @@
 		Exchange(ObjectExternalPackageName, Other.ObjectExternalPackageName);
 		Exchange(ObjectClassPathName, Other.ObjectClassPathName);
 		Exchange(bIsPendingKill, Other.bIsPendingKill);
-		Exchange(SerializedData, Other.SerializedData);
-		Exchange(ReferencedObjects, Other.ReferencedObjects);
-		Exchange(ReferencedNames, Other.ReferencedNames);
-		Exchange(SerializedTaggedData, Other.SerializedTaggedData);
-		Exchange(SerializedObjectIndices, Other.SerializedObjectIndices);
-		Exchange(SerializedNameIndices, Other.SerializedNameIndices);
 	}
 
 	/** The package name of the object when it was serialized, can be dictated either by outer chain or external package */
 	UPROPERTY()
 	FName ObjectPackageName;
-	
+
 	/** The name of the object when it was serialized */
 	UPROPERTY()
 	FName ObjectName;
-	
+
 	/** The path name of the object when it was serialized */
 	UPROPERTY()
 	FName ObjectPathName;
-	
+
 	/** The outer path name of the object when it was serialized */
 	UPROPERTY()
 	FName ObjectOuterPathName;
-	
+
 	/** The external package name of the object when it was serialized, if any */
 	UPROPERTY()
 	FName ObjectExternalPackageName;
-	
+
 	/** The path name of the object's class. */
 	UPROPERTY()
 	FName ObjectClassPathName;
-	
+
 	/** The pending kill state of the object when it was serialized */
 	UPROPERTY()
 	bool bIsPendingKill = false;
+};
 
+USTRUCT()
+struct FTransactionSerializedObject
+{
+	GENERATED_BODY()
+
+public:
+	void Reset()
+	{
+		SerializedData.Reset();
+		ReferencedObjects.Reset();
+		ReferencedNames.Reset();
+	}
+
+	void Swap(FTransactionSerializedObject& Other)
+	{
+		Exchange(SerializedData, Other.SerializedData);
+		Exchange(ReferencedObjects, Other.ReferencedObjects);
+		Exchange(ReferencedNames, Other.ReferencedNames);
+	}
+
 	/** The data stream used to serialize/deserialize record */
 	UPROPERTY()
@@ -315,15 +324,53 @@
 	/** FNames referenced in the object record */
 	UPROPERTY()
 	TArray<FName> ReferencedNames;
-	
+};
+
+USTRUCT()
+struct FTransactionDiffableObject
+{
+	GENERATED_BODY()
+
+public:
+	void SetObject(const UObject* InObject)
+	{
+		ObjectInfo.SetObject(InObject);
+	}
+
+	void Reset()
+	{
+		ObjectInfo.Reset();
+		SerializedObject.Reset();
+		SerializedTaggedData.Reset();
+		SerializedObjectIndices.Reset();
+		SerializedNameIndices.Reset();
+	}
+
+	void Swap(FTransactionDiffableObject& Other)
+	{
+		ObjectInfo.Swap(Other.ObjectInfo);
+		SerializedObject.Swap(Other.SerializedObject);
+		Exchange(SerializedTaggedData, Other.SerializedTaggedData);
+		Exchange(SerializedObjectIndices, Other.SerializedObjectIndices);
+		Exchange(SerializedNameIndices, Other.SerializedNameIndices);
+	}
+
+	/** Information about the object when it was serialized */
+	UPROPERTY()
+	FTransactionSerializedObjectInfo ObjectInfo;
+
+	/** The serialized object data */
+	UPROPERTY()
+	FTransactionSerializedObject SerializedObject;
+
 	/** Information about tagged data (mainly properties) that were serialized within this object */
 	UPROPERTY()
 	TMap<FName, FTransactionSerializedTaggedData> SerializedTaggedData;
-	
+
 	/** Information about the object pointer offsets that were serialized within this object (this maps the tagged data key (or None if there was no tagged data key) to the ReferencedObjects indices of the tagged data) */
 	UPROPERTY()
 	TMap<FName, FTransactionSerializedIndices> SerializedObjectIndices;
-	
+
 	/** Information about the name offsets that were serialized within this object (this maps the tagged data name to the ReferencedNames index of the tagged data) */
 	UPROPERTY()
 	TMap<FName, FTransactionSerializedIndices> SerializedNameIndices;
@@ -335,7 +382,7 @@
 ENGINE_API extern const FName TaggedDataKey_UnknownData;
 ENGINE_API extern const FName TaggedDataKey_ScriptData;
 
-/** Transfers data from an array. */
+/** Core archive to read a transaction object from the buffer. */
 class ENGINE_API FSerializedObjectDataReader : public FArchiveUObject
 {
 public:
@@ -356,9 +403,7 @@
 	int64 Offset = 0;
 };
 
-/**
- * Transfers data to an array.
- */
+/** Core archive to write a transaction object to the buffer. */
 class ENGINE_API FSerializedObjectDataWriter : public FArchiveUObject
 {
 public:
@@ -369,36 +414,57 @@
 	virtual int64 TotalSize() override { return SerializedObject.SerializedData.Num(); }
 
 protected:
-	struct FCachedPropertyKey
-	{
-	public:
-		FName SyncCache(const FArchiveSerializedPropertyChain* InPropertyChain);
-
-	private:
-		FName CachedKey;
-		uint32 LastUpdateCount = 0;
-	};
-
-	FName GetTaggedDataKey() const;
-
-	bool DoesObjectMatchSerializedObject(const UObject* Obj) const;
-
-	virtual void MarkScriptSerializationStart(const UObject* Obj) override;
-	virtual void MarkScriptSerializationEnd(const UObject* Obj) override;
-
 	virtual void Serialize(void* SerData, int64 Num) override;
 
 	using FArchiveUObject::operator<<;
 	virtual FArchive& operator<<(class FName& N) override;
 	virtual FArchive& operator<<(class UObject*& Res) override;
 
+	virtual void OnDataSerialized(int64 InOffset, int64 InNum) {}
+	virtual void OnNameSerialized(int32 InNameIndex) {}
+	virtual void OnObjectSerialized(int32 InObjectIndex) {}
+
 	FTransactionSerializedObject& SerializedObject;
 	int64 Offset = 0;
 
 	TMap<UObject*, int32> ObjectMap;
 	TMap<FName, int32> NameMap;
+};
 
+/** Core archive to write a diffable object to the buffer. */
+class ENGINE_API FDiffableObjectDataWriter : public FSerializedObjectDataWriter
+{
+public:
+	FDiffableObjectDataWriter(FTransactionDiffableObject& InDiffableObject, TArrayView<const FProperty*> InPropertiesToSerialize = TArrayView<const FProperty*>());
+
+protected:
+	FName GetTaggedDataKey() const;
+
+	bool DoesObjectMatchDiffableObject(const UObject* Obj) const;
+
+	virtual bool ShouldSkipProperty(const FProperty* InProperty) const override;
+
+	virtual void MarkScriptSerializationStart(const UObject* Obj) override;
+	virtual void MarkScriptSerializationEnd(const UObject* Obj) override;
+
+	virtual void OnDataSerialized(int64 InOffset, int64 InNum) override;
+	virtual void OnNameSerialized(int32 InNameIndex) override;
+	virtual void OnObjectSerialized(int32 InObjectIndex) override;
+
 private:
+	struct FCachedPropertyKey
+	{
+	public:
+		FName SyncCache(const FArchiveSerializedPropertyChain* InPropertyChain);
+
+	private:
+		FName CachedKey;
+		uint32 LastUpdateCount = 0;
+	};
+
+	FTransactionDiffableObject& DiffableObject;
+	TArrayView<const FProperty*> PropertiesToSerialize;
+
 	bool bIsPerformingScriptSerialization = false;
 
 	mutable bool bWasUsingTaggedDataKey_UnknownData = false;
@@ -411,13 +477,13 @@
 namespace DiffUtil
 {
 
-ENGINE_API void GenerateObjectDiff(const FTransactionSerializedObject& OldSerializedObject, const FTransactionSerializedObject& NewSerializedObject, FTransactionObjectDeltaChange& OutDeltaChange, const bool bFullDiff = true);
+ENGINE_API void GenerateObjectDiff(const FTransactionDiffableObject& OldDiffableObject, const FTransactionDiffableObject& NewDiffableObject, FTransactionObjectDeltaChange& OutDeltaChange, const bool bFullDiff = true);
 
 namespace Internal
 {
 
-ENGINE_API bool AreObjectPointersIdentical(const FTransactionSerializedObject& OldSerializedObject, const FTransactionSerializedObject& NewSerializedObject, const FName TaggedDataKey);
-ENGINE_API bool AreNamesIdentical(const FTransactionSerializedObject& OldSerializedObject, const FTransactionSerializedObject& NewSerializedObject, const FName TaggedDataKey);
+ENGINE_API bool AreObjectPointersIdentical(const FTransactionDiffableObject& OldDiffableObject, const FTransactionDiffableObject& NewDiffableObject, const FName TaggedDataKey);
+ENGINE_API bool AreNamesIdentical(const FTransactionDiffableObject& OldDiffableObject, const FTransactionDiffableObject& NewDiffableObject, const FName TaggedDataKey);
 
 } // namespace Internal
 
