Changing the code to only bind once to the EndPlay event of subactors of an ALOD actor because there are situations where we can unintentionally bind twice to them

diff --git a/Engine/Source/Runtime/Engine/Private/LODActor.cpp b/Engine/Source/Runtime/Engine/Private/LODActor.cpp
--- a/Engine/Source/Runtime/Engine/Private/LODActor.cpp
+++ b/Engine/Source/Runtime/Engine/Private/LODActor.cpp
@@ -534,8 +534,8 @@ void ALODActor::PostRegisterAllComponents()
 		{
 			if (AActor* Actor = ActorPtr.Get())
 			{
-				Actor->OnEndPlay.AddDynamic(this, &ALODActor::OnSubActorEndPlay);
-			}		
+				Actor->OnEndPlay.AddUniqueDynamic(this, &ALODActor::OnSubActorEndPlay);
+			}
 		}
 	}
 
