Optimized BuildingFoundationPostSignificanceUnloadsFirst. Saves half the cost (25us on high end pc).

diff --git a/Engine/Source/Runtime/Engine/Classes/Engine/LevelStreaming.h b/Engine/Source/Runtime/Engine/Classes/Engine/LevelStreaming.h
--- a/Engine/Source/Runtime/Engine/Classes/Engine/LevelStreaming.h
+++ b/Engine/Source/Runtime/Engine/Classes/Engine/LevelStreaming.h
@@ -460,7 +460,7 @@ public:
 
 	/** Returns whether streaming level is visible */
 	UFUNCTION(BlueprintCallable, Category="Game")
-	bool IsLevelVisible() const;
+	bool IsLevelVisible() const { return LoadedLevel && LoadedLevel->bIsVisible; }
 
 	/** Returns whether streaming level is loaded */
 	UFUNCTION(BlueprintCallable, Category="Game")
diff --git a/Engine/Source/Runtime/Engine/Private/LevelStreaming.cpp b/Engine/Source/Runtime/Engine/Private/LevelStreaming.cpp
--- a/Engine/Source/Runtime/Engine/Private/LevelStreaming.cpp
+++ b/Engine/Source/Runtime/Engine/Private/LevelStreaming.cpp
@@ -1539,11 +1539,6 @@ void ULevelStreaming::AsyncLevelLoadComplete(const FName& InPackageName, UPackag
 	TRACE_BOOKMARK(TEXT("RequestLevelComplete - %s"), *InPackageName.ToString());
 }
 
-bool ULevelStreaming::IsLevelVisible() const
-{
-	return LoadedLevel != NULL && LoadedLevel->bIsVisible;
-}
-
 bool ULevelStreaming::IsStreamingStatePending() const
 {
 	UWorld* PersistentWorld = GetWorld();
