Fixed non-unity compile error from missing include.

diff --git a/Engine/Source/Runtime/Engine/Classes/Engine/LevelStreaming.h b/Engine/Source/Runtime/Engine/Classes/Engine/LevelStreaming.h
--- a/Engine/Source/Runtime/Engine/Classes/Engine/LevelStreaming.h
+++ b/Engine/Source/Runtime/Engine/Classes/Engine/LevelStreaming.h
@@ -10,6 +10,7 @@
 #include "Engine/LatentActionManager.h"
 #include "LatentActions.h"
 #include "ProfilingDebugging/ProfilingHelpers.h"
+#include "Engine/Level.h"
 
 #if WITH_EDITOR
 #include "Folder.h"
