Handling torn off actors immediately in client replays to avoid race condition with destroy

diff --git a/Engine/Source/Runtime/Engine/Classes/Engine/DemoNetDriver.h b/Engine/Source/Runtime/Engine/Classes/Engine/DemoNetDriver.h
--- a/Engine/Source/Runtime/Engine/Classes/Engine/DemoNetDriver.h
+++ b/Engine/Source/Runtime/Engine/Classes/Engine/DemoNetDriver.h
@@ -341,6 +341,7 @@ public:
 	virtual void NotifyActorChannelOpen(UActorChannel* Channel, AActor* Actor) override;
 	virtual void NotifyActorChannelCleanedUp(UActorChannel* Channel, EChannelCloseReason CloseReason) override;
 	virtual void NotifyActorClientDormancyChanged(AActor* Actor, ENetDormancy OldDormancyState) override;
+	virtual void NotifyActorTornOff(AActor* Actor) override;
 
 	virtual void ProcessLocalServerPackets() override {}
 	virtual void ProcessLocalClientPackets() override {}
diff --git a/Engine/Source/Runtime/Engine/Classes/Engine/NetDriver.h b/Engine/Source/Runtime/Engine/Classes/Engine/NetDriver.h
--- a/Engine/Source/Runtime/Engine/Classes/Engine/NetDriver.h
+++ b/Engine/Source/Runtime/Engine/Classes/Engine/NetDriver.h
@@ -1733,6 +1733,9 @@ public:
 	/** Called on clients when an actor channel is closed because it went dormant. */
 	ENGINE_API virtual void ClientSetActorDormant(AActor* Actor);
 
+	/** Called on clients when an actor is torn off. */
+	ENGINE_API virtual void ClientSetActorTornOff(AActor* Actor);
+
 	/**
 	 * Returns the current delinquency analytics and resets them.
 	 * This would be similar to calls to Get and Reset separately, except that the caller
diff --git a/Engine/Source/Runtime/Engine/Private/DataChannel.cpp b/Engine/Source/Runtime/Engine/Private/DataChannel.cpp
--- a/Engine/Source/Runtime/Engine/Private/DataChannel.cpp
+++ b/Engine/Source/Runtime/Engine/Private/DataChannel.cpp
@@ -2087,18 +2087,7 @@ void UControlChannel::ReceiveDestructionInfo(FInBunch& Bunch)
 
 			if (TheActor->GetTearOff() && !Connection->Driver->ShouldClientDestroyTearOffActors())
 			{
-				if (!bTornOff)
-				{
-					TheActor->SetRole(ROLE_Authority);
-					TheActor->SetReplicates(false);
-					bTornOff = true;
-					if (TheActor->GetWorld() != nullptr && !IsEngineExitRequested())
-					{
-						TheActor->TornOff();
-					}
-
-					Connection->Driver->NotifyActorTornOff(TheActor);
-				}
+				Connection->Driver->ClientSetActorTornOff(TheActor);
 			}
 			else if (Dormant && (CloseReason == EChannelCloseReason::Dormancy) && !TheActor->GetTearOff())
 			{
@@ -2382,15 +2371,8 @@ bool UActorChannel::CleanUp(const bool bForDestroy, EChannelCloseReason CloseRea
 			{
 				if (!bTornOff)
 				{
-					Actor->SetRole(ROLE_Authority);
-					Actor->SetReplicates(false);
 					bTornOff = true;
-					if (Actor->GetWorld() != NULL && !IsEngineExitRequested())
-					{
-						Actor->TornOff();
-					}
-
-					Connection->Driver->NotifyActorTornOff(Actor);
+					Connection->Driver->ClientSetActorTornOff(Actor);
 				}
 			}
 			else if (Dormant && (CloseReason == EChannelCloseReason::Dormancy) && !Actor->GetTearOff())	
diff --git a/Engine/Source/Runtime/Engine/Private/DemoNetDriver.cpp b/Engine/Source/Runtime/Engine/Private/DemoNetDriver.cpp
--- a/Engine/Source/Runtime/Engine/Private/DemoNetDriver.cpp
+++ b/Engine/Source/Runtime/Engine/Private/DemoNetDriver.cpp
@@ -4986,7 +4986,7 @@ void UDemoNetDriver::OnActorPreDestroy(AActor* DestroyedActor)
 
 	if (IsRecording())
 	{
-		ReplayHelper.OnActorPreDestroy(ClientConnections[0], DestroyedActor);
+		ReplayHelper.RemoveActorFromCheckpoint(ClientConnections[0], DestroyedActor);
 	}
 }
 
@@ -5051,6 +5051,20 @@ void UDemoNetDriver::NotifyActorDestroyed(AActor* Actor, bool IsSeamlessTravel)
 	Super::NotifyActorDestroyed(Actor, IsSeamlessTravel);
 }
 
+void UDemoNetDriver::NotifyActorTornOff(AActor* Actor)
+{
+	if (IsRecording())
+	{
+		// Replicate one last time to the replay stream
+		ReplayHelper.ReplicateActor(Actor, ClientConnections[0], true);
+
+		// Handle being pending in a checkpoint save
+		ReplayHelper.RemoveActorFromCheckpoint(ClientConnections[0], Actor);
+	}
+
+	Super::NotifyActorTornOff(Actor);
+}
+
 void UDemoNetDriver::CleanupOutstandingRewindActors()
 {
 	if (World)
@@ -5327,7 +5341,7 @@ bool UDemoNetDriver::ShouldReplicateFunction(AActor* Actor, UFunction* Function)
 bool UDemoNetDriver::ShouldReplicateActor(AActor* Actor) const
 {
 	// replicate actors that share the demo net driver name, or actors belonging to the game net driver
-	return (Actor && Actor->GetIsReplicated()) && (Super::ShouldReplicateActor(Actor) || (Actor->GetNetDriverName() == NAME_GameNetDriver));
+	return (Actor && (Actor->GetIsReplicated() || Actor->GetTearOff())) && (Super::ShouldReplicateActor(Actor) || (Actor->GetNetDriverName() == NAME_GameNetDriver));
 }
 
 /*
diff --git a/Engine/Source/Runtime/Engine/Private/NetDriver.cpp b/Engine/Source/Runtime/Engine/Private/NetDriver.cpp
--- a/Engine/Source/Runtime/Engine/Private/NetDriver.cpp
+++ b/Engine/Source/Runtime/Engine/Private/NetDriver.cpp
@@ -6500,6 +6500,46 @@ void UNetDriver::NotifyActorChannelCleanedUp(UActorChannel* Channel, EChannelClo
 	}
 }
 
+void UNetDriver::ClientSetActorTornOff(AActor* Actor)
+{
+	check(Actor);
+
+	if (IsServer())
+	{
+		UE_LOG(LogNet, Error, TEXT("ClientSetActorTornOff should only be called on clients."));
+		return;
+	}
+
+	if (Actor->GetIsReplicated() && (Actor->GetRemoteRole() == ROLE_Authority))
+	{
+		Actor->SetRole(ROLE_Authority);
+		Actor->SetReplicates(false);
+
+		if (Actor->GetWorld() != nullptr && !IsEngineExitRequested())
+		{
+			Actor->TornOff();
+		}
+
+		if (World)
+		{
+			if (FWorldContext* const Context = GEngine->GetWorldContextFromWorld(World))
+			{
+				for (FNamedNetDriver& Driver : Context->ActiveNetDrivers)
+				{
+					if (Driver.NetDriver != nullptr && Driver.NetDriver->ShouldReplicateActor(Actor))
+					{
+						Driver.NetDriver->NotifyActorTornOff(Actor);
+					}
+				}
+			}
+		}
+	}
+	else
+	{
+		UE_LOG(LogNet, Warning, TEXT("ClientSetActorTornOff called with invalid actor: %s"), *GetNameSafe(Actor));
+	}
+}
+
 void UNetDriver::NotifyActorTornOff(AActor* Actor)
 {
 
diff --git a/Engine/Source/Runtime/Engine/Private/ReplayHelper.cpp b/Engine/Source/Runtime/Engine/Private/ReplayHelper.cpp
--- a/Engine/Source/Runtime/Engine/Private/ReplayHelper.cpp
+++ b/Engine/Source/Runtime/Engine/Private/ReplayHelper.cpp
@@ -1845,7 +1845,7 @@ bool FReplayHelper::ReplicateActor(AActor* Actor, UNetConnection* Connection, bo
 			{
 				if (Connection->ResendAllDataState == EResendAllDataState::None)		// Don't close the channel if we're forcing them to re-open for checkpoints
 				{
-					Channel->Close(EChannelCloseReason::Destroyed);
+					Channel->Close(Actor->GetTearOff() ? EChannelCloseReason::TearOff : EChannelCloseReason::Destroyed);
 				}
 			}
 		}
@@ -2296,7 +2296,7 @@ void FReplayHelper::RequestCheckpoint()
 	bPendingCheckpointRequest = true;
 }
 
-void FReplayHelper::OnActorPreDestroy(UNetConnection* Connection, AActor* Actor)
+void FReplayHelper::RemoveActorFromCheckpoint(UNetConnection* Connection, AActor* Actor)
 {
 	check(Connection);
 
@@ -2315,7 +2315,7 @@ void FReplayHelper::OnActorPreDestroy(UNetConnection* Connection, AActor* Actor)
 						&& (*PendingIndex >= CheckpointSaveContext.NextAmortizedItem)
 						&& (CheckpointSaveContext.PendingCheckpointActors[*PendingIndex].Actor.Get() == Actor))
 					{
-						UE_LOG(LogDemo, Verbose, TEXT("Destroying actor while it is still in the PendingCheckpointActors list: %s"), *GetNameSafe(Actor));
+						UE_LOG(LogDemo, Verbose, TEXT("Removing actor while it is still in the PendingCheckpointActors list: %s"), *GetNameSafe(Actor));
 
 						FRepActorsCheckpointParams Params
 						{
@@ -2327,9 +2327,12 @@ void FReplayHelper::OnActorPreDestroy(UNetConnection* Connection, AActor* Actor)
 						TArrayView<FPendingCheckPointActor> PendingView(CheckpointSaveContext.PendingCheckpointActors);
 						int32 ActorIndex = 0;
 
-						// Serialize the actor one last time before it gets destroyed
+						// Serialize the actor one last time before it gets removed
 						ProcessCheckpointActors(Connection, PendingView.Slice(*PendingIndex, 1), ActorIndex, Params);
 
+						// Clear the pending actor to avoid any further processing
+						CheckpointSaveContext.PendingCheckpointActors[*PendingIndex].Actor = nullptr;
+
 						// don't allow the channel index to be reused until we're done with the checkpoint
 						Connection->AddReservedChannel(Channel->ChIndex);
 					}
diff --git a/Engine/Source/Runtime/Engine/Public/ReplayHelper.h b/Engine/Source/Runtime/Engine/Public/ReplayHelper.h
--- a/Engine/Source/Runtime/Engine/Public/ReplayHelper.h
+++ b/Engine/Source/Runtime/Engine/Public/ReplayHelper.h
@@ -106,7 +106,7 @@ private:
 
 	void RequestCheckpoint();
 
-	void OnActorPreDestroy(UNetConnection* Connection, AActor* Actor);
+	void RemoveActorFromCheckpoint(UNetConnection* Connection, AActor* Actor);
 
 	void NotifyActorDestroyed(UNetConnection* Connection, AActor* Actor);
 
