Fix for incorrect remote role on client with push model

diff --git a/Engine/Source/Runtime/Engine/Private/DataChannel.cpp b/Engine/Source/Runtime/Engine/Private/DataChannel.cpp
--- a/Engine/Source/Runtime/Engine/Private/DataChannel.cpp
+++ b/Engine/Source/Runtime/Engine/Private/DataChannel.cpp
@@ -3153,7 +3153,15 @@ public:
 	~FScopedRoleDowngrade()
 	{
 		// Upgrade role back to autonomous proxy if needed
-		Actor->SetAutonomousProxy(ActualRemoteRole == ROLE_AutonomousProxy, false);
+		if (Actor->GetRemoteRole() != ActualRemoteRole)
+		{
+			MARK_PROPERTY_DIRTY_FROM_NAME(AActor, RemoteRole, Actor);
+
+			if (ActualRemoteRole == ROLE_AutonomousProxy)
+			{
+				Actor->SetAutonomousProxy(true, false);
+			}
+		}
 	}
 
 private:
