- Added functions to get and set allowed screen orientations through the PlatformMisc class.
- Implemented logic for setting allowed screen orientations on iOS.
- Added possibility to set allowed screen orientation to SCREEN_ORIENTATION_SENSOR  on Android so that the screen orientation can be unlocked to adjust itself according to the sensor.
- Exposed PortraitSensor, LandscapeSensor and FullSensor screen orientations in blueprint library.

diff --git a/Engine/Source/Runtime/ApplicationCore/Private/IOS/IOSView.cpp b/Engine/Source/Runtime/ApplicationCore/Private/IOS/IOSView.cpp
--- a/Engine/Source/Runtime/ApplicationCore/Private/IOS/IOSView.cpp
+++ b/Engine/Source/Runtime/ApplicationCore/Private/IOS/IOSView.cpp
@@ -1043,7 +1043,7 @@ self.accessibilityElements = @[Window.accessibilityContainer];
  */
 - (NSUInteger)supportedInterfaceOrientations
 {
-	return UIInterfaceOrientationMaskAll;
+	return SupportedInterfaceOrientations;
 }
 
 /**
diff --git a/Engine/Source/Runtime/ApplicationCore/Public/IOS/IOSView.h b/Engine/Source/Runtime/ApplicationCore/Public/IOS/IOSView.h
--- a/Engine/Source/Runtime/ApplicationCore/Public/IOS/IOSView.h
+++ b/Engine/Source/Runtime/ApplicationCore/Public/IOS/IOSView.h
@@ -66,6 +66,7 @@ APPLICATIONCORE_API
 	BOOL bSecureTextEntry;
 	
 	volatile int32 KeyboardShowCount;
+	NSUInteger SupportedInterfaceOrientations = UIInterfaceOrientationMaskAll;
 }
 
 #if WITH_ACCESSIBILITY
diff --git a/Engine/Source/Runtime/Core/Private/Android/AndroidPlatformMisc.cpp b/Engine/Source/Runtime/Core/Private/Android/AndroidPlatformMisc.cpp
--- a/Engine/Source/Runtime/Core/Private/Android/AndroidPlatformMisc.cpp
+++ b/Engine/Source/Runtime/Core/Private/Android/AndroidPlatformMisc.cpp
@@ -3110,8 +3110,15 @@ const char* FAndroidMisc::GetThreadName(uint32 ThreadId)
 
 void FAndroidMisc::SetDeviceOrientation(EDeviceScreenOrientation NewDeviceOrentation)
 {
+	SetAllowedDeviceOrientation(NewDeviceOrentation);
+}
+
+void FAndroidMisc::SetAllowedDeviceOrientation(EDeviceScreenOrientation NewAllowedDeviceOrientation)
+{
+	AllowedDeviceOrientation = NewAllowedDeviceOrientation;
+
 #if USE_ANDROID_JNI
-	AndroidThunkCpp_SetOrientation(GetAndroidScreenOrientation(NewDeviceOrentation));
+	AndroidThunkCpp_SetOrientation(GetAndroidScreenOrientation(NewAllowedDeviceOrientation));
 #endif // USE_ANDROID_JNI
 }
 
@@ -3148,6 +3155,9 @@ int32 FAndroidMisc::GetAndroidScreenOrientation(EDeviceScreenOrientation ScreenO
 	case EDeviceScreenOrientation::LandscapeSensor:
 		AndroidScreenOrientation = EAndroidScreenOrientation::SCREEN_ORIENTATION_SENSOR_LANDSCAPE;
 		break;
+	case EDeviceScreenOrientation::FullSensor:
+		AndroidScreenOrientation = EAndroidScreenOrientation::SCREEN_ORIENTATION_SENSOR;
+		break;
 	}
 
 	return static_cast<int32>(AndroidScreenOrientation);
diff --git a/Engine/Source/Runtime/Core/Private/GenericPlatform/GenericPlatformMisc.cpp b/Engine/Source/Runtime/Core/Private/GenericPlatform/GenericPlatformMisc.cpp
--- a/Engine/Source/Runtime/Core/Private/GenericPlatform/GenericPlatformMisc.cpp
+++ b/Engine/Source/Runtime/Core/Private/GenericPlatform/GenericPlatformMisc.cpp
@@ -484,6 +484,8 @@ FString FSHA256Signature::ToString() const
 	bool FGenericPlatformMisc::bPromptForRemoteDebugOnEnsure = false;
 #endif	//#if !UE_BUILD_SHIPPING
 
+EDeviceScreenOrientation FGenericPlatformMisc::AllowedDeviceOrientation = EDeviceScreenOrientation::Unknown;
+
 struct FGenericPlatformMisc::FStaticData
 {
 	FString         RootDir;
@@ -1669,7 +1671,17 @@ EDeviceScreenOrientation FGenericPlatformMisc::GetDeviceOrientation()
 
 void FGenericPlatformMisc::SetDeviceOrientation(EDeviceScreenOrientation NewDeviceOrientation)
 {
-	// not implemented by default
+	SetAllowedDeviceOrientation(NewDeviceOrientation);
+}
+
+EDeviceScreenOrientation FGenericPlatformMisc::GetAllowedDeviceOrientation()
+{
+	return AllowedDeviceOrientation;
+}
+
+void FGenericPlatformMisc::SetAllowedDeviceOrientation(EDeviceScreenOrientation NewAllowedDeviceOrientation)
+{
+	AllowedDeviceOrientation = NewAllowedDeviceOrientation;
 }
 
 int32 FGenericPlatformMisc::GetDeviceVolume()
diff --git a/Engine/Source/Runtime/Core/Private/IOS/IOSPlatformMisc.cpp b/Engine/Source/Runtime/Core/Private/IOS/IOSPlatformMisc.cpp
--- a/Engine/Source/Runtime/Core/Private/IOS/IOSPlatformMisc.cpp
+++ b/Engine/Source/Runtime/Core/Private/IOS/IOSPlatformMisc.cpp
@@ -19,6 +19,7 @@
 #include "Misc/ConfigCacheIni.h"
 #include "Misc/CoreDelegates.h"
 #include "Misc/ScopeExit.h"
+#include "Modules/ModuleManager.h"
 #include "Apple/ApplePlatformCrashContext.h"
 #include "IOS/IOSPlatformCrashContext.h"
 #include "IOS/IOSPlatformPLCrashReporterIncludes.h"
@@ -293,6 +294,47 @@ EDeviceScreenOrientation ConvertFromUIInterfaceOrientation(UIInterfaceOrientatio
 }
 #endif
 
+#if !PLATFORM_TVOS
+UIInterfaceOrientationMask GetUIInterfaceOrientationMask(EDeviceScreenOrientation ScreenOrientation)
+{
+	switch (ScreenOrientation)
+	{
+	default:
+		// Fallthrough...
+	case EDeviceScreenOrientation::Unknown:
+		return UIInterfaceOrientationMaskAll;
+		break;
+	case EDeviceScreenOrientation::Portrait:
+		return UIInterfaceOrientationMaskPortrait;
+		break;
+	case EDeviceScreenOrientation::PortraitUpsideDown:
+		return UIInterfaceOrientationMaskPortraitUpsideDown;
+		break;
+	case EDeviceScreenOrientation::LandscapeLeft:
+		return UIInterfaceOrientationMaskLandscapeLeft;
+		break;
+	case EDeviceScreenOrientation::LandscapeRight:
+		return UIInterfaceOrientationMaskLandscapeRight;
+		break;
+	case EDeviceScreenOrientation::FaceUp:
+		return UIInterfaceOrientationMaskAll;
+		break;
+	case EDeviceScreenOrientation::FaceDown:
+		return UIInterfaceOrientationMaskAll;
+		break;
+	case EDeviceScreenOrientation::PortraitSensor:
+		return UIInterfaceOrientationMaskAll;
+		break;
+	case EDeviceScreenOrientation::LandscapeSensor:
+		return UIInterfaceOrientationMaskLandscape;
+		break;
+	case EDeviceScreenOrientation::FullSensor:
+		return UIInterfaceOrientationMaskAll;
+		break;
+	}
+}
+#endif
+
 #if !PLATFORM_TVOS
 UIInterfaceOrientation GInterfaceOrientation = UIInterfaceOrientationUnknown;
 #endif
@@ -313,10 +355,17 @@ EDeviceScreenOrientation FIOSPlatformMisc::GetDeviceOrientation()
 
 void FIOSPlatformMisc::SetDeviceOrientation(EDeviceScreenOrientation NewDeviceOrientation)
 {
-	// not implemented yet
+	SetAllowedDeviceOrientation(NewDeviceOrientation);
 }
 
-#include "Modules/ModuleManager.h"
+void FIOSPlatformMisc::SetAllowedDeviceOrientation(EDeviceScreenOrientation NewAllowedDeviceOrientation)
+{
+	AllowedDeviceOrientation = NewAllowedDeviceOrientation;
+
+#if !PLATFORM_TVOS
+	[IOSAppDelegate GetDelegate].IOSView->SupportedInterfaceOrientations = GetUIInterfaceOrientationMask(NewAllowedDeviceOrientation);
+#endif
+}
 
 bool FIOSPlatformMisc::HasPlatformFeature(const TCHAR* FeatureName)
 {
diff --git a/Engine/Source/Runtime/Core/Public/Android/AndroidPlatformMisc.h b/Engine/Source/Runtime/Core/Public/Android/AndroidPlatformMisc.h
--- a/Engine/Source/Runtime/Core/Public/Android/AndroidPlatformMisc.h
+++ b/Engine/Source/Runtime/Core/Public/Android/AndroidPlatformMisc.h
@@ -37,7 +37,9 @@ struct CORE_API FAndroidMisc : public FGenericPlatformMisc
 	static void PlatformTearDown();
 	static void PlatformHandleSplashScreen(bool ShowSplashScreen);
     static EDeviceScreenOrientation GetDeviceOrientation() { return DeviceOrientation; }
+	UE_DEPRECATED(5.1, "SetDeviceOrientation is deprecated. Use SetAllowedDeviceOrientation instead.")
 	static void SetDeviceOrientation(EDeviceScreenOrientation NewDeviceOrentation);
+	static void SetAllowedDeviceOrientation(EDeviceScreenOrientation NewAllowedDeviceOrientation);
     
 	FORCEINLINE static int32 GetMaxPathLength()
 	{
diff --git a/Engine/Source/Runtime/Core/Public/GenericPlatform/GenericPlatformMisc.h b/Engine/Source/Runtime/Core/Public/GenericPlatform/GenericPlatformMisc.h
--- a/Engine/Source/Runtime/Core/Public/GenericPlatform/GenericPlatformMisc.h
+++ b/Engine/Source/Runtime/Core/Public/GenericPlatform/GenericPlatformMisc.h
@@ -243,6 +243,9 @@ enum class EDeviceScreenOrientation : uint8
 
 	/** The orientation is landscape, oriented upright with the sensor */
 	LandscapeSensor,
+
+	/** The orientation is no longer locked and adjusts according to the sensor */
+	FullSensor,
 };
 
 
@@ -1515,16 +1518,33 @@ public:
 	static bool IsRunningOnBattery();
 
 	/**
-	 * Returns the orientation of the device: e.g. Portrait, LandscapeRight.
-	 * @see EScreenOrientation
+	 * Returns the current orientation of the device: will be either Portrait, LandscapeLeft, PortraitUpsideDown or LandscapeRight.
+	 * 
+	 * @return An EDeviceScreenOrientation value.
 	 */
 	static EDeviceScreenOrientation GetDeviceOrientation();
+
 	/**
-	 * Change the orientation of the device: e.g. Portrait, LandscapeRight.
-	 * @see EScreenOrientation
+	 * Change the orientation of the device.
 	 */
+	UE_DEPRECATED(5.1, "SetDeviceOrientation is deprecated. Use SetAllowedDeviceOrientation instead.")
 	static void SetDeviceOrientation(EDeviceScreenOrientation NewDeviceOrientation);
 
+	/**
+	 * Returns the allowed orientation of the device. This is NOT the same as GetDeviceOrientation, which only returns Portrait, LandscapeLeft, 
+	 * PortraitUpsideDown or LandscapeRight. The allowed orientation limits what orientation your device can have. So if you set the allowed orientation 
+	 * to LandscapeLeft, GetDeviceOrientation will only ever return LandscapeLeft. But if you set the allowed orientation to LandscapeSensor, you are actually 
+	 * restricting the allowed orientations to LandscapeLeft OR LandscapeRight (depending on the sensor), so GetDeviceOrientation might return LandscapeLeft OR LandscapeRight.
+	 * 
+	 * @return An EDeviceScreenOrientation value.
+	 */
+	static EDeviceScreenOrientation GetAllowedDeviceOrientation();
+
+	/**
+	 * Change the allowed orientation of the device. 
+	 */
+	static void SetAllowedDeviceOrientation(EDeviceScreenOrientation NewAllowedDeviceOrientation);
+
 	/**
 	 * Returns the device volume if the device is capable of returning that information.
 	 *  -1 : Unknown
@@ -1803,6 +1823,8 @@ protected:
 	static bool bPromptForRemoteDebugOnEnsure;
 #endif	//#if !UE_BUILD_SHIPPING
 
+	static EDeviceScreenOrientation AllowedDeviceOrientation;
+
 private:
 	struct FStaticData;
 };
diff --git a/Engine/Source/Runtime/Core/Public/IOS/IOSPlatformMisc.h b/Engine/Source/Runtime/Core/Public/IOS/IOSPlatformMisc.h
--- a/Engine/Source/Runtime/Core/Public/IOS/IOSPlatformMisc.h
+++ b/Engine/Source/Runtime/Core/Public/IOS/IOSPlatformMisc.h
@@ -57,7 +57,9 @@ struct CORE_API FIOSPlatformMisc : public FApplePlatformMisc
 	static bool IsRunningOnBattery();
 	static float GetDeviceTemperatureLevel();
 	static EDeviceScreenOrientation GetDeviceOrientation();
+	UE_DEPRECATED(5.1, "SetDeviceOrientation is deprecated. Use SetAllowedDeviceOrientation instead.")
 	static void SetDeviceOrientation(EDeviceScreenOrientation NewDeviceOrientation);
+	static void SetAllowedDeviceOrientation(EDeviceScreenOrientation NewAllowedDeviceOrientation);
 	static int32 GetDeviceVolume();
 	static void SetBrightness(float Brightness);
 	static float GetBrightness();
diff --git a/Engine/Source/Runtime/Engine/Classes/Kismet/BlueprintPlatformLibrary.h b/Engine/Source/Runtime/Engine/Classes/Kismet/BlueprintPlatformLibrary.h
--- a/Engine/Source/Runtime/Engine/Classes/Kismet/BlueprintPlatformLibrary.h
+++ b/Engine/Source/Runtime/Engine/Classes/Kismet/BlueprintPlatformLibrary.h
@@ -10,6 +10,7 @@
 #include "BlueprintPlatformLibrary.generated.h"
 
 class ILocalNotificationService;
+enum class EDeviceScreenOrientation : uint8;
 
 /**
  * The list of possible device/screen orientation for mobile devices
@@ -38,7 +39,16 @@ namespace EScreenOrientation
 		FaceUp,
 
 		/** The orientation is as if place on a desk with the screen downward */
-		FaceDown
+		FaceDown,
+
+		/** The orientation is portrait, oriented upright with the sensor */
+		PortraitSensor,
+
+		/** The orientation is landscape, oriented upright with the sensor */
+		LandscapeSensor,
+
+		/** The orientation is no longer locked and adjusts according to the sensor */
+		FullSensor
 	};
 }
 
@@ -239,4 +249,22 @@ public:
 	 */
 	UFUNCTION(BlueprintPure, Category="Platform|LocalNotification")
 	static EScreenOrientation::Type GetDeviceOrientation();
+
+	/**
+	 * @return the current allowed device orientation
+	 */
+	UFUNCTION(BlueprintPure, Category = "Platform|LocalNotification")
+	static EScreenOrientation::Type GetAllowedDeviceOrientation();
+
+	/**
+	 * Set the allowed device orientation
+	 */
+	UFUNCTION(BlueprintCallable, Category = "Platform|LocalNotification")
+	static void SetAllowedDeviceOrientation(EScreenOrientation::Type NewAllowedDeviceOrientation);
+
+private:
+
+	static EScreenOrientation::Type ConvertToScreenOrientation(EDeviceScreenOrientation DeviceScreenOrientation);
+
+	static EDeviceScreenOrientation ConvertToDeviceScreenOrientation(EScreenOrientation::Type ScreenOrientation);
 };
diff --git a/Engine/Source/Runtime/Engine/Private/Kismet/BlueprintPlatformLibrary.cpp b/Engine/Source/Runtime/Engine/Private/Kismet/BlueprintPlatformLibrary.cpp
--- a/Engine/Source/Runtime/Engine/Private/Kismet/BlueprintPlatformLibrary.cpp
+++ b/Engine/Source/Runtime/Engine/Private/Kismet/BlueprintPlatformLibrary.cpp
@@ -179,28 +179,86 @@ void UBlueprintPlatformLibrary::GetLaunchNotification(bool& NotificationLaunched
 
 EScreenOrientation::Type UBlueprintPlatformLibrary::GetDeviceOrientation()
 {
-	switch (FPlatformMisc::GetDeviceOrientation())
+	return ConvertToScreenOrientation(FPlatformMisc::GetDeviceOrientation());
+}
+
+
+EScreenOrientation::Type UBlueprintPlatformLibrary::GetAllowedDeviceOrientation()
+{
+	return ConvertToScreenOrientation(FPlatformMisc::GetAllowedDeviceOrientation());
+}
+
+void UBlueprintPlatformLibrary::SetAllowedDeviceOrientation(EScreenOrientation::Type NewAllowedDeviceOrientation)
+{
+	FPlatformMisc::SetAllowedDeviceOrientation(ConvertToDeviceScreenOrientation(NewAllowedDeviceOrientation));
+}
+
+EScreenOrientation::Type UBlueprintPlatformLibrary::ConvertToScreenOrientation(EDeviceScreenOrientation DeviceScreenOrientation)
+{
+	switch (DeviceScreenOrientation)
 	{
-		case EDeviceScreenOrientation::Portrait:
-			return EScreenOrientation::Portrait;
-			
-		case EDeviceScreenOrientation::PortraitUpsideDown:
-			return EScreenOrientation::PortraitUpsideDown;
-			
-		case EDeviceScreenOrientation::LandscapeLeft:
-			return EScreenOrientation::LandscapeLeft;
-
-		case EDeviceScreenOrientation::LandscapeRight:
-			return EScreenOrientation::LandscapeRight;
-
-		case EDeviceScreenOrientation::FaceUp:
-			return EScreenOrientation::FaceUp;
-			
-		case EDeviceScreenOrientation::FaceDown:
-			return EScreenOrientation::FaceDown;
+	case EDeviceScreenOrientation::Portrait:
+		return EScreenOrientation::Portrait;
+
+	case EDeviceScreenOrientation::PortraitUpsideDown:
+		return EScreenOrientation::PortraitUpsideDown;
+
+	case EDeviceScreenOrientation::LandscapeLeft:
+		return EScreenOrientation::LandscapeLeft;
+
+	case EDeviceScreenOrientation::LandscapeRight:
+		return EScreenOrientation::LandscapeRight;
+
+	case EDeviceScreenOrientation::FaceUp:
+		return EScreenOrientation::FaceUp;
+
+	case EDeviceScreenOrientation::FaceDown:
+		return EScreenOrientation::FaceDown;
+
+	case EDeviceScreenOrientation::PortraitSensor:
+		return EScreenOrientation::PortraitSensor;
+
+	case EDeviceScreenOrientation::LandscapeSensor:
+		return EScreenOrientation::LandscapeSensor;
+
+	case EDeviceScreenOrientation::FullSensor:
+		return EScreenOrientation::FullSensor;
 	}
 	return EScreenOrientation::Unknown;
 }
 
+EDeviceScreenOrientation UBlueprintPlatformLibrary::ConvertToDeviceScreenOrientation(EScreenOrientation::Type ScreenOrientation)
+{
+	switch (ScreenOrientation)
+	{
+	case EScreenOrientation::Portrait:
+		return EDeviceScreenOrientation::Portrait;
+
+	case EScreenOrientation::PortraitUpsideDown:
+		return EDeviceScreenOrientation::PortraitUpsideDown;
+
+	case EScreenOrientation::LandscapeLeft:
+		return EDeviceScreenOrientation::LandscapeLeft;
+
+	case EScreenOrientation::LandscapeRight:
+		return EDeviceScreenOrientation::LandscapeRight;
+
+	case EScreenOrientation::FaceUp:
+		return EDeviceScreenOrientation::FaceUp;
+
+	case EScreenOrientation::FaceDown:
+		return EDeviceScreenOrientation::FaceDown;
+
+	case EScreenOrientation::PortraitSensor:
+		return EDeviceScreenOrientation::PortraitSensor;
+
+	case EScreenOrientation::LandscapeSensor:
+		return EDeviceScreenOrientation::LandscapeSensor;
+
+	case EScreenOrientation::FullSensor:
+		return EDeviceScreenOrientation::FullSensor;
+	}
+	return EDeviceScreenOrientation::Unknown;
+}
 
 ILocalNotificationService* UBlueprintPlatformLibrary::platformService = nullptr;
