Fix for smoketest crash with default SoundWave PostLoad

diff --git a/Engine/Source/Runtime/AudioPlatformConfiguration/Public/AudioCompressionSettings.h b/Engine/Source/Runtime/AudioPlatformConfiguration/Public/AudioCompressionSettings.h
--- a/Engine/Source/Runtime/AudioPlatformConfiguration/Public/AudioCompressionSettings.h
+++ b/Engine/Source/Runtime/AudioPlatformConfiguration/Public/AudioCompressionSettings.h
@@ -106,6 +106,7 @@ struct FPlatformAudioCookOverrides
 		{
 			return;
 		}
+		OutSuffix.AppendChar('_');
 
 		int32 CompressionQualityHash = FMath::FloorToInt(InOverrides->CompressionQualityModifier * 100.0f);
 		OutSuffix.AppendInt(CompressionQualityHash);
diff --git a/Engine/Source/Runtime/Engine/Private/SoundWave.cpp b/Engine/Source/Runtime/Engine/Private/SoundWave.cpp
--- a/Engine/Source/Runtime/Engine/Private/SoundWave.cpp
+++ b/Engine/Source/Runtime/Engine/Private/SoundWave.cpp
@@ -158,7 +158,12 @@ void FSoundWaveData::InitializeDataFromSoundWave(USoundWave& InWave)
 	}
 
 	// cache the runtime format for the wave
-	RuntimeFormat = FindRuntimeFormat(InWave);
+	// note if this fails, it will ensure, and keep the "FSoundWaveProxy_InvalidFormat" as its value.
+	FName FoundFormat = FindRuntimeFormat(InWave);
+	if (!FoundFormat.IsNone())
+	{
+		RuntimeFormat = FoundFormat;
+	}
 	
 	SoundWaveKeyCached = FObjectKey(&InWave);
 	CuePoints = InWave.CuePoints;
@@ -189,41 +194,43 @@ FName FSoundWaveData::FindRuntimeFormat(const USoundWave& InWave) const
 	}
 #endif //WITH_EDITOR
 
-	// Force-inline does not use the stream cache chunks at all, so we have to work out the format
-	// from the name of the inner "Format" which has suffixes in the form "FORMAT123456789"
-	if (LoadingBehavior == ESoundWaveLoadingBehavior::ForceInline)
+	// If we can't render audio, we won't contain any compressed data.
+	if (InWave.IsTemplate() || !FApp::CanEverRenderAudio())
 	{
-		TArray<FName> AllValidFormats;
-		FAudioDeviceManager::Get()->GetAudioFormatSettings().GetAllWaveFormats(AllValidFormats);
+		return {};
+	}
 
+	// If the data is "forceinline" it will exist in the CompressedFormatData.
+	if (LoadingBehavior == ESoundWaveLoadingBehavior::ForceInline)
+	{
 		TArray<FName> ContainedFormats;
 		CompressedFormatData.GetContainedFormats(ContainedFormats);
-		if (ContainedFormats.Num() > 0)
-		{
-			// We only support 1 cooked format per wave, so just look at the first one.
-			FString CompressedFormatString = ContainedFormats[0].GetPlainNameString();
-			for (const FName ValidFormat : AllValidFormats)
-			{
-				if (CompressedFormatString.StartsWith(*ValidFormat.GetPlainNameString()))
-				{
-					return ValidFormat;
-				}
-			}
 
-			checkf(false, TEXT("Failed to determine encoded (inlined) AudioFormat: '%s'"), *CompressedFormatString);
-			
-			// Fail.
-			return {};
-		}	
+		// Only one format is supported currently.
+		if (ensureMsgf(ContainedFormats.Num() == 1, TEXT("ContainedFormats::Num()=%d"), ContainedFormats.Num()))
+		{			
+			FName CompressedFormat = ContainedFormats[0];
+			FString CompressedFormatString = CompressedFormat.GetPlainNameString();
 
-		checkf(false, TEXT("Failed to find any encoded formats in container"));
+			// If the format contains a delimiter, it's using a platform cook overrides with a suffix.
+			// Return up to the first delimiter e.g. "ADPCM_some_suffix" -> "ADPCM"
+			int32 DelimIndex = 0;
+			if (CompressedFormatString.FindChar('_', DelimIndex))
+			{
+				return *CompressedFormatString.Left(DelimIndex);
+			}
+			else
+			{
+				return CompressedFormat;
+			}
+		}
 		
 		// Fail.
 		return {};
 	}
+	// ... else the audio format is written in RunningPlatformData (i.e. stream cache).
 	else
 	{
-		// If we're using force inline, the audio format is written in the running platform.
 		return RunningPlatformData.GetAudioFormat();
 	}
 }
@@ -1755,6 +1762,12 @@ void USoundWave::PostLoad()
 		}
 	}
 
+	// update shared state
+	if (!HasAnyFlags(RF_ClassDefaultObject) && FApp::CanEverRenderAudio())
+	{
+		SoundWaveDataPtr->InitializeDataFromSoundWave(*this);
+	}
+
 	// We don't precache default objects and we don't precache in the Editor as the latter will
 	// most likely cause us to run out of memory.
 	if (!GIsEditor && !IsTemplate(RF_ClassDefaultObject) && GEngine)
@@ -1802,8 +1815,6 @@ void USoundWave::PostLoad()
 	INC_FLOAT_STAT_BY(STAT_AudioBufferTime, Duration);
 	INC_FLOAT_STAT_BY(STAT_AudioBufferTimeChannels, NumChannels * Duration);
 
-	// update shared state
-	SoundWaveDataPtr->InitializeDataFromSoundWave(*this);
 }
 
 bool USoundWave::LoadZerothChunk()
