Make packaged build working with cloth asset

diff --git a/Engine/Source/Runtime/Engine/Public/Rendering/SkeletalMeshRenderData.h b/Engine/Source/Runtime/Engine/Public/Rendering/SkeletalMeshRenderData.h
--- a/Engine/Source/Runtime/Engine/Public/Rendering/SkeletalMeshRenderData.h
+++ b/Engine/Source/Runtime/Engine/Public/Rendering/SkeletalMeshRenderData.h
@@ -62,7 +62,7 @@ public:
 #endif
 
 	/** Serialize to/from the specified archive.. */
-	void Serialize(FArchive& Ar, USkinnedAsset* Owner);
+	ENGINE_API void Serialize(FArchive& Ar, USkinnedAsset* Owner);
 
 	/** Initializes rendering resources. */
 	ENGINE_API void InitResources(bool bNeedsVertexColors, TArray<UMorphTarget*>& InMorphTargets, USkinnedAsset* Owner);
