Fix blueprint debugger not being able to get focus if a progress is active during PIE

diff --git a/Engine/Source/Runtime/Engine/Private/GameInstance.cpp b/Engine/Source/Runtime/Engine/Private/GameInstance.cpp
--- a/Engine/Source/Runtime/Engine/Private/GameInstance.cpp
+++ b/Engine/Source/Runtime/Engine/Private/GameInstance.cpp
@@ -428,7 +428,10 @@ FGameInstancePIEResult UGameInstance::StartPlayInEditorGameInstance(ULocalPlayer
 	else
 	{
 		FScopedSlowTask SlowTask(100, NSLOCTEXT("UnrealEd", "StartPlayInEditor", "Starting PIE..."));
-		SlowTask.MakeDialogDelayed(1.0f, false, true);
+		// Disabled for now in PIE until a proper fix is found to give the focus to the blueprint debugger
+		// during beginplay if a breakpoint is hit. The focus is currently held by the slowtask and prevents the debugger from working.
+		// Related JIRA UE-159973
+		SlowTask.MakeDialogDelayed(1.0f, false, false);
 
 		// we're going to be playing in the current world, get it ready for play
 		UWorld* const PlayWorld = GetWorld();
