Inclusive Terminology Push
Renamed terms at BoneContainer and Skeleton
Removed false positives for the tool that checks for unaceptable words at SkeletalMeshComponent and SkinnedMeshComponent

diff --git a/Engine/Source/Runtime/Engine/Classes/Components/SkeletalMeshComponent.h b/Engine/Source/Runtime/Engine/Classes/Components/SkeletalMeshComponent.h
--- a/Engine/Source/Runtime/Engine/Classes/Components/SkeletalMeshComponent.h
+++ b/Engine/Source/Runtime/Engine/Classes/Components/SkeletalMeshComponent.h
@@ -1230,7 +1230,7 @@ public:
 	UFUNCTION(BlueprintCallable, Category = "Clothing", meta = (UnsafeDuringActorConstruction = "true"))
 	void BindClothToLeaderPoseComponent();
 
-	UE_DEPRECATED(5.1, "BindClothToMasterPoseComponent has been deprecated. Please use BindClothToLeaderPoseComponent instead.")
+	UE_DEPRECATED(5.1, "This method has been deprecated. Please use BindClothToLeaderPoseComponent instead.")
 	void BindClothToMasterPoseComponent() { BindClothToLeaderPoseComponent(); }
 
 	/** 
@@ -1243,7 +1243,7 @@ public:
 	UFUNCTION(BlueprintCallable, Category="Clothing", meta=(UnsafeDuringActorConstruction="true"))
 	void UnbindClothFromLeaderPoseComponent(bool bRestoreSimulationSpace = true);
 
-	UE_DEPRECATED(5.1, "UnbindClothFromMasterPoseComponent has been deprecated. Please use UnbindClothFromLeaderPoseComponent instead.")
+	UE_DEPRECATED(5.1, "This method has been deprecated. Please use UnbindClothFromLeaderPoseComponent instead.")
 	void UnbindClothFromMasterPoseComponent(bool bRestoreSimulationSpace = true) { UnbindClothFromLeaderPoseComponent(bRestoreSimulationSpace); }
 
 	/**
@@ -1442,7 +1442,7 @@ public:
 
 	bool IsClothBoundToLeaderComponent() const { return bBindClothToLeaderComponent; }
 	
-	UE_DEPRECATED(5.1, "IsClothBoundToMasterComponent has been deprecated. Please, use IsClothBoundToLeaderComponent instead.")
+	UE_DEPRECATED(5.1, "This method has been deprecated. Please, use IsClothBoundToLeaderComponent instead.")
 	bool IsClothBoundToMasterComponent() const { return IsClothBoundToLeaderComponent(); }
 
 	/** Get the current clothing simulation (read only) */
@@ -1871,7 +1871,7 @@ protected:
 public:
 	virtual void TickPose(float DeltaTime, bool bNeedsValidRootMotion) override;
 	virtual void UpdateFollowerComponent() override;
-	UE_DEPRECATED(5.1, "UpdateSlaveComponent has been deprecated. Please use UpdateFollowerComponent instead.")
+	UE_DEPRECATED(5.1, "This method has been deprecated. Please use UpdateFollowerComponent instead.")
 	virtual void UpdateSlaveComponent() override { UpdateFollowerComponent(); };
 	virtual bool ShouldUpdateTransform(bool bLODHasChanged) const override;
 	virtual bool ShouldTickPose() const override;
diff --git a/Engine/Source/Runtime/Engine/Classes/Components/SkinnedMeshComponent.h b/Engine/Source/Runtime/Engine/Classes/Components/SkinnedMeshComponent.h
--- a/Engine/Source/Runtime/Engine/Classes/Components/SkinnedMeshComponent.h
+++ b/Engine/Source/Runtime/Engine/Classes/Components/SkinnedMeshComponent.h
@@ -417,7 +417,7 @@ protected:
 public:
 	const TArray<int32>& GetLeaderBoneMap() const { return LeaderBoneMap; }
 
-	UE_DEPRECATED(5.1, "GetMasterBoneMap has been deprecated. Please use the GetLeaderBoneMap.")
+	UE_DEPRECATED(5.1, "This method has been deprecated. Please use the GetLeaderBoneMap.")
 	const TArray<int32>& GetMasterBoneMap() const { return GetLeaderBoneMap(); }
 
 	const FExternalMorphWeightData& GetExternalMorphWeights(int32 LOD) const { return ExternalMorphWeightData[LOD]; }
@@ -1259,7 +1259,7 @@ public:
 	 */
 	virtual void UpdateFollowerComponent();
 
-	UE_DEPRECATED(5.1, "UpdateSlaveComponent has been deprecated. Please use UpdateFollowerComponent instead.")
+	UE_DEPRECATED(5.1, "This method has been deprecated. Please use UpdateFollowerComponent instead.")
 	virtual void UpdateSlaveComponent() { UpdateFollowerComponent(); }
 
 	/** 
@@ -1445,26 +1445,26 @@ public:
 	UFUNCTION(BlueprintCallable, Category = "Components|SkinnedMesh")
 	void SetLeaderPoseComponent(USkinnedMeshComponent* NewLeaderBoneComponent, bool bForceUpdate = false);
 
-	UE_DEPRECATED(5.1, "SetMasterPoseComponent has been deprecated. Please use SetLeaderPoseComponent instead.")
+	UE_DEPRECATED(5.1, "This method has been deprecated. Please use SetLeaderPoseComponent instead.")
 	void SetMasterPoseComponent(USkinnedMeshComponent* NewMasterBoneComponent, bool bForceUpdate = false) { SetLeaderPoseComponent(NewMasterBoneComponent, bForceUpdate); }
 
 	/** Return current active list of follower components */
 	const TArray< TWeakObjectPtr<USkinnedMeshComponent> >& GetFollowerPoseComponents() const;
 
-	UE_DEPRECATED(5.1, "GetSlavePoseComponents has been deprecated. Please use GetFollowerPoseComponents instead.")
+	UE_DEPRECATED(5.1, "This method has been deprecated. Please use GetFollowerPoseComponents instead.")
 	const TArray< TWeakObjectPtr<USkinnedMeshComponent> >& GetSlavePoseComponents() const { return GetFollowerPoseComponents(); }
 
 protected:
 	/** Add a follower component to the FollowerPoseComponents array */
 	virtual void AddFollowerPoseComponent(USkinnedMeshComponent* SkinnedMeshComponent);
 
-	UE_DEPRECATED(5.1, "AddSlavePoseComponent has been deprecated. Please use AddFollowerPoseComponent instead.")
+	UE_DEPRECATED(5.1, "This method has been deprecated. Please use AddFollowerPoseComponent instead.")
 	virtual void AddSlavePoseComponent(USkinnedMeshComponent* SkinnedMeshComponent) { AddFollowerPoseComponent(SkinnedMeshComponent); }
 
 	/** Remove a follower component from the FollowerPoseComponents array */
 	virtual void RemoveFollowerPoseComponent(USkinnedMeshComponent* SkinnedMeshComponent);
 
-	UE_DEPRECATED(5.1, "RemoveSlavePoseComponent has been deprecated. Please use RemoveFollowerPoseComponent instead.")
+	UE_DEPRECATED(5.1, "This method has been deprecated. Please use RemoveFollowerPoseComponent instead.")
 	virtual void RemoveSlavePoseComponent(USkinnedMeshComponent* SkinnedMeshComponent) { RemoveFollowerPoseComponent(SkinnedMeshComponent); }
 
 public:
@@ -1475,7 +1475,7 @@ public:
 	 */
 	void RefreshFollowerComponents();
 
-	UE_DEPRECATED(5.1, "RefreshSlaveComponents has been deprecated. Please use RefreshFollowerComponents instead.")
+	UE_DEPRECATED(5.1, "This method has been deprecated. Please use RefreshFollowerComponents instead.")
 	void RefreshSlaveComponents() { RefreshFollowerComponents(); }
 
 	/**
@@ -1483,7 +1483,7 @@ public:
 	 */
 	void UpdateLeaderBoneMap();
 
-	UE_DEPRECATED(5.1, "UpdateMasterBoneMap has been deprecated. Please use UpdateLeaderBoneMap instead.")
+	UE_DEPRECATED(5.1, "This method has been deprecated. Please use UpdateLeaderBoneMap instead.")
 	void UpdateMasterBoneMap() { UpdateLeaderBoneMap(); }
 
 	/**
diff --git a/Engine/Source/Runtime/Engine/Private/Animation/Skeleton.cpp b/Engine/Source/Runtime/Engine/Private/Animation/Skeleton.cpp
--- a/Engine/Source/Runtime/Engine/Private/Animation/Skeleton.cpp
+++ b/Engine/Source/Runtime/Engine/Private/Animation/Skeleton.cpp
@@ -1024,7 +1024,7 @@ bool USkeleton::IsCompatibleMesh(const USkeletalMesh* InSkelMesh, bool bDoParent
 	}
 
 	// originally we made sure at least matches more than 50% 
-	// but then slave components can't play since they're only partial
+	// but then follower components can't play since they're only partial
 	// if the hierarchy matches, and if it's more then 1 bone, we allow
 	return (NumOfBoneMatches > 0);
 }
diff --git a/Engine/Source/Runtime/Engine/Private/BoneContainer.cpp b/Engine/Source/Runtime/Engine/Private/BoneContainer.cpp
--- a/Engine/Source/Runtime/Engine/Private/BoneContainer.cpp
+++ b/Engine/Source/Runtime/Engine/Private/BoneContainer.cpp
@@ -596,9 +596,9 @@ bool FBoneReference::Initialize(const FBoneContainer& RequiredBones)
 	BoneIndex = RequiredBones.GetPoseBoneIndexForBoneName(BoneName);
 
 	bUseSkeletonIndex = false;
-	// If bone name is not found, look into the master skeleton to see if it's found there.
-	// SkeletalMeshes can exclude bones from the master skeleton, and that's OK.
-	// If it's not found in the master skeleton, the bone does not exist at all! so we should report it as a warning.
+	// If bone name is not found, look into the leader skeleton to see if it's found there.
+	// SkeletalMeshes can exclude bones from the leader skeleton, and that's OK.
+	// If it's not found in the leader skeleton, the bone does not exist at all! so we should report it as a warning.
 	if (BoneIndex == INDEX_NONE && BoneName != NAME_None)
 	{
 		if (USkeleton* SkeletonAsset = RequiredBones.GetSkeletonAsset())
diff --git a/Engine/Source/Runtime/Engine/Private/Components/SkinnedMeshComponent.cpp b/Engine/Source/Runtime/Engine/Private/Components/SkinnedMeshComponent.cpp
--- a/Engine/Source/Runtime/Engine/Private/Components/SkinnedMeshComponent.cpp
+++ b/Engine/Source/Runtime/Engine/Private/Components/SkinnedMeshComponent.cpp
@@ -539,7 +539,7 @@ void USkinnedMeshComponent::OnRegister()
 
 	if (LeaderPoseComponent.IsValid())
 	{
-		// we have to make sure it updates the mastesr pose
+		// we have to make sure it updates the leader pose
 		SetLeaderPoseComponent(LeaderPoseComponent.Get(), true);
 	}
 	else
