FSkeletalMeshModel::Serialize: Add missing UsingCustomVersion GUID

diff --git a/Engine/Source/Runtime/Engine/Private/SkeletalMesh.cpp b/Engine/Source/Runtime/Engine/Private/SkeletalMesh.cpp
--- a/Engine/Source/Runtime/Engine/Private/SkeletalMesh.cpp
+++ b/Engine/Source/Runtime/Engine/Private/SkeletalMesh.cpp
@@ -1577,7 +1577,6 @@ void USkeletalMesh::Serialize( FArchive& Ar )
 	}
 #endif // WITH_EDITORONLY_DATA
 
-	Ar.UsingCustomVersion(FSkeletalMeshCustomVersion::GUID);
 	if (Ar.CustomVer(FSkeletalMeshCustomVersion::GUID) >= FSkeletalMeshCustomVersion::SplitModelAndRenderData)
 	{
 		bool bCooked = Ar.IsCooking();
diff --git a/Engine/Source/Runtime/Engine/Private/SkeletalMeshModel.cpp b/Engine/Source/Runtime/Engine/Private/SkeletalMeshModel.cpp
--- a/Engine/Source/Runtime/Engine/Private/SkeletalMeshModel.cpp
+++ b/Engine/Source/Runtime/Engine/Private/SkeletalMeshModel.cpp
@@ -10,13 +10,19 @@
 #include "UObject/UE5MainStreamObjectVersion.h"
 
 FSkeletalMeshModel::FSkeletalMeshModel()
-	: bGuidIsHash(false)	
+	: bGuidIsHash(false)
 {
 }
-
+PRAGMA_DISABLE_OPTIMIZATION
 void FSkeletalMeshModel::Serialize(FArchive& Ar, USkinnedAsset* Owner)
 {
 	DECLARE_SCOPE_CYCLE_COUNTER(TEXT("FSkeletalMeshModel::Serialize"), STAT_SkeletalMeshModel_Serialize, STATGROUP_LoadTime);
+	Ar.UsingCustomVersion(FFortniteMainBranchObjectVersion::GUID);
+	Ar.UsingCustomVersion(FSkeletalMeshCustomVersion::GUID);
+	Ar.UsingCustomVersion(FUE5MainStreamObjectVersion::GUID);
+
+	auto x = Ar.CustomVer(FFortniteMainBranchObjectVersion::GUID);
+
 	bool bIsEditorDataStripped = false;
 	if (Ar.IsSaving() || (Ar.IsLoading() && Ar.CustomVer(FFortniteMainBranchObjectVersion::GUID) >= FFortniteMainBranchObjectVersion::AllowSkeletalMeshToReduceTheBaseLOD))
 	{
@@ -26,9 +32,6 @@ void FSkeletalMeshModel::Serialize(FArchive& Ar, USkinnedAsset* Owner)
 
 	LODModels.Serialize(Ar, Owner);
 
-	Ar.UsingCustomVersion(FSkeletalMeshCustomVersion::GUID);
-	Ar.UsingCustomVersion(FUE5MainStreamObjectVersion::GUID);
-
 	// For old content without a GUID, generate one now from the data
 	if (Ar.IsLoading() && Ar.CustomVer(FSkeletalMeshCustomVersion::GUID) < FSkeletalMeshCustomVersion::SplitModelAndRenderData)
 	{
@@ -130,7 +133,7 @@ void FSkeletalMeshModel::GenerateGUIDFromHash(USkinnedAsset* Owner)
 	}
 	Sha.Final();
 
-	// Retrieve the hash and use it to construct a pseudo-GUID. 
+	// Retrieve the hash and use it to construct a pseudo-GUID.
 	uint32 Hash[5];
 	Sha.GetHash((uint8*)Hash);
 	SkeletalMeshModelGUID = FGuid(Hash[0] ^ Hash[4], Hash[1], Hash[2], Hash[3]);
@@ -168,7 +171,7 @@ FString FSkeletalMeshModel::GetLODModelIdString() const
 		Sha.Update((uint8*)IDArray.GetData(), IDArray.Num() * IDArray.GetTypeSize());
 	}
 	Sha.Final();
-	// Retrieve the hash and use it to construct a pseudo-GUID. 
+	// Retrieve the hash and use it to construct a pseudo-GUID.
 	uint32 Hash[5];
 	Sha.GetHash((uint8*)Hash);
 	return FGuid(Hash[0] ^ Hash[4], Hash[1], Hash[2], Hash[3]).ToString(EGuidFormats::Digits);
@@ -193,4 +196,5 @@ void FSkeletalMeshModel::EmptyOriginalReductionSourceMeshData()
 	OriginalReductionSourceMeshData_DEPRECATED.Empty();
 }
 
-#endif // WITH_EDITOR
\ No newline at end of file
+#endif // WITH_EDITOR
+PRAGMA_ENABLE_OPTIMIZATION
