Amend last commit FSkeletalMeshModel: remove pragma optimization tags

diff --git a/Engine/Source/Runtime/Engine/Private/SkeletalMeshModel.cpp b/Engine/Source/Runtime/Engine/Private/SkeletalMeshModel.cpp
--- a/Engine/Source/Runtime/Engine/Private/SkeletalMeshModel.cpp
+++ b/Engine/Source/Runtime/Engine/Private/SkeletalMeshModel.cpp
@@ -13,7 +13,7 @@ FSkeletalMeshModel::FSkeletalMeshModel()
 	: bGuidIsHash(false)
 {
 }
-PRAGMA_DISABLE_OPTIMIZATION
+
 void FSkeletalMeshModel::Serialize(FArchive& Ar, USkinnedAsset* Owner)
 {
 	DECLARE_SCOPE_CYCLE_COUNTER(TEXT("FSkeletalMeshModel::Serialize"), STAT_SkeletalMeshModel_Serialize, STATGROUP_LoadTime);
@@ -197,4 +197,3 @@ void FSkeletalMeshModel::EmptyOriginalReductionSourceMeshData()
 }
 
 #endif // WITH_EDITOR
-PRAGMA_ENABLE_OPTIMIZATION
