FSkeletalMeshModel::Serialize: Remove unused debug variable...

diff --git a/Engine/Source/Runtime/Engine/Private/SkeletalMeshModel.cpp b/Engine/Source/Runtime/Engine/Private/SkeletalMeshModel.cpp
--- a/Engine/Source/Runtime/Engine/Private/SkeletalMeshModel.cpp
+++ b/Engine/Source/Runtime/Engine/Private/SkeletalMeshModel.cpp
@@ -21,8 +21,6 @@ void FSkeletalMeshModel::Serialize(FArchive& Ar, USkinnedAsset* Owner)
 	Ar.UsingCustomVersion(FSkeletalMeshCustomVersion::GUID);
 	Ar.UsingCustomVersion(FUE5MainStreamObjectVersion::GUID);
 
-	auto x = Ar.CustomVer(FFortniteMainBranchObjectVersion::GUID);
-
 	bool bIsEditorDataStripped = false;
 	if (Ar.IsSaving() || (Ar.IsLoading() && Ar.CustomVer(FFortniteMainBranchObjectVersion::GUID) >= FFortniteMainBranchObjectVersion::AllowSkeletalMeshToReduceTheBaseLOD))
 	{
