World Partition: hide all Asset Registry specific code in FWorldPartitionActorDescUtils.

#rnx

diff --git a/Engine/Plugins/Editor/DataValidation/Source/DataValidation/Private/WorldPartitionChangelistValidator.cpp b/Engine/Plugins/Editor/DataValidation/Source/DataValidation/Private/WorldPartitionChangelistValidator.cpp
--- a/Engine/Plugins/Editor/DataValidation/Source/DataValidation/Private/WorldPartitionChangelistValidator.cpp
+++ b/Engine/Plugins/Editor/DataValidation/Source/DataValidation/Private/WorldPartitionChangelistValidator.cpp
@@ -11,6 +11,7 @@
 
 #include "Engine/World.h"
 #include "WorldPartition/WorldPartitionActorDesc.h"
+#include "WorldPartition/WorldPartitionActorDescUtils.h"
 #include "WorldPartition/ActorDescContainer.h"
 #include "WorldPartition/WorldPartition.h"
 #include "WorldPartition/DataLayer/DataLayerInstanceWithAsset.h"
@@ -78,8 +79,7 @@ EDataValidationResult UWorldPartitionChangelistValidator::ValidateActorsAndDataL
 	auto TryAssociateAssetDataToMap = [&MapToActorsFiles](const FAssetData& AssetData)
 	{
 		// Check that the asset is an actor
-		static FName NAME_ActorMetaDataClass(TEXT("ActorMetaDataClass"));
-		if (AssetData.TagsAndValues.Contains(NAME_ActorMetaDataClass)) // Could check AssetData.GetClass()->IsChildOf<AActor>() but this doesn't handle blueprints, all OFPA/WP Actors will have this tag
+		if (FWorldPartitionActorDescUtils::IsValidActorDescriptorFromAssetData(AssetData))
 		{
 			// WorldPartition actors are all in OFPA mode so they're external
 			// Extract the MapName from the ObjectPath (<PathToPackage>.<mapName>:<level>.<actorName>)
diff --git a/Engine/Source/Runtime/Engine/Private/Actor.cpp b/Engine/Source/Runtime/Engine/Private/Actor.cpp
--- a/Engine/Source/Runtime/Engine/Private/Actor.cpp
+++ b/Engine/Source/Runtime/Engine/Private/Actor.cpp
@@ -69,6 +69,7 @@
 
 #include "WorldPartition/WorldPartition.h"
 #include "WorldPartition/WorldPartitionActorDesc.h"
+#include "WorldPartition/WorldPartitionActorDescUtils.h"
 #include "WorldPartition/WorldPartitionRuntimeCell.h"
 #include "WorldPartition/DataLayer/DataLayerInstance.h"
 #include "WorldPartition/DataLayer/DataLayerAsset.h"
@@ -492,54 +493,10 @@ void AActor::GetAssetRegistryTags(TArray<FAssetRegistryTag>& OutTags) const
 void AActor::GetExtendedAssetRegistryTagsForSave(const ITargetPlatform* TargetPlatform, TArray<FAssetRegistryTag>& OutTags) const
 {
 	Super::GetExtendedAssetRegistryTagsForSave(TargetPlatform, OutTags);
-
+	
 	if (IsPackageExternal() && !IsChildActor())
 	{
-		TUniquePtr<FWorldPartitionActorDesc> ActorDesc(CreateActorDesc());
-
-		// If the actor is not added to a world, we can't retrieve its bounding volume, so try to get the existing one
-		if (ULevel* Level = GetLevel(); !Level || !Level->Actors.Contains(this))
-		{
-			IAssetRegistry& AssetRegistry = FModuleManager::LoadModuleChecked<FAssetRegistryModule>(TEXT("AssetRegistry")).Get();
-
-			FARFilter Filter;
-			Filter.bIncludeOnlyOnDiskAssets = true;
-			Filter.PackageNames.Add(GetPackage()->GetFName());
-
-			TArray<FAssetData> Assets;
-			AssetRegistry.GetAssets(Filter, Assets);
-
-			if (Assets.Num() == 1)
-			{
-				const FAssetData& InAssetData = Assets[0];
-
-				FString ActorMetaDataStr;
-				static FName NAME_ActorMetaData(TEXT("ActorMetaData"));
-				if (InAssetData.GetTagValue(NAME_ActorMetaData, ActorMetaDataStr))
-				{
-					FWorldPartitionActorDescInitData ActorDescInitData;
-					ActorDescInitData.NativeClass = AActor::StaticClass();
-					ActorDescInitData.PackageName = InAssetData.PackageName;
-					ActorDescInitData.ActorPath = InAssetData.ObjectPath;
-					FBase64::Decode(ActorMetaDataStr, ActorDescInitData.SerializedData);
-
-					TUniquePtr<FWorldPartitionActorDesc> NewActorDesc(AActor::StaticCreateClassActorDesc(ActorDescInitData.NativeClass));
-					NewActorDesc->Init(ActorDescInitData);
-
-					ActorDesc->TransferWorldData(NewActorDesc.Get());
-				}
-			}
-		}
-
-		const FString ActorMetaDataClass = GetParentNativeClass(GetClass())->GetPathName();
-		static FName NAME_ActorMetaDataClass(TEXT("ActorMetaDataClass"));
-		OutTags.Add(UObject::FAssetRegistryTag(NAME_ActorMetaDataClass, ActorMetaDataClass, UObject::FAssetRegistryTag::TT_Hidden));
-
-		TArray<uint8> SerializedData;
-		ActorDesc->SerializeTo(SerializedData);
-		const FString ActorMetaData = FBase64::Encode(SerializedData);
-		static FName NAME_ActorMetaData(TEXT("ActorMetaData"));
-		OutTags.Add(UObject::FAssetRegistryTag(NAME_ActorMetaData, ActorMetaData, UObject::FAssetRegistryTag::TT_Hidden));
+		FWorldPartitionActorDescUtils::AppendAssetDataTagsFromActor(this, OutTags);
 	}
 }
 
diff --git a/Engine/Source/Runtime/Engine/Private/WorldPartition/ActorDescContainer.cpp b/Engine/Source/Runtime/Engine/Private/WorldPartition/ActorDescContainer.cpp
--- a/Engine/Source/Runtime/Engine/Private/WorldPartition/ActorDescContainer.cpp
+++ b/Engine/Source/Runtime/Engine/Private/WorldPartition/ActorDescContainer.cpp
@@ -7,9 +7,8 @@
 #include "AssetRegistry/AssetRegistryModule.h"
 #include "WorldPartition/WorldPartitionLog.h"
 #include "WorldPartition/WorldPartitionHandle.h"
-#include "Misc/Base64.h"
+#include "WorldPartition/WorldPartitionActorDescUtils.h"
 #include "UObject/ObjectSaveContext.h"
-#include "UObject/CoreRedirects.h"
 #include "Engine/World.h"
 
 UActorDescContainer::FActorDescContainerInitializeDelegate UActorDescContainer::OnActorDescContainerInitialized;
@@ -49,62 +48,9 @@ void UActorDescContainer::Initialize(UWorld* InWorld, FName InPackageName)
 		AssetRegistry.GetAssets(Filter, Assets);
 	}
 
-	auto GetActorDescriptor = [this](const FAssetData& InAssetData) -> TUniquePtr<FWorldPartitionActorDesc>
-	{
-		FString ActorMetaDataClass;
-		static FName NAME_ActorMetaDataClass(TEXT("ActorMetaDataClass"));
-		if (InAssetData.GetTagValue(NAME_ActorMetaDataClass, ActorMetaDataClass))
-		{
-			FString ActorMetaDataStr;
-			static FName NAME_ActorMetaData(TEXT("ActorMetaData"));
-			if (InAssetData.GetTagValue(NAME_ActorMetaData, ActorMetaDataStr))
-			{
-				FString ActorClassName;
-				FString ActorPackageName;
-				if (!ActorMetaDataClass.Split(TEXT("."), &ActorPackageName, &ActorClassName))
-				{
-					ActorClassName = *ActorMetaDataClass;
-				}
-
-				// Look for a class redirectors
-				FCoreRedirectObjectName OldClassName = FCoreRedirectObjectName(*ActorClassName, NAME_None, *ActorPackageName);
-				FCoreRedirectObjectName NewClassName = FCoreRedirects::GetRedirectedName(ECoreRedirectFlags::Type_Class, OldClassName);
-
-				bool bIsValidClass = true;
-				UClass* ActorClass = UClass::TryFindTypeSlow<UClass>(NewClassName.ToString(), EFindFirstObjectOptions::ExactClass);
-
-				if (!ActorClass)
-				{
-					ActorClass = AActor::StaticClass();
-					bIsValidClass = false;
-				}
-
-				FWorldPartitionActorDescInitData ActorDescInitData;
-				ActorDescInitData.NativeClass = ActorClass;
-				ActorDescInitData.PackageName = InAssetData.PackageName;
-				ActorDescInitData.ActorPath = InAssetData.ObjectPath;
-				FBase64::Decode(ActorMetaDataStr, ActorDescInitData.SerializedData);
-
-				TUniquePtr<FWorldPartitionActorDesc> NewActorDesc(AActor::StaticCreateClassActorDesc(ActorDescInitData.NativeClass));
-
-				NewActorDesc->Init(ActorDescInitData);
-			
-				if (!bIsValidClass)
-				{
-					UE_LOG(LogWorldPartition, Warning, TEXT("Invalid class `%s` for actor guid `%s` ('%s') from package '%s'"), *NewClassName.ToString(), *NewActorDesc->GetGuid().ToString(), *NewActorDesc->GetActorName().ToString(), *NewActorDesc->GetActorPackage().ToString());
-					return nullptr;
-				}
-
-				return NewActorDesc;
-			}
-		}
-
-		return nullptr;
-	};
-
 	for (const FAssetData& Asset : Assets)
 	{
-		TUniquePtr<FWorldPartitionActorDesc> ActorDesc = GetActorDescriptor(Asset);
+		TUniquePtr<FWorldPartitionActorDesc> ActorDesc = FWorldPartitionActorDescUtils::GetActorDescriptorFromAssetData(Asset);
 
 		if (ActorDesc.IsValid())
 		{
@@ -201,21 +147,14 @@ void UActorDescContainer::OnObjectPreSave(UObject* Object, FObjectPreSaveContext
 				check(IsValidChecked(Actor));
 				if (TUniquePtr<FWorldPartitionActorDesc>* ExistingActorDesc = GetActorDescriptor(Actor->GetActorGuid()))
 				{
+					// Existing actor
 					OnActorDescUpdating(ExistingActorDesc->Get());
-
-					TUniquePtr<FWorldPartitionActorDesc> NewActorDesc(Actor->CreateActorDesc());
-
-					// Transfer any internal values not coming from the actor
-					NewActorDesc->TransferFrom(ExistingActorDesc->Get());
-
-					// This moves the TUniquePtr (not the content of the actor desc)
-					*ExistingActorDesc = MoveTemp(NewActorDesc);
-
+					FWorldPartitionActorDescUtils::UpdateActorDescriptorFomActor(Actor, *ExistingActorDesc);
 					OnActorDescUpdated(ExistingActorDesc->Get());
 				}
-				// New actor
 				else
 				{
+					// New actor
 					FWorldPartitionActorDesc* const AddedActorDesc = AddActor(Actor);
 					OnActorDescAdded(AddedActorDesc);
 				}
@@ -237,16 +176,14 @@ void UActorDescContainer::OnPackageDeleted(UPackage* Package)
 void UActorDescContainer::OnObjectsReplaced(const TMap<UObject*, UObject*>& OldToNewObjectMap)
 {
 	// Patch up Actor pointers in ActorDescs
-	for (auto Iter = OldToNewObjectMap.CreateConstIterator(); Iter; ++Iter)
+	for (auto [OldObject, NewObject] : OldToNewObjectMap)
 	{
-		if (AActor* OldActor = Cast<AActor>(Iter->Key))
+		if (AActor* OldActor = Cast<AActor>(OldObject))
 		{
+			AActor* NewActor = CastChecked<AActor>(NewObject);
 			if (FWorldPartitionActorDesc* ActorDesc = GetActorDesc(OldActor->GetActorGuid()))
-			{
-				if (ActorDesc->GetActor() == OldActor)
-				{
-					ActorDesc->ActorPtr = Cast<AActor>(Iter->Value);
-				}
+			{				
+				FWorldPartitionActorDescUtils::ReplaceActorDescriptorPointerFromActor(OldActor, NewActor, ActorDesc);
 			}
 		}
 	}
diff --git a/Engine/Source/Runtime/Engine/Private/WorldPartition/WorldPartitionActorDescUtils.cpp b/Engine/Source/Runtime/Engine/Private/WorldPartition/WorldPartitionActorDescUtils.cpp
new file mode 100644
--- /dev/null
+++ b/Engine/Source/Runtime/Engine/Private/WorldPartition/WorldPartitionActorDescUtils.cpp
@@ -0,0 +1,118 @@
+// Copyright Epic Games, Inc. All Rights Reserved.
+
+#if WITH_EDITOR
+#include "WorldPartition/WorldPartitionActorDescUtils.h"
+#include "WorldPartition/WorldPartitionLog.h"
+#include "GameFramework/Actor.h"
+#include "AssetRegistry/AssetRegistryModule.h"
+#include "AssetRegistry/IAssetRegistry.h"
+#include "UObject/CoreRedirects.h"
+#include "Misc/Base64.h"
+
+static FName NAME_ActorMetaDataClass(TEXT("ActorMetaDataClass"));
+static FName NAME_ActorMetaData(TEXT("ActorMetaData"));
+
+bool FWorldPartitionActorDescUtils::IsValidActorDescriptorFromAssetData(const FAssetData& InAssetData)
+{
+	return InAssetData.FindTag(NAME_ActorMetaDataClass) && InAssetData.FindTag(NAME_ActorMetaData);
+}
+
+TUniquePtr<FWorldPartitionActorDesc> FWorldPartitionActorDescUtils::GetActorDescriptorFromAssetData(const FAssetData& InAssetData)
+{
+	FString ActorMetaDataClass;
+	FString ActorMetaDataStr;
+	if (InAssetData.GetTagValue(NAME_ActorMetaDataClass, ActorMetaDataClass) && InAssetData.GetTagValue(NAME_ActorMetaData, ActorMetaDataStr))
+	{
+		FString ActorClassName;
+		FString ActorPackageName;
+		if (!ActorMetaDataClass.Split(TEXT("."), &ActorPackageName, &ActorClassName))
+		{
+			ActorClassName = *ActorMetaDataClass;
+		}
+
+		// Look for a class redirectors
+		const FCoreRedirectObjectName OldClassName = FCoreRedirectObjectName(*ActorClassName, NAME_None, *ActorPackageName);
+		const FCoreRedirectObjectName NewClassName = FCoreRedirects::GetRedirectedName(ECoreRedirectFlags::Type_Class, OldClassName);
+
+		bool bIsValidClass = true;
+		UClass* ActorClass = UClass::TryFindTypeSlow<UClass>(NewClassName.ToString(), EFindFirstObjectOptions::ExactClass);
+
+		if (!ActorClass)
+		{
+			ActorClass = AActor::StaticClass();
+			bIsValidClass = false;
+		}
+
+		FWorldPartitionActorDescInitData ActorDescInitData;
+		ActorDescInitData.NativeClass = ActorClass;
+		ActorDescInitData.PackageName = InAssetData.PackageName;
+		ActorDescInitData.ActorPath = InAssetData.ObjectPath;
+		FBase64::Decode(ActorMetaDataStr, ActorDescInitData.SerializedData);
+
+		TUniquePtr<FWorldPartitionActorDesc> NewActorDesc(AActor::StaticCreateClassActorDesc(ActorDescInitData.NativeClass));
+
+		NewActorDesc->Init(ActorDescInitData);
+			
+		if (!bIsValidClass)
+		{
+			UE_LOG(LogWorldPartition, Warning, TEXT("Invalid class `%s` for actor guid `%s` ('%s') from package '%s'"), *NewClassName.ToString(), *NewActorDesc->GetGuid().ToString(), *NewActorDesc->GetActorName().ToString(), *NewActorDesc->GetActorPackage().ToString());
+			return nullptr;
+		}
+
+		return NewActorDesc;
+	}
+
+	return nullptr;
+}
+
+void FWorldPartitionActorDescUtils::AppendAssetDataTagsFromActor(const AActor* Actor, TArray<UObject::FAssetRegistryTag>& OutTags)
+{
+	check(Actor->IsPackageExternal());
+	
+	TUniquePtr<FWorldPartitionActorDesc> ActorDesc(Actor->CreateActorDesc());
+
+	// If the actor is not added to a world, we can't retrieve its bounding volume, so try to get the existing one
+	if (ULevel* Level = Actor->GetLevel(); !Level || !Level->Actors.Contains(Actor))
+	{
+		IAssetRegistry& AssetRegistry = FModuleManager::LoadModuleChecked<FAssetRegistryModule>(TEXT("AssetRegistry")).Get();
+
+		FARFilter Filter;
+		Filter.bIncludeOnlyOnDiskAssets = true;
+		Filter.PackageNames.Add(Actor->GetPackage()->GetFName());
+
+		TArray<FAssetData> Assets;
+		AssetRegistry.GetAssets(Filter, Assets);
+
+		if (Assets.Num() == 1)
+		{
+			if (TUniquePtr<FWorldPartitionActorDesc> NewActorDesc = FWorldPartitionActorDescUtils::GetActorDescriptorFromAssetData(Assets[0]))
+			{
+				ActorDesc->TransferWorldData(NewActorDesc.Get());
+			}
+		}
+	}
+
+	const FString ActorMetaDataClass = GetParentNativeClass(Actor->GetClass())->GetPathName();
+	OutTags.Add(UObject::FAssetRegistryTag(NAME_ActorMetaDataClass, ActorMetaDataClass, UObject::FAssetRegistryTag::TT_Hidden));
+
+	TArray<uint8> SerializedData;
+	ActorDesc->SerializeTo(SerializedData);
+	const FString ActorMetaData = FBase64::Encode(SerializedData);
+	OutTags.Add(UObject::FAssetRegistryTag(NAME_ActorMetaData, ActorMetaData, UObject::FAssetRegistryTag::TT_Hidden));
+}
+
+void FWorldPartitionActorDescUtils::UpdateActorDescriptorFomActor(const AActor* Actor, TUniquePtr<FWorldPartitionActorDesc>& ActorDesc)
+{
+	TUniquePtr<FWorldPartitionActorDesc> NewActorDesc(Actor->CreateActorDesc());
+	NewActorDesc->TransferFrom(ActorDesc.Get());
+	ActorDesc = MoveTemp(NewActorDesc);
+}
+
+void FWorldPartitionActorDescUtils::ReplaceActorDescriptorPointerFromActor(const AActor* OldActor, AActor* NewActor, FWorldPartitionActorDesc* ActorDesc)
+{
+	check(OldActor->GetActorGuid() == NewActor->GetActorGuid());
+	check(NewActor->GetActorGuid() == ActorDesc->GetGuid());
+	check(!ActorDesc->ActorPtr.IsValid() || (ActorDesc->ActorPtr == OldActor));
+	ActorDesc->ActorPtr = NewActor;
+}
+#endif
\ No newline at end of file
diff --git a/Engine/Source/Runtime/Engine/Public/WorldPartition/WorldPartitionActorDesc.h b/Engine/Source/Runtime/Engine/Public/WorldPartition/WorldPartitionActorDesc.h
--- a/Engine/Source/Runtime/Engine/Public/WorldPartition/WorldPartitionActorDesc.h
+++ b/Engine/Source/Runtime/Engine/Public/WorldPartition/WorldPartitionActorDesc.h
@@ -61,10 +61,10 @@ inline bool CompareUnsortedArrays(const TArray<FName>& Array1, const TArray<FNam
 class ENGINE_API FWorldPartitionActorDesc
 {
 	friend class AActor;
-	friend class UActorDescContainer;
 	friend class UWorldPartition;
 	friend struct FWorldPartitionHandleImpl;
 	friend struct FWorldPartitionReferenceImpl;
+	friend struct FWorldPartitionActorDescUtils;
 
 public:
 	virtual ~FWorldPartitionActorDesc() {}
diff --git a/Engine/Source/Runtime/Engine/Public/WorldPartition/WorldPartitionActorDescUtils.h b/Engine/Source/Runtime/Engine/Public/WorldPartition/WorldPartitionActorDescUtils.h
new file mode 100644
--- /dev/null
+++ b/Engine/Source/Runtime/Engine/Public/WorldPartition/WorldPartitionActorDescUtils.h
@@ -0,0 +1,20 @@
+// Copyright Epic Games, Inc. All Rights Reserved.
+#pragma once
+
+#if WITH_EDITOR
+#include "CoreMinimal.h"
+#include "UObject/Object.h"
+#include "WorldPartition/WorldPartitionActorDesc.h"
+
+class AActor;
+struct FAssetData;
+
+struct ENGINE_API FWorldPartitionActorDescUtils
+{
+	static bool IsValidActorDescriptorFromAssetData(const FAssetData& InAssetData);
+	static TUniquePtr<FWorldPartitionActorDesc> GetActorDescriptorFromAssetData(const FAssetData& InAssetData);
+	static void AppendAssetDataTagsFromActor(const AActor* Actor, TArray<UObject::FAssetRegistryTag>& OutTags);
+	static void UpdateActorDescriptorFomActor(const AActor* Actor, TUniquePtr<FWorldPartitionActorDesc>& ActorDesc);
+	static void ReplaceActorDescriptorPointerFromActor(const AActor* OldActor, AActor* NewActor, FWorldPartitionActorDesc* ActorDesc);
+};
+#endif
\ No newline at end of file
