Only emit ISPC deprecation warning message when ISPC is compiled for that platform

diff --git a/Engine/Source/Runtime/Engine/Private/SkinnedAsset.cpp b/Engine/Source/Runtime/Engine/Private/SkinnedAsset.cpp
--- a/Engine/Source/Runtime/Engine/Private/SkinnedAsset.cpp
+++ b/Engine/Source/Runtime/Engine/Private/SkinnedAsset.cpp
@@ -164,12 +164,19 @@ static bool IsISPCEnabled()
 	// so when it is false it means one of the two settings are changed. To be backwards compatible, its overridden value takes priority.
 	if (!bAnim_SkeletalMesh_ISPC_Enabled)
 	{
-		UE_LOG(
-			LogSkinnedAsset,
-			Log,
-			TEXT("ANIM_SKELETAL_MESH_ISPC_ENABLED_DEFAULT and CVar 'a.SkeletalMesh.ISPC' are deprecated and will be removed in future releases."
-				 "Please switch to ANIM_SKINNED_ASSET_ISPC_ENABLED_DEFAULT and 'a.SkinnedAsset.ISPC'")
-		);
+#if INTEL_ISPC
+		static bool bWarningLogged = false;
+		if (!bWarningLogged)
+		{
+			UE_LOG(
+				LogSkinnedAsset,
+				Warning,
+				TEXT("ANIM_SKELETAL_MESH_ISPC_ENABLED_DEFAULT and CVar 'a.SkeletalMesh.ISPC' are deprecated and will be removed in future releases."
+					"Please switch to ANIM_SKINNED_ASSET_ISPC_ENABLED_DEFAULT and 'a.SkinnedAsset.ISPC'")
+			);
+			bWarningLogged = true;
+		}
+#endif
 		return false;
 	}
 
