Remove render data check from USkeletalMeshComponent::ComponentIsTouchingSelectionFrustum and ComponentIsTouchingSelectionBox

diff --git a/Engine/Source/Runtime/Engine/Private/Components/SkeletalMeshComponent.cpp b/Engine/Source/Runtime/Engine/Private/Components/SkeletalMeshComponent.cpp
--- a/Engine/Source/Runtime/Engine/Private/Components/SkeletalMeshComponent.cpp
+++ b/Engine/Source/Runtime/Engine/Private/Components/SkeletalMeshComponent.cpp
@@ -3615,38 +3615,38 @@ bool USkeletalMeshComponent::ComponentIsTouchingSelectionBox(const FBox& InSelBB
 	if (!bConsiderOnlyBSP && MeshObject != nullptr)
 	{
 		FSkeletalMeshRenderData* SkelMeshRenderData = GetSkeletalMeshRenderData();
-		check(SkelMeshRenderData);
-		check(SkelMeshRenderData->LODRenderData.Num() > 0);
-
-		// Transform verts into world space. Note that this assumes skeletal mesh is in reference pose...
-		const FSkeletalMeshLODRenderData& LODData = SkelMeshRenderData->LODRenderData[0];
-		for (uint32 VertIdx=0; VertIdx<LODData.StaticVertexBuffers.PositionVertexBuffer.GetNumVertices(); VertIdx++)
+		if (SkelMeshRenderData && SkelMeshRenderData->LODRenderData.Num() > 0)
 		{
-			const FVector3f& VertexPos(LODData.StaticVertexBuffers.PositionVertexBuffer.VertexPosition(VertIdx));
-			const FVector Location = GetComponentTransform().TransformPosition((FVector)VertexPos);
-			const bool bLocationIntersected = FMath::PointBoxIntersection(Location, InSelBBox);
-
-			// If the selection box doesn't have to encompass the entire component and a skeletal mesh vertex has intersected with
-			// the selection box, this component is being touched by the selection box
-			if (!bMustEncompassEntireComponent && bLocationIntersected)
+			// Transform verts into world space. Note that this assumes skeletal mesh is in reference pose...
+			const FSkeletalMeshLODRenderData& LODData = SkelMeshRenderData->LODRenderData[0];
+			for (uint32 VertIdx = 0; VertIdx < LODData.StaticVertexBuffers.PositionVertexBuffer.GetNumVertices(); VertIdx++)
 			{
-				return true;
+				const FVector3f& VertexPos(LODData.StaticVertexBuffers.PositionVertexBuffer.VertexPosition(VertIdx));
+				const FVector Location = GetComponentTransform().TransformPosition((FVector)VertexPos);
+				const bool bLocationIntersected = FMath::PointBoxIntersection(Location, InSelBBox);
+
+				// If the selection box doesn't have to encompass the entire component and a skeletal mesh vertex has intersected with
+				// the selection box, this component is being touched by the selection box
+				if (!bMustEncompassEntireComponent && bLocationIntersected)
+				{
+					return true;
+				}
+
+				// If the selection box has to encompass the entire component and a skeletal mesh vertex didn't intersect with the selection
+				// box, this component does not qualify
+				else if (bMustEncompassEntireComponent && !bLocationIntersected)
+				{
+					return false;
+				}
 			}
 
-			// If the selection box has to encompass the entire component and a skeletal mesh vertex didn't intersect with the selection
-			// box, this component does not qualify
-			else if (bMustEncompassEntireComponent && !bLocationIntersected)
+			// If the selection box has to encompass all of the component and none of the component's verts failed the intersection test, this component
+			// is consider touching
+			if (bMustEncompassEntireComponent)
 			{
-				return false;
+				return true;
 			}
 		}
-
-		// If the selection box has to encompass all of the component and none of the component's verts failed the intersection test, this component
-		// is consider touching
-		if (bMustEncompassEntireComponent)
-		{
-			return true;
-		}
 	}
 
 	return false;
@@ -3657,35 +3657,35 @@ bool USkeletalMeshComponent::ComponentIsTouchingSelectionFrustum(const FConvexVo
 	if (!bConsiderOnlyBSP && MeshObject != nullptr)
 	{
 		FSkeletalMeshRenderData* SkelMeshRenderData = GetSkeletalMeshRenderData();
-		check(SkelMeshRenderData);
-		check(SkelMeshRenderData->LODRenderData.Num() > 0);
-
-		// Transform verts into world space. Note that this assumes skeletal mesh is in reference pose...
-		const FSkeletalMeshLODRenderData& LODData = SkelMeshRenderData->LODRenderData[0];
-		for (uint32 VertIdx = 0; VertIdx < LODData.StaticVertexBuffers.PositionVertexBuffer.GetNumVertices(); VertIdx++)
+		if (SkelMeshRenderData && SkelMeshRenderData->LODRenderData.Num() > 0)
 		{
-			const FVector3f& VertexPos(LODData.StaticVertexBuffers.PositionVertexBuffer.VertexPosition(VertIdx));
-			const FVector Location = GetComponentTransform().TransformPosition((FVector)VertexPos);
-			const bool bLocationIntersected = InFrustum.IntersectSphere(Location, 0.0f);
-
-			// If the selection box doesn't have to encompass the entire component and a skeletal mesh vertex has intersected with
-			// the selection box, this component is being touched by the selection box
-			if (!bMustEncompassEntireComponent && bLocationIntersected)
+			// Transform verts into world space. Note that this assumes skeletal mesh is in reference pose...
+			const FSkeletalMeshLODRenderData& LODData = SkelMeshRenderData->LODRenderData[0];
+			for (uint32 VertIdx = 0; VertIdx < LODData.StaticVertexBuffers.PositionVertexBuffer.GetNumVertices(); VertIdx++)
 			{
-				return true;
-			}
+				const FVector3f& VertexPos(LODData.StaticVertexBuffers.PositionVertexBuffer.VertexPosition(VertIdx));
+				const FVector Location = GetComponentTransform().TransformPosition((FVector)VertexPos);
+				const bool bLocationIntersected = InFrustum.IntersectSphere(Location, 0.0f);
 
-			// If the selection box has to encompass the entire component and a skeletal mesh vertex didn't intersect with the selection
-			// box, this component does not qualify
-			else if (bMustEncompassEntireComponent && !bLocationIntersected)
-			{
-				return false;
+				// If the selection box doesn't have to encompass the entire component and a skeletal mesh vertex has intersected with
+				// the selection box, this component is being touched by the selection box
+				if (!bMustEncompassEntireComponent && bLocationIntersected)
+				{
+					return true;
+				}
+
+				// If the selection box has to encompass the entire component and a skeletal mesh vertex didn't intersect with the selection
+				// box, this component does not qualify
+				else if (bMustEncompassEntireComponent && !bLocationIntersected)
+				{
+					return false;
+				}
 			}
-		}
 
-		// If the selection box has to encompass all of the component and none of the component's verts failed the intersection test, this component
-		// is consider touching
-		return true;
+			// If the selection box has to encompass all of the component and none of the component's verts failed the intersection test, this component
+			// is consider touching
+			return true;
+		}
 	}
 
 	return false;
