Automation Improvements for running editor-based automation, and specifically profiling

* Added IgnoreLogEvents automation command / cvar that can be used to exclude log events from automation reports. Specifically here this prevents errors/warnings in something like PIE from manigesting as automation events
* Automation Controller Manager now waits for an interactive framerate before executing. This ensures automation doesn't occur when async loading is being flushedcritical to startup or level loading and f
* Renamed FWaitForAverageFrameRate to FWaitForInteractiveFramerate and moved default values to be ini-driven
* Editor now disables throttling when -benchmarking is specified
* FIxed issue where Project.Maps.PIE only worked if the provided map was a full object name (/Game/Path/MyMap.MyMap)

diff --git a/Engine/Source/Developer/AutomationController/Private/AutomationCommandline.cpp b/Engine/Source/Developer/AutomationController/Private/AutomationCommandline.cpp
--- a/Engine/Source/Developer/AutomationController/Private/AutomationCommandline.cpp
+++ b/Engine/Source/Developer/AutomationController/Private/AutomationCommandline.cpp
@@ -13,7 +13,6 @@
 #include "HAL/FileManager.h"
 #include "Misc/Paths.h"
 #include "Misc/FileHelper.h"
-#include "AssetRegistry/AssetRegistryModule.h"
 #include "AutomationControllerSettings.h"
 #include "AutomationGroupFilter.h"
 #include "Containers/Ticker.h"
@@ -385,10 +384,10 @@ public:
 		{
 			case EAutomationTestState::Initializing:
 			{
-				FAssetRegistryModule& AssetRegistryModule = FModuleManager::LoadModuleChecked<FAssetRegistryModule>("AssetRegistry");
-				if ( AssetRegistryModule.Get().IsLoadingAssets() == false )
+				if (AutomationController->IsReadyForTests())
 				{
 					AutomationTestState = EAutomationTestState::Idle;
+					UE_LOG(LogAutomationCommandLine, Display, TEXT("Ready to start automation"));
 				}
 				FindWorkerAttempts = 0;
 				break;
@@ -609,6 +608,14 @@ public:
 					AutomationCommandQueue.Add(EAutomationCommand::Quit);
 					Ar.Logf(TEXT("Automation: Quit Command Queued."));
 				}
+				else if (FParse::Command(&TempCmd, TEXT("IgnoreLogEvents")))
+				{
+					if (IConsoleVariable* CVar = IConsoleManager::Get().FindConsoleVariable(TEXT("Automation.CaptureLogEvents")))
+					{
+						Ar.Logf(TEXT("Automation: Suppressing Log Events"));
+						CVar->Set(false);
+					}					
+				}
 				else
 				{
 					Ar.Logf(TEXT("Incorrect automation command syntax! Supported commands are: "));
diff --git a/Engine/Source/Developer/AutomationController/Private/AutomationControllerManager.cpp b/Engine/Source/Developer/AutomationController/Private/AutomationControllerManager.cpp
--- a/Engine/Source/Developer/AutomationController/Private/AutomationControllerManager.cpp
+++ b/Engine/Source/Developer/AutomationController/Private/AutomationControllerManager.cpp
@@ -23,13 +23,14 @@
 #include "Misc/EngineVersion.h"
 #include "Misc/FileHelper.h"
 #include "PlatformHttp.h"
-
 #include "AutomationTelemetry.h"
+#include "AssetRegistry/AssetRegistryModule.h"
 
 #if WITH_EDITOR
 #include "UnrealEdGlobals.h"
 #include "Editor/UnrealEdEngine.h"
 #include "Logging/MessageLog.h"
+#include "Tests/AutomationCommon.h"
 #endif
 #include "Async/ParallelFor.h"
 
@@ -201,6 +202,31 @@ FAutomationControllerManager::FAutomationControllerManager()
 	bResumeRunTest = FParse::Param(FCommandLine::Get(), TEXT("ResumeRunTest"));
 }
 
+bool FAutomationControllerManager::IsReadyForTests()
+{
+	FAssetRegistryModule& AssetRegistryModule = FModuleManager::LoadModuleChecked<FAssetRegistryModule>("AssetRegistry");
+
+	if (AssetRegistryModule.Get().IsLoadingAssets())
+	{
+		return false;
+	}
+
+#if WITH_EDITOR
+	if (InteractiveFrameRateCheck == nullptr)
+	{
+		InteractiveFrameRateCheck = MakeShared<FWaitForInteractiveFrameRate>();
+	}
+
+	if (!InteractiveFrameRateCheck->Update())
+	{
+		return false;
+	}
+
+	InteractiveFrameRateCheck = nullptr;
+#endif // WITH_EDITOR
+	return true;
+}
+
 void FAutomationControllerManager::RequestAvailableWorkers(const FGuid& SessionId)
 {
 	//invalidate previous tests
diff --git a/Engine/Source/Developer/AutomationController/Private/AutomationControllerManager.h b/Engine/Source/Developer/AutomationController/Private/AutomationControllerManager.h
--- a/Engine/Source/Developer/AutomationController/Private/AutomationControllerManager.h
+++ b/Engine/Source/Developer/AutomationController/Private/AutomationControllerManager.h
@@ -201,6 +201,7 @@ public:
 
 	// IAutomationController Interface
 	virtual void RequestAvailableWorkers( const FGuid& InSessionId ) override;
+	virtual bool IsReadyForTests() override;
 	virtual void RequestTests() override;
 	virtual void RunTests( const bool bIsLocalSession) override;
 	virtual void StopTests() override;
@@ -604,6 +605,10 @@ private:
 
 	bool bResumeRunTest;
 
+#if WITH_EDITOR
+	TSharedPtr<class FWaitForInteractiveFrameRate> InteractiveFrameRateCheck;
+#endif
+
 private:
 
 	/** Holds a delegate that is invoked when the controller shuts down. */
diff --git a/Engine/Source/Developer/AutomationController/Public/IAutomationControllerManager.h b/Engine/Source/Developer/AutomationController/Public/IAutomationControllerManager.h
--- a/Engine/Source/Developer/AutomationController/Public/IAutomationControllerManager.h
+++ b/Engine/Source/Developer/AutomationController/Public/IAutomationControllerManager.h
@@ -167,6 +167,12 @@ public:
 	 */
 	virtual void RequestAvailableWorkers( const FGuid& InSessionId ) = 0;
 
+	/** 
+	 * Return true when the manager is ready and capable of executing tests. 
+	 * if RunTests is called earlier it may block until internally these conditions are met
+	 */
+	virtual bool IsReadyForTests() = 0;
+
 	/** Send request for tests that are available to run. */
 	virtual void RequestTests() = 0;
 
diff --git a/Engine/Source/Editor/UnrealEd/Private/EditorEngine.cpp b/Engine/Source/Editor/UnrealEd/Private/EditorEngine.cpp
--- a/Engine/Source/Editor/UnrealEd/Private/EditorEngine.cpp
+++ b/Engine/Source/Editor/UnrealEd/Private/EditorEngine.cpp
@@ -6215,7 +6215,7 @@ bool UEditorEngine::IsAnyViewportRealtime()
 
 bool UEditorEngine::ShouldThrottleCPUUsage() const
 {
-	if (IsRunningCommandlet())
+	if (IsRunningCommandlet() || FApp::IsBenchmarking())
 	{
 		return false;
 	}
diff --git a/Engine/Source/Editor/UnrealEd/Private/PerformanceMonitor.cpp b/Engine/Source/Editor/UnrealEd/Private/PerformanceMonitor.cpp
--- a/Engine/Source/Editor/UnrealEd/Private/PerformanceMonitor.cpp
+++ b/Engine/Source/Editor/UnrealEd/Private/PerformanceMonitor.cpp
@@ -283,7 +283,9 @@ void FPerformanceMonitor::Tick(float DeltaTime)
 	FineMovingAverage.Tick(FPlatformTime::Seconds(), GAverageFPS);
 	CoarseMovingAverage.Tick(FPlatformTime::Seconds(), GAverageFPS);
 
-	bool bMonitorEditorPerformance = GetDefault<UEditorPerformanceSettings>()->bMonitorEditorPerformance;
+	bool bMonitorEditorPerformance = 
+		GetDefault<UEditorPerformanceSettings>()->bMonitorEditorPerformance && FApp::IsBenchmarking() == false;
+
 	if( !bMonitorEditorPerformance || !bIsNotificationAllowed )
 	{
 		return;
diff --git a/Engine/Source/Editor/UnrealEd/Private/Tests/AutomationEditorCommon.cpp b/Engine/Source/Editor/UnrealEd/Private/Tests/AutomationEditorCommon.cpp
--- a/Engine/Source/Editor/UnrealEd/Private/Tests/AutomationEditorCommon.cpp
+++ b/Engine/Source/Editor/UnrealEd/Private/Tests/AutomationEditorCommon.cpp
@@ -1166,6 +1166,7 @@ IMPLEMENT_SIMPLE_AUTOMATION_TEST(FProjectMapsPIETest, "Project.Maps.PIE", EAutom
 
 	const TCHAR* ParsedCmdLine = FCommandLine::Get();
 	FString ParsedMapName;
+	bool FirstMapAlreadyLoaded = false;
 
 	// If there is an explicit list of maps on the command line via -map or -maps the use those.
 	if (FParse::Value(FCommandLine::Get(), TEXT("-maps="), ParsedMapName) || FParse::Value(FCommandLine::Get(), TEXT("-map="), ParsedMapName))
@@ -1185,6 +1186,7 @@ IMPLEMENT_SIMPLE_AUTOMATION_TEST(FProjectMapsPIETest, "Project.Maps.PIE", EAutom
 			FPaths::GetExtension(InitialMapName, /*bIncludeDot=*/true) == FPackageName::GetMapPackageExtension())
 		{
 			PIEMaps.Add(InitialMapName);
+			FirstMapAlreadyLoaded = true;
 			UE_LOG(LogEditorAutomationTests, Display, TEXT("Found Map %s on command line. PIE Test will be restricted to this map"), *InitialMapName);
 		}
 	}
@@ -1233,29 +1235,42 @@ IMPLEMENT_SIMPLE_AUTOMATION_TEST(FProjectMapsPIETest, "Project.Maps.PIE", EAutom
 	
 	for (const FString& Map : PIEMaps)
 	{
-		FString MapPackageName = FPackageName::ObjectPathToPackageName(Map);
+		// Accept any of...
+		// - MyMap
+		// - /Game/MyMap
+		// - /Game/MyMap.MyMap
+		FString MapPackageName = Map;
 
-		if (!FPackageName::IsValidObjectPath(MapPackageName))
+		if (FPackageName::IsValidObjectPath(Map))
 		{
-			if (!FPackageName::SearchForPackageOnDisk(MapPackageName, NULL, &MapPackageName))
-			{
-				UE_LOG(LogEditorAutomationTests, Error, TEXT("Couldn't resolve map for PIE test from %s to valid package name!"), *MapPackageName);
-				continue;
-			}
-		}		
+			MapPackageName = FPackageName::ObjectPathToPackageName(Map);
+		}
+
+		if (!FPackageName::SearchForPackageOnDisk(Map, NULL, &MapPackageName))
+		{
+			UE_LOG(LogEditorAutomationTests, Error, TEXT("Couldn't resolve map for PIE test from %s to valid package name!"), *MapPackageName);
+			continue;
+		}
+
+		UE_LOG(LogEditorAutomationTests, Display, TEXT("Queueing Map %s for PIE Automation"), *MapPackageName);
 		
 		AddCommand(new FEditorAutomationLogCommand(FString::Printf(TEXT("LoadMap-Begin: %s"), *MapPackageName)));
-		AddCommand(new FEditorLoadMap(MapPackageName));
+		if (!FirstMapAlreadyLoaded)
+		{
+			AddCommand(new FEditorLoadMap(MapPackageName));
+		}
 		AddCommand(new FWaitLatentCommand(1.0f));
 		AddCommand(new FEditorAutomationLogCommand(FString::Printf(TEXT("LoadMap-End: %s"), *MapPackageName)));
 		AddCommand(new FEditorAutomationLogCommand(FString::Printf(TEXT("PIE-Begin: %s"), *MapPackageName)));
 		AddCommand(new FStartPIECommand(false));
 		AddCommand(new FWaitForSpecifiedMapToLoadCommand(MapPackageName));  // need at least some frames before starting & ending PIE
-		AddCommand(new FWaitForAverageFrameRate(5.0f));	// wait until the editor reaches something vaguely usable
+		AddCommand(new FWaitForInteractiveFrameRate());	// wait until the editor reaches something vaguely usable
 		AddCommand(new FWaitLatentCommand(AutomationTestSettings->PIETestDuration));
 		AddCommand(new FEndPlayMapCommand());
 		AddCommand(new FWaitForSpecifiedPIEMapToEndCommand(MapPackageName));  // need at least some frames before starting & ending PIE
-		AddCommand(new FEditorAutomationLogCommand(FString::Printf(TEXT("PIE-End: %s"), *Map)));
+		AddCommand(new FEditorAutomationLogCommand(FString::Printf(TEXT("PIE-End: %s"), *MapPackageName)));
+
+		FirstMapAlreadyLoaded = false;
 	}
 
 	return true;
diff --git a/Engine/Source/Runtime/Core/Private/Misc/AutomationTest.cpp b/Engine/Source/Runtime/Core/Private/Misc/AutomationTest.cpp
--- a/Engine/Source/Runtime/Core/Private/Misc/AutomationTest.cpp
+++ b/Engine/Source/Runtime/Core/Private/Misc/AutomationTest.cpp
@@ -18,6 +18,13 @@
 
 DEFINE_LOG_CATEGORY_STATIC(LogAutomationTest, Warning, All);
 
+static TAutoConsoleVariable<bool> CVarAutomationCaptureLogEvents(
+	TEXT("Automation.CaptureLogEvents"),
+	1,
+	TEXT("Consider warning/error log events during a test as impacting the test itself"),
+	ECVF_ReadOnly
+);
+
 CORE_API const TMap<FString, EAutomationTestFlags::Type>& EAutomationTestFlags::GetTestFlagsMap()
 {
 	LLM_SCOPE_BYNAME(TEXT("AutomationTest/Framework"));
@@ -68,6 +75,11 @@ CORE_API ELogVerbosity::Type GetAutomationLogLevel(ELogVerbosity::Type LogVerbos
 	static TArray<FString> SuppressedLogCategories;
 	static FAutomationTestBase* LastTest = nullptr;
 
+	if (CVarAutomationCaptureLogEvents.GetValueOnGameThread() == false)
+	{
+		return ELogVerbosity::NoLogging;
+	}
+
 	if (CurrentTest != LastTest)
 	{
 		// These can be changed in the editor so can't just be cached for the whole session
@@ -153,7 +165,7 @@ void FAutomationTestFramework::FAutomationTestOutputDevice::Serialize( const TCH
 				LocalCurTest->AddWarning(FormattedMsg, STACK_OFFSET);
 			}
 			// Display
-			else
+			else if (EffectiveVerbosity != ELogVerbosity::NoLogging)
 			{
 				LocalCurTest->AddInfo(FormattedMsg, STACK_OFFSET);
 			}
diff --git a/Engine/Source/Runtime/Engine/Classes/Tests/AutomationTestSettings.h b/Engine/Source/Runtime/Engine/Classes/Tests/AutomationTestSettings.h
--- a/Engine/Source/Runtime/Engine/Classes/Tests/AutomationTestSettings.h
+++ b/Engine/Source/Runtime/Engine/Classes/Tests/AutomationTestSettings.h
@@ -444,4 +444,23 @@ public:
 	 */
 	UPROPERTY(EditAnywhere, config, Category = Automation)
 	float PIETestDuration;
+
+	/**
+	 * Default value used for FWaitForInteractiveFrameRate. This is a framerate determine to be suitably "interactive", but may be
+	 * less than the target framerate of the game as this is used for evaluating time to PIE, load the editor etc
+	 */
+	UPROPERTY(EditAnywhere, config, Category = Automation)
+	float DefaultInteractiveFramerate;
+
+	/**
+	 * Default wait time in seconds for FWaitForInteractiveFrameRate. After this time a test will fail.
+	 */
+	UPROPERTY(EditAnywhere, config, Category = Automation)
+	float DefaultInteractiveFramerateWaitTime;
+
+	/**
+	 *  Default time in seconds that DefaultInteractiveFramerate must remain true in FWaitForInteractiveFrameRate
+	 */
+	UPROPERTY(EditAnywhere, config, Category = Automation)
+	float DefaultInteractiveFramerateDuration;
 };
diff --git a/Engine/Source/Runtime/Engine/Private/Tests/AutomationCommon.cpp b/Engine/Source/Runtime/Engine/Private/Tests/AutomationCommon.cpp
--- a/Engine/Source/Runtime/Engine/Private/Tests/AutomationCommon.cpp
+++ b/Engine/Source/Runtime/Engine/Private/Tests/AutomationCommon.cpp
@@ -486,33 +486,84 @@ bool FWaitForSpecifiedMapToLoadCommand::Update()
 	return false;
 }
 
-bool FWaitForAverageFrameRate::Update()
+FWaitForInteractiveFrameRate::FWaitForInteractiveFrameRate(float InDesiredFrameRate, float InDuration, float InMaxWaitTime)
+	: DesiredFrameRate(InDesiredFrameRate)
+	, Duration(InDuration)
+	, MaxWaitTime(InMaxWaitTime)
+	, StartTimeOfWait(0)
+	, StartTimeOfAcceptableFrameRate(0)
 {
-	if (StartTime == 0)
+	UAutomationTestSettings const* AutomationTestSettings = GetDefault<UAutomationTestSettings>();
+
+	if (DesiredFrameRate == 0)
 	{
-		StartTime = FPlatformTime::Seconds();
+		DesiredFrameRate = AutomationTestSettings->DefaultInteractiveFramerate;
 	}
-	else
+
+	if (Duration == 0)
+	{
+		Duration = AutomationTestSettings->DefaultInteractiveFramerateDuration;
+	}
+
+	if (MaxWaitTime == 0)
+	{
+		MaxWaitTime = AutomationTestSettings->DefaultInteractiveFramerateWaitTime;
+	}
+}
+
+
+
+bool FWaitForInteractiveFrameRate::Update()
+{
+	const double TimeNow = FPlatformTime::Seconds();
+
+	if (StartTimeOfWait == 0)
 	{
-		const double ElapsedTime = FPlatformTime::Seconds() - StartTime;
+		UE_LOG(LogEngineAutomationLatentCommand, Log, TEXT("FWaitForInteractiveFrameRate: Starting wait for framerate of >= %.f FPS"), DesiredFrameRate);
+		StartTimeOfWait = TimeNow;
+	}
+
+	// Check if we're at the desired framerate, and if so how long we have been there for
+	extern ENGINE_API float GAverageFPS;
+	bool AtDesiredFramerate = GAverageFPS >= DesiredFrameRate;
+
+	const double ElapsedWaitTime = TimeNow - StartTimeOfWait;
 
-		if (ElapsedTime > Delay)
+	if (!AtDesiredFramerate)
+	{
+		StartTimeOfAcceptableFrameRate = 0;		
+	}
+	else
+	{
+		if (StartTimeOfAcceptableFrameRate == 0)
 		{
-			extern ENGINE_API float GAverageFPS;
-			if (GAverageFPS >= DesiredFrameRate)
-			{
-				return true;
-			}
+			StartTimeOfAcceptableFrameRate = TimeNow;
+		}
+		else
+		{
+			const double TimeAtFrameRate = TimeNow - StartTimeOfAcceptableFrameRate;
 
-			if (ElapsedTime >= MaxWaitTime)
+			if (TimeAtFrameRate >= Duration)
 			{
-				UE_LOG(LogEngineAutomationLatentCommand, Error, TEXT("FWaitForAverageFrameRate: Game did not reach %.02f FPS within %.02f seconds. Giving up."), DesiredFrameRate, MaxWaitTime);
-				
+				UE_LOG(LogEngineAutomationLatentCommand, Log, TEXT("FWaitForInteractiveFrameRate: Hit %.02f FPS for %.f seconds after %.f seconds of waiting"), GAverageFPS, TimeAtFrameRate, ElapsedWaitTime);
 				return true;
 			}
 		}
 	}
 
+	if (ElapsedWaitTime >= MaxWaitTime)
+	{
+		FString Msg = FString::Printf(TEXT("FWaitForInteractiveFrameRate: Game did not reach %.02f FPS within %.02f seconds. Current FPS=%.f. Giving up."), DesiredFrameRate, MaxWaitTime, GAverageFPS);
+		UE_LOG(LogEngineAutomationLatentCommand, Display, TEXT("%s"), *Msg);
+
+		if (FAutomationTestBase* CurrentTest = FAutomationTestFramework::Get().GetCurrentTest())
+		{
+			CurrentTest->AddError(Msg);
+		}
+
+		return true;
+	}
+
 	return false;
 }
 
@@ -613,15 +664,15 @@ bool FWaitForShadersToFinishCompilingInGame::Update()
 #if WITH_EDITOR
 	static double TimeShadersFinishedCompiling = 0;
 	static double LastReportTime = FPlatformTime::Seconds();
-	const double TimeToWaitForJobs = 2.0;
+	static bool DidCompileSomething = false;
+	const double TimeToWaitForJobsToStart = 2.0;
 
 	bool ShadersCompiling = GShaderCompilingManager && GShaderCompilingManager->IsCompiling();
-	bool TexturesCompiling = FTextureCompilingManager::Get().GetNumRemainingTextures() > 0;
-
+	bool AssetsCompiling = FAssetCompilingManager::Get().GetNumRemainingAssets() > 0;
 	
 	double TimeNow = FPlatformTime::Seconds();
 
-	if (ShadersCompiling || TexturesCompiling)
+	if (ShadersCompiling || AssetsCompiling)
 	{
 		if (TimeNow - LastReportTime > 5.0)
 		{
@@ -632,13 +683,14 @@ bool FWaitForShadersToFinishCompilingInGame::Update()
 				UE_LOG(LogEditorAutomationTests, Log, TEXT("Waiting for %i shaders to finish."), GShaderCompilingManager->GetNumRemainingJobs() + GShaderCompilingManager->GetNumPendingJobs());
 			}
 
-			if (TexturesCompiling)
+			if (AssetsCompiling)
 			{
-				UE_LOG(LogEditorAutomationTests, Log, TEXT("Waiting for %i texures to finish."), FTextureCompilingManager::Get().GetNumRemainingTextures());
+				UE_LOG(LogEditorAutomationTests, Log, TEXT("Waiting for %i assets to finish."), FAssetCompilingManager::Get().GetNumRemainingAssets());
 			}
 		}
 
 		TimeShadersFinishedCompiling = 0;
+		DidCompileSomething = true;
 
 		return false;
 	}
@@ -649,14 +701,20 @@ bool FWaitForShadersToFinishCompilingInGame::Update()
 		TimeShadersFinishedCompiling = FPlatformTime::Seconds();
 	}
 
-	if (FPlatformTime::Seconds() - TimeShadersFinishedCompiling < TimeToWaitForJobs)
+	if (FPlatformTime::Seconds() - TimeShadersFinishedCompiling < TimeToWaitForJobsToStart)
 	{
 		return false;
 	}
 
 	// may not be necessary, but just double-check everything is finished and ready
 	GShaderCompilingManager->FinishAllCompilation();
-	UE_LOG(LogEditorAutomationTests, Log, TEXT("Done waiting for shaders to finish."));
+	FAssetCompilingManager::Get().FinishAllCompilation();
+
+	if (DidCompileSomething)
+	{
+		UE_LOG(LogEditorAutomationTests, Log, TEXT("Done waiting for shaders to finish."));
+		DidCompileSomething = false;
+	}
 #endif
 
 	return true;
diff --git a/Engine/Source/Runtime/Engine/Private/Tests/AutomationTestSettings.cpp b/Engine/Source/Runtime/Engine/Private/Tests/AutomationTestSettings.cpp
--- a/Engine/Source/Runtime/Engine/Private/Tests/AutomationTestSettings.cpp
+++ b/Engine/Source/Runtime/Engine/Private/Tests/AutomationTestSettings.cpp
@@ -9,4 +9,8 @@ UAutomationTestSettings::UAutomationTestSettings(const FObjectInitializer& Objec
 {
 	DefaultScreenshotResolution = FIntPoint(1920, 1080);
 	PIETestDuration = 5.f;
+	// These defaults are are also set in BaseEngine.ini
+	DefaultInteractiveFramerate = 10.f;
+	DefaultInteractiveFramerateWaitTime = 10.f * 60.f;		// 10 minutes - assume shaders and other DDC may need to be built
+	DefaultInteractiveFramerateDuration = 5.f;
 }
diff --git a/Engine/Source/Runtime/Engine/Public/Tests/AutomationCommon.h b/Engine/Source/Runtime/Engine/Public/Tests/AutomationCommon.h
--- a/Engine/Source/Runtime/Engine/Public/Tests/AutomationCommon.h
+++ b/Engine/Source/Runtime/Engine/Public/Tests/AutomationCommon.h
@@ -180,34 +180,34 @@ DEFINE_ENGINE_LATENT_AUTOMATION_COMMAND(FWaitForShadersToFinishCompilingInGame);
  * Waits until the average framerate meets to exceeds the specified value. Mostly intended as a way to ensure that a level load etc has completed
  * and an interactive framerate is present.
  *
- * InDelay is how long to wait before checking, InMaxWaitTIme is how long to wait before throwing an error
+ * If values are zero the defaults in [/Script/Engine.AutomationTestSetting] will be used
  */
-class ENGINE_API FWaitForAverageFrameRate : public IAutomationLatentCommand
+class ENGINE_API FWaitForInteractiveFrameRate : public IAutomationLatentCommand
 {
 public:
 
-	FWaitForAverageFrameRate(float InDesiredFrameRate, float InDelay = 1, float InMaxWaitTime = 60)
-		: StartTime(0)
-		, DesiredFrameRate(InDesiredFrameRate)
-		, Delay(InDelay)
-		, MaxWaitTime(InMaxWaitTime)
-	{}
+	FWaitForInteractiveFrameRate(float InDesiredFrameRate = 0, float InDuration = 0, float InMaxWaitTime = 0);
 
 	bool Update() override;
 
-private:
-
-	// Time we began executing
-	double StartTime;
+public:
 
 	// Framerate we want to see
 	float DesiredFrameRate;
 
-	// How long before we start loiking at FPS
-	float Delay;
+	// How long must we maintain this framerate
+	float Duration;
 
 	// Max time so spend waiting
 	float MaxWaitTime;
+
+private:
+
+	// Time we began executing
+	double StartTimeOfWait;
+
+	// how many seconds we've been at the desired framerate
+	double StartTimeOfAcceptableFrameRate;	
 };
 
 /**
