Change PIE seamless travel error to mention new experimental cvar and enable for shootergame

diff --git a/Engine/Source/Runtime/Engine/Private/GameModeBase.cpp b/Engine/Source/Runtime/Engine/Private/GameModeBase.cpp
--- a/Engine/Source/Runtime/Engine/Private/GameModeBase.cpp
+++ b/Engine/Source/Runtime/Engine/Private/GameModeBase.cpp
@@ -418,7 +418,7 @@ bool AGameModeBase::CanServerTravel(const FString& FURL, bool bAbsolute)
 	{
 		if (!UE::GameModeBase::Private::bAllowPIESeamlessTravel)
 		{
-			UE_LOG(LogGameMode, Warning, TEXT("CanServerTravel: Seamless travel currently NOT supported in single process PIE."));
+			UE_LOG(LogGameMode, Warning, TEXT("CanServerTravel: Seamless travel is not supported by default in PIE, set net.AllowPIESeamlessTravel=1 to enable experimental support."));
 			return false;
 		}
 	}
