Fix memory leak when composite data tables have conflicting rows, it needs to call the internal remove when merging tables
Implements PR #9379 (contributed by [REDACTED_PERSON])

diff --git a/Engine/Source/Runtime/Engine/Classes/Engine/DataTable.h b/Engine/Source/Runtime/Engine/Classes/Engine/DataTable.h
--- a/Engine/Source/Runtime/Engine/Classes/Engine/DataTable.h
+++ b/Engine/Source/Runtime/Engine/Classes/Engine/DataTable.h
@@ -91,6 +91,9 @@ protected:
 
 	/** Called to add rows to the data table */
 	ENGINE_API virtual void AddRowInternal(FName RowName, uint8* RowDataPtr);
+
+	/** Deletes the row memory */
+	ENGINE_API virtual void RemoveRowInternal(FName RowName);
 public:
 
 	virtual const TMap<FName, uint8*>& GetRowMap() const { return RowMap; }
diff --git a/Engine/Source/Runtime/Engine/Private/DataTable.cpp b/Engine/Source/Runtime/Engine/Private/DataTable.cpp
--- a/Engine/Source/Runtime/Engine/Private/DataTable.cpp
+++ b/Engine/Source/Runtime/Engine/Private/DataTable.cpp
@@ -422,6 +422,11 @@ void UDataTable::RemoveRow(FName RowName)
 {
 	DATATABLE_CHANGE_SCOPE();
 
+	RemoveRowInternal(RowName);
+}
+
+void UDataTable::RemoveRowInternal(FName RowName)
+{
 	UScriptStruct& EmptyUsingStruct = GetEmptyUsingStruct();
 
 	uint8* RowData = nullptr;
@@ -434,13 +439,14 @@ void UDataTable::RemoveRow(FName RowName)
 	}
 }
 
-	
 void UDataTable::AddRow(FName RowName, const FTableRowBase& RowData)
 {
 	DATATABLE_CHANGE_SCOPE();
 
 	UScriptStruct& EmptyUsingStruct = GetEmptyUsingStruct();
-	RemoveRow(RowName);
+
+	// We want to delete the row memory even for child classes that override remove
+	RemoveRowInternal(RowName);
 		
 	uint8* NewRawRowData = (uint8*)FMemory::Malloc(EmptyUsingStruct.GetStructureSize());
 	
