Shaders: Added more detail to the shader map names passed to DDC and partially fixed timing scopes

#rnx

diff --git a/Engine/Source/Runtime/Engine/Private/Materials/MaterialShader.cpp b/Engine/Source/Runtime/Engine/Private/Materials/MaterialShader.cpp
--- a/Engine/Source/Runtime/Engine/Private/Materials/MaterialShader.cpp
+++ b/Engine/Source/Runtime/Engine/Private/Materials/MaterialShader.cpp
@@ -199,6 +199,17 @@ static UE::DerivedData::FCacheKey GetMaterialShaderMapKey(const FStringView Mate
 	return {Bucket, FIoHash::HashBuffer(MakeMemoryView(FTCHARToUTF8(MaterialShaderMapKey)))};
 }
 
+static UE::DerivedData::FSharedString GetMaterialShaderMapName(const FStringView MaterialPath, const FMaterialShaderMapId& ShaderMapId, const EShaderPlatform Platform)
+{
+	FName FeatureLevelName;
+	GetFeatureLevelName(ShaderMapId.FeatureLevel, FeatureLevelName);
+	return UE::DerivedData::FSharedString(WriteToString<256>(MaterialPath,
+		TEXTVIEW(" ["), LegacyShaderPlatformToShaderFormat(Platform),
+		TEXTVIEW(", "), FeatureLevelName,
+		TEXTVIEW(", "), GetMaterialQualityLevelFName(ShaderMapId.QualityLevel),
+		TEXTVIEW("]")));
+}
+
 static FString FSHA1_HashString(const FString& Key)
 {
 	FSHAHash DDCKeyHash;
@@ -1492,21 +1503,22 @@ TSharedRef<FMaterialShaderMap::FAsyncLoadContext> FMaterialShaderMap::BeginLoadF
 		FSharedBuffer   CachedData;
 		EShaderPlatform Platform;
 		FString         AssetName;
-		TUniquePtr<FRequestOwner> RequestOwner;
+		FRequestOwner   RequestOwner{UE::DerivedData::EPriority::Normal};
 		TRefCountPtr<FMaterialShaderMap> ShaderMap;
 
 		bool IsReady() const override
 		{
-			return RequestOwner == nullptr || RequestOwner->Poll();
+			return RequestOwner.Poll();
 		}
 
 		TRefCountPtr<FMaterialShaderMap> Get() override
 		{
 			// Make sure the async work is complete
-			if (RequestOwner)
+			if (!RequestOwner.Poll())
 			{
-				RequestOwner->Wait();
-				RequestOwner.Reset();
+				COOK_STAT(auto Timer = MaterialShaderCookStats::UsageStats.TimeAsyncWait());
+				COOK_STAT(Timer.TrackCyclesOnly());
+				RequestOwner.Wait();
 			}
 
 			TRACE_CPUPROFILER_EVENT_SCOPE(FMaterialShaderMap::FinishLoadFromDerivedDataCache);
@@ -1565,6 +1577,7 @@ TSharedRef<FMaterialShaderMap::FAsyncLoadContext> FMaterialShaderMap::BeginLoadF
 		{
 			SCOPE_SECONDS_COUNTER(MaterialDDCTime);
 			COOK_STAT(auto Timer = MaterialShaderCookStats::UsageStats.TimeSyncWork());
+			COOK_STAT(Timer.TrackCyclesOnly());
 
 			Result->DataKey = GetMaterialShaderMapKeyString(ShaderMapId, InPlatform);
 			OutDDCKeyDesc = FSHA1_HashString(Result->DataKey);
@@ -1582,7 +1595,7 @@ TSharedRef<FMaterialShaderMap::FAsyncLoadContext> FMaterialShaderMap::BeginLoadF
 				UE_LOG(LogMaterial, Display, TEXT("%s-%s-%s: %s"), *Material->GetAssetName(), *LexToString(ShaderMapId.FeatureLevel), *LexToString(ShaderMapId.QualityLevel), *SpecialEngineDDCKey);
 			}
 
-			bool CheckCache = true;
+			bool bCheckCache = true;
 
 			// If NoShaderDDC then don't check for a material the first time we encounter it to simulate
 			// a cold DDC
@@ -1598,7 +1611,7 @@ TSharedRef<FMaterialShaderMap::FAsyncLoadContext> FMaterialShaderMap::BeginLoadF
 
 				if (!SeenKeys.Contains(KeyHash))
 				{
-					CheckCache = false;
+					bCheckCache = false;
 					SeenKeys.Add(KeyHash);
 				}
 			}
@@ -1606,21 +1619,20 @@ TSharedRef<FMaterialShaderMap::FAsyncLoadContext> FMaterialShaderMap::BeginLoadF
 			// Do not check the DDC if the material isn't persistent, because
 			//   - this results in a lot of DDC requests when editing materials which are almost always going to return nothing.
 			//   - since the get call is synchronous, this can cause a hitch if there's network latency
-			if (CheckCache && Material->IsPersistent())
+			if (bCheckCache && Material->IsPersistent())
 			{
 				FCacheGetValueRequest Request;
-				Request.Name = Material->GetFriendlyName();
+				Request.Name = GetMaterialShaderMapName(Material->GetFullPath(), ShaderMapId, InPlatform);
 				Request.Key = GetMaterialShaderMapKey(Result->DataKey);
 				Result->AssetName = Material->GetAssetName();
 				Result->Platform = InPlatform;
-				Result->RequestOwner = MakeUnique<FRequestOwner>(EPriority::Normal);
 
-				GetCache().GetValue({Request}, *Result->RequestOwner, [Result](FCacheGetValueResponse&& Response)
+				GetCache().GetValue({Request}, Result->RequestOwner, [Result](FCacheGetValueResponse&& Response)
 				{
 					Result->CachedData = Response.Value.GetData().Decompress();
 					// This callback might hold the last reference to Result, which owns RequestOwner, so
 					// we must not cancel in the Owner's destructor; Canceling in a callback will deadlock.
-					Result->RequestOwner->KeepAlive();
+					Result->RequestOwner.KeepAlive();
 				});
 			}
 		}
@@ -1646,7 +1658,7 @@ void FMaterialShaderMap::SaveToDerivedDataCache()
 
 	using namespace UE::DerivedData;
 	FCachePutValueRequest Request;
-	Request.Name = GetFriendlyName();
+	Request.Name = GetMaterialShaderMapName(GetMaterialPath(), ShaderMapId, GetShaderPlatform());
 	Request.Key = GetMaterialShaderMapKey(DataKey);
 	Request.Value = FValue::Compress(MakeSharedBufferFromArray(MoveTemp(SaveData)));
 	FRequestOwner AsyncOwner(EPriority::Normal);
diff --git a/Engine/Source/Runtime/Engine/Private/ShaderCompiler/ShaderCompiler.cpp b/Engine/Source/Runtime/Engine/Private/ShaderCompiler/ShaderCompiler.cpp
--- a/Engine/Source/Runtime/Engine/Private/ShaderCompiler/ShaderCompiler.cpp
+++ b/Engine/Source/Runtime/Engine/Private/ShaderCompiler/ShaderCompiler.cpp
@@ -6852,6 +6852,11 @@ static UE::DerivedData::FCacheKey GetGlobalShaderMapKey(const FGlobalShaderMapId
 	static const UE::DerivedData::FCacheBucket Bucket("GlobalShaderMap");
 	return {Bucket, FIoHash::HashBuffer(MakeMemoryView(FTCHARToUTF8(DataKey)))};
 }
+
+static UE::DerivedData::FSharedString GetGlobalShaderMapName(const FGlobalShaderMapId& ShaderMapId, EShaderPlatform Platform, const FString& Key)
+{
+	return UE::DerivedData::FSharedString(WriteToString<256>(TEXTVIEW("GlobalShaderMap ["), LegacyShaderPlatformToShaderFormat(Platform), TEXTVIEW(", "), Key, TEXTVIEW("]")));
+}
 #endif
 
 /** Saves the platform's shader map to the DDC. It is assumed that the caller will check IsComplete() first before calling the function. */
@@ -6881,7 +6886,7 @@ static void SaveGlobalShaderMapToDerivedDataCache(EShaderPlatform Platform)
 
 			using namespace UE::DerivedData;
 			FCachePutValueRequest Request;
-			Request.Name = TEXTVIEW("GlobalShaderMap");
+			Request.Name = GetGlobalShaderMapName(ShaderMapId, Platform, ShaderFilenameDependencies.Key);
 			Request.Key = GetGlobalShaderMapKey(ShaderMapId, Platform, TargetPlatform, ShaderFilenameDependencies.Value);
 			Request.Value = FValue::Compress(MakeSharedBufferFromArray(MoveTemp(SaveData)));
 			FRequestOwner AsyncOwner(EPriority::Normal);
@@ -7169,7 +7174,7 @@ void CompileGlobalShaderMap(EShaderPlatform Platform, const ITargetPlatform* Tar
 				for (const auto& ShaderFilenameDependencies : ShaderMapId.GetShaderFilenameToDependeciesMap())
 				{
 					FCacheGetValueRequest& Request = Requests.AddDefaulted_GetRef();
-					Request.Name = WriteToString<128>(TEXT("GlobalShaderMap "), LegacyShaderPlatformToShaderFormat(Platform));
+					Request.Name = GetGlobalShaderMapName(ShaderMapId, Platform, ShaderFilenameDependencies.Key);
 					Request.Key = GetGlobalShaderMapKey(ShaderMapId, Platform, TargetPlatform, ShaderFilenameDependencies.Value);
 					Request.UserData = uint64(BufferIndex);
 					++BufferIndex;
