Fix for push model subobject update skips and replay checkpoints

diff --git a/Engine/Source/Runtime/Engine/Private/DataChannel.cpp b/Engine/Source/Runtime/Engine/Private/DataChannel.cpp
--- a/Engine/Source/Runtime/Engine/Private/DataChannel.cpp
+++ b/Engine/Source/Runtime/Engine/Private/DataChannel.cpp
@@ -4885,8 +4885,14 @@ bool UActorChannel::WriteSubObjectInBunch(UObject* Obj, FOutBunch& Bunch, FRepli
 	bool bFoundInvalidReplicator = false;
 	TSharedRef<FObjectReplicator>* FoundReplicator = FindReplicator(Obj, &bFoundInvalidReplicator);
 	const bool bFoundReplicator = (FoundReplicator != nullptr);
-	
-	if (bFoundReplicator && (*FoundReplicator)->CanSkipUpdate(RepFlags))
+
+	// Special case for replay checkpoints where the replicator could already exist but we still need to consider this a new object even if it has an empty layout
+	const bool bNewToReplay = (Connection->ResendAllDataState == EResendAllDataState::SinceOpen)
+		|| ((Connection->ResendAllDataState == EResendAllDataState::SinceCheckpoint) && (!bFoundReplicator || (*FoundReplicator)->IsDirtyForReplay()));
+
+	const bool bCanSkipUpdate = bFoundReplicator && !bNewToReplay && (*FoundReplicator)->CanSkipUpdate(RepFlags);
+
+	if (bCanSkipUpdate)
 	{
 		return false;
 	}
@@ -4907,10 +4913,6 @@ bool UActorChannel::WriteSubObjectInBunch(UObject* Obj, FOutBunch& Bunch, FRepli
 	FReplicationFlags ObjRepFlags = RepFlags;
 	TSharedRef<FObjectReplicator>& ObjectReplicator = !bFoundReplicator ? CreateReplicator(Obj, !bFoundInvalidReplicator) : *FoundReplicator;
 
-	// Special case for replay checkpoints where the replicator could already exist but we still need to consider this a new object even if it has an empty layout
-	const bool bNewToReplay = (Connection->ResendAllDataState == EResendAllDataState::SinceOpen)
-		|| ((Connection->ResendAllDataState == EResendAllDataState::SinceCheckpoint) && ObjectReplicator->IsDirtyForReplay());
-
 	if (!bFoundReplicator || bNewToReplay)
 	{
 		// This is the first time replicating this subobject
