Fix crash (from CL 21199427) : NewActor can be nullptr

#rnx

diff --git a/Engine/Source/Runtime/Engine/Private/WorldPartition/ActorDescContainer.cpp b/Engine/Source/Runtime/Engine/Private/WorldPartition/ActorDescContainer.cpp
--- a/Engine/Source/Runtime/Engine/Private/WorldPartition/ActorDescContainer.cpp
+++ b/Engine/Source/Runtime/Engine/Private/WorldPartition/ActorDescContainer.cpp
@@ -180,7 +180,7 @@ void UActorDescContainer::OnObjectsReplaced(const TMap<UObject*, UObject*>& OldT
 	{
 		if (AActor* OldActor = Cast<AActor>(OldObject))
 		{
-			AActor* NewActor = CastChecked<AActor>(NewObject);
+			AActor* NewActor = Cast<AActor>(NewObject);
 			if (FWorldPartitionActorDesc* ActorDesc = GetActorDesc(OldActor->GetActorGuid()))
 			{				
 				FWorldPartitionActorDescUtils::ReplaceActorDescriptorPointerFromActor(OldActor, NewActor, ActorDesc);
diff --git a/Engine/Source/Runtime/Engine/Private/WorldPartition/WorldPartitionActorDescUtils.cpp b/Engine/Source/Runtime/Engine/Private/WorldPartition/WorldPartitionActorDescUtils.cpp
--- a/Engine/Source/Runtime/Engine/Private/WorldPartition/WorldPartitionActorDescUtils.cpp
+++ b/Engine/Source/Runtime/Engine/Private/WorldPartition/WorldPartitionActorDescUtils.cpp
@@ -110,8 +110,8 @@ void FWorldPartitionActorDescUtils::UpdateActorDescriptorFomActor(const AActor*
 
 void FWorldPartitionActorDescUtils::ReplaceActorDescriptorPointerFromActor(const AActor* OldActor, AActor* NewActor, FWorldPartitionActorDesc* ActorDesc)
 {
-	check(OldActor->GetActorGuid() == NewActor->GetActorGuid());
-	check(NewActor->GetActorGuid() == ActorDesc->GetGuid());
+	check(!NewActor || (OldActor->GetActorGuid() == NewActor->GetActorGuid()));
+	check(!NewActor || (NewActor->GetActorGuid() == ActorDesc->GetGuid()));
 	check(!ActorDesc->ActorPtr.IsValid() || (ActorDesc->ActorPtr == OldActor));
 	ActorDesc->ActorPtr = NewActor;
 }
