Fixed test timeout when running with -nullrhi

Added default values to BaseEngine.ini for reference

diff --git a/Engine/Config/BaseEngine.ini b/Engine/Config/BaseEngine.ini
--- a/Engine/Config/BaseEngine.ini
+++ b/Engine/Config/BaseEngine.ini
@@ -2446,6 +2446,13 @@ MaterialEditorPromotionTest=(DefaultMaterialAsset=(FilePath="../../Content/Edito
 +EditorViewButtons=Binormals
 +EditorViewButtons=UV
 EditorViewButtonsObject=EditorCylinder
+; Consider 5fps as the start of editor interactivity for tests
+DefaultInteractiveFramerate=5
+; Wait for up to 10m mins for interactivity (this is realyl to account for any shader compiles or other DDC tasks)
+DefaultInteractiveFramerateWaitTime=600
+; Require an interactive framerate for this long before passing
+DefaultInteractiveFramerateDuration=5
+
 
 [/Script/NavigationSystem.NavigationSystemV1]
 bAutoCreateNavigationData=true
diff --git a/Engine/Source/Runtime/Engine/Private/Tests/AutomationCommon.cpp b/Engine/Source/Runtime/Engine/Private/Tests/AutomationCommon.cpp
--- a/Engine/Source/Runtime/Engine/Private/Tests/AutomationCommon.cpp
+++ b/Engine/Source/Runtime/Engine/Private/Tests/AutomationCommon.cpp
@@ -515,12 +515,20 @@ FWaitForInteractiveFrameRate::FWaitForInteractiveFrameRate(float InDesiredFrameR
 
 bool FWaitForInteractiveFrameRate::Update()
 {
+	const double kReportInterval = 30;
 	const double TimeNow = FPlatformTime::Seconds();
 
+	if (FApp::CanEverRender() == false)
+	{
+		UE_LOG(LogEngineAutomationLatentCommand, Log, TEXT("FWaitForInteractiveFrameRate: FApp::CanEverRender() == false. Passing wait immediately."));
+		return true;
+	}
+
 	if (StartTimeOfWait == 0)
 	{
 		UE_LOG(LogEngineAutomationLatentCommand, Log, TEXT("FWaitForInteractiveFrameRate: Starting wait for framerate of >= %.f FPS"), DesiredFrameRate);
 		StartTimeOfWait = TimeNow;
+		LastReportTime = TimeNow;
 	}
 
 	// Check if we're at the desired framerate, and if so how long we have been there for
@@ -564,6 +572,12 @@ bool FWaitForInteractiveFrameRate::Update()
 		return true;
 	}
 
+	if (TimeNow - LastReportTime >= kReportInterval)
+	{
+		UE_LOG(LogEngineAutomationLatentCommand, Display, TEXT("FWaitForInteractiveFrameRate: Waited %.f seconds. Will timeout in %.f Current FPS=%.f "), ElapsedWaitTime, MaxWaitTime- ElapsedWaitTime, GAverageFPS);
+		LastReportTime = TimeNow;
+	}
+
 	return false;
 }
 
diff --git a/Engine/Source/Runtime/Engine/Public/Tests/AutomationCommon.h b/Engine/Source/Runtime/Engine/Public/Tests/AutomationCommon.h
--- a/Engine/Source/Runtime/Engine/Public/Tests/AutomationCommon.h
+++ b/Engine/Source/Runtime/Engine/Public/Tests/AutomationCommon.h
@@ -208,6 +208,9 @@ private:
 
 	// how many seconds we've been at the desired framerate
 	double StartTimeOfAcceptableFrameRate;	
+
+	// time we last logged we are waiting
+	double LastReportTime;
 };
 
 /**
