Compile out Lumen GI and Reflections from translucent material shaders when r.Lumen.Supported=0

diff --git a/Engine/Shaders/Private/BasePassPixelShader.usf b/Engine/Shaders/Private/BasePassPixelShader.usf
--- a/Engine/Shaders/Private/BasePassPixelShader.usf
+++ b/Engine/Shaders/Private/BasePassPixelShader.usf
@@ -335,7 +335,7 @@ void GetSkyLighting(FMaterialPixelParameters MaterialParameters, VTPageTableResu
 	OutDiffuseLighting = 0;
 	OutSubsurfaceLighting = 0;
 
-#if MATERIALBLENDING_TRANSLUCENT || MATERIALBLENDING_ADDITIVE || MATERIALBLENDING_MODULATE
+#if (MATERIALBLENDING_TRANSLUCENT || MATERIALBLENDING_ADDITIVE || MATERIALBLENDING_MODULATE) && PROJECT_SUPPORTS_LUMEN && !FORWARD_SHADING
 	if (IsLumenTranslucencyGIEnabled())
 	{
 		// Lumen Dynamic GI + shadowed Skylight
diff --git a/Engine/Shaders/Private/ForwardLightingCommon.ush b/Engine/Shaders/Private/ForwardLightingCommon.ush
--- a/Engine/Shaders/Private/ForwardLightingCommon.ush
+++ b/Engine/Shaders/Private/ForwardLightingCommon.ush
@@ -401,7 +401,7 @@ half3 GetImageBasedReflectionSpecular(FMaterialPixelParameters MaterialParameter
 	float3 SpecularIBL;
 	bool bUseLumenFrontLayerReflection = false;
 
-#if (TRANSLUCENCY_LIGHTING_SURFACE_FORWARDSHADING || TRANSLUCENCY_LIGHTING_SURFACE_LIGHTINGVOLUME) && MATERIALBLENDING_ANY_TRANSLUCENT && !FORWARD_SHADING
+#if (TRANSLUCENCY_LIGHTING_SURFACE_FORWARDSHADING || TRANSLUCENCY_LIGHTING_SURFACE_LIGHTINGVOLUME) && MATERIALBLENDING_ANY_TRANSLUCENT && PROJECT_SUPPORTS_LUMEN && !FORWARD_SHADING
 	
 	bUseLumenFrontLayerReflection = UseFrontLayerReflection(MaterialParameters.ViewBufferUV, MaterialParameters.ScreenPosition.w);
 
diff --git a/Engine/Source/Runtime/Engine/Private/ShaderCompiler/ShaderCompiler.cpp b/Engine/Source/Runtime/Engine/Private/ShaderCompiler/ShaderCompiler.cpp
--- a/Engine/Source/Runtime/Engine/Private/ShaderCompiler/ShaderCompiler.cpp
+++ b/Engine/Source/Runtime/Engine/Private/ShaderCompiler/ShaderCompiler.cpp
@@ -6017,6 +6017,12 @@ void GlobalBeginCompileShader(
 		Input.Environment.SetDefine(TEXT("MATERIAL_ROUGHDIFFUSE"), bMaterialRoughDiffuse ? 1 : 0);
 	}
 
+	{
+		static IConsoleVariable* CVar = IConsoleManager::Get().FindConsoleVariable(TEXT("r.Lumen.Supported"));
+		const bool bLumenSupported = CVar->GetInt() != 0;
+		Input.Environment.SetDefine(TEXT("PROJECT_SUPPORTS_LUMEN"), bLumenSupported ? 1 : 0);
+	}
+
 	{
 		const bool bSupportOIT = FDataDrivenShaderPlatformInfo::GetSupportsOIT(EShaderPlatform(Target.Platform));
 		static IConsoleVariable* CVar = IConsoleManager::Get().FindConsoleVariable(TEXT("r.OIT.SortedPixels"));
