WorldPartition: Fix PIE for Unsaved template worlds

diff --git a/Engine/Source/Runtime/Engine/Private/World.cpp b/Engine/Source/Runtime/Engine/Private/World.cpp
--- a/Engine/Source/Runtime/Engine/Private/World.cpp
+++ b/Engine/Source/Runtime/Engine/Private/World.cpp
@@ -534,11 +534,13 @@ void UWorld::Serialize( FArchive& Ar )
 		GetOutermost()->ThisContainsMap();
 	}
 
-	// Serialize World composition for PIE
+#if WITH_EDITOR
+	// Serialize for PIE
 	if (Ar.GetPortFlags() & PPF_DuplicateForPIE)
 	{
 		Ar << OriginLocation;
 		Ar << RequestedOriginLocation;
+		Ar << OriginalWorldName;
 	}
 	
 	// UWorlds loaded/duplicated for PIE must lose RF_Public and RF_Standalone since they should not be referenced by objects in other packages and they should be GCed normally
@@ -546,6 +548,7 @@ void UWorld::Serialize( FArchive& Ar )
 	{
 		ClearFlags(RF_Public|RF_Standalone);
 	}
+#endif
 }
 
 void UWorld::AddReferencedObjects(UObject* InThis, FReferenceCollector& Collector)
diff --git a/Engine/Source/Runtime/Engine/Private/WorldPartition/WorldPartitionRuntimeLevelStreamingCell.cpp b/Engine/Source/Runtime/Engine/Private/WorldPartition/WorldPartitionRuntimeLevelStreamingCell.cpp
--- a/Engine/Source/Runtime/Engine/Private/WorldPartition/WorldPartitionRuntimeLevelStreamingCell.cpp
+++ b/Engine/Source/Runtime/Engine/Private/WorldPartition/WorldPartitionRuntimeLevelStreamingCell.cpp
@@ -118,7 +118,15 @@ UWorldPartitionLevelStreamingDynamic* UWorldPartitionRuntimeLevelStreamingCell::
 		UWorld* LevelStreamingOuterWorld = IsRunningCommandlet() ? OuterWorld : OwningWorld;
 		UWorldPartitionLevelStreamingDynamic* NewLevelStreaming = NewObject<UWorldPartitionLevelStreamingDynamic>(LevelStreamingOuterWorld, UWorldPartitionLevelStreamingDynamic::StaticClass(), LevelStreamingName, RF_NoFlags, NULL);
 		FString PackageName = !InPackageName.IsEmpty() ? InPackageName : UWorldPartitionLevelStreamingPolicy::GetCellPackagePath(GetFName(), OuterWorld);
-		TSoftObjectPtr<UWorld> WorldAsset(FSoftObjectPath(FString::Printf(TEXT("%s.%s"), *PackageName, *OuterWorld->GetName())));
+				
+		FName WorldName = OuterWorld->GetFName();
+		// In PIE make sure that we are using the proper original world name so that actors resolve their outer property
+		if (OwningWorld->IsPlayInEditor() && !OuterWorld->OriginalWorldName.IsNone())
+		{
+			WorldName = OuterWorld->OriginalWorldName;
+		}
+
+		TSoftObjectPtr<UWorld> WorldAsset(FSoftObjectPath(FString::Printf(TEXT("%s.%s"), *PackageName, *WorldName.ToString())));
 		NewLevelStreaming->SetWorldAsset(WorldAsset);
 		// Transfer WorldPartition's transform to Level
 		NewLevelStreaming->LevelTransform = WorldPartition->GetInstanceTransform();
