Skeletal Mesh - Fix the value of the SkinnedAsset pointer in the SkinnedMeshComponent when the component is overriden by inherited blueprints values and the deprecatred SkeletalMesh pointer is still in use.

diff --git a/Engine/Source/Runtime/Engine/Private/Components/SkinnedMeshComponent.cpp b/Engine/Source/Runtime/Engine/Private/Components/SkinnedMeshComponent.cpp
--- a/Engine/Source/Runtime/Engine/Private/Components/SkinnedMeshComponent.cpp
+++ b/Engine/Source/Runtime/Engine/Private/Components/SkinnedMeshComponent.cpp
@@ -518,7 +518,7 @@ void USkinnedMeshComponent::PostLoad()
 	Super::PostLoad();
 
 	PRAGMA_DISABLE_DEPRECATION_WARNINGS
-	if(!SkinnedAsset)
+	if (SkeletalMesh)
 	{
 		SkinnedAsset = SkeletalMesh;  // Set to the SkeletalMesh pointer for backward compatibility when the SkinnedAsset is invalid
 	}
