Fixed repeated falling movement mode changes during lateral root motion caused by gravity being ignored during liftoffbound computation

diff --git a/Engine/Source/Runtime/Engine/Private/Components/CharacterMovementComponent.cpp b/Engine/Source/Runtime/Engine/Private/Components/CharacterMovementComponent.cpp
--- a/Engine/Source/Runtime/Engine/Private/Components/CharacterMovementComponent.cpp
+++ b/Engine/Source/Runtime/Engine/Private/Components/CharacterMovementComponent.cpp
@@ -4057,7 +4057,7 @@ void UCharacterMovementComponent::ApplyRootMotionToVelocity(float deltaTime)
 		else
 		{
 			// Default bounds - the amount of force gravity is applying this tick
-			LiftoffBound = FMath::Max(GetGravityZ() * deltaTime, UE_SMALL_NUMBER);
+			LiftoffBound = FMath::Max(-GetGravityZ() * deltaTime, UE_SMALL_NUMBER);
 		}
 
 		if( AppliedVelocityDelta.Z > LiftoffBound )
