Changed check to ensure to avoid crashing server if ClientConnectionId assignment fails.

diff --git a/Engine/Source/Runtime/Engine/Private/Net/NetConnectionIdHandler.cpp b/Engine/Source/Runtime/Engine/Private/Net/NetConnectionIdHandler.cpp
--- a/Engine/Source/Runtime/Engine/Private/Net/NetConnectionIdHandler.cpp
+++ b/Engine/Source/Runtime/Engine/Private/Net/NetConnectionIdHandler.cpp
@@ -22,7 +22,7 @@ void FNetConnectionIdHandler::Init(uint32 IdCount)
 uint32 FNetConnectionIdHandler::Allocate()
 {
 	uint32 Id = UsedIds.FindAndSetFirstZeroBit(IdHint);
-	checkf(Id != INDEX_NONE, TEXT("Out of connection IDs. All %d are in use."), UsedIds.Num());
+	ensureMsgf(Id != INDEX_NONE, TEXT("Out of connection IDs. All %d are in use."), UsedIds.Num());
 	IdHint = Id + 1;
 	return Id != INDEX_NONE ? Id : 0U;
 }
