Check for ShaderFormat nullptr in GlobalBeginCompileShader().

#rnx

diff --git a/Engine/Source/Runtime/Engine/Private/ShaderCompiler/ShaderCompiler.cpp b/Engine/Source/Runtime/Engine/Private/ShaderCompiler/ShaderCompiler.cpp
--- a/Engine/Source/Runtime/Engine/Private/ShaderCompiler/ShaderCompiler.cpp
+++ b/Engine/Source/Runtime/Engine/Private/ShaderCompiler/ShaderCompiler.cpp
@@ -6108,6 +6108,7 @@ void GlobalBeginCompileShader(
 
 	// Allow the target shader format to modify the shader input before we add it as a job
 	const IShaderFormat* Format = GetTargetPlatformManagerRef().FindShaderFormat(ShaderFormatName);
+	checkf(Format, TEXT("Shader format %s cannot be found"), *ShaderFormatName.ToString());
 	Format->ModifyShaderCompilerInput(Input);
 
 	// Allow the GBuffer and other shader defines to cause dependend environment changes, but minimizing the #ifdef magic in the shaders, which
