Add delegates for physics state created and destroyed in primitive component.

diff --git a/Engine/Source/Runtime/Engine/Classes/Components/PrimitiveComponent.h b/Engine/Source/Runtime/Engine/Classes/Components/PrimitiveComponent.h
--- a/Engine/Source/Runtime/Engine/Classes/Components/PrimitiveComponent.h
+++ b/Engine/Source/Runtime/Engine/Classes/Components/PrimitiveComponent.h
@@ -166,6 +166,16 @@ struct FRendererStencilMaskEvaluation
 	}
 };
 
+// TODO: Add sleep and wake state change types to this enum, so that the
+// OnComponentWake and OnComponentSleep delegates may be deprecated.
+// Doing so would save a couple bytes per primitive component.
+UENUM(BlueprintType)
+enum class EComponentPhysicsStateChange : uint8
+{
+	Created,
+	Destroyed
+};
+
 /**
  * Delegate for notification of blocking collision against a specific component.  
  * NormalImpulse will be filled in for physics-simulating bodies, but will be zero for swept-component blocking collisions. 
@@ -181,6 +191,8 @@ DECLARE_DYNAMIC_MULTICAST_SPARSE_DELEGATE_TwoParams(FComponentWakeSignature, UPr
 DECLARE_DYNAMIC_MULTICAST_SPARSE_DELEGATE_TwoParams(FComponentSleepSignature, UPrimitiveComponent, OnComponentSleep, UPrimitiveComponent*, SleepingComponent, FName, BoneName);
 /** Delegate for notification when collision settings change. */
 DECLARE_DYNAMIC_MULTICAST_SPARSE_DELEGATE_OneParam(FComponentCollisionSettingsChangedSignature, UPrimitiveComponent, OnComponentCollisionSettingsChangedEvent, UPrimitiveComponent*, ChangedComponent);
+/** Delegate for physics state created */
+DECLARE_DYNAMIC_MULTICAST_SPARSE_DELEGATE_TwoParams(FComponentPhysicsStateChanged, UPrimitiveComponent, OnComponentPhysicsStateChanged, UPrimitiveComponent*, ChangedComponent, EComponentPhysicsStateChange, StateChange);
 
 DECLARE_DYNAMIC_MULTICAST_SPARSE_DELEGATE_OneParam( FComponentBeginCursorOverSignature, UPrimitiveComponent, OnBeginCursorOver, UPrimitiveComponent*, TouchedComponent );
 DECLARE_DYNAMIC_MULTICAST_SPARSE_DELEGATE_OneParam( FComponentEndCursorOverSignature, UPrimitiveComponent, OnEndCursorOver, UPrimitiveComponent*, TouchedComponent );
@@ -1200,6 +1212,12 @@ public:
 	 */
 	FComponentCollisionSettingsChangedSignature OnComponentCollisionSettingsChangedEvent;
 
+	/**
+	 *	Event called when physics state is created or destroyed for this component
+	 */
+	UPROPERTY(BlueprintAssignable, Category = "Physics")
+	FComponentPhysicsStateChanged OnComponentPhysicsStateChanged;
+
 	/** Event called when the mouse cursor is moved over this component and mouse over events are enabled in the player controller */
 	UPROPERTY(BlueprintAssignable, Category="Input|Mouse Input")
 	FComponentBeginCursorOverSignature OnBeginCursorOver;
diff --git a/Engine/Source/Runtime/Engine/Private/Components/PrimitiveComponent.cpp b/Engine/Source/Runtime/Engine/Private/Components/PrimitiveComponent.cpp
--- a/Engine/Source/Runtime/Engine/Private/Components/PrimitiveComponent.cpp
+++ b/Engine/Source/Runtime/Engine/Private/Components/PrimitiveComponent.cpp
@@ -857,6 +857,8 @@ void UPrimitiveComponent::OnCreatePhysicsState()
 #endif // WITH_EDITOR
 		}
 	}
+
+	OnComponentPhysicsStateChanged.Broadcast(this, EComponentPhysicsStateChange::Created);
 }	
 
 void UPrimitiveComponent::EnsurePhysicsStateCreated()
@@ -940,6 +942,8 @@ void UPrimitiveComponent::OnDestroyPhysicsState()
 #endif
 
 	Super::OnDestroyPhysicsState();
+
+	OnComponentPhysicsStateChanged.Broadcast(this, EComponentPhysicsStateChange::Destroyed);
 }
 
 #if UE_ENABLE_DEBUG_DRAWING
