Fixed transaction diffs failing to track name and object indices correctly

#rnx

#localization none
#tests Re-ran affected unit test

diff --git a/Engine/Source/Runtime/Engine/Private/TransactionCommon.cpp b/Engine/Source/Runtime/Engine/Private/TransactionCommon.cpp
--- a/Engine/Source/Runtime/Engine/Private/TransactionCommon.cpp
+++ b/Engine/Source/Runtime/Engine/Private/TransactionCommon.cpp
@@ -412,7 +412,7 @@ void FDiffableObjectDataWriter::OnNameSerialized(int32 InNameIndex)
 	const FName SerializedTaggedDataKey = GetTaggedDataKey();
 	if (!SerializedTaggedDataKey.IsNone())
 	{
-		DiffableObject.SerializedNameIndices.Add(SerializedTaggedDataKey).Indices.AddUnique(InNameIndex);
+		DiffableObject.SerializedNameIndices.FindOrAdd(SerializedTaggedDataKey).Indices.AddUnique(InNameIndex);
 	}
 }
 
@@ -422,7 +422,7 @@ void FDiffableObjectDataWriter::OnObjectSerialized(int32 InObjectIndex)
 	const FName SerializedTaggedDataKey = GetTaggedDataKey();
 	if (!SerializedTaggedDataKey.IsNone())
 	{
-		DiffableObject.SerializedObjectIndices.Add(SerializedTaggedDataKey).Indices.AddUnique(InObjectIndex);
+		DiffableObject.SerializedObjectIndices.FindOrAdd(SerializedTaggedDataKey).Indices.AddUnique(InObjectIndex);
 	}
 }
 
