Fix engine non-unity issues
#rnx

diff --git a/Engine/Plugins/Online/VoiceChat/VoiceChat/Source/Public/VoiceChatResult.h b/Engine/Plugins/Online/VoiceChat/VoiceChat/Source/Public/VoiceChatResult.h
--- a/Engine/Plugins/Online/VoiceChat/VoiceChat/Source/Public/VoiceChatResult.h
+++ b/Engine/Plugins/Online/VoiceChat/VoiceChat/Source/Public/VoiceChatResult.h
@@ -1,6 +1,8 @@
 // Copyright Epic Games, Inc. All Rights Reserved.
 #pragma once
 
+#include "Containers/UnrealString.h"
+
 enum class EVoiceChatResult
 {
 	// The operation succeeded
diff --git a/Engine/Source/Runtime/Engine/Classes/Animation/Skeleton.h b/Engine/Source/Runtime/Engine/Classes/Animation/Skeleton.h
--- a/Engine/Source/Runtime/Engine/Classes/Animation/Skeleton.h
+++ b/Engine/Source/Runtime/Engine/Classes/Animation/Skeleton.h
@@ -16,6 +16,7 @@
 #include "Animation/PreviewAssetAttachComponent.h"
 #include "Animation/SmartName.h"
 #include "Engine/AssetUserData.h"
+#include "Engine/SkeletalMesh.h"
 #include "HAL/CriticalSection.h"
 #include "Interfaces/Interface_AssetUserData.h"
 #include "Interfaces/Interface_PreviewMeshProvider.h"
@@ -25,7 +26,6 @@
 class UAnimSequence;
 class UBlendProfile;
 class URig;
-class USkeletalMesh;
 class USkeletalMeshSocket;
 class FPackageReloadedEvent;
 struct FAssetData;
diff --git a/Engine/Source/Runtime/Engine/Classes/Engine/SkeletalMesh.h b/Engine/Source/Runtime/Engine/Classes/Engine/SkeletalMesh.h
--- a/Engine/Source/Runtime/Engine/Classes/Engine/SkeletalMesh.h
+++ b/Engine/Source/Runtime/Engine/Classes/Engine/SkeletalMesh.h
@@ -21,7 +21,6 @@
 #include "BoneContainer.h"
 #include "Interfaces/Interface_CollisionDataProvider.h"
 #include "EngineTypes.h"
-#include "Engine/Engine.h"
 #include "SkeletalMeshSampling.h"
 #include "PerPlatformProperties.h"
 #include "Misc/EnumClassFlags.h"
diff --git a/Engine/Source/Runtime/Engine/Public/Physics/Experimental/PhysScene_Chaos.h b/Engine/Source/Runtime/Engine/Public/Physics/Experimental/PhysScene_Chaos.h
--- a/Engine/Source/Runtime/Engine/Public/Physics/Experimental/PhysScene_Chaos.h
+++ b/Engine/Source/Runtime/Engine/Public/Physics/Experimental/PhysScene_Chaos.h
@@ -369,7 +369,7 @@ private:
 #endif
 
 	// Allow other code to obtain read-locks when needed
-	friend struct FScopedSceneReadLock;
+	friend struct ChaosInterface::FScopedSceneReadLock;
 	friend struct FScopedSceneLock_Chaos;
 };
 
diff --git a/Engine/Source/Runtime/PhysicsCore/Public/Chaos/ChaosScene.h b/Engine/Source/Runtime/PhysicsCore/Public/Chaos/ChaosScene.h
--- a/Engine/Source/Runtime/PhysicsCore/Public/Chaos/ChaosScene.h
+++ b/Engine/Source/Runtime/PhysicsCore/Public/Chaos/ChaosScene.h
@@ -165,7 +165,7 @@ protected:
 #endif
 
 	// Allow other code to obtain read-locks when needed
-	friend struct FScopedSceneReadLock;
+	friend struct ChaosInterface::FScopedSceneReadLock;
 	friend struct FScopedSceneLock_Chaos;
 
 	//Engine interface BEGIN
diff --git a/Engine/Source/Runtime/PhysicsCore/Public/PhysicsInterfaceDeclaresCore.h b/Engine/Source/Runtime/PhysicsCore/Public/PhysicsInterfaceDeclaresCore.h
--- a/Engine/Source/Runtime/PhysicsCore/Public/PhysicsInterfaceDeclaresCore.h
+++ b/Engine/Source/Runtime/PhysicsCore/Public/PhysicsInterfaceDeclaresCore.h
@@ -42,6 +42,7 @@ namespace ChaosInterface
 	struct FDummyPhysType;
 	struct FDummyPhysActor;
 	template<typename T> struct FDummyCallback;
+	struct FScopedSceneReadLock;
 }
 using FPhysTypeDummy = ChaosInterface::FDummyPhysType;
 using FPhysActorDummy = ChaosInterface::FDummyPhysActor;
