Move implementation to cpp to avoid need for include
#rnx

diff --git a/Engine/Source/Runtime/Engine/Classes/Engine/LevelStreaming.h b/Engine/Source/Runtime/Engine/Classes/Engine/LevelStreaming.h
--- a/Engine/Source/Runtime/Engine/Classes/Engine/LevelStreaming.h
+++ b/Engine/Source/Runtime/Engine/Classes/Engine/LevelStreaming.h
@@ -10,7 +10,6 @@
 #include "Engine/LatentActionManager.h"
 #include "LatentActions.h"
 #include "ProfilingDebugging/ProfilingHelpers.h"
-#include "Engine/Level.h"
 
 #if WITH_EDITOR
 #include "Folder.h"
@@ -461,7 +460,7 @@ public:
 
 	/** Returns whether streaming level is visible */
 	UFUNCTION(BlueprintCallable, Category="Game")
-	bool IsLevelVisible() const { return LoadedLevel && LoadedLevel->bIsVisible; }
+	bool IsLevelVisible() const;
 
 	/** Returns whether streaming level is loaded */
 	UFUNCTION(BlueprintCallable, Category="Game")
diff --git a/Engine/Source/Runtime/Engine/Private/LevelStreaming.cpp b/Engine/Source/Runtime/Engine/Private/LevelStreaming.cpp
--- a/Engine/Source/Runtime/Engine/Private/LevelStreaming.cpp
+++ b/Engine/Source/Runtime/Engine/Private/LevelStreaming.cpp
@@ -394,6 +394,11 @@ UWorld* ULevelStreaming::GetWorld() const
 	}
 }
 
+bool ULevelStreaming::IsLevelVisible() const
+{
+	return LoadedLevel && LoadedLevel->bIsVisible;
+}
+
 void ULevelStreaming::Serialize( FArchive& Ar )
 {
 	Super::Serialize(Ar);
