Fix a bug that PIE doesn't respect to Editor FeatureLevel Settings when have PIE MapOverride setup.

diff --git a/Engine/Source/Runtime/Engine/Private/GameInstance.cpp b/Engine/Source/Runtime/Engine/Private/GameInstance.cpp
--- a/Engine/Source/Runtime/Engine/Private/GameInstance.cpp
+++ b/Engine/Source/Runtime/Engine/Private/GameInstance.cpp
@@ -313,6 +313,7 @@ FGameInstancePIEResult UGameInstance::InitializeForPlayInEditor(int32 PIEInstanc
 			UWorld* WorldToDuplicate = Cast<UWorld>(TargetWorld.TryLoad());
 			if (WorldToDuplicate)
 			{
+				WorldToDuplicate->ChangeFeatureLevel(EditorEngine->EditorWorld->FeatureLevel, false);
 				NewWorld = EditorEngine->CreatePIEWorldByDuplication(*WorldContext, WorldToDuplicate, PIEMapName);
 			}
 		}
