Fixed comment on sequence player

diff --git a/Engine/Source/Runtime/Engine/Classes/Animation/AnimNode_SequencePlayer.h b/Engine/Source/Runtime/Engine/Classes/Animation/AnimNode_SequencePlayer.h
--- a/Engine/Source/Runtime/Engine/Classes/Animation/AnimNode_SequencePlayer.h
+++ b/Engine/Source/Runtime/Engine/Classes/Animation/AnimNode_SequencePlayer.h
@@ -61,7 +61,7 @@ protected:
 	// Performed after PlayRateBasis.
 	virtual const FInputScaleBiasClampConstants& GetPlayRateScaleBiasClampConstants() const { static FInputScaleBiasClampConstants Dummy; return Dummy; }
 
-	// The start position in [0, 1] to use when initializing. When looping, play will still jump back to the beginning when reaching the end.
+	// The start position [range: 0 - sequence length] to use when initializing. When looping, play will still jump back to the beginning when reaching the end.
 	virtual float GetStartPosition() const { return 0.0f; }
 
 	// Should the animation loop back to the start when it reaches the end?
