[NO NON-PII COMMIT MESSAGE]

diff --git a/Engine/Source/Runtime/Engine/Private/Animation/AnimInstance.cpp b/Engine/Source/Runtime/Engine/Private/Animation/AnimInstance.cpp
--- a/Engine/Source/Runtime/Engine/Private/Animation/AnimInstance.cpp
+++ b/Engine/Source/Runtime/Engine/Private/Animation/AnimInstance.cpp
@@ -1780,6 +1780,16 @@ void UAnimInstance::Montage_Advance(float DeltaSeconds)
 
 			MontageInstance->MontageSync_PreUpdate();
 			MontageInstance->Advance(DeltaSeconds, RootMotionParams, bUsingBlendedRootMotion);
+
+			// If MontageInstances has been modified while executing MontageInstance->Advance(), MontageInstance is unsafe to
+			// access further. This happens for example if MontageInstance->Advance() triggers an anim notify in which the user
+			// destroys the owning actor which in turn calls UninitializeAnimation(), or when the anim notify causes any montage
+			// to stop or start playing. We just check here if the current MontageInstance is still safe to access.
+			if (!MontageInstances.IsValidIndex(InstanceIndex) || MontageInstances[InstanceIndex] != MontageInstance)
+			{
+				break;
+			}
+
 			MontageInstance->MontageSync_PostUpdate();
 
 #if DO_CHECK && WITH_EDITORONLY_DATA && 0
