Strata - disable FrontMaterial input for volumetric when strata is off.

diff --git a/Engine/Source/Runtime/Engine/Private/Materials/Material.cpp b/Engine/Source/Runtime/Engine/Private/Materials/Material.cpp
--- a/Engine/Source/Runtime/Engine/Private/Materials/Material.cpp
+++ b/Engine/Source/Runtime/Engine/Private/Materials/Material.cpp
@@ -6216,11 +6216,14 @@ static bool IsPropertyActive_Internal(EMaterialProperty InProperty,
 	}
 	else if (Domain == MD_Volume)
 	{
+		if (InProperty == MP_FrontMaterial)
+		{
+			return bStrataEnabled;
+		}
 		return InProperty == MP_EmissiveColor
 			|| InProperty == MP_SubsurfaceColor
 			|| InProperty == MP_BaseColor
-			|| InProperty == MP_AmbientOcclusion
-			|| InProperty == MP_FrontMaterial;
+			|| InProperty == MP_AmbientOcclusion;
 	}
 	else if (Domain == MD_UI)
 	{
@@ -6229,11 +6232,10 @@ static bool IsPropertyActive_Internal(EMaterialProperty InProperty,
 			|| (InProperty == MP_OpacityMask && BlendMode == BLEND_Masked)
 			|| (InProperty == MP_Opacity && IsTranslucentBlendMode(BlendMode) && BlendMode != BLEND_Modulate)
 			|| (InProperty >= MP_CustomizedUVs0 && InProperty <= MP_CustomizedUVs7);
-		{
-			return true;
-		}
 	}
 
+	// Now processing MD_Surface
+
 	const bool bIsTranslucentBlendMode = IsTranslucentBlendMode(BlendMode);
 	const bool bIsNonDirectionalTranslucencyLightingMode = TranslucencyLightingMode == TLM_VolumetricNonDirectional || TranslucencyLightingMode == TLM_VolumetricPerVertexNonDirectional;
 	const bool bIsVolumetricTranslucencyLightingMode = TranslucencyLightingMode == TLM_VolumetricNonDirectional
