Strata - update error to make it clearer to the users

diff --git a/Engine/Source/Runtime/Engine/Private/Materials/HLSLMaterialTranslator.cpp b/Engine/Source/Runtime/Engine/Private/Materials/HLSLMaterialTranslator.cpp
--- a/Engine/Source/Runtime/Engine/Private/Materials/HLSLMaterialTranslator.cpp
+++ b/Engine/Source/Runtime/Engine/Private/Materials/HLSLMaterialTranslator.cpp
@@ -9519,7 +9519,8 @@ FStrataOperator& FHLSLMaterialTranslator::StrataCompilationRegisterOperator(int3
 
 	if (StrataMaterialExpressionToOperatorIndex.Find(Expression))
 	{
-		Errorf(TEXT("Material %s tries to register a Strata operator twice (asset: %s).\r\n"), *Material->GetDebugName(), *Material->GetAssetPath().ToString());
+		// It is not possible to register/use a Strata BSDF multiple times when creating a topology.
+		Errorf(TEXT("Material %s: It is not possible to uses a Strata BSDF (or any ouput of type StrataData) multiple times within a Strata material topology (asset: %s).\r\n"), *Material->GetDebugName(), *Material->GetAssetPath().ToString());
 		return DefaultOperatorOnError;
 	}
 
