Expose per-component constraint drive strengths etc

diff --git a/Engine/Source/Runtime/Engine/Classes/PhysicsEngine/ConstraintDrives.h b/Engine/Source/Runtime/Engine/Classes/PhysicsEngine/ConstraintDrives.h
--- a/Engine/Source/Runtime/Engine/Classes/PhysicsEngine/ConstraintDrives.h
+++ b/Engine/Source/Runtime/Engine/Classes/PhysicsEngine/ConstraintDrives.h
@@ -106,7 +106,9 @@ private:
 	void SetLinearPositionDrive(bool bEnableXDrive, bool bEnableYDrive, bool bEnableZDrive);
 	void SetLinearVelocityDrive(bool bEnableXDrive, bool bEnableYDrive, bool bEnableZDrive);
 	void SetDriveParams(float InStiffness, float InDamping, float InForceLimit);
+	void SetDriveParams(const FVector& InStiffness, const FVector& InDamping, const FVector& InForceLimit);
 	void GetDriveParams(float& OutStiffness, float& OutDamping, float& OutForceLimit) const;
+	void GetDriveParams(FVector& OutStiffness, FVector& OutDamping, FVector& OutForceLimit) const;
 };
 
 
@@ -178,7 +180,13 @@ private:
 	void SetOrientationDriveSLERP(bool InEnableSLERP);
 	void SetAngularVelocityDriveTwistAndSwing(bool InEnableTwistDrive, bool InEnableSwingDrive);
 	void SetAngularVelocityDriveSLERP(bool InEnableSLERP);
+	// Sets the three drive parameters (swing, twist and slerp) to the same value 
 	void SetDriveParams(float InStiffness, float InDamping, float InForceLimit);
+	// Sets drive parameters in the order swing, twist, slerp
+	void SetDriveParams(const FVector& InStiffness, const FVector& InDamping, const FVector& InForceLimit);
+	// Gets just the swing drive parameters - assuming the single-float set function has previously been used
 	void GetDriveParams(float& OutStiffness, float& OutDamping, float& OutForceLimit) const;
+	// Gets drive parameters in the order swing, twist, slerp
+	void GetDriveParams(FVector& OutStiffness, FVector& OutDamping, FVector& OutForceLimit) const;
 	void SetAngularDriveMode(EAngularDriveMode::Type DriveMode);
 };
diff --git a/Engine/Source/Runtime/Engine/Classes/PhysicsEngine/ConstraintInstance.h b/Engine/Source/Runtime/Engine/Classes/PhysicsEngine/ConstraintInstance.h
--- a/Engine/Source/Runtime/Engine/Classes/PhysicsEngine/ConstraintInstance.h
+++ b/Engine/Source/Runtime/Engine/Classes/PhysicsEngine/ConstraintInstance.h
@@ -747,9 +747,15 @@ public:
 	/** Set the linear drive's strength parameters */
 	void SetLinearDriveParams(float InPositionStrength, float InVelocityStrength, float InForceLimit);
 
-	/** Get the linear drive's strength parameters */
+	/** Set the linear drive's strength parameters per-axis */
+	void SetLinearDriveParams(const FVector& InPositionStrength, const FVector& InVelocityStrength, const FVector& InForceLimit);
+
+	/** Get the linear drive's strength parameters. Assumes all axes are the same so only returns the X values */
 	void GetLinearDriveParams(float& OutPositionStrength, float& OutVelocityStrength, float& OutForceLimit);
 
+	/** Get the linear drive's strength parameters. */
+	void GetLinearDriveParams(FVector& OutPositionStrength, FVector& OutVelocityStrength, FVector& OutForceLimit);
+
 	/** Set which twist and swing orientation drives are enabled. Only applicable when Twist And Swing drive mode is used */
 	void SetOrientationDriveTwistAndSwing(bool bInEnableTwistDrive, bool bInEnableSwingDrive);
 
diff --git a/Engine/Source/Runtime/Engine/Private/PhysicsEngine/ConstraintDrives.cpp b/Engine/Source/Runtime/Engine/Private/PhysicsEngine/ConstraintDrives.cpp
--- a/Engine/Source/Runtime/Engine/Private/PhysicsEngine/ConstraintDrives.cpp
+++ b/Engine/Source/Runtime/Engine/Private/PhysicsEngine/ConstraintDrives.cpp
@@ -79,6 +79,13 @@ void FLinearDriveConstraint::SetDriveParams(float InStiffness, float InDamping,
 	ZDrive.SetDriveParams(InStiffness, InDamping, InForceLimit);
 }
 
+void FLinearDriveConstraint::SetDriveParams(const FVector& InStiffness, const FVector& InDamping, const FVector& InForceLimit)
+{
+	XDrive.SetDriveParams(InStiffness.X, InDamping.X, InForceLimit.X);
+	YDrive.SetDriveParams(InStiffness.Y, InDamping.Y, InForceLimit.Y);
+	ZDrive.SetDriveParams(InStiffness.Z, InDamping.Z, InForceLimit.Z);
+}
+
 void FLinearDriveConstraint::GetDriveParams(float& OutStiffness, float& OutDamping, float& OutForceLimit) const
 {
 	// we set the same value on all drives, just return XDrive ones
@@ -87,6 +94,12 @@ void FLinearDriveConstraint::GetDriveParams(float& OutStiffness, float& OutDampi
 	OutForceLimit = XDrive.MaxForce;
 }
 
+void FLinearDriveConstraint::GetDriveParams(FVector& OutStiffness, FVector& OutDamping, FVector& OutForceLimit) const
+{
+	OutStiffness.Set(XDrive.Stiffness, YDrive.Stiffness, ZDrive.Stiffness);
+	OutDamping.Set(XDrive.Damping, YDrive.Damping, ZDrive.Damping);
+	OutForceLimit.Set(XDrive.MaxForce, YDrive.MaxForce, ZDrive.MaxForce);
+}
 
 void FAngularDriveConstraint::SetAngularDriveMode(EAngularDriveMode::Type DriveMode)
 {
@@ -100,6 +113,13 @@ void FAngularDriveConstraint::SetDriveParams(float InStiffness, float InDamping,
 	SlerpDrive.SetDriveParams(InStiffness, InDamping, InForceLimit);
 }
 
+void FAngularDriveConstraint::SetDriveParams(const FVector& InStiffness, const FVector& InDamping, const FVector& InForceLimit)
+{
+	SwingDrive.SetDriveParams(InStiffness.X, InDamping.X, InForceLimit.X);
+	TwistDrive.SetDriveParams(InStiffness.Y, InDamping.Y, InForceLimit.Y);
+	SlerpDrive.SetDriveParams(InStiffness.Z, InDamping.Z, InForceLimit.Z);
+}
+
 void FAngularDriveConstraint::GetDriveParams(float& OutStiffness, float& OutDamping, float& OutForceLimit) const
 {
 	// we set the same value on all drives, just return SwingDrive ones
@@ -107,3 +127,11 @@ void FAngularDriveConstraint::GetDriveParams(float& OutStiffness, float& OutDamp
 	OutDamping = SwingDrive.Damping;
 	OutForceLimit = SwingDrive.MaxForce;
 }
+
+void FAngularDriveConstraint::GetDriveParams(FVector& OutStiffness, FVector& OutDamping, FVector& OutForceLimit) const
+{
+	OutStiffness.Set(SwingDrive.Stiffness, TwistDrive.Stiffness, TwistDrive.MaxForce);
+	OutDamping.Set(SwingDrive.Damping, TwistDrive.Damping, TwistDrive.MaxForce);
+	OutForceLimit.Set(SwingDrive.MaxForce, TwistDrive.Damping, TwistDrive.MaxForce);
+
+}
diff --git a/Engine/Source/Runtime/Engine/Private/PhysicsEngine/ConstraintInstance.cpp b/Engine/Source/Runtime/Engine/Private/PhysicsEngine/ConstraintInstance.cpp
--- a/Engine/Source/Runtime/Engine/Private/PhysicsEngine/ConstraintInstance.cpp
+++ b/Engine/Source/Runtime/Engine/Private/PhysicsEngine/ConstraintInstance.cpp
@@ -872,12 +872,28 @@ void FConstraintInstance::SetLinearDriveParams(float InSpring, float InDamping,
 	});
 }
 
+/** Function for setting linear motor parameters. */
+void FConstraintInstance::SetLinearDriveParams(const FVector& InSpring, const FVector& InDamping, const FVector& InForceLimit)
+{
+	ProfileInstance.LinearDrive.SetDriveParams(InSpring, InDamping, InForceLimit);
+
+	FPhysicsInterface::ExecuteOnUnbrokenConstraintReadWrite(ConstraintHandle, [this](const FPhysicsConstraintHandle& InUnbrokenConstraint)
+		{
+			FPhysicsInterface::UpdateLinearDrive_AssumesLocked(InUnbrokenConstraint, ProfileInstance.LinearDrive);
+		});
+}
+
 /** Get the linear drive's strength parameters */
 void FConstraintInstance::GetLinearDriveParams(float& OutPositionStrength, float& OutVelocityStrength, float& OutForceLimit)
 {
 	ProfileInstance.LinearDrive.GetDriveParams(OutPositionStrength, OutVelocityStrength, OutForceLimit);
 }
 
+/** Get the linear drive's strength parameters */
+void FConstraintInstance::GetLinearDriveParams(FVector& OutPositionStrength, FVector& OutVelocityStrength, FVector& OutForceLimit)
+{
+	ProfileInstance.LinearDrive.GetDriveParams(OutPositionStrength, OutVelocityStrength, OutForceLimit);
+}
 
 /** Function for setting target angular position. */
 void FConstraintInstance::SetAngularOrientationTarget(const FQuat& InOrientationTarget)
