Level Instance Optimization: pass level instance ids by reference.
Also show level instance package asset when showing loading dialog.

#rnx

diff --git a/Engine/Source/Runtime/Engine/Private/LevelInstance/LevelInstanceSubsystem.cpp b/Engine/Source/Runtime/Engine/Private/LevelInstance/LevelInstanceSubsystem.cpp
--- a/Engine/Source/Runtime/Engine/Private/LevelInstance/LevelInstanceSubsystem.cpp
+++ b/Engine/Source/Runtime/Engine/Private/LevelInstance/LevelInstanceSubsystem.cpp
@@ -100,7 +100,7 @@ bool ULevelInstanceSubsystem::DoesSupportWorldType(EWorldType::Type WorldType) c
 	return Super::DoesSupportWorldType(WorldType) || WorldType == EWorldType::EditorPreview || WorldType == EWorldType::Inactive;
 }
 
-ILevelInstanceInterface* ULevelInstanceSubsystem::GetLevelInstance(FLevelInstanceID LevelInstanceID) const
+ILevelInstanceInterface* ULevelInstanceSubsystem::GetLevelInstance(const FLevelInstanceID& LevelInstanceID) const
 {
 	if (ILevelInstanceInterface*const* LevelInstance = RegisteredLevelInstances.Find(LevelInstanceID))
 	{
@@ -248,7 +248,7 @@ void ULevelInstanceSubsystem::UpdateStreamingState()
 		for (const TPair<ILevelInstanceInterface*, bool>& Pair : LevelInstancesToLoadOrUpdateCopy)
 		{
 #if WITH_EDITOR
-			SlowTask.EnterProgressFrame(1.f, LOCTEXT("LoadingLevelInstance", "Loading Level Instance"));
+			SlowTask.EnterProgressFrame(1.f, FText::Format(LOCTEXT("LoadingLevelInstance", "Loading Level Instance {0}"), FText::FromString(Pair.Key->GetWorldAsset().ToString())));
 #endif
 			LoadLevelInstance(Pair.Key);
 		}
@@ -380,16 +380,16 @@ void ULevelInstanceSubsystem::UnloadLevelInstance(const FLevelInstanceID& LevelI
 			if (ULevel* LoadedLevel = LevelInstance.LevelStreaming->GetLoadedLevel())
 			{
 				ForEachActorInLevel(LoadedLevel, [this](AActor* LevelActor)
+				{
+					if (ILevelInstanceInterface* LevelInstance = Cast<ILevelInstanceInterface>(LevelActor))
 					{
-						if (ILevelInstanceInterface* LevelInstance = Cast<ILevelInstanceInterface>(LevelActor))
-						{
-							// Make sure to remove from pending loads if we are unloading child can't be loaded
-							LevelInstancesToLoadOrUpdate.Remove(LevelInstance);
+						// Make sure to remove from pending loads if we are unloading child can't be loaded
+						LevelInstancesToLoadOrUpdate.Remove(LevelInstance);
 
-							UnloadLevelInstance(LevelInstance->GetLevelInstanceID());
-						}
-						return true;
-					});
+						UnloadLevelInstance(LevelInstance->GetLevelInstanceID());
+					}
+					return true;
+				});
 			}
 
 			ULevelStreamingLevelInstance::UnloadInstance(LevelInstance.LevelStreaming);
@@ -2194,7 +2194,7 @@ bool ULevelInstanceSubsystem::HasChildEdit(const ILevelInstanceInterface* LevelI
 	return ChildEditCountPtr && *ChildEditCountPtr;
 }
 
-void ULevelInstanceSubsystem::OnCommitChild(FLevelInstanceID LevelInstanceID, bool bChildChanged)
+void ULevelInstanceSubsystem::OnCommitChild(const FLevelInstanceID& LevelInstanceID, bool bChildChanged)
 {
 	int32& ChildEditCount = ChildEdits.FindChecked(LevelInstanceID);
 	check(ChildEditCount > 0);
@@ -2206,7 +2206,7 @@ void ULevelInstanceSubsystem::OnCommitChild(FLevelInstanceID LevelInstanceID, bo
 	}
 }
 
-void ULevelInstanceSubsystem::OnEditChild(FLevelInstanceID LevelInstanceID)
+void ULevelInstanceSubsystem::OnEditChild(const FLevelInstanceID& LevelInstanceID)
 {
 	int32& ChildEditCount = ChildEdits.FindOrAdd(LevelInstanceID, 0);
 	// Child edit count can reach 2 maximum in the Context of creating a LevelInstance inside an already editing child level instance
diff --git a/Engine/Source/Runtime/Engine/Public/LevelInstance/LevelInstanceSubsystem.h b/Engine/Source/Runtime/Engine/Public/LevelInstance/LevelInstanceSubsystem.h
--- a/Engine/Source/Runtime/Engine/Public/LevelInstance/LevelInstanceSubsystem.h
+++ b/Engine/Source/Runtime/Engine/Public/LevelInstance/LevelInstanceSubsystem.h
@@ -44,7 +44,7 @@ public:
 	virtual void UpdateStreamingState() override;
 	//~ End UWorldSubsystem Interface.
 
-	ILevelInstanceInterface* GetLevelInstance(FLevelInstanceID LevelInstanceID) const;
+	ILevelInstanceInterface* GetLevelInstance(const FLevelInstanceID& LevelInstanceID) const;
 	ILevelInstanceInterface* GetOwningLevelInstance(const ULevel* Level) const;
 	FLevelInstanceID RegisterLevelInstance(ILevelInstanceInterface* LevelInstance);
 	void UnregisterLevelInstance(ILevelInstanceInterface* LevelInstance);
@@ -135,8 +135,8 @@ private:
 #if WITH_EDITOR
 	void RegisterLoadedLevelStreamingLevelInstanceEditor(ULevelStreamingLevelInstanceEditor* LevelStreaming);
 
-	void OnEditChild(FLevelInstanceID LevelInstanceID);
-	void OnCommitChild(FLevelInstanceID LevelInstanceID, bool bChildChanged);
+	void OnEditChild(const FLevelInstanceID& LevelInstanceID);
+	void OnCommitChild(const FLevelInstanceID& LevelInstanceID, bool bChildChanged);
 
 	bool ForEachLevelInstanceChildImpl(const ILevelInstanceInterface* LevelInstance, bool bRecursive, TFunctionRef<bool(const ILevelInstanceInterface*)> Operation) const;
 	bool ForEachLevelInstanceChildImpl(ILevelInstanceInterface* LevelInstance, bool bRecursive, TFunctionRef<bool(ILevelInstanceInterface*)> Operation) const;
