Strata - material compilation adaptation to fit in memory budget. Very simple for now but this is a start: if the BSDF do not fit in memory, set all the nodes to use parameter blending. If failing, then fail compilation.
BSDF feature flags is now propagated to all operators during the StrataGenerateMaterialTopologyTree pass. This is to be able to compute correct byte size, and later reduce the BSDF feature set one by one.

diff --git a/Engine/Source/Runtime/Engine/Private/Materials/HLSLMaterialTranslator.cpp b/Engine/Source/Runtime/Engine/Private/Materials/HLSLMaterialTranslator.cpp
--- a/Engine/Source/Runtime/Engine/Private/Materials/HLSLMaterialTranslator.cpp
+++ b/Engine/Source/Runtime/Engine/Private/Materials/HLSLMaterialTranslator.cpp
@@ -48,6 +48,14 @@ static inline bool IsDebugTextureSampleEnabled()
 
 bool Engine_IsStrataEnabled();
 
+static uint32 GetStrataBytePerPixel()
+{
+	static const auto CVarStrataBytePerPixel = IConsoleManager::Get().FindTConsoleVariableDataInt(TEXT("r.Strata.BytesPerPixel"));
+	check(CVarStrataBytePerPixel);
+	const uint32 StrataBytePerPixel = CVarStrataBytePerPixel ? CVarStrataBytePerPixel->GetValueOnAnyThread() : 0;
+	return StrataBytePerPixel;
+}
+
 /** @return the vector type containing a given number of components. */
 static inline EMaterialValueType GetVectorType(uint32 NumComponents)
 {
@@ -698,7 +706,7 @@ bool FHLSLMaterialTranslator::Translate()
 			TArray<FShaderCodeChunk> TempChunks;
 			AssignTempScope(TempChunks);
 
-			FrontMaterialExpr->StrataGenerateMaterialTopologyTree(this, nullptr, 0); 
+			FrontMaterialExpr->StrataGenerateMaterialTopologyTree(this, nullptr, 0);
 			if (!StrataGenerateDerivedMaterialOperatorData())
 			{
 				Errorf(TEXT("Strata material errors encountered."));
@@ -1922,9 +1930,7 @@ void FHLSLMaterialTranslator::GetMaterialEnvironment(EShaderPlatform InPlatform,
 		// Now write some feedback to the user
 		{
 			// Output some debug info as comment in code and in the material stat window
-			static const auto CVarStrataBytePerPixel = IConsoleManager::Get().FindTConsoleVariableDataInt(TEXT("r.Strata.BytesPerPixel"));
-			check(CVarStrataBytePerPixel);
-			const uint32 StrataBytePerPixel = CVarStrataBytePerPixel ? CVarStrataBytePerPixel->GetValueOnAnyThread() : 0;
+			const uint32 StrataBytePerPixel = GetStrataBytePerPixel();
 
 			OutEnvironment.SetDefine(TEXT("STRATA_CLAMPED_BSDF_COUNT"), StrataMaterialBSDFCount);
 
@@ -1932,18 +1938,19 @@ void FHLSLMaterialTranslator::GetMaterialEnvironment(EShaderPlatform InPlatform,
 
 			StrataMaterialDescription += FString::Printf(TEXT("----- STRATA -----\r\n"));
 			StrataMaterialDescription += FString::Printf(TEXT("StrataCompilationInfo -\r\n"));
-			StrataMaterialDescription += FString::Printf(TEXT(" - Byte Per Pixel Budget      %u\r\n"), StrataBytePerPixel);
-			StrataMaterialDescription += FString::Printf(TEXT(" - Material complexity        %s\r\n"), bStrataMaterialIsSingle ? TEXT("SINGLE") : (bStrataMaterialIsSimple ? TEXT("SIMPLE") : TEXT("COMPLEX")));
-			StrataMaterialDescription += FString::Printf(TEXT(" - Requested Byte Size        %u (%d UINT32)\r\n"), StrataMaterialRequestedSizeByte, StrataMaterialRequestedSizeByte / 4);
-			StrataMaterialDescription += FString::Printf(TEXT(" - TotalBSDFCount             %i\r\n"), StrataMaterialBSDFCount);
-			StrataMaterialDescription += FString::Printf(TEXT(" - SharedLocalBasesCount      %i\r\n"), FinalUsedSharedLocalBasesCount);
+			StrataMaterialDescription += FString::Printf(TEXT(" - Byte Per Pixel Budget                      %u\r\n"), StrataBytePerPixel);
+			StrataMaterialDescription += FString::Printf(TEXT(" - Requested Byte Size before simplification  %u (%d UINT32)\r\n"), StrataSimplificationStatus.OriginalRequestedByteSize, StrataSimplificationStatus.OriginalRequestedByteSize / 4);
+			StrataMaterialDescription += FString::Printf(TEXT(" - Requested Byte Size after simplification   %u (%d UINT32)\r\n"), StrataMaterialRequestedSizeByte, StrataMaterialRequestedSizeByte / 4);
+			StrataMaterialDescription += FString::Printf(TEXT(" - Material complexity                        %s\r\n"), bStrataMaterialIsSingle ? TEXT("SINGLE") : (bStrataMaterialIsSimple ? TEXT("SIMPLE") : TEXT("COMPLEX")));
+			StrataMaterialDescription += FString::Printf(TEXT(" - TotalBSDFCount                             %i\r\n"), StrataMaterialBSDFCount);
+			StrataMaterialDescription += FString::Printf(TEXT(" - SharedLocalBasesCount                      %i\r\n"), FinalUsedSharedLocalBasesCount);
 
 			for (int32 OpIt = 0; OpIt < StrataMaterialExpressionRegisteredOperators.Num(); ++OpIt)
 			{
 				FStrataOperator& BSDFOperator = StrataMaterialExpressionRegisteredOperators[OpIt];
-				if (BSDFOperator.BSDFIndex == INDEX_NONE && !BSDFOperator.bRootOfParameterBlendingSubTree)
+				if (BSDFOperator.BSDFIndex == INDEX_NONE || BSDFOperator.IsDiscarded())
 				{
-					continue;	// not a bsdf or not the root of a parameter blending subtree so there is not local basis to register
+					continue;	// not a BSDF or if discarded (i.e. not the root of a parameter blending subtree), then there is no local basis to register
 				}
 				if (BSDFOperator.BSDFRegisteredSharedLocalBasis.NormalCodeChunk == INDEX_NONE && BSDFOperator.BSDFRegisteredSharedLocalBasis.TangentCodeChunk == INDEX_NONE)
 				{
@@ -9610,9 +9617,9 @@ void FHLSLMaterialTranslator::StrataEvaluateSharedLocalBases(
 	for (int32 OpIt = 0; OpIt < StrataMaterialExpressionRegisteredOperators.Num(); ++OpIt)
 	{
 		FStrataOperator& BSDFOperator = StrataMaterialExpressionRegisteredOperators[OpIt];
-		if (BSDFOperator.BSDFIndex == INDEX_NONE && !BSDFOperator.bRootOfParameterBlendingSubTree)
+		if (BSDFOperator.BSDFIndex == INDEX_NONE || BSDFOperator.IsDiscarded())
 		{
-			continue;	// not a bsdf or not the root of a parameter blending subtree so there is not local basis to register
+			continue;	// not a BSDF or if discarded (i.e. not the root of a parameter blending subtree), then there is no local basis to register
 		}
 
 		if (BSDFOperator.BSDFRegisteredSharedLocalBasis.NormalCodeChunk == INDEX_NONE && BSDFOperator.BSDFRegisteredSharedLocalBasis.TangentCodeChunk == INDEX_NONE)
@@ -9766,196 +9773,125 @@ bool FHLSLMaterialTranslator::StrataGenerateDerivedMaterialOperatorData()
 		}
 	}
 
-	//
-	// Parse the tree and mark nodes that are the root of a subtree using parameter blending, while other nodes in that tree are forced to use parameter blending.
-	// Allocate BSDFIndex at the same time.
-	// Check BSDF that should be unique.
-	//
+	const uint32 StrataBytePerPixel = GetStrataBytePerPixel();
+	do 
 	{
-		bool bHasUnlit = false;
-		bool bHasVFogCloud = false;
-		bool bHasHair = false;
-		bool bHasEye = false;
-		bool bHasSLW = false;
-		bool bHasSlab = false;
+		// Reset some data for simplifiucation iteration
+		StrataMaterialBSDFCount = 0;
 
-		int VOpTopBranchCountTaken = 0;
-		int VOpBottomBranchCountTaken = 0;
-		bool bStrataUsesVerticalLayering = false;
-		bool bOperatorEncountered = false;
-
-		std::function<void(FStrataOperator&, bool)> WalkOperators = [&](FStrataOperator& CurrentOperator, bool bInsideParameterBlendingSubTree) -> void
+		//
+		// Parse the tree and mark nodes that are the root of a subtree using parameter blending, while other nodes in that tree are forced to use parameter blending.
+		// Allocate BSDFIndex at the same time.
+		// Check BSDF that should be unique.
+		//
 		{
-			const bool bCurrentOpRequestParameterBlending	= CurrentOperator.bNodeRequestParameterBlending;
-			const bool bRootOfParameterBlendingSubTree		= bCurrentOpRequestParameterBlending && !bInsideParameterBlendingSubTree;
-			const bool bUseParameterBlending				= bCurrentOpRequestParameterBlending || bInsideParameterBlendingSubTree;
+			bool bHasUnlit = false;
+			bool bHasVFogCloud = false;
+			bool bHasHair = false;
+			bool bHasEye = false;
+			bool bHasSLW = false;
+			bool bHasSlab = false;
 
-			CurrentOperator.bUseParameterBlending			= bUseParameterBlending;
-			CurrentOperator.bRootOfParameterBlendingSubTree	= bRootOfParameterBlendingSubTree;
+			int VOpTopBranchCountTaken = 0;
+			int VOpBottomBranchCountTaken = 0;
+			bool bStrataUsesVerticalLayering = false;
+			bool bOperatorEncountered = false;
 
-			switch (CurrentOperator.OperatorType)
-			{
-			case STRATA_OPERATOR_VERTICAL:
+			std::function<void(FStrataOperator&, bool)> WalkOperators = [&](FStrataOperator& CurrentOperator, bool bInsideParameterBlendingSubTree) -> void
 			{
-				WalkOperators(StrataMaterialExpressionRegisteredOperators[CurrentOperator.LeftIndex], bUseParameterBlending);
-				WalkOperators(StrataMaterialExpressionRegisteredOperators[CurrentOperator.RightIndex], bUseParameterBlending);
-				bOperatorEncountered = true;
-				break;
-			}
-			case STRATA_OPERATOR_HORIZONTAL:
-			case STRATA_OPERATOR_ADD:
-			{
-				WalkOperators(StrataMaterialExpressionRegisteredOperators[CurrentOperator.LeftIndex], bUseParameterBlending);
-				WalkOperators(StrataMaterialExpressionRegisteredOperators[CurrentOperator.RightIndex], bUseParameterBlending);
-				bOperatorEncountered = true;
-				break;
-			}
-			case STRATA_OPERATOR_WEIGHT:
-			case STRATA_OPERATOR_THINFILM:
-			{
-				WalkOperators(StrataMaterialExpressionRegisteredOperators[CurrentOperator.LeftIndex], bUseParameterBlending);
-				bOperatorEncountered = true;
-				break;
-			}
-			case STRATA_OPERATOR_BSDF:
-			{
-				if (!bInsideParameterBlendingSubTree)
-				{
-					CurrentOperator.BSDFIndex = StrataMaterialBSDFCount++;
-				}
-
-				bHasUnlit |= CurrentOperator.BSDFType == STRATA_BSDF_TYPE_UNLIT;
-				bHasVFogCloud |= CurrentOperator.BSDFType == STRATA_BSDF_TYPE_VOLUMETRICFOGCLOUD;
-				bHasHair |= CurrentOperator.BSDFType == STRATA_BSDF_TYPE_HAIR;
-				bHasEye |= CurrentOperator.BSDFType == STRATA_BSDF_TYPE_EYE;
-				bHasSLW |= CurrentOperator.BSDFType == STRATA_BSDF_TYPE_SINGLELAYERWATER;
-				bHasSlab |= CurrentOperator.BSDFType == STRATA_BSDF_TYPE_SLAB;
-				break;
-			}
-			}
+				const bool bCurrentOpRequestParameterBlending	= CurrentOperator.bNodeRequestParameterBlending || StrataSimplificationStatus.bRunFullSimplification;
+				const bool bRootOfParameterBlendingSubTree		= bCurrentOpRequestParameterBlending && !bInsideParameterBlendingSubTree;
+				const bool bUseParameterBlending				= bCurrentOpRequestParameterBlending || bInsideParameterBlendingSubTree;
 
-			// We mark the top of a parameter blending tree as a BSDF now to allocate a slot for it that can then be used next for non-parameter blending operations.
-			// Intermediate parameter blending BSDF and operation will be done inline and stored in FStrataData.
-			if (bRootOfParameterBlendingSubTree)
-			{
-				CurrentOperator.OperatorType = STRATA_OPERATOR_BSDF;
-				CurrentOperator.BSDFIndex = StrataMaterialBSDFCount++;
-				// We do not reset LeftIndex and RightIndex because those are needed to recover local tangent basis information needed with parameter blending.
-			}
-
-			// When at least one vertical operator exists that is not parameter blending, we can enabled writing to opaque rough refraction buffer.
-			bStrataUsesVerticalLayering = !CurrentOperator.bUseParameterBlending && CurrentOperator.OperatorType == STRATA_OPERATOR_VERTICAL;
-		};
+				CurrentOperator.bUseParameterBlending			= bUseParameterBlending;
+				CurrentOperator.bRootOfParameterBlendingSubTree	= bRootOfParameterBlendingSubTree;
 
-		WalkOperators(*StrataMaterialRootOperator, false);
-
-		const EStrataBlendMode StrataBlendMode = Material->GetStrataBlendMode();
-		const bool bIsOpaqueOrMasked = StrataBlendMode == EStrataBlendMode::SBM_Opaque || StrataBlendMode == EStrataBlendMode::SBM_Masked;
-		bStrataOutputsOpaqueRoughRefractions = bStrataUsesVerticalLayering && bIsOpaqueOrMasked;
-		bStrataMaterialIsUnlitNode = bHasUnlit;
-
-		if ((bHasUnlit || bHasVFogCloud || bHasHair || bHasEye || bHasSLW) && StrataMaterialBSDFCount > 1)
-		{
-			Errorf(TEXT("Unlit, Fog/Cloud, Hair or SingleLayerWater must be used in isolation. See %s (asset: %s).\r\n"), *Material->GetDebugName(), *Material->GetAssetPath().ToString());
-			// Even though we could support Unlit with slab.
-		}
+				switch (CurrentOperator.OperatorType)
+				{
+				case STRATA_OPERATOR_VERTICAL:
+				{
+					WalkOperators(StrataMaterialExpressionRegisteredOperators[CurrentOperator.LeftIndex], bUseParameterBlending);
+					WalkOperators(StrataMaterialExpressionRegisteredOperators[CurrentOperator.RightIndex], bUseParameterBlending);
+					bOperatorEncountered = true;
+					break;
+				}
+				case STRATA_OPERATOR_HORIZONTAL:
+				case STRATA_OPERATOR_ADD:
+				{
+					WalkOperators(StrataMaterialExpressionRegisteredOperators[CurrentOperator.LeftIndex], bUseParameterBlending);
+					WalkOperators(StrataMaterialExpressionRegisteredOperators[CurrentOperator.RightIndex], bUseParameterBlending);
+					bOperatorEncountered = true;
+					break;
+				}
+				case STRATA_OPERATOR_WEIGHT:
+				case STRATA_OPERATOR_THINFILM:
+				{
+					WalkOperators(StrataMaterialExpressionRegisteredOperators[CurrentOperator.LeftIndex], bUseParameterBlending);
+					bOperatorEncountered = true;
+					break;
+				}
+				case STRATA_OPERATOR_BSDF:
+				{
+					if (!bInsideParameterBlendingSubTree)
+					{
+						CurrentOperator.BSDFIndex = StrataMaterialBSDFCount++;
+					}
 
-		if ((bHasUnlit || bHasVFogCloud || bHasHair || bHasEye || bHasSLW) && bOperatorEncountered)
-		{
-			Errorf(TEXT("Unlit, Fog/Cloud, Hair or SingleLayerWater cannot be used with operators. See %s (asset: %s).\r\n"), *Material->GetDebugName(), *Material->GetAssetPath().ToString());
-			// This is because it will results in simpler lighting loops focusin on slab.
-		}
+					bHasUnlit |= CurrentOperator.BSDFType == STRATA_BSDF_TYPE_UNLIT;
+					bHasVFogCloud |= CurrentOperator.BSDFType == STRATA_BSDF_TYPE_VOLUMETRICFOGCLOUD;
+					bHasHair |= CurrentOperator.BSDFType == STRATA_BSDF_TYPE_HAIR;
+					bHasEye |= CurrentOperator.BSDFType == STRATA_BSDF_TYPE_EYE;
+					bHasSLW |= CurrentOperator.BSDFType == STRATA_BSDF_TYPE_SINGLELAYERWATER;
+					bHasSlab |= CurrentOperator.BSDFType == STRATA_BSDF_TYPE_SLAB;
+					break;
+				}
+				}
 
-		if (StrataMaterialBSDFCount > STRATA_MAX_BSDF_COUNT)
-		{
-			Errorf(TEXT("Material tries to register more BSDF than can be supproted (%d > %d). See %s (asset: %s).\r\n"), StrataMaterialBSDFCount, STRATA_MAX_BSDF_COUNT, *Material->GetDebugName(), *Material->GetAssetPath().ToString());
-		}
-	}
+				// We mark the top of a parameter blending tree as a BSDF now to allocate a slot for it that can then be used next for non-parameter blending operations.
+				// Intermediate parameter blending BSDF and operation will be done inline and stored in FStrataData.
+				if (bRootOfParameterBlendingSubTree)
+				{
+					CurrentOperator.OperatorType = STRATA_OPERATOR_BSDF;
+					CurrentOperator.BSDFIndex = StrataMaterialBSDFCount++;
+					// We do not reset LeftIndex and RightIndex because those are needed to recover local tangent basis information needed with parameter blending.
+				}
 
-	// STRATA_TODO: operation using parameter blending are not actually discarded so they still occupy a spot in the operation array in the compiler and in the shader also.
-	// Even thought the compiler will remove unused operators from the shader code, we will still be limited to STRATA_MAX_OPERATOR_COUNT even with parameter blending.
-	// This can be fixed by remapping all operation index and make sure this is the index that is always used in the compiler and specified to the shader.
+				// When at least one vertical operator exists that is not parameter blending, we can enabled writing to opaque rough refraction buffer.
+				bStrataUsesVerticalLayering = !CurrentOperator.bUseParameterBlending && CurrentOperator.OperatorType == STRATA_OPERATOR_VERTICAL;
+			};
 
-	//
-	// Make sure all the types have valid children operator indices
-	//
-	for (auto& It : StrataMaterialExpressionRegisteredOperators)
-	{
-		if (It.IsDiscarded())
-		{
-			continue; // ignore discarded operations in sub tree using parameter blending
-		}
+			WalkOperators(*StrataMaterialRootOperator, false);
 
-		check(It.Index != INDEX_NONE);
+			const EStrataBlendMode StrataBlendMode = Material->GetStrataBlendMode();
+			const bool bIsOpaqueOrMasked = StrataBlendMode == EStrataBlendMode::SBM_Opaque || StrataBlendMode == EStrataBlendMode::SBM_Masked;
+			bStrataOutputsOpaqueRoughRefractions = bStrataUsesVerticalLayering && bIsOpaqueOrMasked;
+			bStrataMaterialIsUnlitNode = bHasUnlit;
 
-		switch (It.OperatorType)
-		{
-		// Operators without any child
-		case STRATA_OPERATOR_BSDF:
-		{
-			if (!It.bUseParameterBlending)
+			if ((bHasUnlit || bHasVFogCloud || bHasHair || bHasEye || bHasSLW) && StrataMaterialBSDFCount > 1)
 			{
-				// When using parameter blending, we need to keep indices to be able to recover local basis information for normal blending.
-				check(It.LeftIndex == INDEX_NONE && It.RightIndex == INDEX_NONE && It.BSDFIndex != INDEX_NONE);
+				Errorf(TEXT("Unlit, Fog/Cloud, Hair or SingleLayerWater must be used in isolation. See %s (asset: %s).\r\n"), *Material->GetDebugName(), *Material->GetAssetPath().ToString());
+				// Even though we could support Unlit with slab.
 			}
-			break;
-		}
-
-		// Operators with two children
-		case STRATA_OPERATOR_HORIZONTAL:
-		case STRATA_OPERATOR_VERTICAL:
-		case STRATA_OPERATOR_ADD:
-		{
-			check(It.RightIndex != INDEX_NONE);
-		}
-		// Fallthrough
 
-		// Operators with a single child
-		case STRATA_OPERATOR_WEIGHT:
-		case STRATA_OPERATOR_THINFILM:
-		{
-			check(It.LeftIndex != INDEX_NONE);
-		}
-		// Fallthrough
-		}
-	}
-
-	//
-	// Compute the maximum depth from the BSDF node for each operator
-	//
-	{
-		std::function<void(FStrataOperator&)>  WalkOperatorsToRoot = [&](FStrataOperator& CurrentOperator) -> void
-		{
-			switch (CurrentOperator.OperatorType)
-			{
-			case STRATA_OPERATOR_WEIGHT:
-			case STRATA_OPERATOR_THINFILM:
-			{
-				CurrentOperator.MaxDistanceFromLeaves = StrataMaterialExpressionRegisteredOperators[CurrentOperator.LeftIndex].MaxDistanceFromLeaves + 1;
-				break;
-			}
-			case STRATA_OPERATOR_VERTICAL:
-			case STRATA_OPERATOR_HORIZONTAL:
-			case STRATA_OPERATOR_ADD:
+			if ((bHasUnlit || bHasVFogCloud || bHasHair || bHasEye || bHasSLW) && bOperatorEncountered)
 			{
-				CurrentOperator.MaxDistanceFromLeaves = FMath::Max(
-					StrataMaterialExpressionRegisteredOperators[CurrentOperator.LeftIndex].MaxDistanceFromLeaves,
-					StrataMaterialExpressionRegisteredOperators[CurrentOperator.RightIndex].MaxDistanceFromLeaves) + 1;
-				break;
-			}
-			case STRATA_OPERATOR_BSDF:
-			{
-				CurrentOperator.MaxDistanceFromLeaves = 0;
-				break;
-			}
+				Errorf(TEXT("Unlit, Fog/Cloud, Hair or SingleLayerWater cannot be used with operators. See %s (asset: %s).\r\n"), *Material->GetDebugName(), *Material->GetAssetPath().ToString());
+				// This is because it will results in simpler lighting loops focusin on slab.
 			}
 
-			if (CurrentOperator.ParentIndex != INDEX_NONE)
+			if (StrataMaterialBSDFCount > STRATA_MAX_BSDF_COUNT)
 			{
-				WalkOperatorsToRoot(StrataMaterialExpressionRegisteredOperators[CurrentOperator.ParentIndex]);
+				Errorf(TEXT("Material tries to register more BSDF than can be supproted (%d > %d). See %s (asset: %s).\r\n"), StrataMaterialBSDFCount, STRATA_MAX_BSDF_COUNT, *Material->GetDebugName(), *Material->GetAssetPath().ToString());
 			}
-		};
+		}
+
+		// STRATA_TODO: operation using parameter blending are not actually discarded so they still occupy a spot in the operation array in the compiler and in the shader also.
+		// Even thought the compiler will remove unused operators from the shader code, we will still be limited to STRATA_MAX_OPERATOR_COUNT even with parameter blending.
+		// This can be fixed by remapping all operation index and make sure this is the index that is always used in the compiler and specified to the shader.
 
+		//
+		// Make sure all the types have valid children operator indices
+		//
 		for (auto& It : StrataMaterialExpressionRegisteredOperators)
 		{
 			if (It.IsDiscarded())
@@ -9963,247 +9899,340 @@ bool FHLSLMaterialTranslator::StrataGenerateDerivedMaterialOperatorData()
 				continue; // ignore discarded operations in sub tree using parameter blending
 			}
 
-			if (It.OperatorType == STRATA_OPERATOR_BSDF)
-			{
-				// Recursively parse all node from BSDF to the root node and update the necessary properties.
-				WalkOperatorsToRoot(It);
-			}
-		}
-	}
-
-	//
-	// Compute IsTop or IsBottom layer using a depth first tree visit while counting vertical right and left branches taken.
-	//
-	{
-		int VOpTopBranchCountTaken = 0;
-		int VOpBottomBranchCountTaken = 0;
+			check(It.Index != INDEX_NONE);
 
-		std::function<void(FStrataOperator&)> WalkOperators = [&](FStrataOperator& CurrentOperator) -> void
-		{
-			switch (CurrentOperator.OperatorType)
+			switch (It.OperatorType)
 			{
-			case STRATA_OPERATOR_VERTICAL:
+			// Operators without any child
+			case STRATA_OPERATOR_BSDF:
 			{
-				VOpTopBranchCountTaken++;
-				WalkOperators(StrataMaterialExpressionRegisteredOperators[CurrentOperator.LeftIndex]);
-				VOpTopBranchCountTaken--;
-				VOpBottomBranchCountTaken++;
-				WalkOperators(StrataMaterialExpressionRegisteredOperators[CurrentOperator.RightIndex]);
-				VOpBottomBranchCountTaken--;
+				if (!It.bUseParameterBlending)
+				{
+					// When using parameter blending, we need to keep indices to be able to recover local basis information for normal blending.
+					check(It.LeftIndex == INDEX_NONE && It.RightIndex == INDEX_NONE && It.BSDFIndex != INDEX_NONE);
+				}
 				break;
 			}
+
+			// Operators with two children
 			case STRATA_OPERATOR_HORIZONTAL:
+			case STRATA_OPERATOR_VERTICAL:
 			case STRATA_OPERATOR_ADD:
 			{
-				WalkOperators(StrataMaterialExpressionRegisteredOperators[CurrentOperator.LeftIndex]);
-				WalkOperators(StrataMaterialExpressionRegisteredOperators[CurrentOperator.RightIndex]);
-				break;
+				check(It.RightIndex != INDEX_NONE);
 			}
+			// Fallthrough
+
+			// Operators with a single child
 			case STRATA_OPERATOR_WEIGHT:
 			case STRATA_OPERATOR_THINFILM:
 			{
-				WalkOperators(StrataMaterialExpressionRegisteredOperators[CurrentOperator.LeftIndex]);
-				break;
+				check(It.LeftIndex != INDEX_NONE);
 			}
-			case STRATA_OPERATOR_BSDF:
-			{
-				const int32 VopCount = VOpTopBranchCountTaken + VOpBottomBranchCountTaken;
-				CurrentOperator.LayerDepth = VOpBottomBranchCountTaken;
-				CurrentOperator.bIsTop = VopCount == 0 || VOpTopBranchCountTaken == VopCount;
-				CurrentOperator.bIsBottom = VopCount == 0 || VOpBottomBranchCountTaken == VopCount;
-				break;
+			// Fallthrough
 			}
-			}
-		};
-
-		WalkOperators(*StrataMaterialRootOperator);
-	}
-
-	//
-	// Compute the size of the material
-	//
-	{
-		// Compute the shared local basis count only
-		// But we cannot use StrataEvaluateSharedLocalBases here, becaise the material has not been compiled yet so al lthe bases would just default to the same.
-		// STRATA_TODO: can we do that material generation in two passes? 
-		//		1- A first one to evaluate the normal/tangent code
-		//		2- Operators are processed and simplification computed based on memory budget
-		//		3- Material is finally compiled for with operator updated to fit in memory budget.
-		uint8 UsedSharedLocalBasesCount = StrataMaterialBSDFCount;
-
-		const uint32 UintByteSize = sizeof(uint32);
-		StrataMaterialRequestedSizeByte = 0;
+		}
 
-		// 1. Evaluate simple/single BSDF
-		bStrataMaterialIsSimple = StrataMaterialBSDFCount == 1;
-		bStrataMaterialIsSingle = StrataMaterialBSDFCount == 1;
-		for (auto& It : StrataMaterialExpressionRegisteredOperators)
+		//
+		// Compute the maximum depth from the BSDF node for each operator
+		//
 		{
-			if (It.IsDiscarded())
+			std::function<void(FStrataOperator&)>  WalkOperatorsToRoot = [&](FStrataOperator& CurrentOperator) -> void
 			{
-				continue; // ignore discarded operations in sub tree using parameter blending
-			}
-
-			switch (It.OperatorType)
-			{
-			case STRATA_OPERATOR_BSDF:
-			{
-				// From the compiler side, we can only assume the top layer has gray scale luminance weight.
-				const bool bMayHaveColoredWeight = !It.bIsTop;
-
-				switch (It.BSDFType)
-				{
-				case STRATA_BSDF_TYPE_SLAB:
+				switch (CurrentOperator.OperatorType)
 				{
-					bStrataMaterialIsSimple = bStrataMaterialIsSimple && !bMayHaveColoredWeight && !It.bBSDFHasAnisotropy && !It.bBSDFHasEdgeColor && !It.bBSDFHasFuzz && !It.bBSDFHasSecondRoughnessOrSimpleClearCoat && !It.bBSDFHasMFPPluggedIn && !It.bBSDFHasSSS;
-					bStrataMaterialIsSingle = bStrataMaterialIsSingle && !bMayHaveColoredWeight && !It.bBSDFHasAnisotropy;
-					break;
-				}
-				case STRATA_BSDF_TYPE_HAIR:
+				case STRATA_OPERATOR_WEIGHT:
+				case STRATA_OPERATOR_THINFILM:
 				{
-					bStrataMaterialIsSimple = false;
-					bStrataMaterialIsSingle = false;
+					CurrentOperator.MaxDistanceFromLeaves = StrataMaterialExpressionRegisteredOperators[CurrentOperator.LeftIndex].MaxDistanceFromLeaves + 1;
 					break;
 				}
-				case STRATA_BSDF_TYPE_EYE:
+				case STRATA_OPERATOR_VERTICAL:
+				case STRATA_OPERATOR_HORIZONTAL:
+				case STRATA_OPERATOR_ADD:
 				{
-					bStrataMaterialIsSimple = false;
-					bStrataMaterialIsSingle = false;
+					CurrentOperator.MaxDistanceFromLeaves = FMath::Max(
+						StrataMaterialExpressionRegisteredOperators[CurrentOperator.LeftIndex].MaxDistanceFromLeaves,
+						StrataMaterialExpressionRegisteredOperators[CurrentOperator.RightIndex].MaxDistanceFromLeaves) + 1;
 					break;
 				}
-				case STRATA_BSDF_TYPE_SINGLELAYERWATER:
+				case STRATA_OPERATOR_BSDF:
 				{
-					bStrataMaterialIsSimple = false;
-					bStrataMaterialIsSingle = false;
+					CurrentOperator.MaxDistanceFromLeaves = 0;
 					break;
 				}
 				}
-				break;
-			}
-			}
-		}
-		bStrataMaterialIsSingle = bStrataMaterialIsSingle && !bStrataMaterialIsSimple;
 
-		// 2. Header
+				if (CurrentOperator.ParentIndex != INDEX_NONE)
+				{
+					WalkOperatorsToRoot(StrataMaterialExpressionRegisteredOperators[CurrentOperator.ParentIndex]);
+				}
+			};
 
-		if (!bStrataMaterialIsSimple && !bStrataMaterialIsSingle) // header written later, 
-		{
-			// Packed Header
-			StrataMaterialRequestedSizeByte += UintByteSize;
+			for (auto& It : StrataMaterialExpressionRegisteredOperators)
+			{
+				if (It.IsDiscarded())
+				{
+					continue; // ignore discarded operations in sub tree using parameter blending
+				}
 
-			// Shared local bases between BSDFs
-			StrataMaterialRequestedSizeByte += UsedSharedLocalBasesCount * STRATA_PACKED_SHAREDLOCALBASIS_STRIDE_BYTES;
-		}
-		else
-		{
-			// Top normal texture is used to read the data
-			StrataMaterialRequestedSizeByte += UintByteSize;
+				if (It.OperatorType == STRATA_OPERATOR_BSDF)
+				{
+					// Recursively parse all node from BSDF to the root node and update the necessary properties.
+					WalkOperatorsToRoot(It);
+				}
+			}
 		}
 
-
-		// 2. The list of BSDFs
-		for (auto& It : StrataMaterialExpressionRegisteredOperators)
+		//
+		// Compute IsTop or IsBottom layer using a depth first tree visit while counting vertical right and left branches taken.
+		//
 		{
-			if (It.IsDiscarded())
-			{
-				continue; // ignore discarded operations in sub tree using parameter blending
-			}
+			int VOpTopBranchCountTaken = 0;
+			int VOpBottomBranchCountTaken = 0;
 
-			switch (It.OperatorType)
+			std::function<void(FStrataOperator&)> WalkOperators = [&](FStrataOperator& CurrentOperator) -> void
 			{
-			case STRATA_OPERATOR_BSDF:
-			{
-				// From the compiler side, we can only assume the top layer has gray scale luminance weight.
-				const bool bMayHaveColoredWeight = !It.bIsTop;
-
-				if (bStrataMaterialIsSimple)
+				switch (CurrentOperator.OperatorType)
+				{
+				case STRATA_OPERATOR_VERTICAL:
 				{
-					// Header
-					StrataMaterialRequestedSizeByte += UintByteSize;
-					// Disney material
-					StrataMaterialRequestedSizeByte += UintByteSize;
-					break; // Stop here
+					VOpTopBranchCountTaken++;
+					WalkOperators(StrataMaterialExpressionRegisteredOperators[CurrentOperator.LeftIndex]);
+					VOpTopBranchCountTaken--;
+					VOpBottomBranchCountTaken++;
+					WalkOperators(StrataMaterialExpressionRegisteredOperators[CurrentOperator.RightIndex]);
+					VOpBottomBranchCountTaken--;
+					break;
 				}
-				else if (bStrataMaterialIsSingle)
+				case STRATA_OPERATOR_HORIZONTAL:
+				case STRATA_OPERATOR_ADD:
 				{
-					// Header
-					StrataMaterialRequestedSizeByte += UintByteSize;
+					WalkOperators(StrataMaterialExpressionRegisteredOperators[CurrentOperator.LeftIndex]);
+					WalkOperators(StrataMaterialExpressionRegisteredOperators[CurrentOperator.RightIndex]);
+					break;
 				}
-				else if (bMayHaveColoredWeight)
+				case STRATA_OPERATOR_WEIGHT:
+				case STRATA_OPERATOR_THINFILM:
 				{
-					// BSDF state
-					StrataMaterialRequestedSizeByte += UintByteSize;
-					// Color weight
-					StrataMaterialRequestedSizeByte += UintByteSize;
-					// Light transmittance weight
-					StrataMaterialRequestedSizeByte += UintByteSize;
+					WalkOperators(StrataMaterialExpressionRegisteredOperators[CurrentOperator.LeftIndex]);
+					break;
 				}
-				else
+				case STRATA_OPERATOR_BSDF:
 				{
-					// BSDF state with gray scale weight
-					StrataMaterialRequestedSizeByte += UintByteSize;
+					const int32 VopCount = VOpTopBranchCountTaken + VOpBottomBranchCountTaken;
+					CurrentOperator.LayerDepth = VOpBottomBranchCountTaken;
+					CurrentOperator.bIsTop = VopCount == 0 || VOpTopBranchCountTaken == VopCount;
+					CurrentOperator.bIsBottom = VopCount == 0 || VOpBottomBranchCountTaken == VopCount;
+					break;
 				}
+				}
+			};
+
+			WalkOperators(*StrataMaterialRootOperator);
+		}
+
+		//
+		// Compute the size of the material
+		//
+		{
+			// Compute the shared local basis count only
+			// But we cannot use StrataEvaluateSharedLocalBases here, becaise the material has not been compiled yet so al lthe bases would just default to the same.
+			// STRATA_TODO: can we do that material generation in two passes? 
+			//		1- A first one to evaluate the normal/tangent code
+			//		2- Operators are processed and simplification computed based on memory budget
+			//		3- Material is finally compiled for with operator updated to fit in memory budget.
+			uint8 UsedSharedLocalBasesCount = StrataMaterialBSDFCount;
 
-				switch (It.BSDFType)
+			const uint32 UintByteSize = sizeof(uint32);
+			StrataMaterialRequestedSizeByte = 0;
+
+			// 1. Evaluate simple/single BSDF
+			bStrataMaterialIsSimple = StrataMaterialBSDFCount == 1;
+			bStrataMaterialIsSingle = StrataMaterialBSDFCount == 1;
+			for (auto& It : StrataMaterialExpressionRegisteredOperators)
+			{
+				if (It.IsDiscarded())
 				{
-				case STRATA_BSDF_TYPE_SLAB:
+					continue; // ignore discarded operations in sub tree using parameter blending
+				}
+
+				switch (It.OperatorType)
 				{
-					// Compute values closer to the reality for HasSSS and IsSimpleVolume, now that we know that we know the topology of the material.
-					const bool bIsSimpleVolume = !It.bIsBottom && It.bBSDFHasMFPPluggedIn;
-					const bool bHasSSS = It.bIsBottom && It.bBSDFHasSSS;
-				
-					StrataMaterialRequestedSizeByte += UintByteSize;
-					StrataMaterialRequestedSizeByte += UintByteSize;
+				case STRATA_OPERATOR_BSDF:
+				{
+					// From the compiler side, we can only assume the top layer has gray scale luminance weight.
+					const bool bMayHaveColoredWeight = !It.bIsTop;
 
-					if (It.bBSDFHasEdgeColor || It.bBSDFHasSecondRoughnessOrSimpleClearCoat)
+					switch (It.BSDFType)
 					{
-						StrataMaterialRequestedSizeByte += UintByteSize;
+					case STRATA_BSDF_TYPE_SLAB:
+					{
+						bStrataMaterialIsSimple = bStrataMaterialIsSimple && !bMayHaveColoredWeight && !It.bBSDFHasAnisotropy && !It.bBSDFHasEdgeColor && !It.bBSDFHasFuzz && !It.bBSDFHasSecondRoughnessOrSimpleClearCoat && !It.bBSDFHasMFPPluggedIn && !It.bBSDFHasSSS;
+						bStrataMaterialIsSingle = bStrataMaterialIsSingle && !bMayHaveColoredWeight && !It.bBSDFHasAnisotropy;
+						break;
 					}
-					if (bHasSSS || bIsSimpleVolume)
+					case STRATA_BSDF_TYPE_HAIR:
 					{
-						StrataMaterialRequestedSizeByte += UintByteSize;
+						bStrataMaterialIsSimple = false;
+						bStrataMaterialIsSingle = false;
+						break;
 					}
-					if (It.bBSDFHasFuzz)
+					case STRATA_BSDF_TYPE_EYE:
 					{
-						StrataMaterialRequestedSizeByte += UintByteSize;
+						bStrataMaterialIsSimple = false;
+						bStrataMaterialIsSingle = false;
+						break;
+					}
+					case STRATA_BSDF_TYPE_SINGLELAYERWATER:
+					{
+						bStrataMaterialIsSimple = false;
+						bStrataMaterialIsSingle = false;
+						break;
+					}
 					}
 					break;
 				}
-				case STRATA_BSDF_TYPE_HAIR:
-				{
-					StrataMaterialRequestedSizeByte += UintByteSize;
-					StrataMaterialRequestedSizeByte += UintByteSize;
-					break;
-				}
-				case STRATA_BSDF_TYPE_EYE:
-				{
-					StrataMaterialRequestedSizeByte += UintByteSize;
-					StrataMaterialRequestedSizeByte += UintByteSize;
-					break;
 				}
-				case STRATA_BSDF_TYPE_SINGLELAYERWATER:
+			}
+			bStrataMaterialIsSingle = bStrataMaterialIsSingle && !bStrataMaterialIsSimple;
+
+			// 2. Header
+
+			if (!bStrataMaterialIsSimple && !bStrataMaterialIsSingle) // header written later, 
+			{
+				// Packed Header
+				StrataMaterialRequestedSizeByte += UintByteSize;
+
+				// Shared local bases between BSDFs
+				StrataMaterialRequestedSizeByte += UsedSharedLocalBasesCount * STRATA_PACKED_SHAREDLOCALBASIS_STRIDE_BYTES;
+			}
+			else
+			{
+				// Top normal texture is used to read the data
+				StrataMaterialRequestedSizeByte += UintByteSize;
+			}
+
+
+			// 2. The list of BSDFs
+			for (auto& It : StrataMaterialExpressionRegisteredOperators)
+			{
+				if (It.IsDiscarded())
 				{
-					StrataMaterialRequestedSizeByte += UintByteSize;
-					StrataMaterialRequestedSizeByte += UintByteSize;
-					StrataMaterialRequestedSizeByte += UintByteSize;
-					break;
+					continue; // ignore discarded operations in sub tree using parameter blending
 				}
-				case STRATA_BSDF_TYPE_UNLIT:
+
+				switch (It.OperatorType)
 				{
-					// Never stored, it goes directly into the scene as emitted luminance.
-					break;
-				}
-				default:
+				case STRATA_OPERATOR_BSDF:
 				{
-					Errorf(TEXT("Unkownd BSDF type encountered in %s (asset: %s).\r\n"), *Material->GetDebugName(), *Material->GetAssetPath().ToString());
+					// From the compiler side, we can only assume the top layer has gray scale luminance weight.
+					const bool bMayHaveColoredWeight = !It.bIsTop;
+
+					if (bStrataMaterialIsSimple)
+					{
+						// Header
+						StrataMaterialRequestedSizeByte += UintByteSize;
+						// Disney material
+						StrataMaterialRequestedSizeByte += UintByteSize;
+						break; // Stop here
+					}
+					else if (bStrataMaterialIsSingle)
+					{
+						// Header
+						StrataMaterialRequestedSizeByte += UintByteSize;
+					}
+					else if (bMayHaveColoredWeight)
+					{
+						// BSDF state
+						StrataMaterialRequestedSizeByte += UintByteSize;
+						// Color weight
+						StrataMaterialRequestedSizeByte += UintByteSize;
+						// Light transmittance weight
+						StrataMaterialRequestedSizeByte += UintByteSize;
+					}
+					else
+					{
+						// BSDF state with gray scale weight
+						StrataMaterialRequestedSizeByte += UintByteSize;
+					}
+
+					switch (It.BSDFType)
+					{
+					case STRATA_BSDF_TYPE_SLAB:
+					{
+						// Compute values closer to the reality for HasSSS and IsSimpleVolume, now that we know that we know the topology of the material.
+						const bool bIsSimpleVolume = !It.bIsBottom && It.bBSDFHasMFPPluggedIn;
+						const bool bHasSSS = It.bIsBottom && It.bBSDFHasSSS;
+
+						StrataMaterialRequestedSizeByte += UintByteSize;
+						StrataMaterialRequestedSizeByte += UintByteSize;
+
+						if (It.bBSDFHasEdgeColor || It.bBSDFHasSecondRoughnessOrSimpleClearCoat)
+						{
+							StrataMaterialRequestedSizeByte += UintByteSize;
+						}
+						if (bHasSSS || bIsSimpleVolume)
+						{
+							StrataMaterialRequestedSizeByte += UintByteSize;
+						}
+						if (It.bBSDFHasFuzz)
+						{
+							StrataMaterialRequestedSizeByte += UintByteSize;
+						}
+						break;
+					}
+					case STRATA_BSDF_TYPE_HAIR:
+					{
+						StrataMaterialRequestedSizeByte += UintByteSize;
+						StrataMaterialRequestedSizeByte += UintByteSize;
+						break;
+					}
+					case STRATA_BSDF_TYPE_EYE:
+					{
+						StrataMaterialRequestedSizeByte += UintByteSize;
+						StrataMaterialRequestedSizeByte += UintByteSize;
+						break;
+					}
+					case STRATA_BSDF_TYPE_SINGLELAYERWATER:
+					{
+						StrataMaterialRequestedSizeByte += UintByteSize;
+						StrataMaterialRequestedSizeByte += UintByteSize;
+						StrataMaterialRequestedSizeByte += UintByteSize;
+						break;
+					}
+					case STRATA_BSDF_TYPE_UNLIT:
+					{
+						// Never stored, it goes directly into the scene as emitted luminance.
+						break;
+					}
+					default:
+					{
+						Errorf(TEXT("Unkownd BSDF type encountered in %s (asset: %s).\r\n"), *Material->GetDebugName(), *Material->GetAssetPath().ToString());
+						break;
+					}
+					}
 					break;
-				}
-				}
-				break;
+				} // case STRATA_OPERATOR_BSDF
+				} // switch (It.OperatorType)
 			}
+
+			StrataSimplificationStatus.bMaterialFitsInMemoryBudget = StrataMaterialRequestedSizeByte <= StrataBytePerPixel;
+			if (!StrataSimplificationStatus.bMaterialFitsInMemoryBudget && StrataSimplificationStatus.bRunFullSimplification)
+			{
+				// If we have already run the full simplification but the material still does not fit in memory, we must fail the material compilation.
+				Errorf(TEXT("Material %s could not be simplified to fit in strata per pixel (asset: %s).\r\n"), *Material->GetDebugName(), *Material->GetAssetPath().ToString());
+				return false;
+			}
+			if (!StrataSimplificationStatus.bRunFullSimplification)
+			{
+				// Record the original requested byte size before simplification, only for the first pass.
+				StrataSimplificationStatus.OriginalRequestedByteSize = StrataMaterialRequestedSizeByte;
 			}
+			StrataSimplificationStatus.bRunFullSimplification = !StrataSimplificationStatus.bMaterialFitsInMemoryBudget;
 		}
 	}
+	while (!StrataSimplificationStatus.bMaterialFitsInMemoryBudget);
 
 	return true; // Success
 }
diff --git a/Engine/Source/Runtime/Engine/Private/Materials/HLSLMaterialTranslator.h b/Engine/Source/Runtime/Engine/Private/Materials/HLSLMaterialTranslator.h
--- a/Engine/Source/Runtime/Engine/Private/Materials/HLSLMaterialTranslator.h
+++ b/Engine/Source/Runtime/Engine/Private/Materials/HLSLMaterialTranslator.h
@@ -435,6 +435,15 @@ protected:
 	bool bStrataUsesConversionFromLegacy;
 	bool bStrataOutputsOpaqueRoughRefractions;
 
+	struct FStrataSimplificationStatus
+	{
+		bool bMaterialFitsInMemoryBudget = false;	// Track whether or not the material fits.
+
+		uint32 OriginalRequestedByteSize = 0;
+		bool bRunFullSimplification = false;	// Simple implementation for now: if the material does not fit, we simply everything.
+	};
+	FStrataSimplificationStatus StrataSimplificationStatus;
+
 	/** Tracks the total number of vt samples in the shader. */
 	uint32 NumVtSamples;
 
diff --git a/Engine/Source/Runtime/Engine/Private/Materials/MaterialExpressions.cpp b/Engine/Source/Runtime/Engine/Private/Materials/MaterialExpressions.cpp
--- a/Engine/Source/Runtime/Engine/Private/Materials/MaterialExpressions.cpp
+++ b/Engine/Source/Runtime/Engine/Private/Materials/MaterialExpressions.cpp
@@ -2187,6 +2187,22 @@ static void AssignOperatorIndexIfNotNull(int32& NextOperatorPin, FStrataOperator
 	NextOperatorPin = Operator ? Operator->Index : INDEX_NONE;
 }
 
+static void CombineFlagForParameterBlending(FStrataOperator& DstOp, FStrataOperator* OpA, FStrataOperator* OpB = nullptr)
+{
+	if (OpA && OpB)
+	{
+		DstOp.CombineFlagsForParameterBlending(*OpA, *OpB);
+	}
+	else if (OpA)
+	{
+		DstOp.CopyFlagsForParameterBlending(*OpA);
+	}
+	else if (OpB)
+	{
+		DstOp.CopyFlagsForParameterBlending(*OpB);
+	}
+}
+
 #endif // WITH_EDITOR
 
 UMaterialExpressionTextureBase::UMaterialExpressionTextureBase(const FObjectInitializer& ObjectInitializer)
@@ -7537,14 +7553,19 @@ FStrataOperator* UMaterialExpressionBlendMaterialAttributes::StrataGenerateMater
 
 	UMaterialExpression* ChildAExpression = A.GetTracedInput().Expression;
 	UMaterialExpression* ChildBExpression = B.GetTracedInput().Expression;
+	FStrataOperator* OpA = nullptr;
+	FStrataOperator* OpB = nullptr;
 	if (ChildAExpression)
 	{
-		AssignOperatorIndexIfNotNull(StrataOperator.LeftIndex, ChildAExpression->StrataGenerateMaterialTopologyTree(Compiler, this, A.OutputIndex));
+		OpA = ChildAExpression->StrataGenerateMaterialTopologyTree(Compiler, this, A.OutputIndex);
+		AssignOperatorIndexIfNotNull(StrataOperator.LeftIndex, OpA);
 	}
 	if (ChildBExpression)
 	{
-		AssignOperatorIndexIfNotNull(StrataOperator.RightIndex, ChildBExpression->StrataGenerateMaterialTopologyTree(Compiler, this, B.OutputIndex));
+		OpB = ChildBExpression->StrataGenerateMaterialTopologyTree(Compiler, this, B.OutputIndex);
+		AssignOperatorIndexIfNotNull(StrataOperator.RightIndex, OpB);
 	}
+	CombineFlagForParameterBlending(StrataOperator, OpA, OpB);
 
 	return &StrataOperator;
 }
@@ -22827,7 +22848,7 @@ FStrataOperator* UMaterialExpressionStrataSimpleClearCoatBSDF::StrataGenerateMat
 	StrataOperator.BSDFType = STRATA_BSDF_TYPE_SLAB;
 	StrataOperator.bBSDFHasEdgeColor = false;
 	StrataOperator.bBSDFHasFuzz = false;
-	StrataOperator.bBSDFHasSecondRoughnessOrSimpleClearCoat = true;	// This node explicitely requires simple clear coat
+	StrataOperator.bBSDFHasSecondRoughnessOrSimpleClearCoat = true;	// This node explicitly requires simple clear coat
 	StrataOperator.bBSDFHasSSS = false;
 	StrataOperator.bBSDFHasMFPPluggedIn = false;
 	StrataOperator.bBSDFHasAnisotropy = false;
@@ -23101,7 +23122,6 @@ int32 UMaterialExpressionStrataConvertToDecal::Compile(class FMaterialCompiler*
 			if (StrataOperator.bRootOfParameterBlendingSubTree)
 			{
 				StrataOperator.BSDFRegisteredSharedLocalBasis = Operator->BSDFRegisteredSharedLocalBasis;
-				StrataOperator.CopyFlagsForParameterBlending(*Operator);
 			}
 
 			OutputCodeChunk = Compiler->StrataWeightParameterBlending(
@@ -23161,10 +23181,13 @@ FStrataOperator* UMaterialExpressionStrataConvertToDecal::StrataGenerateMaterial
 	FStrataOperator& StrataOperator = Compiler->StrataCompilationRegisterOperator(STRATA_OPERATOR_WEIGHT, this, Parent, bUseParameterBlending);
 
 	UMaterialExpression* ChildDecalMaterialExpression = DecalMaterial.GetTracedInput().Expression;
+	FStrataOperator* OpA = nullptr;
 	if (ChildDecalMaterialExpression)
 	{
-		AssignOperatorIndexIfNotNull(StrataOperator.LeftIndex, ChildDecalMaterialExpression->StrataGenerateMaterialTopologyTree(Compiler, this, DecalMaterial.OutputIndex));
+		OpA = ChildDecalMaterialExpression->StrataGenerateMaterialTopologyTree(Compiler, this, DecalMaterial.OutputIndex);
+		AssignOperatorIndexIfNotNull(StrataOperator.LeftIndex, OpA);
 	}
+	CombineFlagForParameterBlending(StrataOperator, OpA);
 
 	return &StrataOperator;
 }
@@ -23652,7 +23675,6 @@ int32 UMaterialExpressionStrataHorizontalMixing::Compile(class FMaterialCompiler
 
 		// Propagate the parameter blended normal
 		StrataOperator.BSDFRegisteredSharedLocalBasis = NewRegisteredSharedLocalBasis;
-		StrataOperator.CombineFlagsForParameterBlending(*BackgroundBSDFOperator, *ForegroundBSDFOperator);
 	}
 	else
 	{
@@ -23712,14 +23734,19 @@ FStrataOperator* UMaterialExpressionStrataHorizontalMixing::StrataGenerateMateri
 
 	UMaterialExpression* ChildAExpression = Background.GetTracedInput().Expression;
 	UMaterialExpression* ChildBExpression = Foreground.GetTracedInput().Expression;
+	FStrataOperator* OpA = nullptr;
+	FStrataOperator* OpB = nullptr;
 	if (ChildAExpression)
 	{
-		AssignOperatorIndexIfNotNull(StrataOperator.LeftIndex, ChildAExpression->StrataGenerateMaterialTopologyTree(Compiler, this, Background.OutputIndex));
+		OpA = ChildAExpression->StrataGenerateMaterialTopologyTree(Compiler, this, Background.OutputIndex);
+		AssignOperatorIndexIfNotNull(StrataOperator.LeftIndex, OpA);
 	}
 	if (ChildBExpression)
 	{
-		AssignOperatorIndexIfNotNull(StrataOperator.RightIndex, ChildBExpression->StrataGenerateMaterialTopologyTree(Compiler, this, Foreground.OutputIndex));
+		OpB = ChildBExpression->StrataGenerateMaterialTopologyTree(Compiler, this, Foreground.OutputIndex);
+		AssignOperatorIndexIfNotNull(StrataOperator.RightIndex, OpB);
 	}
+	CombineFlagForParameterBlending(StrataOperator, OpA, OpB);
 
 	return &StrataOperator;
 }
@@ -23802,7 +23829,6 @@ int32 UMaterialExpressionStrataVerticalLayering::Compile(class FMaterialCompiler
 
 		// Propagate the parameter blended normal
 		StrataOperator.BSDFRegisteredSharedLocalBasis = NewRegisteredSharedLocalBasis;
-		StrataOperator.CombineFlagsForParameterBlending(*TopBSDFOperator, *BaseBSDFOperator);
 	}
 	else
 	{
@@ -23858,14 +23884,19 @@ FStrataOperator* UMaterialExpressionStrataVerticalLayering::StrataGenerateMateri
 
 	UMaterialExpression* ChildAExpression = Top.GetTracedInput().Expression;
 	UMaterialExpression* ChildBExpression = Base.GetTracedInput().Expression;
+	FStrataOperator* OpA = nullptr;
+	FStrataOperator* OpB = nullptr;
 	if (ChildAExpression)
 	{
-		AssignOperatorIndexIfNotNull(StrataOperator.LeftIndex, ChildAExpression->StrataGenerateMaterialTopologyTree(Compiler, this, Top.OutputIndex));
+		OpA = ChildAExpression->StrataGenerateMaterialTopologyTree(Compiler, this, Top.OutputIndex);
+		AssignOperatorIndexIfNotNull(StrataOperator.LeftIndex, OpA);
 	}
 	if (ChildBExpression)
 	{
-		AssignOperatorIndexIfNotNull(StrataOperator.RightIndex, ChildBExpression->StrataGenerateMaterialTopologyTree(Compiler, this, Base.OutputIndex));
+		OpB = ChildBExpression->StrataGenerateMaterialTopologyTree(Compiler, this, Base.OutputIndex);
+		AssignOperatorIndexIfNotNull(StrataOperator.RightIndex, OpB);
 	}
+	CombineFlagForParameterBlending(StrataOperator, OpA, OpB);
 
 	return &StrataOperator;
 }
@@ -23948,7 +23979,6 @@ int32 UMaterialExpressionStrataAdd::Compile(class FMaterialCompiler* Compiler, i
 
 		// Propagate the parameter blended normal
 		StrataOperator.BSDFRegisteredSharedLocalBasis = NewRegisteredSharedLocalBasis;
-		StrataOperator.CombineFlagsForParameterBlending(*ABSDFOperator, *BBSDFOperator);
 	}
 	else
 	{
@@ -24003,14 +24033,19 @@ FStrataOperator* UMaterialExpressionStrataAdd::StrataGenerateMaterialTopologyTre
 
 	UMaterialExpression* ChildAExpression = A.GetTracedInput().Expression;
 	UMaterialExpression* ChildBExpression = B.GetTracedInput().Expression;
+	FStrataOperator* OpA = nullptr;
+	FStrataOperator* OpB = nullptr;
 	if (ChildAExpression)
 	{
-		AssignOperatorIndexIfNotNull(StrataOperator.LeftIndex, ChildAExpression->StrataGenerateMaterialTopologyTree(Compiler, this, A.OutputIndex));
+		OpA = ChildAExpression->StrataGenerateMaterialTopologyTree(Compiler, this, A.OutputIndex);
+		AssignOperatorIndexIfNotNull(StrataOperator.LeftIndex, OpA);
 	}
 	if (ChildBExpression)
 	{
-		AssignOperatorIndexIfNotNull(StrataOperator.RightIndex, ChildBExpression->StrataGenerateMaterialTopologyTree(Compiler, this, B.OutputIndex));
+		OpB = ChildBExpression->StrataGenerateMaterialTopologyTree(Compiler, this, B.OutputIndex);
+		AssignOperatorIndexIfNotNull(StrataOperator.RightIndex, OpB);
 	}
+	CombineFlagForParameterBlending(StrataOperator, OpA, OpB);
 
 	return &StrataOperator;
 }
@@ -24063,7 +24098,6 @@ int32 UMaterialExpressionStrataWeight::Compile(class FMaterialCompiler* Compiler
 		}
 
 		StrataOperator.BSDFRegisteredSharedLocalBasis = Operator->BSDFRegisteredSharedLocalBasis;
-		StrataOperator.CopyFlagsForParameterBlending(*Operator);
 
 		OutputCodeChunk = Compiler->StrataWeightParameterBlending(
 			ACodeChunk, WeightCodeChunk, 
@@ -24114,10 +24148,13 @@ FStrataOperator* UMaterialExpressionStrataWeight::StrataGenerateMaterialTopology
 	FStrataOperator& StrataOperator = Compiler->StrataCompilationRegisterOperator(STRATA_OPERATOR_WEIGHT, this, Parent);
 
 	UMaterialExpression* ChildAExpression = A.GetTracedInput().Expression;
+	FStrataOperator* OpA = nullptr;
 	if (ChildAExpression)
 	{
-		AssignOperatorIndexIfNotNull(StrataOperator.LeftIndex, ChildAExpression->StrataGenerateMaterialTopologyTree(Compiler, this, A.OutputIndex));
+		OpA = ChildAExpression->StrataGenerateMaterialTopologyTree(Compiler, this, A.OutputIndex);
+		AssignOperatorIndexIfNotNull(StrataOperator.LeftIndex, OpA);
 	}
+	CombineFlagForParameterBlending(StrataOperator, OpA);
 
 	return &StrataOperator;
 }
@@ -24162,7 +24199,6 @@ int32 UMaterialExpressionStrataThinFilm::Compile(class FMaterialCompiler* Compil
 			return Compiler->Errorf(TEXT("Missing input on ThinFilm node."));
 		}
 		StrataOperator.BSDFRegisteredSharedLocalBasis = Operator->BSDFRegisteredSharedLocalBasis;
-		StrataOperator.CopyFlagsForParameterBlending(*Operator);
 
 		int32 NormalCodeChunk = Operator->BSDFRegisteredSharedLocalBasis.NormalCodeChunk;
 		OutputCodeChunk = Compiler->StrataThinFilmParameterBlending(ACodeChunk, ThicknessCodeChunk, IORCodeChunk, NormalCodeChunk, StrataOperator.bRootOfParameterBlendingSubTree ? &StrataOperator : nullptr);
@@ -24211,10 +24247,13 @@ FStrataOperator* UMaterialExpressionStrataThinFilm::StrataGenerateMaterialTopolo
 	FStrataOperator& StrataOperator = Compiler->StrataCompilationRegisterOperator(STRATA_OPERATOR_THINFILM, this, Parent);
 
 	UMaterialExpression* ChildAExpression = A.GetTracedInput().Expression;
+	FStrataOperator* OpA = nullptr;
 	if (ChildAExpression)
 	{
-		AssignOperatorIndexIfNotNull(StrataOperator.LeftIndex, ChildAExpression->StrataGenerateMaterialTopologyTree(Compiler, this, A.OutputIndex));
+		OpA = ChildAExpression->StrataGenerateMaterialTopologyTree(Compiler, this, A.OutputIndex);
+		AssignOperatorIndexIfNotNull(StrataOperator.LeftIndex, OpA);
 	}
+	CombineFlagForParameterBlending(StrataOperator, OpA);
 
 	return &StrataOperator;
 }
