Fix to SkyLightEnvMapSample in the HLSL material translator causing a crash if any of its two inputs is in the error state.

diff --git a/Engine/Source/Runtime/Engine/Private/Materials/HLSLMaterialTranslator.cpp b/Engine/Source/Runtime/Engine/Private/Materials/HLSLMaterialTranslator.cpp
--- a/Engine/Source/Runtime/Engine/Private/Materials/HLSLMaterialTranslator.cpp
+++ b/Engine/Source/Runtime/Engine/Private/Materials/HLSLMaterialTranslator.cpp
@@ -9197,10 +9197,16 @@ int32 FHLSLMaterialTranslator::SkyAtmosphereDistantLightScatteredLuminance()
 
 int32 FHLSLMaterialTranslator::SkyLightEnvMapSample(int32 DirectionCodeChunk, int32 RoughnessCodeChunk)
 {
+	if (DirectionCodeChunk == INDEX_NONE || RoughnessCodeChunk == INDEX_NONE)
+	{
+		return INDEX_NONE;
+	}
+
 	if (Material->GetMaterialDomain() != MD_Surface)
 	{
 		return Errorf(TEXT("The SkyLightEnvMapSample node can only be used when material Domain is set to Surface."));
 	}
+
 	return AddCodeChunk(MCT_Float3, TEXT("MaterialExpressionSkyLightEnvMapSample(%s, %s)"), *GetParameterCode(DirectionCodeChunk), *GetParameterCode(RoughnessCodeChunk));
 }
 
