Relax the material quality check when compiling shaders on-demand for COTF.

- It is a perfectly valid scenario to request a low quality level of a material, but if the material doesn't support Low it will default to EMaterialQuality::Num.
- EMaterialQuality::Num just means the quality level doesn't matter to the compiled result.

diff --git a/Engine/Source/Runtime/Engine/Private/Materials/Material.cpp b/Engine/Source/Runtime/Engine/Private/Materials/Material.cpp
--- a/Engine/Source/Runtime/Engine/Private/Materials/Material.cpp
+++ b/Engine/Source/Runtime/Engine/Private/Materials/Material.cpp
@@ -5145,7 +5145,6 @@ void UMaterial::CompileODSCMaterialsForRemoteRecompile(TArray<FODSCRequestPayloa
 		FMaterialResource* MaterialResource = MaterialInterface->GetMaterialResource(Shaders.FeatureLevel, Shaders.QualityLevel);
 		check(MaterialResource);
 		check(MaterialResource->GetFeatureLevel() == Shaders.FeatureLevel);
-		check(MaterialResource->GetQualityLevel() == Shaders.QualityLevel);
 		ResourceArray.Add(MaterialResource);
 	}
 
