Fix package map name mismatch warnings on server in PIE

diff --git a/Engine/Source/Runtime/Engine/Classes/Engine/PackageMapClient.h b/Engine/Source/Runtime/Engine/Classes/Engine/PackageMapClient.h
--- a/Engine/Source/Runtime/Engine/Classes/Engine/PackageMapClient.h
+++ b/Engine/Source/Runtime/Engine/Classes/Engine/PackageMapClient.h
@@ -490,7 +490,7 @@ protected:
 	void ExportNetGUIDHeader();
 
 	void			InternalWriteObject( FArchive& Ar, FNetworkGUID NetGUID, UObject* Object, FString ObjectPathName, UObject* ObjectOuter );	
-	FNetworkGUID	InternalLoadObject( FArchive & Ar, UObject *& Object, int InternalLoadObjectRecursionCount );
+	FNetworkGUID	InternalLoadObject( FArchive & Ar, UObject *& Object, const int32 InternalLoadObjectRecursionCount );
 
 	virtual UObject* ResolvePathAndAssignNetGUID( const FNetworkGUID& NetGUID, const FString& PathName ) override;
 
@@ -500,7 +500,7 @@ protected:
 
 	class UNetConnection* Connection;
 
-	bool ObjectLevelHasFinishedLoading(UObject* Obj);
+	bool ObjectLevelHasFinishedLoading(UObject* Obj) const;
 
 	TArray<TArray<uint8>>				ExportGUIDArchives;
 	TSet< FNetworkGUID >				CurrentExportNetGUIDs;				// Current list of NetGUIDs being written to the Export Bunch.
diff --git a/Engine/Source/Runtime/Engine/Private/PackageMapClient.cpp b/Engine/Source/Runtime/Engine/Private/PackageMapClient.cpp
--- a/Engine/Source/Runtime/Engine/Private/PackageMapClient.cpp
+++ b/Engine/Source/Runtime/Engine/Private/PackageMapClient.cpp
@@ -71,6 +71,23 @@ namespace UE
 
 		int32 MaxSerializedNetExportsPerGroup = 128 * 1024;	// Gameplay tags will be exported into a single large group for replays
 		static FAutoConsoleVariableRef CVarMaxSerializedNetExportsPerGroup(TEXT("net.MaxSerializedNetExportsPerGroup"), MaxSerializedNetExportsPerGroup, TEXT("Maximum number of network exports in each group we would expect to receive in a bunch"));
+
+		static bool ObjectLevelHasFinishedLoading(UObject* Object, UNetDriver* Driver)
+		{
+			if (Object != nullptr && Driver != nullptr && Driver->GetWorld() != nullptr)
+			{
+				// get the level for the object
+				AActor* Actor = Cast<AActor>(Object);
+				ULevel* Level = Actor ? Actor->GetLevel() : Object->GetTypedOuter<ULevel>();
+
+				if (Level != nullptr && Level != Driver->GetWorld()->PersistentLevel)
+				{
+					return Level->bIsVisible;
+				}
+			}
+
+			return true;
+		}
 	};
 };
 
@@ -763,24 +780,24 @@ bool FNetGUIDCache::CanClientLoadObject( const UObject* Object, const FNetworkGU
 }
 
 /** Writes an object NetGUID given the NetGUID and either the object itself, or FString full name of the object. Appends full name/path if necessary */
-void UPackageMapClient::InternalWriteObject( FArchive & Ar, FNetworkGUID NetGUID, UObject* Object, FString ObjectPathName, UObject* ObjectOuter )
+void UPackageMapClient::InternalWriteObject(FArchive & Ar, FNetworkGUID NetGUID, UObject* Object, FString ObjectPathName, UObject* ObjectOuter)
 {
-	check( Ar.IsSaving() );
+	check(Ar.IsSaving());
 
-	const bool bNoLoad = !GuidCache->CanClientLoadObject( Object, NetGUID );
+	const bool bNoLoad = !GuidCache->CanClientLoadObject(Object, NetGUID);
 
-	if ( GuidCache->ShouldAsyncLoad() && IsNetGUIDAuthority() && !GuidCache->IsExportingNetGUIDBunch && !bNoLoad )
+	if (GuidCache->ShouldAsyncLoad() && IsNetGUIDAuthority() && !GuidCache->IsExportingNetGUIDBunch && !bNoLoad)
 	{
 		// These are guids that must exist on the client in a package
 		// The client needs to know about these so it can determine if it has finished loading them
 		// and pause the network stream for that channel if it hasn't
-		MustBeMappedGuidsInLastBunch.AddUnique( NetGUID );
+		MustBeMappedGuidsInLastBunch.AddUnique(NetGUID);
 	}
 
 	Ar << NetGUID;
-	NET_CHECKSUM( Ar );
+	NET_CHECKSUM(Ar);
 
-	if ( !NetGUID.IsValid() )
+	if (!NetGUID.IsValid())
 	{
 		// We're done writing
 		return;
@@ -790,24 +807,24 @@ void UPackageMapClient::InternalWriteObject( FArchive & Ar, FNetworkGUID NetGUID
 	//   note: Default NetGUID is implied to always send path
 	FExportFlags ExportFlags;
 
-	ExportFlags.bHasNetworkChecksum = ( GuidCache->NetworkChecksumMode != FNetGUIDCache::ENetworkChecksumMode::None ) ? 1 : 0;
+	ExportFlags.bHasNetworkChecksum = (GuidCache->NetworkChecksumMode != FNetGUIDCache::ENetworkChecksumMode::None) ? 1 : 0;
 
-	if ( NetGUID.IsDefault() )
+	if (NetGUID.IsDefault())
 	{
 		// Only the client sends default guids
-		check( !IsNetGUIDAuthority() );
+		check(!IsNetGUIDAuthority());
 		ExportFlags.bHasPath = 1;
 
 		Ar << ExportFlags.Value;
 	}
-	else if ( GuidCache->IsExportingNetGUIDBunch )
+	else if (GuidCache->IsExportingNetGUIDBunch)
 	{
 		// Only the server should be exporting guids
-		check( IsNetGUIDAuthority() );
+		check(IsNetGUIDAuthority());
 
-		if ( Object != NULL )
+		if (Object != nullptr)
 		{
-			ExportFlags.bHasPath = ShouldSendFullPath( Object, NetGUID ) ? 1 : 0;
+			ExportFlags.bHasPath = ShouldSendFullPath(Object, NetGUID) ? 1 : 0;
 		}
 		else
 		{
@@ -819,38 +836,38 @@ void UPackageMapClient::InternalWriteObject( FArchive & Ar, FNetworkGUID NetGUID
 		Ar << ExportFlags.Value;
 	}
 
-	if ( ExportFlags.bHasPath )
+	if (ExportFlags.bHasPath)
 	{
-		if ( Object != NULL )
+		if (Object != nullptr)
 		{
-			// If the object isn't NULL, expect an empty path name, then fill it out with the actual info
-			check( ObjectOuter == NULL );
-			check( ObjectPathName.IsEmpty() );
+			// If the object isn't nullptr, expect an empty path name, then fill it out with the actual info
+			check(ObjectOuter == nullptr);
+			check(ObjectPathName.IsEmpty());
 			ObjectPathName = Object->GetName();
 			ObjectOuter = Object->GetOuter();
 		}
 		else
 		{
 			// If we don't have an object, expect an already filled out path name
-			checkf( ObjectOuter != NULL, TEXT("ObjectOuter is null. NetGuid: %s. Object: %s. ObjectPathName: %s"), *NetGUID.ToString(), *GetPathNameSafe(Object), *ObjectPathName );
-			checkf( !ObjectPathName.IsEmpty(), TEXT("ObjectPathName is empty. NetGuid: %s. Object: %s"), *NetGUID.ToString(), *GetPathNameSafe(Object));
+			checkf(ObjectOuter != nullptr, TEXT("ObjectOuter is null. NetGuid: %s. Object: %s. ObjectPathName: %s"), *NetGUID.ToString(), *GetPathNameSafe(Object), *ObjectPathName);
+			checkf(!ObjectPathName.IsEmpty(), TEXT("ObjectPathName is empty. NetGuid: %s. Object: %s"), *NetGUID.ToString(), *GetPathNameSafe(Object));
 		}
 
-		const bool bIsPackage = ( NetGUID.IsStatic() && Object != NULL && Object->GetOuter() == NULL );
+		const bool bIsPackage = (NetGUID.IsStatic() && Object != nullptr && Object->GetOuter() == nullptr);
 
-		check( bIsPackage == ( Cast< UPackage >( Object ) != NULL ) );		// Make sure it really is a package
+		check(bIsPackage == (Cast<UPackage>(Object) != nullptr));		// Make sure it really is a package
 
 		// Serialize reference to outer. This is basically a form of compression.
-		FNetworkGUID OuterNetGUID = GuidCache->GetOrAssignNetGUID( ObjectOuter );
+		FNetworkGUID OuterNetGUID = GuidCache->GetOrAssignNetGUID(ObjectOuter);
 
-		InternalWriteObject( Ar, OuterNetGUID, ObjectOuter, TEXT( "" ), NULL );
+		InternalWriteObject(Ar, OuterNetGUID, ObjectOuter, TEXT( "" ), nullptr);
 
 		// Look for renamed startup actors
 		if (Connection->Driver)
 		{
-			FName SearchPath = FName(*ObjectPathName);
-			FName RenamedPath = Connection->Driver->RenamedStartupActors.FindRef(SearchPath);
-			if (RenamedPath != NAME_None)
+			const FName SearchPath = FName(*ObjectPathName);
+			const FName RenamedPath = Connection->Driver->RenamedStartupActors.FindRef(SearchPath);
+			if (!RenamedPath.IsNone())
 			{
 				ObjectPathName = RenamedPath.ToString();
 			}
@@ -865,27 +882,29 @@ void UPackageMapClient::InternalWriteObject( FArchive & Ar, FNetworkGUID NetGUID
 
 		if ( ExportFlags.bHasNetworkChecksum )
 		{
-			NetworkChecksum = GuidCache->GetNetworkChecksum( Object );
+			NetworkChecksum = GuidCache->GetNetworkChecksum(Object);
 			Ar << NetworkChecksum;
 		}
 
-		FNetGuidCacheObject* CacheObject = GuidCache->ObjectLookup.Find( NetGUID );
-
-		if ( CacheObject != NULL )
+		if (FNetGuidCacheObject* CacheObject = GuidCache->ObjectLookup.Find(NetGUID))
 		{
-			CacheObject->PathName			= FName( *ObjectPathName );
+			if (CacheObject->PathName.IsNone())
+			{
+				CacheObject->PathName = FName(*ObjectPathName);
+			}
+
 			CacheObject->OuterGUID			= OuterNetGUID;
 			CacheObject->bNoLoad			= ExportFlags.bNoLoad;
 			CacheObject->bIgnoreWhenMissing = ExportFlags.bNoLoad;
 			CacheObject->NetworkChecksum	= NetworkChecksum;
 		}
 
-		if ( GuidCache->IsExportingNetGUIDBunch )
+		if (GuidCache->IsExportingNetGUIDBunch)
 		{
-			CurrentExportNetGUIDs.Add( NetGUID );
+			CurrentExportNetGUIDs.Add(NetGUID);
 
-			int32& Count = NetGUIDExportCountMap.FindOrAdd( NetGUID );
-			Count++;
+			int32& ExportCount = NetGUIDExportCountMap.FindOrAdd(NetGUID);
+			ExportCount++;
 		}
 	}
 }
@@ -945,7 +964,7 @@ static void SanityCheckExport(
 }
 
 /** Loads a UObject from an FArchive stream. Reads object path if there, and tries to load object if its not already loaded */
-FNetworkGUID UPackageMapClient::InternalLoadObject( FArchive & Ar, UObject *& Object, const int InternalLoadObjectRecursionCount )
+FNetworkGUID UPackageMapClient::InternalLoadObject( FArchive & Ar, UObject *& Object, const int32 InternalLoadObjectRecursionCount )
 {
 	if ( InternalLoadObjectRecursionCount > INTERNAL_LOAD_OBJECT_RECURSION_LIMIT ) 
 	{
@@ -2269,21 +2288,9 @@ void UPackageMapClient::LogDebugInfo( FOutputDevice & Ar )
 /**
  *	Returns true if Object's outer level has completely finished loading.
  */
-bool UPackageMapClient::ObjectLevelHasFinishedLoading(UObject* Object)
+bool UPackageMapClient::ObjectLevelHasFinishedLoading(UObject* Object) const
 {
-	if (Object != NULL && Connection!= NULL && Connection->Driver != NULL && Connection->Driver->GetWorld() != NULL)
-	{
-		// get the level for the object
-		AActor* Actor = Cast<AActor>(Object);
-		ULevel* Level = Actor ? Actor->GetLevel() : Object->GetTypedOuter<ULevel>();
-		
-		if (Level != NULL && Level != Connection->Driver->GetWorld()->PersistentLevel)
-		{
-			return Level->bIsVisible;
-		}
-	}
-
-	return true;
+	return UE::Net::ObjectLevelHasFinishedLoading(Object, Connection != nullptr ? Connection->Driver : nullptr);
 }
 
 /**
@@ -3202,23 +3209,6 @@ PRAGMA_ENABLE_DEPRECATION_WARNINGS
 	}
 }
 
-static bool ObjectLevelHasFinishedLoading( UObject* Object, UNetDriver* Driver )
-{
-	if ( Object != NULL && Driver != NULL && Driver->GetWorld() != NULL )
-	{
-		// get the level for the object
-		AActor* Actor = Cast<AActor>(Object);
-		ULevel* Level = Actor ? Actor->GetLevel() : Object->GetTypedOuter<ULevel>();
-
-		if ( Level != NULL && Level != Driver->GetWorld()->PersistentLevel )
-		{
-			return Level->bIsVisible;
-		}
-	}
-
-	return true;
-}
-
 UObject* FNetGUIDCache::GetObjectFromNetGUID( const FNetworkGUID& NetGUID, const bool bIgnoreMustBeMapped )
 {
 	LLM_SCOPE_BYTAG(GuidCache);
@@ -3458,13 +3448,12 @@ UObject* FNetGUIDCache::GetObjectFromNetGUID( const FNetworkGUID& NetGUID, const
 		}
 	}
 
-	if ( Object && !ObjectLevelHasFinishedLoading( Object, Driver ) )
+	if ( Object && !UE::Net::ObjectLevelHasFinishedLoading( Object, Driver ) )
 	{
-		UE_LOG( LogNetPackageMap, Verbose, TEXT( "GetObjectFromNetGUID: Forcing object to NULL since level is not loaded yet. Object: %s" ), *Object->GetFullName() );
+		UE_LOG(LogNetPackageMap, Verbose, TEXT("GetObjectFromNetGUID: Forcing object to NULL since level is not loaded yet. Object: %s"), *GetFullNameSafe(Object));
 		return NULL;
 	}
 
-
 	// Assign the resolved object to this guid
 	CacheObjectPtr->Object = Object;		
 
