World Partition: add the name of the loaded world in initialization logs.

#rnx

diff --git a/Engine/Source/Runtime/Engine/Private/WorldPartition/WorldPartition.cpp b/Engine/Source/Runtime/Engine/Private/WorldPartition/WorldPartition.cpp
--- a/Engine/Source/Runtime/Engine/Private/WorldPartition/WorldPartition.cpp
+++ b/Engine/Source/Runtime/Engine/Private/WorldPartition/WorldPartition.cpp
@@ -382,7 +382,7 @@ void UWorldPartition::Initialize(UWorld* InWorld, const FTransform& InTransform)
 	const bool bIsDedicatedServer = IsRunningDedicatedServer();
 	const bool bPIEWorldTravel = (World->WorldType == EWorldType::PIE) && !StreamingPolicy;
 
-	UE_LOG(LogWorldPartition, Log, TEXT("UWorldPartition::Initialize(IsEditor=%d, bPIEWorldTravel=%d IsGame=%d, IsCooking=%d)"), bIsEditor ? 1 : 0, bPIEWorldTravel ? 1 : 0, bIsGame ? 1 : 0, bIsCooking ? 1 : 0);
+	UE_LOG(LogWorldPartition, Log, TEXT("UWorldPartition::Initialize(Asset=%s, IsEditor=%d, bPIEWorldTravel=%d IsGame=%d, IsCooking=%d)"), *OuterWorld->GetName(), bIsEditor ? 1 : 0, bPIEWorldTravel ? 1 : 0, bIsGame ? 1 : 0, bIsCooking ? 1 : 0);
 
 	if (bEnableStreaming)
 	{
