Disable recently added logs that can be spammy when there are a lot of HISM.

diff --git a/Engine/Source/Runtime/Engine/Private/HierarchicalInstancedStaticMesh.cpp b/Engine/Source/Runtime/Engine/Private/HierarchicalInstancedStaticMesh.cpp
--- a/Engine/Source/Runtime/Engine/Private/HierarchicalInstancedStaticMesh.cpp
+++ b/Engine/Source/Runtime/Engine/Private/HierarchicalInstancedStaticMesh.cpp
@@ -3282,9 +3282,12 @@ static void GatherInstanceTransformsInArea(const UHierarchicalInstancedStaticMes
 		const FTransform ToWorldTransform = FTransform(Component.TranslatedInstanceSpaceOrigin) * Component.GetComponentTransform();
 		const FBox WorldNodeBox = FBox(ChildNode.BoundMin, ChildNode.BoundMax).TransformBy(ToWorldTransform);
 
-		UE_VLOG_BOX(&Component, LogStaticMesh, Verbose, FBox(ChildNode.BoundMin, ChildNode.BoundMax), FColor::Red, TEXT("LocalNodeBox"));
-		UE_VLOG_BOX(&Component, LogStaticMesh, Verbose, WorldNodeBox, FColor::Green, TEXT("WorldNodeBox"));
-		UE_VLOG_BOX(&Component, LogStaticMesh, Verbose, AreaBox, FColor::Blue, TEXT("AreaBox"));
+#if 0
+		// Keeping this as it can be useful to debug but disabled as it can be spammy. 
+		UE_VLOG_BOX(&Component, LogStaticMesh, VeryVerbose, FBox(ChildNode.BoundMin, ChildNode.BoundMax), FColor::Red, TEXT("LocalNodeBox"));
+		UE_VLOG_BOX(&Component, LogStaticMesh, VeryVerbose, WorldNodeBox, FColor::Green, TEXT("WorldNodeBox"));
+		UE_VLOG_BOX(&Component, LogStaticMesh, VeryVerbose, AreaBox, FColor::Blue, TEXT("AreaBox"));
+#endif
 
 		if (AreaBox.Intersect(WorldNodeBox))
 		{
