World Partition: fix potential crash for unsaved actors when starting PIE under unknown conditions.

#rnx

diff --git a/Engine/Source/Runtime/Engine/Private/WorldPartition/WorldPartitionRuntimeSpatialHashCell.cpp b/Engine/Source/Runtime/Engine/Private/WorldPartition/WorldPartitionRuntimeSpatialHashCell.cpp
--- a/Engine/Source/Runtime/Engine/Private/WorldPartition/WorldPartitionRuntimeSpatialHashCell.cpp
+++ b/Engine/Source/Runtime/Engine/Private/WorldPartition/WorldPartitionRuntimeSpatialHashCell.cpp
@@ -33,10 +33,13 @@ void UWorldPartitionRuntimeSpatialHashCell::PostDuplicate(bool bDuplicateForPIE)
 		// Make sure actor container isn't under PIE World so those template actors will never be considered part of the world.
 		UnsavedActorsContainer->Rename(nullptr, GetPackage());
 
-		for (auto& ActorPair : UnsavedActorsContainer->Actors)
+		for (auto& [ActorName, Actor] : UnsavedActorsContainer->Actors)
 		{
-			// Don't use AActor::Rename here since the actor is not par of the world, it's only a duplication template.
-			ActorPair.Value->UObject::Rename(nullptr, UnsavedActorsContainer);
+			if (ensure(Actor))
+			{
+				// Don't use AActor::Rename here since the actor is not par of the world, it's only a duplication template.	
+				Actor->UObject::Rename(nullptr, UnsavedActorsContainer);
+			}
 		}
 	}
 }
