Propagate the necessary object flags when duplicating an instanced subobject to another owner (outer).

diff --git a/Engine/Source/Runtime/Engine/Private/ObjectEditorUtils.cpp b/Engine/Source/Runtime/Engine/Private/ObjectEditorUtils.cpp
--- a/Engine/Source/Runtime/Engine/Private/ObjectEditorUtils.cpp
+++ b/Engine/Source/Runtime/Engine/Private/ObjectEditorUtils.cpp
@@ -261,11 +261,21 @@ namespace FObjectEditorUtils
 
 						UObject* DuplicateValue = StaticDuplicateObject(Value, InDestinationObject, Value->GetFName(), RF_AllFlags, nullptr, EDuplicateMode::Normal, EInternalObjectFlags::AllFlags);
 
+						// Ensure that we propagate the necessary flags from the destination object (outer) to the new subobject.
+						EObjectFlags FlagsToPropagate = InDestinationObject->GetMaskedFlags(RF_PropagateToSubObjects);
+						if (InDestinationObject->HasAnyFlags(RF_ClassDefaultObject) && !DuplicateValue->HasAnyFlags(RF_DefaultSubObject | RF_ArchetypeObject))
+						{
+							// Mark the new subobject as a template if its outer is a CDO and if it is not already flagged as a default subobject.
+							FlagsToPropagate |= RF_ArchetypeObject;
+						}
+
+						DuplicateValue->SetFlags(FlagsToPropagate);
+
 						FObjectPropertyBase* DestObjectProperty = CastFieldChecked<FObjectPropertyBase>(InDestinationProperty);
 						DestObjectProperty->SetObjectPropertyValue_InContainer(InTargetPtr, DuplicateValue);
 					}
 
-					// If the outers match, we should look for a corresponding object already in existance
+					// If the outers match, we should look for a corresponding object already in existence
 					// with the same name inside the destination object's outer.
 					if (Value->GetOuter() == SourceObject->GetOuter())
 					{
