CharacterMovement;

Making the GetCharacterMovement() function templated.

-- Code --
Character:
GetCharacterMovement() - Added a templated overload.

#tests compiled, PIE, local change that used it

diff --git a/Engine/Source/Runtime/Engine/Classes/GameFramework/Character.h b/Engine/Source/Runtime/Engine/Classes/GameFramework/Character.h
--- a/Engine/Source/Runtime/Engine/Classes/GameFramework/Character.h
+++ b/Engine/Source/Runtime/Engine/Classes/GameFramework/Character.h
@@ -368,7 +368,12 @@ public:
 	static FName MeshComponentName;
 
 	/** Returns CharacterMovement subobject **/
-	FORCEINLINE class UCharacterMovementComponent* GetCharacterMovement() const { return CharacterMovement; }
+	template <class T>
+	FORCEINLINE_DEBUGGABLE T* GetCharacterMovement() const
+	{
+		return CastChecked<T>(CharacterMovement, ECastCheckedType::NullAllowed);
+	}
+	FORCEINLINE UCharacterMovementComponent* GetCharacterMovement() const { return CharacterMovement; }
 
 	/** Name of the CharacterMovement component. Use this name if you want to use a different class (with ObjectInitializer.SetDefaultSubobjectClass). */
 	static FName CharacterMovementComponentName;
