Prevent array out of bounds crash in Stream cache.

diff --git a/Engine/Source/Runtime/Engine/Classes/Sound/SoundWave.h b/Engine/Source/Runtime/Engine/Classes/Sound/SoundWave.h
--- a/Engine/Source/Runtime/Engine/Classes/Sound/SoundWave.h
+++ b/Engine/Source/Runtime/Engine/Classes/Sound/SoundWave.h
@@ -106,8 +106,6 @@ struct ENGINE_API FStreamedAudioPlatformData
 {
 	GENERATED_USTRUCT_BODY()
 
-	/** Number of audio chunks. */
-	int32 NumChunks;
 	/** Format in which audio chunks are stored. */
 	FName AudioFormat;
 	/** audio data. */
diff --git a/Engine/Source/Runtime/Engine/Private/AudioDerivedData.cpp b/Engine/Source/Runtime/Engine/Private/AudioDerivedData.cpp
--- a/Engine/Source/Runtime/Engine/Private/AudioDerivedData.cpp
+++ b/Engine/Source/Runtime/Engine/Private/AudioDerivedData.cpp
@@ -111,7 +111,7 @@ int32 FStreamedAudioPlatformData::GetNumChunks() const
 	}
 #endif
 
-	return NumChunks;
+	return Chunks.Num();
 }
 
 FName FStreamedAudioPlatformData::GetAudioFormat() const
@@ -572,8 +572,6 @@ class FStreamedAudioCacheDerivedDataWorker : public FNonAbandonableTask
 					NewChunk->BulkData.Unlock();
 				}
 
-				DerivedData->NumChunks = DerivedData->Chunks.Num();
-
 				// Store it in the cache.
 				// @todo: This will remove the streaming bulk data, which we immediately reload below!
 				// Should ideally avoid this redundant work, but it only happens when we actually have
@@ -657,7 +655,7 @@ public:
 			// Load any streaming (not inline) chunks that are necessary for our platform.
 			if (bForDDC)
 			{
-				for (int32 Index = 0; Index < DerivedData->NumChunks; ++Index)
+				for (int32 Index = 0; Index < DerivedData->Chunks.Num(); ++Index)
 				{
 					if (!DerivedData->GetChunkFromDDC(Index, NULL))
 					{
@@ -919,9 +917,8 @@ bool FStreamedAudioPlatformData::TryInlineChunkData()
 #endif //WITH_EDITORONLY_DATA
 
 FStreamedAudioPlatformData::FStreamedAudioPlatformData()
-	: NumChunks(0)
 #if WITH_EDITORONLY_DATA
-	, AsyncTask(nullptr)
+	: AsyncTask(nullptr)
 #endif // #if WITH_EDITORONLY_DATA
 {
 }
@@ -1119,11 +1116,18 @@ void FStreamedAudioPlatformData::Serialize(FArchive& Ar, USoundWave* Owner)
 	}
 #endif
 
+	int32 NumChunks = 0;
+	if (Ar.IsSaving())
+	{
+		NumChunks = Chunks.Num();
+	}
+
 	Ar << NumChunks;
 	Ar << AudioFormat;
 
 	if (Ar.IsLoading())
 	{
+		check(NumChunks >= 0);
 		Chunks.Empty(NumChunks);
 		for (int32 ChunkIndex = 0; ChunkIndex < NumChunks; ++ChunkIndex)
 		{
