Update soundwave assets and run bakes on property undos. Fix for #jira UE-158431 #preflight 62f20955e60c9215b9294b9c

diff --git a/Engine/Source/Runtime/Engine/Private/SoundWave.cpp b/Engine/Source/Runtime/Engine/Private/SoundWave.cpp
--- a/Engine/Source/Runtime/Engine/Private/SoundWave.cpp
+++ b/Engine/Source/Runtime/Engine/Private/SoundWave.cpp
@@ -2459,6 +2459,18 @@ void USoundWave::PostEditChangeProperty(FPropertyChangedEvent& PropertyChangedEv
 	Super::PostEditChangeProperty(PropertyChangedEvent);
 	check(SoundWaveDataPtr);
 
+	if (PropertyChangedEvent.Property == nullptr)
+	{
+		//an empty event property field might mean the update comes from an undo
+		//we can't discern what properties where reverted so we update the asset and bakes
+
+		UpdateAsset();
+		BakeFFTAnalysis();
+		BakeEnvelopeAnalysis();
+
+		return;
+	}
+
 	static const FName CompressionQualityFName = GET_MEMBER_NAME_CHECKED(USoundWave, CompressionQuality);
 	static const FName SampleRateFName = GET_MEMBER_NAME_CHECKED(USoundWave,SampleRateQuality);
 	static const FName StreamingFName = GET_MEMBER_NAME_CHECKED(USoundWave, bStreaming);
@@ -2522,7 +2534,6 @@ void USoundWave::PostEditChangeProperty(FPropertyChangedEvent& PropertyChangedEv
 	}
 }
 
-
 #endif // WITH_EDITOR
 
 #if WITH_EDITORONLY_DATA
