[MLDeformer] Some static code analysis fixes.

diff --git a/Engine/Plugins/Experimental/MLDeformer/NeuralMorphModel/Source/NeuralMorphModel/Private/NeuralMorphModel.cpp b/Engine/Plugins/Experimental/MLDeformer/NeuralMorphModel/Source/NeuralMorphModel/Private/NeuralMorphModel.cpp
--- a/Engine/Plugins/Experimental/MLDeformer/NeuralMorphModel/Source/NeuralMorphModel/Private/NeuralMorphModel.cpp
+++ b/Engine/Plugins/Experimental/MLDeformer/NeuralMorphModel/Source/NeuralMorphModel/Private/NeuralMorphModel.cpp
@@ -46,14 +46,7 @@ void UNeuralMorphModel::Serialize(FArchive& Archive)
 	if (bHasMorphData)
 	{
 		check(MorphTargetSet.IsValid());
-		if (Archive.IsLoading())
-		{
-			Archive << MorphTargetSet->MorphBuffers;
-		}
-		else
-		{
-			Archive << MorphTargetSet->MorphBuffers;
-		}
+		Archive << MorphTargetSet->MorphBuffers;
 	}
 }
 
diff --git a/Engine/Source/Runtime/Engine/Private/SkeletalRenderGPUSkin.cpp b/Engine/Source/Runtime/Engine/Private/SkeletalRenderGPUSkin.cpp
--- a/Engine/Source/Runtime/Engine/Private/SkeletalRenderGPUSkin.cpp
+++ b/Engine/Source/Runtime/Engine/Private/SkeletalRenderGPUSkin.cpp
@@ -2327,7 +2327,10 @@ void FDynamicSkelMeshObjectDataGPUSkin::InitDynamicSkelMeshObjectDataGPUSkin(
 
 	ExternalMorphWeightData = InExternalMorphWeightData;
 	ExternalMorphWeightData.UpdateNumActiveMorphTargets();
-	ExternalMorphSets = InMeshComponent->GetExternalMorphSets(InLODIndex);
+	if (InMeshComponent)
+	{
+		ExternalMorphSets = InMeshComponent->GetExternalMorphSets(InLODIndex);
+	}
 
 	// Gather any bones referenced by shadow shapes
 	FSkeletalMeshSceneProxy* SkeletalMeshProxy = InMeshComponent ? (FSkeletalMeshSceneProxy*)InMeshComponent->SceneProxy : nullptr;
