Sort the sprites into a local array, since FSimpleElementCollector::Draw() can be called on the same object from multiple RDG passes, and they may be running in parallel.

#rnx

diff --git a/Engine/Source/Runtime/Engine/Private/BatchedElements.cpp b/Engine/Source/Runtime/Engine/Private/BatchedElements.cpp
--- a/Engine/Source/Runtime/Engine/Private/BatchedElements.cpp
+++ b/Engine/Source/Runtime/Engine/Private/BatchedElements.cpp
@@ -1043,13 +1043,15 @@ bool FBatchedElements::Draw(FRHICommandList& RHICmdList, const FMeshPassProcesso
 				EBlendModeFilter::Type ValidBlendFilter;
 			};
 
-			Sprites.Sort(FCompareSprite(Filter));
+			TArray<FBatchedSprite> SortedSprites = Sprites;
+
+			SortedSprites.Sort(FCompareSprite(Filter));
 
 			// count the number of sprites that have valid blend modes
 			// (they have been sorted to the front of the list)
 			int32 ValidSpriteCount = 0;
-			while (ValidSpriteCount < Sprites.Num() &&
-				(Filter & GetBlendModeFilter((ESimpleElementBlendMode)Sprites[ValidSpriteCount].BlendMode)))
+			while (ValidSpriteCount < SortedSprites.Num() &&
+				(Filter & GetBlendModeFilter((ESimpleElementBlendMode)SortedSprites[ValidSpriteCount].BlendMode)))
 			{
 				++ValidSpriteCount;
 			}
@@ -1063,7 +1065,7 @@ bool FBatchedElements::Draw(FRHICommandList& RHICmdList, const FMeshPassProcesso
 
 				for (int32 SpriteIndex = 0; SpriteIndex < ValidSpriteCount; SpriteIndex++)
 				{
-					const FBatchedSprite& Sprite = Sprites[SpriteIndex];
+					const FBatchedSprite& Sprite = SortedSprites[SpriteIndex];
 					FSimpleElementVertex* Vertex = &SpriteList[SpriteIndex * 6];
 
 					// Compute the sprite vertices.
@@ -1086,19 +1088,19 @@ bool FBatchedElements::Draw(FRHICommandList& RHICmdList, const FMeshPassProcesso
 				RHICmdList.UnlockBuffer(VertexBufferRHI);
 
 				//First time init
-				const FTexture* CurrentTexture = Sprites[0].Texture;
-				ESimpleElementBlendMode CurrentBlendMode = (ESimpleElementBlendMode)Sprites[0].BlendMode;
+				const FTexture* CurrentTexture = SortedSprites[0].Texture;
+				ESimpleElementBlendMode CurrentBlendMode = (ESimpleElementBlendMode)SortedSprites[0].BlendMode;
 				int32 BatchStartIndex = 0;
-				float CurrentOpacityMask = Sprites[0].OpacityMaskRefVal;
+				float CurrentOpacityMask = SortedSprites[0].OpacityMaskRefVal;
 
 				// Start loop at 1, since we've already started the first batch with the first sprite in the list
 				for (int32 SpriteIndex = 1; SpriteIndex < ValidSpriteCount + 1; SpriteIndex++)
 				{
 					// Need to flush the current batch once we hit the end of the list, or if state of this sprite doesn't match current batch
 					if (SpriteIndex == ValidSpriteCount ||
-						CurrentTexture != Sprites[SpriteIndex].Texture ||
-						CurrentBlendMode != Sprites[SpriteIndex].BlendMode ||
-						CurrentOpacityMask != Sprites[SpriteIndex].OpacityMaskRefVal)
+						CurrentTexture != SortedSprites[SpriteIndex].Texture ||
+						CurrentBlendMode != SortedSprites[SpriteIndex].BlendMode ||
+						CurrentOpacityMask != SortedSprites[SpriteIndex].OpacityMaskRefVal)
 					{
 						const int32 SpriteNum = SpriteIndex - BatchStartIndex;
 						const int32 BaseVertex = BatchStartIndex * 6;
@@ -1111,9 +1113,9 @@ bool FBatchedElements::Draw(FRHICommandList& RHICmdList, const FMeshPassProcesso
 						if (SpriteIndex < ValidSpriteCount)
 						{
 							BatchStartIndex = SpriteIndex;
-							CurrentTexture = Sprites[SpriteIndex].Texture;
-							CurrentBlendMode = (ESimpleElementBlendMode)Sprites[SpriteIndex].BlendMode;
-							CurrentOpacityMask = Sprites[SpriteIndex].OpacityMaskRefVal;
+							CurrentTexture = SortedSprites[SpriteIndex].Texture;
+							CurrentBlendMode = (ESimpleElementBlendMode)SortedSprites[SpriteIndex].BlendMode;
+							CurrentOpacityMask = SortedSprites[SpriteIndex].OpacityMaskRefVal;
 						}
 					}
 				}
diff --git a/Engine/Source/Runtime/Engine/Public/BatchedElements.h b/Engine/Source/Runtime/Engine/Public/BatchedElements.h
--- a/Engine/Source/Runtime/Engine/Public/BatchedElements.h
+++ b/Engine/Source/Runtime/Engine/Public/BatchedElements.h
@@ -342,8 +342,7 @@ private:
 		float OpacityMaskRefVal;
 		uint8 BlendMode;
 	};
-	/** This array is sorted during draw-calls */
-	mutable TArray<FBatchedSprite> Sprites;
+	TArray<FBatchedSprite> Sprites;
 
 	struct FBatchedMeshElement
 	{
