Fix crash when quality is set low for HISM when density scaling is enabled. Ignore instance index from mapping when they are filtered out as INDEX_NONE by scalability.

diff --git a/Engine/Source/Runtime/Engine/Private/InstancedStaticMesh.cpp b/Engine/Source/Runtime/Engine/Private/InstancedStaticMesh.cpp
--- a/Engine/Source/Runtime/Engine/Private/InstancedStaticMesh.cpp
+++ b/Engine/Source/Runtime/Engine/Private/InstancedStaticMesh.cpp
@@ -1469,6 +1469,14 @@ void FInstancedStaticMeshSceneProxy::SetupProxy(UInstancedStaticMeshComponent* I
 		for (int32 InstanceIndex = 0; InstanceIndex < InComponent->GetInstanceCount(); ++InstanceIndex)
 		{
 			const int32 RenderInstanceIndex = InstanceReorderTable.IsValidIndex(InstanceIndex) ? InstanceReorderTable[InstanceIndex] : InstanceIndex;
+
+			if (RenderInstanceIndex == INDEX_NONE)
+			{
+				// Happens when the instance is excluded due to scalability, e.g., in the Build() of
+				// HISM.
+				continue;
+			}
+
 			if (!ensure(InstanceSceneData.IsValidIndex(RenderInstanceIndex)))
 			{
 				continue;
