Animation: Fixed UAnimSequence cache stats and missing debug name, and made compression async during cooking

Reduced UAnimSequence::PreSave time when cooking Lyra from 17.4s to 0.9s.

- Added UAnimSequence::IsCachedCookedPlatformDataLoaded to poll async compression rather than always blocking on it in PreSave.
- Added FDerivedDataAnimationCompression::GetDebugContextString to report the asset name to DDC for debugging, logging, and profiling.
- Added cook stat timing and hit/miss size reporting to FAsyncCompressedAnimationsManagement::OnActiveCompressionFinished.
- Added cook stat timing of FDerivedDataAnimationCompression::Build when executed off the game thread.
- Fixed FAsyncCompressedAnimationsManagement to be ticked during cooking.
- Fixed a memory leak when async compression tasks are canceled.

diff --git a/Engine/Source/Runtime/Engine/Classes/Animation/AnimSequence.h b/Engine/Source/Runtime/Engine/Classes/Animation/AnimSequence.h
--- a/Engine/Source/Runtime/Engine/Classes/Animation/AnimSequence.h
+++ b/Engine/Source/Runtime/Engine/Classes/Animation/AnimSequence.h
@@ -393,6 +393,7 @@ public:
 	virtual void GetPreloadDependencies(TArray<UObject*>& OutDeps) override;
 #if WITH_EDITOR
 	virtual void PostEditChangeProperty(FPropertyChangedEvent& PropertyChangedEvent) override;
+	virtual bool IsCachedCookedPlatformDataLoaded(const ITargetPlatform* TargetPlatform) override;
 #endif // WITH_EDITOR
 	virtual void BeginDestroy() override;
 	virtual void GetAssetRegistryTags(TArray<FAssetRegistryTag>& OutTags) const override;
diff --git a/Engine/Source/Runtime/Engine/Private/Animation/AnimCompressionDerivedData.cpp b/Engine/Source/Runtime/Engine/Private/Animation/AnimCompressionDerivedData.cpp
--- a/Engine/Source/Runtime/Engine/Private/Animation/AnimCompressionDerivedData.cpp
+++ b/Engine/Source/Runtime/Engine/Private/Animation/AnimCompressionDerivedData.cpp
@@ -14,6 +14,7 @@
 #include "AnimationCompression.h"
 #include "UObject/Package.h"
 #include "Misc/CoreDelegates.h"
+#include "ProfilingDebugging/CookStats.h"
 
 #if WITH_EDITOR
 
@@ -39,6 +40,12 @@ const TCHAR* FDerivedDataAnimationCompression::GetVersionString() const
 	return TEXT("3A3F9FBB7A4047278CABE35820CC44D7");
 }
 
+FString FDerivedDataAnimationCompression::GetDebugContextString() const
+{
+	check(DataToCompressPtr.IsValid());
+	return DataToCompressPtr->FullName;
+}
+
 bool FDerivedDataAnimationCompression::Build( TArray<uint8>& OutDataArray )
 {
 	TRACE_CPUPROFILER_EVENT_SCOPE(FDerivedDataAnimationCompression::Build);
@@ -54,6 +61,10 @@ bool FDerivedDataAnimationCompression::Build( TArray<uint8>& OutDataArray )
 		return false;
 	}
 
+	// Update UsageStats only when not running on the game thread since that thread times this at a higher level.
+	COOK_STAT(auto Timer = AnimSequenceCookStats::UsageStats.TimeSyncWork());
+	COOK_STAT(IsInGameThread() ? Timer.Cancel() : Timer.TrackCyclesOnly());
+
 	SCOPE_CYCLE_COUNTER(STAT_AnimCompressionDerivedData);
 	UE_LOG(LogAnimationCompression, Log, TEXT("Building Anim DDC data for %s"), *DataToCompress.FullName);
 
@@ -169,15 +180,22 @@ void FAsyncCompressedAnimationsManagement::Shutdown()
 
 void FAsyncCompressedAnimationsManagement::OnActiveCompressionFinished(int32 ActiveAnimIndex)
 {
+	COOK_STAT(auto Timer = AnimSequenceCookStats::UsageStats.TimeSyncWork());
+
 	FDerivedDataCacheInterface& DerivedDataCache = GetDerivedDataCacheRef();
 
 	FActiveAsyncCompressionTask& Task = ActiveAsyncCompressionTasks[ActiveAnimIndex];
 
-	if(Task.Sequence)
+	TArray<uint8> OutData;
+	bool bBuiltLocally = false;
+	const bool bOk = DerivedDataCache.GetAsynchronousResults(Task.AsyncHandle, OutData, &bBuiltLocally);
+
+	COOK_STAT(using EHitOrMiss = FCookStats::CallStats::EHitOrMiss);
+	COOK_STAT(Timer.AddHitOrMiss(bBuiltLocally ? EHitOrMiss::Miss : EHitOrMiss::Hit, OutData.Num()));
+
+	if (Task.Sequence)
 	{
-		TArray<uint8> OutData;
-		bool bBuiltLocally = false;
-		if (DerivedDataCache.GetAsynchronousResults(Task.AsyncHandle, OutData, &bBuiltLocally))
+		if (bOk)
 		{
 			Task.Sequence->ApplyCompressedData(Task.CacheKey, Task.bPerformFrameStripping, OutData);
 		}
@@ -263,6 +281,11 @@ void FAsyncCompressedAnimationsManagement::Tick(float DeltaTime)
 	}
 }
 
+void FAsyncCompressedAnimationsManagement::TickCook(float DeltaTime, bool bCookCompete)
+{
+	Tick(DeltaTime);
+}
+
 void FAsyncCompressedAnimationsManagement::StartQueuedTasks(int32 MaxActiveTasks)
 {
 	while (ActiveAsyncCompressionTasks.Num() < MaxActiveTasks)
diff --git a/Engine/Source/Runtime/Engine/Private/Animation/AnimCompressionDerivedData.h b/Engine/Source/Runtime/Engine/Private/Animation/AnimCompressionDerivedData.h
--- a/Engine/Source/Runtime/Engine/Private/Animation/AnimCompressionDerivedData.h
+++ b/Engine/Source/Runtime/Engine/Private/Animation/AnimCompressionDerivedData.h
@@ -2,7 +2,7 @@
 
 #pragma once
 
-#include "CoreMinimal.h"
+#include "CoreTypes.h"
 
 #if WITH_EDITOR
 #include "DerivedDataPluginInterface.h"
@@ -11,6 +11,7 @@
 #include "Animation/AnimCompressionTypes.h"
 
 struct FAnimCompressContext;
+struct FDerivedDataUsageStats;
 
 #if WITH_EDITOR
 
@@ -59,6 +60,7 @@ public:
 		return AssetDDCKey;
 	}
 
+	virtual FString GetDebugContextString() const override;
 
 	virtual bool IsBuildThreadsafe() const override
 	{
@@ -74,4 +76,9 @@ public:
 	}
 };
 
+namespace AnimSequenceCookStats
+{
+	extern FDerivedDataUsageStats UsageStats;
+}
+
 #endif	//WITH_EDITOR
diff --git a/Engine/Source/Runtime/Engine/Private/Animation/AnimSequence.cpp b/Engine/Source/Runtime/Engine/Private/Animation/AnimSequence.cpp
--- a/Engine/Source/Runtime/Engine/Private/Animation/AnimSequence.cpp
+++ b/Engine/Source/Runtime/Engine/Private/Animation/AnimSequence.cpp
@@ -80,7 +80,7 @@ DECLARE_CYCLE_STAT(TEXT("AnimSeq EvalCurveData"), STAT_AnimSeq_EvalCurveData, ST
 #if ENABLE_COOK_STATS
 namespace AnimSequenceCookStats
 {
-	static FCookStats::FDDCResourceUsageStats UsageStats;
+	FDerivedDataUsageStats UsageStats;
 	static FCookStatsManager::FAutoRegisterCallback RegisterCookStats([](FCookStatsManager::AddStatFuncRef AddStat)
 	{
 		UsageStats.LogStats(AddStat, TEXT("AnimSequence.Usage"), TEXT(""));
@@ -1174,6 +1174,13 @@ void UAnimSequence::PostEditChangeProperty(FPropertyChangedEvent& PropertyChange
 }
 #endif // WITH_EDITOR
 
+#if WITH_EDITOR
+bool UAnimSequence::IsCachedCookedPlatformDataLoaded(const ITargetPlatform* TargetPlatform)
+{
+	return !bCompressionInProgress;
+}
+#endif // WITH_EDITOR
+
 // @todo DB: Optimize!
 template<typename TimeArray>
 static int32 FindKeyIndex(float Time, const TimeArray& Times)
diff --git a/Engine/Source/Runtime/Engine/Public/Animation/AnimCompressionDerivedDataPublic.h b/Engine/Source/Runtime/Engine/Public/Animation/AnimCompressionDerivedDataPublic.h
--- a/Engine/Source/Runtime/Engine/Public/Animation/AnimCompressionDerivedDataPublic.h
+++ b/Engine/Source/Runtime/Engine/Public/Animation/AnimCompressionDerivedDataPublic.h
@@ -6,7 +6,6 @@
 
 #include "Animation/AnimCompressionTypes.h"
 #include "TickableEditorObject.h"
-#include "Animation/AnimCompressionTypes.h"
 #include "UObject/GCObject.h"
 
 class UAnimSequence;
@@ -54,7 +53,7 @@ struct FQueuedAsyncCompressionWork
 //   Maintains active compressions
 //   tracks memory usage of async compression
 //   Gives API for blocking on compression
-class ENGINE_API FAsyncCompressedAnimationsManagement : public FTickableEditorObject, public FGCObject
+class ENGINE_API FAsyncCompressedAnimationsManagement : public FTickableEditorObject, public FTickableCookObject, public FGCObject
 {
 public:
 	static FAsyncCompressedAnimationsManagement& Get();
@@ -96,6 +95,10 @@ private:
 	virtual TStatId GetStatId() const override;
 	/* End FTickableEditorObject */
 
+	/* Begin FTickableCookObject */
+	virtual void TickCook(float DeltaTime, bool bCookCompete) override;
+	/* End FTickableCookObject */
+
 	// FGCObject interface
 	void AddReferencedObjects(FReferenceCollector& Collector) override;
 	virtual FString GetReferencerName() const override
