Expose a SkeletalMeshComponent setting to make it possible to update the skeletal mesh output when kinematic

diff --git a/Engine/Source/Runtime/Engine/Classes/Components/SkeletalMeshComponent.h b/Engine/Source/Runtime/Engine/Classes/Components/SkeletalMeshComponent.h
--- a/Engine/Source/Runtime/Engine/Classes/Components/SkeletalMeshComponent.h
+++ b/Engine/Source/Runtime/Engine/Classes/Components/SkeletalMeshComponent.h
@@ -629,6 +629,14 @@ public:
 	UPROPERTY(EditAnywhere, AdvancedDisplay, BlueprintReadWrite, Category = SkeletalMesh)
 	uint8 bEnablePhysicsOnDedicatedServer:1;
 
+	/** 
+	 * If true, then the physics bodies will be used to drive the skeletal mesh even when they are
+	 * kinematic (not simulating), otherwise the skeletal mesh will be driven by the animation input 
+	 * when the bodies are kinematic
+	 */
+	UPROPERTY(EditAnywhere, AdvancedDisplay, BlueprintReadWrite, Category = SkeletalMesh)
+	uint8 bUpdateMeshWhenKinematic:1;
+
 	/**
 	 *	If we should pass joint position to joints each frame, so that they can be used by motorized joints to drive the
 	 *	ragdoll based on the animation.
diff --git a/Engine/Source/Runtime/Engine/Private/Components/SkeletalMeshComponent.cpp b/Engine/Source/Runtime/Engine/Private/Components/SkeletalMeshComponent.cpp
--- a/Engine/Source/Runtime/Engine/Private/Components/SkeletalMeshComponent.cpp
+++ b/Engine/Source/Runtime/Engine/Private/Components/SkeletalMeshComponent.cpp
@@ -246,6 +246,8 @@ USkeletalMeshComponent::USkeletalMeshComponent(const FObjectInitializer& ObjectI
 	bEnableUpdateRateOptimizations = false;
 	RagdollAggregateThreshold = UPhysicsSettings::Get()->RagdollAggregateThreshold;
 
+	bUpdateMeshWhenKinematic = false;
+
 	LastPoseTickFrame = 0u;
 
 	bHasCustomNavigableGeometry = EHasCustomNavigableGeometry::Yes;
diff --git a/Engine/Source/Runtime/Engine/Private/PhysicsEngine/PhysAnim.cpp b/Engine/Source/Runtime/Engine/Private/PhysicsEngine/PhysAnim.cpp
--- a/Engine/Source/Runtime/Engine/Private/PhysicsEngine/PhysAnim.cpp
+++ b/Engine/Source/Runtime/Engine/Private/PhysicsEngine/PhysAnim.cpp
@@ -205,6 +205,10 @@ void USkeletalMeshComponent::PerformBlendPhysicsBones(const TArray<FBoneIndexTyp
 			// need to update back to physX so that physX knows where it was after blending
 			FBodyInstance* PhysicsAssetBodyInstance = nullptr;
 
+			// Gets set to true if we have a valid body and the skeletal mesh option has been set to
+			// be driven by kinematic body parts.
+			bool bDriveMeshWhenKinematic = false;
+
 			// If so - get its world space matrix and its parents world space matrix and calc relative atom.
 			if(BodyIndex != INDEX_NONE )
 			{	
@@ -221,8 +225,12 @@ void USkeletalMeshComponent::PerformBlendPhysicsBones(const TArray<FBoneIndexTyp
 #endif
 				PhysicsAssetBodyInstance = Bodies[BodyIndex];
 
+				bDriveMeshWhenKinematic =
+					bUpdateMeshWhenKinematic &&
+					PhysicsAssetBodyInstance->IsValidBodyInstance();
+
 				//if simulated body copy back and blend with animation
-				if(PhysicsAssetBodyInstance->IsInstanceSimulatingPhysics())
+				if(PhysicsAssetBodyInstance->IsInstanceSimulatingPhysics() || bDriveMeshWhenKinematic)
 				{
 					FTransform PhysTM = PhysicsAssetBodyInstance->GetUnrealWorldTransform_AssumesLocked();
 
@@ -230,10 +238,10 @@ void USkeletalMeshComponent::PerformBlendPhysicsBones(const TArray<FBoneIndexTyp
 					WorldBoneTMs[BoneIndex].TM = PhysTM;
 					WorldBoneTMs[BoneIndex].bUpToDate = true;
 
-					float UsePhysWeight = (bBlendPhysics)? 1.f : PhysicsAssetBodyInstance->PhysicsBlendWeight;
+					float UsePhysWeight = (bBlendPhysics || bDriveMeshWhenKinematic)? 1.f : PhysicsAssetBodyInstance->PhysicsBlendWeight;
 
 					// if the body instance is disabled, then we want to use the animation transform and ignore the physics one
-					if (PhysicsAssetBodyInstance->IsPhysicsDisabled())
+					if (PhysicsAssetBodyInstance->IsPhysicsDisabled() && !bDriveMeshWhenKinematic)
 					{
 						UsePhysWeight = 0.0f;
 					}
@@ -241,7 +249,7 @@ void USkeletalMeshComponent::PerformBlendPhysicsBones(const TArray<FBoneIndexTyp
 					// Find this bones parent matrix.
 					FTransform ParentWorldTM;
 
-					// if we wan't 'full weight' we just find 
+					// if we want 'full weight' we just find 
 					if(UsePhysWeight > 0.f)
 					{
 						if (!(ensure(InOutBoneSpaceTransforms.Num())))
@@ -300,7 +308,7 @@ void USkeletalMeshComponent::PerformBlendPhysicsBones(const TArray<FBoneIndexTyp
 			}
 			else
 			{
-				if(bLocalSpaceKinematics || BodyIndex == INDEX_NONE || Bodies[BodyIndex]->IsInstanceSimulatingPhysics())
+				if(bDriveMeshWhenKinematic || bLocalSpaceKinematics || BodyIndex == INDEX_NONE || Bodies[BodyIndex]->IsInstanceSimulatingPhysics())
 				{
 					if (!(ensure(BoneIndex < InOutBoneSpaceTransforms.Num())))
 					{
