UAnimMontage: Fix crash in PostLoad() if there is no valid skeleton

diff --git a/Engine/Source/Runtime/Engine/Private/Animation/AnimMontage.cpp b/Engine/Source/Runtime/Engine/Private/Animation/AnimMontage.cpp
--- a/Engine/Source/Runtime/Engine/Private/Animation/AnimMontage.cpp
+++ b/Engine/Source/Runtime/Engine/Private/Animation/AnimMontage.cpp
@@ -438,8 +438,8 @@ void UAnimMontage::PostLoad()
 	}
 
 	// verify if skeleton matches, otherwise clear it, this can happen if anim sequence has been modified when this hasn't been loaded. 
+	if (const USkeleton* MySkeleton = GetSkeleton())
 	{
-		USkeleton* MySkeleton = GetSkeleton();
 		for (int32 I=0; I<SlotAnimTracks.Num(); ++I)
 		{
 			if ( SlotAnimTracks[I].AnimTrack.AnimSegments.Num() > 0 )
