DataLayers: Use REPNOTIFY_Always until dynamic array rep notify issue is fixed (UE-155774)

diff --git a/Engine/Source/Runtime/Engine/Private/WorldPartition/DataLayer/WorldDataLayers.cpp b/Engine/Source/Runtime/Engine/Private/WorldPartition/DataLayer/WorldDataLayers.cpp
--- a/Engine/Source/Runtime/Engine/Private/WorldPartition/DataLayer/WorldDataLayers.cpp
+++ b/Engine/Source/Runtime/Engine/Private/WorldPartition/DataLayer/WorldDataLayers.cpp
@@ -56,10 +56,13 @@ void AWorldDataLayers::GetLifetimeReplicatedProps(TArray< FLifetimeProperty >& O
 {
 	Super::GetLifetimeReplicatedProps(OutLifetimeProps);
 
-	DOREPLIFETIME(AWorldDataLayers, RepLoadedDataLayerNames);
-	DOREPLIFETIME(AWorldDataLayers, RepActiveDataLayerNames);
-	DOREPLIFETIME(AWorldDataLayers, RepEffectiveLoadedDataLayerNames);
-	DOREPLIFETIME(AWorldDataLayers, RepEffectiveActiveDataLayerNames);
+	// Use RepNotifyCondition = REPNOTIFY_Always because of current issue with dynamic arrays: https://jira.it.epicgames.com/browse/UE-155774
+	FDoRepLifetimeParams Params;
+	Params.RepNotifyCondition = REPNOTIFY_Always;
+	DOREPLIFETIME_WITH_PARAMS_FAST(AWorldDataLayers, RepLoadedDataLayerNames, Params);
+	DOREPLIFETIME_WITH_PARAMS_FAST(AWorldDataLayers, RepActiveDataLayerNames, Params);
+	DOREPLIFETIME_WITH_PARAMS_FAST(AWorldDataLayers, RepEffectiveLoadedDataLayerNames, Params);
+	DOREPLIFETIME_WITH_PARAMS_FAST(AWorldDataLayers, RepEffectiveActiveDataLayerNames, Params);
 }
 
 void AWorldDataLayers::BeginPlay()
