[Engine]
* Fixed so if -SAVEWINPOS is used, window is created in the right position instead of being created in a different position and then moved to saved position quite much later

diff --git a/Engine/Source/Runtime/Engine/Private/GameEngine.cpp b/Engine/Source/Runtime/Engine/Private/GameEngine.cpp
--- a/Engine/Source/Runtime/Engine/Private/GameEngine.cpp
+++ b/Engine/Source/Runtime/Engine/Private/GameEngine.cpp
@@ -258,14 +258,6 @@ void UGameEngine::CreateGameViewport( UGameViewportClient* GameViewportClient )
 	int32 SaveWinPos;
 	if (FParse::Value(FCommandLine::Get(), TEXT("SAVEWINPOS="), SaveWinPos) && SaveWinPos > 0 )
 	{
-		// Get WinX/WinY from GameSettings, apply them if valid.
-		FIntPoint PiePosition = GetGameUserSettings()->GetWindowPosition();
-		if (PiePosition.X >= 0 && PiePosition.Y >= 0)
-		{
-			int32 WinX = GetGameUserSettings()->GetWindowPosition().X;
-			int32 WinY = GetGameUserSettings()->GetWindowPosition().Y;
-			Window->MoveWindowTo(FVector2D(WinX, WinY));
-		}
 		Window->SetOnWindowMoved( FOnWindowMoved::CreateUObject( this, &UGameEngine::OnGameWindowMoved ) );
 	}
 
@@ -555,6 +547,21 @@ TSharedRef<SWindow> UGameEngine::CreateGameWindow()
 		AutoCenterType = EAutoCenter::None;
 	}
 
+	// SAVEWINPOS tells us to load/save window positions to user settings (this is disabled by default)
+	int32 SaveWinPos;
+	if (FParse::Value(FCommandLine::Get(), TEXT("SAVEWINPOS="), SaveWinPos) && SaveWinPos > 0)
+	{
+		// Note GameUserSettings is not instantiated here yet, so we need to read directly from the configs
+		FString ScriptEngineCategory = TEXT("/Script/Engine.Engine");
+		FString GameUserSettingsCategory = TEXT("/Script/Engine.GameUserSettings");
+		GConfig->GetString(*ScriptEngineCategory, TEXT("GameUserSettingsClassName"), GameUserSettingsCategory, GEngineIni);
+		if (GConfig->GetInt(*GameUserSettingsCategory, TEXT("WindowPosX"), WinX, GGameUserSettingsIni) &&
+			GConfig->GetInt(*GameUserSettingsCategory, TEXT("WindowPosY"), WinY, GGameUserSettingsIni))
+		{
+			AutoCenterType = EAutoCenter::None;
+		}
+	}
+
 	// Give the window the max width/height of either the requested resolution, or your available desktop resolution
 	// We need to do this as we request some 4K windows when rendering sequences, and the OS may try and clamp that
 	// window to your available desktop resolution
