Removing usage of master/slave in Tick-Task-Manager

diff --git a/Engine/Source/Runtime/Engine/Private/TickTaskManager.cpp b/Engine/Source/Runtime/Engine/Private/TickTaskManager.cpp
--- a/Engine/Source/Runtime/Engine/Private/TickTaskManager.cpp
+++ b/Engine/Source/Runtime/Engine/Private/TickTaskManager.cpp
@@ -1152,13 +1152,13 @@ public:
 	}
 	// Interface that is private to FTickFunction
 
-	/** Return true if this tick function is in the master list **/
+	/** Return true if this tick function is in the primary list **/
 	bool HasTickFunction(FTickFunction* TickFunction)
 	{
 		return AllEnabledTickFunctions.Contains(TickFunction) || AllDisabledTickFunctions.Contains(TickFunction) || AllCoolingDownTickFunctions.Contains(TickFunction);
 	}
 
-	/** Add the tick function to the master list **/
+	/** Add the tick function to the primary list **/
 	void AddTickFunction(FTickFunction* TickFunction)
 	{
 		check(!HasTickFunction(TickFunction));
@@ -1278,7 +1278,7 @@ public:
 		}
 	}
 
-	/** Remove the tick function from the master list **/
+	/** Remove the tick function from the primary list **/
 	void RemoveTickFunction(FTickFunction* TickFunction)
 	{
 		switch(TickFunction->TickState)
@@ -1412,11 +1412,11 @@ private:
 
 	/** Global Sequencer														*/
 	FTickTaskSequencer&							TickTaskSequencer;
-	/** Master list of enabled tick functions **/
+	/** Primary list of enabled tick functions **/
 	TSet<FTickFunction*>						AllEnabledTickFunctions;
-	/** Master list of enabled tick functions **/
+	/** Primary list of enabled tick functions **/
 	FCoolingDownTickFunctionList				AllCoolingDownTickFunctions;
-	/** Master list of disabled tick functions **/
+	/** Primary list of disabled tick functions **/
 	TSet<FTickFunction*>						AllDisabledTickFunctions;
 	/** Utility array to avoid memory reallocations when collecting functions to reschedule **/
 	TArrayWithThreadsafeAdd<FTickScheduleDetails>				TickFunctionsToReschedule;
@@ -1639,13 +1639,13 @@ public:
 
 	// Interface that is private to FTickFunction
 
-	/** Return true if this tick function is in the master list **/
+	/** Return true if this tick function is in the primary list **/
 	bool HasTickFunction(ULevel* InLevel, FTickFunction* TickFunction)
 	{
 		FTickTaskLevel* Level = TickTaskLevelForLevel(InLevel);
 		return Level->HasTickFunction(TickFunction);
 	}
-	/** Add the tick function to the master list **/
+	/** Add the tick function to the primary list **/
 	void AddTickFunction(ULevel* InLevel, FTickFunction* TickFunction)
 	{
 		check(TickFunction->TickGroup >= 0 && TickFunction->TickGroup < TG_NewlySpawned); // You may not schedule a tick in the newly spawned group...they can only end up there if they are spawned late in a frame.
@@ -1653,7 +1653,7 @@ public:
 		Level->AddTickFunction(TickFunction);
 		TickFunction->InternalData->TickTaskLevel = Level;
 	}
-	/** Remove the tick function from the master list **/
+	/** Remove the tick function from the primary list **/
 	void RemoveTickFunction(FTickFunction* TickFunction)
 	{
 		check(TickFunction->InternalData);
@@ -1844,7 +1844,7 @@ FTickFunction::~FTickFunction()
 
 
 /**
-* Adds the tick function to the master list of tick functions.
+* Adds the tick function to the primary list of tick functions.
 * @param Level - level to place this tick function in
 **/
 void FTickFunction::RegisterTickFunction(ULevel* Level)
@@ -1869,7 +1869,7 @@ void FTickFunction::RegisterTickFunction(ULevel* Level)
 	}
 }
 
-/** Removes the tick function from the master list of tick functions. **/
+/** Removes the tick function from the primary list of tick functions. **/
 void FTickFunction::UnRegisterTickFunction()
 {
 	if (IsTickFunctionRegistered())
