Removing usage of master/slave in GameFramework-JIRA

diff --git a/Engine/Source/Runtime/Engine/Classes/Engine/CloudStorageBase.h b/Engine/Source/Runtime/Engine/Classes/Engine/CloudStorageBase.h
--- a/Engine/Source/Runtime/Engine/Classes/Engine/CloudStorageBase.h
+++ b/Engine/Source/Runtime/Engine/Classes/Engine/CloudStorageBase.h
@@ -182,7 +182,7 @@ class UCloudStorageBase
 
 	/**
 	 * If there was a conflict notification, this will tell the cloud interface
-	 * to choose the version with a given Index to be the master version, and to
+	 * to choose the version with a given Index to be the primary version, and to
 	 * toss any others.
 	 *
 	 * @param Index Conflict version index.
diff --git a/Engine/Source/Runtime/Engine/Classes/Engine/EngineBaseTypes.h b/Engine/Source/Runtime/Engine/Classes/Engine/EngineBaseTypes.h
--- a/Engine/Source/Runtime/Engine/Classes/Engine/EngineBaseTypes.h
+++ b/Engine/Source/Runtime/Engine/Classes/Engine/EngineBaseTypes.h
@@ -297,11 +297,11 @@ public:
 	virtual ~FTickFunction();
 
 	/** 
-	 * Adds the tick function to the master list of tick functions. 
+	 * Adds the tick function to the primary list of tick functions. 
 	 * @param Level - level to place this tick function in
 	 **/
 	void RegisterTickFunction(class ULevel* Level);
-	/** Removes the tick function from the master list of tick functions. **/
+	/** Removes the tick function from the primary list of tick functions. **/
 	void UnRegisterTickFunction();
 	/** See if the tick function is currently registered */
 	bool IsTickFunctionRegistered() const { return (InternalData && InternalData->bRegistered); }
diff --git a/Engine/Source/Runtime/Engine/Classes/Engine/LocalPlayer.h b/Engine/Source/Runtime/Engine/Classes/Engine/LocalPlayer.h
--- a/Engine/Source/Runtime/Engine/Classes/Engine/LocalPlayer.h
+++ b/Engine/Source/Runtime/Engine/Classes/Engine/LocalPlayer.h
@@ -176,14 +176,14 @@ public:
 	/** The FUniqueNetId which this player is associated with. */
 	FUniqueNetIdRepl CachedUniqueNetId;
 
-	/** The master viewport containing this player's view. */
+	/** The primary viewport containing this player's view. */
 	UPROPERTY()
 	TObjectPtr<class UGameViewportClient> ViewportClient;
 
-	/** The coordinates for the upper left corner of the master viewport subregion allocated to this player. 0-1 */
+	/** The coordinates for the upper left corner of the primary viewport subregion allocated to this player. 0-1 */
 	FVector2D Origin;
 
-	/** The size of the master viewport subregion allocated to this player. 0-1 */
+	/** The size of the primary viewport subregion allocated to this player. 0-1 */
 	FVector2D Size;
 
 	/** The location of the player's view the previous frame. */
diff --git a/Engine/Source/Runtime/Engine/Classes/GameFramework/AsyncActionHandleSaveGame.h b/Engine/Source/Runtime/Engine/Classes/GameFramework/AsyncActionHandleSaveGame.h
--- a/Engine/Source/Runtime/Engine/Classes/GameFramework/AsyncActionHandleSaveGame.h
+++ b/Engine/Source/Runtime/Engine/Classes/GameFramework/AsyncActionHandleSaveGame.h
@@ -26,7 +26,7 @@ public:
 	 *
 	 * @param SaveGameObject	Object that contains data about the save game that we want to write out.
 	 * @param SlotName			Name of the save game slot to load from.
-	 * @param UserIndex			For some platforms, master user index to identify the user doing the loading.
+	 * @param UserIndex			For some platforms, primary user index to identify the user doing the loading.
 	 */
 	UFUNCTION(BlueprintCallable, meta=(BlueprintInternalUseOnly="true", Category = "SaveGame",  WorldContext = "WorldContextObject"))
 	static UAsyncActionHandleSaveGame* AsyncSaveGameToSlot(UObject* WorldContextObject, USaveGame* SaveGameObject, const FString& SlotName, const int32 UserIndex);
@@ -37,7 +37,7 @@ public:
 	 * Keep in mind that some platforms may not support trying to load and save at the same time.
 	 *
 	 * @param SlotName			Name of the save game slot to load from.
-	 * @param UserIndex			For some platforms, master user index to identify the user doing the loading.
+	 * @param UserIndex			For some platforms, primary user index to identify the user doing the loading.
 	 */
 	UFUNCTION(BlueprintCallable, meta = (BlueprintInternalUseOnly = "true", Category = "SaveGame", WorldContext = "WorldContextObject"))
 	static UAsyncActionHandleSaveGame* AsyncLoadGameFromSlot(UObject* WorldContextObject, const FString& SlotName, const int32 UserIndex);
diff --git a/Engine/Source/Runtime/Engine/Classes/Kismet/GameplayStatics.h b/Engine/Source/Runtime/Engine/Classes/Kismet/GameplayStatics.h
--- a/Engine/Source/Runtime/Engine/Classes/Kismet/GameplayStatics.h
+++ b/Engine/Source/Runtime/Engine/Classes/Kismet/GameplayStatics.h
@@ -1095,7 +1095,7 @@ public:
 	 * Save the contents of the buffer to a platform-specific save slot/file 
 	 * @param InSaveData		Data to save
 	 * @param SlotName			Name of save game slot to save to.
-	 * @param UserIndex			For some platforms, master user index to identify the user doing the saving.
+	 * @param UserIndex			For some platforms, primary user index to identify the user doing the saving.
 	 * @return					Whether we successfully saved this information
 	 */
 	static bool SaveDataToSlot(const TArray<uint8>& InSaveData, const FString& SlotName, const int32 UserIndex);
@@ -1107,7 +1107,7 @@ public:
 	 *
 	 * @param SaveGameObject	Object that contains data about the save game that we want to write out.
 	 * @param SlotName			Name of the save game slot to load from.
-	 * @param UserIndex			For some platforms, master user index to identify the user doing the loading.
+	 * @param UserIndex			For some platforms, primary user index to identify the user doing the loading.
 	 * @param SavedDelegate		Delegate that will be called on game thread when save succeeds or fails.
 	 */
 	static void AsyncSaveGameToSlot(USaveGame* SaveGameObject, const FString& SlotName, const int32 UserIndex, FAsyncSaveGameToSlotDelegate SavedDelegate = FAsyncSaveGameToSlotDelegate());
@@ -1118,7 +1118,7 @@ public:
 	 *
 	 * @param SaveGameObject	Object that contains data about the save game that we want to write out
 	 * @param SlotName			Name of save game slot to save to.
-	 * @param UserIndex			For some platforms, master user index to identify the user doing the saving.
+	 * @param UserIndex			For some platforms, primary user index to identify the user doing the saving.
 	 * @return					Whether we successfully saved this information
 	 */
 	UFUNCTION(BlueprintCallable, Category="SaveGame")
@@ -1127,7 +1127,7 @@ public:
 	/**
 	 * See if a save game exists with the specified name.
 	 * @param SlotName			Name of save game slot.
-	 * @param UserIndex			For some platforms, master user index to identify the user doing the saving.
+	 * @param UserIndex			For some platforms, primary user index to identify the user doing the saving.
 	 */
 	UFUNCTION(BlueprintCallable, Category="SaveGame")
 	static bool DoesSaveGameExist(const FString& SlotName, const int32 UserIndex);
@@ -1143,7 +1143,7 @@ public:
 	 * Load contents from a slot/file into a buffer of save data.
 	 * @param OutSaveData		Data buffer to load into
 	 * @param SlotName			Name of save game slot to save to.
-	 * @param UserIndex			For some platforms, master user index to identify the user doing the saving.
+	 * @param UserIndex			For some platforms, primary user index to identify the user doing the saving.
 	 * @return					Whether valid save data was found and loaded.
 	 */
 	static bool LoadDataFromSlot(TArray<uint8>& OutSaveData, const FString& SlotName, const int32 UserIndex);
@@ -1154,7 +1154,7 @@ public:
 	 * The passed in delegate will be copied to a worker thread so make sure any payload is thread safe to copy by value
 	 *
 	 * @param SlotName			Name of the save game slot to load from.
-	 * @param UserIndex			For some platforms, master user index to identify the user doing the loading.
+	 * @param UserIndex			For some platforms, primary user index to identify the user doing the loading.
 	 * @param LoadedDelegate	Delegate that will be called on game thread when load succeeds or fails.
 	 */
 	static void AsyncLoadGameFromSlot(const FString& SlotName, const int32 UserIndex, FAsyncLoadGameFromSlotDelegate LoadedDelegate);
@@ -1162,7 +1162,7 @@ public:
 	/** 
 	 * Load the contents from a given slot.
 	 * @param SlotName			Name of the save game slot to load from.
-	 * @param UserIndex			For some platforms, master user index to identify the user doing the loading.
+	 * @param UserIndex			For some platforms, primary user index to identify the user doing the loading.
 	 * @return					Object containing loaded game state (nullptr if load fails)
 	 */
 	UFUNCTION(BlueprintCallable, Category="SaveGame")
@@ -1182,7 +1182,7 @@ public:
 	/**
 	 * Delete a save game in a particular slot.
 	 * @param SlotName			Name of save game slot to delete.
-	 * @param UserIndex			For some platforms, master user index to identify the user doing the deletion.
+	 * @param UserIndex			For some platforms, primary user index to identify the user doing the deletion.
 	 * @return					True if a file was actually able to be deleted. use DoesSaveGameExist to distinguish between delete failures and failure due to file not existing.
 	 */
 	UFUNCTION(BlueprintCallable, Category = "SaveGame")
