Removing usage of master/slave in Player Controller

diff --git a/Engine/Source/Runtime/Engine/Private/PlayerController.cpp b/Engine/Source/Runtime/Engine/Private/PlayerController.cpp
--- a/Engine/Source/Runtime/Engine/Private/PlayerController.cpp
+++ b/Engine/Source/Runtime/Engine/Private/PlayerController.cpp
@@ -3875,17 +3875,17 @@ APlayerState* APlayerController::GetSplitscreenPlayerByIndex(int32 PlayerIndex)
 					if ( ChildRemoteConnection != NULL )
 					{
 						// this player controller is not the primary player in the splitscreen layout
-						UNetConnection* MasterConnection = ChildRemoteConnection->Parent;
+						UNetConnection* PrimaryConnection = ChildRemoteConnection->Parent;
 						if ( PlayerIndex == 0 )
 						{
-							Result = MasterConnection->PlayerController->PlayerState;
+							Result = PrimaryConnection->PlayerController->PlayerState;
 						}
 						else
 						{
 							PlayerIndex--;
-							if ( PlayerIndex >= 0 && PlayerIndex < MasterConnection->Children.Num() )
+							if ( PlayerIndex >= 0 && PlayerIndex < PrimaryConnection->Children.Num() )
 							{
-								ChildRemoteConnection = MasterConnection->Children[PlayerIndex];
+								ChildRemoteConnection = PrimaryConnection->Children[PlayerIndex];
 								Result = ChildRemoteConnection->PlayerController->PlayerState;
 							}
 						}
