PrecomputedVolumetricLightmap: Split CopyDataIntoAtlas into batches of GRHIMaxDispatchThreadGroupsPerDimension

diff --git a/Engine/Source/Runtime/Engine/Private/PrecomputedVolumetricLightmap.cpp b/Engine/Source/Runtime/Engine/Private/PrecomputedVolumetricLightmap.cpp
--- a/Engine/Source/Runtime/Engine/Private/PrecomputedVolumetricLightmap.cpp
+++ b/Engine/Source/Runtime/Engine/Private/PrecomputedVolumetricLightmap.cpp
@@ -846,41 +846,44 @@ void CopyDataIntoAtlas(FRHICommandList& RHICmdList, ERHIFeatureLevel::Type Featu
 
 	FGlobalShaderMap* GlobalShaderMap = GetGlobalShaderMap(FeatureLevel);
 
+	for (int32 BatchOffset = 0; BatchOffset < NumBricks; BatchOffset += GRHIMaxDispatchThreadGroupsPerDimension.X)
 	{
-		FCopyResidentBricksCS::FPermutationDomain PermutationVector;
-		PermutationVector.Set<FCopyResidentBricksCS::FHasSkyBentNormal>(SrcData.SkyBentNormal.Texture.IsValid());
+		{
+			FCopyResidentBricksCS::FPermutationDomain PermutationVector;
+			PermutationVector.Set<FCopyResidentBricksCS::FHasSkyBentNormal>(SrcData.SkyBentNormal.Texture.IsValid());
 
-		TShaderMapRef<FCopyResidentBricksCS> ComputeShader(GlobalShaderMap, PermutationVector);
+			TShaderMapRef<FCopyResidentBricksCS> ComputeShader(GlobalShaderMap, PermutationVector);
 
-		FCopyResidentBricksCS::FParameters Parameters;
+			FCopyResidentBricksCS::FParameters Parameters;
 
-		Parameters.StartPosInOldVolume = SrcOffset;
-		Parameters.StartPosInNewVolume = DestOffset;
+			Parameters.StartPosInOldVolume = BatchOffset + SrcOffset;
+			Parameters.StartPosInNewVolume = BatchOffset + DestOffset;
 
-		Parameters.AmbientVector = SrcData.AmbientVector.Texture;
-		Parameters.SkyBentNormal = SrcData.SkyBentNormal.Texture;
-		Parameters.DirectionalLightShadowing = SrcData.DirectionalLightShadowing.Texture;
+			Parameters.AmbientVector = SrcData.AmbientVector.Texture;
+			Parameters.SkyBentNormal = SrcData.SkyBentNormal.Texture;
+			Parameters.DirectionalLightShadowing = SrcData.DirectionalLightShadowing.Texture;
 
-		Parameters.OutAmbientVector = DestTextureSet.AmbientVector.UAV;
-		Parameters.OutSkyBentNormal = DestTextureSet.SkyBentNormal.UAV;
-		Parameters.OutDirectionalLightShadowing = DestTextureSet.DirectionalLightShadowing.UAV;
+			Parameters.OutAmbientVector = DestTextureSet.AmbientVector.UAV;
+			Parameters.OutSkyBentNormal = DestTextureSet.SkyBentNormal.UAV;
+			Parameters.OutDirectionalLightShadowing = DestTextureSet.DirectionalLightShadowing.UAV;
 
-		FComputeShaderUtils::Dispatch(RHICmdList, ComputeShader, Parameters, FIntVector(NumBricks, 1, 1));
-	}
+			FComputeShaderUtils::Dispatch(RHICmdList, ComputeShader, Parameters, FIntVector(FMath::Min(NumBricks, GRHIMaxDispatchThreadGroupsPerDimension.X), 1, 1));
+		}
 
-	for (int32 i = 0; i < UE_ARRAY_COUNT(SrcData.SHCoefficients); i++)
-	{
-		TShaderMapRef<FCopyResidentBrickSHCoefficientsCS> ComputeShader(GlobalShaderMap);
+		for (int32 i = 0; i < UE_ARRAY_COUNT(SrcData.SHCoefficients); i++)
+		{
+			TShaderMapRef<FCopyResidentBrickSHCoefficientsCS> ComputeShader(GlobalShaderMap);
 
-		FCopyResidentBrickSHCoefficientsCS::FParameters Parameters;
+			FCopyResidentBrickSHCoefficientsCS::FParameters Parameters;
 
-		Parameters.StartPosInOldVolume = SrcOffset;
-		Parameters.StartPosInNewVolume = DestOffset;
+			Parameters.StartPosInOldVolume = BatchOffset + SrcOffset;
+			Parameters.StartPosInNewVolume = BatchOffset + DestOffset;
 
-		Parameters.SHCoefficients = SrcData.SHCoefficients[i].Texture;
-		Parameters.OutSHCoefficients = DestTextureSet.SHCoefficients[i].UAV;
+			Parameters.SHCoefficients = SrcData.SHCoefficients[i].Texture;
+			Parameters.OutSHCoefficients = DestTextureSet.SHCoefficients[i].UAV;
 
-		FComputeShaderUtils::Dispatch(RHICmdList, ComputeShader, Parameters, FIntVector(NumBricks, 1, 1));
+			FComputeShaderUtils::Dispatch(RHICmdList, ComputeShader, Parameters, FIntVector(FMath::Min(NumBricks, GRHIMaxDispatchThreadGroupsPerDimension.X), 1, 1));
+		}
 	}
 
 	// Make all the resources readable again
