Split FActorTransactionAnnotation from Actor.h

#rnx

diff --git a/Engine/Source/Editor/UnrealEd/Private/Kismet2/KismetReinstanceUtilities.cpp b/Engine/Source/Editor/UnrealEd/Private/Kismet2/KismetReinstanceUtilities.cpp
--- a/Engine/Source/Editor/UnrealEd/Private/Kismet2/KismetReinstanceUtilities.cpp
+++ b/Engine/Source/Editor/UnrealEd/Private/Kismet2/KismetReinstanceUtilities.cpp
@@ -8,6 +8,7 @@
 #include "UObject/Package.h"
 #include "Components/SceneComponent.h"
 #include "GameFramework/Actor.h"
+#include "ActorTransactionAnnotation.h"
 #include "Engine/World.h"
 #include "Components/SkeletalMeshComponent.h"
 #include "Components/ChildActorComponent.h"
@@ -1191,7 +1192,7 @@ struct FActorReplacementHelper
 		, TargetWorldTransform(FTransform::Identity)
 		, AttachmentData( MoveTemp(InAttachmentData) )
 	{
-		CachedActorData = StaticCastSharedPtr<AActor::FActorTransactionAnnotation>(OldActor->FindOrCreateTransactionAnnotation());
+		CachedActorData = StaticCastSharedPtr<FActorTransactionAnnotation>(OldActor->FindOrCreateTransactionAnnotation());
 		TArray<AActor*> AttachedActors;
 		OldActor->GetAttachedActors(AttachedActors);
 
@@ -1256,7 +1257,7 @@ private:
 	FActorAttachmentData AttachmentData;
 
 	/** Holds actor component data, etc. that we use to apply */
-	TSharedPtr<AActor::FActorTransactionAnnotation> CachedActorData;
+	TSharedPtr<FActorTransactionAnnotation> CachedActorData;
 
 	TMap<FName, UActorComponent*> OldActorComponentNameMap;
 };
diff --git a/Engine/Source/Runtime/Engine/Classes/GameFramework/Actor.h b/Engine/Source/Runtime/Engine/Classes/GameFramework/Actor.h
--- a/Engine/Source/Runtime/Engine/Classes/GameFramework/Actor.h
+++ b/Engine/Source/Runtime/Engine/Classes/GameFramework/Actor.h
@@ -13,7 +13,7 @@
 #include "Engine/EngineTypes.h"
 #include "Engine/EngineBaseTypes.h"
 #include "PropertyPairsMap.h"
-#include "ComponentInstanceDataCache.h"
+#include "ActorTransactionAnnotation.h"
 #include "Components/ChildActorComponent.h"
 #include "RenderCommandFence.h"
 #include "Engine/Level.h"
@@ -22,7 +22,6 @@
 #include "Engine/ReplicatedState.h"
 
 #if WITH_EDITOR
-#include "Misc/ITransactionObjectAnnotation.h"
 #include "WorldPartition/DataLayer/ActorDataLayer.h"
 #include "Folder.h"
 #endif
@@ -42,6 +41,8 @@ class UPrimitiveComponent;
 struct FAttachedActorInfo;
 struct FNetViewer;
 struct FNetworkObjectInfo;
+class FActorTransactionAnnotation;
+class FComponentInstanceDataCache;
 class UDEPRECATED_DataLayer;
 class UDataLayerAsset;
 class UDataLayerInstance;
@@ -2060,74 +2061,10 @@ public:
 	virtual bool SupportsExternalPackaging() const;
 #endif
 
-	/** Internal struct used to store information about an actor's components during reconstruction */
-	struct FActorRootComponentReconstructionData
-	{
-		/** Struct to store info about attached actors */
-		struct FAttachedActorInfo
-		{
-			TWeakObjectPtr<AActor> Actor;
-			TWeakObjectPtr<USceneComponent> AttachParent;
-			FName AttachParentName;
-			FName SocketName;
-			FTransform RelativeTransform;
-
-			friend FArchive& operator<<(FArchive& Ar, FAttachedActorInfo& ActorInfo);
-		};
-
-		/** The RootComponent's transform */
-		FTransform Transform;
-
-		/** The RootComponent's relative rotation cache (enforces using the same rotator) */
-		FRotationConversionCache TransformRotationCache;
-
-		/** The Actor the RootComponent is attached to */
-		FAttachedActorInfo AttachedParentInfo;
-
-		/** Actors that are attached to this RootComponent */
-		TArray<FAttachedActorInfo> AttachedToInfo;
-
-		friend FArchive& operator<<(FArchive& Ar, FActorRootComponentReconstructionData& RootComponentData);
-	};
-
-	struct FActorTransactionAnnotationData
-	{
-		TWeakObjectPtr<const AActor> Actor;
-		FComponentInstanceDataCache ComponentInstanceData;
-
-		bool bRootComponentDataCached;
-		FActorRootComponentReconstructionData RootComponentData;
-
-		friend ENGINE_API FArchive& operator<<(FArchive& Ar, FActorTransactionAnnotationData& ActorTransactionAnnotationData);
-	};
+	using FActorTransactionAnnotationData = FActorTransactionAnnotationData;
+	using FActorRootComponentReconstructionData = FActorRootComponentReconstructionData;
 
 #if WITH_EDITOR
-	/** Internal struct to track currently active transactions */
-	class FActorTransactionAnnotation : public ITransactionObjectAnnotation
-	{
-	public:
-		/** Create an empty instance */
-		static TSharedRef<FActorTransactionAnnotation> Create();
-
-		/** Create an instance from the given actor, optionally caching root component data */
-		static TSharedRef<FActorTransactionAnnotation> Create(const AActor* InActor, const bool InCacheRootComponentData = true);
-
-		/** Create an instance from the given actor if required (UActorTransactionAnnotation::HasInstanceData would return true), optionally caching root component data */
-		static TSharedPtr<FActorTransactionAnnotation> CreateIfRequired(const AActor* InActor, const bool InCacheRootComponentData = true);
-
-		//~ ITransactionObjectAnnotation interface
-		virtual void AddReferencedObjects(FReferenceCollector& Collector) override;
-		virtual void Serialize(FArchive& Ar) override;
-
-		bool HasInstanceData() const;
-
-		FActorTransactionAnnotationData ActorTransactionAnnotationData;
-
-	private:
-		FActorTransactionAnnotation();
-		FActorTransactionAnnotation(const AActor* InActor, FComponentInstanceDataCache&& InComponentInstanceData, const bool InCacheRootComponentData = true);
-	};
-
 	/** Cached pointer to the transaction annotation data from PostEditUndo to be used in the next RerunConstructionScript */
 	TSharedPtr<FActorTransactionAnnotation> CurrentTransactionAnnotation;
 
diff --git a/Engine/Source/Runtime/Engine/Private/Actor.cpp b/Engine/Source/Runtime/Engine/Private/Actor.cpp
--- a/Engine/Source/Runtime/Engine/Private/Actor.cpp
+++ b/Engine/Source/Runtime/Engine/Private/Actor.cpp
@@ -2,6 +2,7 @@
 
 #include "GameFramework/Actor.h"
 
+#include "ActorTransactionAnnotation.h"
 #include "EngineDefines.h"
 #include "EngineStats.h"
 #include "EngineGlobals.h"
@@ -5941,7 +5942,7 @@ void AActor::SetInstigator(APawn* InInstigator)
 TFunction<bool(const AActor*)> GIsActorSelectedInEditor;
 #endif
 
-FArchive& operator<<(FArchive& Ar, AActor::FActorRootComponentReconstructionData::FAttachedActorInfo& ActorInfo)
+FArchive& operator<<(FArchive& Ar, FActorRootComponentReconstructionData::FAttachedActorInfo& ActorInfo)
 {
 	enum class EVersion : uint8
 	{
@@ -5969,7 +5970,7 @@ FArchive& operator<<(FArchive& Ar, AActor::FActorRootComponentReconstructionData
 	return Ar;
 }
 
-FArchive& operator<<(FArchive& Ar, AActor::FActorRootComponentReconstructionData& RootComponentData)
+FArchive& operator<<(FArchive& Ar, FActorRootComponentReconstructionData& RootComponentData)
 {
 	enum class EVersion : uint8
 	{
@@ -6009,7 +6010,7 @@ FArchive& operator<<(FArchive& Ar, AActor::FActorRootComponentReconstructionData
 	return Ar;
 }
 
-FArchive& operator<<(FArchive& Ar, AActor::FActorTransactionAnnotationData& ActorTransactionAnnotationData)
+FArchive& operator<<(FArchive& Ar, FActorTransactionAnnotationData& ActorTransactionAnnotationData)
 {
 	enum class EVersion : uint8
 	{
diff --git a/Engine/Source/Runtime/Engine/Private/ActorConstruction.cpp b/Engine/Source/Runtime/Engine/Private/ActorConstruction.cpp
--- a/Engine/Source/Runtime/Engine/Private/ActorConstruction.cpp
+++ b/Engine/Source/Runtime/Engine/Private/ActorConstruction.cpp
@@ -17,6 +17,7 @@
 #include "Components/ActorComponent.h"
 #include "Components/SceneComponent.h"
 #include "GameFramework/Actor.h"
+#include "ActorTransactionAnnotation.h"
 #include "Components/PrimitiveComponent.h"
 #include "Components/StaticMeshComponent.h"
 #include "Components/BillboardComponent.h"
diff --git a/Engine/Source/Runtime/Engine/Private/ActorEditor.cpp b/Engine/Source/Runtime/Engine/Private/ActorEditor.cpp
--- a/Engine/Source/Runtime/Engine/Private/ActorEditor.cpp
+++ b/Engine/Source/Runtime/Engine/Private/ActorEditor.cpp
@@ -30,6 +30,7 @@
 
 #include "Editor.h"
 #include "Misc/TransactionObjectEvent.h"
+#include "ActorTransactionAnnotation.h"
 #include "Engine/LevelStreaming.h"
 #include "WorldPartition/WorldPartitionActorDesc.h"
 #include "WorldPartition/DataLayer/IDataLayerEditorModule.h"
@@ -350,17 +351,17 @@ void AActor::DebugShowOneComponentHierarchy( USceneComponent* SceneComp, int32&
 	}
 }
 
-TSharedRef<AActor::FActorTransactionAnnotation> AActor::FActorTransactionAnnotation::Create()
+TSharedRef<FActorTransactionAnnotation> FActorTransactionAnnotation::Create()
 {
 	return MakeShareable(new FActorTransactionAnnotation());
 }
 
-TSharedRef<AActor::FActorTransactionAnnotation> AActor::FActorTransactionAnnotation::Create(const AActor* InActor, const bool InCacheRootComponentData)
+TSharedRef<FActorTransactionAnnotation> FActorTransactionAnnotation::Create(const AActor* InActor, const bool InCacheRootComponentData)
 {
 	return MakeShareable(new FActorTransactionAnnotation(InActor, FComponentInstanceDataCache(InActor), InCacheRootComponentData));
 }
 
-TSharedPtr<AActor::FActorTransactionAnnotation> AActor::FActorTransactionAnnotation::CreateIfRequired(const AActor* InActor, const bool InCacheRootComponentData)
+TSharedPtr<FActorTransactionAnnotation> FActorTransactionAnnotation::CreateIfRequired(const AActor* InActor, const bool InCacheRootComponentData)
 {
 	// Don't create a transaction annotation for something that has no instance data, or a root component that's created by a construction script
 	FComponentInstanceDataCache TempComponentInstanceData(InActor);
@@ -376,12 +377,12 @@ TSharedPtr<AActor::FActorTransactionAnnotation> AActor::FActorTransactionAnnotat
 	return MakeShareable(new FActorTransactionAnnotation(InActor, MoveTemp(TempComponentInstanceData), InCacheRootComponentData));
 }
 
-AActor::FActorTransactionAnnotation::FActorTransactionAnnotation()
+FActorTransactionAnnotation::FActorTransactionAnnotation()
 {
 	ActorTransactionAnnotationData.bRootComponentDataCached = false;
 }
 
-AActor::FActorTransactionAnnotation::FActorTransactionAnnotation(const AActor* InActor, FComponentInstanceDataCache&& InComponentInstanceData, const bool InCacheRootComponentData)
+FActorTransactionAnnotation::FActorTransactionAnnotation(const AActor* InActor, FComponentInstanceDataCache&& InComponentInstanceData, const bool InCacheRootComponentData)
 {
 	ActorTransactionAnnotationData.ComponentInstanceData = MoveTemp(InComponentInstanceData);
 	ActorTransactionAnnotationData.Actor = InActor;
@@ -424,17 +425,17 @@ AActor::FActorTransactionAnnotation::FActorTransactionAnnotation(const AActor* I
 	}
 }
 
-void AActor::FActorTransactionAnnotation::AddReferencedObjects(FReferenceCollector& Collector)
+void FActorTransactionAnnotation::AddReferencedObjects(FReferenceCollector& Collector)
 {
 	ActorTransactionAnnotationData.ComponentInstanceData.AddReferencedObjects(Collector);
 }
 
-void AActor::FActorTransactionAnnotation::Serialize(FArchive& Ar)
+void FActorTransactionAnnotation::Serialize(FArchive& Ar)
 {
 	Ar << ActorTransactionAnnotationData;
 }
 
-bool AActor::FActorTransactionAnnotation::HasInstanceData() const
+bool FActorTransactionAnnotation::HasInstanceData() const
 {
 	return (ActorTransactionAnnotationData.bRootComponentDataCached || ActorTransactionAnnotationData.ComponentInstanceData.HasInstanceData());
 }
diff --git a/Engine/Source/Runtime/Engine/Public/ActorTransactionAnnotation.h b/Engine/Source/Runtime/Engine/Public/ActorTransactionAnnotation.h
new file mode 100644
--- /dev/null
+++ b/Engine/Source/Runtime/Engine/Public/ActorTransactionAnnotation.h
@@ -0,0 +1,83 @@
+// Copyright Epic Games, Inc. All Rights Reserved.
+
+#pragma once
+
+#include "CoreMinimal.h"
+#include "Engine/EngineTypes.h"
+#include "ComponentInstanceDataCache.h"
+
+#if WITH_EDITOR
+#include "Misc/ITransactionObjectAnnotation.h"
+#endif
+
+class AActor;
+class USceneComponent;
+
+/** Internal struct used to store information about an actor's components during reconstruction */
+struct FActorRootComponentReconstructionData
+{
+	/** Struct to store info about attached actors */
+	struct FAttachedActorInfo
+	{
+		TWeakObjectPtr<AActor> Actor;
+		TWeakObjectPtr<USceneComponent> AttachParent;
+		FName AttachParentName;
+		FName SocketName;
+		FTransform RelativeTransform;
+
+		friend FArchive& operator<<(FArchive& Ar, FAttachedActorInfo& ActorInfo);
+	};
+
+	/** The RootComponent's transform */
+	FTransform Transform;
+
+	/** The RootComponent's relative rotation cache (enforces using the same rotator) */
+	FRotationConversionCache TransformRotationCache;
+
+	/** The Actor the RootComponent is attached to */
+	FAttachedActorInfo AttachedParentInfo;
+
+	/** Actors that are attached to this RootComponent */
+	TArray<FAttachedActorInfo> AttachedToInfo;
+
+	friend FArchive& operator<<(FArchive& Ar, FActorRootComponentReconstructionData& RootComponentData);
+};
+
+struct FActorTransactionAnnotationData
+{
+	TWeakObjectPtr<const AActor> Actor;
+	FComponentInstanceDataCache ComponentInstanceData;
+
+	bool bRootComponentDataCached;
+	FActorRootComponentReconstructionData RootComponentData;
+
+	friend ENGINE_API FArchive& operator<<(FArchive& Ar, FActorTransactionAnnotationData& ActorTransactionAnnotationData);
+};
+
+#if WITH_EDITOR
+/** Internal struct to track currently active transactions */
+class FActorTransactionAnnotation : public ITransactionObjectAnnotation
+{
+public:
+	/** Create an empty instance */
+	static TSharedRef<FActorTransactionAnnotation> Create();
+
+	/** Create an instance from the given actor, optionally caching root component data */
+	static TSharedRef<FActorTransactionAnnotation> Create(const AActor* InActor, const bool InCacheRootComponentData = true);
+
+	/** Create an instance from the given actor if required (UActorTransactionAnnotation::HasInstanceData would return true), optionally caching root component data */
+	static TSharedPtr<FActorTransactionAnnotation> CreateIfRequired(const AActor* InActor, const bool InCacheRootComponentData = true);
+
+	//~ ITransactionObjectAnnotation interface
+	virtual void AddReferencedObjects(FReferenceCollector& Collector) override;
+	virtual void Serialize(FArchive& Ar) override;
+
+	bool HasInstanceData() const;
+
+	FActorTransactionAnnotationData ActorTransactionAnnotationData;
+
+private:
+	FActorTransactionAnnotation();
+	FActorTransactionAnnotation(const AActor* InActor, FComponentInstanceDataCache&& InComponentInstanceData, const bool InCacheRootComponentData = true);
+};
+#endif // WITH_EDITOR
