Fix non-unity errors
#rnx

diff --git a/Engine/Source/Developer/CollisionAnalyzer/Private/CollisionAnalyzer.h b/Engine/Source/Developer/CollisionAnalyzer/Private/CollisionAnalyzer.h
--- a/Engine/Source/Developer/CollisionAnalyzer/Private/CollisionAnalyzer.h
+++ b/Engine/Source/Developer/CollisionAnalyzer/Private/CollisionAnalyzer.h
@@ -2,6 +2,8 @@
 
 #pragma once
 
+#include "CollisionDebugDrawingPublic.h"
+
 #include "CoreMinimal.h"
 #include "Engine/EngineTypes.h"
 #include "CollisionQueryParams.h"
diff --git a/Engine/Source/Runtime/Engine/Public/CollisionDebugDrawingPublic.h b/Engine/Source/Runtime/Engine/Public/CollisionDebugDrawingPublic.h
--- a/Engine/Source/Runtime/Engine/Public/CollisionDebugDrawingPublic.h
+++ b/Engine/Source/Runtime/Engine/Public/CollisionDebugDrawingPublic.h
@@ -5,6 +5,8 @@
 #include "CoreMinimal.h"
 #include "Engine/EngineTypes.h"
 
+struct FHitResult;
+
 /** Draw line trace results */
 ENGINE_API void DrawLineTraces(const UWorld* InWorld, const FVector& Start, const FVector& End, const TArray<FHitResult>& Hits, float Lifetime);
 /** Draw sphere sweep results */
