reset seek variables on PCM/ADPCM codecs when reading compressed data - fix for #jira UE-159441 #preflight 62f2a112bc175ec68c0b8d2a

diff --git a/Engine/Source/Runtime/Engine/Private/ADPCMAudioInfo.cpp b/Engine/Source/Runtime/Engine/Private/ADPCMAudioInfo.cpp
--- a/Engine/Source/Runtime/Engine/Private/ADPCMAudioInfo.cpp
+++ b/Engine/Source/Runtime/Engine/Private/ADPCMAudioInfo.cpp
@@ -110,8 +110,7 @@ void FADPCMAudioInfo::SeekToTimeInternal(const float InSeekTime)
 		CurrentChunkBufferOffset = FirstChunkSampleDataOffset;
 		TotalSamplesStreamed = 0;
 		
-		bNewSeekRequest = false;
-		bSeekPendingRead = false;
+		ResetSeekState();
 		return;
 	}
 
@@ -383,6 +382,9 @@ bool FADPCMAudioInfo::ReadCompressedData(uint8* Destination, bool bLooping, uint
 				{
 					// Zero remaining buffer
 					FMemory::Memzero(OutData, BufferSize);
+					
+					ResetSeekState();
+					
 					return true;
 				}
 				else
@@ -422,12 +424,17 @@ bool FADPCMAudioInfo::ReadCompressedData(uint8* Destination, bool bLooping, uint
 				{
 					// Zero remaining buffer
 					FMemory::Memzero(OutData, BufferSize);
+
+					ResetSeekState();
+
 					return true;
 				}
 			}
 		}
 	}
 
+	ResetSeekState();
+
 	return ReachedEndOfSamples;
 }
 
@@ -458,6 +465,12 @@ void FADPCMAudioInfo::ProcessSeekRequest()
 	}
 }
 
+void FADPCMAudioInfo::ResetSeekState()
+{
+	bSeekPendingRead = false;
+	bNewSeekRequest = false;
+}
+
 bool FADPCMAudioInfo::StreamCompressedInfoInternal(const FSoundWaveProxyPtr& InWaveProxy, struct FSoundQualityInfo* QualityInfo)
 {
 	FScopeLock ScopeLock(&CurCompressedChunkHandleCriticalSection);
@@ -774,8 +787,7 @@ bool FADPCMAudioInfo::StreamCompressedData(uint8* Destination, bool bLooping, ui
 						CurrentChunkBufferOffset = FirstChunkSampleDataOffset;
 					}
 
-					bSeekPendingRead = false;
-					bNewSeekRequest = false;
+					ResetSeekState();
 				}
 
 				// Decompress one block for each channel and store it in UncompressedBlockData
@@ -901,8 +913,7 @@ bool FADPCMAudioInfo::StreamCompressedData(uint8* Destination, bool bLooping, ui
 					CurrentChunkBufferOffset = CurrentChunkIndex == FirstChunkSampleDataIndex ? FirstChunkSampleDataOffset : 0;
 				}
 
-				bSeekPendingRead = false;
-				bNewSeekRequest = false;
+				ResetSeekState();
 			}
 			
 			uint32 DecompressedSamplesToCopy = FMath::Min<uint32>(
@@ -953,6 +964,8 @@ bool FADPCMAudioInfo::StreamCompressedData(uint8* Destination, bool bLooping, ui
 
 bool FADPCMAudioInfo::ReleaseStreamChunk(bool bBlockUntilReleased)
 {
+	ResetSeekState();
+
 	if (bBlockUntilReleased)
 	{
 		// Wait for any pending decode tasks to finish.
diff --git a/Engine/Source/Runtime/Engine/Public/ADPCMAudioInfo.h b/Engine/Source/Runtime/Engine/Public/ADPCMAudioInfo.h
--- a/Engine/Source/Runtime/Engine/Public/ADPCMAudioInfo.h
+++ b/Engine/Source/Runtime/Engine/Public/ADPCMAudioInfo.h
@@ -180,6 +180,7 @@ private:
 private:
 	void ProcessSeekRequest();
 	void SeekToTimeInternal(const float InSeekTime);
+	void ResetSeekState();
 
 	std::atomic<float> TargetSeekTime;
 
