Deprecate and remove the usage of FUntypedBulkData

diff --git a/Engine/Plugins/Enterprise/LidarPointCloud/Source/LidarPointCloudRuntime/Private/LidarPointCloudOctree.cpp b/Engine/Plugins/Enterprise/LidarPointCloud/Source/LidarPointCloudRuntime/Private/LidarPointCloudOctree.cpp
--- a/Engine/Plugins/Enterprise/LidarPointCloud/Source/LidarPointCloudRuntime/Private/LidarPointCloudOctree.cpp
+++ b/Engine/Plugins/Enterprise/LidarPointCloud/Source/LidarPointCloudRuntime/Private/LidarPointCloudOctree.cpp
@@ -35,17 +35,28 @@ FArchive& operator<<(FArchive& Ar, FLidarPointCloudPoint_Legacy& P)
 	return Ar;
 }
 
-struct FLidarPointCloudBulkData_Legacy : public FUntypedBulkData
+struct FLidarPointCloudBulkData_Legacy : public FBulkData
 {
 private:
 	int32 ElementSize;
 
 public:
 	FLidarPointCloudBulkData_Legacy(int32 ElementSize) : ElementSize(ElementSize) {}
-	virtual int32 GetElementSize() const override { return ElementSize; }
-	virtual void SerializeElement(FArchive& Ar, void* Data, int64 ElementIndex) override { }
+	void Serialize(FArchive& Ar, UObject* Owner);
+	int32 GetElementCount() const;
 };
 
+void FLidarPointCloudBulkData_Legacy::Serialize(FArchive& Ar, UObject* Owner)
+{
+	const bool bAttemptFileMapping = false;
+	FBulkData::Serialize(Ar, Owner, bAttemptFileMapping, ElementSize, EFileRegionType::None);
+}
+
+int32 FLidarPointCloudBulkData_Legacy::GetElementCount() const
+{
+	return GetBulkDataSize() / ElementSize;
+}
+
 /** Used for grid allocation calculations */
 struct FGridAllocation
 {
diff --git a/Engine/Source/Editor/UnrealEd/Classes/Factories/Factory.h b/Engine/Source/Editor/UnrealEd/Classes/Factories/Factory.h
--- a/Engine/Source/Editor/UnrealEd/Classes/Factories/Factory.h
+++ b/Engine/Source/Editor/UnrealEd/Classes/Factories/Factory.h
@@ -10,7 +10,7 @@
 #include "Factory.generated.h"
 
 class UAssetImportTask;
-struct FUntypedBulkData;
+class FBulkData;
 
 /**
  * Base class for all factories
@@ -174,7 +174,7 @@ public:
 	}
 
 	// @todo document
-	virtual bool ImportUntypedBulkDataFromText(const TCHAR*& Buffer, FUntypedBulkData& BulkData);
+	virtual bool ImportUntypedBulkDataFromText(const TCHAR*& Buffer, FBulkData& BulkData);
 
 	/** Creates a list of file extensions supported by this factory */
 	virtual void GetSupportedFileExtensions(TArray<FString>& OutExtensions) const;
diff --git a/Engine/Source/Editor/UnrealEd/Private/Factories/Factory.cpp b/Engine/Source/Editor/UnrealEd/Private/Factories/Factory.cpp
--- a/Engine/Source/Editor/UnrealEd/Private/Factories/Factory.cpp
+++ b/Engine/Source/Editor/UnrealEd/Private/Factories/Factory.cpp
@@ -528,7 +528,7 @@ bool UFactory::IsSupportedFileExtension(FStringView InExtension) const
 }
 
 
-bool UFactory::ImportUntypedBulkDataFromText(const TCHAR*& Buffer, FUntypedBulkData& BulkData)
+bool UFactory::ImportUntypedBulkDataFromText(const TCHAR*& Buffer, FBulkData& BulkData)
 {
 	FString StrLine;
 	int32 ElementCount = 0;
@@ -566,7 +566,7 @@ bool UFactory::ImportUntypedBulkDataFromText(const TCHAR*& Buffer, FUntypedBulkD
 					check(Size == (ElementSize *ElementCount));
 
 					BulkData.Lock(LOCK_READ_WRITE);
-					void* RawBulkData = BulkData.Realloc(ElementCount);
+					void* RawBulkData = BulkData.Realloc(ElementCount, ElementSize);
 					RawData = (uint8*)RawBulkData;
 					bBulkDataIsLocked = true;
 				}
diff --git a/Engine/Source/Runtime/CoreUObject/Private/Serialization/EditorBulkData.cpp b/Engine/Source/Runtime/CoreUObject/Private/Serialization/EditorBulkData.cpp
--- a/Engine/Source/Runtime/CoreUObject/Private/Serialization/EditorBulkData.cpp
+++ b/Engine/Source/Runtime/CoreUObject/Private/Serialization/EditorBulkData.cpp
@@ -734,7 +734,7 @@ void FEditorBulkData::Serialize(FArchive& Ar, UObject* Owner, bool bAllowRegiste
 		// We've made changes to this serialization path without versioning them, because we assume
 		// IsTransacting means !IsPersistent and therefore all loaded data was saved with the current binary.
 		check(!Ar.IsPersistent()); 
-		// Do not process the transaction if the owner is mid loading (see FUntypedBulkData::Serialize)
+		// Do not process the transaction if the owner is mid loading (see FBulkData::Serialize)
 		bool bNeedsTransaction = Ar.IsSaving() && (!Owner || !Owner->HasAnyFlags(RF_NeedLoad));
 
 		Ar << bNeedsTransaction;
diff --git a/Engine/Source/Runtime/CoreUObject/Public/Serialization/BulkData.h b/Engine/Source/Runtime/CoreUObject/Public/Serialization/BulkData.h
--- a/Engine/Source/Runtime/CoreUObject/Public/Serialization/BulkData.h
+++ b/Engine/Source/Runtime/CoreUObject/Public/Serialization/BulkData.h
@@ -1212,6 +1212,7 @@ public:
 	}
 };
 
+UE_DEPRECATED(5.1, "Use FBulkData/TBulkData");
 struct COREUOBJECT_API FUntypedBulkData : public FBulkData
 {
 	virtual ~FUntypedBulkData() = default;
diff --git a/Engine/Source/Runtime/CoreUObject/Public/Serialization/EditorBulkData.h b/Engine/Source/Runtime/CoreUObject/Public/Serialization/EditorBulkData.h
--- a/Engine/Source/Runtime/CoreUObject/Public/Serialization/EditorBulkData.h
+++ b/Engine/Source/Runtime/CoreUObject/Public/Serialization/EditorBulkData.h
@@ -119,7 +119,7 @@ public:
 	~FEditorBulkData();
 
 	/** 
-	 * Convenience method to make it easier to convert from FUntypedBulkData to FEditorBulkData and sets the Guid 
+	 * Convenience method to make it easier to convert from FBulkData to FEditorBulkData and sets the Guid 
 	 *
 	 * @param BulkData	The bulkdata object to create from.
 	 * @param Guid		A guid associated with the bulkdata object which will be used to identify the payload.
diff --git a/Engine/Source/Runtime/Engine/Private/LightMap.cpp b/Engine/Source/Runtime/Engine/Private/LightMap.cpp
--- a/Engine/Source/Runtime/Engine/Private/LightMap.cpp
+++ b/Engine/Source/Runtime/Engine/Private/LightMap.cpp
@@ -3262,53 +3262,6 @@ void FLegacyLightMap1D::Serialize(FArchive& Ar)
 	SimpleSamples.Serialize( Ar, Owner, INDEX_NONE, false );
 }
 
-/*-----------------------------------------------------------------------------
-	FQuantizedLightSample version of bulk data.
------------------------------------------------------------------------------*/
-
-/**
- * Returns whether single element serialization is required given an archive. This e.g.
- * can be the case if the serialization for an element changes and the single element
- * serialization code handles backward compatibility.
- */
-template<class QuantizedLightSampleType>
-bool TQuantizedLightSampleBulkData<QuantizedLightSampleType>::RequiresSingleElementSerialization( FArchive& Ar )
-{
-	return false;
-}
-
-/**
- * Returns size in bytes of single element.
- *
- * @return Size in bytes of single element
- */
-template<class QuantizedLightSampleType>
-int32 TQuantizedLightSampleBulkData<QuantizedLightSampleType>::GetElementSize() const
-{
-	return sizeof(QuantizedLightSampleType);
-}
-
-/**
- * Serializes an element at a time allowing and dealing with endian conversion and backward compatiblity.
- * 
- * @param Ar			Archive to serialize with
- * @param Data			Base pointer to data
- * @param ElementIndex	Element index to serialize
- */
-template<class QuantizedLightSampleType>
-void TQuantizedLightSampleBulkData<QuantizedLightSampleType>::SerializeElement( FArchive& Ar, void* Data, int64 ElementIndex )
-{
-	QuantizedLightSampleType* QuantizedLightSample = (QuantizedLightSampleType*)Data + ElementIndex;
-	// serialize as colors
-	const uint32 NumCoefficients = sizeof(QuantizedLightSampleType) / sizeof(FColor);
-	for(int32 CoefficientIndex = 0; CoefficientIndex < NumCoefficients; CoefficientIndex++)
-	{
-		uint32 ColorDWORD = QuantizedLightSample->Coefficients[CoefficientIndex].DWColor();
-		Ar << ColorDWORD;
-		QuantizedLightSample->Coefficients[CoefficientIndex] = FColor(ColorDWORD);
-	} 
-};
-
 FArchive& operator<<(FArchive& Ar, FLightMap*& R)
 {
 	uint32 LightMapType = FLightMap::LMT_None;
diff --git a/Engine/Source/Runtime/Engine/Private/SoundWave.cpp b/Engine/Source/Runtime/Engine/Private/SoundWave.cpp
--- a/Engine/Source/Runtime/Engine/Private/SoundWave.cpp
+++ b/Engine/Source/Runtime/Engine/Private/SoundWave.cpp
@@ -1236,36 +1236,6 @@ const FPlatformAudioCookOverrides* USoundWave::GetPlatformCompressionOverridesFo
 	return FPlatformCompressionUtilities::GetCookOverrides();
 }
 
-namespace SoundWavePrivate
-{
-	struct FBulkDataReadScopeLock
-	{
-		FBulkDataReadScopeLock(const FUntypedBulkData& InBulkData)
-		:	BulkData(InBulkData)
-		{
-			RawPtr = BulkData.LockReadOnly();
-		}
-
-		template<typename T>
-		const T* GetData() const
-		{
-			return static_cast<const T*>(RawPtr);
-		}
-			
-		~FBulkDataReadScopeLock()
-		{
-			if (BulkData.IsLocked())
-			{
-				BulkData.Unlock();
-			}
-		}
-
-		private:
-			const FUntypedBulkData& BulkData;
-			const void* RawPtr = nullptr;
-	};
-}
-
 #if WITH_EDITOR
 bool USoundWave::GetImportedSoundWaveData(TArray<uint8>& OutRawPCMData, uint32& OutSampleRate, uint16& OutNumChannels) const
 {
@@ -1287,8 +1257,6 @@ bool USoundWave::GetImportedSoundWaveData(TArray<uint8>& OutRawPCMData, uint32&
 
 bool USoundWave::GetImportedSoundWaveData(TArray<uint8>& OutRawPCMData, uint32& OutSampleRate, TArray<EAudioSpeakers>& OutChannelOrder) const
 {
-	using namespace SoundWavePrivate;
-
 	OutRawPCMData.Reset();
 	OutSampleRate = 0;
 	OutChannelOrder.Reset();
diff --git a/Engine/Source/Runtime/Engine/Private/StaticMeshBuild.cpp b/Engine/Source/Runtime/Engine/Private/StaticMeshBuild.cpp
--- a/Engine/Source/Runtime/Engine/Private/StaticMeshBuild.cpp
+++ b/Engine/Source/Runtime/Engine/Private/StaticMeshBuild.cpp
@@ -848,64 +848,6 @@ struct FStaticMeshTriangle
 	uint32		bExplicitNormals;
 };
 
-struct FStaticMeshTriangleBulkData : public FUntypedBulkData
-{
-	virtual int32 GetElementSize() const
-	{
-		return sizeof(FStaticMeshTriangle);
-	}
-
-	virtual void SerializeElement( FArchive& Ar, void* Data, int64 ElementIndex )
-	{
-		FStaticMeshTriangle& StaticMeshTriangle = *((FStaticMeshTriangle*)Data + ElementIndex);
-		Ar << StaticMeshTriangle.Vertices[0];
-		Ar << StaticMeshTriangle.Vertices[1];
-		Ar << StaticMeshTriangle.Vertices[2];
-		for( int32 VertexIndex=0; VertexIndex<3; VertexIndex++ )
-		{
-			for( int32 UVIndex=0; UVIndex<8; UVIndex++ )
-			{
-				Ar << StaticMeshTriangle.UVs[VertexIndex][UVIndex];
-			}
-        }
-		Ar << StaticMeshTriangle.Colors[0];
-		Ar << StaticMeshTriangle.Colors[1];
-		Ar << StaticMeshTriangle.Colors[2];
-		Ar << StaticMeshTriangle.MaterialIndex;
-		Ar << StaticMeshTriangle.FragmentIndex;
-		Ar << StaticMeshTriangle.SmoothingMask;
-		Ar << StaticMeshTriangle.NumUVs;
-		Ar << StaticMeshTriangle.TangentX[0];
-		Ar << StaticMeshTriangle.TangentX[1];
-		Ar << StaticMeshTriangle.TangentX[2];
-		Ar << StaticMeshTriangle.TangentY[0];
-		Ar << StaticMeshTriangle.TangentY[1];
-		Ar << StaticMeshTriangle.TangentY[2];
-		Ar << StaticMeshTriangle.TangentZ[0];
-		Ar << StaticMeshTriangle.TangentZ[1];
-		Ar << StaticMeshTriangle.TangentZ[2];
-		Ar << StaticMeshTriangle.bOverrideTangentBasis;
-		Ar << StaticMeshTriangle.bExplicitNormals;
-	}
-
-	virtual bool RequiresSingleElementSerialization( FArchive& Ar )
-	{
-		return false;
-	}
-};
-
-struct FFragmentRange
-{
-	int32 BaseIndex;
-	int32 NumPrimitives;
-
-	friend FArchive& operator<<(FArchive& Ar,FFragmentRange& FragmentRange)
-	{
-		Ar << FragmentRange.BaseIndex << FragmentRange.NumPrimitives;
-		return Ar;
-	}
-};
-
 void UStaticMesh::FixupZeroTriangleSections()
 {
 	if (GetRenderData()->MaterialIndexToImportIndex.Num() > 0 && GetRenderData()->LODResources.Num())
diff --git a/Engine/Source/Runtime/Engine/Private/VectorField.cpp b/Engine/Source/Runtime/Engine/Private/VectorField.cpp
--- a/Engine/Source/Runtime/Engine/Private/VectorField.cpp
+++ b/Engine/Source/Runtime/Engine/Private/VectorField.cpp
@@ -412,9 +412,9 @@ void UVectorFieldStatic::UpdateCPUData(bool bDiscardData)
 		if (bDiscardData && SourceData.IsBulkDataLoaded())
 		{
 			// NOTE(mv): This assertion will fail in the case where the bulk data is still available even though the bDiscardInternalCopy
-			//           flag is toggled when FUntypedBulkData::CanLoadFromDisk() also fail. This happens when the user tries to allow 
+			//           flag is toggled when FBulkData::CanLoadFromDisk() also fail. This happens when the user tries to allow 
 			//           CPU access to a newly imported file that isn't reloaded. We still have our valid data, so we just issue a 
-			//           warning and move on. See FUntypedBulkData::GetCopy() for more details. 
+			//           warning and move on. See FBulkData::GetCopy() for more details. 
 			UE_LOG(LogVectorField, Warning, TEXT("SourceData.GetCopy() is supposed to unload the data after copying, but it is still loaded."));
 		}
 
diff --git a/Engine/Source/Runtime/Engine/Public/LightMap.h b/Engine/Source/Runtime/Engine/Public/LightMap.h
--- a/Engine/Source/Runtime/Engine/Public/LightMap.h
+++ b/Engine/Source/Runtime/Engine/Public/LightMap.h
@@ -402,31 +402,9 @@ template <> struct TIsPODType<FQuantizedSimpleLightSample> { enum { Value = true
  * Bulk data array of FQuantizedLightSamples
  */
 template<class QuantizedLightSampleType>
-struct TQuantizedLightSampleBulkData : public FUntypedBulkData
+struct TQuantizedLightSampleBulkData : public TBulkData<QuantizedLightSampleType>
 {
 	typedef QuantizedLightSampleType SampleType;
-	/**
-	 * Returns whether single element serialization is required given an archive. This e.g.
-	 * can be the case if the serialization for an element changes and the single element
-	 * serialization code handles backward compatibility.
-	 */
-	virtual bool RequiresSingleElementSerialization( FArchive& Ar );
-
-	/**
-	 * Returns size in bytes of single element.
-	 *
-	 * @return Size in bytes of single element
-	 */
-	virtual int32 GetElementSize() const;
-
-	/**
-	 * Serializes an element at a time allowing and dealing with endian conversion and backward compatiblity.
-	 * 
-	 * @param Ar			Archive to serialize with
-	 * @param Data			Base pointer to data
-	 * @param ElementIndex	Element index to serialize
-	 */
-	virtual void SerializeElement( FArchive& Ar, void* Data, int64 ElementIndex );
 };
 
 /** A 1D array of incident lighting data. */
@@ -650,4 +628,4 @@ public:
 	FLightmapClusterResourceInput Input;
 	mutable IAllocatedVirtualTexture* AllocatedVT;
 	FUniformBufferRHIRef UniformBuffer;
-};
\ No newline at end of file
+};
diff --git a/Engine/Source/Runtime/Engine/Public/Physics/Experimental/ChaosDerivedDataReader.h b/Engine/Source/Runtime/Engine/Public/Physics/Experimental/ChaosDerivedDataReader.h
--- a/Engine/Source/Runtime/Engine/Public/Physics/Experimental/ChaosDerivedDataReader.h
+++ b/Engine/Source/Runtime/Engine/Public/Physics/Experimental/ChaosDerivedDataReader.h
@@ -13,8 +13,6 @@ namespace Chaos
 	class FConvex;
 }
 
-struct FUntypedBulkData;
-
 template<typename T, int d>
 class FChaosDerivedDataReader
 {
