Added platform id printing to a warning in the HLSL Translator for better debugging.

diff --git a/Engine/Source/Runtime/Engine/Private/Materials/HLSLMaterialTranslator.cpp b/Engine/Source/Runtime/Engine/Private/Materials/HLSLMaterialTranslator.cpp
--- a/Engine/Source/Runtime/Engine/Private/Materials/HLSLMaterialTranslator.cpp
+++ b/Engine/Source/Runtime/Engine/Private/Materials/HLSLMaterialTranslator.cpp
@@ -6571,7 +6571,7 @@ void FHLSLMaterialTranslator::UseSceneTextureId(ESceneTextureId SceneTextureId,
 
 	if (bNeedsGBuffer && IsForwardShadingEnabled(Platform))
 	{
-		Errorf(TEXT("GBuffer scene textures not available with forward shading."));
+		Errorf(TEXT("GBuffer scene textures not available with forward shading (platform id %d)."), Platform);
 	}
 
 	if (SceneTextureId == PPI_Velocity)
