Removed global ActorTransactionAnnotation.h include

#rnx

diff --git a/Engine/Source/Runtime/Engine/Classes/GameFramework/Actor.h b/Engine/Source/Runtime/Engine/Classes/GameFramework/Actor.h
--- a/Engine/Source/Runtime/Engine/Classes/GameFramework/Actor.h
+++ b/Engine/Source/Runtime/Engine/Classes/GameFramework/Actor.h
@@ -13,7 +13,6 @@
 #include "Engine/EngineTypes.h"
 #include "Engine/EngineBaseTypes.h"
 #include "PropertyPairsMap.h"
-#include "ActorTransactionAnnotation.h"
 #include "Components/ChildActorComponent.h"
 #include "RenderCommandFence.h"
 #include "Engine/Level.h"
@@ -2061,9 +2060,6 @@ public:
 	virtual bool SupportsExternalPackaging() const;
 #endif
 
-	using FActorTransactionAnnotationData = FActorTransactionAnnotationData;
-	using FActorRootComponentReconstructionData = FActorRootComponentReconstructionData;
-
 #if WITH_EDITOR
 	/** Cached pointer to the transaction annotation data from PostEditUndo to be used in the next RerunConstructionScript */
 	TSharedPtr<FActorTransactionAnnotation> CurrentTransactionAnnotation;
