Add virtual method to controllers to control if they participate in server travel.

diff --git a/Engine/Source/Runtime/Engine/Classes/GameFramework/Controller.h b/Engine/Source/Runtime/Engine/Classes/GameFramework/Controller.h
--- a/Engine/Source/Runtime/Engine/Classes/GameFramework/Controller.h
+++ b/Engine/Source/Runtime/Engine/Classes/GameFramework/Controller.h
@@ -385,6 +385,9 @@ public:
 	/** Called when the level this controller is in is unloaded via streaming. */
 	virtual void CurrentLevelUnloaded();
 
+	/** Returns whether this controller should persist through seamless server travel */
+	virtual bool ShouldParticipateInServerTravel() const;
+
 protected:
 	/** State entered when inactive (no possessed pawn, not spectating, etc). */
 	virtual void BeginInactiveState();
diff --git a/Engine/Source/Runtime/Engine/Classes/GameFramework/PlayerController.h b/Engine/Source/Runtime/Engine/Classes/GameFramework/PlayerController.h
--- a/Engine/Source/Runtime/Engine/Classes/GameFramework/PlayerController.h
+++ b/Engine/Source/Runtime/Engine/Classes/GameFramework/PlayerController.h
@@ -2145,6 +2145,12 @@ public:
 	 */
 	void SetAsLocalPlayerController() { bIsLocalPlayerController = true; }
 
+	/**
+	 * Whether this controller should persist through seamless server travel
+	 * Player controllers should always be included in server travel
+	 */
+	virtual bool ShouldParticipateInServerTravel() const override { return true; }
+
 private:
 	/** If true, prevent any haptic effects from playing */
 	bool bDisableHaptics : 1;
diff --git a/Engine/Source/Runtime/Engine/Private/Controller.cpp b/Engine/Source/Runtime/Engine/Private/Controller.cpp
--- a/Engine/Source/Runtime/Engine/Private/Controller.cpp
+++ b/Engine/Source/Runtime/Engine/Private/Controller.cpp
@@ -766,4 +766,9 @@ void AController::GetLifetimeReplicatedProps( TArray< FLifetimeProperty > & OutL
 	DOREPLIFETIME_CONDITION_NOTIFY(AController, Pawn, COND_None, REPNOTIFY_Always);
 }
 
+bool AController::ShouldParticipateInServerTravel() const
+{
+	return (PlayerState != nullptr);
+}
+
 #undef LOCTEXT_NAMESPACE
diff --git a/Engine/Source/Runtime/Engine/Private/World.cpp b/Engine/Source/Runtime/Engine/Private/World.cpp
--- a/Engine/Source/Runtime/Engine/Private/World.cpp
+++ b/Engine/Source/Runtime/Engine/Private/World.cpp
@@ -7142,7 +7142,7 @@ UWorld* FSeamlessTravelHandler::Tick()
 				{
 					if (AController* Player = Iterator->Get())
 					{
-						if (Player->PlayerState || Cast<APlayerController>(Player) != nullptr)
+						if (Player->ShouldParticipateInServerTravel())
 						{
 							KeepAnnotation.Set(Player);
 						}
