Unshelved from pending changelist '21116737' (Max Whitehead)
Ovlerap_Geom now includes a bTraceComplex arg

#[fyi] Michael.Lentine, Benn.Gallagher

diff --git a/Engine/Source/Runtime/Engine/Classes/PhysicsEngine/BodyInstance.h b/Engine/Source/Runtime/Engine/Classes/PhysicsEngine/BodyInstance.h
--- a/Engine/Source/Runtime/Engine/Classes/PhysicsEngine/BodyInstance.h
+++ b/Engine/Source/Runtime/Engine/Classes/PhysicsEngine/BodyInstance.h
@@ -1082,9 +1082,10 @@ public:
 	 *  @param  Rotation		Rotation to apply to the shape before testing
 	 *	@param	CollisionShape	Shape to test against
 	 *  @param  OutMTD			The minimum translation direction needed to push the shape out of this BodyInstance. (Optional)
+	 *  @param  TraceComplex    Trace against complex or simple geometry (Defaults simple)
 	 *  @return true if the geometry associated with this body instance overlaps the query shape at the specified location/rotation
 	 */
-	bool OverlapTest(const FVector& Position, const FQuat& Rotation, const struct FCollisionShape& CollisionShape, FMTDResult* OutMTD = nullptr) const;
+	bool OverlapTest(const FVector& Position, const FQuat& Rotation, const struct FCollisionShape& CollisionShape, FMTDResult* OutMTD = nullptr, bool bTraceComplex = false) const;
 
 	/**
 	 *  Test if the bodyinstance overlaps with the specified shape at the specified position/rotation
@@ -1094,9 +1095,10 @@ public:
 	 *  @param  Rotation		Rotation to apply to the shape before testing
 	 *	@param	CollisionShape	Shape to test against
 	 *  @param  OutMTD			The minimum translation direction needed to push the shape out of this BodyInstance. (Optional)
+	 * 	@param  TraceComplex    Trace against complex or simple geometry  (Defaults simple)
 	 *  @return true if the geometry associated with this body instance overlaps the query shape at the specified location/rotation
 	 */
-	bool OverlapTest_AssumesLocked(const FVector& Position, const FQuat& Rotation, const struct FCollisionShape& CollisionShape, FMTDResult* OutMTD = nullptr) const;
+	bool OverlapTest_AssumesLocked(const FVector& Position, const FQuat& Rotation, const struct FCollisionShape& CollisionShape, FMTDResult* OutMTD = nullptr, bool bTraceComplex = false) const;
 
 	/**
 	 *  Test if the bodyinstance overlaps with the specified body instances
@@ -1104,10 +1106,11 @@ public:
 	 *  @param  Position		Position to place our shapes at before testing (shapes of this BodyInstance)
 	 *  @param  Rotation		Rotation to apply to our shapes before testing (shapes of this BodyInstance)
 	 *  @param  Bodies			The bodies we are testing for overlap with. These bodies will be in world space already
+	 *  @param  TraceComplex    Trace against complex or simple geometry (Defaults simple)
 	 *  @return true if any of the bodies passed in overlap with this
 	 */
-	bool OverlapTestForBodies(const FVector& Position, const FQuat& Rotation, const TArray<FBodyInstance*>& Bodies) const;
-	bool OverlapTestForBody(const FVector& Position, const FQuat& Rotation, FBodyInstance* Body) const;
+	bool OverlapTestForBodies(const FVector& Position, const FQuat& Rotation, const TArray<FBodyInstance*>& Bodies, bool bTraceComplex = false) const;
+	bool OverlapTestForBody(const FVector& Position, const FQuat& Rotation, FBodyInstance* Body, bool bTraceComplex = false) const;
 
 	/**
 	 *  Determines the set of components that this body instance would overlap with at the supplied location/rotation
@@ -1231,7 +1234,7 @@ private:
 	static bool IsValidCollisionProfileName(FName InCollisionProfileName);
 
 	template<typename AllocatorType>
-	bool OverlapTestForBodiesImpl(const FVector& Position, const FQuat& Rotation, const TArray<FBodyInstance*, AllocatorType>& Bodies) const;
+	bool OverlapTestForBodiesImpl(const FVector& Position, const FQuat& Rotation, const TArray<FBodyInstance*, AllocatorType>& Bodies, bool bTraceComplex = false) const;
 
 	friend class UPhysicsAsset;
 	friend class UCollisionProfile;
@@ -1287,16 +1290,16 @@ FORCEINLINE_DEBUGGABLE bool FBodyInstance::OverlapMulti(TArray<struct FOverlapRe
 
 /// @endcond
 
-FORCEINLINE_DEBUGGABLE bool FBodyInstance::OverlapTestForBodies(const FVector& Position, const FQuat& Rotation, const TArray<FBodyInstance*>& Bodies) const
+FORCEINLINE_DEBUGGABLE bool FBodyInstance::OverlapTestForBodies(const FVector& Position, const FQuat& Rotation, const TArray<FBodyInstance*>& Bodies, bool bTraceComplex) const
 {
-	return OverlapTestForBodiesImpl(Position, Rotation, Bodies);
+	return OverlapTestForBodiesImpl(Position, Rotation, Bodies, bTraceComplex);
 }
 
-FORCEINLINE_DEBUGGABLE bool FBodyInstance::OverlapTestForBody(const FVector& Position, const FQuat& Rotation, FBodyInstance* Body) const
+FORCEINLINE_DEBUGGABLE bool FBodyInstance::OverlapTestForBody(const FVector& Position, const FQuat& Rotation, FBodyInstance* Body, bool bTraceComplex) const
 {
 	TArray<FBodyInstance*, TInlineAllocator<1>> InlineArray;
 	InlineArray.Add(Body);
-	return OverlapTestForBodiesImpl(Position, Rotation, InlineArray);
+	return OverlapTestForBodiesImpl(Position, Rotation, InlineArray, bTraceComplex);
 }
 
 FORCEINLINE_DEBUGGABLE bool FBodyInstance::IsInstanceSimulatingPhysics() const
@@ -1304,5 +1307,5 @@ FORCEINLINE_DEBUGGABLE bool FBodyInstance::IsInstanceSimulatingPhysics() const
 	return ShouldInstanceSimulatingPhysics() && IsValidBodyInstance();
 }
 
-extern template ENGINE_API bool FBodyInstance::OverlapTestForBodiesImpl(const FVector& Position, const FQuat& Rotation, const TArray<FBodyInstance*>& Bodies) const;
-extern template ENGINE_API bool FBodyInstance::OverlapTestForBodiesImpl(const FVector& Position, const FQuat& Rotation, const TArray<FBodyInstance*, TInlineAllocator<1>>& Bodies) const;
+extern template ENGINE_API bool FBodyInstance::OverlapTestForBodiesImpl(const FVector& Position, const FQuat& Rotation, const TArray<FBodyInstance*>& Bodies, bool bTraceComplex) const;
+extern template ENGINE_API bool FBodyInstance::OverlapTestForBodiesImpl(const FVector& Position, const FQuat& Rotation, const TArray<FBodyInstance*, TInlineAllocator<1>>& Bodies, bool bTraceComplex) const;
diff --git a/Engine/Source/Runtime/Engine/Private/PhysicsEngine/BodyInstance.cpp b/Engine/Source/Runtime/Engine/Private/PhysicsEngine/BodyInstance.cpp
--- a/Engine/Source/Runtime/Engine/Private/PhysicsEngine/BodyInstance.cpp
+++ b/Engine/Source/Runtime/Engine/Private/PhysicsEngine/BodyInstance.cpp
@@ -3556,7 +3556,7 @@ float FBodyInstance::GetDistanceToBody(const FVector& Point, FVector& OutPointOn
 }
 
 template <typename AllocatorType>
-bool FBodyInstance::OverlapTestForBodiesImpl(const FVector& Pos, const FQuat& Rot, const TArray<FBodyInstance*, AllocatorType>& Bodies) const
+bool FBodyInstance::OverlapTestForBodiesImpl(const FVector& Pos, const FQuat& Rot, const TArray<FBodyInstance*, AllocatorType>& Bodies, bool bTraceComplex) const
 {
 	bool bHaveOverlap = false;
 
@@ -3579,7 +3579,7 @@ bool FBodyInstance::OverlapTestForBodiesImpl(const FVector& Pos, const FQuat& Ro
 
 			for (const FBodyInstance* BodyInstance : Bodies)
 			{
-				bHaveOverlap = FPhysicsInterface::Overlap_Geom(BodyInstance, FPhysicsInterface::GetGeometryCollection(Shape), PTestGlobalPose);
+				bHaveOverlap = FPhysicsInterface::Overlap_Geom(BodyInstance, FPhysicsInterface::GetGeometryCollection(Shape), PTestGlobalPose, /*bOutMTD=*/nullptr, bTraceComplex);
 
 				if (bHaveOverlap)
 				{
@@ -3592,11 +3592,11 @@ bool FBodyInstance::OverlapTestForBodiesImpl(const FVector& Pos, const FQuat& Ro
 }
 
 // Explicit template instantiation for the above.
-template bool FBodyInstance::OverlapTestForBodiesImpl(const FVector& Pos, const FQuat& Rot, const TArray<FBodyInstance*>& Bodies) const;
-template bool FBodyInstance::OverlapTestForBodiesImpl(const FVector& Pos, const FQuat& Rot, const TArray<FBodyInstance*, TInlineAllocator<1>>& Bodies) const;
+template bool FBodyInstance::OverlapTestForBodiesImpl(const FVector& Pos, const FQuat& Rot, const TArray<FBodyInstance*>& Bodies, bool bTraceComplex) const;
+template bool FBodyInstance::OverlapTestForBodiesImpl(const FVector& Pos, const FQuat& Rot, const TArray<FBodyInstance*, TInlineAllocator<1>>& Bodies, bool bTraceComplex) const;
 
 
-bool FBodyInstance::OverlapTest(const FVector& Position, const FQuat& Rotation, const struct FCollisionShape& CollisionShape, FMTDResult* OutMTD) const
+bool FBodyInstance::OverlapTest(const FVector& Position, const FQuat& Rotation, const struct FCollisionShape& CollisionShape, FMTDResult* OutMTD, bool bTraceComplex) const
 {
 	SCOPE_CYCLE_COUNTER(STAT_Collision_SceneQueryTotal);
 	SCOPE_CYCLE_COUNTER(STAT_Collision_FBodyInstance_OverlapTest);
@@ -3607,19 +3607,19 @@ bool FBodyInstance::OverlapTest(const FVector& Position, const FQuat& Rotation,
 	{
 		FTransform GeomTransform(Rotation, Position);
 
-		bHasOverlap = FPhysicsInterface::Overlap_Geom(this, CollisionShape, Rotation, GeomTransform, OutMTD);
+		bHasOverlap = FPhysicsInterface::Overlap_Geom(this, CollisionShape, Rotation, GeomTransform, OutMTD, bTraceComplex);
 	});
 
 	return bHasOverlap;
 }
 
-bool FBodyInstance::OverlapTest_AssumesLocked(const FVector& Position, const FQuat& Rotation, const struct FCollisionShape& CollisionShape, FMTDResult* OutMTD /*= nullptr*/) const
+bool FBodyInstance::OverlapTest_AssumesLocked(const FVector& Position, const FQuat& Rotation, const struct FCollisionShape& CollisionShape, FMTDResult* OutMTD /*= nullptr*/, bool bTraceComplex) const
 {
 	SCOPE_CYCLE_COUNTER(STAT_Collision_SceneQueryTotal);
 	SCOPE_CYCLE_COUNTER(STAT_Collision_FBodyInstance_OverlapTest);
 
 	FTransform GeomTransform(Rotation, Position);
-	bool bHasOverlap = FPhysicsInterface::Overlap_Geom(this, CollisionShape, Rotation, GeomTransform, OutMTD);
+	bool bHasOverlap = FPhysicsInterface::Overlap_Geom(this, CollisionShape, Rotation, GeomTransform, OutMTD, bTraceComplex);
 	return bHasOverlap;
 }
 
diff --git a/Engine/Source/Runtime/Engine/Private/PhysicsEngine/Experimental/PhysInterface_Chaos.cpp b/Engine/Source/Runtime/Engine/Private/PhysicsEngine/Experimental/PhysInterface_Chaos.cpp
--- a/Engine/Source/Runtime/Engine/Private/PhysicsEngine/Experimental/PhysInterface_Chaos.cpp
+++ b/Engine/Source/Runtime/Engine/Private/PhysicsEngine/Experimental/PhysInterface_Chaos.cpp
@@ -1170,7 +1170,7 @@ bool FPhysInterface_Chaos::Sweep_Geom(FHitResult& OutHit, const FBodyInstance* I
 	return bSweepHit;
 }
 
-bool Overlap_GeomInternal(const FBodyInstance* InInstance, const Chaos::FImplicitObject& InGeom, const FTransform& GeomTransform, FMTDResult* OutOptResult)
+bool Overlap_GeomInternal(const FBodyInstance* InInstance, const Chaos::FImplicitObject& InGeom, const FTransform& GeomTransform, FMTDResult* OutOptResult, bool bTraceComplex)
 {
 	const FBodyInstance* TargetInstance = InInstance->WeldParent ? InInstance->WeldParent : InInstance;
 	FPhysicsActorHandle RigidBody = TargetInstance->ActorHandle;
@@ -1193,41 +1193,47 @@ bool Overlap_GeomInternal(const FBodyInstance* InInstance, const Chaos::FImplici
 		const Chaos::FPerShapeData* Shape = ShapeRef.Shape;
 		check(Shape);
 
-		if (TargetInstance->IsShapeBoundToBody(ShapeRef))
+		FCollisionFilterData ShapeFilter = Shape->GetQueryData();
+		const bool bShapeIsComplex = (ShapeFilter.Word3 & EPDF_ComplexCollision) != 0;
+		const bool bShapeIsSimple = (ShapeFilter.Word3 & EPDF_SimpleCollision) != 0;
+		if ((bTraceComplex && bShapeIsComplex) || (!bTraceComplex && bShapeIsSimple))
 		{
-			if (OutOptResult)
+			if (TargetInstance->IsShapeBoundToBody(ShapeRef))
 			{
-				Chaos::FMTDInfo MTDInfo;
-				if (Chaos::Utilities::CastHelper(InGeom, ActorTM, [&](const auto& Downcast, const auto& FullActorTM) { return Chaos::OverlapQuery(*Shape->GetGeometry(), FullActorTM, Downcast, GeomTransform, /*Thickness=*/0, &MTDInfo); }))
+				if (OutOptResult)
 				{
-					OutOptResult->Distance = MTDInfo.Penetration;
-					OutOptResult->Direction = MTDInfo.Normal;
-					return true;	//question: should we take most shallow penetration?
+					Chaos::FMTDInfo MTDInfo;
+					if (Chaos::Utilities::CastHelper(InGeom, ActorTM, [&](const auto& Downcast, const auto& FullActorTM) { return Chaos::OverlapQuery(*Shape->GetGeometry(), FullActorTM, Downcast, GeomTransform, /*Thickness=*/0, &MTDInfo); }))
+					{
+						OutOptResult->Distance = MTDInfo.Penetration;
+						OutOptResult->Direction = MTDInfo.Normal;
+						return true;	//question: should we take most shallow penetration?
+					}
 				}
-			}
-			else	//question: why do we even allow user to not pass in MTD info?
-			{
-				if (Chaos::Utilities::CastHelper(InGeom, ActorTM, [&](const auto& Downcast, const auto& FullActorTM) { return Chaos::OverlapQuery(*Shape->GetGeometry(), FullActorTM, Downcast, GeomTransform); }))
+				else	//question: why do we even allow user to not pass in MTD info?
 				{
-					return true;
+					if (Chaos::Utilities::CastHelper(InGeom, ActorTM, [&](const auto& Downcast, const auto& FullActorTM) { return Chaos::OverlapQuery(*Shape->GetGeometry(), FullActorTM, Downcast, GeomTransform); }))
+					{
+						return true;
+					}
 				}
-			}
 
+			}
 		}
 	}
 
 	return false;
 }
 
-bool FPhysInterface_Chaos::Overlap_Geom(const FBodyInstance* InBodyInstance, const FPhysicsGeometryCollection& InGeometry, const FTransform& InShapeTransform, FMTDResult* OutOptResult)
+bool FPhysInterface_Chaos::Overlap_Geom(const FBodyInstance* InBodyInstance, const FPhysicsGeometryCollection& InGeometry, const FTransform& InShapeTransform, FMTDResult* OutOptResult, bool bTraceComplex)
 {
-	return Overlap_GeomInternal(InBodyInstance, InGeometry.GetGeometry(), InShapeTransform, OutOptResult);
+	return Overlap_GeomInternal(InBodyInstance, InGeometry.GetGeometry(), InShapeTransform, OutOptResult, bTraceComplex);
 }
 
-bool FPhysInterface_Chaos::Overlap_Geom(const FBodyInstance* InBodyInstance, const FCollisionShape& InCollisionShape, const FQuat& InShapeRotation, const FTransform& InShapeTransform, FMTDResult* OutOptResult)
+bool FPhysInterface_Chaos::Overlap_Geom(const FBodyInstance* InBodyInstance, const FCollisionShape& InCollisionShape, const FQuat& InShapeRotation, const FTransform& InShapeTransform, FMTDResult* OutOptResult, bool bTraceComplex)
 {
 	FPhysicsShapeAdapter Adaptor(InShapeRotation, InCollisionShape);
-	return Overlap_GeomInternal(InBodyInstance, Adaptor.GetGeometry(), Adaptor.GetGeomPose(InShapeTransform.GetTranslation()), OutOptResult);
+	return Overlap_GeomInternal(InBodyInstance, Adaptor.GetGeometry(), Adaptor.GetGeomPose(InShapeTransform.GetTranslation()), OutOptResult, bTraceComplex);
 }
 
 bool FPhysInterface_Chaos::GetSquaredDistanceToBody(const FBodyInstance* InInstance, const FVector& InPoint, float& OutDistanceSquared, FVector* OutOptPointOnBody)
diff --git a/Engine/Source/Runtime/Engine/Public/Physics/Experimental/PhysInterface_Chaos.h b/Engine/Source/Runtime/Engine/Public/Physics/Experimental/PhysInterface_Chaos.h
--- a/Engine/Source/Runtime/Engine/Public/Physics/Experimental/PhysInterface_Chaos.h
+++ b/Engine/Source/Runtime/Engine/Public/Physics/Experimental/PhysInterface_Chaos.h
@@ -96,8 +96,8 @@ public:
 	// Trace functions for testing specific geometry (not against a world)
 	static bool LineTrace_Geom(FHitResult& OutHit, const FBodyInstance* InInstance, const FVector& InStart, const FVector& InEnd, bool bTraceComplex, bool bExtractPhysMaterial = false);
 	static bool Sweep_Geom(FHitResult& OutHit, const FBodyInstance* InInstance, const FVector& InStart, const FVector& InEnd, const FQuat& InShapeRotation, const FCollisionShape& InShape, bool bSweepComplex);
-	static bool Overlap_Geom(const FBodyInstance* InBodyInstance, const FPhysicsGeometryCollection& InGeometry, const FTransform& InShapeTransform, FMTDResult* OutOptResult = nullptr);
-	static bool Overlap_Geom(const FBodyInstance* InBodyInstance, const FCollisionShape& InCollisionShape, const FQuat& InShapeRotation, const FTransform& InShapeTransform, FMTDResult* OutOptResult = nullptr);
+	static bool Overlap_Geom(const FBodyInstance* InBodyInstance, const FPhysicsGeometryCollection& InGeometry, const FTransform& InShapeTransform, FMTDResult* OutOptResult = nullptr, bool bTraceComplex = false);
+	static bool Overlap_Geom(const FBodyInstance* InBodyInstance, const FCollisionShape& InCollisionShape, const FQuat& InShapeRotation, const FTransform& InShapeTransform, FMTDResult* OutOptResult = nullptr, bool bTraceComplex = false);
 	static bool GetSquaredDistanceToBody(const FBodyInstance* InInstance, const FVector& InPoint, float& OutDistanceSquared, FVector* OutOptPointOnBody = nullptr);
 
 	static void SetMaterials(const FPhysicsShapeHandle& InShape, const TArrayView<UPhysicalMaterial*>InMaterials);
