World Partition: don't trash package more than once when unloading cells in PIE, this happens for child actors.
Original PR #9413: WP: prevent package to be trashed many times (Contributed by [REDACTED_PERSON]).

#rnx

diff --git a/Engine/Source/Runtime/Engine/Private/WorldPartition/WorldPartitionLevelStreamingDynamic.cpp b/Engine/Source/Runtime/Engine/Private/WorldPartition/WorldPartitionLevelStreamingDynamic.cpp
--- a/Engine/Source/Runtime/Engine/Private/WorldPartition/WorldPartitionLevelStreamingDynamic.cpp
+++ b/Engine/Source/Runtime/Engine/Private/WorldPartition/WorldPartitionLevelStreamingDynamic.cpp
@@ -435,18 +435,21 @@ void UWorldPartitionLevelStreamingDynamic::OnCleanupLevel()
 
 		RuntimeLevel->OnCleanupLevel.Remove(OnCleanupLevelDelegateHandle);
 
-		auto TrashPackage = [](UPackage* Package)
+		TSet<UPackage*> TrashedPachages;
+		auto TrashPackage = [&TrashedPachages](UPackage* Package)
 		{
-			// Clears RF_Standalone flag on objects in package (UMetaData)
-			ForEachObjectWithPackage(Package, [](UObject* Object)
+			bool bWasAlreadyInSet;
+			TrashedPachages.Add(Package, &bWasAlreadyInSet);
+
+			if (!bWasAlreadyInSet)
 			{
-				Object->ClearFlags(RF_Standalone);
-				return true;
-			}, false);
+				// Clears RF_Standalone flag on objects in package (UMetaData)
+				ForEachObjectWithPackage(Package, [](UObject* Object) { Object->ClearFlags(RF_Standalone); return true; }, false);
 
-			// Rename package to avoid having to deal with pending kill objects in subsequent RequestLevel calls
-			FName NewPackageName = MakeUniqueObjectName(nullptr, UPackage::StaticClass(), FName(*FString::Printf(TEXT("%s_Trashed"), *Package->GetName())));
-			Package->Rename(*NewPackageName.ToString(), nullptr, REN_ForceNoResetLoaders | REN_DontCreateRedirectors | REN_NonTransactional | REN_DoNotDirty);
+				// Rename package to avoid having to deal with pending kill objects in subsequent RequestLevel calls
+				FName NewPackageName = MakeUniqueObjectName(nullptr, UPackage::StaticClass(), FName(*FString::Printf(TEXT("%s_Trashed"), *Package->GetName())));
+				Package->Rename(*NewPackageName.ToString(), nullptr, REN_ForceNoResetLoaders | REN_DontCreateRedirectors | REN_NonTransactional | REN_DoNotDirty);
+			}
 		};
 		
 		TrashPackage(RuntimeLevel->GetPackage());
