Rename virtual method for controllers to control if they participate in seamless travel so it references the correct travel method.

diff --git a/Engine/Source/Runtime/Engine/Classes/GameFramework/Controller.h b/Engine/Source/Runtime/Engine/Classes/GameFramework/Controller.h
--- a/Engine/Source/Runtime/Engine/Classes/GameFramework/Controller.h
+++ b/Engine/Source/Runtime/Engine/Classes/GameFramework/Controller.h
@@ -385,8 +385,8 @@ public:
 	/** Called when the level this controller is in is unloaded via streaming. */
 	virtual void CurrentLevelUnloaded();
 
-	/** Returns whether this controller should persist through seamless server travel */
-	virtual bool ShouldParticipateInServerTravel() const;
+	/** Returns whether this controller should persist through seamless travel */
+	virtual bool ShouldParticipateInSeamlessTravel() const;
 
 protected:
 	/** State entered when inactive (no possessed pawn, not spectating, etc). */
diff --git a/Engine/Source/Runtime/Engine/Classes/GameFramework/PlayerController.h b/Engine/Source/Runtime/Engine/Classes/GameFramework/PlayerController.h
--- a/Engine/Source/Runtime/Engine/Classes/GameFramework/PlayerController.h
+++ b/Engine/Source/Runtime/Engine/Classes/GameFramework/PlayerController.h
@@ -2146,10 +2146,10 @@ public:
 	void SetAsLocalPlayerController() { bIsLocalPlayerController = true; }
 
 	/**
-	 * Whether this controller should persist through seamless server travel
-	 * Player controllers should always be included in server travel
+	 * Whether this controller should persist through seamless travel
+	 * Player controllers should always be included in seamless travel
 	 */
-	virtual bool ShouldParticipateInServerTravel() const override { return true; }
+	virtual bool ShouldParticipateInSeamlessTravel() const override { return true; }
 
 private:
 	/** If true, prevent any haptic effects from playing */
diff --git a/Engine/Source/Runtime/Engine/Private/Controller.cpp b/Engine/Source/Runtime/Engine/Private/Controller.cpp
--- a/Engine/Source/Runtime/Engine/Private/Controller.cpp
+++ b/Engine/Source/Runtime/Engine/Private/Controller.cpp
@@ -766,7 +766,7 @@ void AController::GetLifetimeReplicatedProps( TArray< FLifetimeProperty > & OutL
 	DOREPLIFETIME_CONDITION_NOTIFY(AController, Pawn, COND_None, REPNOTIFY_Always);
 }
 
-bool AController::ShouldParticipateInServerTravel() const
+bool AController::ShouldParticipateInSeamlessTravel() const
 {
 	return (PlayerState != nullptr);
 }
diff --git a/Engine/Source/Runtime/Engine/Private/World.cpp b/Engine/Source/Runtime/Engine/Private/World.cpp
--- a/Engine/Source/Runtime/Engine/Private/World.cpp
+++ b/Engine/Source/Runtime/Engine/Private/World.cpp
@@ -7142,7 +7142,7 @@ UWorld* FSeamlessTravelHandler::Tick()
 				{
 					if (AController* Player = Iterator->Get())
 					{
-						if (Player->ShouldParticipateInServerTravel())
+						if (Player->ShouldParticipateInSeamlessTravel())
 						{
 							KeepAnnotation.Set(Player);
 						}
