Manually merging this into main from 21.50.  Ran into conflict merging.
Added more draw call catregories to stat drawcount to account for some major drawcalls that are being missed.
#RNX

diff --git a/Engine/Source/Runtime/Engine/Private/Particles/FXSystem.cpp b/Engine/Source/Runtime/Engine/Private/Particles/FXSystem.cpp
--- a/Engine/Source/Runtime/Engine/Private/Particles/FXSystem.cpp
+++ b/Engine/Source/Runtime/Engine/Private/Particles/FXSystem.cpp
@@ -475,6 +475,8 @@ DECLARE_CYCLE_STAT(TEXT("FXPreRender_PrepareCDF"), STAT_CLM_FXPreRender_PrepareC
 DECLARE_CYCLE_STAT(TEXT("FXPreRender_SimulateCDF"), STAT_CLM_FXPreRender_SimulateCDF, STATGROUP_CommandListMarkers);
 DECLARE_CYCLE_STAT(TEXT("FXPreRender_FinalizeCDF"), STAT_CLM_FXPreRender_FinalizeCDF, STATGROUP_CommandListMarkers);
 
+DECLARE_GPU_DRAWCALL_STAT(FXSystemPreRender);
+DECLARE_GPU_DRAWCALL_STAT(FXSystemPostRenderOpaque);
 
 void FFXSystem::PreRender(FRDGBuilder& GraphBuilder, TConstArrayView<FViewInfo> Views, bool bAllowGPUParticleSceneUpdate)
 {
@@ -482,6 +484,7 @@ void FFXSystem::PreRender(FRDGBuilder& GraphBuilder, TConstArrayView<FViewInfo>
 
 	if (RHISupportsGPUParticles() && bAllowGPUParticleSceneUpdate)
 	{
+		RDG_GPU_STAT_SCOPE(GraphBuilder, FXSystemPreRender);
 		RDG_CSV_STAT_EXCLUSIVE_SCOPE(GraphBuilder, FXSystem);
 
 		FRHIUniformBuffer* ViewUniformBuffer = GetReferenceViewUniformBuffer(Views);
@@ -528,6 +531,7 @@ void FFXSystem::PostRenderOpaque(FRDGBuilder& GraphBuilder, TConstArrayView<FVie
 
 	if (RHISupportsGPUParticles() && IsParticleCollisionModeSupported(GetShaderPlatform(), PCM_DepthBuffer) && bAllowGPUParticleUpdate)
 	{
+		RDG_GPU_STAT_SCOPE(GraphBuilder, FXSystemPostRenderOpaque);
 		RDG_CSV_STAT_EXCLUSIVE_SCOPE(GraphBuilder, FXSystem);
 
 		FRHIUniformBuffer* ViewUniformBuffer = GetReferenceViewUniformBuffer(Views);
diff --git a/Engine/Source/Runtime/Renderer/Private/CustomDepthRendering.cpp b/Engine/Source/Runtime/Renderer/Private/CustomDepthRendering.cpp
--- a/Engine/Source/Runtime/Renderer/Private/CustomDepthRendering.cpp
+++ b/Engine/Source/Runtime/Renderer/Private/CustomDepthRendering.cpp
@@ -32,7 +32,7 @@ static TAutoConsoleVariable<int32> CVarCustomDepthTemporalAAJitter(
 	TEXT("If disabled the Engine will remove the TemporalAA Jitter from the Custom Depth Pass. Only has effect when TemporalAA is used."),
 	ECVF_RenderThreadSafe);
 
-DECLARE_GPU_STAT_NAMED(CustomDepth, TEXT("Custom Depth"));
+DECLARE_GPU_DRAWCALL_STAT_NAMED(CustomDepth, TEXT("Custom Depth"));
 
 ECustomDepthPassLocation GetCustomDepthPassLocation(EShaderPlatform Platform)
 {
diff --git a/Engine/Source/Runtime/Renderer/Private/DeferredShadingRenderer.cpp b/Engine/Source/Runtime/Renderer/Private/DeferredShadingRenderer.cpp
--- a/Engine/Source/Runtime/Renderer/Private/DeferredShadingRenderer.cpp
+++ b/Engine/Source/Runtime/Renderer/Private/DeferredShadingRenderer.cpp
@@ -297,7 +297,7 @@ DECLARE_GPU_STAT(RayTracingUpdate);
 DECLARE_GPU_STAT(RayTracingScene);
 DECLARE_GPU_STAT(RayTracingGeometry);
 
-DECLARE_GPU_STAT(Postprocessing);
+DECLARE_GPU_DRAWCALL_STAT(Postprocessing);
 DECLARE_GPU_STAT(VisibilityCommands);
 DECLARE_GPU_STAT(RenderDeferredLighting);
 DECLARE_GPU_STAT(AllocateRendertargets);
@@ -306,7 +306,7 @@ DECLARE_GPU_STAT(SortLights);
 DECLARE_GPU_STAT(PostRenderOpsFX);
 DECLARE_GPU_STAT(GPUSceneUpdate);
 DECLARE_GPU_STAT_NAMED(Unaccounted, TEXT("[unaccounted]"));
-DECLARE_GPU_DRAWCALL_STAT(WaterRendering);
+DECLARE_GPU_STAT(WaterRendering);
 DECLARE_GPU_STAT(HairRendering);
 DEFINE_GPU_DRAWCALL_STAT(VirtualTextureUpdate);
 DECLARE_GPU_STAT(UploadDynamicBuffers);
diff --git a/Engine/Source/Runtime/Renderer/Private/DistortionRendering.cpp b/Engine/Source/Runtime/Renderer/Private/DistortionRendering.cpp
--- a/Engine/Source/Runtime/Renderer/Private/DistortionRendering.cpp
+++ b/Engine/Source/Runtime/Renderer/Private/DistortionRendering.cpp
@@ -16,7 +16,7 @@
 #include "Strata/Strata.h"
 #include "ScreenRendering.h"
 
-DECLARE_GPU_STAT(Distortion);
+DECLARE_GPU_DRAWCALL_STAT(Distortion);
 
 const uint8 kStencilMaskBit = STENCIL_SANDBOX_MASK;
 
diff --git a/Engine/Source/Runtime/Renderer/Private/FogRendering.cpp b/Engine/Source/Runtime/Renderer/Private/FogRendering.cpp
--- a/Engine/Source/Runtime/Renderer/Private/FogRendering.cpp
+++ b/Engine/Source/Runtime/Renderer/Private/FogRendering.cpp
@@ -8,7 +8,7 @@
 #include "SingleLayerWaterRendering.h"
 #include "ScreenPass.h"
 
-DECLARE_GPU_STAT(Fog);
+DECLARE_GPU_DRAWCALL_STAT(Fog);
 
 static TAutoConsoleVariable<int32> CVarFog(
 	TEXT("r.Fog"),
diff --git a/Engine/Source/Runtime/Renderer/Private/LightRendering.cpp b/Engine/Source/Runtime/Renderer/Private/LightRendering.cpp
--- a/Engine/Source/Runtime/Renderer/Private/LightRendering.cpp
+++ b/Engine/Source/Runtime/Renderer/Private/LightRendering.cpp
@@ -34,7 +34,7 @@
 	#define ENABLE_DEBUG_DISCARD_PROP 0
 #endif // !(UE_BUILD_SHIPPING || UE_BUILD_TEST)
 
-DECLARE_GPU_STAT(Lights);
+DECLARE_GPU_DRAWCALL_STAT(Lights);
 
 IMPLEMENT_TYPE_LAYOUT(FLightFunctionSharedParameters);
 IMPLEMENT_TYPE_LAYOUT(FStencilingGeometryShaderParameters);
diff --git a/Engine/Source/Runtime/Renderer/Private/SceneOcclusion.cpp b/Engine/Source/Runtime/Renderer/Private/SceneOcclusion.cpp
--- a/Engine/Source/Runtime/Renderer/Private/SceneOcclusion.cpp
+++ b/Engine/Source/Runtime/Renderer/Private/SceneOcclusion.cpp
@@ -86,6 +86,7 @@ bool UseDownsampledOcclusionQueries()
 }
 
 DEFINE_GPU_STAT(HZB);
+DECLARE_GPU_DRAWCALL_STAT(BeginOcclusionTests);
 
 /** Random table for occlusion **/
 FOcclusionRandomStream GOcclusionRandomStream;
@@ -1409,6 +1410,7 @@ void FDeferredShadingSceneRenderer::RenderOcclusion(
 			PassParameters->RenderTargets.DepthStencil = FDepthStencilBinding(OcclusionDepthTexture, ERenderTargetLoadAction::ELoad, ERenderTargetLoadAction::ELoad, FExclusiveDepthStencil::DepthRead_StencilWrite);
 			PassParameters->RenderTargets.NumOcclusionQueries = NumQueriesForBatch;
 
+			RDG_GPU_STAT_SCOPE(GraphBuilder, BeginOcclusionTests);
 			GraphBuilder.AddPass(
 				RDG_EVENT_NAME("BeginOcclusionTests"),
 				PassParameters,
diff --git a/Engine/Source/Runtime/Renderer/Private/ScreenSpaceRayTracing.cpp b/Engine/Source/Runtime/Renderer/Private/ScreenSpaceRayTracing.cpp
--- a/Engine/Source/Runtime/Renderer/Private/ScreenSpaceRayTracing.cpp
+++ b/Engine/Source/Runtime/Renderer/Private/ScreenSpaceRayTracing.cpp
@@ -75,7 +75,7 @@ static TAutoConsoleVariable<float> CVarSSGISkyDistance(
 	ECVF_Scalability | ECVF_RenderThreadSafe);
 
 
-DECLARE_GPU_STAT_NAMED(ScreenSpaceReflections, TEXT("ScreenSpace Reflections"));
+DECLARE_GPU_DRAWCALL_STAT_NAMED(ScreenSpaceReflections, TEXT("ScreenSpace Reflections"));
 DECLARE_GPU_STAT_NAMED(ScreenSpaceDiffuseIndirect, TEXT("Screen Space Diffuse Indirect"));
 
 
diff --git a/Engine/Source/Runtime/Renderer/Private/ShadowRendering.cpp b/Engine/Source/Runtime/Renderer/Private/ShadowRendering.cpp
--- a/Engine/Source/Runtime/Renderer/Private/ShadowRendering.cpp
+++ b/Engine/Source/Runtime/Renderer/Private/ShadowRendering.cpp
@@ -235,7 +235,7 @@ static TAutoConsoleVariable<int32> CVarHairStrandsCullPerObjectShadowCaster(
 	TEXT("Enable CPU culling of object casting per-object shadow (stationnary object)"),
 	ECVF_RenderThreadSafe);
 
-DEFINE_GPU_DRAWCALL_STAT(ShadowProjection);
+DEFINE_GPU_STAT(ShadowProjection);
 
 // Use Shadow stencil mask is set to 0x07u instead of 0xFF so that that last bit can be used for strata classification without clearing the stencil bit for pre-shadow/per-object static shadow mask
 constexpr uint32 ShadowStencilMask = 0x07u;
diff --git a/Engine/Source/Runtime/Renderer/Private/SingleLayerWaterRendering.cpp b/Engine/Source/Runtime/Renderer/Private/SingleLayerWaterRendering.cpp
--- a/Engine/Source/Runtime/Renderer/Private/SingleLayerWaterRendering.cpp
+++ b/Engine/Source/Runtime/Renderer/Private/SingleLayerWaterRendering.cpp
@@ -23,7 +23,7 @@
 
 DECLARE_GPU_STAT_NAMED(RayTracingWaterReflections, TEXT("Ray Tracing Water Reflections"));
 
-DECLARE_GPU_STAT(SingleLayerWater);
+DECLARE_GPU_DRAWCALL_STAT(SingleLayerWater);
 DECLARE_CYCLE_STAT(TEXT("WaterSingleLayer"), STAT_CLP_WaterSingleLayerPass, STATGROUP_ParallelCommandListMarkers);
 
 static TAutoConsoleVariable<int32> CVarWaterSingleLayer(
diff --git a/Engine/Source/Runtime/Renderer/Private/VT/VirtualTextureSystem.cpp b/Engine/Source/Runtime/Renderer/Private/VT/VirtualTextureSystem.cpp
--- a/Engine/Source/Runtime/Renderer/Private/VT/VirtualTextureSystem.cpp
+++ b/Engine/Source/Runtime/Renderer/Private/VT/VirtualTextureSystem.cpp
@@ -69,7 +69,7 @@ DECLARE_MEMORY_STAT_POOL(TEXT("Total Physical Memory"), STAT_TotalPhysicalMemory
 DECLARE_MEMORY_STAT_POOL(TEXT("Total Pagetable Memory"), STAT_TotalPagetableMemory, STATGROUP_VirtualTextureMemory, FPlatformMemory::MCR_GPU);
 
 DECLARE_GPU_STAT(VirtualTexture);
-DECLARE_GPU_DRAWCALL_STAT(VirtualTextureAllocate);
+DECLARE_GPU_STAT(VirtualTextureAllocate);
 
 static TAutoConsoleVariable<int32> CVarVTVerbose(
 	TEXT("r.VT.Verbose"),
