Removed global ITransaction.h include

#rnx

diff --git a/Engine/Source/Editor/UnrealEd/Private/UnrealEdPrivatePCH.h b/Engine/Source/Editor/UnrealEd/Private/UnrealEdPrivatePCH.h
--- a/Engine/Source/Editor/UnrealEd/Private/UnrealEdPrivatePCH.h
+++ b/Engine/Source/Editor/UnrealEd/Private/UnrealEdPrivatePCH.h
@@ -209,7 +209,6 @@
 #include "Misc/Paths.h"
 #include "Misc/CoreStats.h"
 #include "Templates/ScopedCallback.h"
-#include "Misc/ITransaction.h"
 #include "Serialization/CustomVersion.h"
 #include "Misc/OutputDeviceError.h"
 #include "Misc/ObjectThumbnail.h"
diff --git a/Engine/Source/Editor/UnrealEd/Public/UnrealEdSharedPCH.h b/Engine/Source/Editor/UnrealEd/Public/UnrealEdSharedPCH.h
--- a/Engine/Source/Editor/UnrealEd/Public/UnrealEdSharedPCH.h
+++ b/Engine/Source/Editor/UnrealEd/Public/UnrealEdSharedPCH.h
@@ -202,7 +202,6 @@
 #include "Misc/OutputDeviceError.h"
 #include "Misc/ObjectThumbnail.h"
 #include "Containers/StaticArray.h"
-#include "Misc/ITransaction.h"
 #include "Templates/ScopedCallback.h"
 #include "Misc/CoreStats.h"
 #include "HAL/PlatformProcess.h"
diff --git a/Engine/Source/Runtime/Engine/Private/EnginePrivatePCH.h b/Engine/Source/Runtime/Engine/Private/EnginePrivatePCH.h
--- a/Engine/Source/Runtime/Engine/Private/EnginePrivatePCH.h
+++ b/Engine/Source/Runtime/Engine/Private/EnginePrivatePCH.h
@@ -189,7 +189,6 @@
 #include "Templates/RefCounting.h"
 #include "Async/TaskGraphInterfaces.h"
 #include "Misc/Attribute.h"
-#include "Misc/ITransaction.h"
 #include "Templates/ScopedCallback.h"
 #include "Misc/CoreStats.h"
 #include "Modules/ModuleInterface.h"
diff --git a/Engine/Source/Runtime/Engine/Public/EngineSharedPCH.h b/Engine/Source/Runtime/Engine/Public/EngineSharedPCH.h
--- a/Engine/Source/Runtime/Engine/Public/EngineSharedPCH.h
+++ b/Engine/Source/Runtime/Engine/Public/EngineSharedPCH.h
@@ -197,7 +197,6 @@
 #include "GenericPlatform/GenericWindow.h"
 #include "Math/TransformCalculus.h"
 #include "Math/TransformCalculus2D.h"
-#include "Misc/ITransaction.h"
 #include "Templates/ScopedCallback.h"
 #include "HAL/PlatformProcess.h"
 #include "GenericPlatform/GenericPlatformProcess.h"
