Changelist Validator: fix crash when validating an unloaded actor (its class could also be unloaded), by retrieving the actor's native class from the actor descriptor data if possible.

#rnx

diff --git a/Engine/Plugins/Editor/DataValidation/Source/DataValidation/Private/WorldPartitionChangelistValidator.cpp b/Engine/Plugins/Editor/DataValidation/Source/DataValidation/Private/WorldPartitionChangelistValidator.cpp
--- a/Engine/Plugins/Editor/DataValidation/Source/DataValidation/Private/WorldPartitionChangelistValidator.cpp
+++ b/Engine/Plugins/Editor/DataValidation/Source/DataValidation/Private/WorldPartitionChangelistValidator.cpp
@@ -74,9 +74,9 @@ EDataValidationResult UWorldPartitionChangelistValidator::ValidateActorsAndDataL
 	ISourceControlProvider& SourceControlProvider = ISourceControlModule::Get().GetProvider();
 	FSourceControlChangelistStatePtr ChangelistState = SourceControlProvider.GetState(Changelist->Changelist->AsShared(), EStateCacheUsage::Use);
 
-	// Figure out which world(s) those assets are in and split the files per world
+	// Figure out which world(s) these actors are in and split the files per world, and return the actor's native class
 	TMap<FName, TSet<FAssetData>> MapToActorsFiles;
-	auto TryAssociateActorToMap = [&MapToActorsFiles](const FAssetData& AssetData)
+	auto TryAssociateActorToMap = [&MapToActorsFiles](const FAssetData& AssetData) -> UClass*
 	{
 		// Check that the asset is an actor
 		if (FWorldPartitionActorDescUtils::IsValidActorDescriptorFromAssetData(AssetData))
@@ -99,11 +99,11 @@ EDataValidationResult UWorldPartitionChangelistValidator::ValidateActorsAndDataL
 			if (ActorFiles)	// A null Files indicates a World not using World Partition and OFPA 
 			{
 				ActorFiles->Add(AssetData);
-				return true;
+				return FWorldPartitionActorDescUtils::GetActorNativeClassFromAssetData(AssetData);;
 			}
 		}
 
-		return false;
+		return nullptr;
 	};
 	
 	FAssetRegistryModule& AssetRegistryModule = FModuleManager::LoadModuleChecked<FAssetRegistryModule>(FName("AssetRegistry"));
@@ -124,15 +124,11 @@ EDataValidationResult UWorldPartitionChangelistValidator::ValidateActorsAndDataL
 			
 			for (FAssetData& AssetData : PackageAssetsData)
 			{
-				UClass* AssetClass = AssetData.GetClass();
-
-				if (TryAssociateActorToMap(AssetData))
+				if (UClass* ActorNativeClass = TryAssociateActorToMap(AssetData))
 				{
-					check(AssetClass);
-					check(AssetClass->IsChildOf<AActor>())
-					SubmittingWorldDataLayers = AssetClass->IsChildOf<AWorldDataLayers>();
+					SubmittingWorldDataLayers = ActorNativeClass->IsChildOf<AWorldDataLayers>();
 				}
-				else if (AssetClass)
+				else if (UClass* AssetClass = AssetData.GetClass())
 				{
 					if (AssetClass->IsChildOf<UDataLayerAsset>())
 					{
diff --git a/Engine/Source/Runtime/Engine/Private/WorldPartition/WorldPartitionActorDescUtils.cpp b/Engine/Source/Runtime/Engine/Private/WorldPartition/WorldPartitionActorDescUtils.cpp
--- a/Engine/Source/Runtime/Engine/Private/WorldPartition/WorldPartitionActorDescUtils.cpp
+++ b/Engine/Source/Runtime/Engine/Private/WorldPartition/WorldPartitionActorDescUtils.cpp
@@ -17,11 +17,10 @@ bool FWorldPartitionActorDescUtils::IsValidActorDescriptorFromAssetData(const FA
 	return InAssetData.FindTag(NAME_ActorMetaDataClass) && InAssetData.FindTag(NAME_ActorMetaData);
 }
 
-TUniquePtr<FWorldPartitionActorDesc> FWorldPartitionActorDescUtils::GetActorDescriptorFromAssetData(const FAssetData& InAssetData)
+UClass* FWorldPartitionActorDescUtils::GetActorNativeClassFromAssetData(const FAssetData& InAssetData)
 {
 	FString ActorMetaDataClass;
-	FString ActorMetaDataStr;
-	if (InAssetData.GetTagValue(NAME_ActorMetaDataClass, ActorMetaDataClass) && InAssetData.GetTagValue(NAME_ActorMetaData, ActorMetaDataStr))
+	if (InAssetData.GetTagValue(NAME_ActorMetaDataClass, ActorMetaDataClass))
 	{
 		FString ActorClassName;
 		FString ActorPackageName;
@@ -34,8 +33,19 @@ TUniquePtr<FWorldPartitionActorDesc> FWorldPartitionActorDescUtils::GetActorDesc
 		const FCoreRedirectObjectName OldClassName = FCoreRedirectObjectName(*ActorClassName, NAME_None, *ActorPackageName);
 		const FCoreRedirectObjectName NewClassName = FCoreRedirects::GetRedirectedName(ECoreRedirectFlags::Type_Class, OldClassName);
 
+		return UClass::TryFindTypeSlow<UClass>(NewClassName.ToString(), EFindFirstObjectOptions::ExactClass);
+	}
+	return nullptr;
+}
+
+TUniquePtr<FWorldPartitionActorDesc> FWorldPartitionActorDescUtils::GetActorDescriptorFromAssetData(const FAssetData& InAssetData)
+{
+	FString ActorMetaDataClass;
+	FString ActorMetaDataStr;
+	if (InAssetData.GetTagValue(NAME_ActorMetaDataClass, ActorMetaDataClass) && InAssetData.GetTagValue(NAME_ActorMetaData, ActorMetaDataStr))
+	{
 		bool bIsValidClass = true;
-		UClass* ActorClass = UClass::TryFindTypeSlow<UClass>(NewClassName.ToString(), EFindFirstObjectOptions::ExactClass);
+		UClass* ActorClass = GetActorNativeClassFromAssetData(InAssetData);
 
 		if (!ActorClass)
 		{
@@ -55,7 +65,7 @@ TUniquePtr<FWorldPartitionActorDesc> FWorldPartitionActorDescUtils::GetActorDesc
 			
 		if (!bIsValidClass)
 		{
-			UE_LOG(LogWorldPartition, Warning, TEXT("Invalid class `%s` for actor guid `%s` ('%s') from package '%s'"), *NewClassName.ToString(), *NewActorDesc->GetGuid().ToString(), *NewActorDesc->GetActorName().ToString(), *NewActorDesc->GetActorPackage().ToString());
+			UE_LOG(LogWorldPartition, Warning, TEXT("Invalid class for actor guid `%s` ('%s') from package '%s'"), *NewActorDesc->GetGuid().ToString(), *NewActorDesc->GetActorName().ToString(), *NewActorDesc->GetActorPackage().ToString());
 			return nullptr;
 		}
 
diff --git a/Engine/Source/Runtime/Engine/Public/WorldPartition/WorldPartitionActorDescUtils.h b/Engine/Source/Runtime/Engine/Public/WorldPartition/WorldPartitionActorDescUtils.h
--- a/Engine/Source/Runtime/Engine/Public/WorldPartition/WorldPartitionActorDescUtils.h
+++ b/Engine/Source/Runtime/Engine/Public/WorldPartition/WorldPartitionActorDescUtils.h
@@ -12,6 +12,7 @@ struct FAssetData;
 struct ENGINE_API FWorldPartitionActorDescUtils
 {
 	static bool IsValidActorDescriptorFromAssetData(const FAssetData& InAssetData);
+	static UClass* GetActorNativeClassFromAssetData(const FAssetData& InAssetData);
 	static TUniquePtr<FWorldPartitionActorDesc> GetActorDescriptorFromAssetData(const FAssetData& InAssetData);
 	static void AppendAssetDataTagsFromActor(const AActor* Actor, TArray<UObject::FAssetRegistryTag>& OutTags);
 	static void UpdateActorDescriptorFomActor(const AActor* Actor, TUniquePtr<FWorldPartitionActorDesc>& ActorDesc);
