Skeletal Mesh - Fix backward compatibility of the SkeletalMesh pointer access where both direct access to the SkeletalMesh pointer and the SetSkeletalMesh function are used.

diff --git a/Engine/Source/Runtime/Engine/Private/Components/SkinnedMeshComponent.cpp b/Engine/Source/Runtime/Engine/Private/Components/SkinnedMeshComponent.cpp
--- a/Engine/Source/Runtime/Engine/Private/Components/SkinnedMeshComponent.cpp
+++ b/Engine/Source/Runtime/Engine/Private/Components/SkinnedMeshComponent.cpp
@@ -1892,14 +1892,14 @@ void USkinnedMeshComponent::SetSkinnedAsset(class USkinnedAsset* InSkinnedAsset)
 {
 	PRAGMA_DISABLE_DEPRECATION_WARNINGS
 	SkinnedAsset = InSkinnedAsset;
-	SkeletalMesh = nullptr;  // The SkeletalMesh pointer is no longer valid from here on, blueprint/function access will use getter to retrieve the SkinnedAsset pointer and direct access is deprecated
+	SkeletalMesh = Cast<USkeletalMesh>(InSkinnedAsset);
 	PRAGMA_ENABLE_DEPRECATION_WARNINGS
 }
 
 USkinnedAsset* USkinnedMeshComponent::GetSkinnedAsset() const
 {
 	PRAGMA_DISABLE_DEPRECATION_WARNINGS
-	return SkinnedAsset ? SkinnedAsset : SkeletalMesh;  // Return the SkeletalMesh pointer for backward compatibility when the SkinnedAsset is invalid
+	return !SkeletalMesh ? SkinnedAsset : SkeletalMesh;  // Return the SkeletalMesh pointer for backward compatibility until it is retired
 	PRAGMA_ENABLE_DEPRECATION_WARNINGS
 }
 
