Multi-view-family:  Simplified fix for ghosting issues in GPU skin animation.  Only increment the Scene's frame index for the first view family.  Makes logic consistent between DC.SingleRender 0/1 settings (all view families get the same frame number in either case), and avoids the need to add a view family index to get the correct frame number.   Reverts the previous fix CL 21150022.

#rnx

diff --git a/Engine/Plugins/Runtime/nDisplay/Source/DisplayCluster/Private/Game/EngineClasses/Basics/DisplayClusterViewportClient.cpp b/Engine/Plugins/Runtime/nDisplay/Source/DisplayCluster/Private/Game/EngineClasses/Basics/DisplayClusterViewportClient.cpp
--- a/Engine/Plugins/Runtime/nDisplay/Source/DisplayCluster/Private/Game/EngineClasses/Basics/DisplayClusterViewportClient.cpp
+++ b/Engine/Plugins/Runtime/nDisplay/Source/DisplayCluster/Private/Game/EngineClasses/Basics/DisplayClusterViewportClient.cpp
@@ -830,7 +830,6 @@ void UDisplayClusterViewportClient::Draw(FViewport* InViewport, FCanvas* SceneCa
 			ViewFamilies[0]->bAdditionalViewFamily = false;
 			ViewFamilies[0]->bIsFirstViewInMultipleViewFamily = true;
 			ViewFamilies[0]->bIsMultipleViewFamily = true;
-			ViewFamilies[0]->MultipleViewFamilyIndex = 0;
 
 			for (int32 FamilyIndex = 1; FamilyIndex < ViewFamilies.Num(); FamilyIndex++)
 			{
@@ -838,7 +837,6 @@ void UDisplayClusterViewportClient::Draw(FViewport* InViewport, FCanvas* SceneCa
 				ViewFamily.bAdditionalViewFamily = true;
 				ViewFamily.bIsFirstViewInMultipleViewFamily = false;
 				ViewFamily.bIsMultipleViewFamily = true;
-				ViewFamily.MultipleViewFamilyIndex = FamilyIndex;
 			}
 		}
 
diff --git a/Engine/Source/Runtime/Engine/Private/GPUSkinVertexFactory.cpp b/Engine/Source/Runtime/Engine/Private/GPUSkinVertexFactory.cpp
--- a/Engine/Source/Runtime/Engine/Private/GPUSkinVertexFactory.cpp
+++ b/Engine/Source/Runtime/Engine/Private/GPUSkinVertexFactory.cpp
@@ -718,9 +718,7 @@ public:
 
 				// Bone data is updated whenever animation triggers a dynamic update, animation can skip frames hence the frequency is not necessary every frame.
 				// So check if bone data is updated this frame, if not then the previous frame data is stale and not suitable for motion blur.
-				// When using multiple view families in a frame, View->Family->FrameNumber is incremented for each view. The updated frame number for bone data 
-				// though is the frame number from the first view, so add its view index in the view families on top to get the right frame number to compare with.
-				bool bBoneDataUpdatedThisFrame = View->Family->FrameNumber == ShaderData.UpdatedFrameNumber + View->Family->MultipleViewFamilyIndex;
+				bool bBoneDataUpdatedThisFrame = View->Family->FrameNumber == ShaderData.UpdatedFrameNumber;
 				// If world is paused, use current frame bone matrices, so velocity is canceled and skeletal mesh isn't blurred from motion.
 				bool bPrevious = !View->Family->bWorldIsPaused && bBoneDataUpdatedThisFrame;
 				FRHIShaderResourceView* PreviousData = ShaderData.GetBoneBufferForReading(bPrevious).VertexBufferSRV;
@@ -812,9 +810,7 @@ public:
 
 		// Bone data is updated whenever animation triggers a dynamic update, animation can skip frames hence the frequency is not necessary every frame.
 		// So check if bone data is updated this frame, if not then the previous frame data is stale and not suitable for motion blur.
-		// When using multiple view families in a frame, View->Family->FrameNumber is incremented for each view. The updated frame number for bone data 
-		// though is the frame number from the first view, so add its view index in the view families on top to get the right frame number to compare with.
-		bool bBoneDataUpdatedThisFrame = (View->Family->FrameNumber == PassthroughVertexFactory->GetUpdatedFrameNumber() + View->Family->MultipleViewFamilyIndex);
+		bool bBoneDataUpdatedThisFrame = (View->Family->FrameNumber == PassthroughVertexFactory->GetUpdatedFrameNumber());
 		// If world is paused, use current frame bone matrices, so velocity is canceled and skeletal mesh isn't blurred from motion.
 		bool bVerticesInMotion = !View->Family->bWorldIsPaused && bBoneDataUpdatedThisFrame;
 
diff --git a/Engine/Source/Runtime/Engine/Public/SceneView.h b/Engine/Source/Runtime/Engine/Public/SceneView.h
--- a/Engine/Source/Runtime/Engine/Public/SceneView.h
+++ b/Engine/Source/Runtime/Engine/Public/SceneView.h
@@ -1778,9 +1778,6 @@ public:
 	*/
 	bool bIsFirstViewInMultipleViewFamily = true;
 
-	/** The index of this view inside multiple view families rendered in a single frame. */
-	uint8 MultipleViewFamilyIndex = 0;
-
 	bool bIsSceneTexturesInitialized = false;
 	bool bIsViewFamilyInfo = false;
 
diff --git a/Engine/Source/Runtime/Renderer/Private/SceneRendering.cpp b/Engine/Source/Runtime/Renderer/Private/SceneRendering.cpp
--- a/Engine/Source/Runtime/Renderer/Private/SceneRendering.cpp
+++ b/Engine/Source/Runtime/Renderer/Private/SceneRendering.cpp
@@ -4569,7 +4569,10 @@ void FRendererModule::BeginRenderingViewFamilies(FCanvas* Canvas, TArrayView<FSc
 	if (Scene)
 	{
 		// We allow caching of per-frame, per-scene data
-		Scene->IncrementFrameNumber();
+		if (ViewFamilies[0]->bIsFirstViewInMultipleViewFamily)
+		{
+			Scene->IncrementFrameNumber();
+		}
 		for (FSceneViewFamily* ViewFamily : ViewFamilies)
 		{
 			ViewFamily->FrameNumber = Scene->GetFrameNumber();
