BSP: Fix ABrush::BrushType property always showing RevertToDefault icon in details panel

diff --git a/Engine/Source/Runtime/Engine/Classes/Engine/Brush.h b/Engine/Source/Runtime/Engine/Classes/Engine/Brush.h
--- a/Engine/Source/Runtime/Engine/Classes/Engine/Brush.h
+++ b/Engine/Source/Runtime/Engine/Classes/Engine/Brush.h
@@ -78,7 +78,7 @@ class ENGINE_API ABrush
 	GENERATED_UCLASS_BODY()
 
 	/** Type of brush */
-	UPROPERTY(EditAnywhere, Category=Brush)
+	UPROPERTY(EditAnywhere, Category=Brush, meta = (NoResetToDefault))
 	TEnumAsByte<enum EBrushType> BrushType;
 
 	// Information.
