Fix for sky light not working correctly when toggling effect quality level.

diff --git a/Engine/Source/Runtime/Engine/Private/Components/SkyLightComponent.cpp b/Engine/Source/Runtime/Engine/Private/Components/SkyLightComponent.cpp
--- a/Engine/Source/Runtime/Engine/Private/Components/SkyLightComponent.cpp
+++ b/Engine/Source/Runtime/Engine/Private/Components/SkyLightComponent.cpp
@@ -89,11 +89,26 @@ FAutoConsoleVariableRef CVarSkylightIntensityMultiplier(
 	ECVF_Scalability | ECVF_RenderThreadSafe
 	);
 
+void OnChangeSkylightRealTimeReflectionCapture(IConsoleVariable* Var)
+{
+	// r.SkyLight.RealTimeReflectionCapture is set based on the "Effect" quality level (to be supported, or not, on some different platofrms).
+	// When that quality level changes, real-time sky capture can become disabled. In this case, sky light recapture should be scheduled to match the current quality level.
+	for (TObjectIterator<USkyLightComponent> It; It; ++It)
+	{
+		USkyLightComponent* SkylightComponent = *It;
+		if (IsValid(SkylightComponent) && SkylightComponent->IsRenderStateCreated())
+		{
+			SkylightComponent->SetCaptureIsDirty();
+		}
+	}
+}
+
 int32 GSkylightRealTimeReflectionCapture = 1;
 FAutoConsoleVariableRef CVarSkylightRealTimeReflectionCapture(
 	TEXT("r.SkyLight.RealTimeReflectionCapture"),
 	GSkylightRealTimeReflectionCapture,
 	TEXT("Make sure the sky light real time capture is not run on platform where it is considered out of budget. Cannot be changed at runtime."),
+	FConsoleVariableDelegate::CreateStatic(&OnChangeSkylightRealTimeReflectionCapture),
 	ECVF_Scalability
 	);
 
