World Partition: allow changing PostProcessVolume spatially loaded flag only when the volume has unbound set to false.

#rnx

diff --git a/Engine/Source/Runtime/Engine/Classes/Engine/PostProcessVolume.h b/Engine/Source/Runtime/Engine/Classes/Engine/PostProcessVolume.h
--- a/Engine/Source/Runtime/Engine/Classes/Engine/PostProcessVolume.h
+++ b/Engine/Source/Runtime/Engine/Classes/Engine/PostProcessVolume.h
@@ -76,7 +76,7 @@ class ENGINE_API APostProcessVolume : public AVolume, public IInterface_PostProc
 	virtual void PostUnregisterAllComponents( void ) override;
 
 #if WITH_EDITOR
-	virtual bool CanChangeIsSpatiallyLoadedFlag() const override { return false; }
+	virtual bool CanChangeIsSpatiallyLoadedFlag() const override { return !bUnbound; }
 #endif
 
 protected:
diff --git a/Engine/Source/Runtime/Engine/Private/PostProcessVolume.cpp b/Engine/Source/Runtime/Engine/Private/PostProcessVolume.cpp
--- a/Engine/Source/Runtime/Engine/Private/PostProcessVolume.cpp
+++ b/Engine/Source/Runtime/Engine/Private/PostProcessVolume.cpp
@@ -55,9 +55,11 @@ void APostProcessVolume::PostEditChangeProperty(FPropertyChangedEvent& PropertyC
 {
 	Super::PostEditChangeProperty(PropertyChangedEvent);
 
-	static const FName NAME_Blendables = FName(TEXT("Blendables"));
-	
-	if(PropertyChangedEvent.Property && PropertyChangedEvent.Property->GetFName() == NAME_Blendables)
+	static const FName NAME_Blendables = FName(TEXT("Blendables"));	
+
+	const FName ChangedPropertyName = (PropertyChangedEvent.Property != NULL) ? PropertyChangedEvent.Property->GetFName() : NAME_None;	
+
+	if(ChangedPropertyName == NAME_Blendables)
 	{
 		// remove unsupported types
 		uint32 Count = Settings.WeightedBlendables.Array.Num();
@@ -72,7 +74,13 @@ void APostProcessVolume::PostEditChangeProperty(FPropertyChangedEvent& PropertyC
 		}
 	}
 
-	bIsSpatiallyLoaded = !bUnbound;
+	if (ChangedPropertyName == GET_MEMBER_NAME_CHECKED(APostProcessVolume, bUnbound))
+	{
+		if (bUnbound)
+		{
+			bIsSpatiallyLoaded = false;
+		}
+	}	
 	
 	if (PropertyChangedEvent.Property)
 	{
