World Partition: don't use a deferred context when loading always loaded actors for PIE. We must be in immediate mode here so child actors are registered before adding them to the always loaded actors array. This is needed because the actor can contain a construction script that will invalidate the child actor components and recreate them, leaving a dangling pointer in the array since it happens on the loading context destruction.

#rnx

diff --git a/Engine/Source/Runtime/Engine/Private/WorldPartition/WorldPartitionRuntimeSpatialHash.cpp b/Engine/Source/Runtime/Engine/Private/WorldPartition/WorldPartitionRuntimeSpatialHash.cpp
--- a/Engine/Source/Runtime/Engine/Private/WorldPartition/WorldPartitionRuntimeSpatialHash.cpp
+++ b/Engine/Source/Runtime/Engine/Private/WorldPartition/WorldPartitionRuntimeSpatialHash.cpp
@@ -994,47 +994,47 @@ bool UWorldPartitionRuntimeSpatialHash::CreateStreamingGrid(const FSpatialHashRu
 				const bool bIsCellAlwaysLoaded = (&TempCell == &PartionedActors.GetAlwaysLoadedCell()) && !GridCellDataChunk.HasDataLayers();
 				
 				TArray<IStreamingGenerationContext::FActorInstance> FilteredActors;
-				FWorldPartitionLoadingContext::FDeferred LoadingContext;
+				
+				// We must be in immediate mode here so child actors are registered before adding them to the always loaded actors array. This is
+				// needed because the actor can contain a construction script that will invalidate the child actor components and recreate them, 
+				// leaving a dangling pointer in the array since it happens on the loading context destruction.
+				FWorldPartitionLoadingContext::FImmediate LoadingContext;
 
 				for (const IStreamingGenerationContext::FActorSetInstance* ActorSetInstance : GridCellDataChunk.GetActorSetInstances())
 				{
 					for (const FGuid& ActorGuid : ActorSetInstance->ActorSet->Actors)
 					{
-						if (bIsMainWorldPartition && !IsRunningCookCommandlet())
+						// In PIE, Always loaded cell is not generated. Instead, always loaded actors will be added to AlwaysLoadedActorsForPIE.
+						// This will trigger loading/registration of these actors in the PersistentLevel (if not already loaded).
+						// Then, duplication of world for PIE will duplicate only these actors. 
+						// When stopping PIE, WorldPartition will release these FWorldPartitionReferences which 
+						// will unload actors that were not already loaded in the non PIE world.
+						if (bIsMainWorldPartition && bIsCellAlwaysLoaded && ActorSetInstance->ContainerID.IsMainContainer() && !IsRunningCookCommandlet())
 						{
-							const FWorldPartitionActorDescView& ActorDescView = ActorSetInstance->ContainerInstance->ActorDescViewMap->FindByGuidChecked(ActorGuid);
-
-							// In PIE, Always loaded cell is not generated. Instead, always loaded actors will be added to AlwaysLoadedActorsForPIE.
-							// This will trigger loading/registration of these actors in the PersistentLevel (if not already loaded).
-							// Then, duplication of world for PIE will duplicate only these actors. 
-							// When stopping PIE, WorldPartition will release these FWorldPartitionReferences which 
-							// will unload actors that were not already loaded in the non PIE world.
-							if (bIsCellAlwaysLoaded)
+							// This will load the actor if it isn't already loaded
+							FWorldPartitionReference Reference(WorldPartition, ActorGuid);				
+
+							IStreamingGenerationContext::FActorInstance ActorInstance(ActorGuid, ActorSetInstance);
+							const FWorldPartitionActorDescView& ActorDescView = ActorInstance.GetActorDescView();
+
+							if (AActor* AlwaysLoadedActor = FindObject<AActor>(nullptr, *ActorDescView.GetActorPath().ToString()))
 							{
-								if (ActorSetInstance->ContainerID.IsMainContainer())
-								{
-									// This will load the actor if it isn't already loaded
-									FWorldPartitionReference Reference(WorldPartition, ActorGuid);				
+								AlwaysLoadedActorsForPIE.Emplace(Reference, AlwaysLoadedActor);
 
-									if (AActor* AlwaysLoadedActor = FindObject<AActor>(nullptr, *ActorDescView.GetActorPath().ToString()))
+								// Handle child actors
+								AlwaysLoadedActor->ForEachComponent<UChildActorComponent>(true, [this, &Reference](UChildActorComponent* ChildActorComponent)
+								{
+									if (AActor* ChildActor = ChildActorComponent->GetChildActor())
 									{
-										AlwaysLoadedActorsForPIE.Emplace(Reference, AlwaysLoadedActor);
-
-										// Handle child actors
-										AlwaysLoadedActor->ForEachComponent<UChildActorComponent>(true, [this, &Reference](UChildActorComponent* ChildActorComponent)
-										{
-											if (AActor* ChildActor = ChildActorComponent->GetChildActor())
-											{
-												AlwaysLoadedActorsForPIE.Emplace(Reference, ChildActor);
-											}
-										});
+										AlwaysLoadedActorsForPIE.Emplace(Reference, ChildActor);
 									}
-									continue;
-								}
+								});
 							}
 						}
-
-						FilteredActors.Emplace(ActorGuid, ActorSetInstance);
+						else
+						{
+							FilteredActors.Emplace(ActorGuid, ActorSetInstance);
+						}
 					}					
 				}
 
