Moved bulk of uniform expression updates off the render thread.

diff --git a/Engine/Source/Runtime/Apple/MetalRHI/Private/MetalUniformBuffer.cpp b/Engine/Source/Runtime/Apple/MetalRHI/Private/MetalUniformBuffer.cpp
--- a/Engine/Source/Runtime/Apple/MetalRHI/Private/MetalUniformBuffer.cpp
+++ b/Engine/Source/Runtime/Apple/MetalRHI/Private/MetalUniformBuffer.cpp
@@ -22,8 +22,11 @@ FMetalSuballocatedUniformBuffer::FMetalSuballocatedUniformBuffer(const void *Con
     , Validation(InValidation)
 #endif // METAL_UNIFORM_BUFFER_VALIDATION
 {
-	FMemory::Memcpy(Shadow, Contents, GetSize());
-	CopyResourceTable(Contents, ResourceTable);
+	if (Contents)
+	{
+		FMemory::Memcpy(Shadow, Contents, GetSize());
+		CopyResourceTable(Contents, ResourceTable);
+	}
 }
 
 FMetalSuballocatedUniformBuffer::~FMetalSuballocatedUniformBuffer()
diff --git a/Engine/Source/Runtime/D3D12RHI/Private/D3D12UniformBuffer.cpp b/Engine/Source/Runtime/D3D12RHI/Private/D3D12UniformBuffer.cpp
--- a/Engine/Source/Runtime/D3D12RHI/Private/D3D12UniformBuffer.cpp
+++ b/Engine/Source/Runtime/D3D12RHI/Private/D3D12UniformBuffer.cpp
@@ -117,7 +117,7 @@ FUniformBufferRHIRef FD3D12DynamicRHI::RHICreateUniformBuffer(const void* Conten
 {
 	SCOPE_CYCLE_COUNTER(STAT_D3D12UpdateUniformBufferTime);
 
-	if (Validation == EUniformBufferValidation::ValidateResources)
+	if (Contents && Validation == EUniformBufferValidation::ValidateResources)
 	{
 		ValidateShaderParameterResourcesRHI(Contents, *Layout);
 	}
@@ -143,33 +143,38 @@ FUniformBufferRHIRef FD3D12DynamicRHI::RHICreateUniformBuffer(const void* Conten
 			// Create an offline CBV descriptor
 			NewUniformBuffer->View = new FD3D12ConstantBufferView(Device);
 #endif
-			void* MappedData = nullptr;
-			if (Usage == EUniformBufferUsage::UniformBuffer_MultiFrame)
-			{
-				// Uniform buffers that live for multiple frames must use the more expensive and persistent allocation path
-				FD3D12UploadHeapAllocator& Allocator = GetAdapter().GetUploadHeapAllocator(Device->GetGPUIndex());
-				MappedData = Allocator.AllocUploadResource(NumBytesActualData, D3D12_CONSTANT_BUFFER_DATA_PLACEMENT_ALIGNMENT, NewUniformBuffer->ResourceLocation);
-			}
-			else
-			{
-				// Uniform buffers which will live for 1 frame at the max can be allocated very efficiently from a ring buffer
-				FD3D12FastConstantAllocator& Allocator = GetAdapter().GetTransientUniformBufferAllocator();
-				MappedData = Allocator.Allocate(NumBytesActualData, NewUniformBuffer->ResourceLocation, nullptr);
-			}
-			check(NewUniformBuffer->ResourceLocation.GetOffsetFromBaseOfResource() % 16 == 0);
-
-			// Copy the data to the upload heap
-			check(MappedData != nullptr);
 
-			UpdateUniformBufferConstants(MappedData, Contents, Layout);
+			// Uniform buffers can be created without contents and updated later.
+			if (Contents)
+			{
+				void* MappedData = nullptr;
+				if (Usage == EUniformBufferUsage::UniformBuffer_MultiFrame)
+				{
+					// Uniform buffers that live for multiple frames must use the more expensive and persistent allocation path
+					FD3D12UploadHeapAllocator& Allocator = GetAdapter().GetUploadHeapAllocator(Device->GetGPUIndex());
+					MappedData = Allocator.AllocUploadResource(NumBytesActualData, D3D12_CONSTANT_BUFFER_DATA_PLACEMENT_ALIGNMENT, NewUniformBuffer->ResourceLocation);
+				}
+				else
+				{
+					// Uniform buffers which will live for 1 frame at the max can be allocated very efficiently from a ring buffer
+					FD3D12FastConstantAllocator& Allocator = GetAdapter().GetTransientUniformBufferAllocator();
+					MappedData = Allocator.Allocate(NumBytesActualData, NewUniformBuffer->ResourceLocation, nullptr);
+				}
+				check(NewUniformBuffer->ResourceLocation.GetOffsetFromBaseOfResource() % 16 == 0);
+
+				// Copy the data to the upload heap
+				check(MappedData != nullptr);
+
+				UpdateUniformBufferConstants(MappedData, Contents, Layout);
 
 #if USE_STATIC_ROOT_SIGNATURE
-			NewUniformBuffer->View->Create(NewUniformBuffer->ResourceLocation.GetGPUVirtualAddress(), NumBytesActualData);
+				NewUniformBuffer->View->Create(NewUniformBuffer->ResourceLocation.GetGPUVirtualAddress(), NumBytesActualData);
 #endif
+			}
 		}
 
 		// The GPUVA is used to see if this uniform buffer contains constants or is just a resource table.
-		check((NumBytesActualData > 0) ? (0 != NewUniformBuffer->ResourceLocation.GetGPUVirtualAddress()) : (0 == NewUniformBuffer->ResourceLocation.GetGPUVirtualAddress()));
+		check((Contents && NumBytesActualData > 0) ? (0 != NewUniformBuffer->ResourceLocation.GetGPUVirtualAddress()) : (0 == NewUniformBuffer->ResourceLocation.GetGPUVirtualAddress()));
 		return NewUniformBuffer;
 	});
 
@@ -181,11 +186,14 @@ FUniformBufferRHIRef FD3D12DynamicRHI::RHICreateUniformBuffer(const void* Conten
 
 		for (FD3D12UniformBuffer& CurrentBuffer : *UniformBufferOut)
 		{
-			CurrentBuffer.ResourceTable.Empty(NumResources);
-			for (int32 Index = 0; Index < NumResources; ++Index)
+			CurrentBuffer.ResourceTable.SetNumZeroed(NumResources);
+
+			if (Contents)
 			{
-				//Emplace here instead of assignation prevents a lot of boiler plate code for the destruction/release of TSharedPtr from being generated at all.
-				CurrentBuffer.ResourceTable.Emplace(GetShaderParameterResourceRHI(Contents, Layout->Resources[Index].MemberOffset, Layout->Resources[Index].MemberType));
+				for (int32 Index = 0; Index < NumResources; ++Index)
+				{
+					CurrentBuffer.ResourceTable[Index] = GetShaderParameterResourceRHI(Contents, Layout->Resources[Index].MemberOffset, Layout->Resources[Index].MemberType);
+				}
 			}
 		}
 	}
diff --git a/Engine/Source/Runtime/Engine/Private/Materials/MaterialShared.cpp b/Engine/Source/Runtime/Engine/Private/Materials/MaterialShared.cpp
--- a/Engine/Source/Runtime/Engine/Private/Materials/MaterialShared.cpp
+++ b/Engine/Source/Runtime/Engine/Private/Materials/MaterialShared.cpp
@@ -3575,9 +3575,9 @@ void FUniformExpressionCache::ResetAllocatedVTs()
 	OwnedAllocatedVTs.Reset();
 }
 
-void FMaterialRenderProxy::EvaluateUniformExpressions(FUniformExpressionCache& OutUniformExpressionCache, const FMaterialRenderContext& Context, FRHIComputeCommandList*) const
+void FMaterialRenderProxy::EvaluateUniformExpressions(FUniformExpressionCache& OutUniformExpressionCache, const FMaterialRenderContext& Context, FUniformExpressionCacheUpdater* Updater) const
 {
-	check(IsInParallelRenderingThread());
+	check(IsInRenderingThread());
 
 	SCOPE_CYCLE_COUNTER(STAT_CacheUniformExpressions);
 
@@ -3618,27 +3618,42 @@ void FMaterialRenderProxy::EvaluateUniformExpressions(FUniformExpressionCache& O
 	}
 
 	const FRHIUniformBufferLayout* UniformBufferLayout = ShaderMap->GetUniformBufferLayout();
-	FMemMark Mark(FMemStack::Get());
-	uint8* TempBuffer = FMemStack::Get().PushBytes(UniformBufferLayout->ConstantBufferSize, SHADER_PARAMETER_STRUCT_ALIGNMENT);
-
-	check(TempBuffer != nullptr);
-	UniformExpressionSet.FillUniformBuffer(Context, OutUniformExpressionCache, UniformBufferLayout, TempBuffer, UniformBufferLayout->ConstantBufferSize);
-
-	if (IsValidRef(OutUniformExpressionCache.UniformBuffer) && !OutUniformExpressionCache.UniformBuffer->IsValid())
-	{
-		UE_LOG(LogMaterial, Fatal, TEXT("The Uniformbuffer needs to be valid if it has been set"));
-	}
 
 	if (IsValidRef(OutUniformExpressionCache.UniformBuffer))
 	{
+		if (!OutUniformExpressionCache.UniformBuffer->IsValid())
+		{
+			UE_LOG(LogMaterial, Fatal, TEXT("The Uniformbuffer needs to be valid if it has been set"));
+		}
+
 		// The actual pointer may not match because there are cases (in the editor, during the shader compilation) when material's shader map gets updated without proxy's cache
 		// getting invalidated, but the layout contents must match.
 		check(OutUniformExpressionCache.UniformBuffer->GetLayoutPtr() == UniformBufferLayout || *OutUniformExpressionCache.UniformBuffer->GetLayoutPtr() == *UniformBufferLayout);
-		RHIUpdateUniformBuffer(OutUniformExpressionCache.UniformBuffer, TempBuffer);
+	}
+
+	if (Updater)
+	{
+		if (!IsValidRef(OutUniformExpressionCache.UniformBuffer))
+		{
+			OutUniformExpressionCache.UniformBuffer = RHICreateUniformBuffer(nullptr, UniformBufferLayout, UniformBuffer_MultiFrame);
+		}
+
+		Updater->Add(&OutUniformExpressionCache, &UniformExpressionSet, UniformBufferLayout, Context);
 	}
 	else
 	{
-		OutUniformExpressionCache.UniformBuffer = RHICreateUniformBuffer(TempBuffer, UniformBufferLayout, UniformBuffer_MultiFrame);
+		FMemMark Mark(FMemStack::Get());
+		uint8* TempBuffer = FMemStack::Get().PushBytes(UniformBufferLayout->ConstantBufferSize, SHADER_PARAMETER_STRUCT_ALIGNMENT);
+		UniformExpressionSet.FillUniformBuffer(Context, OutUniformExpressionCache, UniformBufferLayout, TempBuffer, UniformBufferLayout->ConstantBufferSize);
+
+		if (IsValidRef(OutUniformExpressionCache.UniformBuffer))
+		{
+			RHIUpdateUniformBuffer(OutUniformExpressionCache.UniformBuffer, TempBuffer);
+		}
+		else
+		{
+			OutUniformExpressionCache.UniformBuffer = RHICreateUniformBuffer(TempBuffer, UniformBufferLayout, UniformBuffer_MultiFrame);
+		}
 	}
 
 	OutUniformExpressionCache.ParameterCollections = UniformExpressionSet.ParameterCollections;
@@ -3877,6 +3892,38 @@ const FMaterial& FMaterialRenderProxy::GetIncompleteMaterialWithFallback(ERHIFea
 	return *Material;
 }
 
+void FUniformExpressionCacheUpdater::Update(FRHICommandListImmediate& RHICmdListImmediate)
+{
+	if (Items.IsEmpty())
+	{
+		return;
+	}
+
+	FRHICommandList* RHICmdList = new FRHICommandList(FRHIGPUMask::All());
+
+	FGraphEventRef Event = FFunctionGraphTask::CreateAndDispatchWhenReady(
+		[Items = MoveTemp(Items), RHICmdList]
+	{
+		TRACE_CPUPROFILER_EVENT_SCOPE(FUniformExpressionCacheUpdater::Update);
+		FTaskTagScope Scope(ETaskTag::EParallelRenderingThread);
+		FMemMark Mark(FMemStack::Get());
+
+		for (const FItem& Item : Items)
+		{
+			uint8* TempBuffer = FMemStack::Get().PushBytes(Item.UniformBufferLayout->ConstantBufferSize, SHADER_PARAMETER_STRUCT_ALIGNMENT);
+
+			FMaterialRenderContext Context(Item.MaterialRenderProxy, *Item.Material, Item.bShowSelection);
+
+			Item.UniformExpressionSet->FillUniformBuffer(Context, *Item.UniformExpressionCache, Item.UniformBufferLayout, TempBuffer, Item.UniformBufferLayout->ConstantBufferSize);
+
+			RHICmdList->UpdateUniformBuffer(Item.UniformExpressionCache->UniformBuffer, TempBuffer);
+		}
+
+	}, TStatId(), nullptr, ENamedThreads::AnyHiPriThreadHiPriTask);
+
+	RHICmdListImmediate.QueueRenderThreadCommandListSubmit(Event, RHICmdList);
+}
+
 void FMaterialRenderProxy::UpdateDeferredCachedUniformExpressions()
 {
 	LLM_SCOPE(ELLMTag::Materials);
@@ -3887,6 +3934,8 @@ void FMaterialRenderProxy::UpdateDeferredCachedUniformExpressions()
 	CSV_SCOPED_TIMING_STAT_EXCLUSIVE(Material_UpdateDeferredCachedUniformExpressions);
 	QUICK_SCOPE_CYCLE_COUNTER(STAT_UpdateDeferredCachedUniformExpressions);
 
+	FUniformExpressionCacheUpdater Updater;
+
 	for (TSet<FMaterialRenderProxy*>::TConstIterator It(DeferredUniformExpressionCacheRequests); It; ++It)
 	{
 		FMaterialRenderProxy* MaterialProxy = *It;
@@ -3903,11 +3952,13 @@ void FMaterialRenderProxy::UpdateDeferredCachedUniformExpressions()
 			{
 				FMaterialRenderContext MaterialRenderContext(MaterialProxy, *Material, nullptr);
 				MaterialRenderContext.bShowSelection = GIsEditor;
-				MaterialProxy->EvaluateUniformExpressions(MaterialProxy->UniformExpressionCache[(int32)InFeatureLevel], MaterialRenderContext);
+				MaterialProxy->EvaluateUniformExpressions(MaterialProxy->UniformExpressionCache[(int32)InFeatureLevel], MaterialRenderContext, &Updater);
 			}
 		});
 	}
 
+	Updater.Update(FRHICommandListExecutor::GetImmediateCommandList());
+
 	DeferredUniformExpressionCacheRequests.Reset();
 }
 
diff --git a/Engine/Source/Runtime/Engine/Public/MaterialShared.h b/Engine/Source/Runtime/Engine/Public/MaterialShared.h
--- a/Engine/Source/Runtime/Engine/Public/MaterialShared.h
+++ b/Engine/Source/Runtime/Engine/Public/MaterialShared.h
@@ -338,6 +338,15 @@ struct ENGINE_API FMaterialRenderContext
 		const FMaterial& InMaterial,
 		const FSceneView* InView);
 
+	FMaterialRenderContext(
+		const FMaterialRenderProxy* InMaterialRenderProxy,
+		const FMaterial& InMaterial,
+		bool bInShowSelection)
+		: MaterialRenderProxy(InMaterialRenderProxy)
+		, Material(InMaterial)
+		, bShowSelection(bInShowSelection)
+	{}
+
 	void GetTextureParameterValue(const FHashedMaterialParameterInfo& ParameterInfo, int32 TextureIndex, const UTexture*& OutValue) const;
 	void GetTextureParameterValue(const FHashedMaterialParameterInfo& ParameterInfo, int32 TextureIndex, const URuntimeVirtualTexture*& OutValue) const;
 	FGuid GetExternalTextureGuid(const FGuid& ExternalTextureGuid, const FName& ParameterName, int32 SourceTextureIndex) const;
@@ -2471,6 +2480,40 @@ private:
 
 class USubsurfaceProfile;
 
+class ENGINE_API FUniformExpressionCacheUpdater
+{
+public:
+	void Add(FUniformExpressionCache* UniformExpressionCache, const FUniformExpressionSet* UniformExpressionSet, const FRHIUniformBufferLayout* UniformBufferLayout, const FMaterialRenderContext& Context)
+	{
+		Items.Emplace(UniformExpressionCache, UniformExpressionSet, UniformBufferLayout, Context);
+	}
+
+	void Update(FRHICommandListImmediate& RHICmdList);
+
+private:
+	struct FItem
+	{
+		FItem() = default;
+		FItem(FUniformExpressionCache* InUniformExpressionCache, const FUniformExpressionSet* InUniformExpressionSet, const FRHIUniformBufferLayout* InUniformBufferLayout, const FMaterialRenderContext& Context)
+			: UniformExpressionCache(InUniformExpressionCache)
+			, UniformExpressionSet(InUniformExpressionSet)
+			, UniformBufferLayout(InUniformBufferLayout)
+			, MaterialRenderProxy(Context.MaterialRenderProxy)
+			, Material(&Context.Material)
+			, bShowSelection(Context.bShowSelection)
+		{}
+
+		FUniformExpressionCache* UniformExpressionCache = nullptr;
+		const FUniformExpressionSet* UniformExpressionSet = nullptr;
+		const FRHIUniformBufferLayout* UniformBufferLayout = nullptr;
+		const FMaterialRenderProxy* MaterialRenderProxy = nullptr;
+		const FMaterial* Material = nullptr;
+		bool bShowSelection = false;
+	};
+
+	TArray<FItem, FConcurrentLinearArrayAllocator> Items;
+};
+
 /**
  * A material render proxy used by the renderer.
  */
@@ -2490,12 +2533,18 @@ public:
 	/** Destructor. */
 	virtual ~FMaterialRenderProxy();
 
+	UE_DEPRECATED(5.1, "EvaluateUniformExpressions with a command list is deprecated.")
+	void EvaluateUniformExpressions(FUniformExpressionCache& OutUniformExpressionCache, const FMaterialRenderContext& Context, class FRHIComputeCommandList*) const
+	{
+		EvaluateUniformExpressions(OutUniformExpressionCache, Context);
+	}
+
 	/**
 	 * Evaluates uniform expressions and stores them in OutUniformExpressionCache.
 	 * @param OutUniformExpressionCache - The uniform expression cache to build.
 	 * @param MaterialRenderContext - The context for which to cache expressions.
 	 */
-	void EvaluateUniformExpressions(FUniformExpressionCache& OutUniformExpressionCache, const FMaterialRenderContext& Context, class FRHIComputeCommandList* CommandListIfLocalMode = nullptr) const;
+	void EvaluateUniformExpressions(FUniformExpressionCache& OutUniformExpressionCache, const FMaterialRenderContext& Context, FUniformExpressionCacheUpdater* Updater = nullptr) const;
 
 	/**
 	 * Caches uniform expressions for efficient runtime evaluation.
diff --git a/Engine/Source/Runtime/OpenGLDrv/Private/OpenGLUniformBuffer.cpp b/Engine/Source/Runtime/OpenGLDrv/Private/OpenGLUniformBuffer.cpp
--- a/Engine/Source/Runtime/OpenGLDrv/Private/OpenGLUniformBuffer.cpp
+++ b/Engine/Source/Runtime/OpenGLDrv/Private/OpenGLUniformBuffer.cpp
@@ -489,9 +489,12 @@ static void SetLayoutTable(FOpenGLUniformBuffer* NewUniformBuffer, const void* C
 		NewUniformBuffer->ResourceTable.Empty(NumResources);
 		NewUniformBuffer->ResourceTable.AddZeroed(NumResources);
 
-		for (int32 Index = 0; Index < NumResources; ++Index)
+		if (Contents)
 		{
-			NewUniformBuffer->ResourceTable[Index] = GetShaderParameterResourceRHI(Contents, Layout->Resources[Index].MemberOffset, Layout->Resources[Index].MemberType);
+			for (int32 Index = 0; Index < NumResources; ++Index)
+			{
+				NewUniformBuffer->ResourceTable[Index] = GetShaderParameterResourceRHI(Contents, Layout->Resources[Index].MemberOffset, Layout->Resources[Index].MemberType);
+			}
 		}
 	}
 }
@@ -629,7 +632,10 @@ static FUniformBufferRHIRef CreateUniformBuffer(const void* Contents, const FRHI
 
 	check(!bUseEmulatedUBs || (IsValidRef(EmulatedUniformDataRef) && (EmulatedUniformDataRef->Data.Num() * EmulatedUniformDataRef->Data.GetTypeSize() == AllocatedSize)));
 
-	CopyDataToUniformBuffer(bCanCreateOnThisThread, NewUniformBuffer,Contents, Layout->ConstantBufferSize);
+	if (Contents)
+	{
+		CopyDataToUniformBuffer(bCanCreateOnThisThread, NewUniformBuffer, Contents, Layout->ConstantBufferSize);
+	}
 
 	// Initialize the resource table for this uniform buffer.
 	SetLayoutTable(NewUniformBuffer, Contents, Layout, Validation);
@@ -743,7 +749,10 @@ FUniformBufferRHIRef FOpenGLDynamicRHI::RHICreateUniformBuffer(const void* Conte
 	NewUniformBuffer->SetGLUniformBufferParams(AllocatedResource, OffsetInBuffer, PersistentlyMappedBuffer, AllocatedSize, EmulatedUniformDataRef, bStreamDraw);
 
 	check(!bUseEmulatedUBs || (IsValidRef(EmulatedUniformDataRef) && (EmulatedUniformDataRef->Data.Num() * EmulatedUniformDataRef->Data.GetTypeSize() == AllocatedSize)));
-	CopyDataToUniformBuffer(bCanCreateOnThisThread, NewUniformBuffer, Contents, Layout->ConstantBufferSize);
+	if (Contents)
+	{
+		CopyDataToUniformBuffer(bCanCreateOnThisThread, NewUniformBuffer, Contents, Layout->ConstantBufferSize);
+	}
 
 	// Initialize the resource table for this uniform buffer.
 	SetLayoutTable(NewUniformBuffer, Contents, Layout, Validation);
diff --git a/Engine/Source/Runtime/RHI/Private/RHIValidation.cpp b/Engine/Source/Runtime/RHI/Private/RHIValidation.cpp
--- a/Engine/Source/Runtime/RHI/Private/RHIValidation.cpp
+++ b/Engine/Source/Runtime/RHI/Private/RHIValidation.cpp
@@ -1581,10 +1581,11 @@ namespace RHIValidation
 	}
 
 
-	void FUniformBufferResource::InitLifetimeTracking(uint64 FrameID, EUniformBufferUsage Usage)
+	void FUniformBufferResource::InitLifetimeTracking(uint64 FrameID, const void* Contents, EUniformBufferUsage Usage)
 	{
 		AllocatedFrameID = FrameID;
 		UniformBufferUsage = Usage;
+		bContainsNullContents = Contents == nullptr;
 
 #if CAPTURE_UNIFORMBUFFER_ALLOCATION_BACKTRACES
 		AllocatedCallstack = (UniformBufferUsage != UniformBuffer_MultiFrame) ? RHIValidation::CaptureBacktrace() : nullptr;
@@ -1596,6 +1597,7 @@ namespace RHIValidation
 	void FUniformBufferResource::UpdateAllocation(uint64 FrameID)
 	{
 		AllocatedFrameID = FrameID;
+		bContainsNullContents = false;
 
 #if CAPTURE_UNIFORMBUFFER_ALLOCATION_BACKTRACES
 		AllocatedCallstack = (UniformBufferUsage != UniformBuffer_MultiFrame) ? RHIValidation::CaptureBacktrace() : nullptr;
@@ -1607,6 +1609,9 @@ namespace RHIValidation
 	void FUniformBufferResource::ValidateLifeTime()
 	{
 		FValidationRHI* ValidateRHI = (FValidationRHI*)GDynamicRHI;
+
+		RHI_VALIDATION_CHECK(bContainsNullContents == false, TEXT("Uniform buffer created with null contents is now being bound for rendering on an RHI context. The contents must first be updated."));
+
 		if (UniformBufferUsage != UniformBuffer_MultiFrame && AllocatedFrameID != ValidateRHI->RHIThreadFrameID)
 		{
 			FString ErrorMessage = TEXT("Non MultiFrame Uniform buffer has been allocated in a previous frame. The data could have been deleted already!");
@@ -1615,7 +1620,7 @@ namespace RHIValidation
 				ErrorMessage += FString::Printf(TEXT("\nAllocation callstack: (void**)0x%p,32"), AllocatedCallstack);
 			}
 			RHI_VALIDATION_CHECK(false, *ErrorMessage);
-		}		
+		}
 	}
 
 	
diff --git a/Engine/Source/Runtime/RHI/Public/RHIValidation.h b/Engine/Source/Runtime/RHI/Public/RHIValidation.h
--- a/Engine/Source/Runtime/RHI/Public/RHIValidation.h
+++ b/Engine/Source/Runtime/RHI/Public/RHIValidation.h
@@ -362,15 +362,18 @@ public:
 	// FlushType: Thread safe, but varies depending on the RHI override final
 	virtual FUniformBufferRHIRef RHICreateUniformBuffer(const void* Contents, const FRHIUniformBufferLayout* Layout, EUniformBufferUsage Usage, EUniformBufferValidation Validation) override final
 	{
+		check(Layout);
 		check(Layout->Resources.Num() > 0 || Layout->ConstantBufferSize > 0);
 		FUniformBufferRHIRef UniformBuffer = RHI->RHICreateUniformBuffer(Contents, Layout, Usage, Validation);
 		// Use the render thread frame ID for any non RHI thread allocations.
-		UniformBuffer->InitLifetimeTracking(IsInRHIThread() ? RHIThreadFrameID : RenderThreadFrameID, Usage);
+		UniformBuffer->InitLifetimeTracking(IsInRHIThread() ? RHIThreadFrameID : RenderThreadFrameID, Contents, Usage);
 		return UniformBuffer;
 	}
 
 	virtual void RHIUpdateUniformBuffer(FRHICommandListBase& RHICmdList, FRHIUniformBuffer* UniformBufferRHI, const void* Contents) override final
 	{
+		check(UniformBufferRHI);
+		check(Contents);
 		RHI->RHIUpdateUniformBuffer(RHICmdList, UniformBufferRHI, Contents);
 		UniformBufferRHI->UpdateAllocation(RenderThreadFrameID);
 	}
diff --git a/Engine/Source/Runtime/RHI/Public/RHIValidationCommon.h b/Engine/Source/Runtime/RHI/Public/RHIValidationCommon.h
--- a/Engine/Source/Runtime/RHI/Public/RHIValidationCommon.h
+++ b/Engine/Source/Runtime/RHI/Public/RHIValidationCommon.h
@@ -588,10 +588,11 @@ namespace RHIValidation
 	struct FUniformBufferResource
 	{
 		uint64 AllocatedFrameID = 0;
+		bool bContainsNullContents = false;
 		EUniformBufferUsage UniformBufferUsage;
 		void* AllocatedCallstack;
 
-		void InitLifetimeTracking(uint64 FrameID, EUniformBufferUsage Usage);
+		void InitLifetimeTracking(uint64 FrameID, const void* Contents, EUniformBufferUsage Usage);
 		void UpdateAllocation(uint64 FrameID);
 		void ValidateLifeTime();
 	};	
diff --git a/Engine/Source/Runtime/VulkanRHI/Private/VulkanMemory.cpp b/Engine/Source/Runtime/VulkanRHI/Private/VulkanMemory.cpp
--- a/Engine/Source/Runtime/VulkanRHI/Private/VulkanMemory.cpp
+++ b/Engine/Source/Runtime/VulkanRHI/Private/VulkanMemory.cpp
@@ -3561,8 +3561,12 @@ namespace VulkanRHI
 			HandleOOM(false);
 			checkNoEntry();
 		}
-		FMemory::Memcpy(OutAllocation.GetMappedPointer(Device), Contents, Size);
-		OutAllocation.FlushMappedMemory(Device);
+
+		if (Contents)
+		{
+			FMemory::Memcpy(OutAllocation.GetMappedPointer(Device), Contents, Size);
+			OutAllocation.FlushMappedMemory(Device);
+		}
 	}
 	void FMemoryManager::FreeUniformBuffer(FVulkanAllocation& InAllocation)
 	{
diff --git a/Engine/Source/Runtime/VulkanRHI/Private/VulkanUniformBuffer.cpp b/Engine/Source/Runtime/VulkanRHI/Private/VulkanUniformBuffer.cpp
--- a/Engine/Source/Runtime/VulkanRHI/Private/VulkanUniformBuffer.cpp
+++ b/Engine/Source/Runtime/VulkanRHI/Private/VulkanUniformBuffer.cpp
@@ -113,9 +113,13 @@ FVulkanUniformBuffer::FVulkanUniformBuffer(FVulkanDevice& Device, const FRHIUnif
 		// Transfer the resource table to an internal resource-array
 		ResourceTable.Empty(NumResources);
 		ResourceTable.AddZeroed(NumResources);
-		for (uint32 Index = 0; Index < NumResources; ++Index)
+
+		if (Contents)
 		{
-			ResourceTable[Index] = GetShaderParameterResourceRHI(Contents, InLayout->Resources[Index].MemberOffset, InLayout->Resources[Index].MemberType);
+			for (uint32 Index = 0; Index < NumResources; ++Index)
+			{
+				ResourceTable[Index] = GetShaderParameterResourceRHI(Contents, InLayout->Resources[Index].MemberOffset, InLayout->Resources[Index].MemberType);
+			}
 		}
 	}
 
@@ -126,29 +130,32 @@ FVulkanUniformBuffer::FVulkanUniformBuffer(FVulkanDevice& Device, const FRHIUnif
 
 		if (UseRingBuffer(InUsage) && (bInRenderingThread || bInRHIThread))
 		{
-			FRHICommandListImmediate& RHICmdList = FRHICommandListExecutor::GetImmediateCommandList();
-			FVulkanUniformBuffer* UniformBuffer = this;
-			int32 DataSize = InLayout->ConstantBufferSize;
-			
-			// make sure we allocate from RingBuffer on RHIT
-			const bool bCanAllocOnThisThread = RHICmdList.Bypass() || (!IsRunningRHIInSeparateThread() && bInRenderingThread) || bInRHIThread;
-			if (bCanAllocOnThisThread)
+			if (Contents)
 			{
-				FVulkanCommandListContextImmediate& Context = Device.GetImmediateContext();
-				UpdateUniformBufferHelper(Context, UniformBuffer, DataSize, Contents);
-			}
-			else
-			{
-				void* CmdListConstantBufferData = RHICmdList.Alloc(DataSize, 16);
-				FMemory::Memcpy(CmdListConstantBufferData, Contents, DataSize);
-				
-				RHICmdList.EnqueueLambda([UniformBuffer, DataSize, CmdListConstantBufferData](FRHICommandList& CmdList)
+				FRHICommandListImmediate& RHICmdList = FRHICommandListExecutor::GetImmediateCommandList();
+				FVulkanUniformBuffer* UniformBuffer = this;
+				int32 DataSize = InLayout->ConstantBufferSize;
+
+				// make sure we allocate from RingBuffer on RHIT
+				const bool bCanAllocOnThisThread = RHICmdList.Bypass() || (!IsRunningRHIInSeparateThread() && bInRenderingThread) || bInRHIThread;
+				if (bCanAllocOnThisThread)
 				{
-					FVulkanCommandListContext& Context = FVulkanCommandListContext::GetVulkanContext(CmdList.GetContext());
-					UpdateUniformBufferHelper(Context, UniformBuffer, DataSize, CmdListConstantBufferData);
-				});
+					FVulkanCommandListContextImmediate& Context = Device.GetImmediateContext();
+					UpdateUniformBufferHelper(Context, UniformBuffer, DataSize, Contents);
+				}
+				else
+				{
+					void* CmdListConstantBufferData = RHICmdList.Alloc(DataSize, 16);
+					FMemory::Memcpy(CmdListConstantBufferData, Contents, DataSize);
+
+					RHICmdList.EnqueueLambda([UniformBuffer, DataSize, CmdListConstantBufferData](FRHICommandList& CmdList)
+					{
+						FVulkanCommandListContext& Context = FVulkanCommandListContext::GetVulkanContext(CmdList.GetContext());
+						UpdateUniformBufferHelper(Context, UniformBuffer, DataSize, CmdListConstantBufferData);
+					});
 
-				RHICmdList.RHIThreadFence(true);
+					RHICmdList.RHIThreadFence(true);
+				}
 			}
 		}
 		else
diff --git a/Engine/Source/Runtime/Windows/D3D11RHI/Private/D3D11UniformBuffer.cpp b/Engine/Source/Runtime/Windows/D3D11RHI/Private/D3D11UniformBuffer.cpp
--- a/Engine/Source/Runtime/Windows/D3D11RHI/Private/D3D11UniformBuffer.cpp
+++ b/Engine/Source/Runtime/Windows/D3D11RHI/Private/D3D11UniformBuffer.cpp
@@ -154,19 +154,22 @@ static TRefCountPtr<ID3D11Buffer> CreateAndUpdatePooledUniformBuffer(
 
 	check(IsValidRef(UniformBufferResource));
 
-	D3D11_MAPPED_SUBRESOURCE MappedSubresource;
-	// Discard previous results since we always do a full update
-	VERIFYD3D11RESULT_EX(Context->Map(UniformBufferResource, 0, D3D11_MAP_WRITE_DISCARD, 0, &MappedSubresource), Device);
-	check(MappedSubresource.RowPitch >= NumBytes);
-	FMemory::Memcpy(MappedSubresource.pData, Contents, NumBytes);
-	Context->Unmap(UniformBufferResource, 0);
+	if (Contents)
+	{
+		D3D11_MAPPED_SUBRESOURCE MappedSubresource;
+		// Discard previous results since we always do a full update
+		VERIFYD3D11RESULT_EX(Context->Map(UniformBufferResource, 0, D3D11_MAP_WRITE_DISCARD, 0, &MappedSubresource), Device);
+		check(MappedSubresource.RowPitch >= NumBytes);
+		FMemory::Memcpy(MappedSubresource.pData, Contents, NumBytes);
+		Context->Unmap(UniformBufferResource, 0);
+	}
 
 	return UniformBufferResource;
 }
 
 FUniformBufferRHIRef FD3D11DynamicRHI::RHICreateUniformBuffer(const void* Contents, const FRHIUniformBufferLayout* Layout, EUniformBufferUsage Usage, EUniformBufferValidation Validation)
 {
-	if (Validation == EUniformBufferValidation::ValidateResources)
+	if (Contents && Validation == EUniformBufferValidation::ValidateResources)
 	{
 		ValidateShaderParameterResourcesRHI(Contents, *Layout);
 	}
@@ -201,8 +204,14 @@ FUniformBufferRHIRef FD3D11DynamicRHI::RHICreateUniformBuffer(const void* Conten
 			{
 				NewUniformBuffer = new FD3D11UniformBuffer(this, Layout, nullptr, FRingAllocation(), bAllocatedFromPool);
 				NewUniformBuffer->AddRef();
-				void* CPUContent = FMemory::Malloc(NumBytes);
-				FMemory::Memcpy(CPUContent, Contents, NumBytes);
+
+				void* CPUContent = nullptr;
+				
+				if (Contents)
+				{
+					CPUContent = FMemory::Malloc(NumBytes);
+					FMemory::Memcpy(CPUContent, Contents, NumBytes);
+				}
 
 				RunOnRHIThread(
 					[NewUniformBuffer, CPUContent, NumBytes]()
@@ -254,10 +263,13 @@ FUniformBufferRHIRef FD3D11DynamicRHI::RHICreateUniformBuffer(const void* Conten
 		NewUniformBuffer->ResourceTable.Empty(ResourceCount);
 		NewUniformBuffer->ResourceTable.AddZeroed(ResourceCount);
 
-		for (int32 Index = 0; Index < ResourceCount; ++Index)
+		if (Contents)
 		{
-			const auto ResourceParameter = Layout->Resources[Index];
-			NewUniformBuffer->ResourceTable[Index] = GetShaderParameterResourceRHI(Contents, ResourceParameter.MemberOffset, ResourceParameter.MemberType);
+			for (int32 Index = 0; Index < ResourceCount; ++Index)
+			{
+				const auto ResourceParameter = Layout->Resources[Index];
+				NewUniformBuffer->ResourceTable[Index] = GetShaderParameterResourceRHI(Contents, ResourceParameter.MemberOffset, ResourceParameter.MemberType);
+			}
 		}
 	}
 
