Fix typo.

#rnx

diff --git a/Engine/Source/Runtime/Engine/Private/WorldPartition/WorldPartitionLevelStreamingDynamic.cpp b/Engine/Source/Runtime/Engine/Private/WorldPartition/WorldPartitionLevelStreamingDynamic.cpp
--- a/Engine/Source/Runtime/Engine/Private/WorldPartition/WorldPartitionLevelStreamingDynamic.cpp
+++ b/Engine/Source/Runtime/Engine/Private/WorldPartition/WorldPartitionLevelStreamingDynamic.cpp
@@ -435,11 +435,11 @@ void UWorldPartitionLevelStreamingDynamic::OnCleanupLevel()
 
 		RuntimeLevel->OnCleanupLevel.Remove(OnCleanupLevelDelegateHandle);
 
-		TSet<UPackage*> TrashedPachages;
-		auto TrashPackage = [&TrashedPachages](UPackage* Package)
+		TSet<UPackage*> TrashedPackages;
+		auto TrashPackage = [&TrashedPackages](UPackage* Package)
 		{
 			bool bWasAlreadyInSet;
-			TrashedPachages.Add(Package, &bWasAlreadyInSet);
+			TrashedPackages.Add(Package, &bWasAlreadyInSet);
 
 			if (!bWasAlreadyInSet)
 			{
