Add whether or not a contact was a probe contact to external collision event data

diff --git a/Engine/Source/Runtime/Engine/Classes/Engine/EngineTypes.h b/Engine/Source/Runtime/Engine/Classes/Engine/EngineTypes.h
--- a/Engine/Source/Runtime/Engine/Classes/Engine/EngineTypes.h
+++ b/Engine/Source/Runtime/Engine/Classes/Engine/EngineTypes.h
@@ -1587,6 +1587,10 @@ struct ENGINE_API FRigidBodyContactInfo
 	UPROPERTY()
 	float ContactPenetration;
 
+	/** Was this contact generated by a probe constraint */
+	UPROPERTY()
+	bool bContactProbe;
+
 	/** The physical material of the two shapes involved in a contact */
 	UPROPERTY()
 	TObjectPtr<class UPhysicalMaterial> PhysMaterial[2];
@@ -1596,6 +1600,7 @@ struct ENGINE_API FRigidBodyContactInfo
 		: ContactPosition(ForceInit)
 		, ContactNormal(ForceInit)
 		, ContactPenetration(0)
+		, bContactProbe(false)
 	{
 		for (int32 ElementIndex = 0; ElementIndex < 2; ElementIndex++)
 		{
@@ -1606,11 +1611,13 @@ struct ENGINE_API FRigidBodyContactInfo
 	FRigidBodyContactInfo(	const FVector& InContactPosition, 
 							const FVector& InContactNormal, 
 							float InPenetration, 
+							bool bInProbe,
 							UPhysicalMaterial* InPhysMat0, 
 							UPhysicalMaterial* InPhysMat1 )
 		: ContactPosition(InContactPosition)
 		, ContactNormal(InContactNormal)
 		, ContactPenetration(InPenetration)
+		, bContactProbe(bInProbe)
 	{
 		PhysMaterial[0] = InPhysMat0;
 		PhysMaterial[1] = InPhysMat1;
diff --git a/Engine/Source/Runtime/Engine/Private/PhysicsEngine/Experimental/PhysScene_Chaos.cpp b/Engine/Source/Runtime/Engine/Private/PhysicsEngine/Experimental/PhysScene_Chaos.cpp
--- a/Engine/Source/Runtime/Engine/Private/PhysicsEngine/Experimental/PhysScene_Chaos.cpp
+++ b/Engine/Source/Runtime/Engine/Private/PhysicsEngine/Experimental/PhysScene_Chaos.cpp
@@ -762,6 +762,7 @@ void FPhysScene_Chaos::HandleCollisionEvents(const Chaos::FCollisionEventData& E
 								NewContact.ContactNormal = CollisionDataItem.Normal;
 								NewContact.ContactPosition = CollisionDataItem.Location;
 								NewContact.ContactPenetration = CollisionDataItem.PenetrationDepth;
+								NewContact.bContactProbe = CollisionDataItem.bProbe;
 								NotifyInfo.RigidCollisionData.bIsVelocityDeltaUnderThreshold = CollisionDataItem.DeltaVelocity1.IsNearlyZero(MinDeltaVelocityThreshold) && CollisionDataItem.DeltaVelocity2.IsNearlyZero(MinDeltaVelocityThreshold);
 								
 								Chaos::FChaosPhysicsMaterial* InternalMat1 = CollisionDataItem.Mat1.Get();
diff --git a/Engine/Source/Runtime/Experimental/Chaos/Private/EventDefaults.cpp b/Engine/Source/Runtime/Experimental/Chaos/Private/EventDefaults.cpp
--- a/Engine/Source/Runtime/Experimental/Chaos/Private/EventDefaults.cpp
+++ b/Engine/Source/Runtime/Experimental/Chaos/Private/EventDefaults.cpp
@@ -160,8 +160,6 @@ namespace Chaos
 								const FKinematicGeometryParticleHandle* Secondary = Body0 ? Body1 : Body0;
 
 								//
-								const bool bIsProbe = Constraint.IsProbe();
-								const bool bUseManifold = Constraint.GetUseManifold();
 								const int32 NumManifoldPoints = Constraint.GetManifoldPoints().Num();
 
 								if (((Constraint.IsProbe() && NumManifoldPoints > 0) || !Constraint.AccumulatedImpulse.IsZero()) && Primary)
@@ -210,6 +208,7 @@ namespace Chaos
 								Data.AccumulatedImpulse = Constraint.AccumulatedImpulse;
 								Data.Normal = Constraint.CalculateWorldContactNormal();
 								Data.PenetrationDepth = Constraint.GetPhi();
+								Data.bProbe = Constraint.GetIsProbe();
 
 								// @todo(chaos): fix this casting
 								Data.Proxy1 = Particle0 ? const_cast<IPhysicsProxyBase*>(Particle0->PhysicsProxy()) : nullptr;
diff --git a/Engine/Source/Runtime/Experimental/Chaos/Public/Chaos/ExternalCollisionData.h b/Engine/Source/Runtime/Experimental/Chaos/Public/Chaos/ExternalCollisionData.h
--- a/Engine/Source/Runtime/Experimental/Chaos/Public/Chaos/ExternalCollisionData.h
+++ b/Engine/Source/Runtime/Experimental/Chaos/Public/Chaos/ExternalCollisionData.h
@@ -30,6 +30,9 @@ namespace Chaos
 			, Mass1((FReal)0.0)
 			, Mass2((FReal)0.0)
 			, PenetrationDepth((FReal)0.0)
+			, Mat1(FMaterialHandle())
+			, Mat2(FMaterialHandle())
+			, bProbe(false)
 			, Proxy1(nullptr)
 			, Proxy2(nullptr)
 		{}
@@ -48,6 +51,9 @@ namespace Chaos
 			, Mass1(InMass1)
 			, Mass2(InMass2)
 			, PenetrationDepth(InPenetrationDepth)
+			, Mat1(FMaterialHandle())
+			, Mat2(FMaterialHandle())
+			, bProbe(false)
 			, Proxy1(InProxy1)
 			, Proxy2(InProxy2)
 		{}
@@ -66,6 +72,7 @@ namespace Chaos
 		FReal PenetrationDepth;
 		FMaterialHandle Mat1;
 		FMaterialHandle Mat2;
+		bool bProbe;
 
 		// The pointers to the proxies should be used with caution on the Game Thread.
 		// Ideally we only ever use these as table keys when acquiring related structures.
