Add new LevelStreamingDynamic::LoadLevelInstance that take input by a parameter struct.
Add new parameter to control the initial visibility of instanced level.
Add IsInitiallyVisible to LevelInstanceInterface to allow a level instance to control if the level should be made visible.
#rnx

diff --git a/Engine/Source/Runtime/Engine/Classes/Engine/LevelStreamingDynamic.h b/Engine/Source/Runtime/Engine/Classes/Engine/LevelStreamingDynamic.h
--- a/Engine/Source/Runtime/Engine/Classes/Engine/LevelStreamingDynamic.h
+++ b/Engine/Source/Runtime/Engine/Classes/Engine/LevelStreamingDynamic.h
@@ -28,6 +28,36 @@ class ENGINE_API ULevelStreamingDynamic : public ULevelStreaming
 	UPROPERTY(Category=LevelStreaming, EditAnywhere)
 	uint32 bInitiallyVisible:1;
 	
+	struct FLoadLevelInstanceParams
+	{
+		FLoadLevelInstanceParams(UWorld* InWorld, const FString& InLongPackageName, FTransform InLevelTransform)
+			: World(InWorld)
+			, LongPackageName(InLongPackageName)
+			, LevelTransform(InLevelTransform) 
+		{}
+
+		/** World to instance the level into. */
+		UWorld* World = nullptr;
+
+		/** Level long package name to load. */
+		FString LongPackageName;
+
+		/** Transform of the instanced level. */
+		FTransform LevelTransform;
+
+		/** If set, the loaded level package have this name, which is used by other functions like UnloadStreamLevel. Note this is necessary for server and client networking because the level must have the same name on both. */
+		const FString* OptionalLevelNameOverride = nullptr;
+
+		/** If set, the level streaming class will be used instead of ULevelStreamingDynamic. */
+		TSubclassOf<ULevelStreamingDynamic> OptionalLevelStreamingClass = nullptr;
+
+		/** If set, package path is prefixed by /Temp. */
+		bool bLoadAsTempPackage = false;
+
+		/** Set whether the level will be made visible initially. */
+		bool bInitiallyVisible = true;
+	};
+
 	/**  
  	* Stream in a level with a specific location and rotation. You can create multiple instances of the same level!
  	*
@@ -53,6 +83,8 @@ class ENGINE_API ULevelStreamingDynamic : public ULevelStreaming
  	
 	static ULevelStreamingDynamic* LoadLevelInstanceBySoftObjectPtr(UObject* WorldContextObject, TSoftObjectPtr<UWorld> Level, const FTransform LevelTransform, bool& bOutSuccess, const FString& OptionalLevelNameOverride = TEXT(""), TSubclassOf<ULevelStreamingDynamic> OptionalLevelStreamingClass = nullptr, bool bLoadAsTempPackage = false);
 
+	static ULevelStreamingDynamic* LoadLevelInstance(const FLoadLevelInstanceParams& Params, bool& bOutSuccess);
+
 	//~ Begin UObject Interface
 	virtual void PostLoad() override;
 	//~ End UObject Interface
@@ -68,7 +100,7 @@ private:
 	// Counter used by LoadLevelInstance to create unique level names
 	static int32 UniqueLevelInstanceId;
 
- 	static ULevelStreamingDynamic* LoadLevelInstance_Internal(UWorld* World, const FString& LongPackageName, const FTransform LevelTransform, bool& bOutSuccess, const FString& OptionalLevelNameOverride, TSubclassOf<ULevelStreamingDynamic> OptionalLevelStreamingClass, bool bLoadAsTempPackage);
+ 	static ULevelStreamingDynamic* LoadLevelInstance_Internal(const FLoadLevelInstanceParams& Params, bool& bOutSuccess);
 
 };
 
diff --git a/Engine/Source/Runtime/Engine/Private/LevelInstance/LevelInstanceInterface.cpp b/Engine/Source/Runtime/Engine/Private/LevelInstance/LevelInstanceInterface.cpp
--- a/Engine/Source/Runtime/Engine/Private/LevelInstance/LevelInstanceInterface.cpp
+++ b/Engine/Source/Runtime/Engine/Private/LevelInstance/LevelInstanceInterface.cpp
@@ -60,6 +60,16 @@ void ILevelInstanceInterface::UnloadLevelInstance()
 	}
 }
 
+ULevelStreamingLevelInstance* ILevelInstanceInterface::GetLevelStreaming() const
+{
+	if (ULevelInstanceSubsystem* LevelInstanceSubsystem = GetLevelInstanceSubsystem())
+	{
+		return LevelInstanceSubsystem->GetLevelInstanceLevelStreaming(this);
+	}
+
+	return nullptr;
+}
+
 #if WITH_EDITOR
 ULevel* ILevelInstanceInterface::GetLoadedLevel() const
 {
diff --git a/Engine/Source/Runtime/Engine/Private/LevelInstance/LevelInstanceLevelStreaming.cpp b/Engine/Source/Runtime/Engine/Private/LevelInstance/LevelInstanceLevelStreaming.cpp
--- a/Engine/Source/Runtime/Engine/Private/LevelInstance/LevelInstanceLevelStreaming.cpp
+++ b/Engine/Source/Runtime/Engine/Private/LevelInstance/LevelInstanceLevelStreaming.cpp
@@ -198,14 +198,26 @@ ULevelStreamingLevelInstance* ULevelStreamingLevelInstance::LoadInstance(ILevelI
 	// Build a unique and deterministic LevelInstance level instance name by using LevelInstanceID. 
 	// Distinguish game from editor since we don't want to duplicate for PIE already loaded editor instances (not yet supported).
 	const FString Suffix = FString::Printf(TEXT("%s_LevelInstance_%016llx_%d"), *ShortPackageName, LevelInstance->GetLevelInstanceID().GetHash(), LevelInstanceActor->GetWorld()->IsGameWorld() ? 1 : 0);
-	const bool bLoadAsTempPackage = true;
-	ULevelStreamingLevelInstance* LevelStreaming = Cast<ULevelStreamingLevelInstance>(ULevelStreamingDynamic::LoadLevelInstanceBySoftObjectPtr(LevelInstanceActor->GetWorld(), LevelInstance->GetWorldAsset(), LevelInstanceActor->GetActorTransform(), bOutSuccess, Suffix, ULevelStreamingLevelInstance::StaticClass(), bLoadAsTempPackage));
+	
+	UWorld* World = LevelInstanceActor->GetWorld();
+
+	FLoadLevelInstanceParams Params(World, LevelInstance->GetWorldAssetPackage(), LevelInstanceActor->GetActorTransform());
+	Params.OptionalLevelNameOverride = &Suffix;
+	Params.OptionalLevelStreamingClass = ULevelStreamingLevelInstance::StaticClass();
+	Params.bLoadAsTempPackage = true;
+	
+	if (World->IsGameWorld())
+	{
+		Params.bInitiallyVisible = LevelInstance->IsInitiallyVisible();
+	}
+	
+	ULevelStreamingLevelInstance* LevelStreaming = Cast<ULevelStreamingLevelInstance>(ULevelStreamingDynamic::LoadLevelInstance(Params, bOutSuccess));
 	if (bOutSuccess)
 	{
 		LevelStreaming->LevelInstanceID = LevelInstance->GetLevelInstanceID();
 		
 #if WITH_EDITOR
-		if (!LevelInstanceActor->GetWorld()->IsGameWorld())
+		if (!World->IsGameWorld())
 		{
 			GEngine->BlockTillLevelStreamingCompleted(LevelInstanceActor->GetWorld());
 
diff --git a/Engine/Source/Runtime/Engine/Private/LevelInstance/LevelInstanceSubsystem.cpp b/Engine/Source/Runtime/Engine/Private/LevelInstance/LevelInstanceSubsystem.cpp
--- a/Engine/Source/Runtime/Engine/Private/LevelInstance/LevelInstanceSubsystem.cpp
+++ b/Engine/Source/Runtime/Engine/Private/LevelInstance/LevelInstanceSubsystem.cpp
@@ -430,6 +430,19 @@ void ULevelInstanceSubsystem::ForEachLevelInstanceAncestorsAndSelf(AActor* Actor
 	ForEachLevelInstanceAncestors(Actor, Operation);
 }
 
+ULevelStreamingLevelInstance* ULevelInstanceSubsystem::GetLevelInstanceLevelStreaming(const ILevelInstanceInterface* LevelInstance) const
+{
+	if (LevelInstance->HasValidLevelInstanceID())
+	{
+		if (const FLevelInstance* LevelInstanceEntry = LoadedLevelInstances.Find(LevelInstance->GetLevelInstanceID()))
+		{
+			return LevelInstanceEntry->LevelStreaming;
+		}
+	}
+
+	return nullptr;
+}
+
 void ULevelInstanceSubsystem::ForEachLevelInstanceAncestors(AActor* Actor, TFunctionRef<bool(ILevelInstanceInterface*)> Operation) const
 {
 	ILevelInstanceInterface* ParentLevelInstance = nullptr;
diff --git a/Engine/Source/Runtime/Engine/Private/LevelStreaming.cpp b/Engine/Source/Runtime/Engine/Private/LevelStreaming.cpp
--- a/Engine/Source/Runtime/Engine/Private/LevelStreaming.cpp
+++ b/Engine/Source/Runtime/Engine/Private/LevelStreaming.cpp
@@ -2162,9 +2162,13 @@ ULevelStreamingDynamic* ULevelStreamingDynamic::LoadLevelInstance(UObject* World
 	{
 		return nullptr;
 	}
-	FString LongPackageNameObjectPath = ExistingPackageName.ToString() + ObjectRelativePath;
+	const FString LongPackageNameObjectPath = ExistingPackageName.ToString() + ObjectRelativePath;
 
-	return LoadLevelInstance_Internal(World, LongPackageNameObjectPath, FTransform(Rotation, Location), bOutSuccess, OptionalLevelNameOverride, OptionalLevelStreamingClass, bLoadAsTempPackage);
+	FLoadLevelInstanceParams Params(World, LongPackageNameObjectPath, FTransform(Rotation, Location));
+	Params.OptionalLevelNameOverride = OptionalLevelNameOverride.IsEmpty() ? nullptr : &OptionalLevelNameOverride;
+	Params.OptionalLevelStreamingClass = OptionalLevelStreamingClass;
+	Params.bLoadAsTempPackage = bLoadAsTempPackage;
+	return LoadLevelInstance(Params, bOutSuccess);
 }
 
 ULevelStreamingDynamic* ULevelStreamingDynamic::LoadLevelInstanceBySoftObjectPtr(UObject* WorldContextObject, const TSoftObjectPtr<UWorld> Level, const FVector Location, const FRotator Rotation, bool& bOutSuccess, const FString& OptionalLevelNameOverride, TSubclassOf<ULevelStreamingDynamic> OptionalLevelStreamingClass, bool bLoadAsTempPackage)
@@ -2175,30 +2179,46 @@ ULevelStreamingDynamic* ULevelStreamingDynamic::LoadLevelInstanceBySoftObjectPtr
 ULevelStreamingDynamic* ULevelStreamingDynamic::LoadLevelInstanceBySoftObjectPtr(UObject* WorldContextObject, const TSoftObjectPtr<UWorld> Level, const FTransform LevelTransform, bool& bOutSuccess, const FString& OptionalLevelNameOverride, TSubclassOf<ULevelStreamingDynamic> OptionalLevelStreamingClass, bool bLoadAsTempPackage)
 {
 	bOutSuccess = false;
+	if (Level.IsNull())
+	{
+		return nullptr;
+	}
+
 	UWorld* World = GEngine->GetWorldFromContextObject(WorldContextObject, EGetWorldErrorMode::LogAndReturnNull);
-	if (!World)
+
+	FLoadLevelInstanceParams Params(World, Level.GetLongPackageName(), LevelTransform);
+	Params.OptionalLevelNameOverride = OptionalLevelNameOverride.IsEmpty() ? nullptr : &OptionalLevelNameOverride;
+	Params.OptionalLevelStreamingClass = OptionalLevelStreamingClass;
+	Params.bLoadAsTempPackage = bLoadAsTempPackage;
+	return LoadLevelInstance(Params, bOutSuccess);
+}
+
+ULevelStreamingDynamic* ULevelStreamingDynamic::LoadLevelInstance(const FLoadLevelInstanceParams& Params, bool& bOutSuccess)
+{
+	bOutSuccess = false;
+	if (!Params.World)
 	{
 		return nullptr;
 	}
 
-	// Check whether requested map exists, this could be very slow if LevelName is a short package name
-	if (Level.IsNull())
+	if (Params.LongPackageName.IsEmpty())
 	{
 		return nullptr;
 	}
 
-	return LoadLevelInstance_Internal(World, Level.GetLongPackageName(), LevelTransform, bOutSuccess, OptionalLevelNameOverride, OptionalLevelStreamingClass, bLoadAsTempPackage);
+	return LoadLevelInstance_Internal(Params, bOutSuccess);
 }
 
-ULevelStreamingDynamic* ULevelStreamingDynamic::LoadLevelInstance_Internal(UWorld* World, const FString& LongPackageName, const FTransform LevelTransform, bool& bOutSuccess, const FString& OptionalLevelNameOverride, TSubclassOf<ULevelStreamingDynamic> OptionalLevelStreamingClass, bool bLoadAsTempPackage)
+
+ULevelStreamingDynamic* ULevelStreamingDynamic::LoadLevelInstance_Internal(const FLoadLevelInstanceParams& Params, bool& bOutSuccess)
 {
-	const FString PackagePath = FPackageName::GetLongPackagePath(LongPackageName);
-	FString ShortPackageName = FPackageName::GetShortName(LongPackageName);
-	UClass* LevelStreamingClass = OptionalLevelStreamingClass != nullptr ? OptionalLevelStreamingClass.Get() : ULevelStreamingDynamic::StaticClass();
+	const FString PackagePath = FPackageName::GetLongPackagePath(Params.LongPackageName);
+	FString ShortPackageName = FPackageName::GetShortName(Params.LongPackageName);
+	UClass* LevelStreamingClass = Params.OptionalLevelStreamingClass != nullptr ? Params.OptionalLevelStreamingClass.Get() : ULevelStreamingDynamic::StaticClass();
 
-	if (ShortPackageName.StartsWith(World->StreamingLevelsPrefix))
+	if (ShortPackageName.StartsWith(Params.World->StreamingLevelsPrefix))
 	{
-		ShortPackageName.RightChopInline(World->StreamingLevelsPrefix.Len(), false);
+		ShortPackageName.RightChopInline(Params.World->StreamingLevelsPrefix.Len(), false);
 	}
 
 	// Remove PIE prefix if it's there before we actually load the level
@@ -2206,7 +2226,7 @@ ULevelStreamingDynamic* ULevelStreamingDynamic::LoadLevelInstance_Internal(UWorl
 
 	// Determine loaded package name
 	TStringBuilder<512> LevelPackageNameStrBuilder;
-	if (bLoadAsTempPackage)
+	if (Params.bLoadAsTempPackage)
 	{
 		LevelPackageNameStrBuilder.Append(TEXT("/Temp"));
 	}
@@ -2214,27 +2234,27 @@ ULevelStreamingDynamic* ULevelStreamingDynamic::LoadLevelInstance_Internal(UWorl
 	LevelPackageNameStrBuilder.Append(TEXT("/"));
 		
 	bool bNeedsUniqueTest = false;
-	if (OptionalLevelNameOverride.IsEmpty())
+	if (Params.OptionalLevelNameOverride)
 	{
-		LevelPackageNameStrBuilder.Append(ShortPackageName);
-		LevelPackageNameStrBuilder.Append(TEXT("_LevelInstance_"));
-		LevelPackageNameStrBuilder.Append(FString::FromInt(++UniqueLevelInstanceId));
+		// Use the supplied suffix, which is expected to result in a unique package name but we have to check if it is not.
+		LevelPackageNameStrBuilder.Append(*Params.OptionalLevelNameOverride);
+		bNeedsUniqueTest = true;
 	}
 	else
 	{
-		// Use the supplied suffix, which is expected to result in a unique package name but we have to check if it is not.
-		LevelPackageNameStrBuilder.Append(OptionalLevelNameOverride);
-		bNeedsUniqueTest = true;
+		LevelPackageNameStrBuilder.Append(ShortPackageName);
+		LevelPackageNameStrBuilder.Append(TEXT("_LevelInstance_"));
+		LevelPackageNameStrBuilder.Append(FString::FromInt(++UniqueLevelInstanceId));
 	}
 
 	const FString LevelPackageNameStr(LevelPackageNameStrBuilder.ToString());
 	const FName UnmodifiedLevelPackageName = FName(*LevelPackageNameStr);
 #if WITH_EDITOR
-	const bool bIsPlayInEditor = World->IsPlayInEditor();
+	const bool bIsPlayInEditor = Params.World->IsPlayInEditor();
 	int32 PIEInstance = INDEX_NONE;
 	if (bIsPlayInEditor)
 	{
-		const FWorldContext& WorldContext = GEngine->GetWorldContextFromWorldChecked(World);
+		const FWorldContext& WorldContext = GEngine->GetWorldContextFromWorldChecked(Params.World);
 		PIEInstance = WorldContext.PIEInstance;
 	}
 #endif
@@ -2247,7 +2267,7 @@ ULevelStreamingDynamic* ULevelStreamingDynamic::LoadLevelInstance_Internal(UWorl
 			ModifiedLevelPackageName = FName(*UWorld::ConvertToPIEPackageName(LevelPackageNameStr, PIEInstance));
 		}
 #endif
-		if (World->GetStreamingLevels().ContainsByPredicate([&ModifiedLevelPackageName](ULevelStreaming* LS) { return LS && LS->GetWorldAssetPackageFName() == ModifiedLevelPackageName; }))
+		if (Params.World->GetStreamingLevels().ContainsByPredicate([&ModifiedLevelPackageName](ULevelStreaming* LS) { return LS && LS->GetWorldAssetPackageFName() == ModifiedLevelPackageName; }))
 		{
 			// The streaming level already exists, error and return.
 			UE_LOG(LogLevelStreaming, Error, TEXT("LoadLevelInstance called with a name that already exists, returning nullptr. LevelPackageName:%s"), *ModifiedLevelPackageName.ToString());
@@ -2256,7 +2276,7 @@ ULevelStreamingDynamic* ULevelStreamingDynamic::LoadLevelInstance_Internal(UWorl
 	}
     
 	// Setup streaming level object that will load specified map
-	ULevelStreamingDynamic* StreamingLevel = NewObject<ULevelStreamingDynamic>(World, LevelStreamingClass, NAME_None, RF_Transient, NULL);
+	ULevelStreamingDynamic* StreamingLevel = NewObject<ULevelStreamingDynamic>(Params.World, LevelStreamingClass, NAME_None, RF_Transient, NULL);
 
 	FSoftObjectPath WorldAssetPath(*WriteToString<512>(UnmodifiedLevelPackageName, TEXT("."), ShortPackageName));
     StreamingLevel->SetWorldAsset(TSoftObjectPtr<UWorld>(WorldAssetPath));
@@ -2269,17 +2289,17 @@ ULevelStreamingDynamic* ULevelStreamingDynamic::LoadLevelInstance_Internal(UWorl
 #endif // WITH_EDITOR
     StreamingLevel->LevelColor = FColor::MakeRandomColor();
     StreamingLevel->SetShouldBeLoaded(true);
-    StreamingLevel->SetShouldBeVisible(true);
+    StreamingLevel->SetShouldBeVisible(Params.bInitiallyVisible);
     StreamingLevel->bShouldBlockOnLoad = false;
     StreamingLevel->bInitiallyLoaded = true;
-    StreamingLevel->bInitiallyVisible = true;
+    StreamingLevel->bInitiallyVisible = Params.bInitiallyVisible;
 	// Transform
-    StreamingLevel->LevelTransform = LevelTransform;
+    StreamingLevel->LevelTransform = Params.LevelTransform;
 	// Map to Load
     StreamingLevel->PackageNameToLoad = FName(*OnDiskPackageName);
           
     // Add the new level to world.
-    World->AddStreamingLevel(StreamingLevel);
+    Params.World->AddStreamingLevel(StreamingLevel);
       
 	bOutSuccess = true;
     return StreamingLevel;
diff --git a/Engine/Source/Runtime/Engine/Public/LevelInstance/LevelInstanceInterface.h b/Engine/Source/Runtime/Engine/Public/LevelInstance/LevelInstanceInterface.h
--- a/Engine/Source/Runtime/Engine/Public/LevelInstance/LevelInstanceInterface.h
+++ b/Engine/Source/Runtime/Engine/Public/LevelInstance/LevelInstanceInterface.h
@@ -9,6 +9,7 @@
 #include "LevelInstanceInterface.generated.h"
 
 class ULevelInstanceComponent;
+class ULevelStreamingLevelInstance;
 
 UINTERFACE()
 class ENGINE_API ULevelInstanceInterface : public UInterface
@@ -64,7 +65,11 @@ class ENGINE_API ILevelInstanceInterface
 	virtual void LoadLevelInstance();
 
 	virtual void UnloadLevelInstance();
-		
+
+	virtual bool IsInitiallyVisible() const { return true; }
+
+	virtual ULevelStreamingLevelInstance* GetLevelStreaming() const;
+
 #if WITH_EDITOR
 	virtual ULevel* GetLoadedLevel() const;
 
diff --git a/Engine/Source/Runtime/Engine/Public/LevelInstance/LevelInstanceSubsystem.h b/Engine/Source/Runtime/Engine/Public/LevelInstance/LevelInstanceSubsystem.h
--- a/Engine/Source/Runtime/Engine/Public/LevelInstance/LevelInstanceSubsystem.h
+++ b/Engine/Source/Runtime/Engine/Public/LevelInstance/LevelInstanceSubsystem.h
@@ -52,6 +52,7 @@ public:
 	void RequestUnloadLevelInstance(ILevelInstanceInterface* LevelInstance);
 	bool IsLoaded(const ILevelInstanceInterface* LevelInstance) const;
 	void ForEachLevelInstanceAncestorsAndSelf(AActor* Actor, TFunctionRef<bool(ILevelInstanceInterface*)> Operation) const;
+	ULevelStreamingLevelInstance* GetLevelInstanceLevelStreaming(const ILevelInstanceInterface* LevelInstance) const;
 
 #if WITH_EDITOR
 	void Tick();
