Remove reference from Player to PlayerController when PlayerController removes its reference to Player in OnNetCleanup

#localization none
#tests Tested match to match transition on PC

diff --git a/Engine/Source/Runtime/Engine/Private/PlayerController.cpp b/Engine/Source/Runtime/Engine/Private/PlayerController.cpp
--- a/Engine/Source/Runtime/Engine/Private/PlayerController.cpp
+++ b/Engine/Source/Runtime/Engine/Private/PlayerController.cpp
@@ -1361,6 +1361,10 @@ void APlayerController::OnNetCleanup(UNetConnection* Connection)
 	//@note: if we ever implement support for splitscreen players leaving a match without the primary player leaving, we'll need to insert
 	// a call to ClearOnlineDelegates() here so that PlayerController.ClearOnlineDelegates can use the correct ControllerId (which lives
 	// in ULocalPlayer)
+	if (Player && Player->PlayerController == this)
+	{
+		Player->PlayerController = nullptr;
+	}
 	Player = NULL;
 	NetConnection = NULL;	
 	Destroy( true );
