Chaos: Re-enable scene-wide physics collision handler customization.

diff --git a/Engine/Source/Runtime/Engine/Private/PhysicsEngine/Experimental/PhysScene_Chaos.cpp b/Engine/Source/Runtime/Engine/Private/PhysicsEngine/Experimental/PhysScene_Chaos.cpp
--- a/Engine/Source/Runtime/Engine/Private/PhysicsEngine/Experimental/PhysScene_Chaos.cpp
+++ b/Engine/Source/Runtime/Engine/Private/PhysicsEngine/Experimental/PhysScene_Chaos.cpp
@@ -17,6 +17,7 @@
 #include "Components/SkeletalMeshComponent.h"
 #include "PhysicsReplication.h"
 #include "Physics/Experimental/PhysicsUserData_Chaos.h"
+#include "PhysicsEngine/PhysicsCollisionHandler.h"
 
 #include "PhysicsSolver.h"
 #include "ChaosSolversModule.h"
@@ -789,13 +790,13 @@ void FPhysScene_Chaos::HandleCollisionEvents(const Chaos::FCollisionEventData& E
 
 void FPhysScene_Chaos::DispatchPendingCollisionNotifies()
 {
-	//UWorld const* const OwningWorld = GetWorld();
+	UWorld const* const OwningWorld = GetOwningWorld();
 
-	//// Let the game-specific PhysicsCollisionHandler process any physics collisions that took place
-	//if (OwningWorld != nullptr && OwningWorld->PhysicsCollisionHandler != nullptr)
-	//{
-	//	OwningWorld->PhysicsCollisionHandler->HandlePhysicsCollisions_AssumesLocked(PendingCollisionNotifies);
-	//}
+	// Let the game-specific PhysicsCollisionHandler process any physics collisions that took place
+	if (OwningWorld != nullptr && OwningWorld->PhysicsCollisionHandler != nullptr)
+	{
+		OwningWorld->PhysicsCollisionHandler->HandlePhysicsCollisions_AssumesLocked(PendingCollisionNotifies);
+	}
 
 	// Fire any collision notifies in the queue.
 	for (FCollisionNotifyInfo& NotifyInfo : PendingCollisionNotifies)
