Skeletal Mesh:
- Deprecate and redirect USkinnedMeshComponent::SetSkeletalMesh into SetSkinnedMeshAndUpdate.
- Rename SetSkinnedMesh to SetSkinnedMeshAndUpdate to allow default parameter to work in the Blueprints' UI.

diff --git a/Engine/Config/BaseEngine.ini b/Engine/Config/BaseEngine.ini
--- a/Engine/Config/BaseEngine.ini
+++ b/Engine/Config/BaseEngine.ini
@@ -1110,6 +1110,9 @@ DefaultWorldPartitionHLODLayer=/Engine/Maps/Templates/HLODs/HLODLayer_Instanced.
 +FunctionRedirects=(OldName="SkeletalMeshComponent.BindClothToMasterPoseComponent", NewName="BindClothToLeaderPoseComponent")
 +FunctionRedirects=(OldName="SkeletalMeshComponent.UnbindClothFromMasterPoseComponent", NewName="UnbindClothFromLeaderPoseComponent")
 
+; Skeletal Mesh pointer deprecation
++FunctionRedirects=(OldName="SkinnedMeshComponent.SetSkeletalMesh", NewName="SetSkinnedAssetAndUpdate")
+
 [CoreUObject.Metadata]
 ; Note: UnrealHeaderTool should be rerun after making changes to MetadataRedirects to catch any keys specified in class headers
 +MetadataRedirects=(OldKey="K2Protected", NewKey="BlueprintProtected")
diff --git a/Engine/Plugins/Experimental/Animation/PoseSearch/Source/Editor/Private/PoseSearchDebugger.cpp b/Engine/Plugins/Experimental/Animation/PoseSearch/Source/Editor/Private/PoseSearchDebugger.cpp
--- a/Engine/Plugins/Experimental/Animation/PoseSearch/Source/Editor/Private/PoseSearchDebugger.cpp
+++ b/Engine/Plugins/Experimental/Animation/PoseSearch/Source/Editor/Private/PoseSearchDebugger.cpp
@@ -2481,9 +2481,9 @@ void FDebuggerViewModel::UpdateFromTimeline()
 			USkeletalMesh* SkeletalMesh = TSoftObjectPtr<USkeletalMesh>(FSoftObjectPath(SkeletalMeshObjectInfo->PathName)).LoadSynchronous();
 			if (SkeletalMesh)
 			{
-				ActiveComponent->SetSkinnedAsset(SkeletalMesh, true);
-				SelectedComponent->SetSkinnedAsset(SkeletalMesh, true);
-				AssetComponent->SetSkinnedAsset(SkeletalMesh, true);
+				ActiveComponent->SetSkinnedAssetAndUpdate(SkeletalMesh, true);
+				SelectedComponent->SetSkinnedAssetAndUpdate(SkeletalMesh, true);
+				AssetComponent->SetSkinnedAssetAndUpdate(SkeletalMesh, true);
 			}
 			FTransform ComponentWorldTransform;
 			// Active skeleton is simply the traced bone transforms
diff --git a/Engine/Plugins/Runtime/Oculus/OculusVR/Source/OculusInput/Private/OculusHandComponent.cpp b/Engine/Plugins/Runtime/Oculus/OculusVR/Source/OculusInput/Private/OculusHandComponent.cpp
--- a/Engine/Plugins/Runtime/Oculus/OculusVR/Source/OculusInput/Private/OculusHandComponent.cpp
+++ b/Engine/Plugins/Runtime/Oculus/OculusVR/Source/OculusInput/Private/OculusHandComponent.cpp
@@ -51,7 +51,7 @@ void UOculusHandComponent::InitializeSkeletalMesh()
 	{
 		if (UOculusInputFunctionLibrary::GetHandSkeletalMesh(RuntimeSkeletalMesh, SkeletonType, MeshType))
 		{
-			SetSkinnedAsset(RuntimeSkeletalMesh, true);
+			SetSkinnedAssetAndUpdate(RuntimeSkeletalMesh, true);
 			if (MaterialOverride)
 			{
 				SetMaterial(0, MaterialOverride);
diff --git a/Engine/Source/Editor/UnrealEd/Private/EditorEngine.cpp b/Engine/Source/Editor/UnrealEd/Private/EditorEngine.cpp
--- a/Engine/Source/Editor/UnrealEd/Private/EditorEngine.cpp
+++ b/Engine/Source/Editor/UnrealEd/Private/EditorEngine.cpp
@@ -3591,7 +3591,7 @@ struct FConvertStaticMeshActorInfo
 		InternalSetToActor(Actor);
 
 		// Set component properties.
-		if ( bComponentPropsDifferFromDefaults[0] ) MeshComp->SetSkinnedAsset(SkeletalMesh);
+		if ( bComponentPropsDifferFromDefaults[0] ) MeshComp->SetSkeletalMeshAsset(SkeletalMesh);
 		if ( bComponentPropsDifferFromDefaults[1] ) MeshComp->OverrideMaterials		= OverrideMaterials;
 		if ( bComponentPropsDifferFromDefaults[2] ) MeshComp->CachedMaxDrawDistance	= CachedMaxDrawDistance;
 		if ( bComponentPropsDifferFromDefaults[3] ) MeshComp->CastShadow			= CastShadow;
diff --git a/Engine/Source/Runtime/Engine/Classes/Components/SkeletalMeshComponent.h b/Engine/Source/Runtime/Engine/Classes/Components/SkeletalMeshComponent.h
--- a/Engine/Source/Runtime/Engine/Classes/Components/SkeletalMeshComponent.h
+++ b/Engine/Source/Runtime/Engine/Classes/Components/SkeletalMeshComponent.h
@@ -360,7 +360,7 @@ public:
 	 * Set the SkeletalMesh rendered for this mesh.
 	 */
 	UFUNCTION(BlueprintCallable, Category = "Components|SkeletalMesh")
-	void SetSkeletalMeshAsset(class USkeletalMesh* NewMesh) { SetSkeletalMesh(NewMesh, true); }
+	void SetSkeletalMeshAsset(class USkeletalMesh* NewMesh) { SetSkeletalMesh(NewMesh, false); }
 
 	/**
 	 * Get the SkeletalMesh rendered for this mesh.
@@ -1888,7 +1888,16 @@ public:
 	virtual void HideBone( int32 BoneIndex, EPhysBodyOp PhysBodyOption ) override;
 	virtual void UnHideBone( int32 BoneIndex ) override;
 	virtual void SetPhysicsAsset(class UPhysicsAsset* NewPhysicsAsset,bool bForceReInit = false) override;
-	virtual void SetSkeletalMesh(class USkeletalMesh* NewMesh, bool bReinitPose = true) override;
+	virtual void SetSkeletalMesh(class USkeletalMesh* NewMesh, bool bReinitPose = true) override;  // SetSkeletalMesh may remain and become a UFUNCTION but lose its virtual after it is removed from the SkinnedMeshComponent API
+
+	/**
+	 *  Set the new asset to render and update this component (this function is identical to SetSkeletalMesh).
+	 *  The USkinnedAsset pointer is first cast to a USkeletalMesh, therefore this function will only set assets of type USkeletalMesh.
+	 * 
+	 *  @param InSkinnedAsset The new asset.
+	 *  @param bReinitPose    Whether to re-initialize the animation.
+	 */
+	virtual void SetSkinnedAssetAndUpdate(class USkinnedAsset* InSkinnedAsset, bool bReinitPose = true) override;
 
 	// GetSkeletalMesh() is soon to be removed, use GetSkeletalMeshAsset() instead
 	class USkeletalMesh* GetSkeletalMesh() const { return GetSkeletalMeshAsset(); }
@@ -1900,6 +1909,7 @@ public:
 	static void GetSkinnedTangentBasis(USkeletalMeshComponent* Component, int32 VertexIndex, const FSkeletalMeshLODRenderData& Model, const FSkinWeightVertexBuffer& SkinWeightBuffer, TArray<FMatrix44f>& CachedRefToLocals, FVector3f& OutTangentX, FVector3f& OutTangentY, FVector3f& OutTangentZ);
 	static void ComputeSkinnedTangentBasis(USkeletalMeshComponent* Component, TArray<FVector3f>& OutTangenXYZ, TArray<FMatrix44f>& CachedRefToLocals, const FSkeletalMeshLODRenderData& Model, const FSkinWeightVertexBuffer& SkinWeightBuffer);
 
+	UE_DEPRECATED(5.1, "This method has been deprecated. Please use SetSkeletalMesh(NewMesh, false) instead.")
 	void SetSkeletalMeshWithoutResettingAnimation(class USkeletalMesh* NewMesh);
 
 	virtual bool IsPlayingRootMotion() const override;
diff --git a/Engine/Source/Runtime/Engine/Classes/Components/SkinnedMeshComponent.h b/Engine/Source/Runtime/Engine/Classes/Components/SkinnedMeshComponent.h
--- a/Engine/Source/Runtime/Engine/Classes/Components/SkinnedMeshComponent.h
+++ b/Engine/Source/Runtime/Engine/Classes/Components/SkinnedMeshComponent.h
@@ -266,8 +266,8 @@ class ENGINE_API USkinnedMeshComponent : public UMeshComponent, public ILODSyncI
 	friend class FSkinnedMeshComponentRecreateRenderStateContext;
 
 	/** The skeletal mesh used by this component. */
-	UE_DEPRECATED(5.1, "Replaced by SkinnedAsset. Use GetSkinnedAsset()/SetSkinnedAsset() instead.")
-	UPROPERTY(EditAnywhere, BlueprintGetter = GetSkeletalMesh_DEPRECATED, Category = "Mesh|SkeletalAsset", meta = (DisallowedClasses = "/Script/ApexDestruction.DestructibleMesh", DeprecatedProperty, DeprecationMessage = "Use USkeletalMeshComponent::GetSkeletalMesh() or GetSkinnedAsset() instead."))
+	UE_DEPRECATED(5.1, "Replaced by SkinnedAsset. Use GetSkinnedAsset()/SetSkinnedAsset() instead, or GetSkeletalMeshAsset/SetSkeletalMeshAsset() when called from a USkeletalMeshComponent.")
+	UPROPERTY(EditAnywhere, BlueprintGetter = GetSkeletalMesh_DEPRECATED, Category = "Mesh|SkeletalAsset", meta = (DisallowedClasses = "/Script/ApexDestruction.DestructibleMesh", DeprecatedProperty, DeprecationMessage = "Use USkeletalMeshComponent::GetSkeletalMeshAsset() or GetSkinnedAsset() instead."))
 	TObjectPtr<class USkeletalMesh> SkeletalMesh;
 
 private:
@@ -915,8 +915,7 @@ public:
 	 * @param NewMesh New mesh to set for this component
 	 * @param bReinitPose Whether we should keep current pose or reinitialize.
 	 */
-	UE_DEPRECATED(5.1, "Use USkeletalMeshComponent::SetSkeletalMesh() or SetSkinnedAsset() instead.")
-	UFUNCTION(BlueprintCallable, Category="Components|SkinnedMesh")
+	UE_DEPRECATED(5.1, "Use USkeletalMeshComponent::SetSkeletalMesh() or SetSkinnedAssetAndUpdate() instead.")
 	virtual void SetSkeletalMesh(class USkeletalMesh* NewMesh, bool bReinitPose = true);
 
 	/**
@@ -927,17 +926,17 @@ public:
 	 * @return the SkeletalMesh set to this mesh.
 	 */
 	UE_DEPRECATED(5.1, "Use USkeletalMeshComponent::GetSkeletalMeshAsset() or GetSkinnedAsset() instead.")
-	UFUNCTION(BlueprintPure, Category = "Components|SkinnedMesh", DisplayName = "Get Skeletal Mesh", meta = (DeprecatedFunction, DeprecationMessage = "Use USkeletalMeshComponent::GetSkeletalMesh() or GetSkinnedAsset() instead."))
+	UFUNCTION(BlueprintPure, Category = "Components|SkinnedMesh", DisplayName = "Get Skeletal Mesh", meta = (DeprecatedFunction, DeprecationMessage = "Use USkeletalMeshComponent::GetSkeletalMeshAsset() or GetSkinnedAsset() instead."))
 	class USkeletalMesh* GetSkeletalMesh_DEPRECATED() const;
 
 	/**
 	 * Change the SkinnedAsset that is rendered for this Component. Will re-initialize the animation tree etc.
 	 *
-	 * @param NewSkinnedAsset New mesh to set for this component
+	 * @param NewMesh New mesh to set for this component
 	 * @param bReinitPose Whether we should keep current pose or reinitialize.
 	 */
 	UFUNCTION(BlueprintCallable, Category = "Components|SkinnedMesh")
-	void SetSkinnedAsset(class USkinnedAsset* InSkinnedAsset, bool bReinitPose);
+	virtual void SetSkinnedAssetAndUpdate(class USkinnedAsset* NewMesh, bool bReinitPose = true);
 
 	/**
 	 * Change the SkinnedAsset that is rendered without reinitializing this Component.
diff --git a/Engine/Source/Runtime/Engine/Private/Components/SkeletalMeshComponent.cpp b/Engine/Source/Runtime/Engine/Private/Components/SkeletalMeshComponent.cpp
--- a/Engine/Source/Runtime/Engine/Private/Components/SkeletalMeshComponent.cpp
+++ b/Engine/Source/Runtime/Engine/Private/Components/SkeletalMeshComponent.cpp
@@ -2885,7 +2885,12 @@ void USkeletalMeshComponent::SetSkeletalMesh(USkeletalMesh* InSkelMesh, bool bRe
 
 	{
 		FRenderStateRecreator RenderStateRecreator(this);
-		SetSkinnedAsset(InSkelMesh, bReinitPose);
+
+		// The SetSkeletalMesh base implementation is being phased out of USkinnedMeshComponent, and SetSkinnedAssetAndUpdate replaces it,
+		// but the SetSkeletalMesh function will remain to provide continuity and a more convenient USkeletalMesh based API avoiding the uncertainty of the Cast to USkeletalMesh.
+		// Also, since USkeletalMeshComponent::SetSkinnedAssetAndUpdate now calls USkeletalMeshComponent::SetSkeletalMesh, and since SetSkinnedAssetAndUpdate is virtual,
+		// the Super keyword is necessary to make sure that the correct USkinnedMeshComponent base function is called and prevent a recursive call loop.
+		Super::SetSkinnedAssetAndUpdate(InSkelMesh, bReinitPose);
 		
 #if WITH_EDITOR
 		ValidateAnimation();
@@ -2922,6 +2927,11 @@ void USkeletalMeshComponent::SetSkeletalMesh(USkeletalMesh* InSkelMesh, bool bRe
 	IStreamingManager::Get().NotifyPrimitiveUpdated(this);
 }
 
+void USkeletalMeshComponent::SetSkinnedAssetAndUpdate(USkinnedAsset* InSkinnedAsset, bool bReinitPose)
+{
+	SetSkeletalMesh(Cast<USkeletalMesh>(InSkinnedAsset), bReinitPose);
+}
+
 void USkeletalMeshComponent::SetSkeletalMeshWithoutResettingAnimation(USkeletalMesh* InSkelMesh)
 {
 	SetSkeletalMesh(InSkelMesh,false);
diff --git a/Engine/Source/Runtime/Engine/Private/Components/SkinnedMeshComponent.cpp b/Engine/Source/Runtime/Engine/Private/Components/SkinnedMeshComponent.cpp
--- a/Engine/Source/Runtime/Engine/Private/Components/SkinnedMeshComponent.cpp
+++ b/Engine/Source/Runtime/Engine/Private/Components/SkinnedMeshComponent.cpp
@@ -1823,7 +1823,7 @@ FTransform USkinnedMeshComponent::GetRefPoseTransform(int32 BoneIndex) const
 
 void USkinnedMeshComponent::SetSkeletalMesh(USkeletalMesh* NewMesh, bool bReinitPose)
 {
-	SetSkinnedAsset(NewMesh, bReinitPose);
+	SetSkinnedAssetAndUpdate(NewMesh, bReinitPose);
 }
 
 USkeletalMesh* USkinnedMeshComponent::GetSkeletalMesh_DEPRECATED() const
@@ -1831,7 +1831,7 @@ USkeletalMesh* USkinnedMeshComponent::GetSkeletalMesh_DEPRECATED() const
 	return Cast<class USkeletalMesh>(GetSkinnedAsset());
 }
 
-void USkinnedMeshComponent::SetSkinnedAsset(USkinnedAsset* InSkinnedAsset, bool /*bReinitPose*/)
+void USkinnedMeshComponent::SetSkinnedAssetAndUpdate(USkinnedAsset* InSkinnedAsset, bool /*bReinitPose*/)
 {
 	// NOTE: InSkinnedAsset may be nullptr (useful in the editor for removing the skeletal mesh associated with
 	//   this component on-the-fly)
