- Fixed AWorldDataLayers::GetDataLayerInstanceFromAssetName not returning the correct DataLayerInstance
- Fixed UDataLayerInstanceWithAsset::GetDataLayerFullName now return DataLayer Asset's PathName (instead of FullName)
#rnx

diff --git a/Engine/Source/Runtime/Engine/Private/WorldPartition/DataLayer/WorldDataLayers.cpp b/Engine/Source/Runtime/Engine/Private/WorldPartition/DataLayer/WorldDataLayers.cpp
--- a/Engine/Source/Runtime/Engine/Private/WorldPartition/DataLayer/WorldDataLayers.cpp
+++ b/Engine/Source/Runtime/Engine/Private/WorldPartition/DataLayer/WorldDataLayers.cpp
@@ -657,7 +657,7 @@ const UDataLayerInstance* AWorldDataLayers::GetDataLayerInstanceFromAssetName(co
 #if WITH_EDITOR	
 	for (UDataLayerInstance* DataLayerInstance : DataLayerInstances)
 	{
-		if (DataLayerInstance->GetDataLayerFullName().Compare(InDataLayerAssetFullName.ToString(), ESearchCase::IgnoreCase))
+		if (DataLayerInstance->GetDataLayerFullName().Equals(InDataLayerAssetFullName.ToString(), ESearchCase::IgnoreCase))
 		{
 			return DataLayerInstance;
 		}
diff --git a/Engine/Source/Runtime/Engine/Public/WorldPartition/DataLayer/DataLayerInstanceWithAsset.h b/Engine/Source/Runtime/Engine/Public/WorldPartition/DataLayer/DataLayerInstanceWithAsset.h
--- a/Engine/Source/Runtime/Engine/Public/WorldPartition/DataLayer/DataLayerInstanceWithAsset.h
+++ b/Engine/Source/Runtime/Engine/Public/WorldPartition/DataLayer/DataLayerInstanceWithAsset.h
@@ -35,7 +35,7 @@ public:
 	virtual FColor GetDebugColor() const override { return DataLayerAsset != nullptr ? DataLayerAsset->GetDebugColor() : FColor::Black; }
 
 	virtual FString GetDataLayerShortName() const override { return DataLayerAsset != nullptr ? DataLayerAsset->GetName() : GetDataLayerFName().ToString(); }
-	virtual FString GetDataLayerFullName() const override { return DataLayerAsset != nullptr ? DataLayerAsset->GetFullName() : GetDataLayerFName().ToString(); }
+	virtual FString GetDataLayerFullName() const override { return DataLayerAsset != nullptr ? DataLayerAsset->GetPathName() : GetDataLayerFName().ToString(); }
 
 private:
 	UPROPERTY(Category = "Data Layer", EditAnywhere)
