Fixed false positives in transaction diffing caused by name/object index re-ordering

Serialized name/object data would write the index of the name/object entry in the table, however there was no guarantee that this table order was reliable between diff snapshots (eg, adding a new name part way through an object would invalidate all subsequent references, even if they were still using the same name).

This change updates the diff writer to no longer attempt to strip out the name/object data into shared tables (which was/is an optimization for serialized objects kept within the transaction buffer), and instead write the global name index and raw object pointer instead.

Note: This change means that a diffable object snapshot is no longer portable to other instances of the editor.

#rnx

diff --git a/Engine/Source/Editor/UnrealEd/Private/EditorTransaction.cpp b/Engine/Source/Editor/UnrealEd/Private/EditorTransaction.cpp
--- a/Engine/Source/Editor/UnrealEd/Private/EditorTransaction.cpp
+++ b/Engine/Source/Editor/UnrealEd/Private/EditorTransaction.cpp
@@ -486,16 +486,12 @@ void FTransaction::FObjectRecord::AddReferencedObjects( FReferenceCollector& Col
 {
 	Object.AddStructReferencedObjects(Collector);
 
-	auto AddTransactionSerializedObjectReferences = [&Collector](FTransactionSerializedObject& InSerializedObject)
+	auto AddSerializedObjectReferences = [&Collector](FSerializedObject& InSerializedObject)
 	{
 		for (FPersistentObjectRef& ReferencedObject : InSerializedObject.ReferencedObjects)
 		{
 			ReferencedObject.AddStructReferencedObjects(Collector);
 		}
-	};
-	auto AddSerializedObjectReferences = [&Collector, &AddTransactionSerializedObjectReferences](FSerializedObject& InSerializedObject)
-	{
-		AddTransactionSerializedObjectReferences(InSerializedObject);
 
 		if (InSerializedObject.ObjectAnnotation.IsValid())
 		{
@@ -504,14 +500,6 @@ void FTransaction::FObjectRecord::AddReferencedObjects( FReferenceCollector& Col
 	};
 	AddSerializedObjectReferences(SerializedObject);
 	AddSerializedObjectReferences(SerializedObjectFlip);
-	if (DiffableObject)
-	{
-		AddTransactionSerializedObjectReferences(DiffableObject->SerializedObject);
-	}
-	if (DiffableObjectSnapshot)
-	{
-		AddTransactionSerializedObjectReferences(DiffableObjectSnapshot->SerializedObject);
-	}
 
 	if (CustomChange.IsValid())
 	{
diff --git a/Engine/Source/Runtime/Engine/Private/TransactionCommon.cpp b/Engine/Source/Runtime/Engine/Private/TransactionCommon.cpp
--- a/Engine/Source/Runtime/Engine/Private/TransactionCommon.cpp
+++ b/Engine/Source/Runtime/Engine/Private/TransactionCommon.cpp
@@ -235,34 +235,45 @@ FArchive& FSerializedObjectDataReader::operator<<(UObject*& Res)
 }
 
 
-FSerializedObjectDataWriter::FSerializedObjectDataWriter(FTransactionSerializedObject& InSerializedObject)
-	: SerializedObject(InSerializedObject)
+namespace Internal
 {
-	SetIsSaving(true);
 
-	for (int32 ObjIndex = 0; ObjIndex < SerializedObject.ReferencedObjects.Num(); ++ObjIndex)
-	{
-		ObjectMap.Add(SerializedObject.ReferencedObjects[ObjIndex].Get(), ObjIndex);
-	}
-
-	for (int32 NameIndex = 0; NameIndex < SerializedObject.ReferencedNames.Num(); ++NameIndex)
-	{
-		NameMap.Add(SerializedObject.ReferencedNames[NameIndex], NameIndex);
-	}
+FSerializedObjectDataWriterCommon::FSerializedObjectDataWriterCommon(FTransactionSerializedObjectData& InSerializedData)
+	: SerializedData(InSerializedData)
+{
+	SetIsSaving(true);
 }
 
-void FSerializedObjectDataWriter::Serialize(void* SerData, int64 Num)
+void FSerializedObjectDataWriterCommon::Serialize(void* SerData, int64 Num)
 {
 	if (Num)
 	{
 		int64 DataIndex = Offset;
-		SerializedObject.SerializedData.Write(SerData, Offset, Num);
+		SerializedData.Write(SerData, Offset, Num);
 		Offset += Num;
 
 		OnDataSerialized(DataIndex, Num);
 	}
 }
 
+} // namespace Internal
+
+
+FSerializedObjectDataWriter::FSerializedObjectDataWriter(FTransactionSerializedObject& InSerializedObject)
+	: Internal::FSerializedObjectDataWriterCommon(InSerializedObject.SerializedData)
+	, SerializedObject(InSerializedObject)
+{
+	for (int32 ObjIndex = 0; ObjIndex < SerializedObject.ReferencedObjects.Num(); ++ObjIndex)
+	{
+		ObjectMap.Add(SerializedObject.ReferencedObjects[ObjIndex].Get(), ObjIndex);
+	}
+
+	for (int32 NameIndex = 0; NameIndex < SerializedObject.ReferencedNames.Num(); ++NameIndex)
+	{
+		NameMap.Add(SerializedObject.ReferencedNames[NameIndex], NameIndex);
+	}
+}
+
 FArchive& FSerializedObjectDataWriter::operator<<(FName& N)
 {
 	int32 NameIndex = INDEX_NONE;
@@ -277,8 +288,6 @@ FArchive& FSerializedObjectDataWriter::operator<<(FName& N)
 		NameMap.Add(N, NameIndex);
 	}
 
-	OnNameSerialized(NameIndex);
-
 	return *this << NameIndex;
 }
 
@@ -296,28 +305,16 @@ FArchive& FSerializedObjectDataWriter::operator<<(UObject*& Res)
 		ObjectMap.Add(Res, ObjectIndex);
 	}
 
-	OnObjectSerialized(ObjectIndex);
-
 	return *this << ObjectIndex;
 }
 
 
 FDiffableObjectDataWriter::FDiffableObjectDataWriter(FTransactionDiffableObject& InDiffableObject, TArrayView<const FProperty*> InPropertiesToSerialize)
-	: FSerializedObjectDataWriter(InDiffableObject.SerializedObject)
+	: Internal::FSerializedObjectDataWriterCommon(InDiffableObject.SerializedData)
 	, DiffableObject(InDiffableObject)
 	, PropertiesToSerialize(InPropertiesToSerialize)
 {
 	SetWantBinaryPropertySerialization(true);
-
-	for (int32 ObjIndex = 0; ObjIndex < SerializedObject.ReferencedObjects.Num(); ++ObjIndex)
-	{
-		ObjectMap.Add(SerializedObject.ReferencedObjects[ObjIndex].Get(), ObjIndex);
-	}
-
-	for (int32 NameIndex = 0; NameIndex < SerializedObject.ReferencedNames.Num(); ++NameIndex)
-	{
-		NameMap.Add(SerializedObject.ReferencedNames[NameIndex], NameIndex);
-	}
 }
 
 FName FDiffableObjectDataWriter::GetTaggedDataKey() const
@@ -372,7 +369,7 @@ bool FDiffableObjectDataWriter::ShouldSkipProperty(const FProperty* InProperty)
 {
 	return (PropertiesToSerialize.Num() > 0 && !PropertiesToSerialize.Contains(InProperty))
 		|| InProperty->HasAnyPropertyFlags(CPF_Transient | CPF_NonTransactional | CPF_Deprecated)
-		|| FSerializedObjectDataWriter::ShouldSkipProperty(InProperty);
+		|| Internal::FSerializedObjectDataWriterCommon::ShouldSkipProperty(InProperty);
 }
 
 void FDiffableObjectDataWriter::MarkScriptSerializationStart(const UObject* Obj)
@@ -406,24 +403,17 @@ void FDiffableObjectDataWriter::OnDataSerialized(int64 InOffset, int64 InNum)
 	}
 }
 
-void FDiffableObjectDataWriter::OnNameSerialized(int32 InNameIndex)
+FArchive& FDiffableObjectDataWriter::operator<<(FName& N)
 {
-	// Track this name index in the serialized data
-	const FName SerializedTaggedDataKey = GetTaggedDataKey();
-	if (!SerializedTaggedDataKey.IsNone())
-	{
-		DiffableObject.SerializedNameIndices.FindOrAdd(SerializedTaggedDataKey).Indices.AddUnique(InNameIndex);
-	}
+	uint32 NameDisplayIndex = N.GetDisplayIndex().ToUnstableInt();
+	int32 NameNumericSuffix = N.GetNumber();
+	return *this << NameDisplayIndex << NameNumericSuffix;
 }
 
-void FDiffableObjectDataWriter::OnObjectSerialized(int32 InObjectIndex)
+FArchive& FDiffableObjectDataWriter::operator<<(UObject*& Res)
 {
-	// Track this object offset in the serialized data
-	const FName SerializedTaggedDataKey = GetTaggedDataKey();
-	if (!SerializedTaggedDataKey.IsNone())
-	{
-		DiffableObject.SerializedObjectIndices.FindOrAdd(SerializedTaggedDataKey).Indices.AddUnique(InObjectIndex);
-	}
+	PTRINT ResInt = (PTRINT)Res;
+	return *this << ResInt;
 }
 
 FName FDiffableObjectDataWriter::FCachedPropertyKey::SyncCache(const FArchiveSerializedPropertyChain* InPropertyChain)
@@ -458,15 +448,7 @@ void GenerateObjectDiff(const FTransactionDiffableObject& OldDiffableObject, con
 		bool bIsTaggedDataIdentical = OldSerializedTaggedData.DataSize == NewSerializedTaggedData.DataSize;
 		if (bIsTaggedDataIdentical && OldSerializedTaggedData.HasSerializedData())
 		{
-			bIsTaggedDataIdentical = FMemory::Memcmp(OldDiffableObject.SerializedObject.SerializedData.GetPtr(OldSerializedTaggedData.DataOffset), NewDiffableObject.SerializedObject.SerializedData.GetPtr(NewSerializedTaggedData.DataOffset), NewSerializedTaggedData.DataSize) == 0;
-		}
-		if (bIsTaggedDataIdentical)
-		{
-			bIsTaggedDataIdentical = Internal::AreObjectPointersIdentical(OldDiffableObject, NewDiffableObject, TaggedDataKey);
-		}
-		if (bIsTaggedDataIdentical)
-		{
-			bIsTaggedDataIdentical = Internal::AreNamesIdentical(OldDiffableObject, NewDiffableObject, TaggedDataKey);
+			bIsTaggedDataIdentical = FMemory::Memcmp(OldDiffableObject.SerializedData.GetPtr(OldSerializedTaggedData.DataOffset), NewDiffableObject.SerializedData.GetPtr(NewSerializedTaggedData.DataOffset), NewSerializedTaggedData.DataSize) == 0;
 		}
 		return bIsTaggedDataIdentical;
 	};
@@ -546,63 +528,6 @@ void GenerateObjectDiff(const FTransactionDiffableObject& OldDiffableObject, con
 	}
 }
 
-namespace Internal
-{
-
-bool AreObjectPointersIdentical(const FTransactionDiffableObject& OldDiffableObject, const FTransactionDiffableObject& NewDiffableObject, const FName TaggedDataKey)
-{
-	auto GetIndicesArray = [&TaggedDataKey](const FTransactionDiffableObject& DiffableObject) -> const TArray<int32>&
-	{
-		if (const FTransactionSerializedIndices* SerializedObjectIndices = DiffableObject.SerializedObjectIndices.Find(TaggedDataKey))
-		{
-			return SerializedObjectIndices->Indices;
-		}
-
-		static const TArray<int32> EmptyIndicesArray;
-		return EmptyIndicesArray;
-	};
-
-	const TArray<int32>& OldSerializedObjectIndices = GetIndicesArray(OldDiffableObject);
-	const TArray<int32>& NewSerializedObjectIndices = GetIndicesArray(NewDiffableObject);
-
-	bool bAreObjectPointersIdentical = OldSerializedObjectIndices.Num() == NewSerializedObjectIndices.Num();
-	for (int32 ObjIndex = 0; ObjIndex < OldSerializedObjectIndices.Num() && bAreObjectPointersIdentical; ++ObjIndex)
-	{
-		const FTransactionPersistentObjectRef& OldObjectRef = OldDiffableObject.SerializedObject.ReferencedObjects.IsValidIndex(OldSerializedObjectIndices[ObjIndex]) ? OldDiffableObject.SerializedObject.ReferencedObjects[OldSerializedObjectIndices[ObjIndex]] : FTransactionPersistentObjectRef();
-		const FTransactionPersistentObjectRef& NewObjectRef = NewDiffableObject.SerializedObject.ReferencedObjects.IsValidIndex(NewSerializedObjectIndices[ObjIndex]) ? NewDiffableObject.SerializedObject.ReferencedObjects[NewSerializedObjectIndices[ObjIndex]] : FTransactionPersistentObjectRef();
-		bAreObjectPointersIdentical = OldObjectRef == NewObjectRef;
-	}
-	return bAreObjectPointersIdentical;
-}
-
-bool AreNamesIdentical(const FTransactionDiffableObject& OldDiffableObject, const FTransactionDiffableObject& NewDiffableObject, const FName TaggedDataKey)
-{
-	auto GetIndicesArray = [&TaggedDataKey](const FTransactionDiffableObject& DiffableObject) -> const TArray<int32>&
-	{
-		if (const FTransactionSerializedIndices* SerializedNameIndices = DiffableObject.SerializedNameIndices.Find(TaggedDataKey))
-		{
-			return SerializedNameIndices->Indices;
-		}
-
-		static const TArray<int32> EmptyIndicesArray;
-		return EmptyIndicesArray;
-	};
-
-	const TArray<int32>& OldSerializedNameIndices = GetIndicesArray(OldDiffableObject);
-	const TArray<int32>& NewSerializedNameIndices = GetIndicesArray(NewDiffableObject);
-
-	bool bAreNamesIdentical = OldSerializedNameIndices.Num() == NewSerializedNameIndices.Num();
-	for (int32 ObjIndex = 0; ObjIndex < OldSerializedNameIndices.Num() && bAreNamesIdentical; ++ObjIndex)
-	{
-		const FName& OldName = OldDiffableObject.SerializedObject.ReferencedNames.IsValidIndex(OldSerializedNameIndices[ObjIndex]) ? OldDiffableObject.SerializedObject.ReferencedNames[OldSerializedNameIndices[ObjIndex]] : FName();
-		const FName& NewName = NewDiffableObject.SerializedObject.ReferencedNames.IsValidIndex(NewSerializedNameIndices[ObjIndex]) ? NewDiffableObject.SerializedObject.ReferencedNames[NewSerializedNameIndices[ObjIndex]] : FName();
-		bAreNamesIdentical = OldName == NewName;
-	}
-	return bAreNamesIdentical;
-}
-
-} // namespace Internal
-
 } // namespace DiffUtil
 
 } // namespace UE::Transaction
diff --git a/Engine/Source/Runtime/Engine/Public/TransactionCommon.h b/Engine/Source/Runtime/Engine/Public/TransactionCommon.h
--- a/Engine/Source/Runtime/Engine/Public/TransactionCommon.h
+++ b/Engine/Source/Runtime/Engine/Public/TransactionCommon.h
@@ -313,15 +313,15 @@ public:
 		Exchange(ReferencedNames, Other.ReferencedNames);
 	}
 
-	/** The data stream used to serialize/deserialize record */
+	/** The serialized data for the transacted object */
 	UPROPERTY()
 	FTransactionSerializedObjectData SerializedData;
 
-	/** External objects referenced in the transaction */
+	/** External objects referenced by the transacted object */
 	UPROPERTY()
 	TArray<FTransactionPersistentObjectRef> ReferencedObjects;
 	
-	/** FNames referenced in the object record */
+	/** Names referenced by the transacted object */
 	UPROPERTY()
 	TArray<FName> ReferencedNames;
 };
@@ -340,40 +340,28 @@ public:
 	void Reset()
 	{
 		ObjectInfo.Reset();
-		SerializedObject.Reset();
+		SerializedData.Reset();
 		SerializedTaggedData.Reset();
-		SerializedObjectIndices.Reset();
-		SerializedNameIndices.Reset();
 	}
 
 	void Swap(FTransactionDiffableObject& Other)
 	{
 		ObjectInfo.Swap(Other.ObjectInfo);
-		SerializedObject.Swap(Other.SerializedObject);
+		Exchange(SerializedData, Other.SerializedData);
 		Exchange(SerializedTaggedData, Other.SerializedTaggedData);
-		Exchange(SerializedObjectIndices, Other.SerializedObjectIndices);
-		Exchange(SerializedNameIndices, Other.SerializedNameIndices);
 	}
 
 	/** Information about the object when it was serialized */
 	UPROPERTY()
 	FTransactionSerializedObjectInfo ObjectInfo;
 
-	/** The serialized object data */
+	/** The serialized data for the diffable object */
 	UPROPERTY()
-	FTransactionSerializedObject SerializedObject;
+	FTransactionSerializedObjectData SerializedData;
 
 	/** Information about tagged data (mainly properties) that were serialized within this object */
 	UPROPERTY()
 	TMap<FName, FTransactionSerializedTaggedData> SerializedTaggedData;
-
-	/** Information about the object pointer offsets that were serialized within this object (this maps the tagged data key (or None if there was no tagged data key) to the ReferencedObjects indices of the tagged data) */
-	UPROPERTY()
-	TMap<FName, FTransactionSerializedIndices> SerializedObjectIndices;
-
-	/** Information about the name offsets that were serialized within this object (this maps the tagged data name to the ReferencedNames index of the tagged data) */
-	UPROPERTY()
-	TMap<FName, FTransactionSerializedIndices> SerializedNameIndices;
 };
 
 namespace UE::Transaction
@@ -403,27 +391,41 @@ protected:
 	int64 Offset = 0;
 };
 
+namespace Internal
+{
+
 /** Core archive to write a transaction object to the buffer. */
-class ENGINE_API FSerializedObjectDataWriter : public FArchiveUObject
+class ENGINE_API FSerializedObjectDataWriterCommon : public FArchiveUObject
 {
 public:
-	FSerializedObjectDataWriter(FTransactionSerializedObject& InSerializedObject);
+	FSerializedObjectDataWriterCommon(FTransactionSerializedObjectData& InSerializedData);
 
 	virtual int64 Tell() override { return Offset; }
-	virtual void Seek(int64 InPos) override { checkSlow(Offset <= SerializedObject.SerializedData.Num()); Offset = InPos; }
-	virtual int64 TotalSize() override { return SerializedObject.SerializedData.Num(); }
+	virtual void Seek(int64 InPos) override { checkSlow(Offset <= SerializedData.Num()); Offset = InPos; }
+	virtual int64 TotalSize() override { return SerializedData.Num(); }
 
 protected:
 	virtual void Serialize(void* SerData, int64 Num) override;
 
+	virtual void OnDataSerialized(int64 InOffset, int64 InNum) {}
+
+	FTransactionSerializedObjectData& SerializedData;
+	int64 Offset = 0;
+};
+
+} // namespace Internal
+
+/** Core archive to write a transaction object to the buffer. */
+class ENGINE_API FSerializedObjectDataWriter : public Internal::FSerializedObjectDataWriterCommon
+{
+public:
+	FSerializedObjectDataWriter(FTransactionSerializedObject& InSerializedObject);
+
+protected:
 	using FArchiveUObject::operator<<;
 	virtual FArchive& operator<<(class FName& N) override;
 	virtual FArchive& operator<<(class UObject*& Res) override;
 
-	virtual void OnDataSerialized(int64 InOffset, int64 InNum) {}
-	virtual void OnNameSerialized(int32 InNameIndex) {}
-	virtual void OnObjectSerialized(int32 InObjectIndex) {}
-
 	FTransactionSerializedObject& SerializedObject;
 	int64 Offset = 0;
 
@@ -432,7 +434,7 @@ protected:
 };
 
 /** Core archive to write a diffable object to the buffer. */
-class ENGINE_API FDiffableObjectDataWriter : public FSerializedObjectDataWriter
+class ENGINE_API FDiffableObjectDataWriter : public Internal::FSerializedObjectDataWriterCommon
 {
 public:
 	FDiffableObjectDataWriter(FTransactionDiffableObject& InDiffableObject, TArrayView<const FProperty*> InPropertiesToSerialize = TArrayView<const FProperty*>());
@@ -448,8 +450,10 @@ protected:
 	virtual void MarkScriptSerializationEnd(const UObject* Obj) override;
 
 	virtual void OnDataSerialized(int64 InOffset, int64 InNum) override;
-	virtual void OnNameSerialized(int32 InNameIndex) override;
-	virtual void OnObjectSerialized(int32 InObjectIndex) override;
+
+	using FArchiveUObject::operator<<;
+	virtual FArchive& operator<<(class FName& N) override;
+	virtual FArchive& operator<<(class UObject*& Res) override;
 
 private:
 	struct FCachedPropertyKey
@@ -479,14 +483,6 @@ namespace DiffUtil
 
 ENGINE_API void GenerateObjectDiff(const FTransactionDiffableObject& OldDiffableObject, const FTransactionDiffableObject& NewDiffableObject, FTransactionObjectDeltaChange& OutDeltaChange, const bool bFullDiff = true);
 
-namespace Internal
-{
-
-ENGINE_API bool AreObjectPointersIdentical(const FTransactionDiffableObject& OldDiffableObject, const FTransactionDiffableObject& NewDiffableObject, const FName TaggedDataKey);
-ENGINE_API bool AreNamesIdentical(const FTransactionDiffableObject& OldDiffableObject, const FTransactionDiffableObject& NewDiffableObject, const FName TaggedDataKey);
-
-} // namespace Internal
-
 } // namespace DiffUtil
 
 } // namespace UE::Transaction
