Fixed undo bugs occurring when changing a parameter value on a node in the Material Editor either through the Detail View panel or through the pin default value text boxes. The bug was occurring in both directions (changing the value on the Details View was not updating the Pin and vice versa).

diff --git a/Engine/Source/Runtime/Engine/Private/Materials/MaterialExpressions.cpp b/Engine/Source/Runtime/Engine/Private/Materials/MaterialExpressions.cpp
--- a/Engine/Source/Runtime/Engine/Private/Materials/MaterialExpressions.cpp
+++ b/Engine/Source/Runtime/Engine/Private/Materials/MaterialExpressions.cpp
@@ -1073,178 +1073,181 @@ void UMaterialExpression::PinDefaultValueChanged(int32 PinIndex, const FString&
 	// Update the default value of the expression input when pin value changes
 	// Find the UPROPERTYs that have OverridingInputProperty meta data pointing to the input, override their values.
 	TArray<FProperty*> Properties = GetInputPinProperty(PinIndex);
-	if (Properties.Num() > 0)
+	if (Properties.IsEmpty())
+	{
+		return;
+	}
+
+	Modify();
+
+	TArray<FString> PropertyValues;
+	if (Properties.Num() == 1)
+	{
+		PropertyValues.Add(DefaultValue);
+	}
+	else if (Properties.Num() == 2)
+	{
+		// Vector2 is formatted as (X=0.0, Y=0.0)
+		FVector2D Value;
+		Value.InitFromString(DefaultValue);
+		PropertyValues.Add(FString::SanitizeFloat(Value.X));
+		PropertyValues.Add(FString::SanitizeFloat(Value.Y));
+	}
+	else
+	{
+		// Vector3/4 are formatted as numbers separated by commas
+		DefaultValue.ParseIntoArray(PropertyValues, TEXT(","), true);
+		check(PropertyValues.Num() == Properties.Num());
+	}
+
+	for (int32 i = 0; i < Properties.Num(); ++i)
 	{
-		TArray<FString> PropertyValues;
-		if (Properties.Num() == 1)
+		FProperty* Property = Properties[i];
+		const FString& ClampMin = Property->GetMetaData(TEXT("ClampMin"));
+		const FString& ClampMax = Property->GetMetaData(TEXT("ClampMax"));
+
+		FString PropertyValue = PropertyValues[i];
+		FFieldClass* PropertyClass = Property->GetClass();
+		if (PropertyClass == FFloatProperty::StaticClass())
 		{
-			PropertyValues.Add(DefaultValue);
+			float Value = FCString::Atof(*PropertyValue);
+			Value = ClampMin.Len() ? (FMath::Max<float>(FCString::Atof(*ClampMin), Value)) : Value;
+			Value = ClampMax.Len() ? (FMath::Min<float>(FCString::Atof(*ClampMax), Value)) : Value;
+			FFloatProperty* FloatProperty = CastField<FFloatProperty>(Property);
+			FloatProperty->SetPropertyValue_InContainer(this, Value);
 		}
-		else if (Properties.Num() == 2)
+		else if (PropertyClass == FDoubleProperty::StaticClass())
 		{
-			// Vector2 is formatted as (X=0.0, Y=0.0)
-			FVector2D Value;
-			Value.InitFromString(DefaultValue);
-			PropertyValues.Add(FString::SanitizeFloat(Value.X));
-			PropertyValues.Add(FString::SanitizeFloat(Value.Y));
+			double Value = FCString::Atod(*PropertyValue);
+			Value = ClampMin.Len() ? (FMath::Max<double>(FCString::Atod(*ClampMin), Value)) : Value;
+			Value = ClampMax.Len() ? (FMath::Min<double>(FCString::Atod(*ClampMax), Value)) : Value;
+			FDoubleProperty* DoubleProperty = CastField<FDoubleProperty>(Property);
+			DoubleProperty->SetPropertyValue_InContainer(this, Value);
 		}
-		else
+		else if (PropertyClass == FIntProperty::StaticClass())
 		{
-			// Vector3/4 are formatted as numbers separated by commas
-			DefaultValue.ParseIntoArray(PropertyValues, TEXT(","), true);
-			check(PropertyValues.Num() == Properties.Num());
+			int32 Value = FCString::Atoi(*PropertyValue);
+			Value = ClampMin.Len() ? (FMath::Max<int32>(FCString::Atoi(*ClampMin), Value)) : Value;
+			Value = ClampMax.Len() ? (FMath::Min<int32>(FCString::Atoi(*ClampMax), Value)) : Value;
+			FIntProperty* IntProperty = CastField<FIntProperty>(Property);
+			IntProperty->SetPropertyValue_InContainer(this, Value);
 		}
-
-		for (int32 i = 0; i < Properties.Num(); ++i)
+		else if (PropertyClass == FUInt32Property::StaticClass())
 		{
-			FProperty* Property = Properties[i];
-			const FString& ClampMin = Property->GetMetaData(TEXT("ClampMin"));
-			const FString& ClampMax = Property->GetMetaData(TEXT("ClampMax"));
-
-			FString PropertyValue = PropertyValues[i];
-			FFieldClass* PropertyClass = Property->GetClass();
-			if (PropertyClass == FFloatProperty::StaticClass())
-			{
-				float Value = FCString::Atof(*PropertyValue);
-				Value = ClampMin.Len() ? (FMath::Max<float>(FCString::Atof(*ClampMin), Value)) : Value;
-				Value = ClampMax.Len() ? (FMath::Min<float>(FCString::Atof(*ClampMax), Value)) : Value;
-				FFloatProperty* FloatProperty = CastField<FFloatProperty>(Property);
-				FloatProperty->SetPropertyValue_InContainer(this, Value);
-			}
-			else if (PropertyClass == FDoubleProperty::StaticClass())
-			{
-				double Value = FCString::Atod(*PropertyValue);
-				Value = ClampMin.Len() ? (FMath::Max<double>(FCString::Atod(*ClampMin), Value)) : Value;
-				Value = ClampMax.Len() ? (FMath::Min<double>(FCString::Atod(*ClampMax), Value)) : Value;
-				FDoubleProperty* DoubleProperty = CastField<FDoubleProperty>(Property);
-				DoubleProperty->SetPropertyValue_InContainer(this, Value);
-			}
-			else if (PropertyClass == FIntProperty::StaticClass())
+			int32 IntValue = FCString::Atoi(*PropertyValue);
+			IntValue = ClampMin.Len() ? (FMath::Max<int32>(FCString::Atoi(*ClampMin), IntValue)) : IntValue;
+			IntValue = ClampMax.Len() ? (FMath::Min<int32>(FCString::Atoi(*ClampMax), IntValue)) : IntValue;
+			// Make sure the value is not negative
+			uint32 Value = (uint32)FMath::Max(IntValue, 0);
+			FUInt32Property* UInt32Property = CastField<FUInt32Property>(Property);
+			UInt32Property->SetPropertyValue_InContainer(this, Value);
+		}
+		else if (PropertyClass == FByteProperty::StaticClass())
+		{
+			FByteProperty* ByteProperty = CastField<FByteProperty>(Property);
+			uint8 Value;
+			UEnum* Enum = ByteProperty->GetIntPropertyEnum();
+			if (Enum)
 			{
-				int32 Value = FCString::Atoi(*PropertyValue);
-				Value = ClampMin.Len() ? (FMath::Max<int32>(FCString::Atoi(*ClampMin), Value)) : Value;
-				Value = ClampMax.Len() ? (FMath::Min<int32>(FCString::Atoi(*ClampMax), Value)) : Value;
-				FIntProperty* IntProperty = CastField<FIntProperty>(Property);
-				IntProperty->SetPropertyValue_InContainer(this, Value);
+				Value = (uint8)Enum->GetValueByName(FName(PropertyValue));
 			}
-			else if (PropertyClass == FUInt32Property::StaticClass())
+			else
 			{
 				int32 IntValue = FCString::Atoi(*PropertyValue);
 				IntValue = ClampMin.Len() ? (FMath::Max<int32>(FCString::Atoi(*ClampMin), IntValue)) : IntValue;
 				IntValue = ClampMax.Len() ? (FMath::Min<int32>(FCString::Atoi(*ClampMax), IntValue)) : IntValue;
-				// Make sure the value is not negative
-				uint32 Value = (uint32)FMath::Max(IntValue, 0);
-				FUInt32Property* UInt32Property = CastField<FUInt32Property>(Property);
-				UInt32Property->SetPropertyValue_InContainer(this, Value);
+				// Make sure the value doesn't exceed byte limit
+				Value = (uint8)FMath::Max(FMath::Min(IntValue, 255), 0);
 			}
-			else if (PropertyClass == FByteProperty::StaticClass())
+			ByteProperty->SetPropertyValue_InContainer(this, Value);
+		}
+		else if (PropertyClass == FBoolProperty::StaticClass())
+		{
+			bool Value = FCString::ToBool(*PropertyValue);
+			FBoolProperty* BoolProperty = CastField<FBoolProperty>(Property);
+			BoolProperty->SetPropertyValue_InContainer(this, Value);
+		}
+		else if (PropertyClass == FStructProperty::StaticClass())
+		{
+			UScriptStruct* Struct = ((FStructProperty*)Property)->Struct;
+			if (Struct == TBaseStructure<FLinearColor>::Get())
 			{
-				FByteProperty* ByteProperty = CastField<FByteProperty>(Property);
-				uint8 Value;
-				UEnum* Enum = ByteProperty->GetIntPropertyEnum();
-				if (Enum)
+				FLinearColor* ColorProperty = Property->ContainerPtrToValuePtr<FLinearColor>(this);
+				if (Property->HasMetaData(TEXT("HideAlphaChannel")))
 				{
-					Value = (uint8)Enum->GetValueByName(FName(PropertyValue));
+					// This is a 3 element vector
+					TArray<FString> Elements;
+					PropertyValue.ParseIntoArray(Elements, TEXT(","), true);
+					check(Elements.Num() == 3);
+					ColorProperty->R = FCString::Atof(*Elements[0]);
+					ColorProperty->G = FCString::Atof(*Elements[1]);
+					ColorProperty->B = FCString::Atof(*Elements[2]);
 				}
 				else
 				{
-					int32 IntValue = FCString::Atoi(*PropertyValue);
-					IntValue = ClampMin.Len() ? (FMath::Max<int32>(FCString::Atoi(*ClampMin), IntValue)) : IntValue;
-					IntValue = ClampMax.Len() ? (FMath::Min<int32>(FCString::Atoi(*ClampMax), IntValue)) : IntValue;
-					// Make sure the value doesn't exceed byte limit
-					Value = (uint8)FMath::Max(FMath::Min(IntValue, 255), 0);
+					// This is a 4 element vector
+					ColorProperty->InitFromString(PropertyValue);
 				}
-				ByteProperty->SetPropertyValue_InContainer(this, Value);
 			}
-			else if (PropertyClass == FBoolProperty::StaticClass())
+			else if (Struct == TBaseStructure<FVector4>::Get())
 			{
-				bool Value = FCString::ToBool(*PropertyValue);
-				FBoolProperty* BoolProperty = CastField<FBoolProperty>(Property);
-				BoolProperty->SetPropertyValue_InContainer(this, Value);
+				TArray<FString> Elements;
+				PropertyValue.ParseIntoArray(Elements, TEXT(","), true);
+				check(Elements.Num() == 4);
+				FVector4* Value = Property->ContainerPtrToValuePtr<FVector4>(this);
+				Value->X = FCString::Atod(*Elements[0]);
+				Value->Y = FCString::Atod(*Elements[1]);
+				Value->Z = FCString::Atod(*Elements[2]);
+				Value->W = FCString::Atod(*Elements[3]);
 			}
-			else if (PropertyClass == FStructProperty::StaticClass())
+			else if (Struct == TVariantStructure<FVector4d>::Get())
 			{
-				UScriptStruct* Struct = ((FStructProperty*)Property)->Struct;
-				if (Struct == TBaseStructure<FLinearColor>::Get())
-				{
-					FLinearColor* ColorProperty = Property->ContainerPtrToValuePtr<FLinearColor>(this);
-					if (Property->HasMetaData(TEXT("HideAlphaChannel")))
-					{
-						// This is a 3 element vector
-						TArray<FString> Elements;
-						PropertyValue.ParseIntoArray(Elements, TEXT(","), true);
-						check(Elements.Num() == 3);
-						ColorProperty->R = FCString::Atof(*Elements[0]);
-						ColorProperty->G = FCString::Atof(*Elements[1]);
-						ColorProperty->B = FCString::Atof(*Elements[2]);
-					}
-					else
-					{
-						// This is a 4 element vector
-						ColorProperty->InitFromString(PropertyValue);
-					}
-				}
-				else if (Struct == TBaseStructure<FVector4>::Get())
-				{
-					TArray<FString> Elements;
-					PropertyValue.ParseIntoArray(Elements, TEXT(","), true);
-					check(Elements.Num() == 4);
-					FVector4* Value = Property->ContainerPtrToValuePtr<FVector4>(this);
-					Value->X = FCString::Atod(*Elements[0]);
-					Value->Y = FCString::Atod(*Elements[1]);
-					Value->Z = FCString::Atod(*Elements[2]);
-					Value->W = FCString::Atod(*Elements[3]);
-				}
-				else if (Struct == TVariantStructure<FVector4d>::Get())
-				{
-					TArray<FString> Elements;
-					PropertyValue.ParseIntoArray(Elements, TEXT(","), true);
-					check(Elements.Num() == 4);
-					FVector4d* Value = Property->ContainerPtrToValuePtr<FVector4d>(this);
-					Value->X = FCString::Atod(*Elements[0]);
-					Value->Y = FCString::Atod(*Elements[1]);
-					Value->Z = FCString::Atod(*Elements[2]);
-					Value->W = FCString::Atod(*Elements[3]);
-				}
-				else if (Struct == TBaseStructure<FVector>::Get())
-				{
-					TArray<FString> Elements;
-					PropertyValue.ParseIntoArray(Elements, TEXT(","), true);
-					check(Elements.Num() == 3);
-					FVector* Value = Property->ContainerPtrToValuePtr<FVector>(this);
-					Value->X = FCString::Atod(*Elements[0]);
-					Value->Y = FCString::Atod(*Elements[1]);
-					Value->Z = FCString::Atod(*Elements[2]);
-				}
-				else if (Struct == TVariantStructure<FVector3f>::Get())
-				{
-					TArray<FString> Elements;
-					PropertyValue.ParseIntoArray(Elements, TEXT(","), true);
-					check(Elements.Num() == 3);
-					FVector3f* Value = Property->ContainerPtrToValuePtr<FVector3f>(this);
-					Value->X = FCString::Atof(*Elements[0]);
-					Value->Y = FCString::Atof(*Elements[1]);
-					Value->Z = FCString::Atof(*Elements[2]);
-				}
-				else if (Struct == TBaseStructure<FVector2D>::Get())
-				{
-					FVector2D* Value = Property->ContainerPtrToValuePtr<FVector2D>(this);
-					Value->InitFromString(PropertyValue);
-				}
+				TArray<FString> Elements;
+				PropertyValue.ParseIntoArray(Elements, TEXT(","), true);
+				check(Elements.Num() == 4);
+				FVector4d* Value = Property->ContainerPtrToValuePtr<FVector4d>(this);
+				Value->X = FCString::Atod(*Elements[0]);
+				Value->Y = FCString::Atod(*Elements[1]);
+				Value->Z = FCString::Atod(*Elements[2]);
+				Value->W = FCString::Atod(*Elements[3]);
+			}
+			else if (Struct == TBaseStructure<FVector>::Get())
+			{
+				TArray<FString> Elements;
+				PropertyValue.ParseIntoArray(Elements, TEXT(","), true);
+				check(Elements.Num() == 3);
+				FVector* Value = Property->ContainerPtrToValuePtr<FVector>(this);
+				Value->X = FCString::Atod(*Elements[0]);
+				Value->Y = FCString::Atod(*Elements[1]);
+				Value->Z = FCString::Atod(*Elements[2]);
+			}
+			else if (Struct == TVariantStructure<FVector3f>::Get())
+			{
+				TArray<FString> Elements;
+				PropertyValue.ParseIntoArray(Elements, TEXT(","), true);
+				check(Elements.Num() == 3);
+				FVector3f* Value = Property->ContainerPtrToValuePtr<FVector3f>(this);
+				Value->X = FCString::Atof(*Elements[0]);
+				Value->Y = FCString::Atof(*Elements[1]);
+				Value->Z = FCString::Atof(*Elements[2]);
+			}
+			else if (Struct == TBaseStructure<FVector2D>::Get())
+			{
+				FVector2D* Value = Property->ContainerPtrToValuePtr<FVector2D>(this);
+				Value->InitFromString(PropertyValue);
 			}
-
-			FPropertyChangedEvent Event(Property);
-			PostEditChangeProperty(Event);
 		}
 
-		// Update the expression preview and the material to reflect the change
-		Modify();
-		GetDefault<UMaterialGraphSchema>()->ForceVisualizationCacheClear();
-		const UMaterialGraphSchema* Schema = CastChecked<const UMaterialGraphSchema>(GraphNode->GetSchema());
-		Schema->UpdateMaterialOnDefaultValueChanged(GraphNode->GetGraph());
-		// There might be other properties affected by this property change (e.g. propertyA determines if propertyB is read-only) so refresh the detail view
-		Schema->UpdateDetailView(GraphNode->GetGraph());
+		FPropertyChangedEvent Event(Property);
+		PostEditChangeProperty(Event);
 	}
+
+	// Update the expression preview and the material to reflect the change
+	GetDefault<UMaterialGraphSchema>()->ForceVisualizationCacheClear();
+	const UMaterialGraphSchema* Schema = CastChecked<const UMaterialGraphSchema>(GraphNode->GetSchema());
+	Schema->UpdateMaterialOnDefaultValueChanged(GraphNode->GetGraph());
+	// There might be other properties affected by this property change (e.g. propertyA determines if propertyB is read-only) so refresh the detail view
+	Schema->UpdateDetailView(GraphNode->GetGraph());
 }
 
 FString UMaterialExpression::GetInputPinDefaultValue(int32 PinIndex)
@@ -1491,6 +1494,7 @@ void UMaterialExpression::PostEditChangeProperty(FPropertyChangedEvent& Property
 			// Update the pin value of the expression input
 			UEdGraphPin* Pin = GraphNode->GetPinAt(PinIndex);
 			check(Pin);
+			Pin->Modify();
 			Pin->DefaultValue = GetInputPinDefaultValue(PinIndex);
 
 			// If the property is linked as inline toggle to another property, both pins need updating to reflect the change.
