Added unit tests for transaction diffing

#rnx

diff --git a/Engine/Source/Runtime/Engine/Private/Tests/TransactionDiffingTests.cpp b/Engine/Source/Runtime/Engine/Private/Tests/TransactionDiffingTests.cpp
new file mode 100644
--- /dev/null
+++ b/Engine/Source/Runtime/Engine/Private/Tests/TransactionDiffingTests.cpp
@@ -0,0 +1,217 @@
+// Copyright Epic Games, Inc. All Rights Reserved.
+
+#include "Tests/TransactionDiffingTests.h"
+#include "Misc/AutomationTest.h"
+#include "Misc/TransactionObjectEvent.h"
+#include "TransactionCommon.h"
+
+void UTransactionDiffingTestObject::Serialize(FStructuredArchive::FRecord Record)
+{
+	Super::Serialize(Record);
+	Record << SA_VALUE(TEXT("NonPropertyData"), NonPropertyData);
+}
+
+#if WITH_DEV_AUTOMATION_TESTS
+
+namespace TransactionDiffingTests
+{
+
+FTransactionDiffableObject GetDiffableObject(UObject* Object)
+{
+	FTransactionDiffableObject DiffableObject;
+
+	UE::Transaction::FDiffableObjectDataWriter DiffWriter(DiffableObject);
+	Object->Serialize(DiffWriter);
+
+	return DiffableObject;
+}
+
+FTransactionObjectDeltaChange GenerateObjectDiff(const FTransactionDiffableObject& InitialObject, const FTransactionDiffableObject& ModifiedObject)
+{
+	FTransactionObjectDeltaChange DeltaChange;
+	UE::Transaction::DiffUtil::GenerateObjectDiff(InitialObject, ModifiedObject, DeltaChange);
+	return DeltaChange;
+}
+
+constexpr const uint32 TestFlags = EAutomationTestFlags::ApplicationContextMask | EAutomationTestFlags::EngineFilter;
+
+IMPLEMENT_SIMPLE_AUTOMATION_TEST(FEditPropertyDataTest, "System.Engine.Transactions.EditPropertyData", TestFlags)
+bool FEditPropertyDataTest::RunTest(const FString& Parameters)
+{
+	UTransactionDiffingTestObject* DefaultObject = GetMutableDefault<UTransactionDiffingTestObject>();
+	UTransactionDiffingTestObject* ModifiedObject = NewObject<UTransactionDiffingTestObject>();
+	
+	const FTransactionDiffableObject DefaultDiffableObject = GetDiffableObject(DefaultObject);
+
+	ModifiedObject->PropertyData = 10;
+
+	const FTransactionDiffableObject ModifiedDiffableObject = GetDiffableObject(ModifiedObject);
+
+	{
+		const FTransactionObjectDeltaChange DeltaChange = GenerateObjectDiff(DefaultDiffableObject, ModifiedDiffableObject);
+
+		TestFalse(TEXT("bHasNonPropertyChanges"), DeltaChange.bHasNonPropertyChanges);
+		TestEqual(TEXT("ChangedProperties.Num()"), DeltaChange.ChangedProperties.Num(), 1);
+		TestTrue(TEXT("ChangedProperties.Contains('PropertyData')"), DeltaChange.ChangedProperties.Contains(GET_MEMBER_NAME_CHECKED(UTransactionDiffingTestObject, PropertyData)));
+	}
+
+	return true;
+}
+
+IMPLEMENT_SIMPLE_AUTOMATION_TEST(FEditNonPropertyDataTest, "System.Engine.Transactions.EditNonPropertyData", TestFlags)
+bool FEditNonPropertyDataTest::RunTest(const FString& Parameters)
+{
+	UTransactionDiffingTestObject* DefaultObject = GetMutableDefault<UTransactionDiffingTestObject>();
+	UTransactionDiffingTestObject* ModifiedObject = NewObject<UTransactionDiffingTestObject>();
+
+	const FTransactionDiffableObject DefaultDiffableObject = GetDiffableObject(DefaultObject);
+
+	ModifiedObject->NonPropertyData = 10;
+
+	const FTransactionDiffableObject ModifiedDiffableObject = GetDiffableObject(ModifiedObject);
+
+	{
+		const FTransactionObjectDeltaChange DeltaChange = GenerateObjectDiff(DefaultDiffableObject, ModifiedDiffableObject);
+
+		TestTrue(TEXT("bHasNonPropertyChanges"), DeltaChange.bHasNonPropertyChanges);
+		TestEqual(TEXT("ChangedProperties.Num()"), DeltaChange.ChangedProperties.Num(), 0);
+	}
+
+	return true;
+}
+
+IMPLEMENT_SIMPLE_AUTOMATION_TEST(FEditNamesTest, "System.Engine.Transactions.EditNames", TestFlags)
+bool FEditNamesTest::RunTest(const FString& Parameters)
+{
+	UTransactionDiffingTestObject* DefaultObject = GetMutableDefault<UTransactionDiffingTestObject>();
+	UTransactionDiffingTestObject* ModifiedObject = NewObject<UTransactionDiffingTestObject>();
+
+	const FTransactionDiffableObject DefaultDiffableObject = GetDiffableObject(DefaultObject);
+
+	ModifiedObject->AdditionalName = "Test0";
+
+	const FTransactionDiffableObject ModifiedDiffableObject = GetDiffableObject(ModifiedObject);
+
+	{
+		const FTransactionObjectDeltaChange DeltaChange = GenerateObjectDiff(DefaultDiffableObject, ModifiedDiffableObject);
+
+		TestFalse(TEXT("bHasNonPropertyChanges"), DeltaChange.bHasNonPropertyChanges);
+		TestEqual(TEXT("ChangedProperties.Num()"), DeltaChange.ChangedProperties.Num(), 1);
+		TestTrue(TEXT("ChangedProperties.Contains('AdditionalName')"), DeltaChange.ChangedProperties.Contains(GET_MEMBER_NAME_CHECKED(UTransactionDiffingTestObject, AdditionalName)));
+	}
+
+	ModifiedObject->NamesArray.Add("Test1");
+
+	const FTransactionDiffableObject ModifiedDiffableObject2 = GetDiffableObject(ModifiedObject);
+
+	{
+		const FTransactionObjectDeltaChange DeltaChange = GenerateObjectDiff(DefaultDiffableObject, ModifiedDiffableObject2);
+
+		TestFalse(TEXT("bHasNonPropertyChanges"), DeltaChange.bHasNonPropertyChanges);
+		TestEqual(TEXT("ChangedProperties.Num()"), DeltaChange.ChangedProperties.Num(), 2);
+		TestTrue(TEXT("ChangedProperties.Contains('AdditionalName')"), DeltaChange.ChangedProperties.Contains(GET_MEMBER_NAME_CHECKED(UTransactionDiffingTestObject, AdditionalName)));
+		TestTrue(TEXT("ChangedProperties.Contains('NamesArray')"), DeltaChange.ChangedProperties.Contains(GET_MEMBER_NAME_CHECKED(UTransactionDiffingTestObject, NamesArray)));
+	}
+
+	{
+		const FTransactionObjectDeltaChange DeltaChange = GenerateObjectDiff(ModifiedDiffableObject, ModifiedDiffableObject2);
+
+		TestFalse(TEXT("bHasNonPropertyChanges"), DeltaChange.bHasNonPropertyChanges);
+		TestEqual(TEXT("ChangedProperties.Num()"), DeltaChange.ChangedProperties.Num(), 1);
+		TestTrue(TEXT("ChangedProperties.Contains('NamesArray')"), DeltaChange.ChangedProperties.Contains(GET_MEMBER_NAME_CHECKED(UTransactionDiffingTestObject, NamesArray)));
+	}
+
+	return true;
+}
+
+IMPLEMENT_SIMPLE_AUTOMATION_TEST(FEditObjectsTest, "System.Engine.Transactions.EditObjects", TestFlags)
+bool FEditObjectsTest::RunTest(const FString& Parameters)
+{
+	UTransactionDiffingTestObject* DefaultObject = GetMutableDefault<UTransactionDiffingTestObject>();
+	UTransactionDiffingTestObject* ModifiedObject = NewObject<UTransactionDiffingTestObject>();
+
+	const FTransactionDiffableObject DefaultDiffableObject = GetDiffableObject(DefaultObject);
+
+	ModifiedObject->AdditionalObject = NewObject<UTransactionDiffingTestObject>();
+
+	const FTransactionDiffableObject ModifiedDiffableObject = GetDiffableObject(ModifiedObject);
+
+	{
+		const FTransactionObjectDeltaChange DeltaChange = GenerateObjectDiff(DefaultDiffableObject, ModifiedDiffableObject);
+
+		TestFalse(TEXT("bHasNonPropertyChanges"), DeltaChange.bHasNonPropertyChanges);
+		TestEqual(TEXT("ChangedProperties.Num()"), DeltaChange.ChangedProperties.Num(), 1);
+		TestTrue(TEXT("ChangedProperties.Contains('AdditionalObject')"), DeltaChange.ChangedProperties.Contains(GET_MEMBER_NAME_CHECKED(UTransactionDiffingTestObject, AdditionalObject)));
+	}
+
+	ModifiedObject->ObjectsArray.Add(NewObject<UTransactionDiffingTestObject>());
+
+	const FTransactionDiffableObject ModifiedDiffableObject2 = GetDiffableObject(ModifiedObject);
+
+	{
+		const FTransactionObjectDeltaChange DeltaChange = GenerateObjectDiff(DefaultDiffableObject, ModifiedDiffableObject2);
+
+		TestFalse(TEXT("bHasNonPropertyChanges"), DeltaChange.bHasNonPropertyChanges);
+		TestEqual(TEXT("ChangedProperties.Num()"), DeltaChange.ChangedProperties.Num(), 2);
+		TestTrue(TEXT("ChangedProperties.Contains('AdditionalObject')"), DeltaChange.ChangedProperties.Contains(GET_MEMBER_NAME_CHECKED(UTransactionDiffingTestObject, AdditionalObject)));
+		TestTrue(TEXT("ChangedProperties.Contains('ObjectsArray')"), DeltaChange.ChangedProperties.Contains(GET_MEMBER_NAME_CHECKED(UTransactionDiffingTestObject, ObjectsArray)));
+	}
+
+	{
+		const FTransactionObjectDeltaChange DeltaChange = GenerateObjectDiff(ModifiedDiffableObject, ModifiedDiffableObject2);
+
+		TestFalse(TEXT("bHasNonPropertyChanges"), DeltaChange.bHasNonPropertyChanges);
+		TestEqual(TEXT("ChangedProperties.Num()"), DeltaChange.ChangedProperties.Num(), 1);
+		TestTrue(TEXT("ChangedProperties.Contains('ObjectsArray')"), DeltaChange.ChangedProperties.Contains(GET_MEMBER_NAME_CHECKED(UTransactionDiffingTestObject, ObjectsArray)));
+	}
+
+	return true;
+}
+
+IMPLEMENT_SIMPLE_AUTOMATION_TEST(FEditSoftObjectsTest, "System.Engine.Transactions.EditSoftObjects", TestFlags)
+bool FEditSoftObjectsTest::RunTest(const FString& Parameters)
+{
+	UTransactionDiffingTestObject* DefaultObject = GetMutableDefault<UTransactionDiffingTestObject>();
+	UTransactionDiffingTestObject* ModifiedObject = NewObject<UTransactionDiffingTestObject>();
+
+	const FTransactionDiffableObject DefaultDiffableObject = GetDiffableObject(DefaultObject);
+
+	ModifiedObject->AdditionalSoftObject = NewObject<UTransactionDiffingTestObject>();
+
+	const FTransactionDiffableObject ModifiedDiffableObject = GetDiffableObject(ModifiedObject);
+
+	{
+		const FTransactionObjectDeltaChange DeltaChange = GenerateObjectDiff(DefaultDiffableObject, ModifiedDiffableObject);
+
+		TestFalse(TEXT("bHasNonPropertyChanges"), DeltaChange.bHasNonPropertyChanges);
+		TestEqual(TEXT("ChangedProperties.Num()"), DeltaChange.ChangedProperties.Num(), 1);
+		TestTrue(TEXT("ChangedProperties.Contains('AdditionalSoftObject')"), DeltaChange.ChangedProperties.Contains(GET_MEMBER_NAME_CHECKED(UTransactionDiffingTestObject, AdditionalSoftObject)));
+	}
+
+	ModifiedObject->SoftObjectsArray.Add(NewObject<UTransactionDiffingTestObject>());
+
+	const FTransactionDiffableObject ModifiedDiffableObject2 = GetDiffableObject(ModifiedObject);
+
+	{
+		const FTransactionObjectDeltaChange DeltaChange = GenerateObjectDiff(DefaultDiffableObject, ModifiedDiffableObject2);
+
+		TestFalse(TEXT("bHasNonPropertyChanges"), DeltaChange.bHasNonPropertyChanges);
+		TestEqual(TEXT("ChangedProperties.Num()"), DeltaChange.ChangedProperties.Num(), 2);
+		TestTrue(TEXT("ChangedProperties.Contains('AdditionalSoftObject')"), DeltaChange.ChangedProperties.Contains(GET_MEMBER_NAME_CHECKED(UTransactionDiffingTestObject, AdditionalSoftObject)));
+		TestTrue(TEXT("ChangedProperties.Contains('SoftObjectsArray')"), DeltaChange.ChangedProperties.Contains(GET_MEMBER_NAME_CHECKED(UTransactionDiffingTestObject, SoftObjectsArray)));
+	}
+
+	{
+		const FTransactionObjectDeltaChange DeltaChange = GenerateObjectDiff(ModifiedDiffableObject, ModifiedDiffableObject2);
+
+		TestFalse(TEXT("bHasNonPropertyChanges"), DeltaChange.bHasNonPropertyChanges);
+		TestEqual(TEXT("ChangedProperties.Num()"), DeltaChange.ChangedProperties.Num(), 1);
+		TestTrue(TEXT("ChangedProperties.Contains('SoftObjectsArray')"), DeltaChange.ChangedProperties.Contains(GET_MEMBER_NAME_CHECKED(UTransactionDiffingTestObject, SoftObjectsArray)));
+	}
+
+	return true;
+}
+
+} // namespace TransactionDiffingTests
+
+#endif //WITH_DEV_AUTOMATION_TESTS
diff --git a/Engine/Source/Runtime/Engine/Private/Tests/TransactionDiffingTests.h b/Engine/Source/Runtime/Engine/Private/Tests/TransactionDiffingTests.h
new file mode 100644
--- /dev/null
+++ b/Engine/Source/Runtime/Engine/Private/Tests/TransactionDiffingTests.h
@@ -0,0 +1,38 @@
+// Copyright Epic Games, Inc. All Rights Reserved.
+
+#pragma once
+
+#include "UObject/Object.h"
+#include "TransactionDiffingTests.generated.h"
+
+UCLASS()
+class UTransactionDiffingTestObject : public UObject
+{
+	GENERATED_BODY()
+
+public:
+	virtual void Serialize(FStructuredArchive::FRecord Record) override;
+
+	UPROPERTY()
+	TArray<FName> NamesArray;
+
+	UPROPERTY()
+	FName AdditionalName;
+
+	UPROPERTY()
+	TArray<TObjectPtr<UObject>> ObjectsArray;
+
+	UPROPERTY()
+	TObjectPtr<UObject> AdditionalObject;
+
+	UPROPERTY()
+	TArray<TSoftObjectPtr<UObject>> SoftObjectsArray;
+
+	UPROPERTY()
+	TSoftObjectPtr<UObject> AdditionalSoftObject;
+
+	UPROPERTY()
+	int32 PropertyData = 0;
+
+	int32 NonPropertyData = 0;
+};
