Updated TransactionCommon types to no longer be USTRUCT

This was done as we were potentially going to serialize them as part of transaction annotations, but that is no longer the case, plus the diff snapshots are now 100% non-portable.

#rnx

diff --git a/Engine/Source/Editor/UnrealEd/Classes/Editor/Transactor.h b/Engine/Source/Editor/UnrealEd/Classes/Editor/Transactor.h
--- a/Engine/Source/Editor/UnrealEd/Classes/Editor/Transactor.h
+++ b/Engine/Source/Editor/UnrealEd/Classes/Editor/Transactor.h
@@ -38,13 +38,11 @@ class UNREALED_API FTransaction : public ITransaction
 	friend class FTransactionObjectEvent;
 
 protected:
-	using FPersistentObjectRef = FTransactionPersistentObjectRef;
-
 	// Record of an object.
 	class UNREALED_API FObjectRecord
 	{
 	public:
-		struct FSerializedObject : public FTransactionSerializedObject
+		struct FSerializedObject : public UE::Transaction::FSerializedObject
 		{
 			void SetObject(const UObject* InObject)
 			{
@@ -54,20 +52,20 @@ protected:
 
 			void Reset()
 			{
-				FTransactionSerializedObject::Reset();
+				UE::Transaction::FSerializedObject::Reset();
 				ObjectInfo.Reset();
 				ObjectAnnotation.Reset();
 			}
 
 			void Swap(FSerializedObject& Other)
 			{
-				FTransactionSerializedObject::Swap(Other);
+				UE::Transaction::FSerializedObject::Swap(Other);
 				ObjectInfo.Swap(Other.ObjectInfo);
 				Exchange(ObjectAnnotation, Other.ObjectAnnotation);
 			}
 
 			/** Information about the object when it was serialized */
-			FTransactionSerializedObjectInfo ObjectInfo;
+			UE::Transaction::FSerializedObjectInfo ObjectInfo;
 
 			/** Annotation data for the object stored externally */
 			TSharedPtr<ITransactionObjectAnnotation> ObjectAnnotation;
@@ -75,7 +73,7 @@ protected:
 
 		// Variables.
 		/** The object to track */
-		FPersistentObjectRef Object;
+		UE::Transaction::FPersistentObjectRef Object;
 		/** Custom change to apply to this object to undo this record.  Executing the undo will return an object that can be used to redo the change. */
 		TUniquePtr<FChange> CustomChange;
 		/** Array: If an array object, reference to script array */
@@ -110,9 +108,9 @@ protected:
 		/** The serialized object data that will be used when the transaction is flipped */
 		FSerializedObject	SerializedObjectFlip;
 		/** The diffable object data (always null once finalized) */
-		TUniquePtr<FTransactionDiffableObject> DiffableObject;
+		TUniquePtr<UE::Transaction::FDiffableObject> DiffableObject;
 		/** The diffable object data when it was last snapshot (always null once finalized) */
-		TUniquePtr<FTransactionDiffableObject> DiffableObjectSnapshot;
+		TUniquePtr<UE::Transaction::FDiffableObject> DiffableObjectSnapshot;
 		/** Delta change information between the diffable state of the object when the transaction started, and the diffable state of the object when the transaction ended */
 		FTransactionObjectDeltaChange DeltaChange;
 
@@ -168,7 +166,7 @@ protected:
 			{
 				if (Owner)
 				{
-					if (const FObjectRecords* ObjectRecords = Owner->ObjectRecordsMap.Find(FPersistentObjectRef(InObject)))
+					if (const FObjectRecords* ObjectRecords = Owner->ObjectRecordsMap.Find(UE::Transaction::FPersistentObjectRef(InObject)))
 					{
 						for (FObjectRecord* Record : ObjectRecords->Records)
 						{
@@ -226,7 +224,7 @@ protected:
 	TIndirectArray<FObjectRecord> Records;
 
 	/** Map of object records (non-array), for optimized look-up and to prevent an object being serialized to a transaction more than once. */
-	TMap<FPersistentObjectRef, FObjectRecords> ObjectRecordsMap;
+	TMap<UE::Transaction::FPersistentObjectRef, FObjectRecords> ObjectRecordsMap;
 
 	/** Unique identifier for this transaction, used to track it during its lifetime */
 	FGuid		Id;
diff --git a/Engine/Source/Editor/UnrealEd/Private/EditorTransaction.cpp b/Engine/Source/Editor/UnrealEd/Private/EditorTransaction.cpp
--- a/Engine/Source/Editor/UnrealEd/Private/EditorTransaction.cpp
+++ b/Engine/Source/Editor/UnrealEd/Private/EditorTransaction.cpp
@@ -70,7 +70,7 @@ FTransaction::FObjectRecord::FObjectRecord(FTransaction* Owner, UObject* InObjec
 
 		if (!Array)
 		{
-			DiffableObject = MakeUnique<FTransactionDiffableObject>();
+			DiffableObject = MakeUnique<UE::Transaction::FDiffableObject>();
 			DiffableObject->SetObject(CurrentObject);
 			UE::Transaction::FDiffableObjectDataWriter Writer(*DiffableObject);
 			SerializeObject(Writer);
@@ -291,7 +291,7 @@ void FTransaction::FObjectRecord::Finalize( FTransaction* Owner, TSharedPtr<ITra
 				|| SerializedObject.ReferencedObjects != SerializedObjectFlip.ReferencedObjects;
 
 			// Serialize the object so we can diff it
-			FTransactionDiffableObject CurrentDiffableObject;
+			UE::Transaction::FDiffableObject CurrentDiffableObject;
 			{
 				CurrentDiffableObject.SetObject(CurrentObject);
 				UE::Transaction::FDiffableObjectDataWriter Writer(CurrentDiffableObject);
@@ -339,7 +339,7 @@ void FTransaction::FObjectRecord::Snapshot( FTransaction* Owner, TArrayView<cons
 	if (CurrentObject)
 	{
 		// Serialize the object so we can diff it
-		FTransactionDiffableObject CurrentDiffableObject;
+		UE::Transaction::FDiffableObject CurrentDiffableObject;
 		{
 			CurrentDiffableObject.SetObject(CurrentObject);
 			UE::Transaction::FDiffableObjectDataWriter Writer(CurrentDiffableObject, Properties);
@@ -349,14 +349,14 @@ void FTransaction::FObjectRecord::Snapshot( FTransaction* Owner, TArrayView<cons
 		}
 
 		// Diff against the correct serialized data depending on whether we already had a snapshot
-		const FTransactionDiffableObject& InitialDiffableObject = DiffableObjectSnapshot ? *DiffableObjectSnapshot : *DiffableObject;
+		const UE::Transaction::FDiffableObject& InitialDiffableObject = DiffableObjectSnapshot ? *DiffableObjectSnapshot : *DiffableObject;
 		FTransactionObjectDeltaChange SnapshotDeltaChange;
 		UE::Transaction::DiffUtil::GenerateObjectDiff(InitialDiffableObject, CurrentDiffableObject, SnapshotDeltaChange, /*bFullDiff*/false);
 
 		// Update the snapshot data for next time
 		if (!DiffableObjectSnapshot)
 		{
-			DiffableObjectSnapshot = MakeUnique<FTransactionDiffableObject>();
+			DiffableObjectSnapshot = MakeUnique<UE::Transaction::FDiffableObject>();
 		}
 		DiffableObjectSnapshot->Swap(CurrentDiffableObject);
 
@@ -484,13 +484,13 @@ FArchive& operator<<( FArchive& Ar, FTransaction::FObjectRecord& R )
 
 void FTransaction::FObjectRecord::AddReferencedObjects( FReferenceCollector& Collector )
 {
-	Object.AddStructReferencedObjects(Collector);
+	Object.AddReferencedObjects(Collector);
 
 	auto AddSerializedObjectReferences = [&Collector](FSerializedObject& InSerializedObject)
 	{
-		for (FPersistentObjectRef& ReferencedObject : InSerializedObject.ReferencedObjects)
+		for (UE::Transaction::FPersistentObjectRef& ReferencedObject : InSerializedObject.ReferencedObjects)
 		{
-			ReferencedObject.AddStructReferencedObjects(Collector);
+			ReferencedObject.AddReferencedObjects(Collector);
 		}
 
 		if (InSerializedObject.ObjectAnnotation.IsValid())
@@ -519,7 +519,7 @@ bool FTransaction::FObjectRecord::ContainsPieObject() const
 
 	auto SerializedObjectContainPieObjects = [](const FSerializedObject& InSerializedObject) -> bool
 	{
-		for (const FPersistentObjectRef& ReferencedObject : InSerializedObject.ReferencedObjects)
+		for (const UE::Transaction::FPersistentObjectRef& ReferencedObject : InSerializedObject.ReferencedObjects)
 		{
 			const UObject* Obj = ReferencedObject.Get();
 			if (Obj && Obj->GetOutermost()->HasAnyPackageFlags(PKG_PlayInEditor))
@@ -563,9 +563,9 @@ void FTransaction::AddReferencedObjects( FReferenceCollector& Collector )
 		ObjectRecord.AddReferencedObjects( Collector );
 	}
 
-	for (TTuple<FPersistentObjectRef, FObjectRecords>& ObjectRecordsPair : ObjectRecordsMap)
+	for (TTuple<UE::Transaction::FPersistentObjectRef, FObjectRecords>& ObjectRecordsPair : ObjectRecordsMap)
 	{
-		ObjectRecordsPair.Key.AddStructReferencedObjects(Collector);
+		ObjectRecordsPair.Key.AddReferencedObjects(Collector);
 	}
 }
 
@@ -574,7 +574,7 @@ void FTransaction::SaveObject( UObject* Object )
 	check(Object);
 	Object->CheckDefaultSubobjects();
 
-	FObjectRecords& ObjectRecords = ObjectRecordsMap.FindOrAdd(FPersistentObjectRef(Object));
+	FObjectRecords& ObjectRecords = ObjectRecordsMap.FindOrAdd(UE::Transaction::FPersistentObjectRef(Object));
 	if (ObjectRecords.Records.Num() == 0)
 	{
 		// Save the object.
@@ -613,7 +613,7 @@ void FTransaction::StoreUndo(UObject* Object, TUniquePtr<FChange> UndoChange)
 	Object->CheckDefaultSubobjects();
 
 	// Save the undo record
-	FObjectRecords& ObjectRecords = ObjectRecordsMap.FindOrAdd(FPersistentObjectRef(Object));
+	FObjectRecords& ObjectRecords = ObjectRecordsMap.FindOrAdd(UE::Transaction::FPersistentObjectRef(Object));
 	FObjectRecord* UndoRecord = ObjectRecords.Records.Add_GetRef(new FObjectRecord(this, Object, MoveTemp(UndoChange), nullptr, 0, 0, 0, 0, 0, nullptr, nullptr, nullptr));
 	Records.Add(UndoRecord);
 }
@@ -630,7 +630,7 @@ void FTransaction::SnapshotObject( UObject* InObject, TArrayView<const FProperty
 {
 	check(InObject);
 
-	if (const FObjectRecords* ObjectRecords = ObjectRecordsMap.Find(FPersistentObjectRef(InObject)))
+	if (const FObjectRecords* ObjectRecords = ObjectRecordsMap.Find(UE::Transaction::FPersistentObjectRef(InObject)))
 	{
 		for (FObjectRecord* Record : ObjectRecords->Records)
 		{
@@ -645,7 +645,7 @@ void FTransaction::SnapshotObject( UObject* InObject, TArrayView<const FProperty
 
 bool FTransaction::ContainsObject(const UObject* Object) const
 {
-	FPersistentObjectRef PersistentObjectRef(const_cast<UObject*>(Object));
+	UE::Transaction::FPersistentObjectRef PersistentObjectRef(const_cast<UObject*>(Object));
 	if (const FObjectRecords* ObjectRecords = ObjectRecordsMap.Find(PersistentObjectRef))
 	{
 		for (FObjectRecord* Record : ObjectRecords->Records)
diff --git a/Engine/Source/Runtime/Engine/Private/Tests/TransactionDiffingTests.cpp b/Engine/Source/Runtime/Engine/Private/Tests/TransactionDiffingTests.cpp
--- a/Engine/Source/Runtime/Engine/Private/Tests/TransactionDiffingTests.cpp
+++ b/Engine/Source/Runtime/Engine/Private/Tests/TransactionDiffingTests.cpp
@@ -16,9 +16,9 @@ void UTransactionDiffingTestObject::Serialize(FStructuredArchive::FRecord Record
 namespace TransactionDiffingTests
 {
 
-FTransactionDiffableObject GetDiffableObject(UObject* Object)
+UE::Transaction::FDiffableObject GetDiffableObject(UObject* Object)
 {
-	FTransactionDiffableObject DiffableObject;
+	UE::Transaction::FDiffableObject DiffableObject;
 
 	UE::Transaction::FDiffableObjectDataWriter DiffWriter(DiffableObject);
 	Object->Serialize(DiffWriter);
@@ -26,7 +26,7 @@ FTransactionDiffableObject GetDiffableObject(UObject* Object)
 	return DiffableObject;
 }
 
-FTransactionObjectDeltaChange GenerateObjectDiff(const FTransactionDiffableObject& InitialObject, const FTransactionDiffableObject& ModifiedObject)
+FTransactionObjectDeltaChange GenerateObjectDiff(const UE::Transaction::FDiffableObject& InitialObject, const UE::Transaction::FDiffableObject& ModifiedObject)
 {
 	FTransactionObjectDeltaChange DeltaChange;
 	UE::Transaction::DiffUtil::GenerateObjectDiff(InitialObject, ModifiedObject, DeltaChange);
@@ -41,11 +41,11 @@ bool FEditPropertyDataTest::RunTest(const FString& Parameters)
 	UTransactionDiffingTestObject* DefaultObject = GetMutableDefault<UTransactionDiffingTestObject>();
 	UTransactionDiffingTestObject* ModifiedObject = NewObject<UTransactionDiffingTestObject>();
 	
-	const FTransactionDiffableObject DefaultDiffableObject = GetDiffableObject(DefaultObject);
+	const UE::Transaction::FDiffableObject DefaultDiffableObject = GetDiffableObject(DefaultObject);
 
 	ModifiedObject->PropertyData = 10;
 
-	const FTransactionDiffableObject ModifiedDiffableObject = GetDiffableObject(ModifiedObject);
+	const UE::Transaction::FDiffableObject ModifiedDiffableObject = GetDiffableObject(ModifiedObject);
 
 	{
 		const FTransactionObjectDeltaChange DeltaChange = GenerateObjectDiff(DefaultDiffableObject, ModifiedDiffableObject);
@@ -64,11 +64,11 @@ bool FEditNonPropertyDataTest::RunTest(const FString& Parameters)
 	UTransactionDiffingTestObject* DefaultObject = GetMutableDefault<UTransactionDiffingTestObject>();
 	UTransactionDiffingTestObject* ModifiedObject = NewObject<UTransactionDiffingTestObject>();
 
-	const FTransactionDiffableObject DefaultDiffableObject = GetDiffableObject(DefaultObject);
+	const UE::Transaction::FDiffableObject DefaultDiffableObject = GetDiffableObject(DefaultObject);
 
 	ModifiedObject->NonPropertyData = 10;
 
-	const FTransactionDiffableObject ModifiedDiffableObject = GetDiffableObject(ModifiedObject);
+	const UE::Transaction::FDiffableObject ModifiedDiffableObject = GetDiffableObject(ModifiedObject);
 
 	{
 		const FTransactionObjectDeltaChange DeltaChange = GenerateObjectDiff(DefaultDiffableObject, ModifiedDiffableObject);
@@ -86,11 +86,11 @@ bool FEditNamesTest::RunTest(const FString& Parameters)
 	UTransactionDiffingTestObject* DefaultObject = GetMutableDefault<UTransactionDiffingTestObject>();
 	UTransactionDiffingTestObject* ModifiedObject = NewObject<UTransactionDiffingTestObject>();
 
-	const FTransactionDiffableObject DefaultDiffableObject = GetDiffableObject(DefaultObject);
+	const UE::Transaction::FDiffableObject DefaultDiffableObject = GetDiffableObject(DefaultObject);
 
 	ModifiedObject->AdditionalName = "Test0";
 
-	const FTransactionDiffableObject ModifiedDiffableObject = GetDiffableObject(ModifiedObject);
+	const UE::Transaction::FDiffableObject ModifiedDiffableObject = GetDiffableObject(ModifiedObject);
 
 	{
 		const FTransactionObjectDeltaChange DeltaChange = GenerateObjectDiff(DefaultDiffableObject, ModifiedDiffableObject);
@@ -102,7 +102,7 @@ bool FEditNamesTest::RunTest(const FString& Parameters)
 
 	ModifiedObject->NamesArray.Add("Test1");
 
-	const FTransactionDiffableObject ModifiedDiffableObject2 = GetDiffableObject(ModifiedObject);
+	const UE::Transaction::FDiffableObject ModifiedDiffableObject2 = GetDiffableObject(ModifiedObject);
 
 	{
 		const FTransactionObjectDeltaChange DeltaChange = GenerateObjectDiff(DefaultDiffableObject, ModifiedDiffableObject2);
@@ -130,11 +130,11 @@ bool FEditObjectsTest::RunTest(const FString& Parameters)
 	UTransactionDiffingTestObject* DefaultObject = GetMutableDefault<UTransactionDiffingTestObject>();
 	UTransactionDiffingTestObject* ModifiedObject = NewObject<UTransactionDiffingTestObject>();
 
-	const FTransactionDiffableObject DefaultDiffableObject = GetDiffableObject(DefaultObject);
+	const UE::Transaction::FDiffableObject DefaultDiffableObject = GetDiffableObject(DefaultObject);
 
 	ModifiedObject->AdditionalObject = NewObject<UTransactionDiffingTestObject>();
 
-	const FTransactionDiffableObject ModifiedDiffableObject = GetDiffableObject(ModifiedObject);
+	const UE::Transaction::FDiffableObject ModifiedDiffableObject = GetDiffableObject(ModifiedObject);
 
 	{
 		const FTransactionObjectDeltaChange DeltaChange = GenerateObjectDiff(DefaultDiffableObject, ModifiedDiffableObject);
@@ -146,7 +146,7 @@ bool FEditObjectsTest::RunTest(const FString& Parameters)
 
 	ModifiedObject->ObjectsArray.Add(NewObject<UTransactionDiffingTestObject>());
 
-	const FTransactionDiffableObject ModifiedDiffableObject2 = GetDiffableObject(ModifiedObject);
+	const UE::Transaction::FDiffableObject ModifiedDiffableObject2 = GetDiffableObject(ModifiedObject);
 
 	{
 		const FTransactionObjectDeltaChange DeltaChange = GenerateObjectDiff(DefaultDiffableObject, ModifiedDiffableObject2);
@@ -174,11 +174,11 @@ bool FEditSoftObjectsTest::RunTest(const FString& Parameters)
 	UTransactionDiffingTestObject* DefaultObject = GetMutableDefault<UTransactionDiffingTestObject>();
 	UTransactionDiffingTestObject* ModifiedObject = NewObject<UTransactionDiffingTestObject>();
 
-	const FTransactionDiffableObject DefaultDiffableObject = GetDiffableObject(DefaultObject);
+	const UE::Transaction::FDiffableObject DefaultDiffableObject = GetDiffableObject(DefaultObject);
 
 	ModifiedObject->AdditionalSoftObject = NewObject<UTransactionDiffingTestObject>();
 
-	const FTransactionDiffableObject ModifiedDiffableObject = GetDiffableObject(ModifiedObject);
+	const UE::Transaction::FDiffableObject ModifiedDiffableObject = GetDiffableObject(ModifiedObject);
 
 	{
 		const FTransactionObjectDeltaChange DeltaChange = GenerateObjectDiff(DefaultDiffableObject, ModifiedDiffableObject);
@@ -190,7 +190,7 @@ bool FEditSoftObjectsTest::RunTest(const FString& Parameters)
 
 	ModifiedObject->SoftObjectsArray.Add(NewObject<UTransactionDiffingTestObject>());
 
-	const FTransactionDiffableObject ModifiedDiffableObject2 = GetDiffableObject(ModifiedObject);
+	const UE::Transaction::FDiffableObject ModifiedDiffableObject2 = GetDiffableObject(ModifiedObject);
 
 	{
 		const FTransactionObjectDeltaChange DeltaChange = GenerateObjectDiff(DefaultDiffableObject, ModifiedDiffableObject2);
diff --git a/Engine/Source/Runtime/Engine/Private/TransactionCommon.cpp b/Engine/Source/Runtime/Engine/Private/TransactionCommon.cpp
--- a/Engine/Source/Runtime/Engine/Private/TransactionCommon.cpp
+++ b/Engine/Source/Runtime/Engine/Private/TransactionCommon.cpp
@@ -5,7 +5,10 @@
 #include "UObject/UnrealType.h"
 #include "Components/ActorComponent.h"
 
-FTransactionPersistentObjectRef::FTransactionPersistentObjectRef(UObject* InObject)
+namespace UE::Transaction
+{
+
+FPersistentObjectRef::FPersistentObjectRef(UObject* InObject)
 {
 	RootObject = InObject;
 	{
@@ -53,7 +56,7 @@ FTransactionPersistentObjectRef::FTransactionPersistentObjectRef(UObject* InObje
 	checkSlow(Get() == InObject);
 }
 
-UObject* FTransactionPersistentObjectRef::Get() const
+UObject* FPersistentObjectRef::Get() const
 {
 	if (ReferenceType == EReferenceType::SubObject)
 	{
@@ -116,7 +119,7 @@ UObject* FTransactionPersistentObjectRef::Get() const
 	return RootObject;
 }
 
-void FTransactionPersistentObjectRef::AddStructReferencedObjects(FReferenceCollector& Collector)
+void FPersistentObjectRef::AddReferencedObjects(FReferenceCollector& Collector)
 {
 	Collector.AddReferencedObject(RootObject);
 
@@ -136,17 +139,17 @@ void FTransactionPersistentObjectRef::AddStructReferencedObjects(FReferenceColle
 }
 
 
-FTransactionSerializedTaggedData FTransactionSerializedTaggedData::FromOffsetAndSize(const int64 InOffset, const int64 InSize)
+FSerializedTaggedData FSerializedTaggedData::FromOffsetAndSize(const int64 InOffset, const int64 InSize)
 {
-	return FTransactionSerializedTaggedData{ InOffset, InSize };
+	return FSerializedTaggedData{ InOffset, InSize };
 }
 
-FTransactionSerializedTaggedData FTransactionSerializedTaggedData::FromStartAndEnd(const int64 InStart, const int64 InEnd)
+FSerializedTaggedData FSerializedTaggedData::FromStartAndEnd(const int64 InStart, const int64 InEnd)
 {
-	return FTransactionSerializedTaggedData{ InStart, InEnd - InStart };
+	return FSerializedTaggedData{ InStart, InEnd - InStart };
 }
 
-void FTransactionSerializedTaggedData::AppendSerializedData(const int64 InOffset, const int64 InSize)
+void FSerializedTaggedData::AppendSerializedData(const int64 InOffset, const int64 InSize)
 {
 	if (DataOffset == INDEX_NONE)
 	{
@@ -160,7 +163,7 @@ void FTransactionSerializedTaggedData::AppendSerializedData(const int64 InOffset
 	}
 }
 
-void FTransactionSerializedTaggedData::AppendSerializedData(const FTransactionSerializedTaggedData& InData)
+void FSerializedTaggedData::AppendSerializedData(const FSerializedTaggedData& InData)
 {
 	if (InData.HasSerializedData())
 	{
@@ -168,19 +171,19 @@ void FTransactionSerializedTaggedData::AppendSerializedData(const FTransactionSe
 	}
 }
 
-bool FTransactionSerializedTaggedData::HasSerializedData() const
+bool FSerializedTaggedData::HasSerializedData() const
 {
 	return DataOffset != INDEX_NONE && DataSize != 0;
 }
 
 
-void FTransactionSerializedObjectData::Read(void* Dest, int64 Offset, int64 Num) const
+void FSerializedObjectData::Read(void* Dest, int64 Offset, int64 Num) const
 {
 	checkSlow(Offset + Num <= Data.Num());
 	FMemory::Memcpy(Dest, &Data[Offset], Num);
 }
 
-void FTransactionSerializedObjectData::Write(const void* Src, int64 Offset, int64 Num)
+void FSerializedObjectData::Write(const void* Src, int64 Offset, int64 Num)
 {
 	if (Offset == Data.Num())
 	{
@@ -190,13 +193,10 @@ void FTransactionSerializedObjectData::Write(const void* Src, int64 Offset, int6
 }
 
 
-namespace UE::Transaction
-{
-
 const FName TaggedDataKey_UnknownData = ".UnknownData";
 const FName TaggedDataKey_ScriptData = ".ScriptData";
 
-FSerializedObjectDataReader::FSerializedObjectDataReader(const FTransactionSerializedObject& InSerializedObject)
+FSerializedObjectDataReader::FSerializedObjectDataReader(const FSerializedObject& InSerializedObject)
 	: SerializedObject(InSerializedObject)
 {
 	SetIsLoading(true);
@@ -238,7 +238,7 @@ FArchive& FSerializedObjectDataReader::operator<<(UObject*& Res)
 namespace Internal
 {
 
-FSerializedObjectDataWriterCommon::FSerializedObjectDataWriterCommon(FTransactionSerializedObjectData& InSerializedData)
+FSerializedObjectDataWriterCommon::FSerializedObjectDataWriterCommon(FSerializedObjectData& InSerializedData)
 	: SerializedData(InSerializedData)
 {
 	SetIsSaving(true);
@@ -259,7 +259,7 @@ void FSerializedObjectDataWriterCommon::Serialize(void* SerData, int64 Num)
 } // namespace Internal
 
 
-FSerializedObjectDataWriter::FSerializedObjectDataWriter(FTransactionSerializedObject& InSerializedObject)
+FSerializedObjectDataWriter::FSerializedObjectDataWriter(FSerializedObject& InSerializedObject)
 	: Internal::FSerializedObjectDataWriterCommon(InSerializedObject.SerializedData)
 	, SerializedObject(InSerializedObject)
 {
@@ -301,7 +301,7 @@ FArchive& FSerializedObjectDataWriter::operator<<(UObject*& Res)
 	}
 	else if (Res)
 	{
-		ObjectIndex = SerializedObject.ReferencedObjects.Add(FTransactionPersistentObjectRef(Res));
+		ObjectIndex = SerializedObject.ReferencedObjects.Add(FPersistentObjectRef(Res));
 		ObjectMap.Add(Res, ObjectIndex);
 	}
 
@@ -309,7 +309,7 @@ FArchive& FSerializedObjectDataWriter::operator<<(UObject*& Res)
 }
 
 
-FDiffableObjectDataWriter::FDiffableObjectDataWriter(FTransactionDiffableObject& InDiffableObject, TArrayView<const FProperty*> InPropertiesToSerialize)
+FDiffableObjectDataWriter::FDiffableObjectDataWriter(FDiffableObject& InDiffableObject, TArrayView<const FProperty*> InPropertiesToSerialize)
 	: Internal::FSerializedObjectDataWriterCommon(InDiffableObject.SerializedData)
 	, DiffableObject(InDiffableObject)
 	, PropertiesToSerialize(InPropertiesToSerialize)
@@ -398,7 +398,7 @@ void FDiffableObjectDataWriter::OnDataSerialized(int64 InOffset, int64 InNum)
 	const FName SerializedTaggedDataKey = GetTaggedDataKey();
 	if (!SerializedTaggedDataKey.IsNone())
 	{
-		FTransactionSerializedTaggedData& SerializedTaggedData = DiffableObject.SerializedTaggedData.FindOrAdd(SerializedTaggedDataKey);
+		FSerializedTaggedData& SerializedTaggedData = DiffableObject.SerializedTaggedData.FindOrAdd(SerializedTaggedDataKey);
 		SerializedTaggedData.AppendSerializedData(InOffset, InNum);
 	}
 }
@@ -440,9 +440,9 @@ FName FDiffableObjectDataWriter::FCachedPropertyKey::SyncCache(const FArchiveSer
 namespace DiffUtil
 {
 
-void GenerateObjectDiff(const FTransactionDiffableObject& OldDiffableObject, const FTransactionDiffableObject& NewDiffableObject, FTransactionObjectDeltaChange& OutDeltaChange, const bool bFullDiff)
+void GenerateObjectDiff(const FDiffableObject& OldDiffableObject, const FDiffableObject& NewDiffableObject, FTransactionObjectDeltaChange& OutDeltaChange, const bool bFullDiff)
 {
-	auto IsTaggedDataBlockIdentical = [&OldDiffableObject, &NewDiffableObject, &OutDeltaChange](const FName TaggedDataKey, const FTransactionSerializedTaggedData& OldSerializedTaggedData, const FTransactionSerializedTaggedData& NewSerializedTaggedData) -> bool
+	auto IsTaggedDataBlockIdentical = [&OldDiffableObject, &NewDiffableObject, &OutDeltaChange](const FName TaggedDataKey, const FSerializedTaggedData& OldSerializedTaggedData, const FSerializedTaggedData& NewSerializedTaggedData) -> bool
 	{
 		// Binary compare the serialized data to see if something has changed for this property
 		bool bIsTaggedDataIdentical = OldSerializedTaggedData.DataSize == NewSerializedTaggedData.DataSize;
@@ -477,14 +477,14 @@ void GenerateObjectDiff(const FTransactionDiffableObject& OldDiffableObject, con
 		OutDeltaChange.bHasPendingKillChange |= OldDiffableObject.ObjectInfo.bIsPendingKill != NewDiffableObject.ObjectInfo.bIsPendingKill;
 	}
 
-	for (const TPair<FName, FTransactionSerializedTaggedData>& NewNamePropertyPair : NewDiffableObject.SerializedTaggedData)
+	for (const TPair<FName, FSerializedTaggedData>& NewNamePropertyPair : NewDiffableObject.SerializedTaggedData)
 	{
 		if (!ShouldCompareTaggedData(NewNamePropertyPair.Key))
 		{
 			continue;
 		}
 
-		const FTransactionSerializedTaggedData* OldSerializedTaggedData = OldDiffableObject.SerializedTaggedData.Find(NewNamePropertyPair.Key);
+		const FSerializedTaggedData* OldSerializedTaggedData = OldDiffableObject.SerializedTaggedData.Find(NewNamePropertyPair.Key);
 		if (ShouldCompareAsNonPropertyData(NewNamePropertyPair.Key))
 		{
 			if (bFullDiff && !OutDeltaChange.bHasNonPropertyChanges && (!OldSerializedTaggedData || !IsTaggedDataBlockIdentical(NewNamePropertyPair.Key, *OldSerializedTaggedData, NewNamePropertyPair.Value)))
@@ -508,7 +508,7 @@ void GenerateObjectDiff(const FTransactionDiffableObject& OldDiffableObject, con
 
 	if (bFullDiff)
 	{
-		for (const TPair<FName, FTransactionSerializedTaggedData>& OldNamePropertyPair : OldDiffableObject.SerializedTaggedData)
+		for (const TPair<FName, FSerializedTaggedData>& OldNamePropertyPair : OldDiffableObject.SerializedTaggedData)
 		{
 			if (!ShouldCompareTaggedData(OldNamePropertyPair.Key) || NewDiffableObject.SerializedTaggedData.Contains(OldNamePropertyPair.Key))
 			{
diff --git a/Engine/Source/Runtime/Engine/Public/TransactionCommon.h b/Engine/Source/Runtime/Engine/Public/TransactionCommon.h
--- a/Engine/Source/Runtime/Engine/Public/TransactionCommon.h
+++ b/Engine/Source/Runtime/Engine/Public/TransactionCommon.h
@@ -7,11 +7,13 @@
 #include "UObject/WeakObjectPtr.h"
 #include "Serialization/ArchiveUObject.h"
 #include "Serialization/ArchiveSerializedPropertyChain.h"
-#include "TransactionCommon.generated.h"
 
 class FReferenceCollector;
 struct FTransactionObjectDeltaChange;
 
+namespace UE::Transaction
+{
+
 /**
 	This type is necessary because the blueprint system is destroying and creating
 	CDOs at edit time (usually on compile, but also on load), but also stores user
@@ -24,14 +26,11 @@ struct FTransactionObjectDeltaChange;
 	(rather than the CDO) to store editor data. The CDO can then be re-instanced (or not)
 	as runtime code requires.
 */
-USTRUCT()
-struct ENGINE_API FTransactionPersistentObjectRef
+struct ENGINE_API FPersistentObjectRef
 {
-	GENERATED_BODY()
-
 public:
-	FTransactionPersistentObjectRef() = default;
-	explicit FTransactionPersistentObjectRef(UObject* InObject);
+	FPersistentObjectRef() = default;
+	explicit FPersistentObjectRef(UObject* InObject);
 
 	bool Serialize(FArchive& Ar)
 	{
@@ -39,7 +38,7 @@ public:
 		return true;
 	}
 
-	friend FArchive& operator<<(FArchive& Ar, FTransactionPersistentObjectRef& ReferencedObject)
+	friend FArchive& operator<<(FArchive& Ar, FPersistentObjectRef& ReferencedObject)
 	{
 		Ar << (std::underlying_type_t<EReferenceType>&)ReferencedObject.ReferenceType;
 		Ar << ReferencedObject.RootObject;
@@ -47,19 +46,19 @@ public:
 		return Ar;
 	}
 
-	friend bool operator==(const FTransactionPersistentObjectRef& LHS, const FTransactionPersistentObjectRef& RHS)
+	friend bool operator==(const FPersistentObjectRef& LHS, const FPersistentObjectRef& RHS)
 	{
 		return LHS.ReferenceType == RHS.ReferenceType
 			&& LHS.RootObject == RHS.RootObject
 			&& (LHS.ReferenceType != EReferenceType::SubObject || LHS.SubObjectHierarchyIDs == RHS.SubObjectHierarchyIDs);
 	}
 
-	friend bool operator!=(const FTransactionPersistentObjectRef& LHS, const FTransactionPersistentObjectRef& RHS)
+	friend bool operator!=(const FPersistentObjectRef& LHS, const FPersistentObjectRef& RHS)
 	{
 		return !(LHS == RHS);
 	}
 
-	friend uint32 GetTypeHash(const FTransactionPersistentObjectRef& InObjRef)
+	friend uint32 GetTypeHash(const FPersistentObjectRef& InObjRef)
 	{
 		return HashCombine(GetTypeHash(InObjRef.ReferenceType), GetTypeHash(InObjRef.RootObject));
 	}
@@ -76,7 +75,7 @@ public:
 
 	UObject* Get() const;
 	
-	void AddStructReferencedObjects(FReferenceCollector& Collector);
+	void AddReferencedObjects(FReferenceCollector& Collector);
 
 private:
 	/** This enum represents all of the different special cases we are handling with this type */
@@ -100,28 +99,14 @@ private:
 	mutable TArray<TWeakObjectPtr<UObject>, TInlineAllocator<4>> CachedSubObjectHierarchy;
 };
 
-template<>
-struct TStructOpsTypeTraits<FTransactionPersistentObjectRef> : public TStructOpsTypeTraitsBase2<FTransactionPersistentObjectRef>
-{
-	enum
-	{
-		WithIdenticalViaEquality = true,
-		WithAddStructReferencedObjects = true,
-		WithSerializer = true,
-	};
-};
-
-USTRUCT()
-struct ENGINE_API FTransactionSerializedTaggedData
+struct ENGINE_API FSerializedTaggedData
 {
-	GENERATED_BODY()
-
 public:
-	static FTransactionSerializedTaggedData FromOffsetAndSize(const int64 InOffset, const int64 InSize);
-	static FTransactionSerializedTaggedData FromStartAndEnd(const int64 InStart, const int64 InEnd);
+	static FSerializedTaggedData FromOffsetAndSize(const int64 InOffset, const int64 InSize);
+	static FSerializedTaggedData FromStartAndEnd(const int64 InStart, const int64 InEnd);
 
 	void AppendSerializedData(const int64 InOffset, const int64 InSize);
-	void AppendSerializedData(const FTransactionSerializedTaggedData& InData);
+	void AppendSerializedData(const FSerializedTaggedData& InData);
 
 	bool HasSerializedData() const;
 
@@ -135,31 +120,26 @@ public:
 		return DataOffset + DataSize;
 	}
 
-	friend bool operator==(const FTransactionSerializedTaggedData& LHS, const FTransactionSerializedTaggedData& RHS)
+	friend bool operator==(const FSerializedTaggedData& LHS, const FSerializedTaggedData& RHS)
 	{
 		return LHS.DataOffset == RHS.DataOffset
 			&& LHS.DataSize == RHS.DataSize;
 	}
 
-	friend bool operator!=(const FTransactionSerializedTaggedData& LHS, const FTransactionSerializedTaggedData& RHS)
+	friend bool operator!=(const FSerializedTaggedData& LHS, const FSerializedTaggedData& RHS)
 	{
 		return !(LHS == RHS);
 	}
 
 	/** Offset to the start of the tagged data within the serialized object */
-	UPROPERTY()
 	int64 DataOffset = INDEX_NONE;
 
 	/** Size (in bytes) of the tagged data within the serialized object */
-	UPROPERTY()
 	int64 DataSize = 0;
 };
 
-USTRUCT()
-struct ENGINE_API FTransactionSerializedObjectData
+struct ENGINE_API FSerializedObjectData
 {
-	GENERATED_BODY()
-
 public:
 	bool Serialize(FArchive& Ar)
 	{
@@ -167,18 +147,18 @@ public:
 		return true;
 	}
 
-	friend FArchive& operator<<(FArchive& Ar, FTransactionSerializedObjectData& SerializedData)
+	friend FArchive& operator<<(FArchive& Ar, FSerializedObjectData& SerializedData)
 	{
 		Ar << SerializedData.Data;
 		return Ar;
 	}
 
-	friend bool operator==(const FTransactionSerializedObjectData& LHS, const FTransactionSerializedObjectData& RHS)
+	friend bool operator==(const FSerializedObjectData& LHS, const FSerializedObjectData& RHS)
 	{
 		return LHS.Data == RHS.Data;
 	}
 
-	friend bool operator!=(const FTransactionSerializedObjectData& LHS, const FTransactionSerializedObjectData& RHS)
+	friend bool operator!=(const FSerializedObjectData& LHS, const FSerializedObjectData& RHS)
 	{
 		return !(LHS == RHS);
 	}
@@ -205,31 +185,8 @@ private:
 	TArray64<uint8> Data;
 };
 
-template<>
-struct TStructOpsTypeTraits<FTransactionSerializedObjectData> : public TStructOpsTypeTraitsBase2<FTransactionSerializedObjectData>
+struct FSerializedObjectInfo
 {
-	enum
-	{
-		WithIdenticalViaEquality = true,
-		WithSerializer = true,
-	};
-};
-
-USTRUCT()
-struct FTransactionSerializedIndices
-{
-	GENERATED_BODY()
-
-public:
-	UPROPERTY()
-	TArray<int32> Indices;
-};
-
-USTRUCT()
-struct FTransactionSerializedObjectInfo
-{
-	GENERATED_BODY()
-
 public:
 	void SetObject(const UObject* InObject)
 	{
@@ -253,7 +210,7 @@ public:
 		bIsPendingKill = false;
 	}
 
-	void Swap(FTransactionSerializedObjectInfo& Other)
+	void Swap(FSerializedObjectInfo& Other)
 	{
 		Exchange(ObjectPackageName, Other.ObjectPackageName);
 		Exchange(ObjectName, Other.ObjectName);
@@ -265,39 +222,29 @@ public:
 	}
 
 	/** The package name of the object when it was serialized, can be dictated either by outer chain or external package */
-	UPROPERTY()
 	FName ObjectPackageName;
 
 	/** The name of the object when it was serialized */
-	UPROPERTY()
 	FName ObjectName;
 
 	/** The path name of the object when it was serialized */
-	UPROPERTY()
 	FName ObjectPathName;
 
 	/** The outer path name of the object when it was serialized */
-	UPROPERTY()
 	FName ObjectOuterPathName;
 
 	/** The external package name of the object when it was serialized, if any */
-	UPROPERTY()
 	FName ObjectExternalPackageName;
 
 	/** The path name of the object's class. */
-	UPROPERTY()
 	FName ObjectClassPathName;
 
 	/** The pending kill state of the object when it was serialized */
-	UPROPERTY()
 	bool bIsPendingKill = false;
 };
 
-USTRUCT()
-struct FTransactionSerializedObject
+struct FSerializedObject
 {
-	GENERATED_BODY()
-
 public:
 	void Reset()
 	{
@@ -306,7 +253,7 @@ public:
 		ReferencedNames.Reset();
 	}
 
-	void Swap(FTransactionSerializedObject& Other)
+	void Swap(FSerializedObject& Other)
 	{
 		Exchange(SerializedData, Other.SerializedData);
 		Exchange(ReferencedObjects, Other.ReferencedObjects);
@@ -314,23 +261,17 @@ public:
 	}
 
 	/** The serialized data for the transacted object */
-	UPROPERTY()
-	FTransactionSerializedObjectData SerializedData;
+	FSerializedObjectData SerializedData;
 
 	/** External objects referenced by the transacted object */
-	UPROPERTY()
-	TArray<FTransactionPersistentObjectRef> ReferencedObjects;
+	TArray<FPersistentObjectRef> ReferencedObjects;
 	
 	/** Names referenced by the transacted object */
-	UPROPERTY()
 	TArray<FName> ReferencedNames;
 };
 
-USTRUCT()
-struct FTransactionDiffableObject
+struct FDiffableObject
 {
-	GENERATED_BODY()
-
 public:
 	void SetObject(const UObject* InObject)
 	{
@@ -344,7 +285,7 @@ public:
 		SerializedTaggedData.Reset();
 	}
 
-	void Swap(FTransactionDiffableObject& Other)
+	void Swap(FDiffableObject& Other)
 	{
 		ObjectInfo.Swap(Other.ObjectInfo);
 		Exchange(SerializedData, Other.SerializedData);
@@ -352,21 +293,15 @@ public:
 	}
 
 	/** Information about the object when it was serialized */
-	UPROPERTY()
-	FTransactionSerializedObjectInfo ObjectInfo;
+	FSerializedObjectInfo ObjectInfo;
 
 	/** The serialized data for the diffable object */
-	UPROPERTY()
-	FTransactionSerializedObjectData SerializedData;
+	FSerializedObjectData SerializedData;
 
 	/** Information about tagged data (mainly properties) that were serialized within this object */
-	UPROPERTY()
-	TMap<FName, FTransactionSerializedTaggedData> SerializedTaggedData;
+	TMap<FName, FSerializedTaggedData> SerializedTaggedData;
 };
 
-namespace UE::Transaction
-{
-
 ENGINE_API extern const FName TaggedDataKey_UnknownData;
 ENGINE_API extern const FName TaggedDataKey_ScriptData;
 
@@ -374,7 +309,7 @@ ENGINE_API extern const FName TaggedDataKey_ScriptData;
 class ENGINE_API FSerializedObjectDataReader : public FArchiveUObject
 {
 public:
-	FSerializedObjectDataReader(const FTransactionSerializedObject& InSerializedObject);
+	FSerializedObjectDataReader(const FSerializedObject& InSerializedObject);
 
 	virtual int64 Tell() override { return Offset; }
 	virtual void Seek(int64 InPos) override { Offset = InPos; }
@@ -387,7 +322,7 @@ protected:
 	virtual FArchive& operator<<(class FName& N) override;
 	virtual FArchive& operator<<(class UObject*& Res) override;
 
-	const FTransactionSerializedObject& SerializedObject;
+	const FSerializedObject& SerializedObject;
 	int64 Offset = 0;
 };
 
@@ -398,7 +333,7 @@ namespace Internal
 class ENGINE_API FSerializedObjectDataWriterCommon : public FArchiveUObject
 {
 public:
-	FSerializedObjectDataWriterCommon(FTransactionSerializedObjectData& InSerializedData);
+	FSerializedObjectDataWriterCommon(FSerializedObjectData& InSerializedData);
 
 	virtual int64 Tell() override { return Offset; }
 	virtual void Seek(int64 InPos) override { checkSlow(Offset <= SerializedData.Num()); Offset = InPos; }
@@ -409,7 +344,7 @@ protected:
 
 	virtual void OnDataSerialized(int64 InOffset, int64 InNum) {}
 
-	FTransactionSerializedObjectData& SerializedData;
+	FSerializedObjectData& SerializedData;
 	int64 Offset = 0;
 };
 
@@ -419,14 +354,14 @@ protected:
 class ENGINE_API FSerializedObjectDataWriter : public Internal::FSerializedObjectDataWriterCommon
 {
 public:
-	FSerializedObjectDataWriter(FTransactionSerializedObject& InSerializedObject);
+	FSerializedObjectDataWriter(FSerializedObject& InSerializedObject);
 
 protected:
 	using FArchiveUObject::operator<<;
 	virtual FArchive& operator<<(class FName& N) override;
 	virtual FArchive& operator<<(class UObject*& Res) override;
 
-	FTransactionSerializedObject& SerializedObject;
+	FSerializedObject& SerializedObject;
 	int64 Offset = 0;
 
 	TMap<UObject*, int32> ObjectMap;
@@ -437,7 +372,7 @@ protected:
 class ENGINE_API FDiffableObjectDataWriter : public Internal::FSerializedObjectDataWriterCommon
 {
 public:
-	FDiffableObjectDataWriter(FTransactionDiffableObject& InDiffableObject, TArrayView<const FProperty*> InPropertiesToSerialize = TArrayView<const FProperty*>());
+	FDiffableObjectDataWriter(FDiffableObject& InDiffableObject, TArrayView<const FProperty*> InPropertiesToSerialize = TArrayView<const FProperty*>());
 
 protected:
 	FName GetTaggedDataKey() const;
@@ -466,7 +401,7 @@ private:
 		uint32 LastUpdateCount = 0;
 	};
 
-	FTransactionDiffableObject& DiffableObject;
+	FDiffableObject& DiffableObject;
 	TArrayView<const FProperty*> PropertiesToSerialize;
 
 	bool bIsPerformingScriptSerialization = false;
@@ -481,7 +416,7 @@ private:
 namespace DiffUtil
 {
 
-ENGINE_API void GenerateObjectDiff(const FTransactionDiffableObject& OldDiffableObject, const FTransactionDiffableObject& NewDiffableObject, FTransactionObjectDeltaChange& OutDeltaChange, const bool bFullDiff = true);
+ENGINE_API void GenerateObjectDiff(const FDiffableObject& OldDiffableObject, const FDiffableObject& NewDiffableObject, FTransactionObjectDeltaChange& OutDeltaChange, const bool bFullDiff = true);
 
 } // namespace DiffUtil
 
