Disable pin unlink verification ensure when the owning node is being destroyed which stops it from ensuring about orphaned bad nodes
Move the node modify calls to only happen if the pin connection will actually be changed

diff --git a/Engine/Source/Runtime/Engine/Private/EdGraph/EdGraphPin.cpp b/Engine/Source/Runtime/Engine/Private/EdGraph/EdGraphPin.cpp
--- a/Engine/Source/Runtime/Engine/Private/EdGraph/EdGraphPin.cpp
+++ b/Engine/Source/Runtime/Engine/Private/EdGraph/EdGraphPin.cpp
@@ -508,12 +508,9 @@ UEdGraphPin* UEdGraphPin::CreatePin(UEdGraphNode* InOwningNode)
 
 void UEdGraphPin::MakeLinkTo(UEdGraphPin* ToPin)
 {
-	Modify();
-
 	if (ToPin)
 	{
 		check(!bWasTrashed);
-		ToPin->Modify();
 
 		// Make sure we don't already link to it
 		if (!LinkedTo.Contains(ToPin))
@@ -524,6 +521,10 @@ void UEdGraphPin::MakeLinkTo(UEdGraphPin* ToPin)
 			ensureMsgf(!ToPin->LinkedTo.Contains(this), TEXT("%s"), *GetLinkInfoString( LOCTEXT("MakeLinkTo", "MakeLinkTo").ToString(), LOCTEXT("IsLinked", "is linked with pin").ToString(), ToPin));
 			ensureMsgf(MyNode->GetOuter() == ToPin->GetOwningNode()->GetOuter(), TEXT("%s"), *GetLinkInfoString( LOCTEXT("MakeLinkTo", "MakeLinkTo").ToString(), LOCTEXT("OuterMismatch", "has a different outer than pin").ToString(), ToPin)); // Ensure both pins belong to the same graph
 
+			// Notify owning nodes about upcoming change
+			Modify();
+			ToPin->Modify();
+
 			// Add to both lists
 			LinkedTo.Add(ToPin);
 			ToPin->LinkedTo.Add(this);
@@ -537,18 +538,24 @@ void UEdGraphPin::MakeLinkTo(UEdGraphPin* ToPin)
 
 void UEdGraphPin::BreakLinkTo(UEdGraphPin* ToPin)
 {
-	Modify();
-
 	if (ToPin)
 	{
-		ToPin->Modify();
-
 		// If we do indeed link to the passed in pin...
 		if (LinkedTo.Contains(ToPin))
 		{
-			// Check that the other pin links to us
-			ensureAlwaysMsgf(ToPin->LinkedTo.Contains(this), TEXT("%s"), *GetLinkInfoString(LOCTEXT("BreakLinkTo", "BreakLinkTo").ToString(), LOCTEXT("NotLinked", "not reciprocally linked with pin").ToString(), ToPin));
-			ToPin->LinkedTo.Remove(this);
+			Modify();
+
+			if (ToPin->LinkedTo.Contains(this))
+			{
+				ToPin->Modify();
+				ToPin->LinkedTo.Remove(this);
+			}
+			else if (OwningNode && !OwningNode->HasAnyFlags(RF_BeginDestroyed))
+			{
+				// Ensure that the other pin links to us but ignore it if our parent is invalid or being destroyed
+				ensureAlwaysMsgf(ToPin->LinkedTo.Contains(this), TEXT("%s"), *GetLinkInfoString(LOCTEXT("BreakLinkTo", "BreakLinkTo").ToString(), LOCTEXT("NotLinked", "not reciprocally linked with pin").ToString(), ToPin));
+			}
+
 			LinkedTo.Remove(ToPin);
 		}
 		else
