Moved some utils to TransactionCommon

#rnx

diff --git a/Engine/Source/Runtime/Engine/Private/Tests/TransactionDiffingTests.cpp b/Engine/Source/Runtime/Engine/Private/Tests/TransactionDiffingTests.cpp
--- a/Engine/Source/Runtime/Engine/Private/Tests/TransactionDiffingTests.cpp
+++ b/Engine/Source/Runtime/Engine/Private/Tests/TransactionDiffingTests.cpp
@@ -16,39 +16,22 @@ void UTransactionDiffingTestObject::Serialize(FStructuredArchive::FRecord Record
 namespace TransactionDiffingTests
 {
 
-UE::Transaction::FDiffableObject GetDiffableObject(UObject* Object)
-{
-	UE::Transaction::FDiffableObject DiffableObject;
-
-	UE::Transaction::FDiffableObjectDataWriter DiffWriter(DiffableObject);
-	Object->Serialize(DiffWriter);
-
-	return DiffableObject;
-}
-
-FTransactionObjectDeltaChange GenerateObjectDiff(const UE::Transaction::FDiffableObject& InitialObject, const UE::Transaction::FDiffableObject& ModifiedObject)
-{
-	FTransactionObjectDeltaChange DeltaChange;
-	UE::Transaction::DiffUtil::GenerateObjectDiff(InitialObject, ModifiedObject, DeltaChange);
-	return DeltaChange;
-}
-
 constexpr const uint32 TestFlags = EAutomationTestFlags::ApplicationContextMask | EAutomationTestFlags::EngineFilter;
 
 IMPLEMENT_SIMPLE_AUTOMATION_TEST(FEditPropertyDataTest, "System.Engine.Transactions.EditPropertyData", TestFlags)
 bool FEditPropertyDataTest::RunTest(const FString& Parameters)
 {
-	UTransactionDiffingTestObject* DefaultObject = GetMutableDefault<UTransactionDiffingTestObject>();
+	const UTransactionDiffingTestObject* DefaultObject = GetDefault<UTransactionDiffingTestObject>();
 	UTransactionDiffingTestObject* ModifiedObject = NewObject<UTransactionDiffingTestObject>();
 	
-	const UE::Transaction::FDiffableObject DefaultDiffableObject = GetDiffableObject(DefaultObject);
+	const UE::Transaction::FDiffableObject DefaultDiffableObject = UE::Transaction::DiffUtil::GetDiffableObject(DefaultObject);
 
 	ModifiedObject->PropertyData = 10;
 
-	const UE::Transaction::FDiffableObject ModifiedDiffableObject = GetDiffableObject(ModifiedObject);
+	const UE::Transaction::FDiffableObject ModifiedDiffableObject = UE::Transaction::DiffUtil::GetDiffableObject(ModifiedObject);
 
 	{
-		const FTransactionObjectDeltaChange DeltaChange = GenerateObjectDiff(DefaultDiffableObject, ModifiedDiffableObject);
+		const FTransactionObjectDeltaChange DeltaChange = UE::Transaction::DiffUtil::GenerateObjectDiff(DefaultDiffableObject, ModifiedDiffableObject);
 
 		TestFalse(TEXT("bHasNonPropertyChanges"), DeltaChange.bHasNonPropertyChanges);
 		TestEqual(TEXT("ChangedProperties.Num()"), DeltaChange.ChangedProperties.Num(), 1);
@@ -61,17 +44,17 @@ bool FEditPropertyDataTest::RunTest(const FString& Parameters)
 IMPLEMENT_SIMPLE_AUTOMATION_TEST(FEditNonPropertyDataTest, "System.Engine.Transactions.EditNonPropertyData", TestFlags)
 bool FEditNonPropertyDataTest::RunTest(const FString& Parameters)
 {
-	UTransactionDiffingTestObject* DefaultObject = GetMutableDefault<UTransactionDiffingTestObject>();
+	const UTransactionDiffingTestObject* DefaultObject = GetDefault<UTransactionDiffingTestObject>();
 	UTransactionDiffingTestObject* ModifiedObject = NewObject<UTransactionDiffingTestObject>();
 
-	const UE::Transaction::FDiffableObject DefaultDiffableObject = GetDiffableObject(DefaultObject);
+	const UE::Transaction::FDiffableObject DefaultDiffableObject = UE::Transaction::DiffUtil::GetDiffableObject(DefaultObject);
 
 	ModifiedObject->NonPropertyData = 10;
 
-	const UE::Transaction::FDiffableObject ModifiedDiffableObject = GetDiffableObject(ModifiedObject);
+	const UE::Transaction::FDiffableObject ModifiedDiffableObject = UE::Transaction::DiffUtil::GetDiffableObject(ModifiedObject);
 
 	{
-		const FTransactionObjectDeltaChange DeltaChange = GenerateObjectDiff(DefaultDiffableObject, ModifiedDiffableObject);
+		const FTransactionObjectDeltaChange DeltaChange = UE::Transaction::DiffUtil::GenerateObjectDiff(DefaultDiffableObject, ModifiedDiffableObject);
 
 		TestTrue(TEXT("bHasNonPropertyChanges"), DeltaChange.bHasNonPropertyChanges);
 		TestEqual(TEXT("ChangedProperties.Num()"), DeltaChange.ChangedProperties.Num(), 0);
@@ -83,17 +66,17 @@ bool FEditNonPropertyDataTest::RunTest(const FString& Parameters)
 IMPLEMENT_SIMPLE_AUTOMATION_TEST(FEditNamesTest, "System.Engine.Transactions.EditNames", TestFlags)
 bool FEditNamesTest::RunTest(const FString& Parameters)
 {
-	UTransactionDiffingTestObject* DefaultObject = GetMutableDefault<UTransactionDiffingTestObject>();
+	const UTransactionDiffingTestObject* DefaultObject = GetDefault<UTransactionDiffingTestObject>();
 	UTransactionDiffingTestObject* ModifiedObject = NewObject<UTransactionDiffingTestObject>();
 
-	const UE::Transaction::FDiffableObject DefaultDiffableObject = GetDiffableObject(DefaultObject);
+	const UE::Transaction::FDiffableObject DefaultDiffableObject = UE::Transaction::DiffUtil::GetDiffableObject(DefaultObject);
 
 	ModifiedObject->AdditionalName = "Test0";
 
-	const UE::Transaction::FDiffableObject ModifiedDiffableObject = GetDiffableObject(ModifiedObject);
+	const UE::Transaction::FDiffableObject ModifiedDiffableObject = UE::Transaction::DiffUtil::GetDiffableObject(ModifiedObject);
 
 	{
-		const FTransactionObjectDeltaChange DeltaChange = GenerateObjectDiff(DefaultDiffableObject, ModifiedDiffableObject);
+		const FTransactionObjectDeltaChange DeltaChange = UE::Transaction::DiffUtil::GenerateObjectDiff(DefaultDiffableObject, ModifiedDiffableObject);
 
 		TestFalse(TEXT("bHasNonPropertyChanges"), DeltaChange.bHasNonPropertyChanges);
 		TestEqual(TEXT("ChangedProperties.Num()"), DeltaChange.ChangedProperties.Num(), 1);
@@ -102,10 +85,10 @@ bool FEditNamesTest::RunTest(const FString& Parameters)
 
 	ModifiedObject->NamesArray.Add("Test1");
 
-	const UE::Transaction::FDiffableObject ModifiedDiffableObject2 = GetDiffableObject(ModifiedObject);
+	const UE::Transaction::FDiffableObject ModifiedDiffableObject2 = UE::Transaction::DiffUtil::GetDiffableObject(ModifiedObject);
 
 	{
-		const FTransactionObjectDeltaChange DeltaChange = GenerateObjectDiff(DefaultDiffableObject, ModifiedDiffableObject2);
+		const FTransactionObjectDeltaChange DeltaChange = UE::Transaction::DiffUtil::GenerateObjectDiff(DefaultDiffableObject, ModifiedDiffableObject2);
 
 		TestFalse(TEXT("bHasNonPropertyChanges"), DeltaChange.bHasNonPropertyChanges);
 		TestEqual(TEXT("ChangedProperties.Num()"), DeltaChange.ChangedProperties.Num(), 2);
@@ -114,7 +97,7 @@ bool FEditNamesTest::RunTest(const FString& Parameters)
 	}
 
 	{
-		const FTransactionObjectDeltaChange DeltaChange = GenerateObjectDiff(ModifiedDiffableObject, ModifiedDiffableObject2);
+		const FTransactionObjectDeltaChange DeltaChange = UE::Transaction::DiffUtil::GenerateObjectDiff(ModifiedDiffableObject, ModifiedDiffableObject2);
 
 		TestFalse(TEXT("bHasNonPropertyChanges"), DeltaChange.bHasNonPropertyChanges);
 		TestEqual(TEXT("ChangedProperties.Num()"), DeltaChange.ChangedProperties.Num(), 1);
@@ -127,17 +110,17 @@ bool FEditNamesTest::RunTest(const FString& Parameters)
 IMPLEMENT_SIMPLE_AUTOMATION_TEST(FEditObjectsTest, "System.Engine.Transactions.EditObjects", TestFlags)
 bool FEditObjectsTest::RunTest(const FString& Parameters)
 {
-	UTransactionDiffingTestObject* DefaultObject = GetMutableDefault<UTransactionDiffingTestObject>();
+	const UTransactionDiffingTestObject* DefaultObject = GetDefault<UTransactionDiffingTestObject>();
 	UTransactionDiffingTestObject* ModifiedObject = NewObject<UTransactionDiffingTestObject>();
 
-	const UE::Transaction::FDiffableObject DefaultDiffableObject = GetDiffableObject(DefaultObject);
+	const UE::Transaction::FDiffableObject DefaultDiffableObject = UE::Transaction::DiffUtil::GetDiffableObject(DefaultObject);
 
 	ModifiedObject->AdditionalObject = NewObject<UTransactionDiffingTestObject>();
 
-	const UE::Transaction::FDiffableObject ModifiedDiffableObject = GetDiffableObject(ModifiedObject);
+	const UE::Transaction::FDiffableObject ModifiedDiffableObject = UE::Transaction::DiffUtil::GetDiffableObject(ModifiedObject);
 
 	{
-		const FTransactionObjectDeltaChange DeltaChange = GenerateObjectDiff(DefaultDiffableObject, ModifiedDiffableObject);
+		const FTransactionObjectDeltaChange DeltaChange = UE::Transaction::DiffUtil::GenerateObjectDiff(DefaultDiffableObject, ModifiedDiffableObject);
 
 		TestFalse(TEXT("bHasNonPropertyChanges"), DeltaChange.bHasNonPropertyChanges);
 		TestEqual(TEXT("ChangedProperties.Num()"), DeltaChange.ChangedProperties.Num(), 1);
@@ -146,10 +129,10 @@ bool FEditObjectsTest::RunTest(const FString& Parameters)
 
 	ModifiedObject->ObjectsArray.Add(NewObject<UTransactionDiffingTestObject>());
 
-	const UE::Transaction::FDiffableObject ModifiedDiffableObject2 = GetDiffableObject(ModifiedObject);
+	const UE::Transaction::FDiffableObject ModifiedDiffableObject2 = UE::Transaction::DiffUtil::GetDiffableObject(ModifiedObject);
 
 	{
-		const FTransactionObjectDeltaChange DeltaChange = GenerateObjectDiff(DefaultDiffableObject, ModifiedDiffableObject2);
+		const FTransactionObjectDeltaChange DeltaChange = UE::Transaction::DiffUtil::GenerateObjectDiff(DefaultDiffableObject, ModifiedDiffableObject2);
 
 		TestFalse(TEXT("bHasNonPropertyChanges"), DeltaChange.bHasNonPropertyChanges);
 		TestEqual(TEXT("ChangedProperties.Num()"), DeltaChange.ChangedProperties.Num(), 2);
@@ -158,7 +141,7 @@ bool FEditObjectsTest::RunTest(const FString& Parameters)
 	}
 
 	{
-		const FTransactionObjectDeltaChange DeltaChange = GenerateObjectDiff(ModifiedDiffableObject, ModifiedDiffableObject2);
+		const FTransactionObjectDeltaChange DeltaChange = UE::Transaction::DiffUtil::GenerateObjectDiff(ModifiedDiffableObject, ModifiedDiffableObject2);
 
 		TestFalse(TEXT("bHasNonPropertyChanges"), DeltaChange.bHasNonPropertyChanges);
 		TestEqual(TEXT("ChangedProperties.Num()"), DeltaChange.ChangedProperties.Num(), 1);
@@ -171,17 +154,17 @@ bool FEditObjectsTest::RunTest(const FString& Parameters)
 IMPLEMENT_SIMPLE_AUTOMATION_TEST(FEditSoftObjectsTest, "System.Engine.Transactions.EditSoftObjects", TestFlags)
 bool FEditSoftObjectsTest::RunTest(const FString& Parameters)
 {
-	UTransactionDiffingTestObject* DefaultObject = GetMutableDefault<UTransactionDiffingTestObject>();
+	const UTransactionDiffingTestObject* DefaultObject = GetDefault<UTransactionDiffingTestObject>();
 	UTransactionDiffingTestObject* ModifiedObject = NewObject<UTransactionDiffingTestObject>();
 
-	const UE::Transaction::FDiffableObject DefaultDiffableObject = GetDiffableObject(DefaultObject);
+	const UE::Transaction::FDiffableObject DefaultDiffableObject = UE::Transaction::DiffUtil::GetDiffableObject(DefaultObject);
 
 	ModifiedObject->AdditionalSoftObject = NewObject<UTransactionDiffingTestObject>();
 
-	const UE::Transaction::FDiffableObject ModifiedDiffableObject = GetDiffableObject(ModifiedObject);
+	const UE::Transaction::FDiffableObject ModifiedDiffableObject = UE::Transaction::DiffUtil::GetDiffableObject(ModifiedObject);
 
 	{
-		const FTransactionObjectDeltaChange DeltaChange = GenerateObjectDiff(DefaultDiffableObject, ModifiedDiffableObject);
+		const FTransactionObjectDeltaChange DeltaChange = UE::Transaction::DiffUtil::GenerateObjectDiff(DefaultDiffableObject, ModifiedDiffableObject);
 
 		TestFalse(TEXT("bHasNonPropertyChanges"), DeltaChange.bHasNonPropertyChanges);
 		TestEqual(TEXT("ChangedProperties.Num()"), DeltaChange.ChangedProperties.Num(), 1);
@@ -190,10 +173,10 @@ bool FEditSoftObjectsTest::RunTest(const FString& Parameters)
 
 	ModifiedObject->SoftObjectsArray.Add(NewObject<UTransactionDiffingTestObject>());
 
-	const UE::Transaction::FDiffableObject ModifiedDiffableObject2 = GetDiffableObject(ModifiedObject);
+	const UE::Transaction::FDiffableObject ModifiedDiffableObject2 = UE::Transaction::DiffUtil::GetDiffableObject(ModifiedObject);
 
 	{
-		const FTransactionObjectDeltaChange DeltaChange = GenerateObjectDiff(DefaultDiffableObject, ModifiedDiffableObject2);
+		const FTransactionObjectDeltaChange DeltaChange = UE::Transaction::DiffUtil::GenerateObjectDiff(DefaultDiffableObject, ModifiedDiffableObject2);
 
 		TestFalse(TEXT("bHasNonPropertyChanges"), DeltaChange.bHasNonPropertyChanges);
 		TestEqual(TEXT("ChangedProperties.Num()"), DeltaChange.ChangedProperties.Num(), 2);
@@ -202,7 +185,7 @@ bool FEditSoftObjectsTest::RunTest(const FString& Parameters)
 	}
 
 	{
-		const FTransactionObjectDeltaChange DeltaChange = GenerateObjectDiff(ModifiedDiffableObject, ModifiedDiffableObject2);
+		const FTransactionObjectDeltaChange DeltaChange = UE::Transaction::DiffUtil::GenerateObjectDiff(ModifiedDiffableObject, ModifiedDiffableObject2);
 
 		TestFalse(TEXT("bHasNonPropertyChanges"), DeltaChange.bHasNonPropertyChanges);
 		TestEqual(TEXT("ChangedProperties.Num()"), DeltaChange.ChangedProperties.Num(), 1);
diff --git a/Engine/Source/Runtime/Engine/Private/TransactionCommon.cpp b/Engine/Source/Runtime/Engine/Private/TransactionCommon.cpp
--- a/Engine/Source/Runtime/Engine/Private/TransactionCommon.cpp
+++ b/Engine/Source/Runtime/Engine/Private/TransactionCommon.cpp
@@ -440,6 +440,35 @@ FName FDiffableObjectDataWriter::FCachedPropertyKey::SyncCache(const FArchiveSer
 namespace DiffUtil
 {
 
+FDiffableObject GetDiffableObject(const UObject* Object, const EGetDiffableObjectMode Mode)
+{
+	FDiffableObject DiffableObject;
+	DiffableObject.SetObject(Object);
+
+	FDiffableObjectDataWriter DiffWriter(DiffableObject);
+	switch (Mode)
+	{
+	case EGetDiffableObjectMode::SerializeObject:
+		const_cast<UObject*>(Object)->Serialize(DiffWriter);
+		break;
+	case EGetDiffableObjectMode::SerializeProperties:
+		const_cast<UObject*>(Object)->SerializeScriptProperties(DiffWriter);
+		break;
+	default:
+		checkf(false, TEXT("Unknown EGetDiffableObjectMode!"));
+		break;
+	}
+
+	return DiffableObject;
+}
+
+FTransactionObjectDeltaChange GenerateObjectDiff(const FDiffableObject& OldDiffableObject, const FDiffableObject& NewDiffableObject, const bool bFullDiff)
+{
+	FTransactionObjectDeltaChange DeltaChange;
+	GenerateObjectDiff(OldDiffableObject, NewDiffableObject, DeltaChange, bFullDiff);
+	return DeltaChange;
+}
+
 void GenerateObjectDiff(const FDiffableObject& OldDiffableObject, const FDiffableObject& NewDiffableObject, FTransactionObjectDeltaChange& OutDeltaChange, const bool bFullDiff)
 {
 	auto IsTaggedDataBlockIdentical = [&OldDiffableObject, &NewDiffableObject, &OutDeltaChange](const FName TaggedDataKey, const FSerializedTaggedData& OldSerializedTaggedData, const FSerializedTaggedData& NewSerializedTaggedData) -> bool
diff --git a/Engine/Source/Runtime/Engine/Public/TransactionCommon.h b/Engine/Source/Runtime/Engine/Public/TransactionCommon.h
--- a/Engine/Source/Runtime/Engine/Public/TransactionCommon.h
+++ b/Engine/Source/Runtime/Engine/Public/TransactionCommon.h
@@ -416,6 +416,14 @@ private:
 namespace DiffUtil
 {
 
+enum EGetDiffableObjectMode : uint8
+{
+	SerializeObject,
+	SerializeProperties,
+};
+ENGINE_API FDiffableObject GetDiffableObject(const UObject* Object, const EGetDiffableObjectMode Mode = EGetDiffableObjectMode::SerializeObject);
+
+ENGINE_API FTransactionObjectDeltaChange GenerateObjectDiff(const FDiffableObject& OldDiffableObject, const FDiffableObject& NewDiffableObject, const bool bFullDiff = true);
 ENGINE_API void GenerateObjectDiff(const FDiffableObject& OldDiffableObject, const FDiffableObject& NewDiffableObject, FTransactionObjectDeltaChange& OutDeltaChange, const bool bFullDiff = true);
 
 } // namespace DiffUtil
