Fixed bIsPendingKill being set incorrectly

#rnx

diff --git a/Engine/Source/Runtime/Engine/Public/TransactionCommon.h b/Engine/Source/Runtime/Engine/Public/TransactionCommon.h
--- a/Engine/Source/Runtime/Engine/Public/TransactionCommon.h
+++ b/Engine/Source/Runtime/Engine/Public/TransactionCommon.h
@@ -196,7 +196,7 @@ public:
 		ObjectOuterPathName = InObject->GetOuter() ? FName(*InObject->GetOuter()->GetPathName()) : FName();
 		ObjectExternalPackageName = InObject->GetExternalPackage() ? InObject->GetExternalPackage()->GetFName() : FName();
 		ObjectClassPathName = FName(*InObject->GetClass()->GetPathName());
-		bIsPendingKill = IsValid(InObject);
+		bIsPendingKill = !IsValid(InObject);
 	}
 
 	void Reset()
