Fix for loading cooked SoundWaves in editor

diff --git a/Engine/Source/Runtime/Engine/Classes/Sound/SoundWave.h b/Engine/Source/Runtime/Engine/Classes/Sound/SoundWave.h
--- a/Engine/Source/Runtime/Engine/Classes/Sound/SoundWave.h
+++ b/Engine/Source/Runtime/Engine/Classes/Sound/SoundWave.h
@@ -1044,6 +1044,8 @@ public:
 
 	/** Returns if an async task for a certain platform has finished. */
 	bool IsCompressedDataReady(FName Format, const FPlatformAudioCookOverrides* CompressionOverrides) const;
+
+	bool IsLoadedFromCookedData() const;
 #endif //WITH_EDITOR
 
 	virtual void BeginGetCompressedData(FName Format, const FPlatformAudioCookOverrides* CompressionOverrides);
@@ -1230,6 +1232,9 @@ public:
 		, bIsSeekable(0)
 		, bShouldUseStreamCaching(0)
 		, bLoadingBehaviorOverridden(0)
+#if WITH_EDITOR
+		, bLoadedFromCookedData(0)
+#endif //WITH_EDITOR
 	{
 	}
 
@@ -1347,6 +1352,10 @@ private:
 	uint8 bShouldUseStreamCaching : 1;
 	uint8 bLoadingBehaviorOverridden : 1;
 
+#if WITH_EDITOR
+	uint8 bLoadedFromCookedData : 1;
+#endif //WITH_EDITOR
+
 	friend class USoundWave;
 }; // class FSoundWaveData
 
diff --git a/Engine/Source/Runtime/Engine/Private/SoundWave.cpp b/Engine/Source/Runtime/Engine/Private/SoundWave.cpp
--- a/Engine/Source/Runtime/Engine/Private/SoundWave.cpp
+++ b/Engine/Source/Runtime/Engine/Private/SoundWave.cpp
@@ -179,6 +179,10 @@ void FSoundWaveData::InitializeDataFromSoundWave(USoundWave& InWave)
 	bIsStreaming = InWave.IsStreaming(nullptr);
 	SoundAssetCompressionType = InWave.GetSoundAssetCompressionType();
 	bShouldUseStreamCaching = InWave.ShouldUseStreamCaching();
+
+#if WITH_EDITOR
+	bLoadedFromCookedData = InWave.IsLoadedFromCookedData();
+#endif //WITH_EDITOR
 }
 
 FName FSoundWaveData::FindRuntimeFormat(const USoundWave& InWave) const
@@ -429,48 +433,51 @@ bool FSoundWaveData::LoadZerothChunk()
 		if (ZerothChunkData.GetView().Num() == 0)
 		{
 #if WITH_EDITOR
-			// In editor, we wait until the zeroth chunk is required to force finish compilation
-			// and retrieve it from DDC.  The rational is that async compilation will either prefetch
-			// chunks or rebuild them if unavailable so that this DDC request will be fulfilled from local cache.
-			uint8* TempChunkBuffer = nullptr;
-			int32 ChunkSizeInBytes = 0;
-
-			// GetNumChunks will force finish compilation if an async task still exists.
-			// Do not query chunk from DDC if no chunk exists as it would ensure inside GetChunkFromDDC.
-			if (RunningPlatformData.GetNumChunks() > 0)
+			if (!bLoadedFromCookedData)
 			{
-				ChunkSizeInBytes = GetChunkFromDDC(0, &TempChunkBuffer, true);
-			}
+				// In editor, we wait until the zeroth chunk is required to force finish compilation
+				// and retrieve it from DDC.  The rational is that async compilation will either prefetch
+				// chunks or rebuild them if unavailable so that this DDC request will be fulfilled from local cache.
+				uint8* TempChunkBuffer = nullptr;
+				int32 ChunkSizeInBytes = 0;
+
+				// GetNumChunks will force finish compilation if an async task still exists.
+				// Do not query chunk from DDC if no chunk exists as it would ensure inside GetChunkFromDDC.
+				if (RunningPlatformData.GetNumChunks() > 0)
+				{
+					ChunkSizeInBytes = GetChunkFromDDC(0, &TempChunkBuffer, true);
+				}
 
-			// Since we block for the DDC in the previous call we should always have the chunk loaded.
-			if (ChunkSizeInBytes != 0)
-			{
-				ZerothChunkData.Reset(TempChunkBuffer, ChunkSizeInBytes);
+				// Since we block for the DDC in the previous call we should always have the chunk loaded.
+				if (ChunkSizeInBytes != 0)
+				{
+					ZerothChunkData.Reset(TempChunkBuffer, ChunkSizeInBytes);
+				}
+				else
+				{
+					ZerothChunkData.Empty();
+					UE_LOG(LogAudio, Warning, TEXT("FSoundWaveData::LoadZerothChunk: Unsuccessful load of zeroth chunk from DDC. Asset requires manual re-cook.\n\tAsset: '%s'\n\tDerivedDataKey: '%s'"),
+						*GetFName().ToString(),
+						*RunningPlatformData.DerivedDataKey);
+					return false;
+				}
 			}
 			else
+#endif // WITH_EDITOR
 			{
-				ZerothChunkData.Empty();
-				UE_LOG(LogAudio, Warning, TEXT("FSoundWaveData::LoadZerothChunk: Unsuccessful load of zeroth chunk from DDC. Asset requires manual re-cook.\n\tAsset: '%s'\n\tDerivedDataKey: '%s'"),
-					*GetFName().ToString(),
-					*RunningPlatformData.DerivedDataKey);
-				return false;
+				// The zeroth chunk is cooked out to RunningPlatformData, so retrieve it.
+				check(GetNumChunks() > 0);
+				FStreamedAudioChunk& ZerothChunk = GetChunk(0);
+
+				// Sanity check to ensure bulk size is set up
+				UE_CLOG(ZerothChunk.BulkData.GetBulkDataSize() != ZerothChunk.DataSize, LogAudio, Warning
+					, TEXT("SoundWave '%s' bulk data serialized out had a mismatched size with the DataSize field."
+						"\nBulk Data Reported Size: %d"
+						"\nBulk Data Actual Size: %ld")
+					, *GetFName().ToString(), ZerothChunk.DataSize, ZerothChunk.BulkData.GetBulkDataSize());
+
+				ZerothChunkData = ZerothChunk.BulkData.GetCopyAsBuffer(ZerothChunk.AudioDataSize, true);
 			}
-#else // WITH_EDITOR
-
-			// The zeroth chunk is cooked out to RunningPlatformData, so retrieve it.
-			check(GetNumChunks() > 0);
-			FStreamedAudioChunk& ZerothChunk = GetChunk(0);
-
-			// Sanity check to ensure bulk size is set up
-			UE_CLOG(ZerothChunk.BulkData.GetBulkDataSize() != ZerothChunk.DataSize, LogAudio, Warning
-				, TEXT("SoundWave '%s' bulk data serialized out had a mismatched size with the DataSize field."
-					"\nBulk Data Reported Size: %d"
-					"\nBulk Data Actual Size: %ld")
-				, *GetFName().ToString(), ZerothChunk.DataSize, ZerothChunk.BulkData.GetBulkDataSize());
-
-			ZerothChunkData = ZerothChunk.BulkData.GetCopyAsBuffer(ZerothChunk.AudioDataSize, true);
-
-#endif // WITH_EDITOR
 		}
 	}
 
@@ -1463,7 +1470,7 @@ void USoundWave::BeginGetCompressedData(FName Format, const FPlatformAudioCookOv
 	check(SoundWaveDataPtr);
 
 #if WITH_EDITOR
-	if (IsTemplate() || IsRunningDedicatedServer())
+	if (IsTemplate() || IsRunningDedicatedServer() || bLoadedFromCookedData)
 	{
 		return;
 	}
@@ -1510,6 +1517,11 @@ bool USoundWave::IsCompressedDataReady(FName Format, const FPlatformAudioCookOve
 	return true;
 }
 
+bool USoundWave::IsLoadedFromCookedData() const
+{
+	return bLoadedFromCookedData;
+}
+
 #endif
 
 FByteBulkData* USoundWave::GetCompressedData(FName Format, const FPlatformAudioCookOverrides* CompressionOverrides)
