Splitting a.Runtime.ISPC and the matching preprocessor define into individual controls per algorithm as some of them show gains while others regress performance.

diff --git a/Engine/Source/Runtime/Engine/Private/Animation/AnimationRuntime.cpp b/Engine/Source/Runtime/Engine/Private/Animation/AnimationRuntime.cpp
--- a/Engine/Source/Runtime/Engine/Private/Animation/AnimationRuntime.cpp
+++ b/Engine/Source/Runtime/Engine/Private/Animation/AnimationRuntime.cpp
@@ -15,9 +15,6 @@
 #include "Animation/AnimationPoseData.h"
 #include "Animation/BlendProfile.h"
 #include "Animation/MirrorDataTable.h"
-#if INTEL_ISPC
-#include "AnimationRuntime.ispc.generated.h"
-#endif
 
 DEFINE_LOG_CATEGORY(LogAnimation);
 DEFINE_LOG_CATEGORY(LogRootMotion);
@@ -30,24 +27,76 @@ DECLARE_CYCLE_STAT(TEXT("BlendPosesPerBoneFilter"), STAT_BlendPosesPerBoneFilter
 //////////////////////////////////////////////////////////////////////////
 
 #if INTEL_ISPC
+
+#include "AnimationRuntime.ispc.generated.h"
+
 static_assert(sizeof(ispc::FTransform) == sizeof(FTransform), "sizeof(ispc::FTransform) != sizeof(FTransform)");
 static_assert(sizeof(ispc::FVector) == sizeof(FVector), "sizeof(ispc::FVector) != sizeof(FVector)");
 static_assert(sizeof(ispc::FVector4) == sizeof(FQuat), "sizeof(ispc::FVector4) != sizeof(FQuat)");
 static_assert(sizeof(ispc::FPerBoneBlendWeight) == sizeof(FPerBoneBlendWeight), "sizeof(ispc::FPerBoneBlendWeight) != sizeof(FPerBoneBlendWeight)");
+
+#if !defined(ANIM_BLEND_POSE_OVERWRITE_ISPC_ENABLED_DEFAULT)
+#define ANIM_BLEND_POSE_OVERWRITE_ISPC_ENABLED_DEFAULT 1
 #endif
 
-#if !defined(ANIM_RUNTIME_ISPC_ENABLED_DEFAULT)
-#define ANIM_RUNTIME_ISPC_ENABLED_DEFAULT 1
+#if !defined(ANIM_BLEND_POSE_ACCUMULATE_ISPC_ENABLED_DEFAULT)
+#define ANIM_BLEND_POSE_ACCUMULATE_ISPC_ENABLED_DEFAULT 1
 #endif
 
-// Support run-time toggling on supported platforms in non-shipping configurations
-#if !INTEL_ISPC || UE_BUILD_SHIPPING
-static constexpr bool bAnim_Runtime_ISPC_Enabled = INTEL_ISPC && ANIM_RUNTIME_ISPC_ENABLED_DEFAULT;
-#else
-static bool bAnim_Runtime_ISPC_Enabled = ANIM_RUNTIME_ISPC_ENABLED_DEFAULT;
-static FAutoConsoleVariableRef CVarAnimRuntimeISPCEnabled(TEXT("a.Runtime.ISPC"), bAnim_Runtime_ISPC_Enabled, TEXT("Whether to use ISPC optimizations in animation runtime"));
+#if !defined(ANIM_BLEND_CURVES_ISPC_ENABLED_DEFAULT)
+#define ANIM_BLEND_CURVES_ISPC_ENABLED_DEFAULT 1
+#endif
+
+#if !defined(ANIM_LERP_BONE_TRANSFORMS_ISPC_ENABLED_DEFAULT)
+#define ANIM_LERP_BONE_TRANSFORMS_ISPC_ENABLED_DEFAULT 1
+#endif
+
+#if !defined(ANIM_CONVERT_POSE_TO_MESH_ROTATION_ISPC_ENABLED_DEFAULT)
+#define ANIM_CONVERT_POSE_TO_MESH_ROTATION_ISPC_ENABLED_DEFAULT 1
+#endif
+
+#if !defined(ANIM_CONVERT_MESH_ROTATION_TO_LOCAL_SPACE_ISPC_ENABLED_DEFAULT)
+#define ANIM_CONVERT_MESH_ROTATION_TO_LOCAL_SPACE_ISPC_ENABLED_DEFAULT 1
+#endif
+
+#if !defined(ANIM_ACCUMULATE_LOCAL_SPACE_ADDITIVE_POSE_ISPC_ENABLED_DEFAULT)
+#define ANIM_ACCUMULATE_LOCAL_SPACE_ADDITIVE_POSE_ISPC_ENABLED_DEFAULT 1
 #endif
 
+#if !defined(ANIM_BLEND_POSES_PER_BONE_FILTER_ISPC_ENABLED_DEFAULT)
+#define ANIM_BLEND_POSES_PER_BONE_FILTER_ISPC_ENABLED_DEFAULT 1
+#endif
+
+#if UE_BUILD_SHIPPING
+static constexpr bool bAnim_BlendPoseOverwrite_ISPC_Enabled = ANIM_BLEND_POSE_OVERWRITE_ISPC_ENABLED_DEFAULT;
+static constexpr bool bAnim_BlendPoseAccumulate_ISPC_Enabled = ANIM_BLEND_POSE_ACCUMULATE_ISPC_ENABLED_DEFAULT;
+static constexpr bool bAnim_BlendCurves_ISPC_Enabled = ANIM_BLEND_CURVES_ISPC_ENABLED_DEFAULT;
+static constexpr bool bAnim_LerpBoneTransforms_ISPC_Enabled = ANIM_LERP_BONE_TRANSFORMS_ISPC_ENABLED_DEFAULT;
+static constexpr bool bAnim_ConvertPoseToMeshRotation_ISPC_Enabled = ANIM_CONVERT_POSE_TO_MESH_ROTATION_ISPC_ENABLED_DEFAULT;
+static constexpr bool bAnim_ConvertMeshRotationPoseToLocalSpace_ISPC_Enabled = ANIM_CONVERT_MESH_ROTATION_TO_LOCAL_SPACE_ISPC_ENABLED_DEFAULT;
+static constexpr bool bAnim_AccumulateLocalSpaceAdditivePose_ISPC_Enabled = ANIM_ACCUMULATE_LOCAL_SPACE_ADDITIVE_POSE_ISPC_ENABLED_DEFAULT;
+static constexpr bool bAnim_BlendPosesPerBoneFilter_ISPC_Enabled = ANIM_BLEND_POSES_PER_BONE_FILTER_ISPC_ENABLED_DEFAULT;
+#else
+static bool bAnim_BlendPoseOverwrite_ISPC_Enabled = INTEL_ISPC && ANIM_BLEND_POSE_OVERWRITE_ISPC_ENABLED_DEFAULT;
+static FAutoConsoleVariableRef CVarBlendPoseOverwriteISPCEnabled(TEXT("a.BlendPoseOverwrite.ISPC"), bAnim_BlendPoseOverwrite_ISPC_Enabled, TEXT("Whether to use ISPC optimizations for over-write pose blending"));
+static bool bAnim_BlendPoseAccumulate_ISPC_Enabled = INTEL_ISPC && ANIM_BLEND_POSE_ACCUMULATE_ISPC_ENABLED_DEFAULT;
+static FAutoConsoleVariableRef CVarBlendPoseAccumulateISPCEnabled(TEXT("a.BlendPoseAccumulate.ISPC"), bAnim_BlendPoseAccumulate_ISPC_Enabled, TEXT("Whether to use ISPC optimizations for accumulation pose blending"));
+static bool bAnim_BlendCurves_ISPC_Enabled = INTEL_ISPC && ANIM_BLEND_CURVES_ISPC_ENABLED_DEFAULT;
+static FAutoConsoleVariableRef CVarBlendCurvesISPCEnabled(TEXT("a.BlendCurves.ISPC"), bAnim_BlendCurves_ISPC_Enabled, TEXT("Whether to use ISPC optimizations for curve blending"));
+static bool bAnim_LerpBoneTransforms_ISPC_Enabled = INTEL_ISPC && ANIM_LERP_BONE_TRANSFORMS_ISPC_ENABLED_DEFAULT;
+static FAutoConsoleVariableRef CVarLerpBoneTransformsISPCEnabled(TEXT("a.LerpBoneTransforms.ISPC"), bAnim_LerpBoneTransforms_ISPC_Enabled, TEXT("Whether to use ISPC optimizations for interpolating bone transforms"));
+static bool bAnim_ConvertPoseToMeshRotation_ISPC_Enabled = INTEL_ISPC && ANIM_CONVERT_POSE_TO_MESH_ROTATION_ISPC_ENABLED_DEFAULT;
+static FAutoConsoleVariableRef CVarConvertPoseToMeshRotationISPCEnabled(TEXT("a.ConvertPoseToMeshRotation.ISPC"), bAnim_ConvertPoseToMeshRotation_ISPC_Enabled, TEXT("Whether to use ISPC optimizations for converting local space rotations to mesh space"));
+static bool bAnim_ConvertMeshRotationPoseToLocalSpace_ISPC_Enabled = INTEL_ISPC && ANIM_CONVERT_MESH_ROTATION_TO_LOCAL_SPACE_ISPC_ENABLED_DEFAULT;
+static FAutoConsoleVariableRef CVarConvertMeshRotationPoseToLocalSpace(TEXT("a.ConvertMeshRotationPoseToLocalSpace.ISPC"), bAnim_ConvertMeshRotationPoseToLocalSpace_ISPC_Enabled, TEXT("Whether to use ISPC optimizations for converting mesh space rotations to local space"));
+static bool bAnim_AccumulateLocalSpaceAdditivePose_ISPC_Enabled = INTEL_ISPC && ANIM_ACCUMULATE_LOCAL_SPACE_ADDITIVE_POSE_ISPC_ENABLED_DEFAULT;
+static FAutoConsoleVariableRef CVarAccumulateLocalSpaceAdditivePose(TEXT("a.AccumulateLocalSpaceAdditivePose.ISPC"), bAnim_AccumulateLocalSpaceAdditivePose_ISPC_Enabled, TEXT("Whether to use ISPC optimizations for accumulating local space additive pose"));
+static bool bAnim_BlendPosesPerBoneFilter_ISPC_Enabled = INTEL_ISPC && ANIM_BLEND_POSES_PER_BONE_FILTER_ISPC_ENABLED_DEFAULT;
+static FAutoConsoleVariableRef CVarBlendPosesPerBoneFilter(TEXT("a.BlendPosesPerBoneFilter.ISPC"), bAnim_BlendPosesPerBoneFilter_ISPC_Enabled, TEXT("Whether to use ISPC optimizations for blending poses with a per-bone filter"));
+#endif // UE_BUILD_SHIPPING
+
+#endif // INTEL_ISPC
+
 void FAnimationRuntime::NormalizeRotations(const FBoneContainer& RequiredBones, /*inout*/ FTransformArrayA2& Atoms)
 {
 	check( Atoms.Num() == RequiredBones.GetNumBones() );
@@ -100,17 +149,17 @@ FORCEINLINE void BlendPose(const FCompactPose& SourcePose, FCompactPose& ResultP
 template <>
 FORCEINLINE void BlendPose<ETransformBlendMode::Overwrite>(const FCompactPose& SourcePose, FCompactPose& ResultPose, const float BlendWeight)
 {
-	if (bAnim_Runtime_ISPC_Enabled)
-	{
 #if INTEL_ISPC
+	if (bAnim_BlendPoseOverwrite_ISPC_Enabled)
+	{
 		ispc::BlendTransformOverwrite(
 			(ispc::FTransform*)&SourcePose.GetBones()[0],
 			(ispc::FTransform*)&ResultPose.GetBones()[0],
 			BlendWeight,
 			SourcePose.GetNumBones());
-#endif
 	}
 	else
+#endif
 	{
 		for (FCompactPoseBoneIndex BoneIndex : SourcePose.ForEachBoneIndex())
 		{
@@ -122,17 +171,17 @@ FORCEINLINE void BlendPose<ETransformBlendMode::Overwrite>(const FCompactPose& S
 template <>
 FORCEINLINE void BlendPose<ETransformBlendMode::Accumulate>(const FCompactPose& SourcePose, FCompactPose& ResultPose, const float BlendWeight)
 {
-	if (bAnim_Runtime_ISPC_Enabled)
-	{
 #if INTEL_ISPC
+	if (bAnim_BlendPoseAccumulate_ISPC_Enabled)
+	{
 		ispc::BlendTransformAccumulate(
 			(ispc::FTransform*)&SourcePose.GetBones()[0],
 			(ispc::FTransform*)&ResultPose.GetBones()[0],
 			BlendWeight,
 			SourcePose.GetNumBones());
-#endif
 	}
 	else
+#endif
 	{
 		for (FCompactPoseBoneIndex BoneIndex : SourcePose.ForEachBoneIndex())
 		{
@@ -171,9 +220,9 @@ FORCEINLINE void BlendCurves(const TArrayView<const FBlendedCurve> SourceCurves,
 {
 	if (SourceCurves.Num() > 0)
 	{
-		if (bAnim_Runtime_ISPC_Enabled)
-		{
 #if INTEL_ISPC
+		if (bAnim_BlendCurves_ISPC_Enabled)
+		{
 			OutCurve.InitFrom(SourceCurves[0]);
 			for (int32 CurveIndex = 0; CurveIndex < SourceCurves.Num(); ++CurveIndex)
 			{
@@ -187,9 +236,9 @@ FORCEINLINE void BlendCurves(const TArrayView<const FBlendedCurve> SourceCurves,
 					SourceWeights[CurveIndex]
 				);
 			}
-#endif
 		}
 		else
+#endif
 		{
 			OutCurve.Override(SourceCurves[0], SourceWeights[0]);
 			for (int32 CurveIndex = 1; CurveIndex < SourceCurves.Num(); ++CurveIndex)
@@ -854,9 +903,9 @@ void FAnimationRuntime::LerpBoneTransforms(TArray<FTransform>& A, const TArray<F
 		return;
 	}
 
-	if (bAnim_Runtime_ISPC_Enabled)
-	{
 #if INTEL_ISPC
+	if (bAnim_LerpBoneTransforms_ISPC_Enabled)
+	{
 		ispc::LerpBoneTransforms(
 			(ispc::FTransform*)A.GetData(),
 			(ispc::FTransform*)B.GetData(),
@@ -864,9 +913,9 @@ void FAnimationRuntime::LerpBoneTransforms(TArray<FTransform>& A, const TArray<F
 			RequiredBonesArray.GetData(),
 			RequiredBonesArray.Num()
 		);
-#endif
 	}
 	else
+#endif
 	{
 		FTransform* ATransformData = A.GetData();
 		const FTransform* BTransformData = B.GetData();
@@ -964,18 +1013,16 @@ void FAnimationRuntime::ConvertPoseToMeshRotation(FCompactPose& LocalPose)
 {
 	SCOPE_CYCLE_COUNTER(STAT_ConvertPoseToMeshRot);
 
-	// Convert all rotations to mesh space
-	// only the root bone doesn't have a parent. So skip it to save a branch in the iteration.
-	if (bAnim_Runtime_ISPC_Enabled)
-	{
 #if INTEL_ISPC
+	if (bAnim_ConvertPoseToMeshRotation_ISPC_Enabled)
+	{
 		ispc::ConvertPoseToMeshRotation(
 			(ispc::FTransform*)&LocalPose.GetBones()[0],
 			(int32*)&LocalPose.GetBoneContainer().GetCompactPoseParentBoneArray().GetData()[0],
 			LocalPose.GetNumBones());
-#endif
 	}
 	else
+#endif
 	{
 		for (FCompactPoseBoneIndex BoneIndex(1); BoneIndex < LocalPose.GetNumBones(); ++BoneIndex)
 		{
@@ -991,18 +1038,16 @@ void FAnimationRuntime::ConvertMeshRotationPoseToLocalSpace(FCompactPose& Pose)
 {
 	SCOPE_CYCLE_COUNTER(STAT_ConvertMeshRotPoseToLocalSpace);
 
-	// Convert all rotations to mesh space
-	// only the root bone doesn't have a parent. So skip it to save a branch in the iteration.
-	if (bAnim_Runtime_ISPC_Enabled)
-	{
 #if INTEL_ISPC
+	if (bAnim_ConvertMeshRotationPoseToLocalSpace_ISPC_Enabled)
+	{
 		ispc::ConvertMeshRotationPoseToLocalSpace(
 			(ispc::FTransform*)&Pose.GetBones()[0],
 			(int32*)&Pose.GetBoneContainer().GetCompactPoseParentBoneArray().GetData()[0],
 			Pose.GetNumBones());
-#endif
 	}
 	else
+#endif
 	{
 		for (FCompactPoseBoneIndex BoneIndex(Pose.GetNumBones() - 1); BoneIndex > 0; --BoneIndex)
 		{
@@ -1079,17 +1124,17 @@ void FAnimationRuntime::AccumulateLocalSpaceAdditivePoseInternal(FCompactPose& B
 		const ScalarRegister VBlendWeight(Weight);
 		if (FAnimWeight::IsFullWeight(Weight))
 		{
-			if (bAnim_Runtime_ISPC_Enabled)
-			{
 #if INTEL_ISPC
+			if (bAnim_AccumulateLocalSpaceAdditivePose_ISPC_Enabled)
+			{
 				ispc::AccumulateWithAdditiveScale(
 					(ispc::FTransform*)&BasePose.GetBones()[0],
 					(ispc::FTransform*)&AdditivePose.GetBones()[0],
 					Weight,
 					BasePose.GetNumBones());
-#endif
 			}
 			else
+#endif
 			{
 				// fast path, no need to weight additive.
 				for (FCompactPoseBoneIndex BoneIndex : BasePose.ForEachBoneIndex())
@@ -1766,9 +1811,9 @@ FAnimationPoseData& OutAnimationPoseData, TArray<FPerBoneBlendWeight>& BoneBlend
 	// blend poses with both mesh space rotation and scaling (we assume uniform scale)
 	if (bMeshSpaceRotationBlend && bMeshSpaceScaleBlend)
 	{
-		if (bAnim_Runtime_ISPC_Enabled)
-		{
 #if INTEL_ISPC
+		if (bAnim_BlendPosesPerBoneFilter_ISPC_Enabled)
+		{
 			ispc::BlendPosesPerBoneFilterScaleRotation(
 				(ispc::FTransform*)OutPose.GetBones().GetData(),
 				(ispc::FTransform*)BasePose.GetBones().GetData(),
@@ -1783,9 +1828,9 @@ FAnimationPoseData& OutAnimationPoseData, TArray<FPerBoneBlendWeight>& BoneBlend
 				(ispc::FPerBoneBlendWeight*)BoneBlendWeights.GetData(),
 				(int32*)BoneContainer.GetCompactPoseParentBoneArray().GetData(),
 				BasePose.GetNumBones());
-#endif
 		}
 		else
+#endif
 		{
 			for (const FCompactPoseBoneIndex BoneIndex : BasePose.ForEachBoneIndex())
 			{
