Minor change I accidentally reverted since the last review revision; INTEL_ISPC is guaranteed true for these defines.

diff --git a/Engine/Source/Runtime/Engine/Private/Animation/AnimationRuntime.cpp b/Engine/Source/Runtime/Engine/Private/Animation/AnimationRuntime.cpp
--- a/Engine/Source/Runtime/Engine/Private/Animation/AnimationRuntime.cpp
+++ b/Engine/Source/Runtime/Engine/Private/Animation/AnimationRuntime.cpp
@@ -77,21 +77,21 @@ static constexpr bool bAnim_ConvertMeshRotationPoseToLocalSpace_ISPC_Enabled = A
 static constexpr bool bAnim_AccumulateLocalSpaceAdditivePose_ISPC_Enabled = ANIM_ACCUMULATE_LOCAL_SPACE_ADDITIVE_POSE_ISPC_ENABLED_DEFAULT;
 static constexpr bool bAnim_BlendPosesPerBoneFilter_ISPC_Enabled = ANIM_BLEND_POSES_PER_BONE_FILTER_ISPC_ENABLED_DEFAULT;
 #else
-static bool bAnim_BlendPoseOverwrite_ISPC_Enabled = INTEL_ISPC && ANIM_BLEND_POSE_OVERWRITE_ISPC_ENABLED_DEFAULT;
+static bool bAnim_BlendPoseOverwrite_ISPC_Enabled = ANIM_BLEND_POSE_OVERWRITE_ISPC_ENABLED_DEFAULT;
 static FAutoConsoleVariableRef CVarBlendPoseOverwriteISPCEnabled(TEXT("a.BlendPoseOverwrite.ISPC"), bAnim_BlendPoseOverwrite_ISPC_Enabled, TEXT("Whether to use ISPC optimizations for over-write pose blending"));
-static bool bAnim_BlendPoseAccumulate_ISPC_Enabled = INTEL_ISPC && ANIM_BLEND_POSE_ACCUMULATE_ISPC_ENABLED_DEFAULT;
+static bool bAnim_BlendPoseAccumulate_ISPC_Enabled = ANIM_BLEND_POSE_ACCUMULATE_ISPC_ENABLED_DEFAULT;
 static FAutoConsoleVariableRef CVarBlendPoseAccumulateISPCEnabled(TEXT("a.BlendPoseAccumulate.ISPC"), bAnim_BlendPoseAccumulate_ISPC_Enabled, TEXT("Whether to use ISPC optimizations for accumulation pose blending"));
-static bool bAnim_BlendCurves_ISPC_Enabled = INTEL_ISPC && ANIM_BLEND_CURVES_ISPC_ENABLED_DEFAULT;
+static bool bAnim_BlendCurves_ISPC_Enabled = ANIM_BLEND_CURVES_ISPC_ENABLED_DEFAULT;
 static FAutoConsoleVariableRef CVarBlendCurvesISPCEnabled(TEXT("a.BlendCurves.ISPC"), bAnim_BlendCurves_ISPC_Enabled, TEXT("Whether to use ISPC optimizations for curve blending"));
-static bool bAnim_LerpBoneTransforms_ISPC_Enabled = INTEL_ISPC && ANIM_LERP_BONE_TRANSFORMS_ISPC_ENABLED_DEFAULT;
+static bool bAnim_LerpBoneTransforms_ISPC_Enabled = ANIM_LERP_BONE_TRANSFORMS_ISPC_ENABLED_DEFAULT;
 static FAutoConsoleVariableRef CVarLerpBoneTransformsISPCEnabled(TEXT("a.LerpBoneTransforms.ISPC"), bAnim_LerpBoneTransforms_ISPC_Enabled, TEXT("Whether to use ISPC optimizations for interpolating bone transforms"));
-static bool bAnim_ConvertPoseToMeshRotation_ISPC_Enabled = INTEL_ISPC && ANIM_CONVERT_POSE_TO_MESH_ROTATION_ISPC_ENABLED_DEFAULT;
+static bool bAnim_ConvertPoseToMeshRotation_ISPC_Enabled = ANIM_CONVERT_POSE_TO_MESH_ROTATION_ISPC_ENABLED_DEFAULT;
 static FAutoConsoleVariableRef CVarConvertPoseToMeshRotationISPCEnabled(TEXT("a.ConvertPoseToMeshRotation.ISPC"), bAnim_ConvertPoseToMeshRotation_ISPC_Enabled, TEXT("Whether to use ISPC optimizations for converting local space rotations to mesh space"));
-static bool bAnim_ConvertMeshRotationPoseToLocalSpace_ISPC_Enabled = INTEL_ISPC && ANIM_CONVERT_MESH_ROTATION_TO_LOCAL_SPACE_ISPC_ENABLED_DEFAULT;
+static bool bAnim_ConvertMeshRotationPoseToLocalSpace_ISPC_Enabled = ANIM_CONVERT_MESH_ROTATION_TO_LOCAL_SPACE_ISPC_ENABLED_DEFAULT;
 static FAutoConsoleVariableRef CVarConvertMeshRotationPoseToLocalSpace(TEXT("a.ConvertMeshRotationPoseToLocalSpace.ISPC"), bAnim_ConvertMeshRotationPoseToLocalSpace_ISPC_Enabled, TEXT("Whether to use ISPC optimizations for converting mesh space rotations to local space"));
-static bool bAnim_AccumulateLocalSpaceAdditivePose_ISPC_Enabled = INTEL_ISPC && ANIM_ACCUMULATE_LOCAL_SPACE_ADDITIVE_POSE_ISPC_ENABLED_DEFAULT;
+static bool bAnim_AccumulateLocalSpaceAdditivePose_ISPC_Enabled = ANIM_ACCUMULATE_LOCAL_SPACE_ADDITIVE_POSE_ISPC_ENABLED_DEFAULT;
 static FAutoConsoleVariableRef CVarAccumulateLocalSpaceAdditivePose(TEXT("a.AccumulateLocalSpaceAdditivePose.ISPC"), bAnim_AccumulateLocalSpaceAdditivePose_ISPC_Enabled, TEXT("Whether to use ISPC optimizations for accumulating local space additive pose"));
-static bool bAnim_BlendPosesPerBoneFilter_ISPC_Enabled = INTEL_ISPC && ANIM_BLEND_POSES_PER_BONE_FILTER_ISPC_ENABLED_DEFAULT;
+static bool bAnim_BlendPosesPerBoneFilter_ISPC_Enabled = ANIM_BLEND_POSES_PER_BONE_FILTER_ISPC_ENABLED_DEFAULT;
 static FAutoConsoleVariableRef CVarBlendPosesPerBoneFilter(TEXT("a.BlendPosesPerBoneFilter.ISPC"), bAnim_BlendPosesPerBoneFilter_ISPC_Enabled, TEXT("Whether to use ISPC optimizations for blending poses with a per-bone filter"));
 #endif // UE_BUILD_SHIPPING
 
