Fixed race condition in async uniform expression cache updates.

 1) Added a variant of FillUniformBuffer which takes an array of allocated VT's directly, and a copy of the array is made. This avoids a race condition where InvalidateUniformExpressionCache could be called during the async update.
 2) Added a delegate to wait on the task when render commands are flushed, as this is when material resources can be destroyed.

diff --git a/Engine/Source/Runtime/Engine/Private/Materials/MaterialShared.cpp b/Engine/Source/Runtime/Engine/Private/Materials/MaterialShared.cpp
--- a/Engine/Source/Runtime/Engine/Private/Materials/MaterialShared.cpp
+++ b/Engine/Source/Runtime/Engine/Private/Materials/MaterialShared.cpp
@@ -86,6 +86,13 @@ FAutoConsoleVariableRef CVarDeferUniformExpressionCaching(
 	ECVF_RenderThreadSafe
 	);
 
+int32 GUniformExpressionCacheAsyncUpdates = 1;
+FAutoConsoleVariableRef CVarUniformExpressionCacheAsyncUpdates(
+	TEXT("r.UniformExpressionCacheAsyncUpdates"),
+	GUniformExpressionCacheAsyncUpdates,
+	TEXT("Whether to allow async updates of uniform expression caches."),
+	ECVF_RenderThreadSafe);
+
 bool Engine_IsStrataEnabled();
 
 struct FAllowCachingStaticParameterValues
@@ -3892,6 +3899,44 @@ const FMaterial& FMaterialRenderProxy::GetIncompleteMaterialWithFallback(ERHIFea
 	return *Material;
 }
 
+class FUniformExpressionCacheAsyncUpdateTask
+{
+public:
+	FUniformExpressionCacheAsyncUpdateTask()
+	{
+		DelegateHandle = FCoreRenderDelegates::OnFlushRenderingCommandsEnd.AddLambda([this]
+		{
+			Wait();
+		});
+	}
+
+	~FUniformExpressionCacheAsyncUpdateTask()
+	{
+		Wait();
+		FCoreRenderDelegates::OnFlushRenderingCommandsEnd.Remove(DelegateHandle);
+	}
+
+	void SetTask(const FGraphEventRef& InTask)
+	{
+		check(!Task);
+		Task = InTask;
+	}
+
+	void Wait()
+	{
+		if (Task)
+		{
+			Task->Wait(ENamedThreads::GetRenderThread_Local());
+			Task = nullptr;
+		}
+	}
+
+private:
+	FGraphEventRef Task;
+	FDelegateHandle DelegateHandle;
+
+} GUniformExpressionCacheAsyncUpdateTask;
+
 void FUniformExpressionCacheUpdater::Update(FRHICommandListImmediate& RHICmdListImmediate)
 {
 	if (Items.IsEmpty())
@@ -3899,6 +3944,8 @@ void FUniformExpressionCacheUpdater::Update(FRHICommandListImmediate& RHICmdList
 		return;
 	}
 
+	GUniformExpressionCacheAsyncUpdateTask.Wait();
+
 	FRHICommandList* RHICmdList = new FRHICommandList(FRHIGPUMask::All());
 
 	FGraphEventRef Event = FFunctionGraphTask::CreateAndDispatchWhenReady(
@@ -3914,14 +3961,16 @@ void FUniformExpressionCacheUpdater::Update(FRHICommandListImmediate& RHICmdList
 
 			FMaterialRenderContext Context(Item.MaterialRenderProxy, *Item.Material, Item.bShowSelection);
 
-			Item.UniformExpressionSet->FillUniformBuffer(Context, *Item.UniformExpressionCache, Item.UniformBufferLayout, TempBuffer, Item.UniformBufferLayout->ConstantBufferSize);
+			Item.UniformExpressionSet->FillUniformBuffer(Context, Item.AllocatedVTs, Item.UniformBufferLayout, TempBuffer, Item.UniformBufferLayout->ConstantBufferSize);
 
-			RHICmdList->UpdateUniformBuffer(Item.UniformExpressionCache->UniformBuffer, TempBuffer);
+			RHICmdList->UpdateUniformBuffer(Item.UniformBuffer, TempBuffer);
 		}
 
 	}, TStatId(), nullptr, ENamedThreads::AnyHiPriThreadHiPriTask);
 
 	RHICmdListImmediate.QueueRenderThreadCommandListSubmit(Event, RHICmdList);
+
+	GUniformExpressionCacheAsyncUpdateTask.SetTask(Event);
 }
 
 void FMaterialRenderProxy::UpdateDeferredCachedUniformExpressions()
@@ -3935,6 +3984,7 @@ void FMaterialRenderProxy::UpdateDeferredCachedUniformExpressions()
 	QUICK_SCOPE_CYCLE_COUNTER(STAT_UpdateDeferredCachedUniformExpressions);
 
 	FUniformExpressionCacheUpdater Updater;
+	FUniformExpressionCacheUpdater* UpdaterIfEnabled = !GRHICommandList.Bypass() && GUniformExpressionCacheAsyncUpdates > 0 ? &Updater : nullptr;
 
 	for (TSet<FMaterialRenderProxy*>::TConstIterator It(DeferredUniformExpressionCacheRequests); It; ++It)
 	{
@@ -3952,12 +4002,15 @@ void FMaterialRenderProxy::UpdateDeferredCachedUniformExpressions()
 			{
 				FMaterialRenderContext MaterialRenderContext(MaterialProxy, *Material, nullptr);
 				MaterialRenderContext.bShowSelection = GIsEditor;
-				MaterialProxy->EvaluateUniformExpressions(MaterialProxy->UniformExpressionCache[(int32)InFeatureLevel], MaterialRenderContext, &Updater);
+				MaterialProxy->EvaluateUniformExpressions(MaterialProxy->UniformExpressionCache[(int32)InFeatureLevel], MaterialRenderContext, UpdaterIfEnabled);
 			}
 		});
 	}
 
-	Updater.Update(FRHICommandListExecutor::GetImmediateCommandList());
+	if (UpdaterIfEnabled)
+	{
+		Updater.Update(FRHICommandListExecutor::GetImmediateCommandList());
+	}
 
 	DeferredUniformExpressionCacheRequests.Reset();
 }
diff --git a/Engine/Source/Runtime/Engine/Private/Materials/MaterialUniformExpressions.cpp b/Engine/Source/Runtime/Engine/Private/Materials/MaterialUniformExpressions.cpp
--- a/Engine/Source/Runtime/Engine/Private/Materials/MaterialUniformExpressions.cpp
+++ b/Engine/Source/Runtime/Engine/Private/Materials/MaterialUniformExpressions.cpp
@@ -667,6 +667,11 @@ void FUniformExpressionSet::GetTextureValue(int32 Index, const FMaterialRenderCo
 }
 
 void FUniformExpressionSet::FillUniformBuffer(const FMaterialRenderContext& MaterialRenderContext, const FUniformExpressionCache& UniformExpressionCache, const FRHIUniformBufferLayout* UniformBufferLayout, uint8* TempBuffer, int TempBufferSize) const
+{
+	FillUniformBuffer(MaterialRenderContext, UniformExpressionCache.AllocatedVTs, UniformBufferLayout, TempBuffer, TempBufferSize);
+}
+
+void FUniformExpressionSet::FillUniformBuffer(const FMaterialRenderContext& MaterialRenderContext, TConstArrayView<IAllocatedVirtualTexture*> AllocatedVTs, const FRHIUniformBufferLayout* UniformBufferLayout, uint8* TempBuffer, int TempBufferSize) const
 {
 	using namespace UE::Shader;
 	check(IsInParallelRenderingThread());
@@ -680,10 +685,10 @@ void FUniformExpressionSet::FillUniformBuffer(const FMaterialRenderContext& Mate
 		check(BufferCursor <= TempBuffer + TempBufferSize);
 
 		// Dump virtual texture per page table uniform data
-		check(UniformExpressionCache.AllocatedVTs.Num() == VTStacks.Num());
+		check(AllocatedVTs.Num() == VTStacks.Num());
 		for ( int32 VTStackIndex = 0; VTStackIndex < VTStacks.Num(); ++VTStackIndex)
 		{
-			const IAllocatedVirtualTexture* AllocatedVT = UniformExpressionCache.AllocatedVTs[VTStackIndex];
+			const IAllocatedVirtualTexture* AllocatedVT = AllocatedVTs[VTStackIndex];
 			FUintVector4* VTPackedPageTableUniform = (FUintVector4*)BufferCursor;
 			if (AllocatedVT)
 			{
@@ -715,7 +720,7 @@ void FUniformExpressionSet::FillUniformBuffer(const FMaterialRenderContext& Mate
 				if (Texture != nullptr)
 				{
 					const FVTPackedStackAndLayerIndex StackAndLayerIndex = GetVTStackAndLayerIndex(ExpressionIndex);
-					const IAllocatedVirtualTexture* AllocatedVT = UniformExpressionCache.AllocatedVTs[StackAndLayerIndex.StackIndex];
+					const IAllocatedVirtualTexture* AllocatedVT = AllocatedVTs[StackAndLayerIndex.StackIndex];
 					if (AllocatedVT)
 					{
 						AllocatedVT->GetPackedUniform(VTPackedUniform, StackAndLayerIndex.LayerIndex);
@@ -1199,7 +1204,7 @@ void FUniformExpressionSet::FillUniformBuffer(const FMaterialRenderContext& Mate
 			void** ResourceTablePageIndirectionBuffer = (void**)((uint8*)BufferCursor + 0 * SHADER_PARAMETER_POINTER_ALIGNMENT);
 			BufferCursor = ((uint8*)BufferCursor) + SHADER_PARAMETER_POINTER_ALIGNMENT;
 
-			const IAllocatedVirtualTexture* AllocatedVT = UniformExpressionCache.AllocatedVTs[VTStackIndex];
+			const IAllocatedVirtualTexture* AllocatedVT = AllocatedVTs[VTStackIndex];
 			if (AllocatedVT != nullptr)
 			{
 				FRHITexture* PageTable0RHI = AllocatedVT->GetPageTableTexture(0u);
@@ -1253,7 +1258,7 @@ void FUniformExpressionSet::FillUniformBuffer(const FMaterialRenderContext& Mate
 					FVirtualTexture2DResource* VTResource = (FVirtualTexture2DResource*)Texture->GetResource();
 					check(VTResource);
 
-					const IAllocatedVirtualTexture* AllocatedVT = UniformExpressionCache.AllocatedVTs[StackAndLayerIndex.StackIndex];
+					const IAllocatedVirtualTexture* AllocatedVT = AllocatedVTs[StackAndLayerIndex.StackIndex];
 					if (AllocatedVT != nullptr)
 					{
 						FRHIShaderResourceView* PhysicalViewRHI = AllocatedVT->GetPhysicalTextureSRV(StackAndLayerIndex.LayerIndex, VTResource->bSRGB);
diff --git a/Engine/Source/Runtime/Engine/Public/MaterialShared.h b/Engine/Source/Runtime/Engine/Public/MaterialShared.h
--- a/Engine/Source/Runtime/Engine/Public/MaterialShared.h
+++ b/Engine/Source/Runtime/Engine/Public/MaterialShared.h
@@ -575,6 +575,8 @@ public:
 
 	ENGINE_API void FillUniformBuffer(const FMaterialRenderContext& MaterialRenderContext, const FUniformExpressionCache& UniformExpressionCache, const FRHIUniformBufferLayout* UniformBufferLayout, uint8* TempBuffer, int TempBufferSize) const;
 
+	ENGINE_API void FillUniformBuffer(const FMaterialRenderContext& MaterialRenderContext, TConstArrayView<IAllocatedVirtualTexture*> AllocatedVTs, const FRHIUniformBufferLayout* UniformBufferLayout, uint8* TempBuffer, int TempBufferSize) const;
+
 	// Get a combined hash of all referenced Texture2D's underlying RHI textures, going through TextureReferences. Can be used to tell if any texture has gone through texture streaming mip changes recently.
 	ENGINE_API uint32 GetReferencedTexture2DRHIHash(const FMaterialRenderContext& MaterialRenderContext) const;
 
@@ -2443,16 +2445,10 @@ struct FUniformExpressionCache
 	TArray<IAllocatedVirtualTexture*> OwnedAllocatedVTs;
 	/** Ids of parameter collections needed for rendering. */
 	TArray<FGuid> ParameterCollections;
-	/** True if the cache is up to date. */
-	bool bUpToDate;
-
 	/** Shader map that was used to cache uniform expressions on this material.  This is used for debugging and verifying correct behavior. */
-	const FMaterialShaderMap* CachedUniformExpressionShaderMap;
-
-	FUniformExpressionCache() :
-		bUpToDate(false),
-		CachedUniformExpressionShaderMap(NULL)
-	{}
+	const FMaterialShaderMap* CachedUniformExpressionShaderMap = nullptr;
+	/** True if the cache is up to date. */
+	bool bUpToDate = false;
 
 	/** Destructor. */
 	ENGINE_API ~FUniformExpressionCache();
@@ -2495,7 +2491,8 @@ private:
 	{
 		FItem() = default;
 		FItem(FUniformExpressionCache* InUniformExpressionCache, const FUniformExpressionSet* InUniformExpressionSet, const FRHIUniformBufferLayout* InUniformBufferLayout, const FMaterialRenderContext& Context)
-			: UniformExpressionCache(InUniformExpressionCache)
+			: UniformBuffer(InUniformExpressionCache->UniformBuffer)
+			, AllocatedVTs(InUniformExpressionCache->AllocatedVTs)
 			, UniformExpressionSet(InUniformExpressionSet)
 			, UniformBufferLayout(InUniformBufferLayout)
 			, MaterialRenderProxy(Context.MaterialRenderProxy)
@@ -2503,7 +2500,8 @@ private:
 			, bShowSelection(Context.bShowSelection)
 		{}
 
-		FUniformExpressionCache* UniformExpressionCache = nullptr;
+		TRefCountPtr<FRHIUniformBuffer> UniformBuffer;
+		TArray<IAllocatedVirtualTexture*, FConcurrentLinearArrayAllocator> AllocatedVTs;
 		const FUniformExpressionSet* UniformExpressionSet = nullptr;
 		const FRHIUniformBufferLayout* UniformBufferLayout = nullptr;
 		const FMaterialRenderProxy* MaterialRenderProxy = nullptr;
