Speculative fix for CIS error

#rnx

diff --git a/Engine/Source/Runtime/Engine/Public/TransactionCommon.h b/Engine/Source/Runtime/Engine/Public/TransactionCommon.h
--- a/Engine/Source/Runtime/Engine/Public/TransactionCommon.h
+++ b/Engine/Source/Runtime/Engine/Public/TransactionCommon.h
@@ -7,6 +7,7 @@
 #include "UObject/WeakObjectPtr.h"
 #include "Serialization/ArchiveUObject.h"
 #include "Serialization/ArchiveSerializedPropertyChain.h"
+#include "Templates/TypeHash.h"
 
 class FReferenceCollector;
 struct FTransactionObjectDeltaChange;
@@ -60,7 +61,7 @@ public:
 
 	friend uint32 GetTypeHash(const FPersistentObjectRef& InObjRef)
 	{
-		return HashCombine(GetTypeHash(InObjRef.ReferenceType), GetTypeHash(InObjRef.RootObject));
+		return HashCombine(::GetTypeHash(InObjRef.ReferenceType), GetTypeHash(InObjRef.RootObject));
 	}
 
 	bool IsRootObjectReference() const
