Fixed assert when waiting on uniform expression task due to call to FlushRenderingThread()

diff --git a/Engine/Source/Runtime/Engine/Private/Materials/MaterialShared.cpp b/Engine/Source/Runtime/Engine/Private/Materials/MaterialShared.cpp
--- a/Engine/Source/Runtime/Engine/Private/Materials/MaterialShared.cpp
+++ b/Engine/Source/Runtime/Engine/Private/Materials/MaterialShared.cpp
@@ -3926,7 +3926,7 @@ public:
 	{
 		if (Task)
 		{
-			Task->Wait(ENamedThreads::GetRenderThread_Local());
+			Task->Wait();
 			Task = nullptr;
 		}
 	}
