Fix a bug of counting skeletal&static meshes memory with mesh streaming enabled.
The IncrementMemoryStats is called at the FSkeletalMeshLODRenderData::InitResources and FSkeletalMeshStreamIn::DoFinishUpdate.

diff --git a/Engine/Source/Runtime/Engine/Private/SkeletalMeshLODRenderData.cpp b/Engine/Source/Runtime/Engine/Private/SkeletalMeshLODRenderData.cpp
--- a/Engine/Source/Runtime/Engine/Private/SkeletalMeshLODRenderData.cpp
+++ b/Engine/Source/Runtime/Engine/Private/SkeletalMeshLODRenderData.cpp
@@ -161,7 +161,10 @@ FArchive& operator<<(FArchive& Ar, FSkelMeshRenderSection& S)
 
 void FSkeletalMeshLODRenderData::InitResources(bool bNeedsVertexColors, int32 LODIndex, TArray<UMorphTarget*>& InMorphTargets, USkinnedAsset* Owner)
 {
-	IncrementMemoryStats(bNeedsVertexColors);
+	if (bStreamedDataInlined)
+	{
+		IncrementMemoryStats(bNeedsVertexColors);
+	}
 
 	MultiSizeIndexContainer.InitResources();
 
diff --git a/Engine/Source/Runtime/Engine/Private/StaticMesh.cpp b/Engine/Source/Runtime/Engine/Private/StaticMesh.cpp
--- a/Engine/Source/Runtime/Engine/Private/StaticMesh.cpp
+++ b/Engine/Source/Runtime/Engine/Private/StaticMesh.cpp
@@ -1315,7 +1315,10 @@ void FStaticMeshLODResources::UpdateVertexMemoryStats() const
 
 void FStaticMeshLODResources::InitResources(UStaticMesh* Parent)
 {
-	UpdateIndexMemoryStats<true>();
+	if (bBuffersInlined)
+	{
+		UpdateIndexMemoryStats<true>();
+	}
 
 	BeginInitResource(&IndexBuffer);
 	if(bHasWireframeIndices)
@@ -1376,7 +1379,10 @@ void FStaticMeshLODResources::InitResources(UStaticMesh* Parent)
 	ENQUEUE_RENDER_COMMAND(UpdateMemoryStats)(
 		[this](FRHICommandListImmediate&)
 	{
-		UpdateVertexMemoryStats<true>();
+		if (bBuffersInlined)
+		{
+			UpdateVertexMemoryStats<true>();
+		}
 	});
 #endif
 
