If a connection is set to InternalAck, in ReplicateProperties, don't use the backward compatible path if RepFlags.bSerializePropertyNames is true since it doesn't support that flag.

diff --git a/Engine/Source/Runtime/Engine/Private/RepLayout.cpp b/Engine/Source/Runtime/Engine/Private/RepLayout.cpp
--- a/Engine/Source/Runtime/Engine/Private/RepLayout.cpp
+++ b/Engine/Source/Runtime/Engine/Private/RepLayout.cpp
@@ -2081,7 +2081,7 @@ bool FRepLayout::ReplicateProperties(
 	}
 
 	// Send the final merged change list
-	if (OwningChannel->Connection->IsInternalAck())
+	if (OwningChannel->Connection->IsInternalAck() && !RepFlags.bSerializePropertyNames)
 	{
 		// Remember all properties that have changed since this channel was first opened in case we need it (for bResendAllDataSinceOpen)
 		// We use UnfilteredChanged so LifetimeChangelist contains all properties, regardless of Active state.
