Fixed false positives in transaction diffing caused by instanced sub-objects

Objects may be replaced during re-instancing (eg, SCS components on an actor), but the object reference needs to be considered identical for diffing as the objects still have the same path name. This was previously the case when the diffing code was using FPersistentObjectRef (which compared the path names for sub-objects).

#rnx

diff --git a/Engine/Source/Runtime/Engine/Private/TransactionCommon.cpp b/Engine/Source/Runtime/Engine/Private/TransactionCommon.cpp
--- a/Engine/Source/Runtime/Engine/Private/TransactionCommon.cpp
+++ b/Engine/Source/Runtime/Engine/Private/TransactionCommon.cpp
@@ -412,8 +412,14 @@ FArchive& FDiffableObjectDataWriter::operator<<(FName& N)
 
 FArchive& FDiffableObjectDataWriter::operator<<(UObject*& Res)
 {
-	PTRINT ResInt = (PTRINT)Res;
-	return *this << ResInt;
+	FName ResPathName;
+	if (Res)
+	{
+		FNameBuilder ResPathNameBuilder;
+		Res->GetPathName(nullptr, ResPathNameBuilder);
+		ResPathName = FName(ResPathNameBuilder);
+	}
+	return *this << ResPathName;
 }
 
 FName FDiffableObjectDataWriter::FCachedPropertyKey::SyncCache(const FArchiveSerializedPropertyChain* InPropertyChain)
