Fix groom manual skin cache not firing.
Issue was introduced by GetCachedGeometry() returning sections even when core skin cache wasn't enabled.

diff --git a/Engine/Source/Runtime/Engine/Private/SkeletalRenderGPUSkin.cpp b/Engine/Source/Runtime/Engine/Private/SkeletalRenderGPUSkin.cpp
--- a/Engine/Source/Runtime/Engine/Private/SkeletalRenderGPUSkin.cpp
+++ b/Engine/Source/Runtime/Engine/Private/SkeletalRenderGPUSkin.cpp
@@ -2200,7 +2200,13 @@ const TArray<FMatrix44f>& FSkeletalMeshObjectGPUSkin::GetReferenceToLocalMatrice
 bool FSkeletalMeshObjectGPUSkin::GetCachedGeometry(FCachedGeometry& OutCachedGeometry) const
 {
 	OutCachedGeometry = FCachedGeometry();
-	
+
+	// Cached geometry is only available if we are using skin cache or a mesh deformer.	
+	if (DynamicData == nullptr || !(DynamicData->bIsSkinCacheAllowed || DynamicData->bHasMeshDeformer))
+	{
+		return false;
+	}
+
 	const int32 LodIndex = GetLOD();
 	if (SkeletalMeshRenderData == nullptr || !SkeletalMeshRenderData->LODRenderData.IsValidIndex(LodIndex))
 	{
