Fix broken "Selected Morph Target Vertices" editor visualization

#github 9174

diff --git a/Engine/Source/Runtime/Engine/Private/SkeletalRenderCPUSkin.cpp b/Engine/Source/Runtime/Engine/Private/SkeletalRenderCPUSkin.cpp
--- a/Engine/Source/Runtime/Engine/Private/SkeletalRenderCPUSkin.cpp
+++ b/Engine/Source/Runtime/Engine/Private/SkeletalRenderCPUSkin.cpp
@@ -1183,8 +1183,8 @@ static void CalculateMorphTargetWeights(FFinalSkinVertex* DestVertex, FSkeletalM
 
 	for (FFinalSkinVertex* ClearVert = DestVertex; ClearVert != EndVert; ++ClearVert)
 	{
-		DestVertex->TextureCoordinates[0].X = 0.0f;
-		DestVertex->TextureCoordinates[0].Y = 0.0f;
+		ClearVert->TextureCoordinates[0].X = 0.0f;
+		ClearVert->TextureCoordinates[0].Y = 0.0f;
 	}
 
 	for (const UMorphTarget* Morphtarget : InMorphTargetOfInterest)
@@ -1194,8 +1194,8 @@ static void CalculateMorphTargetWeights(FFinalSkinVertex* DestVertex, FSkeletalM
 		for (int32 MorphVertexIndex = 0; MorphVertexIndex < NumDeltas; ++MorphVertexIndex)
 		{
 			FFinalSkinVertex* SetVert = DestVertex + MTLODVertices[MorphVertexIndex].SourceIdx;
-			DestVertex->TextureCoordinates[0].X += 1.0f;
-			DestVertex->TextureCoordinates[0].Y += 1.0f;
+			SetVert->TextureCoordinates[0].X += 1.0f;
+			SetVert->TextureCoordinates[0].Y += 1.0f;
 		}
 	}
 }
