Fix Single Shadow Sample shadows

diff --git a/Engine/Source/Runtime/Engine/Private/PrimitiveUniformShaderParameters.cpp b/Engine/Source/Runtime/Engine/Private/PrimitiveUniformShaderParameters.cpp
--- a/Engine/Source/Runtime/Engine/Private/PrimitiveUniformShaderParameters.cpp
+++ b/Engine/Source/Runtime/Engine/Private/PrimitiveUniformShaderParameters.cpp
@@ -159,6 +159,7 @@ FPrimitiveSceneShaderData::FPrimitiveSceneShaderData(const FPrimitiveSceneProxy*
 		.InstancePayloadDataOffset(Proxy->GetPrimitiveSceneInfo()->GetInstancePayloadDataOffset())
 		.InstancePayloadDataStride(Proxy->GetPrimitiveSceneInfo()->GetInstancePayloadDataStride())
 		.HasCapsuleRepresentation(Proxy->HasDynamicIndirectShadowCasterRepresentation())
+		.UseSingleSampleShadowFromStationaryLights(Proxy->UseSingleSampleShadowFromStationaryLights())
 		.ReceivesDecals(Proxy->ReceivesDecals())
 		.ShouldCacheShadow(Proxy->ShouldCacheShadow())
 		.OutputVelocity(bOutputVelocity)
diff --git a/Engine/Source/Runtime/Renderer/Private/GPUScene.cpp b/Engine/Source/Runtime/Renderer/Private/GPUScene.cpp
--- a/Engine/Source/Runtime/Renderer/Private/GPUScene.cpp
+++ b/Engine/Source/Runtime/Renderer/Private/GPUScene.cpp
@@ -983,6 +983,8 @@ void FGPUScene::UploadGeneral(FRDGBuilder& GraphBuilder, FScene *Scene, FRDGExte
 				// 1. do (potentially) parallel work next
 				ParallelFor(NumPrimitiveDataUploads, [&UploadDataSourceAdapter, this](int32 ItemIndex)
 				{
+					FTaskTagScope TaskTagScope(ETaskTag::EParallelRenderingThread);
+
 					FPrimitiveUploadInfo UploadInfo;
 					UploadDataSourceAdapter.GetPrimitiveInfo(ItemIndex, UploadInfo);
 
