Strata - fixed byte count per pixel evalaution for some materials.

diff --git a/Engine/Source/Runtime/Engine/Private/Materials/HLSLMaterialTranslator.cpp b/Engine/Source/Runtime/Engine/Private/Materials/HLSLMaterialTranslator.cpp
--- a/Engine/Source/Runtime/Engine/Private/Materials/HLSLMaterialTranslator.cpp
+++ b/Engine/Source/Runtime/Engine/Private/Materials/HLSLMaterialTranslator.cpp
@@ -10057,6 +10057,8 @@ bool FHLSLMaterialTranslator::StrataGenerateDerivedMaterialOperatorData()
 			// 1. Evaluate simple/single BSDF
 			bStrataMaterialIsSimple = StrataMaterialBSDFCount == 1;
 			bStrataMaterialIsSingle = StrataMaterialBSDFCount == 1;
+			bool bIsFastWaterPath = false;
+			bool bCustomEncoding = false;
 			for (auto& It : StrataMaterialExpressionRegisteredOperators)
 			{
 				if (It.IsDiscarded())
@@ -10083,18 +10085,21 @@ bool FHLSLMaterialTranslator::StrataGenerateDerivedMaterialOperatorData()
 					{
 						bStrataMaterialIsSimple = false;
 						bStrataMaterialIsSingle = false;
+						bCustomEncoding			= true;
 						break;
 					}
 					case STRATA_BSDF_TYPE_EYE:
 					{
 						bStrataMaterialIsSimple = false;
 						bStrataMaterialIsSingle = false;
+						bCustomEncoding			= true;
 						break;
 					}
 					case STRATA_BSDF_TYPE_SINGLELAYERWATER:
 					{
 						bStrataMaterialIsSimple = false;
 						bStrataMaterialIsSingle = false;
+						bIsFastWaterPath		= true;
 						break;
 					}
 					}
@@ -10106,7 +10111,7 @@ bool FHLSLMaterialTranslator::StrataGenerateDerivedMaterialOperatorData()
 
 			// 2. Header
 
-			if (!bStrataMaterialIsSimple && !bStrataMaterialIsSingle) // header written later, 
+			if (!bStrataMaterialIsSimple && !bStrataMaterialIsSingle && !bCustomEncoding && !bIsFastWaterPath) // header written later, 
 			{
 				// Packed Header
 				StrataMaterialRequestedSizeByte += UintByteSize;
@@ -10114,11 +10119,9 @@ bool FHLSLMaterialTranslator::StrataGenerateDerivedMaterialOperatorData()
 				// Shared local bases between BSDFs
 				StrataMaterialRequestedSizeByte += UsedSharedLocalBasesCount * STRATA_PACKED_SHAREDLOCALBASIS_STRIDE_BYTES;
 			}
-			else
-			{
-				// Top normal texture is used to read the data
-				StrataMaterialRequestedSizeByte += UintByteSize;
-			}
+			// Note:
+			//  - We do not need to account for the Top Normal texture when evaluating the material byte count for the optimisation algorithm.
+			//  - This is because we only need to optimise for the Strata uint matertial buffer.
 
 
 			// 2. The list of BSDFs
@@ -10149,6 +10152,21 @@ bool FHLSLMaterialTranslator::StrataGenerateDerivedMaterialOperatorData()
 						// Header
 						StrataMaterialRequestedSizeByte += UintByteSize;
 					}
+					else if (bCustomEncoding)
+					{
+						// Header
+						StrataMaterialRequestedSizeByte += UintByteSize;
+					}
+					else if (bIsFastWaterPath)
+					{
+						// Header + Data
+						StrataMaterialRequestedSizeByte += UintByteSize;
+						// Data
+						StrataMaterialRequestedSizeByte += UintByteSize;	
+						// Optional separted dir light luminance
+						StrataMaterialRequestedSizeByte += UintByteSize;
+						break; // Stop here
+					}
 					else if (bMayHaveColoredWeight)
 					{
 						// BSDF state
@@ -10191,21 +10209,21 @@ bool FHLSLMaterialTranslator::StrataGenerateDerivedMaterialOperatorData()
 					}
 					case STRATA_BSDF_TYPE_HAIR:
 					{
+						// Custom encoding
 						StrataMaterialRequestedSizeByte += UintByteSize;
 						StrataMaterialRequestedSizeByte += UintByteSize;
 						break;
 					}
 					case STRATA_BSDF_TYPE_EYE:
 					{
+						// Custom encoding
 						StrataMaterialRequestedSizeByte += UintByteSize;
 						StrataMaterialRequestedSizeByte += UintByteSize;
 						break;
 					}
 					case STRATA_BSDF_TYPE_SINGLELAYERWATER:
 					{
-						StrataMaterialRequestedSizeByte += UintByteSize;
-						StrataMaterialRequestedSizeByte += UintByteSize;
-						StrataMaterialRequestedSizeByte += UintByteSize;
+						Errorf(TEXT("Strata error: single layer water should go through the its dedicated fast path in %s (asset: %s).\r\n"), *Material->GetDebugName(), *Material->GetAssetPath().ToString());
 						break;
 					}
 					case STRATA_BSDF_TYPE_UNLIT:
