Minor fixes from pose watch manager extension merge
* Added GetElement to UPoseWatch.
* Fixed clicking on a pose watch element to show its properties in the details panel.
* Fixed blend profiles not showing when selecting a viewport mask.
* Fixed upgrade path for old pose watches.
* Deprecated SetColor and GetColor in UPoseWatch as only elements have colors now.

diff --git a/Engine/Source/Editor/AnimationBlueprintEditor/Private/AnimationNodes/SPoseWatchOverlay.cpp b/Engine/Source/Editor/AnimationBlueprintEditor/Private/AnimationNodes/SPoseWatchOverlay.cpp
--- a/Engine/Source/Editor/AnimationBlueprintEditor/Private/AnimationNodes/SPoseWatchOverlay.cpp
+++ b/Engine/Source/Editor/AnimationBlueprintEditor/Private/AnimationNodes/SPoseWatchOverlay.cpp
@@ -47,9 +47,13 @@ FSlateColor SPoseWatchOverlay::GetPoseViewColor() const
 	UPoseWatch* CurPoseWatch = PoseWatch.Get();
 	if (CurPoseWatch)
 	{
-		FLinearColor OutColor = CurPoseWatch->GetColor();
-		OutColor.A = CurPoseWatch->GetShouldDeleteOnDeselect() ? AlphaTemporary : AlphaPermanent;
-		return FSlateColor(OutColor);
+		TObjectPtr<UPoseWatchPoseElement> PoseElement = CurPoseWatch->GetFirstElementOfType<UPoseWatchPoseElement>();
+		if (PoseElement)
+		{
+			FLinearColor OutColor = PoseElement->GetColor();
+			OutColor.A = CurPoseWatch->GetShouldDeleteOnDeselect() ? AlphaTemporary : AlphaPermanent;
+			return FSlateColor(OutColor);
+		}
 	}
 	return FSlateColor(FColor::Black);
 }
diff --git a/Engine/Source/Editor/Persona/Private/BlendProfileCustomization.cpp b/Engine/Source/Editor/Persona/Private/BlendProfileCustomization.cpp
--- a/Engine/Source/Editor/Persona/Private/BlendProfileCustomization.cpp
+++ b/Engine/Source/Editor/Persona/Private/BlendProfileCustomization.cpp
@@ -111,6 +111,13 @@ USkeleton* FBlendProfileCustomization::GetSkeletonFromOuter(const UObject* Outer
 	{
 		AnimBlueprint = Cast<UAnimBlueprint>(FBlueprintEditorUtils::FindBlueprintForNode(OuterPoseWatch->Node.Get()));
 	}
+	else if (const UPoseWatchPoseElement* OuterPoseElement = Cast<UPoseWatchPoseElement>(Outer))
+	{
+		if (const UPoseWatch* OuterPoseElementParent = OuterPoseElement->GetParent())
+		{
+			AnimBlueprint = Cast<UAnimBlueprint>(FBlueprintEditorUtils::FindBlueprintForNode(OuterPoseElementParent->Node.Get()));
+		}
+	}
 
 	// If outer belongs to an anim blueprint, grab its skeleton.
 	if (AnimBlueprint)
diff --git a/Engine/Source/Editor/Persona/Private/PoseWatchManagerPoseWatchTreeItem.cpp b/Engine/Source/Editor/Persona/Private/PoseWatchManagerPoseWatchTreeItem.cpp
--- a/Engine/Source/Editor/Persona/Private/PoseWatchManagerPoseWatchTreeItem.cpp
+++ b/Engine/Source/Editor/Persona/Private/PoseWatchManagerPoseWatchTreeItem.cpp
@@ -222,14 +222,4 @@ bool FPoseWatchManagerPoseWatchTreeItem::IsExpanded() const
 	return PoseWatch->GetIsExpanded();
 }
 
-FColor FPoseWatchManagerPoseWatchTreeItem::GetColor() const
-{
-	return PoseWatch->GetColor();
-}
-
-void FPoseWatchManagerPoseWatchTreeItem::SetColor(const FColor& InColor)
-{
-	PoseWatch->SetColor(InColor);
-}
-
 #undef LOCTEXT_NAMESPACE
diff --git a/Engine/Source/Editor/Persona/Private/PoseWatchManagerPoseWatchTreeItem.h b/Engine/Source/Editor/Persona/Private/PoseWatchManagerPoseWatchTreeItem.h
--- a/Engine/Source/Editor/Persona/Private/PoseWatchManagerPoseWatchTreeItem.h
+++ b/Engine/Source/Editor/Persona/Private/PoseWatchManagerPoseWatchTreeItem.h
@@ -36,7 +36,5 @@ public:
 	virtual bool IsEnabled() const override;
 	virtual void SetIsExpanded(const bool bIsExpanded) override;
 	virtual bool IsExpanded() const override;
-	virtual FColor GetColor() const override;
-	virtual void SetColor(const FColor& InColor) override;
 	/* End IPoseWatchManagerTreeItem Implementation */
 };
diff --git a/Engine/Source/Editor/Persona/Private/SPoseWatchManager.cpp b/Engine/Source/Editor/Persona/Private/SPoseWatchManager.cpp
--- a/Engine/Source/Editor/Persona/Private/SPoseWatchManager.cpp
+++ b/Engine/Source/Editor/Persona/Private/SPoseWatchManager.cpp
@@ -700,20 +700,13 @@ void SPoseWatchManager::OnManagerTreeSelectionChanged(FPoseWatchManagerTreeItemP
 	{
 		TGuardValue<bool> ReentrantGuard(bIsReentrant, true);
 
+		BlueprintEditor->GetInspector()->ShowDetailsForSingleObject(nullptr);
 		if (TreeItem.IsValid())
 		{
-			if (const FPoseWatchManagerPoseWatchTreeItem* PoseWatchItem = TreeItem.Get()->CastTo<FPoseWatchManagerPoseWatchTreeItem>())
+			if (const FPoseWatchManagerElementTreeItem* ElementItem = TreeItem->CastTo<FPoseWatchManagerElementTreeItem>())
 			{
-				BlueprintEditor->GetInspector()->ShowDetailsForSingleObject(PoseWatchItem->PoseWatch.Get());
+				BlueprintEditor->GetInspector()->ShowDetailsForSingleObject(ElementItem->PoseWatchElement.Get());
 			}
-			else
-			{
-				BlueprintEditor->GetInspector()->ShowDetailsForSingleObject(nullptr);
-			}
-		}
-		else
-		{
-			BlueprintEditor->GetInspector()->ShowDetailsForSingleObject(nullptr);
 		}
 		OnItemSelectionChanged.Broadcast(TreeItem, SelectInfo);
 	}
diff --git a/Engine/Source/Runtime/Engine/Classes/Engine/PoseWatch.h b/Engine/Source/Runtime/Engine/Classes/Engine/PoseWatch.h
--- a/Engine/Source/Runtime/Engine/Classes/Engine/PoseWatch.h
+++ b/Engine/Source/Runtime/Engine/Classes/Engine/PoseWatch.h
@@ -295,6 +295,7 @@ public:
 	bool GetIsNodeEnabled() const;
 
 	/** Returns the color to display the pose watch using */
+	UE_DEPRECATED(5.1, "Node watchs no longer have colors, use the color of its elements instead.")
 	FColor GetColor() const;
 
 	/** Set whether this should display its children in the pose watch manager window */
@@ -327,6 +328,7 @@ public:
 	void SetIsVisible(bool bInIsVisible);
 
 	/** Sets the display color of this pose watch in the UI and viewport */
+	UE_DEPRECATED(5.1, "Node watches no longer have colors, use the color of its elements instead.")
 	void SetColor(const FColor& InColor);
 
 	/** Sets whether this pose watch should delete after deselecting it's assigned node (Editor preference) */
@@ -371,6 +373,20 @@ public:
 	/** Returns true if this pose watch contains the specified element */
 	bool Contains(const TObjectPtr<UPoseWatchElement> InElement) { return Elements.Contains(InElement); }
 	bool Contains(const UPoseWatchElement* const InElement) { return Elements.ContainsByPredicate([InElement](const TObjectPtr<UPoseWatchElement>& ContainedElement) { return InElement == ContainedElement.Get(); }); }
+
+	/** Returns the first element with the matching type */
+	template< class TElementType > TObjectPtr<TElementType> GetFirstElementOfType()
+	{
+		for (UPoseWatchElement* Element : Elements)
+		{
+			if (TElementType* FoundElement = Cast<TElementType>(Element))
+			{
+				return FoundElement;
+			}
+		}
+		return nullptr;
+	}
+
 	
 private:
 	/** Returns a unique name for a new pose watch placed within InParent */
@@ -413,7 +429,7 @@ protected:
 	bool bIsExpanded = true;
 
 	UPROPERTY()
-	FColor Color;
+	FColor Color_DEPRECATED;
 
 	UPROPERTY()
 	FText Label;
@@ -435,7 +451,7 @@ template< class TElementType > TObjectPtr<TElementType> UPoseWatch::AddElement(c
 	PoseWatchElement->SetParent(this);
 	PoseWatchElement->SetUniqueLabel(InLabel);
 	PoseWatchElement->SetIconName(InIconName);
-	PoseWatchElement->SetColor(GetColor());
+	PoseWatchElement->SetColor(FColor::MakeRandomColor());
 	Elements.Add(PoseWatchElement);
 
 	return PoseWatchElement;
diff --git a/Engine/Source/Runtime/Engine/Private/Engine/PoseWatch.cpp b/Engine/Source/Runtime/Engine/Private/Engine/PoseWatch.cpp
--- a/Engine/Source/Runtime/Engine/Private/Engine/PoseWatch.cpp
+++ b/Engine/Source/Runtime/Engine/Private/Engine/PoseWatch.cpp
@@ -332,7 +332,7 @@ bool UPoseWatchFolder::HasPoseWatchDescendents() const
 
 FText UPoseWatchElement::GetDefaultLabel() const
 {
-	return LOCTEXT("NewPoseWatchElement", "NewPoseWatchElement");
+	return LOCTEXT("NewPoseWatchElement", "Pose Watch");
 }
 
 FText UPoseWatchElement::GetLabel() const
@@ -475,6 +475,9 @@ UPoseWatchPoseElement::UPoseWatchPoseElement(const FObjectInitializer& ObjectIni
 	bInvertViewportMask = false;
 	BlendScaleThreshold = 0.f;
 	ViewportOffset = FVector3d::ZeroVector;
+
+	Label = GetDefaultLabel();
+	SetIconName("ClassIcon.PoseAsset");
 #endif
 }
 
@@ -494,7 +497,6 @@ UPoseWatch::UPoseWatch(const FObjectInitializer& ObjectInitializer)
 {
 #if WITH_EDITOR
 	Label = GetDefaultLabel();
-	SetColor(PoseWatchUtil::ChoosePoseWatchColor());
 #endif
 }
 
@@ -559,7 +561,7 @@ bool UPoseWatch::GetIsEnabled() const
 
 FColor UPoseWatch::GetColor() const
 {
-	return Color;
+	return Color_DEPRECATED;
 }
 
 bool UPoseWatch::GetShouldDeleteOnDeselect() const
@@ -655,7 +657,7 @@ void UPoseWatch::SetIsVisible(bool bInIsVisible)
 
 void UPoseWatch::SetColor(const FColor& InColor)
 {
-	Color = InColor;
+	Color_DEPRECATED = InColor;
 
 	for (TObjectPtr<UPoseWatchElement> PoseWatchElement : Elements)
 	{
@@ -758,19 +760,15 @@ void UPoseWatch::Serialize(FArchive& Ar)
 
 	if (Ar.CustomVer(FFortniteMainBranchObjectVersion::GUID) < FFortniteMainBranchObjectVersion::PoseWatchMigrateSkeletonDrawParametersToPoseElement)
 	{
-		// Initialize all PoseWatchElement's Skeleton Draw Parameters from those defined in the parent Pose Watch.
-		for (TObjectPtr<UPoseWatchElement> PoseWatchElement : Elements)
-		{
-			UPoseWatchPoseElement* const PoseWatchPoseElement = Cast<UPoseWatchPoseElement>(PoseWatchElement.Get());
-
-			if (PoseWatchPoseElement)
-			{
-				PoseWatchPoseElement->ViewportMask = ViewportMask_DEPRECATED;
-				PoseWatchPoseElement->bInvertViewportMask = bInvertViewportMask_DEPRECATED;
-				PoseWatchPoseElement->BlendScaleThreshold = BlendScaleThreshold_DEPRECATED;
-				PoseWatchPoseElement->ViewportOffset = ViewportOffset_DEPRECATED;
-			}
-		}
+		// Transfer pose watch skeleton draw parameters to a new pose watch pose element.
+		TObjectPtr<UPoseWatchPoseElement> NewPoseWatchElement = NewObject<UPoseWatchPoseElement>(this);
+		NewPoseWatchElement->SetParent(this);
+		NewPoseWatchElement->SetColor(Color_DEPRECATED);
+		NewPoseWatchElement->ViewportMask = ViewportMask_DEPRECATED;
+		NewPoseWatchElement->bInvertViewportMask = bInvertViewportMask_DEPRECATED;
+		NewPoseWatchElement->BlendScaleThreshold = BlendScaleThreshold_DEPRECATED;
+		NewPoseWatchElement->ViewportOffset = ViewportOffset_DEPRECATED;
+		Elements.Add(NewPoseWatchElement);
 	}
 }
 #endif // WITH_EDITORONLY_DATA
