Fix all the places that were not using GRHIRayTracingScratchBufferAlignment or BUF_RayTracingScratch when allocating scratch buffer.

diff --git a/Engine/Plugins/Experimental/GPULightmass/Source/GPULightmass/Private/LightmapRenderer.cpp b/Engine/Plugins/Experimental/GPULightmass/Source/GPULightmass/Private/LightmapRenderer.cpp
--- a/Engine/Plugins/Experimental/GPULightmass/Source/GPULightmass/Private/LightmapRenderer.cpp
+++ b/Engine/Plugins/Experimental/GPULightmass/Source/GPULightmass/Private/LightmapRenderer.cpp
@@ -1105,7 +1105,7 @@ bool FSceneRenderState::SetupRayTracingScene(int32 LODIndex)
 			FRHIResourceCreateInfo ScratchBufferCreateInfo(TEXT("LightmassRayTracingScratchBuffer"));
 			FBufferRHIRef ScratchBuffer = RHICreateBuffer(
 				uint32(SizeInfo.BuildScratchSize),
-				BUF_UnorderedAccess, GRHIRayTracingScratchBufferAlignment, ERHIAccess::UAVCompute,
+				BUF_StructuredBuffer | BUF_RayTracingScratch, GRHIRayTracingScratchBufferAlignment, ERHIAccess::UAVCompute,
 				ScratchBufferCreateInfo);
 
 			FRWBufferStructured InstanceBuffer;
diff --git a/Engine/Source/Runtime/Engine/Private/RayTracingSkinnedGeometry.cpp b/Engine/Source/Runtime/Engine/Private/RayTracingSkinnedGeometry.cpp
--- a/Engine/Source/Runtime/Engine/Private/RayTracingSkinnedGeometry.cpp
+++ b/Engine/Source/Runtime/Engine/Private/RayTracingSkinnedGeometry.cpp
@@ -65,7 +65,7 @@ uint32 FRayTracingSkinnedGeometryUpdateQueue::ComputeScratchBufferSize() const
 {
 	TRACE_CPUPROFILER_EVENT_SCOPE(FRayTracingSkinnedGeometryUpdateQueue::ComputeScratchBufferSize);
 
-	const uint64 ScratchAlignment = GRHIRayTracingAccelerationStructureAlignment;
+	const uint64 ScratchAlignment = GRHIRayTracingScratchBufferAlignment;
 	uint32 ScratchBLASSize = 0;
 
 	if (ToUpdate.Num() && GRayTracingUseTransientForScratch > 0)
@@ -150,7 +150,7 @@ void FRayTracingSkinnedGeometryUpdateQueue::Commit(FRHICommandListImmediate & RH
 			PrimitivesToUpdates = 0;
 		};
 
-		const uint64 ScratchAlignment = GRHIRayTracingAccelerationStructureAlignment;
+		const uint64 ScratchAlignment = GRHIRayTracingScratchBufferAlignment;
 		uint32 ScratchBLASCurrentOffset = 0;
 		uint32 ScratchBLASNextOffset = 0;
 
@@ -211,7 +211,7 @@ void FRayTracingSkinnedGeometryUpdateQueue::Commit(FRDGBuilder& GraphBuilder)
 	const uint32 BLASScratchSize = ComputeScratchBufferSize();
 	if (BLASScratchSize > 0)
 	{
-		const uint32 ScratchAlignment = GRHIRayTracingAccelerationStructureAlignment;
+		const uint32 ScratchAlignment = GRHIRayTracingScratchBufferAlignment;
 		
 		FRDGBufferDesc ScratchBufferDesc;
 		ScratchBufferDesc.Usage = EBufferUsageFlags::RayTracingScratch | EBufferUsageFlags::StructuredBuffer;
diff --git a/Engine/Source/Runtime/Renderer/Private/DeferredShadingRenderer.cpp b/Engine/Source/Runtime/Renderer/Private/DeferredShadingRenderer.cpp
--- a/Engine/Source/Runtime/Renderer/Private/DeferredShadingRenderer.cpp
+++ b/Engine/Source/Runtime/Renderer/Private/DeferredShadingRenderer.cpp
@@ -1727,7 +1727,7 @@ bool FDeferredShadingSceneRenderer::DispatchRayTracingWorldUpdates(FRDGBuilder&
 	const uint32 BLASScratchSize = Scene->GetRayTracingDynamicGeometryCollection()->ComputeScratchBufferSize();
 	if (BLASScratchSize > 0)
 	{
-		const uint32 ScratchAlignment = GRHIRayTracingAccelerationStructureAlignment;
+		const uint32 ScratchAlignment = GRHIRayTracingScratchBufferAlignment;
 		FRDGBufferDesc ScratchBufferDesc;
 		ScratchBufferDesc.Usage = EBufferUsageFlags::RayTracingScratch | EBufferUsageFlags::StructuredBuffer;
 		ScratchBufferDesc.BytesPerElement = ScratchAlignment;
diff --git a/Engine/Source/Runtime/Renderer/Private/RayTracing/RayTracingDynamicGeometry.cpp b/Engine/Source/Runtime/Renderer/Private/RayTracing/RayTracingDynamicGeometry.cpp
--- a/Engine/Source/Runtime/Renderer/Private/RayTracing/RayTracingDynamicGeometry.cpp
+++ b/Engine/Source/Runtime/Renderer/Private/RayTracing/RayTracingDynamicGeometry.cpp
@@ -568,7 +568,7 @@ uint32 FRayTracingDynamicGeometryCollection::ComputeScratchBufferSize()
 {
 	TRACE_CPUPROFILER_EVENT_SCOPE(FRayTracingDynamicGeometryCollection::ComputeScratchBufferSize);
 
-	const uint64 ScratchAlignment = GRHIRayTracingAccelerationStructureAlignment;
+	const uint64 ScratchAlignment = GRHIRayTracingScratchBufferAlignment;
 
 	uint32 BLASScratchSize = 0;
 
diff --git a/Engine/Source/Runtime/VulkanRHI/Private/VulkanRayTracing.cpp b/Engine/Source/Runtime/VulkanRHI/Private/VulkanRayTracing.cpp
--- a/Engine/Source/Runtime/VulkanRHI/Private/VulkanRayTracing.cpp
+++ b/Engine/Source/Runtime/VulkanRHI/Private/VulkanRayTracing.cpp
@@ -461,7 +461,7 @@ void FVulkanRayTracingScene::BuildAccelerationStructure(
 	if (InScratchBuffer == nullptr)
 	{
 		FRHIResourceCreateInfo ScratchBufferCreateInfo(TEXT("BuildScratchTLAS"));
-		ScratchBuffer = ResourceCast(RHICreateBuffer(SizeInfo.BuildScratchSize, BUF_UnorderedAccess | BUF_StructuredBuffer, 0, ERHIAccess::UAVCompute, ScratchBufferCreateInfo).GetReference());
+		ScratchBuffer = ResourceCast(RHICreateBuffer(SizeInfo.BuildScratchSize, BUF_StructuredBuffer | BUF_RayTracingScratch, 0, ERHIAccess::UAVCompute, ScratchBufferCreateInfo).GetReference());
 		InScratchBuffer = ScratchBuffer.GetReference();
 		InScratchOffset = 0;
 	}
