Skeletal Mesh - Replace the newly created GetSkeletalMesh function by GetSkeletalMeshAsset so it matches the setter SetSkeletalMeshAsset function (which itself couldn't be named SetSkeletalMesh due to

diff --git a/Engine/Source/Runtime/Engine/Private/Components/SkeletalMeshComponent.cpp b/Engine/Source/Runtime/Engine/Private/Components/SkeletalMeshComponent.cpp
--- a/Engine/Source/Runtime/Engine/Private/Components/SkeletalMeshComponent.cpp
+++ b/Engine/Source/Runtime/Engine/Private/Components/SkeletalMeshComponent.cpp
@@ -182,7 +182,7 @@ public:
 		if (USkeletalMeshComponent* Comp = SkeletalMeshComponent.Get())
 		{
 			FScopeCycleCounterUObject ComponentScope(Comp);
-			FScopeCycleCounterUObject MeshScope(Comp->GetSkeletalMesh());
+			FScopeCycleCounterUObject MeshScope(Comp->GetSkeletalMeshAsset());
 
 			if(IsValidRef(Comp->ParallelAnimationEvaluationTask))
 			{
@@ -472,12 +472,12 @@ extern TAutoConsoleVariable<int32> CVarEnableClothPhysics;
 
 bool USkeletalMeshComponent::CanSimulateClothing() const
 {
-	if(!GetSkeletalMesh() || !bAllowClothActors || !CVarEnableClothPhysics.GetValueOnAnyThread())
+	if(!GetSkeletalMeshAsset() || !bAllowClothActors || !CVarEnableClothPhysics.GetValueOnAnyThread())
 	{
 		return false;
 	}
 
-	return GetSkeletalMesh()->HasActiveClothingAssets() && !IsNetMode(NM_DedicatedServer);
+	return GetSkeletalMeshAsset()->HasActiveClothingAssets() && !IsNetMode(NM_DedicatedServer);
 }
 
 void USkeletalMeshComponent::UpdateClothTickRegisteredState()
@@ -531,8 +531,8 @@ bool USkeletalMeshComponent::NeedToSpawnAnimScriptInstance() const
 {
 	IAnimClassInterface* AnimClassInterface = IAnimClassInterface::GetFromClass(AnimClass);
 	const USkeleton* AnimSkeleton = (AnimClassInterface) ? AnimClassInterface->GetTargetSkeleton() : nullptr;
-	const bool bSkeletonCompatible = (GetSkeletalMesh() && AnimSkeleton && GetSkeletalMesh()->GetSkeleton()) ? GetSkeletalMesh()->GetSkeleton()->IsCompatible(AnimSkeleton) : false;
-	const bool bSkelMeshCompatible = (GetSkeletalMesh() && AnimSkeleton) ? AnimSkeleton->IsCompatibleMesh(GetSkeletalMesh(), false) : false;
+	const bool bSkeletonCompatible = (GetSkeletalMeshAsset() && AnimSkeleton && GetSkeletalMeshAsset()->GetSkeleton()) ? GetSkeletalMeshAsset()->GetSkeleton()->IsCompatible(AnimSkeleton) : false;
+	const bool bSkelMeshCompatible = (GetSkeletalMeshAsset() && AnimSkeleton) ? AnimSkeleton->IsCompatibleMesh(GetSkeletalMeshAsset(), false) : false;
 	const bool bAnimSkelValid = !AnimClassInterface || (bSkeletonCompatible && bSkelMeshCompatible);
 
 	if (AnimationMode == EAnimationMode::AnimationBlueprint && AnimClass && bAnimSkelValid)
@@ -552,10 +552,10 @@ bool USkeletalMeshComponent::NeedToSpawnAnimScriptInstance() const
 
 bool USkeletalMeshComponent::NeedToSpawnPostPhysicsInstance(bool bForceReinit) const
 {
-	if(GetSkeletalMesh())
+	if(GetSkeletalMeshAsset())
 	{
 		const UClass* MainInstanceClass = *AnimClass;
-		const UClass* ClassToUse = *GetSkeletalMesh()->GetPostProcessAnimBlueprint();
+		const UClass* ClassToUse = *GetSkeletalMeshAsset()->GetPostProcessAnimBlueprint();
 		const UClass* CurrentClass = PostProcessAnimInstance ? PostProcessAnimInstance->GetClass() : nullptr;
 
 		// We need to have an instance, and we have the wrong class (different or null)
@@ -601,7 +601,7 @@ void USkeletalMeshComponent::OnRegister()
 		ClothingSimulationFactory = UClothingSimulationFactory::GetDefaultClothingSimulationFactoryClass();
 	}
 
-	USkeletalMesh* SkelMesh = GetSkeletalMesh();
+	USkeletalMesh* SkelMesh = GetSkeletalMeshAsset();
 	// Look up for the best simulation factory to support each asset
 	if (ClothingSimulationFactory && SkelMesh)
 	{
@@ -747,13 +747,13 @@ void USkeletalMeshComponent::InitAnim(bool bForceReinit)
 
 	// a lot of places just call InitAnim without checking Mesh, so 
 	// I'm moving the check here
-	if ( GetSkeletalMesh() != nullptr && IsRegistered() )
+	if ( GetSkeletalMeshAsset() != nullptr && IsRegistered() )
 	{
 		//clear cache UID since we don't know if skeleton changed
 		CachedAnimCurveUidVersion = 0;
 
 		// we still need this in case users doesn't call tick, but sent to renderer
-		MorphTargetWeights.SetNumZeroed(GetSkeletalMesh()->GetMorphTargets().Num());
+		MorphTargetWeights.SetNumZeroed(GetSkeletalMeshAsset()->GetMorphTargets().Num());
 
 		// We may be doing parallel evaluation on the current anim instance
 		// Calling this here with true will block this init till that thread completes
@@ -768,8 +768,8 @@ void USkeletalMeshComponent::InitAnim(bool bForceReinit)
 		const USkeleton* AnimSkeleton = (AnimScriptInstance)? AnimScriptInstance->CurrentSkeleton : nullptr;
 
 		const bool bClearAnimInstance = AnimScriptInstance && !AnimSkeleton;
-		const bool bSkeletonMismatch = AnimSkeleton && (AnimScriptInstance->CurrentSkeleton!=GetSkeletalMesh()->GetSkeleton());
-		const bool bSkeletonCompatible = AnimSkeleton && GetSkeletalMesh()->GetSkeleton() && !bSkeletonMismatch && GetSkeletalMesh()->GetSkeleton()->IsCompatible(AnimSkeleton);
+		const bool bSkeletonMismatch = AnimSkeleton && (AnimScriptInstance->CurrentSkeleton!=GetSkeletalMeshAsset()->GetSkeleton());
+		const bool bSkeletonCompatible = AnimSkeleton && GetSkeletalMeshAsset()->GetSkeleton() && !bSkeletonMismatch && GetSkeletalMeshAsset()->GetSkeleton()->IsCompatible(AnimSkeleton);
 
 		LastPoseTickFrame = 0;
 
@@ -804,9 +804,9 @@ void USkeletalMeshComponent::InitAnim(bool bForceReinit)
 				else
 				{
 					PRAGMA_DISABLE_DEPRECATION_WARNINGS
-					BoneSpaceTransforms = GetSkeletalMesh()->GetRefSkeleton().GetRefBonePose();
+					BoneSpaceTransforms = GetSkeletalMeshAsset()->GetRefSkeleton().GetRefBonePose();
 					//Mini RefreshBoneTransforms (the bit we actually care about)
-					GetSkeletalMesh()->FillComponentSpaceTransforms(BoneSpaceTransforms, FillComponentSpaceTransformsRequiredBones, GetEditableComponentSpaceTransforms());
+					GetSkeletalMeshAsset()->FillComponentSpaceTransforms(BoneSpaceTransforms, FillComponentSpaceTransformsRequiredBones, GetEditableComponentSpaceTransforms());
 					PRAGMA_ENABLE_DEPRECATION_WARNINGS
 					bNeedToFlipSpaceBaseBuffers = true; // Have updated space bases so need to flip
 					FlipEditableSpaceBases();
@@ -856,7 +856,7 @@ bool USkeletalMeshComponent::InitializeAnimScriptInstance(bool bForceReinit, boo
 
 	if (IsRegistered())
 	{
-		USkeletalMesh* SkelMesh = GetSkeletalMesh();
+		USkeletalMesh* SkelMesh = GetSkeletalMeshAsset();
 		check(SkelMesh);
 
 		if (NeedToSpawnAnimScriptInstance())
@@ -1057,7 +1057,7 @@ void USkeletalMeshComponent::PostEditChangeProperty(FPropertyChangedEvent& Prope
 			InitAnim(false);
 		}
 
-		USkeletalMesh* SkelMesh = GetSkeletalMesh();
+		USkeletalMesh* SkelMesh = GetSkeletalMeshAsset();
 		PRAGMA_DISABLE_DEPRECATION_WARNINGS
 		if(PropertyThatChanged->GetFName() == GET_MEMBER_NAME_CHECKED( USkeletalMeshComponent, SkeletalMeshAsset))
 		PRAGMA_ENABLE_DEPRECATION_WARNINGS
@@ -1172,9 +1172,9 @@ TSoftObjectPtr<UObject> USkeletalMeshComponent::GetDefaultAnimatingRig() const
 	{
 		return DefaultAnimatingRigOverride;
 	}
-	if (GetSkeletalMesh())
+	if (GetSkeletalMeshAsset())
 	{
-		return GetSkeletalMesh()->GetDefaultAnimatingRig();
+		return GetSkeletalMeshAsset()->GetDefaultAnimatingRig();
 	}
 	return nullptr;
 }
@@ -1216,7 +1216,7 @@ void USkeletalMeshComponent::TickAnimation(float DeltaTime, bool bNeedsValidRoot
 		RecalcRequiredCurves();
 	}
 
-	if (GetSkeletalMesh() != nullptr)
+	if (GetSkeletalMeshAsset() != nullptr)
 	{
 		// We're about to UpdateAnimation, this will potentially queue events that we'll need to dispatch.
 		bNeedsQueuedAnimEventsDispatched = true;
@@ -1291,9 +1291,9 @@ void USkeletalMeshComponent::UpdateVisualizeLODString(FString& DebugString)
 	Super::UpdateVisualizeLODString(DebugString);
 
 	uint32 NumVertices = 0;
-	if (GetSkeletalMesh())
+	if (GetSkeletalMeshAsset())
 	{
-		if (FSkeletalMeshRenderData* RenderData = GetSkeletalMesh()->GetResourceForRendering())
+		if (FSkeletalMeshRenderData* RenderData = GetSkeletalMeshAsset()->GetResourceForRendering())
 		{
 			if (RenderData->LODRenderData.IsValidIndex(GetPredictedLODLevel()))
 			{
@@ -1442,9 +1442,9 @@ void USkeletalMeshComponent::ResetMorphTargetCurves()
 {
 	ActiveMorphTargets.Reset();
 
-	if (GetSkeletalMesh())
+	if (GetSkeletalMeshAsset())
 	{
-		MorphTargetWeights.SetNum(GetSkeletalMesh()->GetMorphTargets().Num());
+		MorphTargetWeights.SetNum(GetSkeletalMeshAsset()->GetMorphTargets().Num());
 
 		// we need this code to ensure the buffer gets cleared whether or not you have morphtarget curve set
 		// the case, where you had morphtargets weight on, and when you clear the weight, you want to make sure 
@@ -1462,11 +1462,11 @@ void USkeletalMeshComponent::ResetMorphTargetCurves()
 
 void USkeletalMeshComponent::UpdateMorphTargetOverrideCurves()
 {
-	if (GetSkeletalMesh())
+	if (GetSkeletalMeshAsset())
 	{
 		if (MorphTargetCurves.Num() > 0)
 		{
-			FAnimationRuntime::AppendActiveMorphTargets(GetSkeletalMesh(), MorphTargetCurves, ActiveMorphTargets, MorphTargetWeights);
+			FAnimationRuntime::AppendActiveMorphTargets(GetSkeletalMeshAsset(), MorphTargetCurves, ActiveMorphTargets, MorphTargetWeights);
 		}
 	}
 }
@@ -1660,14 +1660,14 @@ static void MergeInBoneIndexArrays(TArray<FBoneIndexType>& BaseArray, const TArr
 // if you call RecalcRequiredBones, curve should be refreshed
 void USkeletalMeshComponent::RecalcRequiredCurves()
 {
-	if (!GetSkeletalMesh())
+	if (!GetSkeletalMeshAsset())
 	{
 		return;
 	}
 
-	if (GetSkeletalMesh()->GetSkeleton())
+	if (GetSkeletalMeshAsset()->GetSkeleton())
 	{
-		CachedCurveUIDList = GetSkeletalMesh()->GetSkeleton()->GetDefaultCurveUIDList();
+		CachedCurveUIDList = GetSkeletalMeshAsset()->GetSkeleton()->GetDefaultCurveUIDList();
 	}
 
 	const FCurveEvaluationOption CurveEvalOption(bAllowAnimCurveEvaluation, &DisallowedAnimCurves, GetPredictedLODLevel());
@@ -1696,7 +1696,7 @@ void USkeletalMeshComponent::ComputeRequiredBones(TArray<FBoneIndexType>& OutReq
 	OutRequiredBones.Reset();
 	OutFillComponentSpaceTransformsRequiredBones.Reset();
 
-	USkeletalMesh* SkelMesh = GetSkeletalMesh();
+	USkeletalMesh* SkelMesh = GetSkeletalMeshAsset();
 	if (!SkelMesh)
 	{
 		return;
@@ -1867,14 +1867,14 @@ void USkeletalMeshComponent::ComputeRequiredBones(TArray<FBoneIndexType>& OutReq
 
 void USkeletalMeshComponent::RecalcRequiredBones(int32 LODIndex)
 {
-	if (!GetSkeletalMesh())
+	if (!GetSkeletalMeshAsset())
 	{
 		return;
 	}
 
 	ComputeRequiredBones(RequiredBones, FillComponentSpaceTransformsRequiredBones, LODIndex, /*bIgnorePhysicsAsset=*/ false);
 	PRAGMA_DISABLE_DEPRECATION_WARNINGS
-	BoneSpaceTransforms = GetSkeletalMesh()->GetRefSkeleton().GetRefBonePose();
+	BoneSpaceTransforms = GetSkeletalMeshAsset()->GetRefSkeleton().GetRefBonePose();
 	PRAGMA_ENABLE_DEPRECATION_WARNINGS
 	// make sure animation requiredBone to mark as dirty
 	if (AnimScriptInstance)
@@ -1903,15 +1903,15 @@ void USkeletalMeshComponent::RecalcRequiredBones(int32 LODIndex)
 
 void USkeletalMeshComponent::MarkRequiredCurveUpToDate()
 {
-	if (GetSkeletalMesh() && GetSkeletalMesh()->GetSkeleton())
+	if (GetSkeletalMeshAsset() && GetSkeletalMeshAsset()->GetSkeleton())
 	{
-		CachedAnimCurveUidVersion = GetSkeletalMesh()->GetSkeleton()->GetAnimCurveUidVersion();
+		CachedAnimCurveUidVersion = GetSkeletalMeshAsset()->GetSkeleton()->GetAnimCurveUidVersion();
 	}
 }
 
 bool USkeletalMeshComponent::AreRequiredCurvesUpToDate() const
 {
-	return (!GetSkeletalMesh() || !GetSkeletalMesh()->GetSkeleton() || CachedAnimCurveUidVersion == GetSkeletalMesh()->GetSkeleton()->GetAnimCurveUidVersion());
+	return (!GetSkeletalMeshAsset() || !GetSkeletalMeshAsset()->GetSkeleton() || CachedAnimCurveUidVersion == GetSkeletalMeshAsset()->GetSkeleton()->GetAnimCurveUidVersion());
 }
 
 void USkeletalMeshComponent::EvaluateAnimation(const USkeletalMesh* InSkeletalMesh, UAnimInstance* InAnimInstance, FVector& OutRootBoneTranslation, FBlendedHeapCurve& OutCurve, FCompactPose& OutPose, UE::Anim::FHeapAttributeContainer& OutAttributes) const
@@ -1954,18 +1954,18 @@ void USkeletalMeshComponent::UpdateFollowerComponent()
 		// we changed order of morphtarget to be overriden by SetMorphTarget from BP
 		// so this has to go first
 		// now propagate BP-driven curves from the leader SMC...
-		if (GetSkeletalMesh())
+		if (GetSkeletalMeshAsset())
 		{
-			check(MorphTargetWeights.Num() == GetSkeletalMesh()->GetMorphTargets().Num());
+			check(MorphTargetWeights.Num() == GetSkeletalMeshAsset()->GetMorphTargets().Num());
 			if (LeaderSMC->MorphTargetCurves.Num() > 0)
 			{
-				FAnimationRuntime::AppendActiveMorphTargets(GetSkeletalMesh(), LeaderSMC->MorphTargetCurves, ActiveMorphTargets, MorphTargetWeights);
+				FAnimationRuntime::AppendActiveMorphTargets(GetSkeletalMeshAsset(), LeaderSMC->MorphTargetCurves, ActiveMorphTargets, MorphTargetWeights);
 			}
 
 			// if follower also has it, add it here. 
 			if (MorphTargetCurves.Num() > 0)
 			{
-				FAnimationRuntime::AppendActiveMorphTargets(GetSkeletalMesh(), MorphTargetCurves, ActiveMorphTargets, MorphTargetWeights);
+				FAnimationRuntime::AppendActiveMorphTargets(GetSkeletalMeshAsset(), MorphTargetCurves, ActiveMorphTargets, MorphTargetWeights);
 			}
 		}
 
@@ -2243,7 +2243,7 @@ void USkeletalMeshComponent::RefreshBoneTransforms(FActorComponentTickFunction*
 
 	check(IsInGameThread()); //Only want to call this from the game thread as we set up tasks etc
 	
-	if (!GetSkeletalMesh() || GetNumComponentSpaceTransforms() == 0)
+	if (!GetSkeletalMeshAsset() || GetNumComponentSpaceTransforms() == 0)
 	{
 		return;
 	}
@@ -2268,7 +2268,7 @@ void USkeletalMeshComponent::RefreshBoneTransforms(FActorComponentTickFunction*
 	//Dont mark cache as invalid if we aren't performing optimization anyway
 	PRAGMA_DISABLE_DEPRECATION_WARNINGS
 	const bool bInvalidCachedBones = bDoEvaluationRateOptimization &&
-									 ((BoneSpaceTransforms.Num() != GetSkeletalMesh()->GetRefSkeleton().GetNum())
+									 ((BoneSpaceTransforms.Num() != GetSkeletalMeshAsset()->GetRefSkeleton().GetNum())
 									 || (BoneSpaceTransforms.Num() != CachedBoneSpaceTransforms.Num())
 									 || (GetNumComponentSpaceTransforms() != CachedComponentSpaceTransforms.Num()));
 
@@ -2310,7 +2310,7 @@ void USkeletalMeshComponent::RefreshBoneTransforms(FActorComponentTickFunction*
 		return;
 	}
 
-	AnimEvaluationContext.SkeletalMesh = GetSkeletalMesh();
+	AnimEvaluationContext.SkeletalMesh = GetSkeletalMeshAsset();
 	AnimEvaluationContext.AnimInstance = AnimScriptInstance;
 	AnimEvaluationContext.PostProcessAnimInstance = (ShouldEvaluatePostProcessInstance())? ToRawPtr(PostProcessAnimInstance): nullptr;
 
@@ -2456,9 +2456,9 @@ void USkeletalMeshComponent::DispatchParallelEvaluationTasks(FActorComponentTick
 
 #if WITH_EDITOR
 	// We can only finish compilation on the game-thread, so wait here before spawning eval tasks.
-	if (GetSkeletalMesh() && GetSkeletalMesh()->IsCompiling())
+	if (GetSkeletalMeshAsset() && GetSkeletalMeshAsset()->IsCompiling())
 	{
-		FSkinnedAssetCompilingManager::Get().FinishCompilation({ GetSkeletalMesh() });
+		FSkinnedAssetCompilingManager::Get().FinishCompilation({ GetSkeletalMeshAsset() });
 	}
 #endif
 
@@ -2490,7 +2490,7 @@ void USkeletalMeshComponent::DispatchParallelTickPose(FActorComponentTickFunctio
 {
 	check(TickFunction);
 	
-	if(GetSkeletalMesh() != nullptr)
+	if(GetSkeletalMeshAsset() != nullptr)
 	{
 		if ((AnimScriptInstance && AnimScriptInstance->NeedsUpdate()) ||
 			(PostProcessAnimInstance && PostProcessAnimInstance->NeedsUpdate()))
@@ -2514,7 +2514,7 @@ void USkeletalMeshComponent::DispatchParallelTickPose(FActorComponentTickFunctio
 				}
 
 				// Do a mini-setup of the eval context
-				AnimEvaluationContext.SkeletalMesh = GetSkeletalMesh();
+				AnimEvaluationContext.SkeletalMesh = GetSkeletalMeshAsset();
 				AnimEvaluationContext.AnimInstance = AnimScriptInstance;
 
 				// We dont set up the Curve here, as we dont use it in Update()
@@ -2574,7 +2574,7 @@ void USkeletalMeshComponent::PostAnimEvaluation(FAnimationEvaluationContext& Eva
 		if (EvaluationContext.bDuplicateToCacheCurve)
 		{
 			ensureAlwaysMsgf(AnimCurves.IsValid(), TEXT("Animation Curve is invalid (%s). TotalCount(%d) "),
-				*GetPathNameSafe(GetSkeletalMesh()), AnimCurves.NumValidCurveCount);
+				*GetPathNameSafe(GetSkeletalMeshAsset()), AnimCurves.NumValidCurveCount);
 			CachedCurve.CopyFrom(AnimCurves);
 		}
 
@@ -2627,7 +2627,7 @@ void USkeletalMeshComponent::PostAnimEvaluation(FAnimationEvaluationContext& Eva
 
 			PRAGMA_DISABLE_DEPRECATION_WARNINGS
 			FAnimationRuntime::LerpBoneTransforms(BoneSpaceTransforms, CachedBoneSpaceTransforms, Alpha, RequiredBones);
-			GetSkeletalMesh()->FillComponentSpaceTransforms(BoneSpaceTransforms, FillComponentSpaceTransformsRequiredBones, GetEditableComponentSpaceTransforms());
+			GetSkeletalMeshAsset()->FillComponentSpaceTransforms(BoneSpaceTransforms, FillComponentSpaceTransformsRequiredBones, GetEditableComponentSpaceTransforms());
 			PRAGMA_ENABLE_DEPRECATION_WARNINGS
 			// interpolate curve
 			AnimCurves.LerpTo(CachedCurve, Alpha);
@@ -2756,10 +2756,10 @@ void USkeletalMeshComponent::ApplyAnimationCurvesToComponent(const TMap<FName, f
 	}
 
 	const bool bContainsMorphCurves = InAnimationMorphCurves && InAnimationMorphCurves->Num() > 0;
-	if (GetSkeletalMesh() && bContainsMorphCurves)
+	if (GetSkeletalMeshAsset() && bContainsMorphCurves)
 	{
 		// we want to append to existing curves - i.e. BP driven curves 
-		FAnimationRuntime::AppendActiveMorphTargets(GetSkeletalMesh(), *InAnimationMorphCurves, ActiveMorphTargets, MorphTargetWeights);
+		FAnimationRuntime::AppendActiveMorphTargets(GetSkeletalMeshAsset(), *InAnimationMorphCurves, ActiveMorphTargets, MorphTargetWeights);
 	}
 
 	/** Push through curves to follower components */
@@ -2861,7 +2861,7 @@ void USkeletalMeshComponent::SetSkeletalMesh(USkeletalMesh* InSkelMesh, bool bRe
 	QUICK_SCOPE_CYCLE_COUNTER(STAT_SetSkeletalMesh);
 	SCOPE_CYCLE_UOBJECT(NewSkelMesh, InSkelMesh);
 
-	if (InSkelMesh == GetSkeletalMesh())
+	if (InSkelMesh == GetSkeletalMeshAsset())
 	{
 		// do nothing if the input mesh is the same mesh we're already using.
 		return;
@@ -2945,9 +2945,9 @@ bool USkeletalMeshComponent::AllocateTransformData()
 	if ( Super::AllocateTransformData() )
 	{
 		PRAGMA_DISABLE_DEPRECATION_WARNINGS
-		if(BoneSpaceTransforms.Num() != GetSkeletalMesh()->GetRefSkeleton().GetNum() )
+		if(BoneSpaceTransforms.Num() != GetSkeletalMeshAsset()->GetRefSkeleton().GetNum() )
 		{
-			BoneSpaceTransforms = GetSkeletalMesh()->GetRefSkeleton().GetRefBonePose();
+			BoneSpaceTransforms = GetSkeletalMeshAsset()->GetRefSkeleton().GetRefBonePose();
 		}
 		PRAGMA_ENABLE_DEPRECATION_WARNINGS
 
@@ -3178,7 +3178,7 @@ void USkeletalMeshComponent::HideBone( int32 BoneIndex, EPhysBodyOp PhysBodyOpti
 {
 	Super::HideBone(BoneIndex, PhysBodyOption);
 
-	if (!GetSkeletalMesh())
+	if (!GetSkeletalMeshAsset())
 	{
 		return;
 	}
@@ -3195,7 +3195,7 @@ void USkeletalMeshComponent::HideBone( int32 BoneIndex, EPhysBodyOp PhysBodyOpti
 
 		if (PhysBodyOption != PBO_None)
 		{
-			FName HideBoneName = GetSkeletalMesh()->GetRefSkeleton().GetBoneName(BoneIndex);
+			FName HideBoneName = GetSkeletalMeshAsset()->GetRefSkeleton().GetBoneName(BoneIndex);
 			if (PhysBodyOption == PBO_Term)
 			{
 				TermBodiesBelow(HideBoneName);
@@ -3204,7 +3204,7 @@ void USkeletalMeshComponent::HideBone( int32 BoneIndex, EPhysBodyOp PhysBodyOpti
 	}
 	else
 	{
-		UE_LOG(LogSkeletalMesh, Warning, TEXT("HideBone[%s]: Invalid Body Index (%d) has entered. This component doesn't contain buffer for the given body."), *GetPathNameSafe(GetSkeletalMesh()), BoneIndex);
+		UE_LOG(LogSkeletalMesh, Warning, TEXT("HideBone[%s]: Invalid Body Index (%d) has entered. This component doesn't contain buffer for the given body."), *GetPathNameSafe(GetSkeletalMeshAsset()), BoneIndex);
 	}
 }
 
@@ -3212,7 +3212,7 @@ void USkeletalMeshComponent::UnHideBone( int32 BoneIndex )
 {
 	Super::UnHideBone(BoneIndex);
 
-	if (!GetSkeletalMesh())
+	if (!GetSkeletalMeshAsset())
 	{
 		return;
 	}
@@ -3226,7 +3226,7 @@ void USkeletalMeshComponent::UnHideBone( int32 BoneIndex )
 	{
 		bRequiredBonesUpToDate = false;
 
-		//FName HideBoneName = GetSkeletalMesh()->GetRefSkeleton().GetBoneName(BoneIndex);
+		//FName HideBoneName = GetSkeletalMeshAsset()->GetRefSkeleton().GetBoneName(BoneIndex);
 		// It's okay to turn this on for terminated bodies
 		// It won't do any if BodyData isn't found
 		// @JTODO
@@ -3234,7 +3234,7 @@ void USkeletalMeshComponent::UnHideBone( int32 BoneIndex )
 	}
 	else
 	{
-		UE_LOG(LogSkeletalMesh, Warning, TEXT("UnHideBone[%s]: Invalid Body Index (%d) has entered. This component doesn't contain buffer for the given body."), *GetPathNameSafe(GetSkeletalMesh()), BoneIndex);
+		UE_LOG(LogSkeletalMesh, Warning, TEXT("UnHideBone[%s]: Invalid Body Index (%d) has entered. This component doesn't contain buffer for the given body."), *GetPathNameSafe(GetSkeletalMeshAsset()), BoneIndex);
 	}
 }
 
@@ -3349,7 +3349,7 @@ void USkeletalMeshComponent::SetAnimationMode(EAnimationMode::Type InAnimationMo
 	// when mode is swapped, make sure to reinitialize
 	// even if it was same mode, this was due to users who wants to use BP construction script to do this
 	// if you use it in the construction script, it gets serialized, but it never instantiate. 
-	if(GetSkeletalMesh() != nullptr && (bNeedChange || AnimationMode == EAnimationMode::AnimationBlueprint))
+	if(GetSkeletalMeshAsset() != nullptr && (bNeedChange || AnimationMode == EAnimationMode::AnimationBlueprint))
 	{
 		if (InitializeAnimScriptInstance(true))
 		{
@@ -3729,7 +3729,7 @@ void USkeletalMeshComponent::UnregisterOnSkeletalMeshPropertyChanged( FDelegateH
 
 void USkeletalMeshComponent::ValidateAnimation()
 {
-	USkeletalMesh* SkelMesh = GetSkeletalMesh();
+	USkeletalMesh* SkelMesh = GetSkeletalMeshAsset();
 	if (SkelMesh && SkelMesh->GetSkeleton() == nullptr)
 	{
 		UE_LOG(LogAnimation, Warning, TEXT("SkeletalMesh %s has no skeleton. This needs to fixed before an animation can be set"), *SkelMesh->GetFullName());
@@ -3852,7 +3852,7 @@ void USkeletalMeshComponent::RefreshMorphTargets()
 {
 	ResetMorphTargetCurves();
 
-	if (GetSkeletalMesh() && AnimScriptInstance)
+	if (GetSkeletalMeshAsset() && AnimScriptInstance)
 	{
 		// as this can be called from any worker thread (i.e. from CreateRenderState_Concurrent) we cant currently be doing parallel evaluation
 		check(!IsRunningParallelEvaluation());
@@ -3912,7 +3912,7 @@ void USkeletalMeshComponent::ParallelDuplicateAndInterpolate(FAnimationEvaluatio
 		if (InAnimEvaluationContext.bDuplicateToCacheCurve)
 		{
 			ensureAlwaysMsgf(InAnimEvaluationContext.Curve.IsValid(), TEXT("Animation Curve is invalid (%s). TotalCount(%d) "),
-				*GetPathNameSafe(GetSkeletalMesh()), InAnimEvaluationContext.Curve.NumValidCurveCount);
+				*GetPathNameSafe(GetSkeletalMeshAsset()), InAnimEvaluationContext.Curve.NumValidCurveCount);
 			InAnimEvaluationContext.CachedCurve.CopyFrom(InAnimEvaluationContext.Curve);
 		}
 
@@ -3979,7 +3979,7 @@ void USkeletalMeshComponent::CompleteParallelAnimationEvaluation(bool bDoPostAni
 	SCOPED_NAMED_EVENT(USkeletalMeshComponent_CompleteParallelAnimationEvaluation, FColor::Yellow);
 	ParallelAnimationEvaluationTask.SafeRelease(); //We are done with this task now, clean up!
 
-	if (bDoPostAnimEvaluation && (AnimEvaluationContext.AnimInstance == AnimScriptInstance) && (AnimEvaluationContext.SkeletalMesh == GetSkeletalMesh()) && (AnimEvaluationContext.ComponentSpaceTransforms.Num() == GetNumComponentSpaceTransforms()))
+	if (bDoPostAnimEvaluation && (AnimEvaluationContext.AnimInstance == AnimScriptInstance) && (AnimEvaluationContext.SkeletalMesh == GetSkeletalMeshAsset()) && (AnimEvaluationContext.ComponentSpaceTransforms.Num() == GetNumComponentSpaceTransforms()))
 	{
 		SwapEvaluationContextBuffers();
 
@@ -4030,7 +4030,7 @@ void USkeletalMeshComponent::BindClothToLeaderPoseComponent()
 {
 	if(USkeletalMeshComponent* LeaderComp = Cast<USkeletalMeshComponent>(LeaderPoseComponent.Get()))
 	{
-		if(GetSkeletalMesh() != LeaderComp->GetSkeletalMesh())
+		if(GetSkeletalMeshAsset() != LeaderComp->GetSkeletalMeshAsset())
 		{
 			// Not the same mesh, can't bind
 			return;
@@ -4138,7 +4138,7 @@ void USkeletalMeshComponent::FinalizeBoneTransform()
 
 void USkeletalMeshComponent::GetCurrentRefToLocalMatrices(TArray<FMatrix44f>& OutRefToLocals, int32 InLodIdx) const
 {
-	USkeletalMesh* SkelMesh = GetSkeletalMesh();
+	USkeletalMesh* SkelMesh = GetSkeletalMeshAsset();
 	if(SkelMesh)
 	{
 		FSkeletalMeshRenderData* RenderData = SkelMesh->GetResourceForRendering();
@@ -4273,7 +4273,7 @@ void USkeletalMeshComponent::RemoveFollowerPoseComponent(USkinnedMeshComponent*
 
 void USkeletalMeshComponent::SnapshotPose(FPoseSnapshot& Snapshot)
 {
-	USkeletalMesh* SkelMesh = GetSkeletalMesh();
+	USkeletalMesh* SkelMesh = GetSkeletalMeshAsset();
 	if (ensureAsRuntimeWarning(SkelMesh != nullptr))
 	{
 		const TArray<FTransform>& ComponentSpaceTMs = GetComponentSpaceTransforms();
@@ -4428,9 +4428,9 @@ void USkeletalMeshComponent::SetAllowedAnimCurvesEvaluation(const TArray<FName>&
 			const TArray<FName>& AllowedList;
 		};
 
-		if (GetSkeletalMesh())
+		if (GetSkeletalMeshAsset())
 		{
-			USkeleton* Skeleton = GetSkeletalMesh()->GetSkeleton();
+			USkeleton* Skeleton = GetSkeletalMeshAsset()->GetSkeleton();
 			if (Skeleton)
 			{
 				const FSmartNameMapping* Mapping = Skeleton->GetSmartNameContainer(USkeleton::AnimCurveMappingName);
@@ -4505,7 +4505,7 @@ bool USkeletalMeshComponent::FindAttributeChecked(const FName& BoneName, const F
 	OutValue = DefaultValue;
 	bool bFound = false;
 
-	USkeletalMesh* SkelMesh = GetSkeletalMesh();
+	USkeletalMesh* SkelMesh = GetSkeletalMeshAsset();
 	if (SkelMesh)
 	{
 		const UE::Anim::FMeshAttributeContainer& Attributes = GetCustomAttributes();
diff --git a/Engine/Source/Runtime/Engine/Private/LocalPlayer.cpp b/Engine/Source/Runtime/Engine/Private/LocalPlayer.cpp
--- a/Engine/Source/Runtime/Engine/Private/LocalPlayer.cpp
+++ b/Engine/Source/Runtime/Engine/Private/LocalPlayer.cpp
@@ -1259,7 +1259,7 @@ bool ULocalPlayer::HandleListSkelMeshesCommand( const TCHAR* Cmd, FOutputDevice&
 	for( TObjectIterator<USkeletalMeshComponent> It; It; ++It )
 	{
 		USkeletalMeshComponent* SkeletalMeshComponent = *It;
-		USkeletalMesh* SkeletalMesh = SkeletalMeshComponent->GetSkeletalMesh();
+		USkeletalMesh* SkeletalMesh = SkeletalMeshComponent->GetSkeletalMeshAsset();
 
 		if( !SkeletalMeshComponent->IsTemplate() )
 		{
diff --git a/Engine/Source/Runtime/Engine/Private/Particles/ParticleModules_Location.cpp b/Engine/Source/Runtime/Engine/Private/Particles/ParticleModules_Location.cpp
--- a/Engine/Source/Runtime/Engine/Private/Particles/ParticleModules_Location.cpp
+++ b/Engine/Source/Runtime/Engine/Private/Particles/ParticleModules_Location.cpp
@@ -1467,7 +1467,7 @@ void UParticleModuleLocationBoneSocket::RegeneratePreSelectedIndices(FModuleLoca
 {
 	if (SourceIndexMode == EBoneSocketSourceIndexMode::PreSelectedIndices)
 	{
-		int32 MaxIndex = SourceType == BONESOCKETSOURCE_Sockets ? SourceComponent->GetSkeletalMesh()->NumSockets() : SourceComponent->GetNumBones();
+		int32 MaxIndex = SourceType == BONESOCKETSOURCE_Sockets ? SourceComponent->GetSkeletalMeshAsset()->NumSockets() : SourceComponent->GetNumBones();
 		for (int32 i = 0; i < NumPreSelectedIndices; ++i)
 		{
 			//Should we provide sequential selection here? Does that make sense for the pre selected list?
@@ -1708,7 +1708,7 @@ void UParticleModuleLocationBoneSocket::FinalUpdate(FParticleEmitterInstance* Ow
 			BEGIN_UPDATE_LOOP;
 			{
 				FModuleLocationBoneSocketParticlePayload* ParticlePayload = (FModuleLocationBoneSocketParticlePayload*)((uint8*)&Particle + Offset);
-				if (SourceComponent && SourceComponent->GetSkeletalMesh())
+				if (SourceComponent && SourceComponent->GetSkeletalMeshAsset())
 				{
 					USkeletalMeshSocket* Socket;
 					FVector SocketOffset;
@@ -1916,9 +1916,9 @@ int32 UParticleModuleLocationBoneSocket::GetMaxSourceIndex(FModuleLocationBoneSo
 		{
 			if (SourceType == BONESOCKETSOURCE_Sockets)
 			{
-				if (ensure(SourceComponent->GetSkeletalMesh()))
+				if (ensure(SourceComponent->GetSkeletalMeshAsset()))
 				{
-					return SourceComponent->GetSkeletalMesh()->NumSockets();
+					return SourceComponent->GetSkeletalMeshAsset()->NumSockets();
 				}
 				return 0;
 			}
@@ -1997,7 +1997,7 @@ bool UParticleModuleLocationBoneSocket::GetSocketInfoForSourceIndex(FModuleLocat
 {
 	if (!ensureMsgf(SourceType == BONESOCKETSOURCE_Sockets, TEXT("Invalid source type %d for %s"), SourceType, *GetPathName()) ||
 		!ensureMsgf(SourceComponent, TEXT("Null SkeletalMeshComponent for %s"), *GetPathName()) ||
-		!ensureMsgf(SourceComponent->GetSkeletalMesh(), TEXT("Null SkeletalMesh on Component %s for %s"), *SourceComponent->GetPathName(), *GetPathName()))
+		!ensureMsgf(SourceComponent->GetSkeletalMeshAsset(), TEXT("Null SkeletalMesh on Component %s for %s"), *SourceComponent->GetPathName(), *GetPathName()))
 	{
 		return false;
 	}
@@ -2008,7 +2008,7 @@ bool UParticleModuleLocationBoneSocket::GetSocketInfoForSourceIndex(FModuleLocat
 		{
 			if (ensureMsgf(SourceLocations.IsValidIndex(SourceIndex), TEXT("Invalid index of %d for %s"), SourceIndex, *GetPathName()))
 			{
-				OutSocket = SourceComponent->GetSkeletalMesh()->FindSocket(SourceLocations[SourceIndex].BoneSocketName);
+				OutSocket = SourceComponent->GetSkeletalMeshAsset()->FindSocket(SourceLocations[SourceIndex].BoneSocketName);
 				OutOffset = SourceLocations[SourceIndex].Offset + UniversalOffset;
 			}
 			else
@@ -2022,7 +2022,7 @@ bool UParticleModuleLocationBoneSocket::GetSocketInfoForSourceIndex(FModuleLocat
 			if (ensureMsgf(InstancePayload, TEXT("Invalid instance payload parameter on GetSocketInfoForSourceIndex %d for %s"), SourceIndex, *GetPathName()) &&
 				ensureMsgf(SourceIndex >= 0 && SourceIndex < InstancePayload->PreSelectedBoneSocketIndices.Num(), TEXT("Invalid index of %d for %s"), SourceIndex, *GetPathName()))
 			{
-				OutSocket = SourceComponent->GetSkeletalMesh()->GetSocketByIndex(InstancePayload->PreSelectedBoneSocketIndices[SourceIndex]);
+				OutSocket = SourceComponent->GetSkeletalMeshAsset()->GetSocketByIndex(InstancePayload->PreSelectedBoneSocketIndices[SourceIndex]);
 				OutOffset = UniversalOffset;
 			} 
 			else
@@ -2034,7 +2034,7 @@ bool UParticleModuleLocationBoneSocket::GetSocketInfoForSourceIndex(FModuleLocat
 		break;
 		case EBoneSocketSourceIndexMode::Direct:
 		{
-			OutSocket = SourceComponent->GetSkeletalMesh()->GetSocketByIndex(SourceIndex);
+			OutSocket = SourceComponent->GetSkeletalMeshAsset()->GetSocketByIndex(SourceIndex);
 			OutOffset = UniversalOffset;
 		}
 		break;
@@ -2098,7 +2098,7 @@ bool UParticleModuleLocationBoneSocket::GetParticleLocation(FModuleLocationBoneS
 
 	if (SourceType == BONESOCKETSOURCE_Sockets)
 	{
-		if (InSkelMeshComponent->GetSkeletalMesh())
+		if (InSkelMeshComponent->GetSkeletalMeshAsset())
 		{
 			USkeletalMeshSocket* Socket;
 			FVector SocketOffset;
@@ -2583,7 +2583,7 @@ void UParticleModuleLocationSkelVertSurface::UpdateBoneIndicesList(FParticleEmit
 		int32 InsertionIndex = 0;
 		for (int32 FindBoneIdx = 0; FindBoneIdx < ValidAssociatedBones.Num(); FindBoneIdx++)
 		{
-			const int32 BoneIdx = SkelMeshComp->GetSkeletalMesh()->GetRefSkeleton().FindBoneIndex(ValidAssociatedBones[FindBoneIdx]);
+			const int32 BoneIdx = SkelMeshComp->GetSkeletalMeshAsset()->GetRefSkeleton().FindBoneIndex(ValidAssociatedBones[FindBoneIdx]);
 			if (BoneIdx != INDEX_NONE && ValidAssociatedBones.Num() > InsertionIndex)
 			{
 				InstancePayload->ValidAssociatedBoneIndices[InsertionIndex++] = BoneIdx;
@@ -2747,10 +2747,10 @@ void UParticleModuleLocationSkelVertSurface::GetSkeletalMeshComponentSource(FPar
 
 		bool bMeshIsValid = false;
 		int32 MinLOD = INDEX_NONE;
-		if (NewSkelMeshComp && NewSkelMeshComp->GetScene() && NewSkelMeshComp->GetSkeletalMesh())
+		if (NewSkelMeshComp && NewSkelMeshComp->GetScene() && NewSkelMeshComp->GetSkeletalMeshAsset())
 		{
 			FSkeletalMeshRenderData* SkelMeshResource = NewSkelMeshComp->GetSkeletalMeshRenderData();
-			MinLOD = SkelMeshResource->GetFirstValidLODIdx(NewSkelMeshComp->GetSkeletalMesh()->GetMinLodIdx());
+			MinLOD = SkelMeshResource->GetFirstValidLODIdx(NewSkelMeshComp->GetSkeletalMeshAsset()->GetMinLodIdx());
 
 			if (MinLOD != INDEX_NONE)
 			{
@@ -2786,7 +2786,7 @@ void UParticleModuleLocationSkelVertSurface::GetSkeletalMeshComponentSource(FPar
 			UE_LOG(LogParticles, Warning, TEXT("Attempting to use Cascade SkelVertSurface module on mesh without valid data."));
 			UE_LOG(LogParticles, Warning, TEXT("Likely due to CPU side buffers being stripped with r.FreeSkeletalMeshBuffers=1."));
 			UE_LOG(LogParticles, Warning, TEXT("This emitter should probably be culled for this platform using detail mode."));
-			UE_LOG(LogParticles, Warning, TEXT("Mesh: %s"), NewSkelMeshComp ? *NewSkelMeshComp->GetSkeletalMesh()->GetFullName() : TEXT("NULL"));
+			UE_LOG(LogParticles, Warning, TEXT("Mesh: %s"), NewSkelMeshComp ? *NewSkelMeshComp->GetSkeletalMeshAsset()->GetFullName() : TEXT("NULL"));
 			UE_LOG(LogParticles, Warning, TEXT("Comp: %s"), *Owner->Component->GetFullName());
 			UE_LOG(LogParticles, Warning, TEXT("System: %s"), *Owner->Component->Template->GetFullName());
 			UE_LOG(LogParticles, Warning, TEXT("----------------------------------------------------------------------"));
@@ -2964,7 +2964,7 @@ bool UParticleModuleLocationSkelVertSurface::VertInfluencedByActiveBoneTyped(
 			int32 BoneIndex = Section.BoneMap[WeightBuffer->GetBoneIndex(Section.GetVertexBufferIndex() + VertIndex, InfluenceIndex)];
 			if (InSkelMeshComponent->LeaderPoseComponent.IsValid())
 			{
-				check(LeaderBoneMap.Num() == InSkelMeshComponent->GetSkeletalMesh()->GetRefSkeleton().GetNum());
+				check(LeaderBoneMap.Num() == InSkelMeshComponent->GetSkeletalMeshAsset()->GetRefSkeleton().GetNum());
 				BoneIndex = LeaderBoneMap[BoneIndex];
 			}
 
diff --git a/Engine/Source/Runtime/Engine/Private/PhysicsEngine/PhysAnim.cpp b/Engine/Source/Runtime/Engine/Private/PhysicsEngine/PhysAnim.cpp
--- a/Engine/Source/Runtime/Engine/Private/PhysicsEngine/PhysAnim.cpp
+++ b/Engine/Source/Runtime/Engine/Private/PhysicsEngine/PhysAnim.cpp
@@ -132,7 +132,7 @@ void UpdateWorldBoneTM(TAssetWorldBoneTMArray& WorldBoneTMs, const TArray<FTrans
 	else
 	{
 		// If not root, use our cached world-space bone transforms.
-		int32 ParentIndex = SkelComp->GetSkeletalMesh()->GetRefSkeleton().GetParentIndex(BoneIndex);
+		int32 ParentIndex = SkelComp->GetSkeletalMeshAsset()->GetRefSkeleton().GetParentIndex(BoneIndex);
 		UpdateWorldBoneTM(WorldBoneTMs, InBoneSpaceTransforms, ParentIndex, SkelComp, LocalToWorldTM, Scale3D);
 		ParentTM = WorldBoneTMs[ParentIndex].TM;
 	}
@@ -201,7 +201,7 @@ void USkeletalMeshComponent::PerformBlendPhysicsBones(const TArray<FBoneIndexTyp
 			int32 BoneIndex = InRequiredBones[i];
 
 			// See if this is a physics bone..
-			int32 BodyIndex = PhysicsAsset->FindBodyIndex(GetSkeletalMesh()->GetRefSkeleton().GetBoneName(BoneIndex));
+			int32 BodyIndex = PhysicsAsset->FindBodyIndex(GetSkeletalMeshAsset()->GetRefSkeleton().GetBoneName(BoneIndex));
 			// need to update back to physX so that physX knows where it was after blending
 			FBodyInstance* PhysicsAssetBodyInstance = nullptr;
 
@@ -217,7 +217,7 @@ void USkeletalMeshComponent::PerformBlendPhysicsBones(const TArray<FBoneIndexTyp
 				if ( !ensure(Bodies.IsValidIndex(BodyIndex)) )
 				{
 					UE_LOG(LogPhysics, Warning, TEXT("%s(Mesh %s, PhysicsAsset %s)"), 
-						*GetName(), *GetNameSafe(GetSkeletalMesh()), *GetNameSafe(PhysicsAsset));
+						*GetName(), *GetNameSafe(GetSkeletalMeshAsset()), *GetNameSafe(PhysicsAsset));
 					UE_LOG(LogPhysics, Warning, TEXT(" - # of BodySetup (%d), # of Bodies (%d), Invalid BodyIndex(%d)"), 
 						PhysicsAsset->SkeletalBodySetups.Num(), Bodies.Num(), BodyIndex);
 					continue;
@@ -264,7 +264,7 @@ void USkeletalMeshComponent::PerformBlendPhysicsBones(const TArray<FBoneIndexTyp
 						else
 						{
 							// If not root, get parent TM from cache (making sure its up-to-date).
-							int32 ParentIndex = GetSkeletalMesh()->GetRefSkeleton().GetParentIndex(BoneIndex);
+							int32 ParentIndex = GetSkeletalMeshAsset()->GetRefSkeleton().GetParentIndex(BoneIndex);
 							UpdateWorldBoneTM(WorldBoneTMs, InOutBoneSpaceTransforms, ParentIndex, this, LocalToWorldTM, TotalScale3D);
 							ParentWorldTM = WorldBoneTMs[ParentIndex].TM;
 						}
@@ -314,7 +314,7 @@ void USkeletalMeshComponent::PerformBlendPhysicsBones(const TArray<FBoneIndexTyp
 					{
 						continue;
 					}
-					const int32 ParentIndex = GetSkeletalMesh()->GetRefSkeleton().GetParentIndex(BoneIndex);
+					const int32 ParentIndex = GetSkeletalMeshAsset()->GetRefSkeleton().GetParentIndex(BoneIndex);
 					InOutComponentSpaceTransforms[BoneIndex] = InOutBoneSpaceTransforms[BoneIndex] * InOutComponentSpaceTransforms[ParentIndex];
 
 					/**
@@ -365,7 +365,7 @@ void USkeletalMeshComponent::BlendInPhysicsInternal(FTickFunction& ThisTickFunct
 	check(IsInGameThread());
 
 	// Can't do anything without a SkeletalMesh
-	if( !GetSkeletalMesh())
+	if( !GetSkeletalMeshAsset())
 	{
 		return;
 	}
@@ -535,13 +535,13 @@ void USkeletalMeshComponent::UpdateKinematicBonesToAnim(const TArray<FTransform>
 
 #if !(UE_BUILD_SHIPPING || UE_BUILD_TEST)
 	// If desired, draw the skeleton at the point where we pass it to the physics.
-	if (bShowPrePhysBones && GetSkeletalMesh() && InSpaceBases.Num() == GetSkeletalMesh()->GetRefSkeleton().GetNum())
+	if (bShowPrePhysBones && GetSkeletalMeshAsset() && InSpaceBases.Num() == GetSkeletalMeshAsset()->GetRefSkeleton().GetNum())
 	{
 		for (int32 i = 1; i<InSpaceBases.Num(); i++)
 		{
 			FVector ThisPos = CurrentLocalToWorld.TransformPosition(InSpaceBases[i].GetLocation());
 
-			int32 ParentIndex = GetSkeletalMesh()->GetRefSkeleton().GetParentIndex(i);
+			int32 ParentIndex = GetSkeletalMeshAsset()->GetRefSkeleton().GetParentIndex(i);
 			FVector ParentPos = CurrentLocalToWorld.TransformPosition(InSpaceBases[ParentIndex].GetLocation());
 
 			World->LineBatcher->DrawLine(ThisPos, ParentPos, AnimSkelDrawColor, SDPG_Foreground);
@@ -554,7 +554,7 @@ void USkeletalMeshComponent::UpdateKinematicBonesToAnim(const TArray<FTransform>
 #if !(UE_BUILD_SHIPPING || UE_BUILD_TEST)
 	if( !MeshScale3D.IsUniform() )
 	{
-		UE_LOG(LogPhysics, Log, TEXT("USkeletalMeshComponent::UpdateKinematicBonesToAnim : Non-uniform scale factor (%s) can cause physics to mismatch for %s  SkelMesh: %s"), *MeshScale3D.ToString(), *GetFullName(), GetSkeletalMesh() ? *GetSkeletalMesh()->GetFullName() : TEXT("NULL"));
+		UE_LOG(LogPhysics, Log, TEXT("USkeletalMeshComponent::UpdateKinematicBonesToAnim : Non-uniform scale factor (%s) can cause physics to mismatch for %s  SkelMesh: %s"), *MeshScale3D.ToString(), *GetFullName(), GetSkeletalMeshAsset() ? *GetSkeletalMeshAsset()->GetFullName() : TEXT("NULL"));
 	}
 #endif
 
@@ -562,11 +562,11 @@ void USkeletalMeshComponent::UpdateKinematicBonesToAnim(const TArray<FTransform>
 	if (bEnablePerPolyCollision == false)
 	{
 		const UPhysicsAsset* const PhysicsAsset = GetPhysicsAsset();
-		if (PhysicsAsset && GetSkeletalMesh() && Bodies.Num() > 0)
+		if (PhysicsAsset && GetSkeletalMeshAsset() && Bodies.Num() > 0)
 		{
 #if !(UE_BUILD_SHIPPING || UE_BUILD_TEST)
 			if (!ensureMsgf(PhysicsAsset->SkeletalBodySetups.Num() == Bodies.Num(), TEXT("Mesh (%s) has PhysicsAsset(%s), and BodySetup(%d) and Bodies(%d) don't match"),
-						*GetSkeletalMesh()->GetName(), *PhysicsAsset->GetName(), PhysicsAsset->SkeletalBodySetups.Num(), Bodies.Num()))
+						*GetSkeletalMeshAsset()->GetName(), *PhysicsAsset->GetName(), PhysicsAsset->SkeletalBodySetups.Num(), Bodies.Num()))
 			{
 				return;
 			}
@@ -601,7 +601,7 @@ void USkeletalMeshComponent::UpdateKinematicBonesToAnim(const TArray<FTransform>
 							{
 								BodyName = PhysicsAsset->SkeletalBodySetups[i]->BoneName;
 							}
-							UE_LOG(LogPhysics, Log, TEXT("UpdateRBBones: WARNING: Failed to find bone '%s' need by PhysicsAsset '%s' in SkeletalMesh '%s'."), *BodyName.ToString(), *PhysicsAsset->GetName(), *GetSkeletalMesh()->GetName());
+							UE_LOG(LogPhysics, Log, TEXT("UpdateRBBones: WARNING: Failed to find bone '%s' need by PhysicsAsset '%s' in SkeletalMesh '%s'."), *BodyName.ToString(), *PhysicsAsset->GetName(), *GetSkeletalMeshAsset()->GetName());
 						}
 						else
 						{
@@ -612,7 +612,7 @@ void USkeletalMeshComponent::UpdateKinematicBonesToAnim(const TArray<FTransform>
 								{
 									BodyName = PhysicsAsset->SkeletalBodySetups[i]->BoneName;
 								}
-								UE_LOG(LogPhysics, Warning, TEXT("BoneIndex %d out of range of SpaceBases (Size %d) on PhysicsAsset '%s' in SkeletalMesh '%s' for bone '%s'"), BoneIndex, InSpaceBases.Num(), *PhysicsAsset->GetName(), *GetSkeletalMesh()->GetName(), *BodyName.ToString());
+								UE_LOG(LogPhysics, Warning, TEXT("BoneIndex %d out of range of SpaceBases (Size %d) on PhysicsAsset '%s' in SkeletalMesh '%s' for bone '%s'"), BoneIndex, InSpaceBases.Num(), *PhysicsAsset->GetName(), *GetSkeletalMeshAsset()->GetName(), *BodyName.ToString());
 								continue;
 							}
 
@@ -626,7 +626,7 @@ void USkeletalMeshComponent::UpdateKinematicBonesToAnim(const TArray<FTransform>
 									BodyName = PhysicsAsset->SkeletalBodySetups[i]->BoneName;
 								}
 
-								UE_LOG(LogPhysics, Warning, TEXT("UpdateKinematicBonesToAnim: Trying to set transform with bad data %s on PhysicsAsset '%s' in SkeletalMesh '%s' for bone '%s'"), *BoneTransform.ToHumanReadableString(), *PhysicsAsset->GetName(), *GetSkeletalMesh()->GetName(), *BodyName.ToString());
+								UE_LOG(LogPhysics, Warning, TEXT("UpdateKinematicBonesToAnim: Trying to set transform with bad data %s on PhysicsAsset '%s' in SkeletalMesh '%s' for bone '%s'"), *BoneTransform.ToHumanReadableString(), *PhysicsAsset->GetName(), *GetSkeletalMeshAsset()->GetName(), *BodyName.ToString());
 								BoneTransform.DiagnosticCheck_IsValid();	//In special nan mode we want to actually ensure
 
 								continue;
@@ -673,7 +673,7 @@ void USkeletalMeshComponent::UpdateKinematicBonesToAnim(const TArray<FTransform>
 							//It's not clear whether this should be a warning. There are certainly cases where you interpolate the blend weight towards 0. The blend feature needs some work which will probably change this in the future.
 							//Making it Verbose for now
 							UE_LOG(LogPhysics, Verbose, TEXT("%s(Mesh %s, PhysicsAsset %s, Bone %s) is simulating, but no blending. "),
-								*GetName(), *GetNameSafe(GetSkeletalMesh()), *GetNameSafe(PhysicsAsset), *BodyInst->BodySetup.Get()->BoneName.ToString());
+								*GetName(), *GetNameSafe(GetSkeletalMeshAsset()), *GetNameSafe(PhysicsAsset), *BodyInst->BodySetup.Get()->BoneName.ToString());
 						}
 					}
 				}
@@ -732,7 +732,7 @@ void USkeletalMeshComponent::UpdateRBJointMotors()
 	}
 
 	const UPhysicsAsset* const PhysicsAsset = GetPhysicsAsset();
-	if(PhysicsAsset && Constraints.Num() > 0 && GetSkeletalMesh())
+	if(PhysicsAsset && Constraints.Num() > 0 && GetSkeletalMeshAsset())
 	{
 		check( PhysicsAsset->ConstraintSetup.Num() == Constraints.Num() );
 
@@ -744,7 +744,7 @@ void USkeletalMeshComponent::UpdateRBJointMotors()
 			FConstraintInstance* CI = Constraints[i];
 
 			FName JointChildBoneName = CS->DefaultInstance.GetChildBoneName();
-			int32 BoneIndex = GetSkeletalMesh()->GetRefSkeleton().FindBoneIndex(JointChildBoneName);
+			int32 BoneIndex = GetSkeletalMeshAsset()->GetRefSkeleton().FindBoneIndex(JointChildBoneName);
 
 			// If we found this bone, and a visible bone that is not the root, and its joint is motorised in some way..
 			if( (BoneIndex != INDEX_NONE) && (BoneIndex != 0) &&
@@ -763,8 +763,8 @@ void USkeletalMeshComponent::UpdateRBJointMotors()
 				// We need this to compensate for welding, where graphics and physics parents may not be the same.
 				FMatrix ControlBodyToParentBoneTM = FMatrix::Identity;
 
-				int32 TestBoneIndex = GetSkeletalMesh()->GetRefSkeleton().GetParentIndex(BoneIndex); // This give the 'graphics' parent of this bone
-				bool bFoundControlBody = (GetSkeletalMesh()->GetRefSkeleton().GetBoneName(TestBoneIndex) == CS->DefaultInstance.ConstraintBone2); // ConstraintBone2 is the 'physics' parent of this joint.
+				int32 TestBoneIndex = GetSkeletalMeshAsset()->GetRefSkeleton().GetParentIndex(BoneIndex); // This give the 'graphics' parent of this bone
+				bool bFoundControlBody = (GetSkeletalMeshAsset()->GetRefSkeleton().GetBoneName(TestBoneIndex) == CS->DefaultInstance.ConstraintBone2); // ConstraintBone2 is the 'physics' parent of this joint.
 
 				while(!bFoundControlBody)
 				{
@@ -786,7 +786,7 @@ void USkeletalMeshComponent::UpdateRBJointMotors()
 					ControlBodyToParentBoneTM = ControlBodyToParentBoneTM * RelTM;
 
 					// Move on to parent
-					TestBoneIndex = GetSkeletalMesh()->GetRefSkeleton().GetParentIndex(TestBoneIndex);
+					TestBoneIndex = GetSkeletalMeshAsset()->GetRefSkeleton().GetParentIndex(TestBoneIndex);
 
 					// If we are at the root - bail out.
 					if(TestBoneIndex == 0)
@@ -795,7 +795,7 @@ void USkeletalMeshComponent::UpdateRBJointMotors()
 					}
 
 					// See if this is the controlling body
-					bFoundControlBody = (GetSkeletalMesh()->GetRefSkeleton().GetBoneName(TestBoneIndex) == CS->DefaultInstance.ConstraintBone2);
+					bFoundControlBody = (GetSkeletalMeshAsset()->GetRefSkeleton().GetBoneName(TestBoneIndex) == CS->DefaultInstance.ConstraintBone2);
 				}
 
 				// If after that we didn't find a parent body, we can' do this, so skip.
diff --git a/Engine/Source/Runtime/Engine/Private/PhysicsEngine/PhysicalAnimationComponent.cpp b/Engine/Source/Runtime/Engine/Private/PhysicsEngine/PhysicalAnimationComponent.cpp
--- a/Engine/Source/Runtime/Engine/Private/PhysicsEngine/PhysicalAnimationComponent.cpp
+++ b/Engine/Source/Runtime/Engine/Private/PhysicsEngine/PhysicalAnimationComponent.cpp
@@ -214,7 +214,7 @@ FTransform ComputeWorldSpaceTargetTM(const USkeletalMeshComponent& SkeletalMeshC
 
 FTransform ComputeLocalSpaceTargetTM(const USkeletalMeshComponent& SkeletalMeshComponent, const UPhysicsAsset& PhysAsset, const TArray<FTransform>& LocalTransforms, int32 BoneIndex)
 {
-	const FReferenceSkeleton& RefSkeleton = SkeletalMeshComponent.GetSkeletalMesh()->GetRefSkeleton();
+	const FReferenceSkeleton& RefSkeleton = SkeletalMeshComponent.GetSkeletalMeshAsset()->GetRefSkeleton();
 	FTransform AccumulatedDelta = LocalTransforms[BoneIndex];
 	int32 CurBoneIdx = BoneIndex;
 	while ((CurBoneIdx = RefSkeleton.GetParentIndex(CurBoneIdx)) != INDEX_NONE)
@@ -256,9 +256,9 @@ FTransform ComputeTargetTM(const FPhysicalAnimationData& PhysAnimData, const USk
 void UPhysicalAnimationComponent::UpdateTargetActors(ETeleportType TeleportType)
 {
 	UPhysicsAsset* PhysAsset = SkeletalMeshComponent ? SkeletalMeshComponent->GetPhysicsAsset() : nullptr;
-	if (PhysAsset && SkeletalMeshComponent->GetSkeletalMesh())
+	if (PhysAsset && SkeletalMeshComponent->GetSkeletalMeshAsset())
 	{
-		const FReferenceSkeleton& RefSkeleton = SkeletalMeshComponent->GetSkeletalMesh()->GetRefSkeleton();
+		const FReferenceSkeleton& RefSkeleton = SkeletalMeshComponent->GetSkeletalMeshAsset()->GetRefSkeleton();
 
 		// Note we use GetEditableComponentSpaceTransforms because we need to update target actors in the midst of the 
 		// various anim ticks, before buffers are flipped (which happens in the skel mesh component's post-physics tick)
@@ -357,7 +357,7 @@ void UPhysicalAnimationComponent::UpdatePhysicsEngineImp()
 {
 	bPhysicsEngineNeedsUpdating = false;
 	UPhysicsAsset* PhysAsset = SkeletalMeshComponent ? SkeletalMeshComponent->GetPhysicsAsset() : nullptr;
-	if(PhysAsset && SkeletalMeshComponent->GetSkeletalMesh())
+	if(PhysAsset && SkeletalMeshComponent->GetSkeletalMeshAsset())
 	{
 		//TODO: This is hacky and assumes constraints can only be added and not removed. True for now, but bad in general!
 		const int32 NumData = DriveData.Num();
@@ -368,7 +368,7 @@ void UPhysicalAnimationComponent::UpdatePhysicsEngineImp()
 		// Note we use GetEditableComponentSpaceTransforms because we need to update target actors in the midst of the 
 		// various anim ticks, before buffers are flipped (which happens in the skel mesh component's post-physics tick)
 		const TArray<FTransform>& SpaceBases = SkeletalMeshComponent->GetEditableComponentSpaceTransforms();
-		const FReferenceSkeleton& RefSkeleton = SkeletalMeshComponent->GetSkeletalMesh()->GetRefSkeleton();
+		const FReferenceSkeleton& RefSkeleton = SkeletalMeshComponent->GetSkeletalMeshAsset()->GetRefSkeleton();
 
 		FPhysicsCommand::ExecuteWrite(SkeletalMeshComponent, [&]()
 		{
diff --git a/Engine/Source/Runtime/Engine/Private/SkeletalDebugRendering.cpp b/Engine/Source/Runtime/Engine/Private/SkeletalDebugRendering.cpp
--- a/Engine/Source/Runtime/Engine/Private/SkeletalDebugRendering.cpp
+++ b/Engine/Source/Runtime/Engine/Private/SkeletalDebugRendering.cpp
@@ -181,7 +181,7 @@ void DrawBonesFromPoseWatch(
 	if (Pose.GetNumBones() == 0 ||
 		Pose.GetBoneContainer().GetCompactPoseNumBones() == 0 ||
 		!MeshComponent ||
-		!MeshComponent->GetSkeletalMesh() ||
+		!MeshComponent->GetSkeletalMeshAsset() ||
 		!PoseWatchPoseElement
 		)
 	{
@@ -245,7 +245,7 @@ void DrawBonesFromPoseWatch(
 		PDI,
 		MeshComponent->GetComponentLocation() + RelativeOffset,
 		RequiredBones,
-		MeshComponent->GetSkeletalMesh()->GetRefSkeleton(),
+		MeshComponent->GetSkeletalMeshAsset()->GetRefSkeleton(),
 		WorldTransforms,
 		/*SelectedBones*/TArray<int32>(),
 		/*BoneColors*/TArray<FLinearColor>(),
diff --git a/Engine/Source/Runtime/Engine/Private/SkeletalMesh.cpp b/Engine/Source/Runtime/Engine/Private/SkeletalMesh.cpp
--- a/Engine/Source/Runtime/Engine/Private/SkeletalMesh.cpp
+++ b/Engine/Source/Runtime/Engine/Private/SkeletalMesh.cpp
@@ -345,7 +345,7 @@ void FScopedSkeletalMeshPostEditChange::SetSkeletalMesh(USkeletalMesh* InSkeleta
 			for (TObjectIterator<USkeletalMeshComponent> It; It; ++It)
 			{
 				USkeletalMeshComponent* SkelComp = *It;
-				if (SkelComp->GetSkeletalMesh() == SkeletalMesh)
+				if (SkelComp->GetSkeletalMeshAsset() == SkeletalMesh)
 				{
 					ComponentReregisterContexts.Add(new FComponentReregisterContext(SkelComp));
 				}
@@ -1205,7 +1205,7 @@ void RefreshSkelMeshOnPhysicsAssetChange(const USkeletalMesh* InSkeletalMesh)
 		{
 			USkeletalMeshComponent* SkeletalMeshComponent = Cast<USkeletalMeshComponent>(*Iter);
 			// if PhysicsAssetOverride is NULL, it uses SkeletalMesh Physics Asset, so I'll need to update here
-			if  (SkeletalMeshComponent->GetSkeletalMesh() == InSkeletalMesh &&
+			if  (SkeletalMeshComponent->GetSkeletalMeshAsset() == InSkeletalMesh &&
 				 SkeletalMeshComponent->PhysicsAssetOverride == NULL)
 			{
 				// it needs to recreate IF it already has been created
@@ -1333,7 +1333,7 @@ void USkeletalMesh::PostEditChangeProperty(FPropertyChangedEvent& PropertyChange
 		for(TObjectIterator<USkeletalMeshComponent> It; It; ++It)
 		{
 			USkeletalMeshComponent* MeshComponent = *It;
-			if(MeshComponent && !MeshComponent->IsTemplate() && MeshComponent->GetSkeletalMesh() == this)
+			if(MeshComponent && !MeshComponent->IsTemplate() && MeshComponent->GetSkeletalMeshAsset() == this)
 			{
 				ComponentsToReregister.Add(*It);
 			}
@@ -1396,7 +1396,7 @@ void USkeletalMesh::PostEditUndo()
 		USkeletalMeshComponent* MeshComponent = *It;
 		if( MeshComponent && 
 			!MeshComponent->IsTemplate() &&
-			MeshComponent->GetSkeletalMesh() == this )
+			MeshComponent->GetSkeletalMeshAsset() == this )
 		{
 			FComponentReregisterContext Context(MeshComponent);
 		}
@@ -3456,7 +3456,7 @@ void USkeletalMesh::InitMorphTargetsAndRebuildRenderData()
 		// reset all morphtarget for all components
 		for (TObjectIterator<USkeletalMeshComponent> It; It; ++It)
 		{
-			if (It->GetSkeletalMesh() == this)
+			if (It->GetSkeletalMeshAsset() == this)
 			{
 				It->RefreshMorphTargets();
 			}
@@ -6822,7 +6822,7 @@ FSkinnedMeshComponentRecreateRenderStateContext::FSkinnedMeshComponentRecreateRe
 	bool bRequireFlushRenderingCommands = false;
 	for (TObjectIterator<USkeletalMeshComponent> It; It; ++It)
 	{
-		if (It->GetSkeletalMesh() == InSkeletalMesh)
+		if (It->GetSkeletalMeshAsset() == InSkeletalMesh)
 		{
 			checkf(!It->IsUnreachable(), TEXT("%s"), *It->GetFullName());
 
diff --git a/Engine/Source/Runtime/Engine/Private/SkeletalMeshComponentPhysics.cpp b/Engine/Source/Runtime/Engine/Private/SkeletalMeshComponentPhysics.cpp
--- a/Engine/Source/Runtime/Engine/Private/SkeletalMeshComponentPhysics.cpp
+++ b/Engine/Source/Runtime/Engine/Private/SkeletalMeshComponentPhysics.cpp
@@ -160,10 +160,10 @@ void USkeletalMeshComponent::CreateBodySetup()
 		BodySetup = NewObject<UBodySetup>(this);
 	}
 
-	if (GetSkeletalMesh())
+	if (GetSkeletalMeshAsset())
 	{
-		const USkeletalMesh* SkeletalMeshConst = GetSkeletalMesh();
-		GetSkeletalMesh()->CreateBodySetup();
+		const USkeletalMesh* SkeletalMeshConst = GetSkeletalMeshAsset();
+		GetSkeletalMeshAsset()->CreateBodySetup();
 		UBodySetup* OriginalBodySetup = SkeletalMeshConst->GetBodySetup();
 		BodySetup->CopyBodyPropertiesFrom(OriginalBodySetup);
 		BodySetup->CookedFormatDataOverride = &OriginalBodySetup->CookedFormatData;
@@ -184,11 +184,11 @@ UBodySetup* USkeletalMeshComponent::GetBodySetup()
 	if (bEnablePerPolyCollision == false)
 	{
 		UPhysicsAsset * const PhysicsAsset = GetPhysicsAsset();
-		if (GetSkeletalMesh() && PhysicsAsset)
+		if (GetSkeletalMeshAsset() && PhysicsAsset)
 		{
-			for (int32 i = 0; i < GetSkeletalMesh()->GetRefSkeleton().GetNum(); i++)
+			for (int32 i = 0; i < GetSkeletalMeshAsset()->GetRefSkeleton().GetNum(); i++)
 			{
-				int32 BodyIndex = PhysicsAsset->FindBodyIndex(GetSkeletalMesh()->GetRefSkeleton().GetBoneName(i));
+				int32 BodyIndex = PhysicsAsset->FindBodyIndex(GetSkeletalMeshAsset()->GetRefSkeleton().GetBoneName(i));
 				if (BodyIndex != INDEX_NONE)
 				{
 					return PhysicsAsset->SkeletalBodySetups[BodyIndex];
@@ -651,13 +651,13 @@ int32 USkeletalMeshComponent::FindRootBodyIndex() const
 {
 	// Find root physics body
 	int32 RootBodyIndex = RootBodyData.BodyIndex;
-	if(RootBodyIndex == INDEX_NONE && GetSkeletalMesh())
+	if(RootBodyIndex == INDEX_NONE && GetSkeletalMeshAsset())
 	{
 		if(const UPhysicsAsset* PhysicsAsset = GetPhysicsAsset())
 		{
-			for (int32 i = 0; i< GetSkeletalMesh()->GetRefSkeleton().GetNum(); i++)
+			for (int32 i = 0; i< GetSkeletalMeshAsset()->GetRefSkeleton().GetNum(); i++)
 			{
-				int32 BodyInstIndex = PhysicsAsset->FindBodyIndex(GetSkeletalMesh()->GetRefSkeleton().GetBoneName(i));
+				int32 BodyInstIndex = PhysicsAsset->FindBodyIndex(GetSkeletalMeshAsset()->GetRefSkeleton().GetBoneName(i));
 				if (BodyInstIndex != INDEX_NONE)
 				{
 					RootBodyIndex = BodyInstIndex;
@@ -684,7 +684,7 @@ void USkeletalMeshComponent::InitArticulated(FPhysScene* PhysScene)
 
 	UPhysicsAsset* const PhysicsAsset = GetPhysicsAsset();
 
-	if (PhysScene == nullptr || PhysicsAsset == nullptr || GetSkeletalMesh() == nullptr || !ShouldCreatePhysicsState())
+	if (PhysScene == nullptr || PhysicsAsset == nullptr || GetSkeletalMeshAsset() == nullptr || !ShouldCreatePhysicsState())
 	{
 		return;
 	}
@@ -774,9 +774,9 @@ void USkeletalMeshComponent::InstantiatePhysicsAssetBodies(const UPhysicsAsset&
 
 void USkeletalMeshComponent::InstantiatePhysicsAssetRefPose(const UPhysicsAsset& PhysAsset, const FVector& Scale3D, TArray<FBodyInstance*>& OutBodies, TArray<FConstraintInstance*>& OutConstraints, FPhysScene* PhysScene /*= nullptr*/, USkeletalMeshComponent* OwningComponent /*= nullptr*/, int32 UseRootBodyIndex /*= INDEX_NONE*/, const FPhysicsAggregateHandle& UseAggregate, bool bCreateBodiesInRefPose) const
 {
-	if(GetSkeletalMesh())
+	if(GetSkeletalMeshAsset())
 	{
-		const FReferenceSkeleton& RefSkeleton = GetSkeletalMesh()->GetRefSkeleton();
+		const FReferenceSkeleton& RefSkeleton = GetSkeletalMeshAsset()->GetRefSkeleton();
 
 		if (bCreateBodiesInRefPose)
 		{
@@ -1028,7 +1028,7 @@ void USkeletalMeshComponent::TermArticulated()
 void USkeletalMeshComponent::TermBodiesBelow(FName ParentBoneName)
 {
 	UPhysicsAsset* const PhysicsAsset = GetPhysicsAsset();
-	if(PhysicsAsset && GetSkeletalMesh() && Bodies.Num() > 0)
+	if(PhysicsAsset && GetSkeletalMeshAsset() && Bodies.Num() > 0)
 	{
 		check(Bodies.Num() == PhysicsAsset->SkeletalBodySetups.Num());
 
@@ -1048,7 +1048,7 @@ void USkeletalMeshComponent::TermBodiesBelow(FName ParentBoneName)
 			int32 JointBoneIndex = GetBoneIndex(JointChildBoneName);
 
 			// If constraint has bone in mesh, and is either the parent or child of it, term it
-			if(	JointBoneIndex != INDEX_NONE && (JointChildBoneName == ParentBoneName || GetSkeletalMesh()->GetRefSkeleton().BoneIsChildOf(JointBoneIndex, ParentBoneIndex)) )
+			if(	JointBoneIndex != INDEX_NONE && (JointChildBoneName == ParentBoneName || GetSkeletalMeshAsset()->GetRefSkeleton().BoneIsChildOf(JointBoneIndex, ParentBoneIndex)) )
 			{
 				Constraints[i]->TermConstraint();
 			}
@@ -1064,7 +1064,7 @@ void USkeletalMeshComponent::TermBodiesBelow(FName ParentBoneName)
 				int32 BodyBoneIndex = GetBoneIndex(BodyName);
 
 				// If body has bone in mesh, and is either the parent or child of it, term it
-				if(	BodyBoneIndex != INDEX_NONE && (BodyName == ParentBoneName || GetSkeletalMesh()->GetRefSkeleton().BoneIsChildOf(BodyBoneIndex, ParentBoneIndex)) )
+				if(	BodyBoneIndex != INDEX_NONE && (BodyName == ParentBoneName || GetSkeletalMeshAsset()->GetRefSkeleton().BoneIsChildOf(BodyBoneIndex, ParentBoneIndex)) )
 				{
 					Bodies[i]->TermBody();
 				}
@@ -1655,7 +1655,7 @@ void USkeletalMeshComponent::UpdateMeshForBrokenConstraints()
 {
 	UPhysicsAsset * const PhysicsAsset = GetPhysicsAsset();
 	// Needs to have a SkeletalMesh, and PhysicsAsset.
-	if( !GetSkeletalMesh() || !PhysicsAsset )
+	if( !GetSkeletalMeshAsset() || !PhysicsAsset )
 	{
 		return;
 	}
@@ -1684,7 +1684,7 @@ void USkeletalMeshComponent::UpdateMeshForBrokenConstraints()
 				UBodySetup* PhysicsAssetBodySetup = PhysicsAsset->SkeletalBodySetups[BodySetupIndex];
 				int32 BoneIndex = GetBoneIndex(PhysicsAssetBodySetup->BoneName);
 				if( BoneIndex != INDEX_NONE && 
-					(BoneIndex == JointBoneIndex || GetSkeletalMesh()->GetRefSkeleton().BoneIsChildOf(BoneIndex, JointBoneIndex)) )
+					(BoneIndex == JointBoneIndex || GetSkeletalMeshAsset()->GetRefSkeleton().BoneIsChildOf(BoneIndex, JointBoneIndex)) )
 				{
 					DEBUGBROKENCONSTRAINTUPDATE(UE_LOG(LogSkeletalMesh, Log, TEXT("    Found Child Bone: (%d) %s"), BoneIndex, *PhysicsAssetBodySetup->BoneName.ToString());)
 
@@ -1825,7 +1825,7 @@ int32 USkeletalMeshComponent::ForEachBodyBelow(FName BoneName, bool bIncludeSelf
 	else
 	{
 		UPhysicsAsset* const PhysicsAsset = GetPhysicsAsset();
-		if (!PhysicsAsset || !GetSkeletalMesh())
+		if (!PhysicsAsset || !GetSkeletalMeshAsset())
 		{
 			return 0;
 		}
@@ -1839,7 +1839,7 @@ int32 USkeletalMeshComponent::ForEachBodyBelow(FName BoneName, bool bIncludeSelf
 
 		TArray<int32> BodyIndices;
 		BodyIndices.Reserve(Bodies.Num());
-		PhysicsAsset->GetBodyIndicesBelow(BodyIndices, BoneName, GetSkeletalMesh(), bIncludeSelf);
+		PhysicsAsset->GetBodyIndicesBelow(BodyIndices, BoneName, GetSkeletalMeshAsset(), bIncludeSelf);
 
 		int32 NumBodiesFound = 0;
 		for (int32 BodyIdx : BodyIndices)
@@ -2062,12 +2062,12 @@ void USkeletalMeshComponent::SetPhysicsAsset(UPhysicsAsset* InPhysicsAsset, bool
 		UpdateHasValidBodies();
 
 		// Component should be re-attached here, so create physics.
-		if(GetSkeletalMesh())
+		if(GetSkeletalMeshAsset())
 		{
 			// Because we don't know what bones the new PhysicsAsset might want, we have to force an update to _all_ bones in the skeleton.
-			RequiredBones.Reset(GetSkeletalMesh()->GetRefSkeleton().GetNum());
-			RequiredBones.AddUninitialized(GetSkeletalMesh()->GetRefSkeleton().GetNum() );
-			for(int32 i=0; i< GetSkeletalMesh()->GetRefSkeleton().GetNum(); i++)
+			RequiredBones.Reset(GetSkeletalMeshAsset()->GetRefSkeleton().GetNum());
+			RequiredBones.AddUninitialized(GetSkeletalMeshAsset()->GetRefSkeleton().GetNum() );
+			for(int32 i=0; i< GetSkeletalMeshAsset()->GetRefSkeleton().GetNum(); i++)
 			{
 				RequiredBones[i] = (FBoneIndexType)i;
 			}
@@ -2184,8 +2184,8 @@ FVector GetTypedSkinnedVertexPositionWithCloth(USkeletalMeshComponent* Component
 	const FSkelMeshRenderSection& Section = LODData.RenderSections[SectionIndex];
 
 	// only if this component has clothing and is showing simulated results	
-	if (Component->GetSkeletalMesh() &&
-		Component->GetSkeletalMesh()->GetMeshClothingAssets().Num() > 0 &&
+	if (Component->GetSkeletalMeshAsset() &&
+		Component->GetSkeletalMeshAsset()->GetMeshClothingAssets().Num() > 0 &&
 		!Component->bDisableClothSimulation &&
 		Component->ClothBlendWeight > 0.0f // if cloth blend weight is 0.0, only showing skinned vertices regardless of simulation positions
 		)
@@ -2240,12 +2240,12 @@ FVector3f USkeletalMeshComponent::GetSkinnedVertexPosition(USkeletalMeshComponen
 void USkeletalMeshComponent::ComputeSkinnedPositions(USkeletalMeshComponent* Component, TArray<FVector3f> & OutPositions, TArray<FMatrix44f>& CachedRefToLocals, const FSkeletalMeshLODRenderData& LODData, const FSkinWeightVertexBuffer& SkinWeightBuffer)
 {
 	// Fail if no mesh
-	if (!Component->GetSkeletalMesh())
+	if (!Component->GetSkeletalMeshAsset())
 	{
 		return;
 	}
 
-	if (Component->GetSkeletalMesh()->GetMeshClothingAssets().Num() > 0 &&
+	if (Component->GetSkeletalMeshAsset()->GetMeshClothingAssets().Num() > 0 &&
 		!Component->bDisableClothSimulation &&
 		Component->ClothBlendWeight > 0.0f // if cloth blend weight is 0.0, only showing skinned vertices regardless of simulation positions
 		)
@@ -2280,7 +2280,7 @@ void USkeletalMeshComponent::ComputeSkinnedPositions(USkeletalMeshComponent* Com
 
 			if (bClothVertex)
 			{
-				int32 AssetIndex = Component->GetSkeletalMesh()->GetClothingAssetIndex(ClothAssetGuid);
+				int32 AssetIndex = Component->GetSkeletalMeshAsset()->GetClothingAssetIndex(ClothAssetGuid);
 				if (AssetIndex != INDEX_NONE)
 				{
 					const FClothSimulData* ActorData = ClothData.Find(AssetIndex);
@@ -2360,7 +2360,7 @@ void USkeletalMeshComponent::GetSkinnedTangentBasis(USkeletalMeshComponent* Comp
 void USkeletalMeshComponent::ComputeSkinnedTangentBasis(USkeletalMeshComponent* Component, TArray<FVector3f>& OutTangenXYZ, TArray<FMatrix44f>& CachedRefToLocals, const FSkeletalMeshLODRenderData& LODData, const FSkinWeightVertexBuffer& SkinWeightBuffer)
 {
 	// Fail if no mesh
-	if (!Component->GetSkeletalMesh())
+	if (!Component->GetSkeletalMeshAsset())
 	{
 		return;
 	}
@@ -2454,7 +2454,7 @@ bool USkeletalMeshComponent::GetClosestPointOnPhysicsAsset(const FVector& WorldP
 
 	bool bSuccess = false;
 	const UPhysicsAsset* PhysicsAsset = GetPhysicsAsset();
-	const FReferenceSkeleton* RefSkeleton = GetSkeletalMesh() ? &GetSkeletalMesh()->GetRefSkeleton() : nullptr;
+	const FReferenceSkeleton* RefSkeleton = GetSkeletalMeshAsset() ? &GetSkeletalMeshAsset()->GetRefSkeleton() : nullptr;
 	if(PhysicsAsset && RefSkeleton)
 	{
 		const TArray<FTransform>& BoneTransforms = GetComponentSpaceTransforms();
@@ -2666,7 +2666,7 @@ void USkeletalMeshComponent::RecreateClothingActors()
 
 	ReleaseAllClothingResources();
 
-	USkeletalMesh* SkelMesh = GetSkeletalMesh();
+	USkeletalMesh* SkelMesh = GetSkeletalMeshAsset();
 	if(!bAllowClothActors || SkelMesh == nullptr || !IsRegistered())
 	{
 		return;
@@ -2826,7 +2826,7 @@ void USkeletalMeshComponent::ExtractCollisionsForCloth(
 {
 	// Extract collisions from this mesh 'raw', as this isn't a mesh that has cloth simulation
 	// (but we want it to affect other meshes with cloth simulation)
-	if(SourceComponent->GetSkeletalMesh() && PhysicsAsset)
+	if(SourceComponent->GetSkeletalMeshAsset() && PhysicsAsset)
 	{
 		FTransform ComponentToComponentTransform;
 		if(SourceComponent != DestClothComponent)
@@ -2837,12 +2837,12 @@ void USkeletalMeshComponent::ExtractCollisionsForCloth(
 		}
 
 		// Init cache on first copy
-		if(!ClothCollisionSource.bCached || ClothCollisionSource.CachedSkeletalMesh.Get() != SourceComponent->GetSkeletalMesh())
+		if(!ClothCollisionSource.bCached || ClothCollisionSource.CachedSkeletalMesh.Get() != SourceComponent->GetSkeletalMeshAsset())
 		{
 			for(const USkeletalBodySetup* SkeletalBodySetup : PhysicsAsset->SkeletalBodySetups)
 			{
 				// Cache bones
-				int32 MeshBoneIndex = SourceComponent->GetSkeletalMesh()->GetRefSkeleton().FindBoneIndex(SkeletalBodySetup->BoneName);
+				int32 MeshBoneIndex = SourceComponent->GetSkeletalMeshAsset()->GetRefSkeleton().FindBoneIndex(SkeletalBodySetup->BoneName);
 				if(MeshBoneIndex != INDEX_NONE)
 				{
 					// Cache spheres & capsules form physics asset
@@ -2904,7 +2904,7 @@ void USkeletalMeshComponent::ExtractCollisionsForCloth(
 				}
 			}
 
-			ClothCollisionSource.CachedSkeletalMesh = SourceComponent->GetSkeletalMesh();
+			ClothCollisionSource.CachedSkeletalMesh = SourceComponent->GetSkeletalMeshAsset();
 			ClothCollisionSource.bCached = true;
 		}
 
@@ -3337,7 +3337,7 @@ void USkeletalMeshComponent::ProcessClothCollisionWithEnvironment()
 				// Possibly a skeletal mesh, extract it's clothing collisions if necessary
 				USkeletalMeshComponent* SkelComp = Cast<USkeletalMeshComponent>(Component.Get());
 
-				if(SkelComp && SkelComp->GetSkeletalMesh())
+				if(SkelComp && SkelComp->GetSkeletalMeshAsset())
 				{
 					if(SkelComp == this)
 					{
@@ -3602,7 +3602,7 @@ public:
 
 	void DoTask(ENamedThreads::Type CurrentThread, const FGraphEventRef& MyCompletionGraphEvent)
 	{
-		FScopeCycleCounterUObject ContextScope(SkeletalMeshComponent.GetSkeletalMesh());
+		FScopeCycleCounterUObject ContextScope(SkeletalMeshComponent.GetSkeletalMeshAsset());
 		SCOPE_CYCLE_COUNTER(STAT_ClothTotalTime);
 		CSV_SCOPED_TIMING_STAT(Animation, Cloth);
 
@@ -3759,7 +3759,7 @@ DECLARE_CYCLE_STAT(TEXT("Cloth Sim"), STAT_ClothSimTime, STATGROUP_Physics);
 
 bool USkeletalMeshComponent::GetClothSimulatedPosition_GameThread(const FGuid& AssetGuid, int32 VertexIndex, FVector& OutSimulPos) const
 {
-	if(!GetSkeletalMesh())
+	if(!GetSkeletalMeshAsset())
 	{
 		// Can't proceed without a mesh
 		return false;
@@ -3767,7 +3767,7 @@ bool USkeletalMeshComponent::GetClothSimulatedPosition_GameThread(const FGuid& A
 		
 	bool bSucceed = false;
 
-	int32 AssetIndex = GetSkeletalMesh()->GetClothingAssetIndex(AssetGuid);
+	int32 AssetIndex = GetSkeletalMeshAsset()->GetClothingAssetIndex(AssetGuid);
 
 	if(AssetIndex != INDEX_NONE)
 	{
@@ -3785,7 +3785,7 @@ bool USkeletalMeshComponent::GetClothSimulatedPosition_GameThread(const FGuid& A
 
 void USkeletalMeshComponent::TickClothing(float DeltaTime, FTickFunction& ThisTickFunction)
 {
-	if (GetSkeletalMesh() == nullptr || !ClothingSimulation || CVarEnableClothPhysics.GetValueOnGameThread() == 0)
+	if (GetSkeletalMeshAsset() == nullptr || !ClothingSimulation || CVarEnableClothPhysics.GetValueOnGameThread() == 0)
 	{
 		return;
 	}
diff --git a/Engine/Source/Runtime/Engine/Private/UnrealEngine.cpp b/Engine/Source/Runtime/Engine/Private/UnrealEngine.cpp
--- a/Engine/Source/Runtime/Engine/Private/UnrealEngine.cpp
+++ b/Engine/Source/Runtime/Engine/Private/UnrealEngine.cpp
@@ -7894,7 +7894,7 @@ bool UEngine::HandleMergeMeshCommand( const TCHAR* Cmd, FOutputDevice& Ar, UWorl
 		if (PlayerController && PlayerController->GetCharacter() != NULL && PlayerController->GetCharacter()->GetMesh())
 		{
 			PlayerPawn = PlayerController->GetCharacter();
-			PlayerMesh = PlayerController->GetCharacter()->GetMesh()->GetSkeletalMesh();
+			PlayerMesh = PlayerController->GetCharacter()->GetMesh()->GetSkeletalMeshAsset();
 			break;
 		}
 	}
diff --git a/Engine/Source/Runtime/Experimental/GeometryCollectionEngine/Private/GeometryCollection/PhysicsAssetSimulation.cpp b/Engine/Source/Runtime/Experimental/GeometryCollectionEngine/Private/GeometryCollection/PhysicsAssetSimulation.cpp
--- a/Engine/Source/Runtime/Experimental/GeometryCollectionEngine/Private/GeometryCollection/PhysicsAssetSimulation.cpp
+++ b/Engine/Source/Runtime/Experimental/GeometryCollectionEngine/Private/GeometryCollection/PhysicsAssetSimulation.cpp
@@ -36,7 +36,7 @@ void FPhysicsAssetSimulationUtil::BuildParams(const UObject* Caller, const AActo
 	Chaos::FVec3 ActorOrigin(ActorOriginVec);
 	Chaos::FVec3 ActorBoxExtent(ActorBoxExtentVec);
 
-	const USkeletalMesh* SkeletalMesh = SkelMeshComponent->GetSkeletalMesh();
+	const USkeletalMesh* SkeletalMesh = SkelMeshComponent->GetSkeletalMeshAsset();
 	if (!SkeletalMesh)
 	{
 		return;
@@ -345,7 +345,7 @@ void FPhysicsAssetSimulationUtil::BuildParams(const UObject* Caller, const AActo
 		int32 TmpBoneIndex = BoneIndex;
 		int32 SocketIndex = INDEX_NONE;
 		if (USkeletalMeshSocket const* Socket =
-			SkelMeshComponent->GetSkeletalMesh()->FindSocketInfo(
+			SkelMeshComponent->GetSkeletalMeshAsset()->FindSocketInfo(
 				Params.BoneHierarchy.GetAnalyticShapeGroup(BoneIndex)->GetBoneName(),
 				BoneXf,
 				TmpBoneIndex,
@@ -375,7 +375,7 @@ bool FPhysicsAssetSimulationUtil::UpdateAnimState(const UObject* Caller, const A
 {
 	int32 UpdatedBones = 0;
 
-	const USkeletalMesh* SkeletalMesh = SkelMeshComponent->GetSkeletalMesh();
+	const USkeletalMesh* SkeletalMesh = SkelMeshComponent->GetSkeletalMeshAsset();
 	
 	if (!SkeletalMesh)
 	{
@@ -413,7 +413,7 @@ bool FPhysicsAssetSimulationUtil::UpdateAnimState(const UObject* Caller, const A
 			const int32 SocketIndex = Params.BoneHierarchy.GetSocketIndexForBone(BoneIndex);
 			if (SocketIndex != INDEX_NONE)
 			{
-				USkeletalMeshSocket *Socket = SkelMeshComponent->GetSkeletalMesh()->GetSocketByIndex(SocketIndex);
+				USkeletalMeshSocket *Socket = SkelMeshComponent->GetSkeletalMeshAsset()->GetSocketByIndex(SocketIndex);
 				Params.BoneHierarchy.SetAnimLocalSpaceTransform(
 					BoneIndex,
 					Socket->GetSocketLocalTransform());
diff --git a/Engine/Source/Runtime/MovieSceneTracks/Private/Evaluation/MovieSceneSkeletalAnimationTemplate.cpp b/Engine/Source/Runtime/MovieSceneTracks/Private/Evaluation/MovieSceneSkeletalAnimationTemplate.cpp
--- a/Engine/Source/Runtime/MovieSceneTracks/Private/Evaluation/MovieSceneSkeletalAnimationTemplate.cpp
+++ b/Engine/Source/Runtime/MovieSceneTracks/Private/Evaluation/MovieSceneSkeletalAnimationTemplate.cpp
@@ -31,8 +31,8 @@ bool ShouldUsePreviewPlayback(IMovieScenePlayer& Player, UObject& RuntimeObject)
 
 bool CanPlayAnimation(USkeletalMeshComponent* SkeletalMeshComponent, UAnimSequenceBase* AnimAssetBase)
 {
-	return (SkeletalMeshComponent->GetSkeletalMesh() && SkeletalMeshComponent->GetSkeletalMesh()->GetSkeleton() &&
-		(!AnimAssetBase || SkeletalMeshComponent->GetSkeletalMesh()->GetSkeleton()->IsCompatible(AnimAssetBase->GetSkeleton())));
+	return (SkeletalMeshComponent->GetSkeletalMeshAsset() && SkeletalMeshComponent->GetSkeletalMeshAsset()->GetSkeleton() &&
+		(!AnimAssetBase || SkeletalMeshComponent->GetSkeletalMeshAsset()->GetSkeleton()->IsCompatible(AnimAssetBase->GetSkeleton())));
 }
 
 void ResetAnimSequencerInstance(UObject& InObject, const UE::MovieScene::FRestoreStateParams& Params)
@@ -142,10 +142,10 @@ struct FPreAnimatedAnimationTokenProducer : IMovieScenePreAnimatedTokenProducer
 				{
 					Component->AnimScriptInstance = CachedAnimInstance.Get();
 					CachedAnimInstance.Reset();
-					if (Component->AnimScriptInstance && Component->GetSkeletalMesh() && Component->AnimScriptInstance->CurrentSkeleton != Component->GetSkeletalMesh()->GetSkeleton())
+					if (Component->AnimScriptInstance && Component->GetSkeletalMeshAsset() && Component->AnimScriptInstance->CurrentSkeleton != Component->GetSkeletalMeshAsset()->GetSkeleton())
 					{
 						//the skeleton may have changed so need to recalc required bones as needed.
-						Component->AnimScriptInstance->CurrentSkeleton = Component->GetSkeletalMesh()->GetSkeleton();
+						Component->AnimScriptInstance->CurrentSkeleton = Component->GetSkeletalMeshAsset()->GetSkeleton();
 						//Need at least RecalcRequiredbones and UpdateMorphTargetrs
 						Component->InitializeAnimScriptInstance(true);
 					}
@@ -281,7 +281,7 @@ namespace MovieScene
 			ensureMsgf(InObject, TEXT("Attempting to evaluate an Animation track with a null object."));
 
 			USkeletalMeshComponent* SkeletalMeshComponent = SkeletalMeshComponentFromObject(InObject);
-			if (!SkeletalMeshComponent || !SkeletalMeshComponent->GetSkeletalMesh())
+			if (!SkeletalMeshComponent || !SkeletalMeshComponent->GetSkeletalMeshAsset())
 			{
 				return;
 			}
