CIS - Rollback removal of GetSkeletalMesh().

diff --git a/Engine/Source/Runtime/Engine/Classes/Components/SkeletalMeshComponent.h b/Engine/Source/Runtime/Engine/Classes/Components/SkeletalMeshComponent.h
--- a/Engine/Source/Runtime/Engine/Classes/Components/SkeletalMeshComponent.h
+++ b/Engine/Source/Runtime/Engine/Classes/Components/SkeletalMeshComponent.h
@@ -1899,6 +1899,9 @@ public:
 	 */
 	virtual void SetSkinnedAssetAndUpdate(class USkinnedAsset* InSkinnedAsset, bool bReinitPose = true) override;
 
+	// GetSkeletalMesh() is soon to be removed, use GetSkeletalMeshAsset() instead
+	class USkeletalMesh* GetSkeletalMesh() const { return GetSkeletalMeshAsset(); }
+
 	static FVector3f GetSkinnedVertexPosition(USkeletalMeshComponent* Component, int32 VertexIndex, const FSkeletalMeshLODRenderData& Model, const FSkinWeightVertexBuffer& SkinWeightBuffer);
 	static FVector3f GetSkinnedVertexPosition(USkeletalMeshComponent* Component, int32 VertexIndex, const FSkeletalMeshLODRenderData& Model, const FSkinWeightVertexBuffer& SkinWeightBuffer, TArray<FMatrix44f>& CachedRefToLocals);
 	static void ComputeSkinnedPositions(USkeletalMeshComponent* Component, TArray<FVector3f> & OutPositions, TArray<FMatrix44f>& CachedRefToLocals, const FSkeletalMeshLODRenderData& Model, const FSkinWeightVertexBuffer& SkinWeightBuffer);
