* Fixed complex collision display for spline mesh components
* Factorized FMeshBatch spline setup code

diff --git a/Engine/Source/Runtime/Engine/Private/SplineMeshSceneProxy.cpp b/Engine/Source/Runtime/Engine/Private/SplineMeshSceneProxy.cpp
--- a/Engine/Source/Runtime/Engine/Private/SplineMeshSceneProxy.cpp
+++ b/Engine/Source/Runtime/Engine/Private/SplineMeshSceneProxy.cpp
@@ -70,20 +70,25 @@ SIZE_T FSplineMeshSceneProxy::GetTypeHash() const
 	return reinterpret_cast<size_t>(&UniquePointer);
 }
 
+void FSplineMeshSceneProxy::SetupMeshBatchForSpline(int32 InLODIndex, FMeshBatch& OutMeshBatch) const
+{
+	const FStaticMeshVertexFactories& VFs = RenderData->LODVertexFactories[InLODIndex];
+	check(OutMeshBatch.Elements.Num() == 1);
+	OutMeshBatch.VertexFactory = OutMeshBatch.Elements[0].bUserDataIsColorVertexBuffer ? VFs.SplineVertexFactoryOverrideColorVertexBuffer : VFs.SplineVertexFactory;
+	check(OutMeshBatch.VertexFactory);
+	OutMeshBatch.Elements[0].SplineMeshSceneProxy = const_cast<FSplineMeshSceneProxy*>(this);
+	OutMeshBatch.Elements[0].bIsSplineProxy = true;
+	OutMeshBatch.Elements[0].PrimitiveUniformBuffer = GetUniformBuffer();
+	OutMeshBatch.ReverseCulling ^= (SplineParams.StartScale.X < 0) ^ (SplineParams.StartScale.Y < 0);
+}
+
 bool FSplineMeshSceneProxy::GetShadowMeshElement(int32 LODIndex, int32 BatchIndex, uint8 InDepthPriorityGroup, FMeshBatch& OutMeshBatch, bool bDitheredLODTransition) const
 {
 	//checkf(LODIndex == 0, TEXT("Getting spline static mesh element with invalid LOD [%d]"), LODIndex);
 
 	if (FStaticMeshSceneProxy::GetShadowMeshElement(LODIndex, BatchIndex, InDepthPriorityGroup, OutMeshBatch, bDitheredLODTransition))
 	{
-		const FStaticMeshVertexFactories& VFs = RenderData->LODVertexFactories[LODIndex];
-		check(OutMeshBatch.Elements.Num() == 1);
-		OutMeshBatch.VertexFactory = OutMeshBatch.Elements[0].bUserDataIsColorVertexBuffer ? VFs.SplineVertexFactoryOverrideColorVertexBuffer : VFs.SplineVertexFactory;
-		check(OutMeshBatch.VertexFactory);
-		OutMeshBatch.Elements[0].SplineMeshSceneProxy = const_cast<FSplineMeshSceneProxy*>(this);
-		OutMeshBatch.Elements[0].bIsSplineProxy = true;
-		OutMeshBatch.Elements[0].PrimitiveUniformBuffer = GetUniformBuffer();
-		OutMeshBatch.ReverseCulling ^= (SplineParams.StartScale.X < 0) ^ (SplineParams.StartScale.Y < 0);
+		SetupMeshBatchForSpline(LODIndex, OutMeshBatch);
 		return true;
 	}
 	return false;
@@ -95,14 +100,7 @@ bool FSplineMeshSceneProxy::GetMeshElement(int32 LODIndex, int32 BatchIndex, int
 
 	if (FStaticMeshSceneProxy::GetMeshElement(LODIndex, BatchIndex, SectionIndex, InDepthPriorityGroup, bUseSelectionOutline, bAllowPreCulledIndices, OutMeshBatch))
 	{
-		const FStaticMeshVertexFactories& VFs = RenderData->LODVertexFactories[LODIndex];
-		check(OutMeshBatch.Elements.Num() == 1);
-		OutMeshBatch.VertexFactory = OutMeshBatch.Elements[0].bUserDataIsColorVertexBuffer ? VFs.SplineVertexFactoryOverrideColorVertexBuffer : VFs.SplineVertexFactory;
-		check(OutMeshBatch.VertexFactory);
-		OutMeshBatch.Elements[0].SplineMeshSceneProxy = const_cast<FSplineMeshSceneProxy*>(this);
-		OutMeshBatch.Elements[0].bIsSplineProxy = true;
-		OutMeshBatch.Elements[0].PrimitiveUniformBuffer = GetUniformBuffer();
-		OutMeshBatch.ReverseCulling ^= (SplineParams.StartScale.X < 0) ^ (SplineParams.StartScale.Y < 0);
+		SetupMeshBatchForSpline(LODIndex, OutMeshBatch);
 		return true;
 	}
 	return false;
@@ -114,15 +112,21 @@ bool FSplineMeshSceneProxy::GetWireframeMeshElement(int32 LODIndex, int32 BatchI
 
 	if (FStaticMeshSceneProxy::GetWireframeMeshElement(LODIndex, BatchIndex, WireframeRenderProxy, InDepthPriorityGroup, bAllowPreCulledIndices, OutMeshBatch))
 	{
-		const FStaticMeshVertexFactories& VFs = RenderData->LODVertexFactories[LODIndex];
-		check(OutMeshBatch.Elements.Num() == 1);
-		OutMeshBatch.VertexFactory = OutMeshBatch.Elements[0].bUserDataIsColorVertexBuffer ? VFs.SplineVertexFactoryOverrideColorVertexBuffer : VFs.SplineVertexFactory;
-		check(OutMeshBatch.VertexFactory);
-		OutMeshBatch.Elements[0].SplineMeshSceneProxy = const_cast<FSplineMeshSceneProxy*>(this);
-		OutMeshBatch.Elements[0].bIsSplineProxy = true;
-		OutMeshBatch.Elements[0].PrimitiveUniformBuffer = GetUniformBuffer();
-		OutMeshBatch.ReverseCulling ^= (SplineParams.StartScale.X < 0) ^ (SplineParams.StartScale.Y < 0);
+		SetupMeshBatchForSpline(LODIndex, OutMeshBatch);
 		return true;
 	}
 	return false;
 }
+
+
+bool FSplineMeshSceneProxy::GetCollisionMeshElement(int32 LODIndex, int32 BatchIndex, int32 SectionIndex, uint8 InDepthPriorityGroup, const FMaterialRenderProxy* RenderProxy, FMeshBatch& OutMeshBatch) const
+{
+	//checkf(LODIndex == 0, TEXT("Getting spline static mesh element with invalid LOD [%d]"), LODIndex);
+
+	if (FStaticMeshSceneProxy::GetCollisionMeshElement(LODIndex, BatchIndex, SectionIndex, InDepthPriorityGroup, RenderProxy, OutMeshBatch))
+	{
+		SetupMeshBatchForSpline(LODIndex, OutMeshBatch);
+		return true;
+	}
+	return false;
+}
\ No newline at end of file
diff --git a/Engine/Source/Runtime/Engine/Private/SplineMeshSceneProxy.h b/Engine/Source/Runtime/Engine/Private/SplineMeshSceneProxy.h
--- a/Engine/Source/Runtime/Engine/Private/SplineMeshSceneProxy.h
+++ b/Engine/Source/Runtime/Engine/Private/SplineMeshSceneProxy.h
@@ -104,6 +104,9 @@ public:
 	/** Sets up a wireframe FMeshBatch for a specific LOD. */
 	virtual bool GetWireframeMeshElement(int32 LODIndex, int32 BatchIndex, const FMaterialRenderProxy* WireframeRenderProxy, uint8 InDepthPriorityGroup, bool bAllowPreCulledIndices, FMeshBatch& OutMeshBatch) const override;
 
+	/** Sets up a collision FMeshBatch for a specific LOD and element. */
+	virtual bool GetCollisionMeshElement(int32 LODIndex, int32 BatchIndex, int32 ElementIndex, uint8 InDepthPriorityGroup, const FMaterialRenderProxy* RenderProxy, FMeshBatch& OutMeshBatch) const override;
+
 #if RHI_RAYTRACING
 	virtual bool HasRayTracingRepresentation() const override { return false; }
 	virtual bool IsRayTracingRelevant() const override final { return false; }
@@ -117,6 +120,10 @@ public:
 
 	// 	  virtual uint32 GetMemoryFootprint( void ) const { return 0; }
 
+private:
+	void SetupMeshBatchForSpline(int32 InLODIndex, FMeshBatch& OutMeshBatch) const;
+
+public:
 	/** Parameters that define the spline, used to deform mesh */
 	FSplineMeshParams SplineParams;
 	/** Axis (in component space) that is used to determine X axis for co-ordinates along spline */
