Level Instance: implement edit asset functionnality (Ctrl+E).

#rnx

diff --git a/Engine/Source/Developer/AssetTools/Private/AssetTypeActions/AssetTypeActions_Actor.cpp b/Engine/Source/Developer/AssetTools/Private/AssetTypeActions/AssetTypeActions_Actor.cpp
new file mode 100644
--- /dev/null
+++ b/Engine/Source/Developer/AssetTools/Private/AssetTypeActions/AssetTypeActions_Actor.cpp
@@ -0,0 +1,21 @@
+// Copyright Epic Games, Inc. All Rights Reserved.
+
+#include "AssetTypeActions/AssetTypeActions_Actor.h"
+
+#define LOCTEXT_NAMESPACE "AssetTypeActions"
+
+void FAssetTypeActions_Actor::OpenAssetEditor(const TArray<UObject*>& InObjects, TSharedPtr<class IToolkitHost> EditWithinLevelEditor)
+{
+	for(UObject* Object : InObjects)
+	{
+		if (AActor* Actor = Cast<AActor>(Object))
+		{
+			if (Actor->OpenAssetEditor())
+			{
+				return;
+			}
+		}
+	}
+}
+
+#undef LOCTEXT_NAMESPACE
diff --git a/Engine/Source/Developer/AssetTools/Private/AssetTypeActions/AssetTypeActions_Actor.h b/Engine/Source/Developer/AssetTools/Private/AssetTypeActions/AssetTypeActions_Actor.h
--- a/Engine/Source/Developer/AssetTools/Private/AssetTypeActions/AssetTypeActions_Actor.h
+++ b/Engine/Source/Developer/AssetTools/Private/AssetTypeActions/AssetTypeActions_Actor.h
@@ -15,4 +15,5 @@ public:
 	virtual UClass* GetSupportedClass() const override { return AActor::StaticClass(); }
 	virtual uint32 GetCategories() override { return EAssetTypeCategories::None; }
 	virtual FString GetObjectDisplayName(UObject* Object) const override { return CastChecked<AActor>(Object)->GetActorLabel(); }
+	virtual void OpenAssetEditor( const TArray<UObject*>& InObjects, TSharedPtr<class IToolkitHost> EditWithinLevelEditor = TSharedPtr<IToolkitHost>() ) override;
 };
diff --git a/Engine/Source/Runtime/Engine/Classes/GameFramework/Actor.h b/Engine/Source/Runtime/Engine/Classes/GameFramework/Actor.h
--- a/Engine/Source/Runtime/Engine/Classes/GameFramework/Actor.h
+++ b/Engine/Source/Runtime/Engine/Classes/GameFramework/Actor.h
@@ -2393,6 +2393,9 @@ public:
 	/** Similar to GetReferencedContentObjects, but for soft referenced objects */
 	virtual bool GetSoftReferencedContentObjects(TArray<FSoftObjectPath>& SoftObjects) const;
 
+	/** Used to allow actor classes to open an actor specific editor */
+	virtual bool OpenAssetEditor() { return false; }
+
 	/** Returns NumUncachedStaticLightingInteractions for this actor */
 	const int32 GetNumUncachedStaticLightingInteractions() const;
 
diff --git a/Engine/Source/Runtime/Engine/Private/LevelInstance/LevelInstanceActor.cpp b/Engine/Source/Runtime/Engine/Private/LevelInstance/LevelInstanceActor.cpp
--- a/Engine/Source/Runtime/Engine/Private/LevelInstance/LevelInstanceActor.cpp
+++ b/Engine/Source/Runtime/Engine/Private/LevelInstance/LevelInstanceActor.cpp
@@ -250,6 +250,12 @@ bool ALevelInstance::IsLockLocation() const
 	return Super::IsLockLocation() || LevelInstanceActorImpl.IsLockLocation();
 }
 
+bool ALevelInstance::GetReferencedContentObjects(TArray<UObject*>& Objects) const
+{
+	Objects.Add(const_cast<ALevelInstance*>(this));
+	return true;
+}
+
 bool ALevelInstance::GetSoftReferencedContentObjects(TArray<FSoftObjectPath>& SoftObjects) const
 {
 	if (WorldAsset.ToSoftObjectPath().IsValid())
@@ -260,6 +266,11 @@ bool ALevelInstance::GetSoftReferencedContentObjects(TArray<FSoftObjectPath>& So
 	return false;
 }
 
+bool ALevelInstance::OpenAssetEditor()
+{
+	return CanEnterEdit() ? EnterEdit() : false;
+}
+
 FBox ALevelInstance::GetComponentsBoundingBox(bool bNonColliding, bool bIncludeFromChildActors) const
 {
 	FBox Box = Super::GetComponentsBoundingBox(bNonColliding, bIncludeFromChildActors);
diff --git a/Engine/Source/Runtime/Engine/Public/LevelInstance/LevelInstanceActor.h b/Engine/Source/Runtime/Engine/Public/LevelInstance/LevelInstanceActor.h
--- a/Engine/Source/Runtime/Engine/Public/LevelInstance/LevelInstanceActor.h
+++ b/Engine/Source/Runtime/Engine/Public/LevelInstance/LevelInstanceActor.h
@@ -94,7 +94,9 @@ public:
 	virtual FBox GetComponentsBoundingBox(bool bNonColliding = false, bool bIncludeFromChildActors = false) const override;
 	virtual FBox GetStreamingBounds() const override;
 	virtual bool IsLockLocation() const override;
+	virtual bool GetReferencedContentObjects(TArray<UObject*>& Objects) const override;
 	virtual bool GetSoftReferencedContentObjects(TArray<FSoftObjectPath>& SoftObjects) const override;
+	virtual bool OpenAssetEditor() override;
 	// End of AActor interface
 
 	DECLARE_MULTICAST_DELEGATE_OneParam(FOnLevelInstanceActorPostLoad, ALevelInstance*);
