Updating comments relating to changes to remove usage of master/slave in GameFramework-JIRA

diff --git a/Engine/Source/Runtime/Engine/Classes/GameFramework/AsyncActionHandleSaveGame.h b/Engine/Source/Runtime/Engine/Classes/GameFramework/AsyncActionHandleSaveGame.h
--- a/Engine/Source/Runtime/Engine/Classes/GameFramework/AsyncActionHandleSaveGame.h
+++ b/Engine/Source/Runtime/Engine/Classes/GameFramework/AsyncActionHandleSaveGame.h
@@ -26,7 +26,7 @@ public:
 	 *
 	 * @param SaveGameObject	Object that contains data about the save game that we want to write out.
 	 * @param SlotName			Name of the save game slot to load from.
-	 * @param UserIndex			For some platforms, primary user index to identify the user doing the loading.
+	 * @param UserIndex			The platform user index that identifies the user doing the saving, ignored on some platforms.
 	 */
 	UFUNCTION(BlueprintCallable, meta=(BlueprintInternalUseOnly="true", Category = "SaveGame",  WorldContext = "WorldContextObject"))
 	static UAsyncActionHandleSaveGame* AsyncSaveGameToSlot(UObject* WorldContextObject, USaveGame* SaveGameObject, const FString& SlotName, const int32 UserIndex);
@@ -37,7 +37,7 @@ public:
 	 * Keep in mind that some platforms may not support trying to load and save at the same time.
 	 *
 	 * @param SlotName			Name of the save game slot to load from.
-	 * @param UserIndex			For some platforms, primary user index to identify the user doing the loading.
+	 * @param UserIndex			The platform user index that identifies the user doing the saving, ignored on some platforms.
 	 */
 	UFUNCTION(BlueprintCallable, meta = (BlueprintInternalUseOnly = "true", Category = "SaveGame", WorldContext = "WorldContextObject"))
 	static UAsyncActionHandleSaveGame* AsyncLoadGameFromSlot(UObject* WorldContextObject, const FString& SlotName, const int32 UserIndex);
diff --git a/Engine/Source/Runtime/Engine/Classes/Kismet/GameplayStatics.h b/Engine/Source/Runtime/Engine/Classes/Kismet/GameplayStatics.h
--- a/Engine/Source/Runtime/Engine/Classes/Kismet/GameplayStatics.h
+++ b/Engine/Source/Runtime/Engine/Classes/Kismet/GameplayStatics.h
@@ -1095,7 +1095,7 @@ public:
 	 * Save the contents of the buffer to a platform-specific save slot/file 
 	 * @param InSaveData		Data to save
 	 * @param SlotName			Name of save game slot to save to.
-	 * @param UserIndex			For some platforms, primary user index to identify the user doing the saving.
+	 * @param UserIndex			The platform user index that identifies the user doing the saving, ignored on some platforms.
 	 * @return					Whether we successfully saved this information
 	 */
 	static bool SaveDataToSlot(const TArray<uint8>& InSaveData, const FString& SlotName, const int32 UserIndex);
@@ -1107,7 +1107,7 @@ public:
 	 *
 	 * @param SaveGameObject	Object that contains data about the save game that we want to write out.
 	 * @param SlotName			Name of the save game slot to load from.
-	 * @param UserIndex			For some platforms, primary user index to identify the user doing the loading.
+	 * @param UserIndex			The platform user index that identifies the user doing the saving, ignored on some platforms.
 	 * @param SavedDelegate		Delegate that will be called on game thread when save succeeds or fails.
 	 */
 	static void AsyncSaveGameToSlot(USaveGame* SaveGameObject, const FString& SlotName, const int32 UserIndex, FAsyncSaveGameToSlotDelegate SavedDelegate = FAsyncSaveGameToSlotDelegate());
@@ -1118,7 +1118,7 @@ public:
 	 *
 	 * @param SaveGameObject	Object that contains data about the save game that we want to write out
 	 * @param SlotName			Name of save game slot to save to.
-	 * @param UserIndex			For some platforms, primary user index to identify the user doing the saving.
+	 * @param UserIndex			The platform user index that identifies the user doing the saving, ignored on some platforms.
 	 * @return					Whether we successfully saved this information
 	 */
 	UFUNCTION(BlueprintCallable, Category="SaveGame")
@@ -1127,7 +1127,7 @@ public:
 	/**
 	 * See if a save game exists with the specified name.
 	 * @param SlotName			Name of save game slot.
-	 * @param UserIndex			For some platforms, primary user index to identify the user doing the saving.
+	 * @param UserIndex			The platform user index that identifies the user doing the saving, ignored on some platforms.
 	 */
 	UFUNCTION(BlueprintCallable, Category="SaveGame")
 	static bool DoesSaveGameExist(const FString& SlotName, const int32 UserIndex);
@@ -1143,7 +1143,7 @@ public:
 	 * Load contents from a slot/file into a buffer of save data.
 	 * @param OutSaveData		Data buffer to load into
 	 * @param SlotName			Name of save game slot to save to.
-	 * @param UserIndex			For some platforms, primary user index to identify the user doing the saving.
+	 * @param UserIndex			The platform user index that identifies the user doing the saving, ignored on some platforms.
 	 * @return					Whether valid save data was found and loaded.
 	 */
 	static bool LoadDataFromSlot(TArray<uint8>& OutSaveData, const FString& SlotName, const int32 UserIndex);
@@ -1154,7 +1154,7 @@ public:
 	 * The passed in delegate will be copied to a worker thread so make sure any payload is thread safe to copy by value
 	 *
 	 * @param SlotName			Name of the save game slot to load from.
-	 * @param UserIndex			For some platforms, primary user index to identify the user doing the loading.
+	 * @param UserIndex			The platform user index that identifies the user doing the saving, ignored on some platforms.
 	 * @param LoadedDelegate	Delegate that will be called on game thread when load succeeds or fails.
 	 */
 	static void AsyncLoadGameFromSlot(const FString& SlotName, const int32 UserIndex, FAsyncLoadGameFromSlotDelegate LoadedDelegate);
@@ -1162,7 +1162,7 @@ public:
 	/** 
 	 * Load the contents from a given slot.
 	 * @param SlotName			Name of the save game slot to load from.
-	 * @param UserIndex			For some platforms, primary user index to identify the user doing the loading.
+	 * @param UserIndex			The platform user index that identifies the user doing the saving, ignored on some platforms.
 	 * @return					Object containing loaded game state (nullptr if load fails)
 	 */
 	UFUNCTION(BlueprintCallable, Category="SaveGame")
@@ -1182,7 +1182,7 @@ public:
 	/**
 	 * Delete a save game in a particular slot.
 	 * @param SlotName			Name of save game slot to delete.
-	 * @param UserIndex			For some platforms, primary user index to identify the user doing the deletion.
+	 * @param UserIndex			The platform user index that identifies the user doing the saving, ignored on some platforms.
 	 * @return					True if a file was actually able to be deleted. use DoesSaveGameExist to distinguish between delete failures and failure due to file not existing.
 	 */
 	UFUNCTION(BlueprintCallable, Category = "SaveGame")
