Removing usage of master/slave in Player Camera Manager

diff --git a/Engine/Source/Runtime/Engine/Classes/Camera/PlayerCameraManager.h b/Engine/Source/Runtime/Engine/Classes/Camera/PlayerCameraManager.h
--- a/Engine/Source/Runtime/Engine/Classes/Camera/PlayerCameraManager.h
+++ b/Engine/Source/Runtime/Engine/Classes/Camera/PlayerCameraManager.h
@@ -302,7 +302,7 @@ public:
 	UPROPERTY(EditAnywhere, BlueprintReadWrite, Category = Debug)
 	FVector ViewTargetOffset;
 
-	/** If bound, broadcast on fade start (with fade time) instead of manually altering audio device's master volume directly */
+	/** If bound, broadcast on fade start (with fade time) instead of manually altering audio device's primary volume directly */
 	UPROPERTY(BlueprintAssignable, Transient)
 	FOnAudioFadeChangeSignature OnAudioFadeChangeEvent;
 
@@ -659,7 +659,7 @@ public:
 	virtual void UnlockOrthoWidth();
 
 	/**
-	 * Master function to retrieve Camera's actual view point.
+	 * Primary function to retrieve Camera's actual view point.
 	 * Consider calling PlayerController::GetPlayerViewPoint() instead.
 	 *
 	 * @param	OutCamLoc	Returned camera location
