Disable onscreen display of dynamic res subsystem metrics if they're disabled to reduce noise and unnecessary perf overhead (debug rendering is expensive, so this skews historical perf data)

diff --git a/Engine/Source/Runtime/Engine/Private/UnrealClient.cpp b/Engine/Source/Runtime/Engine/Private/UnrealClient.cpp
--- a/Engine/Source/Runtime/Engine/Private/UnrealClient.cpp
+++ b/Engine/Source/Runtime/Engine/Private/UnrealClient.cpp
@@ -728,7 +728,7 @@ int32 FStatUnitData::DrawStat(FViewport* InViewport, FCanvas* InCanvas, int32 In
 			{
 				const DynamicRenderScaling::FBudget& Budget = **BudgetIt;
 				const DynamicRenderScaling::FHeuristicSettings& HeuristicSettings = Budget.GetSettings();
-				if (Budget == GDynamicPrimaryResolutionFraction)
+				if (Budget == GDynamicPrimaryResolutionFraction || !HeuristicSettings.IsEnabled())
 				{
 					continue;
 				}
@@ -743,11 +743,7 @@ int32 FStatUnitData::DrawStat(FViewport* InViewport, FCanvas* InCanvas, int32 In
 				FColor Color = (ResolutionFraction < AlertResolutionFraction) ? StatRed : ((ResolutionFraction < FMath::Min(ResolutionFraction * 0.97f, 1.0f)) ? StatOrange : StatGreen);
 
 				DrawTitleString(*DisplayName, /* UnitGraphColor = */ FColor::White);
-				if (!HeuristicSettings.IsEnabled())
-				{
-					InCanvas->DrawShadowedString(X2, InY, TEXT("OFF"), Font, FColor(160, 160, 160));
-				}
-				else if (HeuristicSettings.Model == DynamicRenderScaling::EHeuristicModel::Quadratic)
+				if (HeuristicSettings.Model == DynamicRenderScaling::EHeuristicModel::Quadratic)
 				{
 					InCanvas->DrawShadowedString(X2, InY, *FString::Printf(TEXT("%3.1f%% x %3.1f%%"), ScreenPercentage, ScreenPercentage), Font, Color);
 				}
