Subobject replication
* Added utility function to tell if a replicated subobject is registered or not
* GameplayAbilities will assert if they send an RPC before being registered.

diff --git a/Engine/Plugins/Runtime/GameplayAbilities/Source/GameplayAbilities/Private/Abilities/GameplayAbility.cpp b/Engine/Plugins/Runtime/GameplayAbilities/Source/GameplayAbilities/Private/Abilities/GameplayAbility.cpp
--- a/Engine/Plugins/Runtime/GameplayAbilities/Source/GameplayAbilities/Private/Abilities/GameplayAbility.cpp
+++ b/Engine/Plugins/Runtime/GameplayAbilities/Source/GameplayAbilities/Private/Abilities/GameplayAbility.cpp
@@ -117,6 +117,15 @@ bool UGameplayAbility::CallRemoteFunction(UFunction* Function, void* Parameters,
 	FWorldContext* const Context = GEngine->GetWorldContextFromWorld(GetWorld());
 	if (Context != nullptr)
 	{
+#if DO_CHECK
+		if (UAbilitySystemComponent* AbilitySystemComponent = GetAbilitySystemComponentFromActorInfo())
+		{
+			// If the component is using the subobject list, detect subobjects that send RPC's before they are registered.
+			const bool bIsSubObjectRegistered = !AbilitySystemComponent->IsUsingRegisteredSubObjectList() || Owner->IsActorComponentReplicatedSubObjectRegistered(AbilitySystemComponent, this);
+			checkf(bIsSubObjectRegistered, TEXT("%s sent an RPC before it was registered as a subobject on %s::%s"), *GetName(), *Owner->GetName(), *AbilitySystemComponent->GetName());
+		}
+#endif
+		
 		for (FNamedNetDriver& Driver : Context->ActiveNetDrivers)
 		{
 			if (Driver.NetDriver != nullptr && Driver.NetDriver->ShouldReplicateFunction(Owner, Function))
diff --git a/Engine/Source/Runtime/Engine/Classes/Components/ActorComponent.h b/Engine/Source/Runtime/Engine/Classes/Components/ActorComponent.h
--- a/Engine/Source/Runtime/Engine/Classes/Components/ActorComponent.h
+++ b/Engine/Source/Runtime/Engine/Classes/Components/ActorComponent.h
@@ -523,6 +523,11 @@ public:
 	*/
 	void RemoveReplicatedSubObject(UObject* SubObject);
 
+	/**
+	* Tells if the object is registered to be replicated by this actor component.
+	*/
+	bool IsReplicatedSubObjectRegistered(const UObject* SubObject) const;
+
 #if WITH_EDITORONLY_DATA
 	/** Returns whether this component is an editor-only object or not */
 	virtual bool IsEditorOnly() const override { return bIsEditorOnly; }
diff --git a/Engine/Source/Runtime/Engine/Classes/GameFramework/Actor.h b/Engine/Source/Runtime/Engine/Classes/GameFramework/Actor.h
--- a/Engine/Source/Runtime/Engine/Classes/GameFramework/Actor.h
+++ b/Engine/Source/Runtime/Engine/Classes/GameFramework/Actor.h
@@ -3679,6 +3679,15 @@ public:
 	*/
 	void RemoveActorComponentReplicatedSubObject(UActorComponent* OwnerComponent, UObject* SubObject);
 
+	/** Tells if the object has been registered as a replicated subobject of this actor */
+	bool IsReplicatedSubObjectRegistered(const UObject* SubObject) const;
+
+	/** Tells if the component has been registered as a replicated component */
+	bool IsReplicatedActorComponentRegistered(const UActorComponent* ReplicatedComponent) const;
+
+	/** Tells if an object owned by a component has been registered as a replicated subobject of the component */
+	bool IsActorComponentReplicatedSubObjectRegistered(const UActorComponent* OwnerComponent, const UObject* SubObject) const;
+
 private:
 	/** Collection of SubObjects that get replicated when this actor gets replicated. */
 	UE::Net::FSubObjectRegistry ReplicatedSubObjects;
@@ -3693,6 +3702,9 @@ private:
 	/** Remove a component from the replicated list */
 	void RemoveReplicatedComponent(UActorComponent* Component);
 
+	/** Constructs the list of replicated components and their netcondition */
+	void BuildReplicatedComponentsInfo();
+
 protected:
 	/** Set of replicated components, stored as an array to save space as this is generally not very large */
 	TArray<UActorComponent*> ReplicatedComponents;
diff --git a/Engine/Source/Runtime/Engine/Private/Actor.cpp b/Engine/Source/Runtime/Engine/Private/Actor.cpp
--- a/Engine/Source/Runtime/Engine/Private/Actor.cpp
+++ b/Engine/Source/Runtime/Engine/Private/Actor.cpp
@@ -4012,19 +4012,7 @@ void AActor::DispatchBeginPlay(bool bFromLevelStreaming)
 		bActorBeginningPlayFromLevelStreaming = bFromLevelStreaming;
 		ActorHasBegunPlay = EActorBeginPlayState::BeginningPlay;
 
-		// Ask the actor class if they want to override any replicated components
-		for (UActorComponent* ReplicatedComponent : ReplicatedComponents)
-		{
-			const ELifetimeCondition NetCondition = AllowActorComponentToReplicate(ReplicatedComponent);
-
-			const int32 Index = ReplicatedComponentsInfo.AddUnique(UE::Net::FReplicatedComponentInfo(ReplicatedComponent));
-			ReplicatedComponentsInfo[Index].NetCondition = NetCondition;
-
-			if (!ReplicatedComponent->IsReadyForReplication())
-			{
-				ReplicatedComponent->ReadyForReplication();
-			}
-		}
+		BuildReplicatedComponentsInfo();
 
 #if UE_WITH_IRIS
 		BeginReplication();
diff --git a/Engine/Source/Runtime/Engine/Private/ActorReplication.cpp b/Engine/Source/Runtime/Engine/Private/ActorReplication.cpp
--- a/Engine/Source/Runtime/Engine/Private/ActorReplication.cpp
+++ b/Engine/Source/Runtime/Engine/Private/ActorReplication.cpp
@@ -678,7 +678,7 @@ void AActor::RemoveActorComponentReplicatedSubObject(UActorComponent* OwnerCompo
 	if (FReplicatedComponentInfo* ComponentInfo = ReplicatedComponentsInfo.FindByKey(OwnerComponent))
 	{
 		bool bWasRemoved = ComponentInfo->SubObjects.RemoveSubObject(SubObject);
-		
+
 		UE_CLOG(bWasRemoved, LogNetSubObject, Verbose, TEXT("%s::%s (0x%p) removed replicated subobject %s (0x%p)"), *GetName(), *OwnerComponent->GetName(), OwnerComponent, *SubObject->GetName(), SubObject);
 
 #if UE_WITH_IRIS
@@ -687,6 +687,46 @@ void AActor::RemoveActorComponentReplicatedSubObject(UActorComponent* OwnerCompo
 	}
 }
 
+bool AActor::IsReplicatedSubObjectRegistered(const UObject* SubObject) const
+{
+	return ReplicatedSubObjects.IsSubObjectInRegistry(SubObject);
+}
+
+bool AActor::IsReplicatedActorComponentRegistered(const UActorComponent* ReplicatedComponent) const
+{
+	return ReplicatedComponentsInfo.FindByKey(ReplicatedComponent) != nullptr;
+}
+
+bool AActor::IsActorComponentReplicatedSubObjectRegistered(const UActorComponent* OwnerComponent, const UObject* SubObject) const
+{
+	check(OwnerComponent);
+
+	if (const FReplicatedComponentInfo* ComponentInfo = ReplicatedComponentsInfo.FindByKey(OwnerComponent))
+	{
+		return ComponentInfo->SubObjects.IsSubObjectInRegistry(SubObject);
+	}
+	return false;
+}
+
+void AActor::BuildReplicatedComponentsInfo()
+{
+	checkf(HasActorBegunPlay() == false, TEXT("BuildReplicatedComponentsInfo can only be called before BeginPlay."));
+
+	for (UActorComponent* ReplicatedComponent : ReplicatedComponents)
+	{
+		// Ask the actor if they want to override the replicated component
+		const ELifetimeCondition NetCondition = AllowActorComponentToReplicate(ReplicatedComponent);
+
+		const int32 Index = ReplicatedComponentsInfo.AddUnique(UE::Net::FReplicatedComponentInfo(ReplicatedComponent));
+		ReplicatedComponentsInfo[Index].NetCondition = NetCondition;
+
+		if (!ReplicatedComponent->IsReadyForReplication())
+		{
+			ReplicatedComponent->ReadyForReplication();
+		}
+	}
+}
+
 void AActor::GetSubobjectsWithStableNamesForNetworking(TArray<UObject*> &ObjList)
 {	
 	// For experimenting with replicating ALL stably named components initially
diff --git a/Engine/Source/Runtime/Engine/Private/Components/ActorComponent.cpp b/Engine/Source/Runtime/Engine/Private/Components/ActorComponent.cpp
--- a/Engine/Source/Runtime/Engine/Private/Components/ActorComponent.cpp
+++ b/Engine/Source/Runtime/Engine/Private/Components/ActorComponent.cpp
@@ -2123,6 +2123,16 @@ void UActorComponent::RemoveReplicatedSubObject(UObject* SubObject)
 	}
 }
 
+bool UActorComponent::IsReplicatedSubObjectRegistered(const UObject* SubObject) const
+{
+	if (AActor* MyOwner=GetOwner())
+	{
+		return MyOwner->IsActorComponentReplicatedSubObjectRegistered(this, SubObject);
+	}
+
+	return false;
+}
+
 bool UActorComponent::ReplicateSubobjects(UActorChannel* Channel, FOutBunch* Bunch, FReplicationFlags* RepFlags)
 {
 	UActorChannel::SetCurrentSubObjectOwner(this);
diff --git a/Engine/Source/Runtime/Net/Core/Private/Net/Core/Misc/NetSubObjectRegistry.cpp b/Engine/Source/Runtime/Net/Core/Private/Net/Core/Misc/NetSubObjectRegistry.cpp
--- a/Engine/Source/Runtime/Net/Core/Private/Net/Core/Misc/NetSubObjectRegistry.cpp
+++ b/Engine/Source/Runtime/Net/Core/Private/Net/Core/Misc/NetSubObjectRegistry.cpp
@@ -45,7 +45,7 @@ bool FSubObjectRegistry::RemoveSubObject(UObject* InSubObject)
 	return false;
 }
 
-bool FSubObjectRegistry::IsSubObjectInRegistry(UObject* SubObject) const
+bool FSubObjectRegistry::IsSubObjectInRegistry(const UObject* SubObject) const
 {
 	return Registry.FindByKey(SubObject) != nullptr;
 }
diff --git a/Engine/Source/Runtime/Net/Core/Public/Net/Core/Misc/NetSubObjectRegistry.h b/Engine/Source/Runtime/Net/Core/Public/Net/Core/Misc/NetSubObjectRegistry.h
--- a/Engine/Source/Runtime/Net/Core/Public/Net/Core/Misc/NetSubObjectRegistry.h
+++ b/Engine/Source/Runtime/Net/Core/Public/Net/Core/Misc/NetSubObjectRegistry.h
@@ -63,14 +63,14 @@ public:
 		}
 
 		bool operator==(const FEntry& rhs) const { return Key == rhs.Key; }
-		bool operator==(UObject* rhs) const { return SubObject == rhs; }
+		bool operator==(const UObject* rhs) const { return SubObject == rhs; }
 	};
 
 	/** Returns the list of registered subobjects */
 	const TArray<FSubObjectRegistry::FEntry>& GetRegistryList() const { return Registry; }
 
 	/** Returns true if the subobject is contained in this list */
-	bool IsSubObjectInRegistry(UObject* SubObject) const;
+	bool IsSubObjectInRegistry(const UObject* SubObject) const;
 
 private:
 
