World Partition: temporary workaround, don't maintain dirty actor packages alive when streaming is disabled.

#rnx

diff --git a/Engine/Source/Runtime/Engine/Private/WorldPartition/WorldPartition.cpp b/Engine/Source/Runtime/Engine/Private/WorldPartition/WorldPartition.cpp
--- a/Engine/Source/Runtime/Engine/Private/WorldPartition/WorldPartition.cpp
+++ b/Engine/Source/Runtime/Engine/Private/WorldPartition/WorldPartition.cpp
@@ -272,23 +272,26 @@ void UWorldPartition::OnGCPostReachabilityAnalysis()
 
 void UWorldPartition::OnPackageDirtyStateChanged(UPackage* Package)
 {
-	if (AActor* Actor = AActor::FindActorInPackage(Package); ShouldHandleActorEvent(Actor))
+	if (IsStreamingEnabled())
 	{
-		if (FWorldPartitionHandle ActorHandle(this, Actor->GetActorGuid()); ActorHandle.IsValid())
+		if (AActor* Actor = AActor::FindActorInPackage(Package); ShouldHandleActorEvent(Actor))
 		{
-			if (Package->IsDirty())
+			if (FWorldPartitionHandle ActorHandle(this, Actor->GetActorGuid()); ActorHandle.IsValid())
 			{
-				DirtyActors.Add(ActorHandle.ToReference(), Actor);
-			}
-			else
-			{
-				// If we hold the last reference to that actor (or no reference are held at all), pin it to avoid unloading
-				if (PinnedActors && (ActorHandle->GetHardRefCount() <= 1))
+				if (Package->IsDirty())
 				{
-					PinnedActors->AddActors({ ActorHandle });
+					DirtyActors.Add(ActorHandle.ToReference(), Actor);
 				}
+				else
+				{
+					// If we hold the last reference to that actor (or no reference are held at all), pin it to avoid unloading
+					if (PinnedActors && (ActorHandle->GetHardRefCount() <= 1))
+					{
+						PinnedActors->AddActors({ ActorHandle });
+					}
 
-				DirtyActors.Remove(ActorHandle.ToReference());
+					DirtyActors.Remove(ActorHandle.ToReference());
+				}
 			}
 		}
 	}
