Added units onto SLW output node.

diff --git a/Engine/Source/Runtime/Engine/Classes/Materials/MaterialExpressionSingleLayerWaterMaterialOutput.h b/Engine/Source/Runtime/Engine/Classes/Materials/MaterialExpressionSingleLayerWaterMaterialOutput.h
--- a/Engine/Source/Runtime/Engine/Classes/Materials/MaterialExpressionSingleLayerWaterMaterialOutput.h
+++ b/Engine/Source/Runtime/Engine/Classes/Materials/MaterialExpressionSingleLayerWaterMaterialOutput.h
@@ -13,11 +13,11 @@ class UMaterialExpressionSingleLayerWaterMaterialOutput : public UMaterialExpres
 {
 	GENERATED_UCLASS_BODY()
 
-	/** Input for scattering coefficient describing how light scatter around and is absorbed. Valid range is [0,+inf[. */
+	/** Input for scattering coefficient describing how light scatter around and is absorbed. Valid range is [0,+inf[. Unit is 1/cm. */
 	UPROPERTY()
 	FExpressionInput ScatteringCoefficients;
 
-	/** Input for scattering coefficient describing how light bounce is absorbed. Valid range is [0,+inf[. */
+	/** Input for scattering coefficient describing how light bounce is absorbed. Valid range is [0,+inf[. Unit is 1/cm. */
 	UPROPERTY()
 	FExpressionInput AbsorptionCoefficients;
 		
