Improved controls for DistanceFieldApproxAO node.
- In the current version, modifying any of the node parameters can drastically change the output (particularly changing number of steps). This is because the main contributor to the result is the first sample, however it's position is affected by the three parameters and so changes to the parameters have unintuitive impact on the result.
- This CL adds a new parameter to the node, BaseDistance, which is used to directly control the position of the first sample. The rest of the samples are distributed between BaseDistance and Radius based on StepScale parameter. When NumSamples = 1, the single sample is placed at BaseDistance and the other parameters are ignored.

diff --git a/Engine/Source/Runtime/Engine/Classes/Materials/MaterialExpressionDistanceFieldApproxAO.h b/Engine/Source/Runtime/Engine/Classes/Materials/MaterialExpressionDistanceFieldApproxAO.h
--- a/Engine/Source/Runtime/Engine/Classes/Materials/MaterialExpressionDistanceFieldApproxAO.h
+++ b/Engine/Source/Runtime/Engine/Classes/Materials/MaterialExpressionDistanceFieldApproxAO.h
@@ -20,10 +20,17 @@ class UMaterialExpressionDistanceFieldApproxAO : public UMaterialExpression
 	UPROPERTY(meta = (RequiredInput = "false", ToolTip = "Defaults to current world normal if not specified"))
 	FExpressionInput Normal;
 
+	UPROPERTY(meta = (RequiredInput = "false"))
+	FExpressionInput BaseDistance;
+
+	/** only used if BaseDistance is not hooked up */
+	UPROPERTY(EditAnywhere, Category = MaterialExpressionDistanceFieldApproxAO, meta=(OverridingInputProperty = "BaseDistance"))
+	float BaseDistanceDefault;
+
 	UPROPERTY(meta = (RequiredInput = "false"))
 	FExpressionInput Radius;
 
-	/** only used if TextureSize is not hooked up */
+	/** only used if Radius is not hooked up */
 	UPROPERTY(EditAnywhere, Category = MaterialExpressionDistanceFieldApproxAO, meta=(OverridingInputProperty = "Radius"))
 	float RadiusDefault;
 
diff --git a/Engine/Source/Runtime/Engine/Private/Materials/HLSLMaterialTranslator.cpp b/Engine/Source/Runtime/Engine/Private/Materials/HLSLMaterialTranslator.cpp
--- a/Engine/Source/Runtime/Engine/Private/Materials/HLSLMaterialTranslator.cpp
+++ b/Engine/Source/Runtime/Engine/Private/Materials/HLSLMaterialTranslator.cpp
@@ -9058,7 +9058,7 @@ int32 FHLSLMaterialTranslator::DistanceFieldGradient(int32 PositionArg)
 	return AddCodeChunk(MCT_Float3, TEXT("GetDistanceFieldGradientGlobal(%s)"), *CoerceParameter(PositionArg, MCT_Float3));
 }
 
-int32 FHLSLMaterialTranslator::DistanceFieldApproxAO(int32 PositionArg, int32 NormalArg, int32 RadiusArg, uint32 NumSteps, float StepScale)
+int32 FHLSLMaterialTranslator::DistanceFieldApproxAO(int32 PositionArg, int32 NormalArg, int32 BaseDistanceArg, int32 RadiusArg, uint32 NumSteps, float StepScale)
 {
 	if (ErrorUnlessFeatureLevelSupported(ERHIFeatureLevel::SM5) == INDEX_NONE)
 	{
@@ -9075,6 +9075,11 @@ int32 FHLSLMaterialTranslator::DistanceFieldApproxAO(int32 PositionArg, int32 No
 		return INDEX_NONE;
 	}
 
+	if (BaseDistanceArg == INDEX_NONE)
+	{
+		return INDEX_NONE;
+	}
+
 	if (RadiusArg == INDEX_NONE)
 	{
 		return INDEX_NONE;
@@ -9085,12 +9090,26 @@ int32 FHLSLMaterialTranslator::DistanceFieldApproxAO(int32 PositionArg, int32 No
 	NumSteps = FMath::Clamp(NumSteps, 1, 4);
 	StepScale = FMath::Max(StepScale, 1.0f);
 
-	int32 NumStepsConst = Constant(NumSteps);
-	int32 NumStepsMinus1Const = Constant(NumSteps - 1);
-	int32 StepScaleConst = Constant(StepScale);
+	const int32 NumStepsConst = Constant(NumSteps);
+	const int32 NumStepsMinus1Const = Constant(NumSteps - 1);
+	const int32 StepScaleConst = Constant(StepScale);
+
+	int32 StepDistance;
+	int32 DistanceBias;
+	int32 MaxDistance;
 
-	int32 StepDistance = Div(RadiusArg, Power(StepScaleConst, NumStepsMinus1Const));
-	int32 DistanceBiasConst = Constant(0);
+	if (NumSteps == 1)
+	{
+		StepDistance = Constant(0);
+		DistanceBias = BaseDistanceArg;
+		MaxDistance = BaseDistanceArg;
+	}
+	else
+	{
+		StepDistance = Div(Sub(RadiusArg, BaseDistanceArg), Sub(Power(StepScaleConst, NumStepsMinus1Const), Constant(1)));
+		DistanceBias = Sub(BaseDistanceArg, StepDistance);
+		MaxDistance = RadiusArg;
+	}
 
 	// LWC_TODO: update for LWC position
 	return AddCodeChunk(MCT_Float,
@@ -9100,8 +9119,8 @@ int32 FHLSLMaterialTranslator::DistanceFieldApproxAO(int32 PositionArg, int32 No
 		*GetParameterCode(NumStepsConst),
 		*CoerceParameter(StepDistance, MCT_Float),
 		*GetParameterCode(StepScaleConst),
-		*GetParameterCode(DistanceBiasConst),
-		*CoerceParameter(RadiusArg, MCT_Float));
+		*CoerceParameter(DistanceBias, MCT_Float),
+		*CoerceParameter(MaxDistance, MCT_Float));
 }
 
 int32 FHLSLMaterialTranslator::SamplePhysicsField(int32 PositionArg, const int32 OutputType, const int32 TargetIndex)
diff --git a/Engine/Source/Runtime/Engine/Private/Materials/HLSLMaterialTranslator.h b/Engine/Source/Runtime/Engine/Private/Materials/HLSLMaterialTranslator.h
--- a/Engine/Source/Runtime/Engine/Private/Materials/HLSLMaterialTranslator.h
+++ b/Engine/Source/Runtime/Engine/Private/Materials/HLSLMaterialTranslator.h
@@ -976,7 +976,7 @@ protected:
 	virtual int32 GetHairColorFromMelanin(int32 Melanin, int32 Redness, int32 DyeColor) override;
 	virtual int32 DistanceToNearestSurface(int32 PositionArg) override;
 	virtual int32 DistanceFieldGradient(int32 PositionArg) override;
-	virtual int32 DistanceFieldApproxAO(int32 PositionArg, int32 NormalArg, int32 RadiusArg, uint32 NumSteps, float StepScale) override;
+	virtual int32 DistanceFieldApproxAO(int32 PositionArg, int32 NormalArg, int32 BaseDistanceArg, int32 RadiusArg, uint32 NumSteps, float StepScale) override;
 	virtual int32 SamplePhysicsField(int32 PositionArg, const int32 OutputType, const int32 TargetIndex) override;
 	virtual int32 AtmosphericFogColor(int32 WorldPosition) override;
 	virtual int32 AtmosphericLightVector() override;
diff --git a/Engine/Source/Runtime/Engine/Private/Materials/MaterialExpressions.cpp b/Engine/Source/Runtime/Engine/Private/Materials/MaterialExpressions.cpp
--- a/Engine/Source/Runtime/Engine/Private/Materials/MaterialExpressions.cpp
+++ b/Engine/Source/Runtime/Engine/Private/Materials/MaterialExpressions.cpp
@@ -18415,6 +18415,7 @@ UMaterialExpressionDistanceFieldApproxAO::UMaterialExpressionDistanceFieldApprox
 #endif
 
 	NumSteps = 1;
+	BaseDistanceDefault = 15;
 	RadiusDefault = 150;
 	StepScaleDefault = 3.0f;
 }
@@ -18444,9 +18445,10 @@ int32 UMaterialExpressionDistanceFieldApproxAO::Compile(class FMaterialCompiler*
 		NormalArg = Compiler->VertexNormal();
 	}
 
+	int32 BaseDistanceArg = BaseDistance.GetTracedInput().Expression ? BaseDistance.Compile(Compiler) : Compiler->Constant(BaseDistanceDefault);
 	int32 RadiusArg = Radius.GetTracedInput().Expression ? Radius.Compile(Compiler) : Compiler->Constant(RadiusDefault);
 
-	return Compiler->DistanceFieldApproxAO(PositionArg, NormalArg, RadiusArg, NumSteps, StepScaleDefault);
+	return Compiler->DistanceFieldApproxAO(PositionArg, NormalArg, BaseDistanceArg, RadiusArg, NumSteps, StepScaleDefault);
 }
 
 void UMaterialExpressionDistanceFieldApproxAO::GetCaption(TArray<FString>& OutCaptions) const
diff --git a/Engine/Source/Runtime/Engine/Public/MaterialCompiler.h b/Engine/Source/Runtime/Engine/Public/MaterialCompiler.h
--- a/Engine/Source/Runtime/Engine/Public/MaterialCompiler.h
+++ b/Engine/Source/Runtime/Engine/Public/MaterialCompiler.h
@@ -503,7 +503,7 @@ public:
 	virtual int32 BlackBody( int32 Temp ) = 0;
 	virtual int32 DistanceToNearestSurface(int32 PositionArg) = 0;
 	virtual int32 DistanceFieldGradient(int32 PositionArg) = 0;
-	virtual int32 DistanceFieldApproxAO(int32 PositionArg, int32 NormalArg, int32 RadiusArg, uint32 NumSteps, float StepScale) = 0;
+	virtual int32 DistanceFieldApproxAO(int32 PositionArg, int32 NormalArg, int32 BaseDistanceArg, int32 RadiusArg, uint32 NumSteps, float StepScale) = 0;
 	virtual int32 SamplePhysicsField(int32 PositionArg, const int32 OutputType, const int32 TargetIndex) = 0;
 	virtual int32 DepthOfFieldFunction(int32 Depth, int32 FunctionValueIndex) = 0;
 	virtual int32 AtmosphericFogColor(int32 WorldPosition) = 0;
@@ -905,7 +905,7 @@ public:
 	virtual int32 BlackBody(int32 Temp) override { return Compiler->BlackBody(Temp); }
 	virtual int32 DistanceToNearestSurface(int32 PositionArg) override { return Compiler->DistanceToNearestSurface(PositionArg); }
 	virtual int32 DistanceFieldGradient(int32 PositionArg) override { return Compiler->DistanceFieldGradient(PositionArg); }
-	virtual int32 DistanceFieldApproxAO(int32 PositionArg, int32 NormalArg, int32 RadiusArg, uint32 NumSteps, float StepScale) override { return Compiler->DistanceFieldApproxAO(PositionArg, NormalArg, RadiusArg, NumSteps, StepScale); }
+	virtual int32 DistanceFieldApproxAO(int32 PositionArg, int32 NormalArg, int32 BaseDistanceArg, int32 RadiusArg, uint32 NumSteps, float StepScale) override { return Compiler->DistanceFieldApproxAO(PositionArg, NormalArg, BaseDistanceArg, RadiusArg, NumSteps, StepScale); }
 	virtual int32 SamplePhysicsField(int32 PositionArg, const int32 OutputType, const int32 TargetIndex)  override { return Compiler->SamplePhysicsField(PositionArg, OutputType, TargetIndex); }
 	virtual int32 PerInstanceRandom() override { return Compiler->PerInstanceRandom(); }
 	virtual int32 PerInstanceFadeAmount() override { return Compiler->PerInstanceFadeAmount(); }
