SkeletalMesh - Remove the now unused GetSkeletalMesh() from the USkeletalMeshComponent API.

GetSkeletalMesh() has been introduced only recently in 5.1, and therefore doesn't need to be deprecated to be removed.
It has since been renamed GetSkeletalMeshAsset() in order to match the setter function, SetSkeletalMeshAsset(), that couldn't be named SetSkeletalMesh() due to already having such a named function doing a slightly different thing (setting the asset and resetting/updating the component and animation).
Licensees will not be affected by this change since GetSkeletalMesh() has never been released in 5.0.
All code have already been updated, and projects shouldn't be affected by this change.

diff --git a/Engine/Source/Runtime/Engine/Classes/Components/SkeletalMeshComponent.h b/Engine/Source/Runtime/Engine/Classes/Components/SkeletalMeshComponent.h
--- a/Engine/Source/Runtime/Engine/Classes/Components/SkeletalMeshComponent.h
+++ b/Engine/Source/Runtime/Engine/Classes/Components/SkeletalMeshComponent.h
@@ -1899,9 +1899,6 @@ public:
 	 */
 	virtual void SetSkinnedAssetAndUpdate(class USkinnedAsset* InSkinnedAsset, bool bReinitPose = true) override;
 
-	// GetSkeletalMesh() is soon to be removed, use GetSkeletalMeshAsset() instead
-	class USkeletalMesh* GetSkeletalMesh() const { return GetSkeletalMeshAsset(); }
-
 	static FVector3f GetSkinnedVertexPosition(USkeletalMeshComponent* Component, int32 VertexIndex, const FSkeletalMeshLODRenderData& Model, const FSkinWeightVertexBuffer& SkinWeightBuffer);
 	static FVector3f GetSkinnedVertexPosition(USkeletalMeshComponent* Component, int32 VertexIndex, const FSkeletalMeshLODRenderData& Model, const FSkinWeightVertexBuffer& SkinWeightBuffer, TArray<FMatrix44f>& CachedRefToLocals);
 	static void ComputeSkinnedPositions(USkeletalMeshComponent* Component, TArray<FVector3f> & OutPositions, TArray<FMatrix44f>& CachedRefToLocals, const FSkeletalMeshLODRenderData& Model, const FSkinWeightVertexBuffer& SkinWeightBuffer);
