Fix eye shader compilation when strata is enabled.

diff --git a/Engine/Source/Runtime/Engine/Private/Materials/HLSLMaterialTranslator.cpp b/Engine/Source/Runtime/Engine/Private/Materials/HLSLMaterialTranslator.cpp
--- a/Engine/Source/Runtime/Engine/Private/Materials/HLSLMaterialTranslator.cpp
+++ b/Engine/Source/Runtime/Engine/Private/Materials/HLSLMaterialTranslator.cpp
@@ -10588,7 +10588,7 @@ int32 FHLSLMaterialTranslator::StrataHairBSDF(int32 BaseColor, int32 Scatter, in
 int32 FHLSLMaterialTranslator::StrataEyeBSDF(int32 DiffuseAlbedo, int32 Roughness, int32 IrisMask, int32 IrisDistance, int32 EmissiveColor, int32 CorneaNormal, int32 IrisNormal, int32 IrisPlaneNormal, int32 SSSProfileId, const FString& SharedLocalBasisIndexMacro)
 {
 	return AddCodeChunk(
-		MCT_Strata, TEXT("GetStrataEyeBSDF(%s, %s, %s, %s, %s, %s, %s, %s) /* Cornea:%s Iris:%s */"),
+		MCT_Strata, TEXT("GetStrataEyeBSDF(%s, %s, %s, %s, %s, %s, %s, %s, %s) /* Cornea:%s Iris:%s */"),
 		*StrataGetCastParameterCode(DiffuseAlbedo, MCT_Float3),
 		*StrataGetCastParameterCode(Roughness, MCT_Float),
 		*StrataGetCastParameterCode(IrisMask, MCT_Float),
