Strata - fix bad normal on decals.

diff --git a/Engine/Source/Runtime/Engine/Private/Materials/HLSLMaterialTranslator.cpp b/Engine/Source/Runtime/Engine/Private/Materials/HLSLMaterialTranslator.cpp
--- a/Engine/Source/Runtime/Engine/Private/Materials/HLSLMaterialTranslator.cpp
+++ b/Engine/Source/Runtime/Engine/Private/Materials/HLSLMaterialTranslator.cpp
@@ -8489,7 +8489,12 @@ int32 FHLSLMaterialTranslator::TransformPosition(EMaterialCommonBasis SourceCoor
 
 int32 FHLSLMaterialTranslator::TransformNormalFromRequestedBasisToWorld(int32 NormalCodeChunk)
 {
-	return Material->IsTangentSpaceNormal() ? TransformBase(MCB_Tangent, MCB_World, NormalCodeChunk, 0) : NormalCodeChunk;
+	if (Material->GetMaterialDomain() == MD_DeferredDecal || Material->IsTangentSpaceNormal())
+	{
+		// See TransformTangentNormalToWorld definitions in MaterialTemplate.ush
+		return AddCodeChunk(MCT_Float3, TEXT("TransformTangentNormalToWorld(Parameters.TangentToWorld, %s)"), *GetParameterCode(NormalCodeChunk));
+	}
+	return NormalCodeChunk;
 }
 
 int32 FHLSLMaterialTranslator::DynamicParameter(FLinearColor& DefaultValue, uint32 ParameterIndex)
