"(" and ")" characters are not accepted by the analytics backend as a valid name so code added to rectify this.

diff --git a/Engine/Source/Runtime/Engine/Private/Analytics/StudioAnalytics.cpp b/Engine/Source/Runtime/Engine/Private/Analytics/StudioAnalytics.cpp
--- a/Engine/Source/Runtime/Engine/Private/Analytics/StudioAnalytics.cpp
+++ b/Engine/Source/Runtime/Engine/Private/Analytics/StudioAnalytics.cpp
@@ -192,8 +192,10 @@ void FStudioAnalytics::FireEvent_Loading(const FString& LoadingName, double Seco
 
 		for (const FDerivedDataCacheResourceStat& Stat : DDCResourceStats)
 		{
-			const FString BaseName = TEXT("DDC.Resource.") + Stat.AssetType;
+			FString BaseName = TEXT("DDC.Resource.") + Stat.AssetType;
 
+			BaseName = BaseName.Replace(TEXT("("), TEXT("")).Replace(TEXT(")"), TEXT(""));
+			
 			{
 				FString AttrName = BaseName + TEXT(".BuildCount");
 				Attributes.Emplace(MoveTemp(AttrName), Stat.BuildCount);
