Renamed comments and LLM_SCOPE tag from audio to animation at AnimationStreaming.cpp/h

diff --git a/Engine/Source/Runtime/Engine/Private/Animation/AnimationStreaming.cpp b/Engine/Source/Runtime/Engine/Private/Animation/AnimationStreaming.cpp
--- a/Engine/Source/Runtime/Engine/Private/Animation/AnimationStreaming.cpp
+++ b/Engine/Source/Runtime/Engine/Private/Animation/AnimationStreaming.cpp
@@ -392,7 +392,7 @@ SIZE_T FAnimationStreamingManager::GetMemorySizeForAnim(const UAnimStreamable* A
 
 void FAnimationStreamingManager::UpdateResourceStreaming(float DeltaTime, bool bProcessEverything /*= false*/)
 {
-	LLM_SCOPE(ELLMTag::Audio);
+	LLM_SCOPE(ELLMTag::Animation);
 
 	FScopeLock Lock(&CriticalSection);
 
diff --git a/Engine/Source/Runtime/Engine/Public/Animation/AnimationStreaming.h b/Engine/Source/Runtime/Engine/Public/Animation/AnimationStreaming.h
--- a/Engine/Source/Runtime/Engine/Public/Animation/AnimationStreaming.h
+++ b/Engine/Source/Runtime/Engine/Public/Animation/AnimationStreaming.h
@@ -1,7 +1,7 @@
 // Copyright Epic Games, Inc. All Rights Reserved.
 
 /*=============================================================================
-AudioStreaming.h: Definitions of classes used for audio streaming.
+AnimationStreaming.h: Definitions of classes used for animation streaming.
 =============================================================================*/
 
 #pragma once
@@ -58,7 +58,7 @@ struct FStreamingAnimationData final
 	void FreeResources();
 
 	/**
-	 * Sets up the streaming wave data and loads the first chunk of audio for instant play
+	 * Sets up the streaming wave data and loads the first chunk of animation for instant play
 	 *
 	 * @param Anim	The streamable animation we are managing
 	 */
@@ -111,7 +111,7 @@ public:
 	/** AnimStreamable this streaming data is for */
 	UAnimStreamable* StreamableAnim;
 
-	/* Contains pointers to Chunks of audio data that have been streamed in */
+	/* Contains pointers to Chunks of animation data that have been streamed in */
 	TArray<FLoadedAnimationChunk> LoadedChunks;
 
 	mutable FCriticalSection LoadedChunksCritcalSection;
@@ -123,13 +123,13 @@ public:
 
 	TArray<uint32> LoadFailedChunks;
 
-	/** Ptr to owning audio streaming manager. */
+	/** Ptr to owning animation streaming manager. */
 	FAnimationStreamingManager* AnimationStreamingManager;
 };
 
 
 /**
-* Streaming manager dealing with audio.
+* Streaming manager dealing with animation.
 */
 struct FAnimationStreamingManager : public IAnimationStreamingManager
 {
@@ -149,14 +149,14 @@ struct FAnimationStreamingManager : public IAnimationStreamingManager
 	virtual void NotifyLevelOffset( class ULevel* Level, const FVector& Offset ) override;
 	// End IStreamingManager interface
 
-	// IAudioStreamingManager interface
+	// IAnimationStreamingManager interface
 	virtual void AddStreamingAnim(UAnimStreamable* Anim) override;
 	virtual bool RemoveStreamingAnim(UAnimStreamable* Anim) override;
 	virtual SIZE_T GetMemorySizeForAnim(const UAnimStreamable* Anim) override;
 	virtual const FCompressedAnimSequence* GetLoadedChunk(const UAnimStreamable* Anim, uint32 ChunkIndex, bool bTrackAsRequested) const override;
-	// End IAudioStreamingManager interface
+	// End IAnimationStreamingManager interface
 
-	/** Called when an async callback is made on an async loading audio chunk request. */
+	/** Called when an async callback is made on an async loading animation chunk request. */
 	void OnAsyncFileCallback(FStreamingAnimationData* StreamingAnimData, int32 ChunkIndex, int64 ReadSize, IBulkDataIORequest* ReadRequest, bool bWasCancelled);
 
 protected:
