Fix Strata eye/hair BSDF when created using dedicated node.

diff --git a/Engine/Source/Runtime/Engine/Private/Materials/HLSLMaterialTranslator.cpp b/Engine/Source/Runtime/Engine/Private/Materials/HLSLMaterialTranslator.cpp
--- a/Engine/Source/Runtime/Engine/Private/Materials/HLSLMaterialTranslator.cpp
+++ b/Engine/Source/Runtime/Engine/Private/Materials/HLSLMaterialTranslator.cpp
@@ -10575,10 +10575,10 @@ int32 FHLSLMaterialTranslator::StrataUnlitBSDF(int32 EmissiveColor, int32 Transm
 	);
 }
 
-int32 FHLSLMaterialTranslator::StrataHairBSDF(int32 BaseColor, int32 Scatter, int32 Specular, int32 Roughness, int32 Backlit, int32 EmissiveColor, int32 Tangent, const FString& SharedLocalBasisIndexMacro)
+int32 FHLSLMaterialTranslator::StrataHairBSDF(int32 BaseColor, int32 Scatter, int32 Specular, int32 Roughness, int32 Backlit, int32 EmissiveColor, int32 Tangent, const FString& SharedLocalBasisIndexMacro, FStrataOperator* PromoteToOperator)
 {
 	return AddCodeChunk(
-		MCT_Strata, TEXT("GetStrataHairBSDF(%s, %s, %s, %s, %s, %s, %s) /* %s */"),
+		MCT_Strata, TEXT("PromoteParameterBlendedBSDFToOperator(GetStrataHairBSDF(%s, %s, %s, %s, %s, %s, %s), Parameters.StrataTree, %u, %u, %u, %u) /* Tangent:%s */"),
 		*StrataGetCastParameterCode(BaseColor,			MCT_Float3),
 		*StrataGetCastParameterCode(Scatter,			MCT_Float),
 		*StrataGetCastParameterCode(Specular,			MCT_Float),
@@ -10586,14 +10586,18 @@ int32 FHLSLMaterialTranslator::StrataHairBSDF(int32 BaseColor, int32 Scatter, in
 		*StrataGetCastParameterCode(Backlit,			MCT_Float),
 		*StrataGetCastParameterCode(EmissiveColor,		MCT_Float3),
 		*SharedLocalBasisIndexMacro,
+		PromoteToOperator->Index,
+		PromoteToOperator->BSDFIndex,
+		PromoteToOperator->LayerDepth,
+		PromoteToOperator->bIsBottom ? 1 : 0,
 		*GetParameterCode(Tangent)
 	);
 }
 
-int32 FHLSLMaterialTranslator::StrataEyeBSDF(int32 DiffuseAlbedo, int32 Roughness, int32 IrisMask, int32 IrisDistance, int32 EmissiveColor, int32 CorneaNormal, int32 IrisNormal, int32 IrisPlaneNormal, int32 SSSProfileId, const FString& SharedLocalBasisIndexMacro)
+int32 FHLSLMaterialTranslator::StrataEyeBSDF(int32 DiffuseAlbedo, int32 Roughness, int32 IrisMask, int32 IrisDistance, int32 EmissiveColor, int32 CorneaNormal, int32 IrisNormal, int32 IrisPlaneNormal, int32 SSSProfileId, const FString& SharedLocalBasisIndexMacro, FStrataOperator* PromoteToOperator)
 {
 	return AddCodeChunk(
-		MCT_Strata, TEXT("GetStrataEyeBSDF(%s, %s, %s, %s, %s, %s, %s, %s, %s) /* Cornea:%s Iris:%s */"),
+		MCT_Strata, TEXT("PromoteParameterBlendedBSDFToOperator(GetStrataEyeBSDF(%s, %s, %s, %s, %s, %s, %s, %s, %s), Parameters.StrataTree, %u, %u, %u, %u) /* Cornea:%s Iris:%s */"),
 		*StrataGetCastParameterCode(DiffuseAlbedo, MCT_Float3),
 		*StrataGetCastParameterCode(Roughness, MCT_Float),
 		*StrataGetCastParameterCode(IrisMask, MCT_Float),
@@ -10603,6 +10607,10 @@ int32 FHLSLMaterialTranslator::StrataEyeBSDF(int32 DiffuseAlbedo, int32 Roughnes
 		*StrataGetCastParameterCode(SSSProfileId, MCT_Float),
 		*StrataGetCastParameterCode(EmissiveColor, MCT_Float3),
 		*SharedLocalBasisIndexMacro,
+		PromoteToOperator->Index,
+		PromoteToOperator->BSDFIndex,
+		PromoteToOperator->LayerDepth,
+		PromoteToOperator->bIsBottom ? 1 : 0,
 		*GetParameterCode(CorneaNormal),
 		*GetParameterCode(IrisNormal)
 	);
@@ -10614,7 +10622,7 @@ int32 FHLSLMaterialTranslator::StrataSingleLayerWaterBSDF(
 	FStrataOperator* PromoteToOperator)
 {
 	return AddCodeChunk(
-		MCT_Strata, TEXT("PromoteParameterBlendedBSDFToOperator(GetStrataSingleLayerWaterBSDF(%s, %s, %s, %s, %s, %s, %s, %s, %s, %s, %s), Parameters.StrataTree, %u, %u, %u, %u) /* %s */"),
+		MCT_Strata, TEXT("PromoteParameterBlendedBSDFToOperator(GetStrataSingleLayerWaterBSDF(%s, %s, %s, %s, %s, %s, %s, %s, %s, %s, %s), Parameters.StrataTree, %u, %u, %u, %u) /* Normal:%s */"),
 		*StrataGetCastParameterCode(BaseColor,				MCT_Float3),
 		*StrataGetCastParameterCode(Metallic,				MCT_Float),
 		*StrataGetCastParameterCode(Specular,				MCT_Float),
diff --git a/Engine/Source/Runtime/Engine/Private/Materials/HLSLMaterialTranslator.h b/Engine/Source/Runtime/Engine/Private/Materials/HLSLMaterialTranslator.h
--- a/Engine/Source/Runtime/Engine/Private/Materials/HLSLMaterialTranslator.h
+++ b/Engine/Source/Runtime/Engine/Private/Materials/HLSLMaterialTranslator.h
@@ -1056,8 +1056,8 @@ protected:
 		int32 CustomTangent_Tangent) override;
 	virtual int32 StrataVolumetricFogCloudBSDF(int32 Albedo, int32 Extinction, int32 EmissiveColor, int32 AmbientOcclusion) override;
 	virtual int32 StrataUnlitBSDF(int32 EmissiveColor, int32 TransmittanceColor) override;
-	virtual int32 StrataHairBSDF(int32 BaseColor, int32 Scatter, int32 Specular, int32 Roughness, int32 Backlit, int32 EmissiveColor, int32 Tangent, const FString& SharedLocalBasisIndexMacro) override;
-	virtual int32 StrataEyeBSDF(int32 DiffuseColor, int32 Roughness, int32 IrisMask, int32 IrisDistance, int32 EmissiveColor, int32 CorneaNormal, int32 IrisNormal, int32 IrisPlaneNormal, int32 SSSProfileId, const FString& SharedLocalBasisIndexMacro) override;
+	virtual int32 StrataHairBSDF(int32 BaseColor, int32 Scatter, int32 Specular, int32 Roughness, int32 Backlit, int32 EmissiveColor, int32 Tangent, const FString& SharedLocalBasisIndexMacro, FStrataOperator* PromoteToOperator) override;
+	virtual int32 StrataEyeBSDF(int32 DiffuseColor, int32 Roughness, int32 IrisMask, int32 IrisDistance, int32 EmissiveColor, int32 CorneaNormal, int32 IrisNormal, int32 IrisPlaneNormal, int32 SSSProfileId, const FString& SharedLocalBasisIndexMacro, FStrataOperator* PromoteToOperator) override;
 	virtual int32 StrataSingleLayerWaterBSDF(
 		int32 BaseColor, int32 Metallic, int32 Specular, int32 Roughness, 
 		int32 EmissiveColor, int32 TopMaterialOpacity, int32 WaterAlbedo, int32 WaterExtinction, int32 WaterPhaseG, 
diff --git a/Engine/Source/Runtime/Engine/Private/Materials/MaterialExpressions.cpp b/Engine/Source/Runtime/Engine/Private/Materials/MaterialExpressions.cpp
--- a/Engine/Source/Runtime/Engine/Private/Materials/MaterialExpressions.cpp
+++ b/Engine/Source/Runtime/Engine/Private/Materials/MaterialExpressions.cpp
@@ -23295,6 +23295,18 @@ int32 UMaterialExpressionStrataHairBSDF::Compile(class FMaterialCompiler* Compil
 	int32 TangentCodeChunk = CompileWithDefaultTangentWS(Compiler, Tangent);
 	const FStrataRegisteredSharedLocalBasis NewRegisteredSharedLocalBasis = StrataCompilationInfoCreateSharedLocalBasis(Compiler, TangentCodeChunk);
 
+	FStrataOperator& StrataOperator = Compiler->StrataCompilationGetOperator(this);
+	StrataOperator.BSDFRegisteredSharedLocalBasis = NewRegisteredSharedLocalBasis;
+
+	if (StrataOperator.bUseParameterBlending)
+	{
+		return Compiler->Errorf(TEXT("Strata Hair BSDF node cannot be used with parameter blending."));
+	}
+	else if (StrataOperator.bRootOfParameterBlendingSubTree)
+	{
+		return Compiler->Errorf(TEXT("Strata Hair BSDF node cannot be the root of a parameter blending sub tree."));
+	}
+
 	int32 OutputCodeChunk = Compiler->StrataHairBSDF(
 		CompileWithDefaultFloat3(Compiler, BaseColor,	0.0f, 0.0f, 0.0f),
 		CompileWithDefaultFloat1(Compiler, Scatter,		0.0f),
@@ -23303,7 +23315,8 @@ int32 UMaterialExpressionStrataHairBSDF::Compile(class FMaterialCompiler* Compil
 		CompileWithDefaultFloat1(Compiler, Backlit,		0.0f),
 		CompileWithDefaultFloat3(Compiler, EmissiveColor,1.0f, 0.0f, 0.0f),
 		TangentCodeChunk,
-		Compiler->GetStrataSharedLocalBasisIndexMacro(NewRegisteredSharedLocalBasis));
+		Compiler->GetStrataSharedLocalBasisIndexMacro(NewRegisteredSharedLocalBasis),
+		&StrataOperator);
 
 	return OutputCodeChunk;
 }
@@ -23400,6 +23413,18 @@ int32 UMaterialExpressionStrataEyeBSDF::Compile(class FMaterialCompiler* Compile
 		SSSProfileCodeChunk = Compiler->ForceCast(Compiler->ScalarParameter(GetSubsurfaceProfileParameterName(), 1.0f), MCT_Float1);
 	}
 
+	FStrataOperator& StrataOperator = Compiler->StrataCompilationGetOperator(this);
+	StrataOperator.BSDFRegisteredSharedLocalBasis = NewRegisteredSharedLocalBasis;
+
+	if (StrataOperator.bUseParameterBlending)
+	{
+		return Compiler->Errorf(TEXT("Strata Eye BSDF node cannot be used with parameter blending."));
+	}
+	else if (StrataOperator.bRootOfParameterBlendingSubTree)
+	{
+		return Compiler->Errorf(TEXT("Strata Eye BSDF node cannot be the root of a parameter blending sub tree."));
+	}
+
 	int32 OutputCodeChunk = Compiler->StrataEyeBSDF(
 		CompileWithDefaultFloat3(Compiler, DiffuseColor, 0.0f, 0.0f, 0.0f),
 		CompileWithDefaultFloat1(Compiler, Roughness,	 0.5f),
@@ -23410,7 +23435,8 @@ int32 UMaterialExpressionStrataEyeBSDF::Compile(class FMaterialCompiler* Compile
 		SSSProfileCodeChunk != INDEX_NONE ? SSSProfileCodeChunk : Compiler->Constant(0.0f),
 		CompileWithDefaultFloat3(Compiler, EmissiveColor,0.0f, 0.0f, 0.0f),
 		CorneaNormalCodeChunk,
-		Compiler->GetStrataSharedLocalBasisIndexMacro(NewRegisteredSharedLocalBasis));
+		Compiler->GetStrataSharedLocalBasisIndexMacro(NewRegisteredSharedLocalBasis),
+		&StrataOperator);
 
 	return OutputCodeChunk;
 }
diff --git a/Engine/Source/Runtime/Engine/Public/MaterialCompiler.h b/Engine/Source/Runtime/Engine/Public/MaterialCompiler.h
--- a/Engine/Source/Runtime/Engine/Public/MaterialCompiler.h
+++ b/Engine/Source/Runtime/Engine/Public/MaterialCompiler.h
@@ -590,8 +590,8 @@ public:
 		int32 CustomTangent_Tangent) = 0;
 	virtual int32 StrataVolumetricFogCloudBSDF(int32 Albedo, int32 Extinction, int32 EmissiveColor, int32 AmbientOcclusion) = 0;
 	virtual int32 StrataUnlitBSDF(int32 EmissiveColor, int32 TransmittanceColor) = 0;
-	virtual int32 StrataHairBSDF(int32 BaseColor, int32 Scatter, int32 Specular, int32 Roughness, int32 Backlit, int32 EmissiveColor, int32 Tangent, const FString& SharedLocalBasisIndexMacro) = 0;
-	virtual int32 StrataEyeBSDF(int32 DiffuseAlbedo, int32 Roughness, int32 IrisMask, int32 IrisDistance, int32 EmissiveColor, int32 CorneaNormal, int32 IrisNormal, int32 IrisPlaneNormal, int32 SSSProfileId, const FString& SharedLocalBasisIndexMacro) = 0;
+	virtual int32 StrataHairBSDF(int32 BaseColor, int32 Scatter, int32 Specular, int32 Roughness, int32 Backlit, int32 EmissiveColor, int32 Tangent, const FString& SharedLocalBasisIndexMacro, FStrataOperator* PromoteToOperator) = 0;
+	virtual int32 StrataEyeBSDF(int32 DiffuseAlbedo, int32 Roughness, int32 IrisMask, int32 IrisDistance, int32 EmissiveColor, int32 CorneaNormal, int32 IrisNormal, int32 IrisPlaneNormal, int32 SSSProfileId, const FString& SharedLocalBasisIndexMacro, FStrataOperator* PromoteToOperator) = 0;
 	virtual int32 StrataSingleLayerWaterBSDF(
 		int32 BaseColor, int32 Metallic, int32 Specular, int32 Roughness, int32 EmissiveColor, int32 TopMaterialOpacity, 
 		int32 WaterAlbedo, int32 WaterExtinction, int32 WaterPhaseG, int32 ColorScaleBehindWater, int32 Normal, const FString& SharedLocalBasisIndexMacro, FStrataOperator* PromoteToOperator) = 0;
@@ -1178,14 +1178,14 @@ public:
 		return Compiler->StrataUnlitBSDF(EmissiveColor, TransmittanceColor);
 	}
 
-	virtual int32 StrataHairBSDF(int32 BaseColor, int32 Scatter, int32 Specular, int32 Roughness, int32 Backlit, int32 EmissiveColor, int32 Tangent, const FString& SharedLocalBasisIndexMacro) override
+	virtual int32 StrataHairBSDF(int32 BaseColor, int32 Scatter, int32 Specular, int32 Roughness, int32 Backlit, int32 EmissiveColor, int32 Tangent, const FString& SharedLocalBasisIndexMacro, FStrataOperator* PromoteToOperator) override
 	{
-		return Compiler->StrataHairBSDF(BaseColor, Scatter, Specular, Roughness, Backlit, EmissiveColor, Tangent, SharedLocalBasisIndexMacro);
+		return Compiler->StrataHairBSDF(BaseColor, Scatter, Specular, Roughness, Backlit, EmissiveColor, Tangent, SharedLocalBasisIndexMacro, PromoteToOperator);
 	}
 
-	virtual int32 StrataEyeBSDF(int32 DiffuseAlbedo, int32 Roughness, int32 IrisMask, int32 IrisDistance, int32 EmissiveColor, int32 CorneaNormal, int32 IrisNormal, int32 IrisPlaneNormal, int32 SSSProfileId, const FString& SharedLocalBasisIndexMacro) override
+	virtual int32 StrataEyeBSDF(int32 DiffuseAlbedo, int32 Roughness, int32 IrisMask, int32 IrisDistance, int32 EmissiveColor, int32 CorneaNormal, int32 IrisNormal, int32 IrisPlaneNormal, int32 SSSProfileId, const FString& SharedLocalBasisIndexMacro, FStrataOperator* PromoteToOperator) override
 	{
-		return Compiler->StrataEyeBSDF(DiffuseAlbedo, Roughness, IrisMask, IrisDistance, EmissiveColor, CorneaNormal, IrisNormal, IrisPlaneNormal, SSSProfileId, SharedLocalBasisIndexMacro);
+		return Compiler->StrataEyeBSDF(DiffuseAlbedo, Roughness, IrisMask, IrisDistance, EmissiveColor, CorneaNormal, IrisNormal, IrisPlaneNormal, SSSProfileId, SharedLocalBasisIndexMacro, PromoteToOperator);
 	}
 
 	virtual int32 StrataSingleLayerWaterBSDF(
