Fix bForceCopyCrossGPU not initialized in all viewport clients other than UDisplayClusterViewportClient

diff --git a/Engine/Source/Runtime/Engine/Public/SceneView.h b/Engine/Source/Runtime/Engine/Public/SceneView.h
--- a/Engine/Source/Runtime/Engine/Public/SceneView.h
+++ b/Engine/Source/Runtime/Engine/Public/SceneView.h
@@ -1765,7 +1765,7 @@ public:
 
 #if WITH_MGPU
 	/** Force a cross-GPU copy of all persistent resources after rendering this view family (except those marked MultiGPUGraphIgnore) */
-	bool bForceCopyCrossGPU;
+	bool bForceCopyCrossGPU = false;
 #endif
 
 	/** Whether this view is one of multiple view families rendered in a single frame.  Affects occlusion query synchronization logic. */
