World Partition: when an actor package is saved, don't try to pin the saved actors during the dirty package event, do it during the next tick instead, to avoid calls to StaticFindObject triggered by the world outliner reacting to newly added actors.

#rnx

diff --git a/Engine/Source/Runtime/Engine/Private/WorldPartition/WorldPartition.cpp b/Engine/Source/Runtime/Engine/Private/WorldPartition/WorldPartition.cpp
--- a/Engine/Source/Runtime/Engine/Private/WorldPartition/WorldPartition.cpp
+++ b/Engine/Source/Runtime/Engine/Private/WorldPartition/WorldPartition.cpp
@@ -282,16 +282,6 @@ void UWorldPartition::OnPackageDirtyStateChanged(UPackage* Package)
 				{
 					DirtyActors.Add(ActorHandle.ToReference(), Actor);
 				}
-				else
-				{
-					// If we hold the last reference to that actor (or no reference are held at all), pin it to avoid unloading
-					if (PinnedActors && (ActorHandle->GetHardRefCount() <= 1))
-					{
-						PinnedActors->AddActors({ ActorHandle });
-					}
-
-					DirtyActors.Remove(ActorHandle.ToReference());
-				}
 			}
 		}
 	}
@@ -1107,6 +1097,26 @@ void UWorldPartition::Tick(float DeltaSeconds)
 		EditorHash->Tick(DeltaSeconds);
 	}
 
+	if (DirtyActors.Num())
+	{
+		for (auto DirtyActorIt = DirtyActors.CreateIterator(); DirtyActorIt; ++DirtyActorIt)
+		{
+			const FWorldPartitionReference& ActorReference = DirtyActorIt->Key;
+			const AActor* Actor = DirtyActorIt->Value;
+
+			if (!Actor->GetPackage()->IsDirty())
+			{
+				// If we hold the last reference to that actor (or no reference are held at all), pin it to avoid unloading
+				if (PinnedActors && ActorReference->GetHardRefCount() <= 1)
+				{
+					PinnedActors->AddActors({ ActorReference.ToHandle() });
+				}
+
+				DirtyActorIt.RemoveCurrent();
+			}
+		}
+	}
+
 	if (bForceGarbageCollection)
 	{
 		GEngine->ForceGarbageCollection(bForceGarbageCollectionPurge);
