Fixed crashes when deleting static mesh instances through the array manipulator in the details panel when one or more of the deleted instances is selected.

When deleting static mesh instances through the array manipulator in the details panel, there was no call to deselect any potentially selected static mesh instances. This resulted in references to these instances being held at the point they were being deleted, which the Typed Elements guard against. These changes introduce explict calls to deselect any instances before they get deleted. Deselecting will however cause a callback to UInstancedStaticMeshComponent, that causes the internal list of selected instances to be updated even though the instance entry may already be removed. This situation can happen because deleting through the array manipulator removes the index from PerInstanceSMData while deleting through other means such as the contextual menu or keyboard shortcut requires UInstancedStaticMeshComponent to explicitly remove the entry. For consistency all deletions behave as if the entry has already been removed from PerInstanceSMData, but an enum has been used for the currently supported delete cases in UInstancedStaticMeshComponent so behavior can be specialized if needed in the future.

diff --git a/Engine/Source/Editor/LevelEditor/Private/SLevelEditor.cpp b/Engine/Source/Editor/LevelEditor/Private/SLevelEditor.cpp
--- a/Engine/Source/Editor/LevelEditor/Private/SLevelEditor.cpp
+++ b/Engine/Source/Editor/LevelEditor/Private/SLevelEditor.cpp
@@ -269,6 +269,9 @@ void SLevelEditor::Initialize( const TSharedRef<SDockTab>& OwnerTab, const TShar
 	GUnrealEd->GetSelectedActors()->SetElementSelectionSet(SelectedElements);
 	GUnrealEd->GetSelectedComponents()->SetElementSelectionSet(SelectedElements);
 
+	// Setup a callback to deselect instances of Instanced Static Meshes so there are no lingering references.
+	FSMInstanceElementIdMap::Get().OnInstancePreRemoval().AddSP(this, &SLevelEditor::OnIsmInstanceRemoving);
+
 	CommonActions = NewObject<UTypedElementCommonActions>();
 	CommonActions->AddToRoot();
 
@@ -1941,6 +1944,12 @@ void SLevelEditor::HandleAssetsDeleted(const TArray<UClass*>& DeletedClasses)
 	}
 }
 
+void SLevelEditor::OnIsmInstanceRemoving(const FSMInstanceElementId& SMInstanceElementId, [[maybe_unused]] int32 InstanceIndex)
+{
+	SelectedElements->DeselectElement(
+		UEngineElementsLibrary::AcquireEditorSMInstanceElementHandle(SMInstanceElementId), FTypedElementSelectionOptions{});
+}
+
 void SLevelEditor::OnElementSelectionChanged(const UTypedElementSelectionSet* SelectionSet, bool bForceRefresh)
 {
 	for (TSharedRef<SActorDetails> ActorDetails : GetAllActorDetails())
diff --git a/Engine/Source/Editor/LevelEditor/Private/SLevelEditor.h b/Engine/Source/Editor/LevelEditor/Private/SLevelEditor.h
--- a/Engine/Source/Editor/LevelEditor/Private/SLevelEditor.h
+++ b/Engine/Source/Editor/LevelEditor/Private/SLevelEditor.h
@@ -248,6 +248,9 @@ private:
 	/** Handles deletion of assets */
 	void HandleAssetsDeleted(const TArray<UClass*>& DeletedClasses);
 
+	/** Handles events where an instance of an Instanced Static Mesh is about to be removed. */
+	void OnIsmInstanceRemoving(const FSMInstanceElementId& SMInstanceElementId, int32 InstanceIndex);
+
 	/** Called when element selection changes */
 	void OnElementSelectionChanged(const UTypedElementSelectionSet* SelectionSet, bool bForceRefresh = false);
 
diff --git a/Engine/Source/Runtime/Engine/Classes/Components/InstancedStaticMeshComponent.h b/Engine/Source/Runtime/Engine/Classes/Components/InstancedStaticMeshComponent.h
--- a/Engine/Source/Runtime/Engine/Classes/Components/InstancedStaticMeshComponent.h
+++ b/Engine/Source/Runtime/Engine/Classes/Components/InstancedStaticMeshComponent.h
@@ -315,6 +315,19 @@ public:
 
 	/** Indicates that the user has purposedly chosen to show the instance list in the details panel, despite the performance warning. */
 	bool bForceShowAllInstancesDetails = false;
+
+	/** The reason why a deletion operation is currently happening. */
+	enum class EInstanceDeletionReason : uint8
+	{
+		NotDeleting, /** There is currently no deletion in progress. */
+		EntryAlreadyRemoved, /** The instance has been deleted externally. Data synchronization in progress. */
+		EntryRemoval, /** The instance is being removed. */
+		Clearing /** All instances are being removed. */
+	};
+	/** This will be set to the appropriate state when one or more instances are in the process of being
+	 *  deleted. This is primarily used for functions that round trip to this class, such as callbacks
+	 *  for deselecting instances. */
+	EInstanceDeletionReason DeletionState = EInstanceDeletionReason::NotDeleting;
 #endif
 	/** Physics representation of the instance bodies. */
 	TArray<FBodyInstance*> InstanceBodies;
diff --git a/Engine/Source/Runtime/Engine/Private/Elements/SMInstance/SMInstanceElementId.cpp b/Engine/Source/Runtime/Engine/Private/Elements/SMInstance/SMInstanceElementId.cpp
--- a/Engine/Source/Runtime/Engine/Private/Elements/SMInstance/SMInstanceElementId.cpp
+++ b/Engine/Source/Runtime/Engine/Private/Elements/SMInstance/SMInstanceElementId.cpp
@@ -224,6 +224,8 @@ void FSMInstanceElementIdMap::OnInstanceIndexUpdated(UInstancedStaticMeshCompone
 				{
 					const uint64 InstanceId = *InstanceIdPtr;
 
+					OnInstancePreRemovalDelegate.Broadcast(FSMInstanceElementId{ InComponent, InstanceId }, IndexUpdateData.Index);
+
 					ISMEntry->InstanceIndexToIdMap.Remove(IndexUpdateData.Index);
 					ISMEntry->InstanceIdToIndexMap.Remove(InstanceId);
 
@@ -282,6 +284,14 @@ void FSMInstanceElementIdMap::OnInstanceIndexUpdated(UInstancedStaticMeshCompone
 
 void FSMInstanceElementIdMap::ClearInstanceData_NoLock(UInstancedStaticMeshComponent* InComponent, FSMInstanceElementIdMapEntry& InEntry)
 {
+	if (OnInstancePreRemovalDelegate.IsBound())
+	{
+		for (const TTuple<uint64, int32>& InstanceIdToIndexPair : InEntry.InstanceIdToIndexMap)
+		{
+			OnInstancePreRemovalDelegate.Broadcast(FSMInstanceElementId{ InComponent, InstanceIdToIndexPair.Key }, InstanceIdToIndexPair.Value);
+		}
+	}
+
 	// Take the current instance IDs, so that we can find notify the elements have been unmapped
 	TMap<uint64, int32> PreviousInstanceIdToIndexMap = MoveTemp(InEntry.InstanceIdToIndexMap);
 
diff --git a/Engine/Source/Runtime/Engine/Private/InstancedStaticMesh.cpp b/Engine/Source/Runtime/Engine/Private/InstancedStaticMesh.cpp
--- a/Engine/Source/Runtime/Engine/Private/InstancedStaticMesh.cpp
+++ b/Engine/Source/Runtime/Engine/Private/InstancedStaticMesh.cpp
@@ -3240,6 +3240,10 @@ bool UInstancedStaticMeshComponent::SetCustomData(int32 InstanceIndex, const TAr
 
 bool UInstancedStaticMeshComponent::RemoveInstanceInternal(int32 InstanceIndex, bool InstanceAlreadyRemoved)
 {
+#if WITH_EDITOR
+	DeletionState = InstanceAlreadyRemoved ? EInstanceDeletionReason::EntryAlreadyRemoved : EInstanceDeletionReason::EntryRemoval;
+#endif
+
 	// remove instance
 	if (!InstanceAlreadyRemoved && PerInstanceSMData.IsValidIndex(InstanceIndex))
 	{
@@ -3305,6 +3309,9 @@ bool UInstancedStaticMeshComponent::RemoveInstanceInternal(int32 InstanceIndex,
 	// Force recreation of the render data
 	InstanceUpdateCmdBuffer.Edit();
 	MarkRenderStateDirty();
+#if WITH_EDITOR
+	DeletionState = EInstanceDeletionReason::NotDeleting;
+#endif
 	return true;
 }
 
@@ -4004,6 +4011,10 @@ void UInstancedStaticMeshComponent::GetStreamingRenderAssetInfo(FStreamingTextur
 
 void UInstancedStaticMeshComponent::ClearInstances()
 {
+#if WITH_EDITOR
+	DeletionState = EInstanceDeletionReason::Clearing;
+#endif
+
 	const int32 PrevNumInstances = GetInstanceCount();
 
 	// Clear all the per-instance data
@@ -4030,6 +4041,10 @@ void UInstancedStaticMeshComponent::ClearInstances()
 	}
 
 	FNavigationSystem::UpdateComponentData(*this);
+
+#if WITH_EDITOR
+	DeletionState = EInstanceDeletionReason::NotDeleting;
+#endif
 }
 
 int32 UInstancedStaticMeshComponent::GetInstanceCount() const
@@ -4657,7 +4672,7 @@ bool UInstancedStaticMeshComponent::IsInstanceSelected(int32 InInstanceIndex) co
 void UInstancedStaticMeshComponent::SelectInstance(bool bInSelected, int32 InInstanceIndex, int32 InInstanceCount)
 {
 #if WITH_EDITOR
-	if (InInstanceCount > 0)
+	if (InInstanceCount > 0 && DeletionState == EInstanceDeletionReason::NotDeleting)
 	{
 		if (PerInstanceSMData.Num() != SelectedInstances.Num())
 		{
diff --git a/Engine/Source/Runtime/Engine/Public/Elements/SMInstance/SMInstanceElementId.h b/Engine/Source/Runtime/Engine/Public/Elements/SMInstance/SMInstanceElementId.h
--- a/Engine/Source/Runtime/Engine/Public/Elements/SMInstance/SMInstanceElementId.h
+++ b/Engine/Source/Runtime/Engine/Public/Elements/SMInstance/SMInstanceElementId.h
@@ -119,6 +119,15 @@ public:
 		return OnInstanceRemappedDelegate;
 	}
 
+	/**
+	 * Delegate called when a static mesh instance is about to be removed from the instances array on its component, and the associated ID for a SMInstance element will be unmapped.
+	 */
+	DECLARE_MULTICAST_DELEGATE_TwoParams(FOnInstancePreRemoval, const FSMInstanceElementId& /*SMInstanceElementId*/, int32 /*InstanceIndex*/);
+	FOnInstancePreRemoval& OnInstancePreRemoval()
+	{
+		return OnInstancePreRemovalDelegate;
+	}
+
 	/**
 	 * Delegate called when a static mesh instance has been removed from the instances array on its component, and the associated ID for a SMInstance element has been unmapped.
 	 */
@@ -177,6 +186,7 @@ private:
 
 	FOnInstanceRemapped OnInstanceRemappedDelegate;
 
+	FOnInstancePreRemoval OnInstancePreRemovalDelegate;
 	FOnInstanceRemoved OnInstanceRemovedDelegate;
 
 	FDelegateHandle OnInstanceIndexUpdatedHandle;
