Fix performance issues when setting up active morph targets from curves.

Reduces the call time of FAnimationRuntime::AppendActiveMorphTargets for a MetaHuman head from 2.5ms to 200us.

diff --git a/Engine/Source/Runtime/Engine/Classes/Components/SkinnedMeshComponent.h b/Engine/Source/Runtime/Engine/Classes/Components/SkinnedMeshComponent.h
--- a/Engine/Source/Runtime/Engine/Classes/Components/SkinnedMeshComponent.h
+++ b/Engine/Source/Runtime/Engine/Classes/Components/SkinnedMeshComponent.h
@@ -112,34 +112,10 @@ namespace EBoneSpaces
 }
 
 /** Struct used to indicate one active morph target that should be applied to this USkeletalMesh when rendered. */
-struct FActiveMorphTarget
-{
-	/** The Morph Target that we want to apply. */
-	class UMorphTarget* MorphTarget;
-
-	/** Index into the array of weights for the Morph target, between 0.0 and 1.0. */
-	int32 WeightIndex;
-
-	FActiveMorphTarget()
-		: MorphTarget(nullptr)
-		, WeightIndex(-1)
-	{
-	}
 
-	FActiveMorphTarget(class UMorphTarget* InTarget, int32 InWeightIndex)
-	:	MorphTarget(InTarget)
-	,	WeightIndex(InWeightIndex)
-	{
-	}
+/** WeightIndex is an into the MorphTargetWeights array */
+using FMorphTargetWeightMap = TMap<const UMorphTarget* /* MorphTarget */, int32 /* WeightIndex */>;
 
-	bool operator==(const FActiveMorphTarget& Other) const
-	{
-		// if target is same, we consider same 
-		// any equal operator to check if it's same, 
-		// we just check if this is same anim
-		return MorphTarget == Other.MorphTarget;
-	}
-};
 
 /** An external morph target set. External morph targets are managed by systems outside of the skinned meshes. */
 struct ENGINE_API FExternalMorphSet
@@ -488,8 +464,8 @@ public:
 	void SetDebugDrawColor(const FLinearColor& InColor) { DebugDrawColor = InColor; }
 #endif
 
-	/** Array indicating all active morph targets. This array is updated inside RefreshBoneTransforms based on the Anim Blueprint. */
-	TArray<FActiveMorphTarget> ActiveMorphTargets;
+	/** Array indicating all active morph targets. This map is updated inside RefreshBoneTransforms based on the Anim Blueprint. */
+	FMorphTargetWeightMap ActiveMorphTargets;
 
 	/** Array of weights for all morph targets. This array is updated inside RefreshBoneTransforms based on the Anim Blueprint. */
 	TArray<float> MorphTargetWeights;
diff --git a/Engine/Source/Runtime/Engine/Classes/Engine/SkeletalMesh.h b/Engine/Source/Runtime/Engine/Classes/Engine/SkeletalMesh.h
--- a/Engine/Source/Runtime/Engine/Classes/Engine/SkeletalMesh.h
+++ b/Engine/Source/Runtime/Engine/Classes/Engine/SkeletalMesh.h
@@ -1824,7 +1824,7 @@ public:
 	}
 
 	UE_DEPRECATED(4.27, "Please do not access this member directly; use USkeletalMesh::GetMorphTargets() or USkeletalMesh::SetMorphTargets().")
-	UPROPERTY(BlueprintGetter = GetMorphTargets, BlueprintSetter = SetMorphTargets, Category = Mesh)
+	UPROPERTY(BlueprintGetter = GetMorphTargetsPtrConv, BlueprintSetter = SetMorphTargets, Category = Mesh)
 	TArray<TObjectPtr<UMorphTarget>> MorphTargets;
 
 	static FName GetMorphTargetsMemberName()
@@ -1844,8 +1844,19 @@ public:
 	}
 
 	/** USkinnedAsset interface. */
+	virtual const TArray<TObjectPtr<UMorphTarget>>& GetMorphTargets() const override
+	{
+		WaitUntilAsyncPropertyReleased(ESkeletalMeshAsyncProperties::MorphTargets, ESkinnedAssetAsyncPropertyLockType::ReadOnly);
+		PRAGMA_DISABLE_DEPRECATION_WARNINGS
+		return MorphTargets;
+		PRAGMA_ENABLE_DEPRECATION_WARNINGS
+	}
+
+	// NOTE: BP compiler doesn't support TObjectPtr as an argument type for UFUNCTION so this converting call is
+	// required. For many morphs, this can be expensive, since it needs to resolve _all_ TObjectPtrs and construct a new
+	// array for it.
 	UFUNCTION(BlueprintGetter)
-	virtual const TArray<UMorphTarget*>& GetMorphTargets() const override
+	TArray<UMorphTarget*> GetMorphTargetsPtrConv() const
 	{
 		WaitUntilAsyncPropertyReleased(ESkeletalMeshAsyncProperties::MorphTargets, ESkinnedAssetAsyncPropertyLockType::ReadOnly);
 		PRAGMA_DISABLE_DEPRECATION_WARNINGS
diff --git a/Engine/Source/Runtime/Engine/Classes/Engine/SkinnedAsset.h b/Engine/Source/Runtime/Engine/Classes/Engine/SkinnedAsset.h
--- a/Engine/Source/Runtime/Engine/Classes/Engine/SkinnedAsset.h
+++ b/Engine/Source/Runtime/Engine/Classes/Engine/SkinnedAsset.h
@@ -180,7 +180,7 @@ public:
 	/** Return the morph targets. */
 	virtual TArray<TObjectPtr<UMorphTarget>>& GetMorphTargets()
 	{ static TArray<TObjectPtr<UMorphTarget>> Dummy; return Dummy; }
-	virtual const TArray<UMorphTarget*>& GetMorphTargets() const
+	virtual const TArray<TObjectPtr<UMorphTarget>>& GetMorphTargets() const
 	{ static const TArray<TObjectPtr<UMorphTarget>> Dummy; return Dummy; }
 
 	/**Find a named MorphTarget from the morph targets */
diff --git a/Engine/Source/Runtime/Engine/Private/Animation/AnimationRuntime.cpp b/Engine/Source/Runtime/Engine/Private/Animation/AnimationRuntime.cpp
--- a/Engine/Source/Runtime/Engine/Private/Animation/AnimationRuntime.cpp
+++ b/Engine/Source/Runtime/Engine/Private/Animation/AnimationRuntime.cpp
@@ -2410,21 +2410,12 @@ void FAnimationRuntime::FillUpComponentSpaceTransformsRetargetBasePose(const USk
 }
 #endif // WITH_EDITOR
 
-/** See if an array of ActiveMorphTargets already contains the supplied anim */
-static int32 FindMorphTarget(const TArray<FActiveMorphTarget>& ActiveMorphTargets, UMorphTarget* InMorphTarget)
-{
-	for(int32 i=0; i<ActiveMorphTargets.Num(); i++)
-	{
-		if(ActiveMorphTargets[i].MorphTarget == InMorphTarget)
-		{
-			return i;
-		}
-	}
-
-	return INDEX_NONE;
-}
-
-void FAnimationRuntime::AppendActiveMorphTargets(const USkeletalMesh* InSkeletalMesh, const TMap<FName, float>& MorphCurveAnims, TArray<FActiveMorphTarget>& InOutActiveMorphTargets, TArray<float>& InOutMorphTargetWeights)
+void FAnimationRuntime::AppendActiveMorphTargets(
+	const USkeletalMesh* InSkeletalMesh,
+	const TMap<FName, float>& MorphCurveAnims,
+	FMorphTargetWeightMap& InOutActiveMorphTargets,
+	TArray<float>& InOutMorphTargetWeights
+	)
 {
 	if (!InSkeletalMesh)
 	{
@@ -2443,53 +2434,55 @@ void FAnimationRuntime::AppendActiveMorphTargets(const USkeletalMesh* InSkeletal
 	//
 	// if somehow it gets rendered without going through these places, there will be crash. Renderer expect the buffer size being same. 
 
-	if(MorphCurveAnims.Num() > 0)
+	if(MorphCurveAnims.IsEmpty())
 	{
-		const int32 NumMorphTargets = InSkeletalMesh->GetMorphTargets().Num();
-		InOutMorphTargetWeights.SetNumZeroed(NumMorphTargets);
+		return;
+	}
+	
+	const int32 NumMorphTargets = InSkeletalMesh->GetMorphTargets().Num();
+	InOutMorphTargetWeights.SetNumZeroed(NumMorphTargets);
 
-		if(NumMorphTargets > 0)
-		{
-			// Then go over the CurveKeys finding morph targets by name
-			for(const TPair<FName, float>& MorphCurveAnim : MorphCurveAnims)
+	if(NumMorphTargets == 0)
+	{
+		return;
+	}
+	
+	// Then go over the CurveKeys finding morph targets by name
+	for(const TPair<FName, float>& MorphCurveAnim : MorphCurveAnims)
+	{
+		const FName& CurveName = MorphCurveAnim.Key;
+		const float Weight = MorphCurveAnim.Value;
+
+		// Find morph reference
+		int32 SkeletalMorphIndex = INDEX_NONE;
+		const UMorphTarget* Target = InSkeletalMesh->FindMorphTargetAndIndex(CurveName, SkeletalMorphIndex);
+		if (Target != nullptr)
+		{
+			// See if this morph target already has an entry
+			const int32* FoundMorphIndex = InOutActiveMorphTargets.Find(Target);
+			
+			// If it has a valid weight
+			if (FMath::Abs(Weight) > MinMorphTargetBlendWeight)
 			{
-				const FName& CurveName = MorphCurveAnim.Key;
-				const float Weight = MorphCurveAnim.Value;
-
-				// Find morph reference
-				int32 SkeletalMorphIndex = INDEX_NONE;
-				UMorphTarget* Target = InSkeletalMesh->FindMorphTargetAndIndex(CurveName, SkeletalMorphIndex);
-				if (Target != nullptr)
+				// If not, add it
+				if (FoundMorphIndex == nullptr)
 				{
-					// If it has a valid weight
-					if (FMath::Abs(Weight) > MinMorphTargetBlendWeight)
-					{
-						// See if this morph target already has an entry
-						int32 MorphIndex = FindMorphTarget(InOutActiveMorphTargets, Target);
-						// If not, add it
-						if (MorphIndex == INDEX_NONE)
-						{
-							InOutActiveMorphTargets.Add(FActiveMorphTarget(Target, SkeletalMorphIndex));
-							InOutMorphTargetWeights[SkeletalMorphIndex] = Weight;
-						}
-						else
-						{
-							// If it does, use the max weight
-							check(SkeletalMorphIndex == InOutActiveMorphTargets[MorphIndex].WeightIndex);
-							InOutMorphTargetWeights[SkeletalMorphIndex] = Weight;
-						}
-					}
-					else
-					{
-						int32 MorphIndex = FindMorphTarget(InOutActiveMorphTargets, Target);
-						if (MorphIndex != INDEX_NONE)
-						{
-							// clear weight
-							InOutMorphTargetWeights[SkeletalMorphIndex] = 0.f;
-						}
-					}
+					InOutActiveMorphTargets.Add(Target, SkeletalMorphIndex);
+					InOutMorphTargetWeights[SkeletalMorphIndex] = Weight;
+				}
+				else
+				{
+					// If it does, use the max weight
+					check(SkeletalMorphIndex == *FoundMorphIndex);
+					InOutMorphTargetWeights[SkeletalMorphIndex] = Weight;
 				}
 			}
+			else if (FoundMorphIndex != nullptr)
+			{
+				// The target weight is below the minimum. Force to zero.
+				check(SkeletalMorphIndex == *FoundMorphIndex);
+				InOutMorphTargetWeights[SkeletalMorphIndex] = 0.f;
+			}
 		}
 	}
 }
diff --git a/Engine/Source/Runtime/Engine/Private/SkeletalRenderCPUSkin.cpp b/Engine/Source/Runtime/Engine/Private/SkeletalRenderCPUSkin.cpp
--- a/Engine/Source/Runtime/Engine/Private/SkeletalRenderCPUSkin.cpp
+++ b/Engine/Source/Runtime/Engine/Private/SkeletalRenderCPUSkin.cpp
@@ -17,7 +17,7 @@
 struct FMorphTargetDelta;
 
 template<typename VertexType, int32 NumberOfUVs>
-static void SkinVertices(FFinalSkinVertex* DestVertex, FMatrix44f* ReferenceToLocal, int32 LODIndex, FSkeletalMeshLODRenderData& LOD, FSkinWeightVertexBuffer& WeightBuffer, TArray<FActiveMorphTarget>& ActiveMorphTargets, TArray<float>& MorphTargetWeights, const TMap<int32, FClothSimulData>& ClothSimulUpdateData, float ClothBlendWeight, const FMatrix& WorldToLocal);
+static void SkinVertices(FFinalSkinVertex* DestVertex, FMatrix44f* ReferenceToLocal, int32 LODIndex, FSkeletalMeshLODRenderData& LOD, FSkinWeightVertexBuffer& WeightBuffer, const FMorphTargetWeightMap& ActiveMorphTargets, const TArray<float>& MorphTargetWeights, const TMap<int32, FClothSimulData>& ClothSimulUpdateData, float ClothBlendWeight, const FMatrix& WorldToLocal);
 
 #define INFLUENCE_0		0
 #define INFLUENCE_1		1
@@ -131,8 +131,8 @@ void FSkeletalMeshObjectCPUSkin::EnableOverlayRendering(bool bEnabled, const TAr
 void FSkeletalMeshObjectCPUSkin::Update(
 	int32 LODIndex,
 	USkinnedMeshComponent* InMeshComponent,
-	const TArray<FActiveMorphTarget>& ActiveMorphTargets,
-	const TArray<float>& MorphTargetWeights,
+	const FMorphTargetWeightMap& InActiveMorphTargets,
+	const TArray<float>& InMorphTargetWeights,
 	EPreviousBoneTransformUpdateMode PreviousBoneTransformUpdateMode,
 	const FExternalMorphWeightData& InExternalMorphWeightData)
 {
@@ -140,7 +140,7 @@ void FSkeletalMeshObjectCPUSkin::Update(
 	{
 		// create the new dynamic data for use by the rendering thread
 		// this data is only deleted when another update is sent
-		FDynamicSkelMeshObjectDataCPUSkin* NewDynamicData = new FDynamicSkelMeshObjectDataCPUSkin(InMeshComponent,SkeletalMeshRenderData,LODIndex,ActiveMorphTargets, MorphTargetWeights);
+		FDynamicSkelMeshObjectDataCPUSkin* NewDynamicData = new FDynamicSkelMeshObjectDataCPUSkin(InMeshComponent,SkeletalMeshRenderData,LODIndex,InActiveMorphTargets, InMorphTargetWeights);
 
 		// We prepare the next frame but still have the value from the last one
 		uint32 FrameNumberToPrepare = GFrameNumber + 1;
@@ -575,7 +575,7 @@ FDynamicSkelMeshObjectDataCPUSkin::FDynamicSkelMeshObjectDataCPUSkin(
 	USkinnedMeshComponent* InMeshComponent,
 	FSkeletalMeshRenderData* InSkelMeshRenderData,
 	int32 InLODIndex,
-	const TArray<FActiveMorphTarget>& InActiveMorphTargets,
+	const FMorphTargetWeightMap& InActiveMorphTargets,
 	const TArray<float>& InMorphTargetWeights
 	)
 :	LODIndex(InLODIndex)
@@ -600,51 +600,45 @@ FDynamicSkelMeshObjectDataCPUSkin::FDynamicSkelMeshObjectDataCPUSkin(
 /** Struct used to hold temporary info during morph target blending */
 struct FMorphTargetInfo
 {
-	/** Info about morphtarget to blend */
-	FActiveMorphTarget			ActiveMorphTarget;
+	/** The index into the morph weight list */
+	int32						WeightIndex = INDEX_NONE;
+	
 	/** Index of next delta to try applying. This prevents us looking at every delta for every vertex. */
-	int32						NextDeltaIndex;
+	int32						NextDeltaIndex = INDEX_NONE;
 	/** Array of deltas to apply to mesh, sorted based on the index of the base mesh vert that they affect. */
-	const FMorphTargetDelta*	Deltas;
+	const FMorphTargetDelta*	Deltas = nullptr;
 	/** How many deltas are in array */
-	int32						NumDeltas;
+	int32						NumDeltas = 0;
 };
 
 /**
  *	Init set of info structs to hold temporary state while blending morph targets in.
  * @return							number of active morphs that are valid
  */
-static uint32 InitEvalInfos(const TArray<FActiveMorphTarget>& ActiveMorphTargets, const TArray<float>& MorphTargetWeights, int32 LODIndex, TArray<FMorphTargetInfo>& OutEvalInfos)
+static uint32 InitEvalInfos(const FMorphTargetWeightMap& InActiveMorphTargets, const TArray<float>& MorphTargetWeights, int32 LODIndex, TArray<FMorphTargetInfo>& OutEvalInfos)
 {
 	uint32 NumValidMorphTargets=0;
 
-	for( int32 MorphIdx=0; MorphIdx < ActiveMorphTargets.Num(); MorphIdx++ )
+	for(const TTuple<const UMorphTarget*, int32>& MorphItem: InActiveMorphTargets)
 	{
 		FMorphTargetInfo NewInfo;
+		const UMorphTarget* MorphTarget = MorphItem.Key;
+		const int32 WeightIndex = MorphItem.Value;
 
-		const FActiveMorphTarget& ActiveMorphTarget = ActiveMorphTargets[MorphIdx];
-		const float ActiveMorphAbsVertexWeight = FMath::Abs(MorphTargetWeights[ActiveMorphTarget.WeightIndex]);
+		const float ActiveMorphAbsVertexWeight = FMath::Abs(MorphTargetWeights[WeightIndex]);
 
-		if( ActiveMorphTarget.MorphTarget != NULL &&
+		if( MorphTarget != nullptr &&
 			ActiveMorphAbsVertexWeight >= MinMorphTargetBlendWeight &&
 			ActiveMorphAbsVertexWeight <= MaxMorphTargetBlendWeight &&
-			ActiveMorphTarget.MorphTarget->HasDataForLOD(LODIndex) )
+			MorphItem.Key->HasDataForLOD(LODIndex) )
 		{
 			// start at the first vertex since they affect base mesh verts in ascending order
-			NewInfo.ActiveMorphTarget = ActiveMorphTarget;
+			NewInfo.WeightIndex = WeightIndex;
 			NewInfo.NextDeltaIndex = 0;
-			NewInfo.Deltas = ActiveMorphTarget.MorphTarget->GetMorphTargetDelta(LODIndex, NewInfo.NumDeltas);
+			NewInfo.Deltas = MorphTarget->GetMorphTargetDelta(LODIndex, NewInfo.NumDeltas);
 
 			NumValidMorphTargets++;
 		}
-		else
-		{
-			// invalidate the indices for any invalid morph models
-			NewInfo.ActiveMorphTarget = FActiveMorphTarget();
-			NewInfo.NextDeltaIndex = INDEX_NONE;
-			NewInfo.Deltas = nullptr;
-			NewInfo.NumDeltas = 0;
-		}			
 
 		OutEvalInfos.Add(NewInfo);
 	}
@@ -710,7 +704,7 @@ FORCEINLINE void UpdateMorphedVertex( VertexType& MorphedVertex, const VertexTyp
 			Info.NextDeltaIndex < Info.NumDeltas &&
 			Info.Deltas[Info.NextDeltaIndex].SourceIdx == CurBaseVertIdx )
 		{
-			ApplyMorphBlend( MorphedVertex, Info.Deltas[Info.NextDeltaIndex], MorphWeights[Info.ActiveMorphTarget.WeightIndex] );
+			ApplyMorphBlend( MorphedVertex, Info.Deltas[Info.NextDeltaIndex], MorphWeights[Info.WeightIndex] );
 
 			// Update 'next delta to use'
 			Info.NextDeltaIndex += 1;
@@ -1065,8 +1059,8 @@ static void SkinVertices(
 	int32 LODIndex, 
 	FSkeletalMeshLODRenderData& LOD,
 	FSkinWeightVertexBuffer& WeightBuffer,
-	TArray<FActiveMorphTarget>& ActiveMorphTargets, 
-	TArray<float>& MorphTargetWeights, 
+	const FMorphTargetWeightMap& InActiveMorphTargets, 
+	const TArray<float>& MorphTargetWeights, 
 	const TMap<int32, FClothSimulData>& ClothSimulUpdateData, 
 	float ClothBlendWeight, 
 	const FMatrix& WorldToLocal)
@@ -1076,7 +1070,7 @@ static void SkinVertices(
 
 	// Create array to track state during morph blending
 	TArray<FMorphTargetInfo> MorphEvalInfos;
-	uint32 NumValidMorphs = InitEvalInfos(ActiveMorphTargets, MorphTargetWeights, LODIndex, MorphEvalInfos);
+	uint32 NumValidMorphs = InitEvalInfos(InActiveMorphTargets, MorphTargetWeights, LODIndex, MorphEvalInfos);
 
 	const uint32 MaxGPUSkinBones = FGPUBaseSkinVertexFactory::GetMaxGPUSkinBones();
 	check(MaxGPUSkinBones <= FGPUBaseSkinVertexFactory::GHardwareMaxGPUSkinBones);
diff --git a/Engine/Source/Runtime/Engine/Private/SkeletalRenderCPUSkin.h b/Engine/Source/Runtime/Engine/Private/SkeletalRenderCPUSkin.h
--- a/Engine/Source/Runtime/Engine/Private/SkeletalRenderCPUSkin.h
+++ b/Engine/Source/Runtime/Engine/Private/SkeletalRenderCPUSkin.h
@@ -40,7 +40,7 @@ public:
 		USkinnedMeshComponent* InMeshComponent,
 		FSkeletalMeshRenderData* InSkelMeshRenderData,
 		int32 InLODIndex,
-		const TArray<FActiveMorphTarget>& InActiveMorphTargets,
+		const FMorphTargetWeightMap& InActiveMorphTargets,
 		const TArray<float>& InMorphTargetWeights
 		);
 
@@ -61,7 +61,7 @@ public:
 	/** currently LOD for bones being updated */
 	int32 LODIndex;
 	/** Morphs to blend when skinning verts */
-	TArray<FActiveMorphTarget> ActiveMorphTargets;
+	FMorphTargetWeightMap ActiveMorphTargets;
 	/** Morph Weights to blend when skinning verts */
 	TArray<float> MorphTargetWeights;
 
@@ -98,7 +98,7 @@ public:
 	//~ Begin FSkeletalMeshObject Interface
 	virtual void InitResources(USkinnedMeshComponent* InMeshComponent) override;
 	virtual void ReleaseResources() override;
-	virtual void Update(int32 LODIndex,USkinnedMeshComponent* InMeshComponent,const TArray<FActiveMorphTarget>& ActiveMorphTargets, const TArray<float>& MorphTargetsWeights, EPreviousBoneTransformUpdateMode PreviousBoneTransformUpdateMode, const FExternalMorphWeightData& InExternalMorphWeightData) override;
+	virtual void Update(int32 LODIndex,USkinnedMeshComponent* InMeshComponent,const FMorphTargetWeightMap& InActiveMorphTargets, const TArray<float>& MorphTargetsWeights, EPreviousBoneTransformUpdateMode PreviousBoneTransformUpdateMode, const FExternalMorphWeightData& InExternalMorphWeightData) override;
 	void UpdateDynamicData_RenderThread(FRHICommandListImmediate& RHICmdList, FDynamicSkelMeshObjectDataCPUSkin* InDynamicData, uint32 FrameNumberToPrepare, uint32 RevisionNumber);
 	virtual void EnableOverlayRendering(bool bEnabled, const TArray<int32>* InBonesOfInterest, const TArray<UMorphTarget*>* InMorphTargetOfInterest) override;
 	virtual void CacheVertices(int32 LODIndex, bool bForce) const override;
diff --git a/Engine/Source/Runtime/Engine/Private/SkeletalRenderGPUSkin.cpp b/Engine/Source/Runtime/Engine/Private/SkeletalRenderGPUSkin.cpp
--- a/Engine/Source/Runtime/Engine/Private/SkeletalRenderGPUSkin.cpp
+++ b/Engine/Source/Runtime/Engine/Private/SkeletalRenderGPUSkin.cpp
@@ -407,18 +407,18 @@ void FSkeletalMeshObjectGPUSkin::ReleaseMorphResources()
 void FSkeletalMeshObjectGPUSkin::Update(
 	int32 LODIndex,
 	USkinnedMeshComponent* InMeshComponent,
-	const TArray<FActiveMorphTarget>& ActiveMorphTargets,
-	const TArray<float>& MorphTargetWeights,
+	const FMorphTargetWeightMap& InActiveMorphTargets,
+	const TArray<float>& InMorphTargetWeights,
 	EPreviousBoneTransformUpdateMode PreviousBoneTransformUpdateMode,
 	const FExternalMorphWeightData& InExternalMorphWeightData)
 {
 	LLM_SCOPE(ELLMTag::SkeletalMesh);
 
 	// make sure morph data has been initialized for each LOD
-	if(InMeshComponent && !bMorphResourcesInitialized && (ActiveMorphTargets.Num() > 0 || !InExternalMorphWeightData.MorphSets.IsEmpty()))
+	if(InMeshComponent && !bMorphResourcesInitialized && (!InActiveMorphTargets.IsEmpty() || !InExternalMorphWeightData.MorphSets.IsEmpty()))
 	{
 		// initialized on-the-fly in order to avoid creating extra vertex streams for each skel mesh instance
-		InitMorphResources(InMeshComponent->bPerBoneMotionBlur, MorphTargetWeights);		
+		InitMorphResources(InMeshComponent->bPerBoneMotionBlur, InMorphTargetWeights);		
 	}
 
 	// create the new dynamic data for use by the rendering thread
@@ -429,8 +429,8 @@ void FSkeletalMeshObjectGPUSkin::Update(
 		SkeletalMeshRenderData,
 		this,
 		LODIndex,
-		ActiveMorphTargets,
-		MorphTargetWeights,
+		InActiveMorphTargets,
+		InMorphTargetWeights,
 		PreviousBoneTransformUpdateMode,
 		InExternalMorphWeightData);
 
@@ -1301,7 +1301,7 @@ void FSkeletalMeshObjectGPUSkin::FSkeletalMeshObjectLOD::UpdateSkinWeights(FSkel
 	FSkeletalMeshLODRenderData& LODData = SkelMeshRenderData->LODRenderData[LODIndex];	
 	if (CompLODInfo)
 	{
-		FSkinWeightVertexBuffer* NewMeshObjectWeightBuffer = nullptr;
+		FSkinWeightVertexBuffer* NewMeshObjectWeightBuffer;
 
 		if (CompLODInfo->OverrideSkinWeights &&
 			CompLODInfo->OverrideSkinWeights->GetNumVertices() == LODData.StaticVertexBuffers.PositionVertexBuffer.GetNumVertices())
@@ -1327,7 +1327,7 @@ void FSkeletalMeshObjectGPUSkin::FSkeletalMeshObjectLOD::UpdateSkinWeights(FSkel
 			FVertexFactoryBuffers VertexBuffers;
 			GetVertexBuffers(VertexBuffers, LODData);
 
-			FSkeletalMeshObjectGPUSkin::FSkeletalMeshObjectLOD* Self = this;
+			FSkeletalMeshObjectLOD* Self = this;
 			ENQUEUE_RENDER_COMMAND(UpdateSkinWeightsGPUSkin)(
 				[NewMeshObjectWeightBuffer, VertexBuffers, Self](FRHICommandListImmediate& RHICmdList)
 			{
@@ -1338,7 +1338,7 @@ void FSkeletalMeshObjectGPUSkin::FSkeletalMeshObjectLOD::UpdateSkinWeights(FSkel
 	
 }
 
-void FSkeletalMeshObjectGPUSkin::FSkeletalMeshObjectLOD::UpdateMorphVertexBufferCPU(const TArray<FActiveMorphTarget>& ActiveMorphTargets, const TArray<float>& MorphTargetWeights, 
+void FSkeletalMeshObjectGPUSkin::FSkeletalMeshObjectLOD::UpdateMorphVertexBufferCPU(const FMorphTargetWeightMap& InActiveMorphTargets, const TArray<float>& MorphTargetWeights, 
 																					const TArray<int32>& SectionIdsUseByActiveMorphTargets, bool bGPUSkinCacheEnabled, FMorphVertexBuffer& MorphVertexBuffer)
 {
 	SCOPE_CYCLE_COUNTER(STAT_MorphVertexBuffer_Update);
@@ -1400,19 +1400,20 @@ void FSkeletalMeshObjectGPUSkin::FSkeletalMeshObjectLOD::UpdateMorphVertexBuffer
 			SCOPE_CYCLE_COUNTER(STAT_MorphVertexBuffer_ApplyDelta);
 
 			// iterate over all active morph targets and accumulate their vertex deltas
-			for (int32 AnimIdx = 0; AnimIdx < ActiveMorphTargets.Num(); AnimIdx++)
+			for(const TTuple<const UMorphTarget*, int32>& MorphItem: InActiveMorphTargets)
 			{
-				const FActiveMorphTarget& MorphTarget = ActiveMorphTargets[AnimIdx];
-				checkSlow(MorphTarget.MorphTarget != NULL);
-				checkSlow(MorphTarget.MorphTarget->HasDataForLOD(LODIndex));
-				const float MorphTargetWeight = MorphTargetWeights.IsValidIndex(MorphTarget.WeightIndex) ? MorphTargetWeights[MorphTarget.WeightIndex] : 0.0f;
+				const UMorphTarget* MorphTarget = MorphItem.Key;
+				const int32 WeightIndex = MorphItem.Value;
+				checkSlow(MorphTarget != nullptr);
+				checkSlow(MorphTarget->HasDataForLOD(LODIndex));
+				const float MorphTargetWeight = MorphTargetWeights.IsValidIndex(WeightIndex) ? MorphTargetWeights[WeightIndex] : 0.0f;
 				const float MorphAbsWeight = FMath::Abs(MorphTargetWeight);
 				checkSlow(MorphAbsWeight >= MinMorphTargetBlendWeight && MorphAbsWeight <= MaxMorphTargetBlendWeight);
 
 
 				// Get deltas
 				int32 NumDeltas;
-				const FMorphTargetDelta* Deltas = MorphTarget.MorphTarget->GetMorphTargetDelta(LODIndex, NumDeltas);
+				const FMorphTargetDelta* Deltas = MorphTarget->GetMorphTargetDelta(LODIndex, NumDeltas);
 
 				// iterate over the vertices that this lod model has changed
 				for (int32 MorphVertIdx = 0; MorphVertIdx < NumDeltas; MorphVertIdx++)
@@ -1534,9 +1535,10 @@ FGPUBaseSkinVertexFactory const* FSkeletalMeshObjectGPUSkin::GetBaseSkinVertexFa
 	// use the morph enabled vertex factory if any active morphs are set
 	if (DynamicData->NumWeightedActiveMorphTargets > 0)
 	{
-		for (FActiveMorphTarget& ActiveMorphTarget : DynamicData->ActiveMorphTargets)
+		for(const TTuple<const UMorphTarget*, int32>& MorphItem: DynamicData->ActiveMorphTargets)
 		{
-			if (ActiveMorphTarget.MorphTarget->HasDataForSection(LODIndex, ChunkIdx))
+			const UMorphTarget* MorphTarget = MorphItem.Key;
+			if (MorphTarget->HasDataForSection(LODIndex, ChunkIdx))
 			{
 				return LOD.GPUSkinVertexFactories.MorphVertexFactories[ChunkIdx].Get();
 			}
@@ -2388,7 +2390,7 @@ void FDynamicSkelMeshObjectDataGPUSkin::InitDynamicSkelMeshObjectDataGPUSkin(
 	FSkeletalMeshRenderData* InSkeletalMeshRenderData,
 	FSkeletalMeshObjectGPUSkin* InMeshObject,
 	int32 InLODIndex,
-	const TArray<FActiveMorphTarget>& InActiveMorphTargets,
+	const FMorphTargetWeightMap& InActiveMorphTargets,
 	const TArray<float>& InMorphTargetWeights, 
 	EPreviousBoneTransformUpdateMode PreviousBoneTransformUpdateMode,
 	const FExternalMorphWeightData& InExternalMorphWeightData)
@@ -2397,7 +2399,6 @@ void FDynamicSkelMeshObjectDataGPUSkin::InitDynamicSkelMeshObjectDataGPUSkin(
 	check(!ActiveMorphTargets.Num() && !ReferenceToLocal.Num() && !ClothingSimData.Num() && !MorphTargetWeights.Num());
 
 	// append instead of equals to avoid alloc
-	ActiveMorphTargets.Append(InActiveMorphTargets);
 	MorphTargetWeights.Append(InMorphTargetWeights);
 	NumWeightedActiveMorphTargets = 0;
 
@@ -2460,28 +2461,29 @@ void FDynamicSkelMeshObjectDataGPUSkin::InitDynamicSkelMeshObjectDataGPUSkin(
 	}	
 #endif
 	SectionIdsUseByActiveMorphTargets.Empty();
+
 	// find number of morphs that are currently weighted and will affect the mesh
-	for( int32 MorphIdx=ActiveMorphTargets.Num()-1; MorphIdx >= 0; MorphIdx-- )
+	ActiveMorphTargets.Reserve(InActiveMorphTargets.Num());
+	for(const TTuple<const UMorphTarget*, int32>& MorphItem: InActiveMorphTargets)
 	{
-		const FActiveMorphTarget& Morph = ActiveMorphTargets[MorphIdx];
-		const float MorphTargetWeight = MorphTargetWeights[Morph.WeightIndex];
+		const UMorphTarget* MorphTarget = MorphItem.Key;
+		const int32 WeightIndex = MorphItem.Value;
+		const float MorphTargetWeight = MorphTargetWeights[WeightIndex];
 		const float MorphAbsWeight = FMath::Abs(MorphTargetWeight);
 
-		if( Morph.MorphTarget != NULL && 
+		if( MorphTarget != nullptr && 
 			MorphAbsWeight >= MinMorphTargetBlendWeight &&
 			MorphAbsWeight <= MaxMorphTargetBlendWeight &&
-			Morph.MorphTarget->HasDataForLOD(LODIndex) ) 
+			MorphTarget->HasDataForLOD(LODIndex) ) 
 		{
 			NumWeightedActiveMorphTargets++;
-			const TArray<int32>& MorphSectionIndices = Morph.MorphTarget->GetMorphLODModels()[LODIndex].SectionIndices;
+			const TArray<int32>& MorphSectionIndices = MorphTarget->GetMorphLODModels()[LODIndex].SectionIndices;
 			for (int32 SecId = 0; SecId < MorphSectionIndices.Num(); ++SecId)
 			{
 				SectionIdsUseByActiveMorphTargets.AddUnique(MorphSectionIndices[SecId]);
 			}
-		}
-		else
-		{
-			ActiveMorphTargets.RemoveAt(MorphIdx, 1, false);
+
+			ActiveMorphTargets.Add(MorphTarget, WeightIndex);
 		}
 	}
 
@@ -2511,29 +2513,32 @@ void FDynamicSkelMeshObjectDataGPUSkin::InitDynamicSkelMeshObjectDataGPUSkin(
 #endif
 }
 
-bool FDynamicSkelMeshObjectDataGPUSkin::ActiveMorphTargetsEqual( const TArray<FActiveMorphTarget>& CompareActiveMorphTargets, const TArray<float>& CompareMorphTargetWeights)
+bool FDynamicSkelMeshObjectDataGPUSkin::ActiveMorphTargetsEqual(
+	const FMorphTargetWeightMap& InCompareActiveMorphTargets,
+	const TArray<float>& CompareMorphTargetWeights
+	)
 {
-	bool Result=true;
-	if( CompareActiveMorphTargets.Num() == ActiveMorphTargets.Num() )
+	if (InCompareActiveMorphTargets.Num() != ActiveMorphTargets.Num())
+	{
+		return false;
+	}
+
+	for(const TTuple<const UMorphTarget*, int32>& MorphItem: ActiveMorphTargets)
 	{
-		for( int32 MorphIdx=0; MorphIdx < ActiveMorphTargets.Num(); MorphIdx++ )
+		const UMorphTarget* MorphTarget = MorphItem.Key;
+		const int32 WeightIndex = MorphItem.Value;
+		const int32* CompareWeightIndex = InCompareActiveMorphTargets.Find(MorphTarget);
+		if (CompareWeightIndex == nullptr)
 		{
-			const FActiveMorphTarget& Morph = ActiveMorphTargets[MorphIdx];
-			const FActiveMorphTarget& CompMorph = CompareActiveMorphTargets[MorphIdx];
+			return false;
+		}
 
-			if( Morph.MorphTarget != CompMorph.MorphTarget ||
-				FMath::Abs(MorphTargetWeights[Morph.WeightIndex] - CompareMorphTargetWeights[CompMorph.WeightIndex]) >= GMorphTargetWeightThreshold)
-			{
-				Result=false;
-				break;
-			}
+		if( FMath::Abs(MorphTargetWeights[WeightIndex] - CompareMorphTargetWeights[*CompareWeightIndex]) >= GMorphTargetWeightThreshold)
+		{
+			return false;
 		}
 	}
-	else
-	{
-		Result = false;
-	}
-	return Result;
+	return true;
 }
 
 bool FDynamicSkelMeshObjectDataGPUSkin::UpdateClothSimulationData(USkinnedMeshComponent* InMeshComponent)
diff --git a/Engine/Source/Runtime/Engine/Private/SkeletalRenderGPUSkin.h b/Engine/Source/Runtime/Engine/Private/SkeletalRenderGPUSkin.h
--- a/Engine/Source/Runtime/Engine/Private/SkeletalRenderGPUSkin.h
+++ b/Engine/Source/Runtime/Engine/Private/SkeletalRenderGPUSkin.h
@@ -62,7 +62,7 @@ public:
 		FSkeletalMeshRenderData* InSkeletalMeshRenderData,
 		FSkeletalMeshObjectGPUSkin* InMeshObject,
 		int32 InLODIndex,
-		const TArray<FActiveMorphTarget>& InActiveMorphTargets,
+		const FMorphTargetWeightMap& InActiveMorphTargets,
 		const TArray<float>& InMorphTargetWeights, 
 		EPreviousBoneTransformUpdateMode PreviousBoneTransformUpdateMode,
 		const FExternalMorphWeightData& InExternalMorphWeightData);
@@ -86,7 +86,7 @@ public:
 	int32 RayTracingLODIndex;
 #endif
 	/** current morph targets active on this mesh */
-	TArray<FActiveMorphTarget> ActiveMorphTargets;
+	FMorphTargetWeightMap ActiveMorphTargets;
 	/** All morph target weights on this mesh */
 	TArray<float> MorphTargetWeights;
 	/** All section ID impacted by active morph target on this mesh */
@@ -122,7 +122,7 @@ public:
 	* @param MorphTargetWeights - array of morphs weights to compare
 	* @return true if both sets of active morphs are equal
 	*/
-	bool ActiveMorphTargetsEqual(const TArray<FActiveMorphTarget>& CompareActiveMorphTargets, const TArray<float>& CompareMorphTargetWeights);
+	bool ActiveMorphTargetsEqual(const FMorphTargetWeightMap& InCompareActiveMorphTargets, const TArray<float>& CompareMorphTargetWeights);
 	
 	/** Returns the size of memory allocated by render data */
 	virtual void GetResourceSizeEx(FResourceSizeEx& CumulativeResourceSize)
@@ -334,7 +334,7 @@ public:
 	//~ Begin FSkeletalMeshObject Interface
 	virtual void InitResources(USkinnedMeshComponent* InMeshComponent) override;
 	virtual void ReleaseResources() override;
-	virtual void Update(int32 LODIndex,USkinnedMeshComponent* InMeshComponent,const TArray<FActiveMorphTarget>& ActiveMorphTargets, const TArray<float>& MorphTargetWeights, EPreviousBoneTransformUpdateMode PreviousBoneTransformUpdateMode, const FExternalMorphWeightData& InExternalMorphWeightData) override;
+	virtual void Update(int32 LODIndex,USkinnedMeshComponent* InMeshComponent,const FMorphTargetWeightMap& InActiveMorphTargets, const TArray<float>& InMorphTargetWeights, EPreviousBoneTransformUpdateMode PreviousBoneTransformUpdateMode, const FExternalMorphWeightData& InExternalMorphWeightData) override;
 	void UpdateDynamicData_RenderThread(FGPUSkinCache* GPUSkinCache, FRHICommandListImmediate& RHICmdList, FDynamicSkelMeshObjectDataGPUSkin* InDynamicData, FSceneInterface* Scene, uint32 FrameNumberToPrepare, uint32 RevisionNumber);
 	virtual void PreGDMECallback(FGPUSkinCache* GPUSkinCache, uint32 FrameNumber) override;
 	virtual const FVertexFactory* GetSkinVertexFactory(const FSceneView* View, int32 LODIndex,int32 ChunkIdx, ESkinVertexFactoryMode VFMode = ESkinVertexFactoryMode::Default) const override;
@@ -611,7 +611,7 @@ protected:
 		 * @param ActiveMorphTargets - Morph to accumulate. assumed to be weighted and have valid targets
 		 * @param MorphTargetWeights - All Morph weights
 		 */
-		void UpdateMorphVertexBufferCPU(const TArray<FActiveMorphTarget>& ActiveMorphTargets, const TArray<float>& MorphTargetWeights, const TArray<int32>& SectionIdsUseByActiveMorphTargets, 
+		void UpdateMorphVertexBufferCPU(const FMorphTargetWeightMap& InActiveMorphTargets, const TArray<float>& MorphTargetWeights, const TArray<int32>& SectionIdsUseByActiveMorphTargets, 
 										bool bGPUSkinCacheEnabled, FMorphVertexBuffer& MorphVertexBuffer);
 
 		void UpdateMorphVertexBufferGPU(FRHICommandListImmediate& RHICmdList, const TArray<float>& MorphTargetWeights, const FMorphTargetVertexInfoBuffers& MorphTargetVertexInfoBuffers, 
diff --git a/Engine/Source/Runtime/Engine/Private/SkeletalRenderStatic.h b/Engine/Source/Runtime/Engine/Private/SkeletalRenderStatic.h
--- a/Engine/Source/Runtime/Engine/Private/SkeletalRenderStatic.h
+++ b/Engine/Source/Runtime/Engine/Private/SkeletalRenderStatic.h
@@ -26,7 +26,7 @@ public:
 	//~ Begin FSkeletalMeshObject Interface
 	virtual void InitResources(USkinnedMeshComponent* InMeshComponent) override;
 	virtual void ReleaseResources() override;
-	virtual void Update(int32 LODIndex,USkinnedMeshComponent* InMeshComponent,const TArray<FActiveMorphTarget>& ActiveMorphTargets, const TArray<float>& MorphTargetWeights, EPreviousBoneTransformUpdateMode PreviousBoneTransformUpdateMode, const FExternalMorphWeightData& InExternalMorphWeightData) override {};
+	virtual void Update(int32 LODIndex,USkinnedMeshComponent* InMeshComponent,const FMorphTargetWeightMap& InActiveMorphTargets, const TArray<float>& MorphTargetWeights, EPreviousBoneTransformUpdateMode PreviousBoneTransformUpdateMode, const FExternalMorphWeightData& InExternalMorphWeightData) override {};
 	//virtual void UpdateRecomputeTangent(int32 MaterialIndex, int32 LODIndex, bool bRecomputeTangent) override {};
 	virtual void EnableOverlayRendering(bool bEnabled, const TArray<int32>* InBonesOfInterest, const TArray<UMorphTarget*>* InMorphTargetOfInterest) override {};
 	virtual void CacheVertices(int32 LODIndex, bool bForce) const override {};
diff --git a/Engine/Source/Runtime/Engine/Public/AnimationRuntime.h b/Engine/Source/Runtime/Engine/Public/AnimationRuntime.h
--- a/Engine/Source/Runtime/Engine/Public/AnimationRuntime.h
+++ b/Engine/Source/Runtime/Engine/Public/AnimationRuntime.h
@@ -597,7 +597,7 @@ public:
 	/**
 	* Combine CurveKeys (that reference morph targets by name) and ActiveAnims (that reference morphs by reference) into the ActiveMorphTargets array.
 	*/
- 	static void AppendActiveMorphTargets(const USkeletalMesh* InSkeletalMesh, const TMap<FName, float>& InMorphCurveAnims, TArray<FActiveMorphTarget>& InOutActiveMorphTargets, TArray<float>& InOutMorphTargetWeights);
+ 	static void AppendActiveMorphTargets(const USkeletalMesh* InSkeletalMesh, const TMap<FName, float>& InMorphCurveAnims, FMorphTargetWeightMap& InOutActiveMorphTargets, TArray<float>& InOutMorphTargetWeights);
 
 	/**
 	* Retarget a single bone transform, to apply right after extraction.
diff --git a/Engine/Source/Runtime/Engine/Public/SkeletalRenderPublic.h b/Engine/Source/Runtime/Engine/Public/SkeletalRenderPublic.h
--- a/Engine/Source/Runtime/Engine/Public/SkeletalRenderPublic.h
+++ b/Engine/Source/Runtime/Engine/Public/SkeletalRenderPublic.h
@@ -71,7 +71,7 @@ public:
 	 * @param	InSkeletalMeshComponen - parent prim component doing the updating
 	 * @param	ActiveMorphs - morph targets to blend with during skinning
 	 */
-	virtual void Update(int32 LODIndex,USkinnedMeshComponent* InMeshComponent,const TArray<FActiveMorphTarget>& ActiveMorphTargets, const TArray<float>& MorphTargetWeights, EPreviousBoneTransformUpdateMode PreviousBoneTransformUpdateMode, const FExternalMorphWeightData& InExternalMorphWeightData) = 0;
+	virtual void Update(int32 LODIndex,USkinnedMeshComponent* InMeshComponent, const FMorphTargetWeightMap& InActiveMorphTargets, const TArray<float>& MorphTargetWeights, EPreviousBoneTransformUpdateMode PreviousBoneTransformUpdateMode, const FExternalMorphWeightData& InExternalMorphWeightData) = 0;
 
 	/**
 	 * Called by FSkeletalMeshObject prior to GDME. This allows the GPU skin version to update bones etc now that we know we are going to render
