PR #9487: Expose WorldPartitionRuntimeCell as ENGINE_API. (Contributed by [REDACTED_PERSON])

#rnx

diff --git a/Engine/Source/Runtime/Engine/Public/WorldPartition/WorldPartitionRuntimeCell.h b/Engine/Source/Runtime/Engine/Public/WorldPartition/WorldPartitionRuntimeCell.h
--- a/Engine/Source/Runtime/Engine/Public/WorldPartition/WorldPartitionRuntimeCell.h
+++ b/Engine/Source/Runtime/Engine/Public/WorldPartition/WorldPartitionRuntimeCell.h
@@ -133,7 +133,7 @@ static_assert(EWorldPartitionRuntimeCellState::Unloaded < EWorldPartitionRuntime
  * Represents a PIE/Game streaming cell which points to external actor/data chunk packages
  */
 UCLASS(Abstract, Within = WorldPartition)
-class UWorldPartitionRuntimeCell : public UObject, public IWorldPartitionCell
+class ENGINE_API UWorldPartitionRuntimeCell : public UObject, public IWorldPartitionCell
 {
 	GENERATED_UCLASS_BODY()
 
