Get rid of config option Kismet.bRerunConstructionDuringEditorStreaming, as it makes no sense to disable this project-wide while some actors might need to run construction scripts in PIE for levels that were loaded from disk rather than duplicated.

#rnx

diff --git a/Engine/Source/Runtime/Engine/Private/World.cpp b/Engine/Source/Runtime/Engine/Private/World.cpp
--- a/Engine/Source/Runtime/Engine/Private/World.cpp
+++ b/Engine/Source/Runtime/Engine/Private/World.cpp
@@ -2914,17 +2914,9 @@ void UWorld::AddToWorld( ULevel* Level, const FTransform& LevelTransform, bool b
 		TGuardValue<bool> IsEditorLoadingPackage(GIsEditorLoadingPackage, (GIsEditor ? true : GIsEditorLoadingPackage));
 #endif
 
-		bool bRerunConstructionScript = !FPlatformProperties::RequiresCookedData();
-		if (bRerunConstructionScript)
-		{
-			// Config bool that allows disabling all construction scripts during PIE level streaming.
-			bool bRerunConstructionDuringEditorStreaming = true;
-			GConfig->GetBool(TEXT("Kismet"), TEXT("bRerunConstructionDuringEditorStreaming"), /*out*/ bRerunConstructionDuringEditorStreaming, GEngineIni);
-
-			// We don't need to rerun construction scripts if we have cooked data or we are playing in editor unless the PIE world was loaded
-			// from disk rather than duplicated
-			bRerunConstructionScript = !(IsGameWorld() && (Level->bHasRerunConstructionScripts || !bRerunConstructionDuringEditorStreaming));
-		}
+		// We don't need to rerun construction scripts if we have cooked data or we are playing in editor unless the PIE world was loaded
+		// from disk rather than duplicated
+		const bool bRerunConstructionScript = !FPlatformProperties::RequiresCookedData() && (!IsGameWorld() || !Level->bHasRerunConstructionScripts);
 
 		// Prepare context used to store batch/parallelize AddPrimitive calls
 		FRegisterComponentContext Context(this);
@@ -4831,14 +4823,10 @@ void UWorld::InitializeActorsForPlay(const FURL& InURL, bool bResetTime, FRegist
 		URL = InURL;
 	}
 
-	// Config bool that allows disabling all construction scripts during PIE level streaming.
-	bool bRerunConstructionDuringEditorStreaming = true;
-	GConfig->GetBool(TEXT("Kismet"), TEXT("bRerunConstructionDuringEditorStreaming"), /*out*/ bRerunConstructionDuringEditorStreaming, GEngineIni);
-
 	// Update world and the components of all levels.	
 	// We don't need to rerun construction scripts if we have cooked data or we are playing in editor unless the PIE world was loaded
 	// from disk rather than duplicated
-	const bool bRerunConstructionScript = !(FPlatformProperties::RequiresCookedData() || (IsGameWorld() && (PersistentLevel->bHasRerunConstructionScripts || PersistentLevel->bWasDuplicatedForPIE || !bRerunConstructionDuringEditorStreaming)));
+	const bool bRerunConstructionScript = !(FPlatformProperties::RequiresCookedData() || (IsGameWorld() && (PersistentLevel->bHasRerunConstructionScripts || PersistentLevel->bWasDuplicatedForPIE)));
 	UpdateWorldComponents( bRerunConstructionScript, true, Context);
 
 	// Init level gameplay info.
