Allow transient properties when doing property diffing in transactions.

diff --git a/Engine/Source/Runtime/Engine/Private/TransactionCommon.cpp b/Engine/Source/Runtime/Engine/Private/TransactionCommon.cpp
--- a/Engine/Source/Runtime/Engine/Private/TransactionCommon.cpp
+++ b/Engine/Source/Runtime/Engine/Private/TransactionCommon.cpp
@@ -368,7 +368,7 @@ bool FDiffableObjectDataWriter::DoesObjectMatchDiffableObject(const UObject* Obj
 bool FDiffableObjectDataWriter::ShouldSkipProperty(const FProperty* InProperty) const
 {
 	return (PropertiesToSerialize.Num() > 0 && !PropertiesToSerialize.Contains(InProperty))
-		|| InProperty->HasAnyPropertyFlags(CPF_Transient | CPF_NonTransactional | CPF_Deprecated)
+		|| InProperty->HasAnyPropertyFlags(CPF_NonTransactional | CPF_Deprecated)
 		|| Internal::FSerializedObjectDataWriterCommon::ShouldSkipProperty(InProperty);
 }
 
