World Partition: don't trigger post load state changes when the effective count of loaded/unloaded actors is unchanged.

#rnx

diff --git a/Engine/Source/Runtime/Engine/Private/WorldPartition/WorldPartitionActorLoaderInterface.cpp b/Engine/Source/Runtime/Engine/Private/WorldPartition/WorldPartitionActorLoaderInterface.cpp
--- a/Engine/Source/Runtime/Engine/Private/WorldPartition/WorldPartitionActorLoaderInterface.cpp
+++ b/Engine/Source/Runtime/Engine/Private/WorldPartition/WorldPartitionActorLoaderInterface.cpp
@@ -59,23 +59,20 @@ void IWorldPartitionActorLoaderInterface::ILoaderAdapter::Unload()
 	{
 		if (UWorldPartition* WorldPartition = World->GetWorldPartition())
 		{
-			TArray<FWorldPartitionHandle> ActorsToUnload;
-			for (const TPair<FGuid, TMap<FGuid, FWorldPartitionReference>>& ActorRefToUnload : ActorReferences)
-			{
-				ActorsToUnload.Emplace(WorldPartition, ActorRefToUnload.Key);
-			}
-			
 			UnregisterDelegates();
 
 			FScopedSlowTask SlowTask(1, LOCTEXT("UpdatingLoading", "Updating loading..."));
 			SlowTask.MakeDialogDelayed(1.0f);
 
+			int32 NumUnloads = 0;
 			{
 				FWorldPartitionLoadingContext::FDeferred LoadingContext;
 				ActorReferences.Empty();
-				SlowTask.EnterProgressFrame(1);
+				NumUnloads = LoadingContext.GetNumUnregistrations();
 			}
 
+			SlowTask.EnterProgressFrame(1);
+
 			bLoaded = false;
 
 			if (bUserCreated)
@@ -83,7 +80,10 @@ void IWorldPartitionActorLoaderInterface::ILoaderAdapter::Unload()
 				WorldPartition->OnUserCreatedRegionUnloaded();
 			}
 
-			PostLoadedStateChanged(true);
+			if (NumUnloads)
+			{
+				PostLoadedStateChanged(0, NumUnloads);
+			}
 		}
 	}
 }
@@ -172,35 +172,45 @@ void IWorldPartitionActorLoaderInterface::ILoaderAdapter::RefreshLoadedState()
 
 			if (ActorsToLoad.Num() || ActorsToUnload.Num())
 			{
-				FWorldPartitionLoadingContext::FDeferred LoadingContext;
-
-				if (ActorsToLoad.Num())
+				int32 NumLoads = 0;
+				int32 NumUnloads = 0;
 				{
-					FScopedSlowTask SlowTask(ActorsToLoad.Num(), LOCTEXT("ActorsLoading", "Loading actors..."));
-					SlowTask.MakeDialogDelayed(1.0f);
+					FWorldPartitionLoadingContext::FDeferred LoadingContext;
 
-					for (FWorldPartitionHandle& ActorToLoad : ActorsToLoad)
+					if (ActorsToLoad.Num())
 					{
-						FWorldPartitionHandle ActorHandle(WorldPartition, ActorToLoad->GetGuid());
-						AddReferenceToActor(ActorHandle);
-						SlowTask.EnterProgressFrame(1);
+						FScopedSlowTask SlowTask(ActorsToLoad.Num(), LOCTEXT("ActorsLoading", "Loading actors..."));
+						SlowTask.MakeDialogDelayed(1.0f);
+
+						for (FWorldPartitionHandle& ActorToLoad : ActorsToLoad)
+						{
+							FWorldPartitionHandle ActorHandle(WorldPartition, ActorToLoad->GetGuid());
+							AddReferenceToActor(ActorHandle);
+							SlowTask.EnterProgressFrame(1);
+						}
 					}
-				}
 
-				if (ActorsToUnload.Num())
-				{
-					FScopedSlowTask SlowTask(ActorsToUnload.Num(), LOCTEXT("ActorsUnoading", "Unloading actors..."));
-					SlowTask.MakeDialogDelayed(1.0f);
-
-					for (FWorldPartitionHandle& ActorToUnload : ActorsToUnload)
+					if (ActorsToUnload.Num())
 					{
-						RemoveReferenceToActor(ActorToUnload);
-						SlowTask.EnterProgressFrame(1);
+						FScopedSlowTask SlowTask(ActorsToUnload.Num(), LOCTEXT("ActorsUnoading", "Unloading actors..."));
+						SlowTask.MakeDialogDelayed(1.0f);
+
+						for (FWorldPartitionHandle& ActorToUnload : ActorsToUnload)
+						{
+							RemoveReferenceToActor(ActorToUnload);
+							SlowTask.EnterProgressFrame(1);
+						}
 					}
+
+					NumLoads = LoadingContext.GetNumRegistrations();
+					NumUnloads = LoadingContext.GetNumUnregistrations();
 				}
-			}
 
-			PostLoadedStateChanged(ActorsToUnload.Num() > 0);
+				if (NumLoads || NumUnloads)
+				{
+					PostLoadedStateChanged(NumLoads, NumUnloads);
+				}
+			}
 		}
 	}
 }
@@ -211,11 +221,13 @@ bool IWorldPartitionActorLoaderInterface::ILoaderAdapter::ShouldActorBeLoaded(co
 	return PassActorDescFilter(Actor) && PassDataLayersFilter(Actor);
 };
 
-void IWorldPartitionActorLoaderInterface::ILoaderAdapter::PostLoadedStateChanged(bool bUnloadedActors)
+void IWorldPartitionActorLoaderInterface::ILoaderAdapter::PostLoadedStateChanged(int32 NumLoads, int32 NumUnloads)
 {
+	check(NumLoads || NumUnloads);
+
 	if (!IsRunningCommandlet())
 	{
-		if (bUnloadedActors)
+		if (NumUnloads)
 		{
 			GEditor->SelectNone(true, true);
 		}
@@ -225,7 +237,7 @@ void IWorldPartitionActorLoaderInterface::ILoaderAdapter::PostLoadedStateChanged
 
 		GEditor->ResetTransaction(LOCTEXT("LoadingEditorActorResetTrans", "Editor Actors Loading State Changed"));
 
-		if (bUnloadedActors)
+		if (NumUnloads)
 		{
 			GEngine->ForceGarbageCollection(true);
 		}
@@ -266,4 +278,4 @@ void IWorldPartitionActorLoaderInterface::ILoaderAdapter::OnActorDataLayersEdito
 }
 #endif
 
-#undef LOCTEXT_NAMESPACE
\ No newline at end of file
+#undef LOCTEXT_NAMESPACE
diff --git a/Engine/Source/Runtime/Engine/Private/WorldPartition/WorldPartitionHandle.cpp b/Engine/Source/Runtime/Engine/Private/WorldPartition/WorldPartitionHandle.cpp
--- a/Engine/Source/Runtime/Engine/Private/WorldPartition/WorldPartitionHandle.cpp
+++ b/Engine/Source/Runtime/Engine/Private/WorldPartition/WorldPartitionHandle.cpp
@@ -174,6 +174,7 @@ void FWorldPartitionLoadingContext::FDeferred::RegisterActor(FWorldPartitionActo
 		bool bIsAlreadyInSetPtr;
 		ContainerOps.FindOrAdd(Container).Registrations.Add(ActorDesc, &bIsAlreadyInSetPtr);
 		check(!bIsAlreadyInSetPtr);
+		NumRegistrations++;
 
 		check(!ContainerOps.FindChecked(Container).Unregistrations.Contains(ActorDesc));
 	}
@@ -190,6 +191,7 @@ void FWorldPartitionLoadingContext::FDeferred::UnregisterActor(FWorldPartitionAc
 		bool bIsAlreadyInSetPtr;
 		ContainerOps.FindOrAdd(Container).Unregistrations.Add(ActorDesc, &bIsAlreadyInSetPtr);
 		check(!bIsAlreadyInSetPtr);
+		NumUnregistrations++;
 
 		check(!ContainerOps.FindChecked(Container).Registrations.Contains(ActorDesc));
 	}
diff --git a/Engine/Source/Runtime/Engine/Public/WorldPartition/WorldPartitionActorLoaderInterface.h b/Engine/Source/Runtime/Engine/Public/WorldPartition/WorldPartitionActorLoaderInterface.h
--- a/Engine/Source/Runtime/Engine/Public/WorldPartition/WorldPartitionActorLoaderInterface.h
+++ b/Engine/Source/Runtime/Engine/Public/WorldPartition/WorldPartitionActorLoaderInterface.h
@@ -54,7 +54,7 @@ public:
 		bool PassDataLayersFilter(const FWorldPartitionHandle& Actor) const;
 		void RefreshLoadedState();
 
-		void PostLoadedStateChanged(bool bUnloadedActors);
+		void PostLoadedStateChanged(int32 NumLoads, int32 NumUnloads);
 		void AddReferenceToActor(FWorldPartitionHandle& Actor);
 		void RemoveReferenceToActor(FWorldPartitionHandle& Actor);
 		void OnActorDataLayersEditorLoadingStateChanged(bool bFromUserOperation);		
diff --git a/Engine/Source/Runtime/Engine/Public/WorldPartition/WorldPartitionHandle.h b/Engine/Source/Runtime/Engine/Public/WorldPartition/WorldPartitionHandle.h
--- a/Engine/Source/Runtime/Engine/Public/WorldPartition/WorldPartitionHandle.h
+++ b/Engine/Source/Runtime/Engine/Public/WorldPartition/WorldPartitionHandle.h
@@ -51,6 +51,9 @@ public:
 	public:
 		~FDeferred();
 
+		int32 GetNumRegistrations() const { return NumRegistrations; }
+		int32 GetNumUnregistrations() const { return NumUnregistrations; }
+
 	private:
 		virtual void RegisterActor(FWorldPartitionActorDesc* ActorDesc) override;
 		virtual void UnregisterActor(FWorldPartitionActorDesc* ActorDesc) override;
@@ -62,6 +65,9 @@ public:
 		};
 
 		TMap<UActorDescContainer*, FContainerOps> ContainerOps;
+
+		int32 NumRegistrations = 0;
+		int32 NumUnregistrations = 0;
 	};
 
 	static void LoadAndRegisterActor(FWorldPartitionActorDesc* ActorDesc);
