Fix potential crash recording a client replay

diff --git a/Engine/Source/Runtime/Engine/Private/DemoNetDriver.cpp b/Engine/Source/Runtime/Engine/Private/DemoNetDriver.cpp
--- a/Engine/Source/Runtime/Engine/Private/DemoNetDriver.cpp
+++ b/Engine/Source/Runtime/Engine/Private/DemoNetDriver.cpp
@@ -5053,7 +5053,7 @@ void UDemoNetDriver::NotifyActorDestroyed(AActor* Actor, bool IsSeamlessTravel)
 
 void UDemoNetDriver::NotifyActorTornOff(AActor* Actor)
 {
-	if (IsRecording())
+	if (IsRecording() && IsValid(Actor))
 	{
 		// Replicate one last time to the replay stream
 		ReplayHelper.ReplicateActor(Actor, ClientConnections[0], true);
