Add new PreLevelRemovedFromWorld callback for cases where the handler needs to act on the level while it still has its OwningWorld.

diff --git a/Engine/Source/Runtime/Engine/Classes/Engine/World.h b/Engine/Source/Runtime/Engine/Classes/Engine/World.h
--- a/Engine/Source/Runtime/Engine/Classes/Engine/World.h
+++ b/Engine/Source/Runtime/Engine/Classes/Engine/World.h
@@ -4245,6 +4245,11 @@ public:
 	// Sent when a ULevel is added to the world via UWorld::AddToWorld
 	static FOnLevelChanged			LevelAddedToWorld;
 
+	// Sent before a ULevel is removed from the world via UWorld::RemoveFromWorld or 
+	// LoadMap (a NULL object means the LoadMap case, because all levels will be 
+	// removed from the world without a RemoveFromWorld call for each)
+	static FOnLevelChanged			PreLevelRemovedFromWorld;
+
 	// Sent when a ULevel is removed from the world via UWorld::RemoveFromWorld or 
 	// LoadMap (a NULL object means the LoadMap case, because all levels will be 
 	// removed from the world without a RemoveFromWorld call for each)
diff --git a/Engine/Source/Runtime/Engine/Private/UnrealEngine.cpp b/Engine/Source/Runtime/Engine/Private/UnrealEngine.cpp
--- a/Engine/Source/Runtime/Engine/Private/UnrealEngine.cpp
+++ b/Engine/Source/Runtime/Engine/Private/UnrealEngine.cpp
@@ -14344,6 +14344,8 @@ bool UEngine::LoadMap( FWorldContext& WorldContext, FURL URL, class UPendingNetG
 	// Unload the current world
 	if( WorldContext.World() )
 	{
+		// Notify listeners that all levels will be removed before we start tearing down the world.
+		FWorldDelegates::PreLevelRemovedFromWorld.Broadcast(nullptr, WorldContext.World());
 		WorldContext.World()->BeginTearingDown();
 
 		if(!URL.HasOption(TEXT("quiet")) )
diff --git a/Engine/Source/Runtime/Engine/Private/World.cpp b/Engine/Source/Runtime/Engine/Private/World.cpp
--- a/Engine/Source/Runtime/Engine/Private/World.cpp
+++ b/Engine/Source/Runtime/Engine/Private/World.cpp
@@ -366,6 +366,7 @@ FWorldDelegates::FWorldCleanupEvent FWorldDelegates::OnWorldCleanup;
 FWorldDelegates::FWorldCleanupEvent FWorldDelegates::OnPostWorldCleanup;
 FWorldDelegates::FWorldEvent FWorldDelegates::OnPreWorldFinishDestroy;
 FWorldDelegates::FOnLevelChanged FWorldDelegates::LevelAddedToWorld;
+FWorldDelegates::FOnLevelChanged FWorldDelegates::PreLevelRemovedFromWorld;
 FWorldDelegates::FOnLevelChanged FWorldDelegates::LevelRemovedFromWorld;
 FWorldDelegates::FLevelOffsetEvent FWorldDelegates::PostApplyLevelOffset;
 FWorldDelegates::FWorldGetAssetTags FWorldDelegates::GetAssetTags;
@@ -3115,6 +3116,7 @@ void UWorld::RemoveFromWorld( ULevel* Level, bool bAllowIncrementalRemoval, FNet
 
 	auto BeginRemoval = [Level, this]()
 	{
+		FWorldDelegates::PreLevelRemovedFromWorld.Broadcast(Level, this);
 		Level->bIsBeingRemoved = true;
 
 		if (Level->bRequireFullVisibilityToRender && Level->bIsVisible)
@@ -7628,6 +7630,7 @@ bool UWorld::RemoveLevel( ULevel* InLevel )
 	bool bRemovedLevel = false;
 	if(ContainsLevel( InLevel ) == true )
 	{
+		FWorldDelegates::PreLevelRemovedFromWorld.Broadcast(InLevel, this);
 		bRemovedLevel = true;
 		
 #if WITH_EDITOR
