Enable mobile ambient occlusin on LDR

diff --git a/Engine/Source/Runtime/Engine/Private/SceneUtils.cpp b/Engine/Source/Runtime/Engine/Private/SceneUtils.cpp
--- a/Engine/Source/Runtime/Engine/Private/SceneUtils.cpp
+++ b/Engine/Source/Runtime/Engine/Private/SceneUtils.cpp
@@ -150,7 +150,7 @@ ENGINE_API uint32 GetDefaultMSAACount(const FStaticFeatureLevel InFeatureLevel,
 				bool bMobilePixelProjectedReflection = IsUsingMobilePixelProjectedReflection(ShaderPlatform);
 				
 				static const auto MobileEarlyZPassCVar = IConsoleManager::Get().FindTConsoleVariableDataInt(TEXT("r.Mobile.EarlyZPass"));
-				const bool bMobileAmbientOcclusion = UseMobileAmbientOcclusion(ShaderPlatform) && IsMobileHDR();
+				const bool bMobileAmbientOcclusion = IsMobileAmbientOcclusionEnabled(ShaderPlatform);
 				const bool bMobileUsesShadowMaskTexture = MobileUsesShadowMaskTexture(ShaderPlatform);
 				bool bIsFullDepthPrepassEnabled = bMobileAmbientOcclusion || bMobileUsesShadowMaskTexture || MobileEarlyZPassCVar->GetValueOnAnyThread() > 0;
 
diff --git a/Engine/Source/Runtime/RenderCore/Private/RenderUtils.cpp b/Engine/Source/Runtime/RenderCore/Private/RenderUtils.cpp
--- a/Engine/Source/Runtime/RenderCore/Private/RenderUtils.cpp
+++ b/Engine/Source/Runtime/RenderCore/Private/RenderUtils.cpp
@@ -1130,7 +1130,7 @@ RENDERCORE_API int32 EnabledShadingModelsMask(const FStaticShaderPlatform Platfo
 RENDERCORE_API uint64 GMobileAmbientOcclusionPlatformMask = 0;
 static_assert(SP_NumPlatforms <= sizeof(GMobileAmbientOcclusionPlatformMask) * 8, "GMobileAmbientOcclusionPlatformMask must be large enough to support all shader platforms");
 
-RENDERCORE_API bool UseMobileAmbientOcclusion(const FStaticShaderPlatform Platform)
+RENDERCORE_API bool IsMobileAmbientOcclusionEnabled(const FStaticShaderPlatform Platform)
 {
 	return (IsMobilePlatform(Platform) && !!(GMobileAmbientOcclusionPlatformMask & (1ull << Platform)));
 }
diff --git a/Engine/Source/Runtime/RenderCore/Public/RenderUtils.h b/Engine/Source/Runtime/RenderCore/Public/RenderUtils.h
--- a/Engine/Source/Runtime/RenderCore/Public/RenderUtils.h
+++ b/Engine/Source/Runtime/RenderCore/Public/RenderUtils.h
@@ -498,7 +498,7 @@ RENDERCORE_API bool AllowPerPixelShadingModels(const FStaticShaderPlatform Platf
 
 RENDERCORE_API int32 EnabledShadingModelsMask(const FStaticShaderPlatform Platform);
 
-RENDERCORE_API bool UseMobileAmbientOcclusion(const FStaticShaderPlatform Platform);
+RENDERCORE_API bool IsMobileAmbientOcclusionEnabled(const FStaticShaderPlatform Platform);
 
 RENDERCORE_API bool IsMobileDistanceFieldEnabled(const FStaticShaderPlatform Platform);
 
diff --git a/Engine/Source/Runtime/Renderer/Private/PostProcess/PostProcessAmbientOcclusionMobile.cpp b/Engine/Source/Runtime/Renderer/Private/PostProcess/PostProcessAmbientOcclusionMobile.cpp
--- a/Engine/Source/Runtime/Renderer/Private/PostProcess/PostProcessAmbientOcclusionMobile.cpp
+++ b/Engine/Source/Runtime/Renderer/Private/PostProcess/PostProcessAmbientOcclusionMobile.cpp
@@ -74,11 +74,6 @@ DECLARE_GPU_STAT_NAMED(MobileSSAO, TEXT("SSAO"));
 
 // --------------------------------------------------------------------------------------------------------------------
 
-bool IsMobileAmbientOcclusionEnabled(EShaderPlatform ShaderPlatform)
-{
-	return UseMobileAmbientOcclusion(ShaderPlatform) && IsMobileHDR();
-}
-
 bool IsUsingMobileAmbientOcclusion(EShaderPlatform ShaderPlatform)
 {
 	static const auto MobileAmbientOcclusionQualityCVar = IConsoleManager::Get().FindTConsoleVariableDataInt(TEXT("r.Mobile.AmbientOcclusionQuality"));
diff --git a/Engine/Source/Runtime/Renderer/Private/PostProcess/PostProcessAmbientOcclusionMobile.h b/Engine/Source/Runtime/Renderer/Private/PostProcess/PostProcessAmbientOcclusionMobile.h
--- a/Engine/Source/Runtime/Renderer/Private/PostProcess/PostProcessAmbientOcclusionMobile.h
+++ b/Engine/Source/Runtime/Renderer/Private/PostProcess/PostProcessAmbientOcclusionMobile.h
@@ -10,5 +10,4 @@
 
 FRDGTextureRef CreateMobileScreenSpaceAOTexture(FRDGBuilder& GraphBuilder, const struct FSceneTexturesConfig& Config);
 
-bool IsMobileAmbientOcclusionEnabled(EShaderPlatform ShaderPlatform);
 bool IsUsingMobileAmbientOcclusion(EShaderPlatform ShaderPlatform);
\ No newline at end of file
