Actor Descriptor: avoid an assert calling StaticFindObject during save, since we are only retrieving a native class.

#rnx

diff --git a/Engine/Source/Runtime/Engine/Private/WorldPartition/WorldPartitionActorDescUtils.cpp b/Engine/Source/Runtime/Engine/Private/WorldPartition/WorldPartitionActorDescUtils.cpp
--- a/Engine/Source/Runtime/Engine/Private/WorldPartition/WorldPartitionActorDescUtils.cpp
+++ b/Engine/Source/Runtime/Engine/Private/WorldPartition/WorldPartitionActorDescUtils.cpp
@@ -22,6 +22,10 @@ UClass* FWorldPartitionActorDescUtils::GetActorNativeClassFromAssetData(const FA
 	FString ActorMetaDataClass;
 	if (InAssetData.GetTagValue(NAME_ActorMetaDataClass, ActorMetaDataClass))
 	{
+		// Avoid an assert when calling StaticFindObject during save to retrieve the actor's class.
+		// Since we are only looking for a native class, the call to StaticFindObject is legit.
+		TGuardValue<bool> GIsSavingPackageGuard(GIsSavingPackage, false);
+
 		FString ActorClassName;
 		FString ActorPackageName;
 		if (!ActorMetaDataClass.Split(TEXT("."), &ActorPackageName, &ActorClassName))
@@ -75,20 +79,20 @@ TUniquePtr<FWorldPartitionActorDesc> FWorldPartitionActorDescUtils::GetActorDesc
 	return nullptr;
 }
 
-void FWorldPartitionActorDescUtils::AppendAssetDataTagsFromActor(const AActor* Actor, TArray<UObject::FAssetRegistryTag>& OutTags)
+void FWorldPartitionActorDescUtils::AppendAssetDataTagsFromActor(const AActor* InActor, TArray<UObject::FAssetRegistryTag>& OutTags)
 {
-	check(Actor->IsPackageExternal());
+	check(InActor->IsPackageExternal());
 	
-	TUniquePtr<FWorldPartitionActorDesc> ActorDesc(Actor->CreateActorDesc());
+	TUniquePtr<FWorldPartitionActorDesc> ActorDesc(InActor->CreateActorDesc());
 
 	// If the actor is not added to a world, we can't retrieve its bounding volume, so try to get the existing one
-	if (ULevel* Level = Actor->GetLevel(); !Level || !Level->Actors.Contains(Actor))
+	if (ULevel* Level = InActor->GetLevel(); !Level || !Level->Actors.Contains(InActor))
 	{
 		IAssetRegistry& AssetRegistry = FModuleManager::LoadModuleChecked<FAssetRegistryModule>(TEXT("AssetRegistry")).Get();
 
 		FARFilter Filter;
 		Filter.bIncludeOnlyOnDiskAssets = true;
-		Filter.PackageNames.Add(Actor->GetPackage()->GetFName());
+		Filter.PackageNames.Add(InActor->GetPackage()->GetFName());
 
 		TArray<FAssetData> Assets;
 		AssetRegistry.GetAssets(Filter, Assets);
@@ -102,7 +106,7 @@ void FWorldPartitionActorDescUtils::AppendAssetDataTagsFromActor(const AActor* A
 		}
 	}
 
-	const FString ActorMetaDataClass = GetParentNativeClass(Actor->GetClass())->GetPathName();
+	const FString ActorMetaDataClass = GetParentNativeClass(InActor->GetClass())->GetPathName();
 	OutTags.Add(UObject::FAssetRegistryTag(NAME_ActorMetaDataClass, ActorMetaDataClass, UObject::FAssetRegistryTag::TT_Hidden));
 
 	TArray<uint8> SerializedData;
@@ -111,18 +115,18 @@ void FWorldPartitionActorDescUtils::AppendAssetDataTagsFromActor(const AActor* A
 	OutTags.Add(UObject::FAssetRegistryTag(NAME_ActorMetaData, ActorMetaData, UObject::FAssetRegistryTag::TT_Hidden));
 }
 
-void FWorldPartitionActorDescUtils::UpdateActorDescriptorFomActor(const AActor* Actor, TUniquePtr<FWorldPartitionActorDesc>& ActorDesc)
+void FWorldPartitionActorDescUtils::UpdateActorDescriptorFomActor(const AActor* InActor, TUniquePtr<FWorldPartitionActorDesc>& ActorDesc)
 {
-	TUniquePtr<FWorldPartitionActorDesc> NewActorDesc(Actor->CreateActorDesc());
+	TUniquePtr<FWorldPartitionActorDesc> NewActorDesc(InActor->CreateActorDesc());
 	NewActorDesc->TransferFrom(ActorDesc.Get());
 	ActorDesc = MoveTemp(NewActorDesc);
 }
 
-void FWorldPartitionActorDescUtils::ReplaceActorDescriptorPointerFromActor(const AActor* OldActor, AActor* NewActor, FWorldPartitionActorDesc* ActorDesc)
+void FWorldPartitionActorDescUtils::ReplaceActorDescriptorPointerFromActor(const AActor* InOldActor, AActor* InNewActor, FWorldPartitionActorDesc* InActorDesc)
 {
-	check(!NewActor || (OldActor->GetActorGuid() == NewActor->GetActorGuid()));
-	check(!NewActor || (NewActor->GetActorGuid() == ActorDesc->GetGuid()));
-	check(!ActorDesc->ActorPtr.IsValid() || (ActorDesc->ActorPtr == OldActor));
-	ActorDesc->ActorPtr = NewActor;
+	check(!InNewActor || (InOldActor->GetActorGuid() == InNewActor->GetActorGuid()));
+	check(!InNewActor || (InNewActor->GetActorGuid() == InActorDesc->GetGuid()));
+	check(!InActorDesc->ActorPtr.IsValid() || (InActorDesc->ActorPtr == InOldActor));
+	InActorDesc->ActorPtr = InNewActor;
 }
 #endif
\ No newline at end of file
diff --git a/Engine/Source/Runtime/Engine/Public/WorldPartition/WorldPartitionActorDescUtils.h b/Engine/Source/Runtime/Engine/Public/WorldPartition/WorldPartitionActorDescUtils.h
--- a/Engine/Source/Runtime/Engine/Public/WorldPartition/WorldPartitionActorDescUtils.h
+++ b/Engine/Source/Runtime/Engine/Public/WorldPartition/WorldPartitionActorDescUtils.h
@@ -11,11 +11,47 @@ struct FAssetData;
 
 struct ENGINE_API FWorldPartitionActorDescUtils
 {
+	/** 
+	 * Checks if the provided asset data contains a valid actor descriptor.
+	 * @param InAssetData	The asset data to look into.
+	 * @return				Whether the provideed asset data contains a valid actor descriptor or not.
+	 */
 	static bool IsValidActorDescriptorFromAssetData(const FAssetData& InAssetData);
+
+	/** 
+	 * Retrieve the actor's native class from the asset data.
+	 * @param InAssetData	The asset data to look into.
+	 * @return				The actor's native class.
+	 */
 	static UClass* GetActorNativeClassFromAssetData(const FAssetData& InAssetData);
+
+	/** 
+	 * Creates a valid actor descriptor from the provided asset data.
+	 * @param InAssetData	The asset data to look into.
+	 * @return				Actor descriptor retrieved from the provided asset data..
+	 */
 	static TUniquePtr<FWorldPartitionActorDesc> GetActorDescriptorFromAssetData(const FAssetData& InAssetData);
-	static void AppendAssetDataTagsFromActor(const AActor* Actor, TArray<UObject::FAssetRegistryTag>& OutTags);
-	static void UpdateActorDescriptorFomActor(const AActor* Actor, TUniquePtr<FWorldPartitionActorDesc>& ActorDesc);
-	static void ReplaceActorDescriptorPointerFromActor(const AActor* OldActor, AActor* NewActor, FWorldPartitionActorDesc* ActorDesc);
+
+	/** 
+	 * Appends the actor's actor descriptor data into the provided asset registry tags.
+	 * @param InActor		The actor that will append its actor descriptor.
+	 * @param OutTags		Output tags to output into.
+	 */
+	static void AppendAssetDataTagsFromActor(const AActor* InActor, TArray<UObject::FAssetRegistryTag>& OutTags);
+
+	/** 
+	 * Update an actor descriptor with new values coming from the provided actor.
+	 * @param InActor		The actor to update from..
+	 * @param OutActorDesc	Actor descriptor unique pointer that will get updated with the new actor descriptor.
+	 */
+	static void UpdateActorDescriptorFomActor(const AActor* InActor, TUniquePtr<FWorldPartitionActorDesc>& OutActorDesc);
+
+	/** 
+	 * Replaces the actor descriptor's actor pointer with the provided new actor pointer.
+	 * @param InOldActor	The old actor that the provided actor descriptor was representing.
+	 * @param InNewActor	The new actor that the provided actor descriptor should be representing.
+	 * @param InActorDesc	Actor descriptor that will get its actor pointer updated.
+	 */
+	static void ReplaceActorDescriptorPointerFromActor(const AActor* InOldActor, AActor* InNewActor, FWorldPartitionActorDesc* InActorDesc);
 };
 #endif
\ No newline at end of file
