Strata - made the ConvertToDecal node force using parameter blending now that it is a requiremeent.
#prelfight none

diff --git a/Engine/Shaders/Shared/StrataDefinitions.h b/Engine/Shaders/Shared/StrataDefinitions.h
--- a/Engine/Shaders/Shared/StrataDefinitions.h
+++ b/Engine/Shaders/Shared/StrataDefinitions.h
@@ -8,7 +8,7 @@
 #pragma once
 
 // Change this to force recompilation of all strata dependent shaders (use https://www.random.org/cgi-bin/randbyte?nbytes=4&format=h)
-#define STRATA_SHADER_VERSION 0x97e993a4
+#define STRATA_SHADER_VERSION 0xc5265453
 
 
 
diff --git a/Engine/Source/Runtime/Engine/Classes/Materials/MaterialExpressionStrata.h b/Engine/Source/Runtime/Engine/Classes/Materials/MaterialExpressionStrata.h
--- a/Engine/Source/Runtime/Engine/Classes/Materials/MaterialExpressionStrata.h
+++ b/Engine/Source/Runtime/Engine/Classes/Materials/MaterialExpressionStrata.h
@@ -781,10 +781,6 @@ class UMaterialExpressionStrataConvertToDecal : public UMaterialExpressionStrata
 	UPROPERTY()
 	FExpressionInput Coverage;
 
-	// When a material is converted to decals, we use parameter blending to ensure that a single slab of matter is provided (current restriction).
-	UPROPERTY()
-	uint32 bUseParameterBlending : 1;
-
 	//~ Begin UMaterialExpression Interface
 #if WITH_EDITOR
 	virtual int32 Compile(class FMaterialCompiler* Compiler, int32 OutputIndex) override;
diff --git a/Engine/Source/Runtime/Engine/Private/Materials/MaterialExpressions.cpp b/Engine/Source/Runtime/Engine/Private/Materials/MaterialExpressions.cpp
--- a/Engine/Source/Runtime/Engine/Private/Materials/MaterialExpressions.cpp
+++ b/Engine/Source/Runtime/Engine/Private/Materials/MaterialExpressions.cpp
@@ -23040,7 +23040,6 @@ FStrataOperator* UMaterialExpressionStrataPostProcess::StrataGenerateMaterialTop
 
 UMaterialExpressionStrataConvertToDecal::UMaterialExpressionStrataConvertToDecal(const FObjectInitializer& ObjectInitializer)
 	: Super(ObjectInitializer)
-	, bUseParameterBlending(1)
 {
 	struct FConstructorStatics
 	{
@@ -23066,35 +23065,23 @@ int32 UMaterialExpressionStrataConvertToDecal::Compile(class FMaterialCompiler*
 
 	int32 OutputCodeChunk = INDEX_NONE;
 	FStrataOperator& StrataOperator = Compiler->StrataCompilationGetOperator(this);
-	if (bUseParameterBlending)
+	if (!StrataOperator.bUseParameterBlending)
 	{
-		if (!StrataOperator.bUseParameterBlending)
-		{
-			return Compiler->Errorf(TEXT("Strata Convert To Decal node is not receiveing a parameter blended strata material tree."));
-		}
-		else if (!StrataOperator.bRootOfParameterBlendingSubTree)
-		{
-			return Compiler->Errorf(TEXT("Strata Convert To Decal node must be the root of a parameter blending sub tree."));
-		}
-
-		if (StrataOperator.bUseParameterBlending)
-		{
-			// Propagate the parameter blended normal
-			FStrataOperator* Operator = Compiler->StrataCompilationGetOperatorFromIndex(StrataOperator.LeftIndex);
-			if (StrataOperator.bRootOfParameterBlendingSubTree)
-			{
-				StrataOperator.BSDFRegisteredSharedLocalBasis = Operator->BSDFRegisteredSharedLocalBasis;
-			}
-
-			OutputCodeChunk = Compiler->StrataWeightParameterBlending(
-				DecalMaterialCodeChunk, CoverageCodeChunk,
-				StrataOperator.bRootOfParameterBlendingSubTree ? &StrataOperator : nullptr);
-		}
+		return Compiler->Errorf(TEXT("Strata Convert To Decal node must recveive StrataData a parameter blended strata material sub tree."));
 	}
-	else
+	if (!StrataOperator.bRootOfParameterBlendingSubTree)
 	{
-		OutputCodeChunk = Compiler->StrataWeight(DecalMaterialCodeChunk, CoverageCodeChunk, StrataOperator.Index, StrataOperator.MaxDistanceFromLeaves);
+		return Compiler->Errorf(TEXT("Strata Convert To Decal node must be the root of a parameter blending sub tree: no more strata operations can be applied a over its output."));
 	}
+
+	// Propagate the parameter blended normal
+	FStrataOperator* Operator = Compiler->StrataCompilationGetOperatorFromIndex(StrataOperator.LeftIndex);
+	StrataOperator.BSDFRegisteredSharedLocalBasis = Operator->BSDFRegisteredSharedLocalBasis;
+
+	OutputCodeChunk = Compiler->StrataWeightParameterBlending(
+		DecalMaterialCodeChunk, CoverageCodeChunk,
+		StrataOperator.bRootOfParameterBlendingSubTree ? &StrataOperator : nullptr);
+
 	return OutputCodeChunk;
 }
 
@@ -23140,6 +23127,7 @@ void UMaterialExpressionStrataConvertToDecal::GatherStrataMaterialInfo(FStrataMa
 
 FStrataOperator* UMaterialExpressionStrataConvertToDecal::StrataGenerateMaterialTopologyTree(class FMaterialCompiler* Compiler, class UMaterialExpression* Parent, int32 OutputIndex)
 {
+	const bool bUseParameterBlending = true;
 	FStrataOperator& StrataOperator = Compiler->StrataCompilationRegisterOperator(STRATA_OPERATOR_WEIGHT, this, Parent, bUseParameterBlending);
 
 	UMaterialExpression* ChildDecalMaterialExpression = DecalMaterial.GetTracedInput().Expression;
