Fixed bug when detecting and fixing duplicate actor folders.

diff --git a/Engine/Source/Runtime/Engine/Classes/Engine/Level.h b/Engine/Source/Runtime/Engine/Classes/Engine/Level.h
--- a/Engine/Source/Runtime/Engine/Classes/Engine/Level.h
+++ b/Engine/Source/Runtime/Engine/Classes/Engine/Level.h
@@ -120,7 +120,7 @@ struct ENGINE_API FActorFolderSet
 
 public:
 
-	void Add(UActorFolder* InActorFolder) { ActorFolders.Add(InActorFolder); }
+	void Add(UActorFolder* InActorFolder);
 	int32 Remove(UActorFolder* InActorFolder) { return ActorFolders.Remove(InActorFolder); }
 	bool IsEmpty() const { return ActorFolders.IsEmpty(); }
 	const TSet<TObjectPtr<UActorFolder>>& GetActorFolders() const { return ActorFolders; }
diff --git a/Engine/Source/Runtime/Engine/Private/ActorFolder.cpp b/Engine/Source/Runtime/Engine/Private/ActorFolder.cpp
--- a/Engine/Source/Runtime/Engine/Private/ActorFolder.cpp
+++ b/Engine/Source/Runtime/Engine/Private/ActorFolder.cpp
@@ -191,7 +191,7 @@ void UActorFolder::MarkAsDeleted()
 			}
 		}
 		return true;
-	}, /*bSkipDeleted*/ false);
+	}, /*bSkipDeleted*/ true);
 
 	// Sort in descending order so children will be deleted before parents
 	FoldersToDelete.Sort([](const UActorFolder& FolderA, const UActorFolder& FolderB)
diff --git a/Engine/Source/Runtime/Engine/Private/Level.cpp b/Engine/Source/Runtime/Engine/Private/Level.cpp
--- a/Engine/Source/Runtime/Engine/Private/Level.cpp
+++ b/Engine/Source/Runtime/Engine/Private/Level.cpp
@@ -237,6 +237,35 @@ void FLevelPartitionOperationScope::DestroyTransientLevel(ULevel* Level)
 }
 #endif
 
+/*-----------------------------------------------------------------------------
+FActorFolderSet implementation.
+-----------------------------------------------------------------------------*/
+#if WITH_EDITOR
+static bool GIsFixingActorFolders = false;
+#endif
+
+void FActorFolderSet::Add(UActorFolder* InActorFolder)
+{
+	check(InActorFolder);
+
+#if WITH_EDITOR
+	if (!IsEmpty() && !GIsFixingActorFolders)
+	{
+		const FName AddedFolderPath = InActorFolder->GetPath();
+		for (const TObjectPtr<UActorFolder>& Folder : ActorFolders)
+		{
+			if (AddedFolderPath.IsEqual(Folder->GetPath(), ENameCase::IgnoreCase))
+			{
+				UE_LOG(LogLevel, Error, TEXT("Adding duplicate actor folder %s with path %s."), *InActorFolder->GetName(), *AddedFolderPath.ToString());
+				break;
+			}
+		}
+	}
+#endif
+
+	ActorFolders.Add(InActorFolder);
+}
+
 /*-----------------------------------------------------------------------------
 ULevel implementation.
 -----------------------------------------------------------------------------*/
@@ -2654,7 +2683,7 @@ void ULevel::OnFolderMarkAsDeleted(UActorFolder* InActorFolder)
 	{
 		Modify(false);
 		FActorFolderSet& Folders = FolderLabelToActorFolders.FindChecked(InActorFolder->GetLabel());
-		verify(Folders.Remove(InActorFolder));
+		verify(Folders.Remove(InActorFolder) || GIsFixingActorFolders);
 		if (Folders.IsEmpty())
 		{
 			FolderLabelToActorFolders.Remove(InActorFolder->GetLabel());
@@ -2679,7 +2708,7 @@ void ULevel::OnFolderLabelChanged(UActorFolder* InActorFolder, const FString& In
 	{
 		Modify(false);
 		FActorFolderSet& OldLabelFolders = FolderLabelToActorFolders.FindChecked(InOldFolderLabel);
-		verify(OldLabelFolders.Remove(InActorFolder));
+		verify(OldLabelFolders.Remove(InActorFolder) || GIsFixingActorFolders);
 		if (OldLabelFolders.IsEmpty())
 		{
 			FolderLabelToActorFolders.Remove(InOldFolderLabel);
@@ -2698,6 +2727,8 @@ void ULevel::FixupActorFolders()
 
 	if (IsUsingActorFolders())
 	{
+		TGuardValue<bool> FixingActorFolders(GIsFixingActorFolders, true);
+
 		// At this point, LoadedExternalActorFolders are fully loaded, transfer them to the ActorFolders list.
 		for (UActorFolder* LoadedActorFolder : LoadedExternalActorFolders)
 		{
@@ -2706,43 +2737,100 @@ void ULevel::FixupActorFolders()
 		}
 		LoadedExternalActorFolders.Empty();
 
-		const bool bFixDuplicateFolders = !IsRunningCommandlet();
-	
-		// Discover duplicate paths first to prioritize non-duplicate paths
-		TMap<FName, UActorFolder*> ExistingPaths;
-		TArray<UActorFolder*> FoldersToDelete;
-		TMap<UActorFolder*, UActorFolder*> DuplicateFolders;
-		ForEachActorFolder([bFixDuplicateFolders, &ExistingPaths, &DuplicateFolders, &FoldersToDelete](UActorFolder* ActorFolder)
+		ForEachActorFolder([](UActorFolder* ActorFolder)
 		{
 			// Detects and clears invalid parent folder
 			ActorFolder->FixupParentFolder();
+			return true;
+		}, /*bSkipDeleted*/ true);
 
-			if (bFixDuplicateFolders)
+		const bool bFixDuplicateFolders = !IsRunningCommandlet();
+		if (bFixDuplicateFolders)
+		{
+			// Build a sorted list of duplicate paths and a map of duplicate path to corresponding actor folders
+			TArray<FName> SortedDuplicatePaths;
+			TMap<FName, TArray<UActorFolder*>> PathToFolders;
+			ForEachActorFolder([&SortedDuplicatePaths, &PathToFolders](UActorFolder* ActorFolder)
 			{
-				FName FolderPath = ActorFolder->GetPath();
-				UActorFolder** ExistingFolder = ExistingPaths.Find(FolderPath);
-				if (!ExistingFolder)
+				FName Path = ActorFolder->GetPath();
+				if (PathToFolders.FindOrAdd(ActorFolder->GetPath()).Add(ActorFolder) == 1)
 				{
-					ExistingPaths.Add(FolderPath, ActorFolder);
+					SortedDuplicatePaths.Add(Path);
 				}
-				else
+				return true;
+			}, /*bSkipDeleted*/ true);
+
+			SortedDuplicatePaths.Sort([](const FName& FolderPathA, const FName& FolderPathB)
+			{
+				return FolderPathA.LexicalLess(FolderPathB);
+			});
+
+			TSet<UActorFolder*> FoldersToDelete;
+			auto HasParentToDelete = [&FoldersToDelete](UActorFolder* InFolder)
+			{
+				UActorFolder* Parent = InFolder->GetParent();
+				while (Parent)
+				{
+					if (FoldersToDelete.Contains(Parent))
+					{
+						return true;
+					}
+					Parent = Parent->GetParent();
+				}
+				return false;
+			};
+
+			TSet<UActorFolder*> FoldersToKeep;
+			TMap<UActorFolder*, UActorFolder*> DuplicateFolders;
+
+			for (FName& DuplicatePath : SortedDuplicatePaths)
+			{
+				// Choose to keep the folder that has no parent to delete
+				TArray<UActorFolder*>& Folders = PathToFolders.FindChecked(DuplicatePath);
+				UActorFolder* FolderToKeep = nullptr;
+				for (UActorFolder* Folder : Folders)
 				{
-					FoldersToDelete.Add(ActorFolder);
-					DuplicateFolders.Add(ActorFolder, *ExistingFolder);
+					if (!HasParentToDelete(Folder))
+					{
+						FolderToKeep = Folder;
+						break;
+					}
+				}
+
+				// Validation
+				check(FolderToKeep);
+				bool bAlreadyExists = false;
+				FoldersToKeep.Add(FolderToKeep, &bAlreadyExists);
+				check(!bAlreadyExists);
+
+				// All other folders are considered duplicated and will be marked as deleted
+				for (UActorFolder* Folder : Folders)
+				{
+					if (Folder != FolderToKeep)
+					{
+						FoldersToDelete.Add(Folder);
+
+						// Cleanup FolderLabelToActorFolders from duplicate folder as we don't want GetActorFolder(Path) to resolve to a duplicate ActorFolder
+						FActorFolderSet* FoundSet = FolderLabelToActorFolders.Find(Folder->GetLabel());
+						// We should have at least 2 elements in the set (the one we keep and one or multiple duplicates)
+						check(FoundSet && FoundSet->GetActorFolders().Num() > 1 && FoundSet->GetActorFolders().Contains(Folder));
+						verify(FoundSet->Remove(Folder));
+						check(!FoundSet->IsEmpty());
+
+						// Keep a map that will be used to retrieve the folder we keep for each duplicate/deleted folder
+						DuplicateFolders.Add(Folder, FolderToKeep);
+					}
 				}
 			}
-			return true;
-		}, /*bSkipDeleted*/ true);
 
-		if (bFixDuplicateFolders)
-		{
 			// Sort in descending order so children will be deleted before parents
-			FoldersToDelete.Sort([](const UActorFolder& FolderA, const UActorFolder& FolderB)
+			TArray<UActorFolder*> SortedFoldersToDelete = FoldersToDelete.Array();
+			SortedFoldersToDelete.Sort([](const UActorFolder& FolderA, const UActorFolder& FolderB)
 			{
 				return FolderB.GetPath().LexicalLess(FolderA.GetPath());
 			});
 
-			for (UActorFolder* FolderToDelete : FoldersToDelete)
+			for (UActorFolder* FolderToDelete : SortedFoldersToDelete)
 			{
 				UActorFolder* NewParent = DuplicateFolders.FindChecked(FolderToDelete);
 				// First move duplicate folder under the single folder we keep. Use a unique name to avoid dealing with name clash
