After SkinCache entry is released we also reset RayTracingInitializer which means that in the next update if there are no geometry segments in the initializer we need to recreate RayTracingGeometry.

diff --git a/Engine/Source/Runtime/Engine/Private/SkeletalRenderGPUSkin.cpp b/Engine/Source/Runtime/Engine/Private/SkeletalRenderGPUSkin.cpp
--- a/Engine/Source/Runtime/Engine/Private/SkeletalRenderGPUSkin.cpp
+++ b/Engine/Source/Runtime/Engine/Private/SkeletalRenderGPUSkin.cpp
@@ -792,7 +792,9 @@ void FSkeletalMeshObjectGPUSkin::UpdateRayTracingGeometry(FSkeletalMeshLODRender
 	{
 		const bool bAnySegmentUsesWorldPositionOffset = DynamicData != nullptr ? DynamicData->bAnySegmentUsesWorldPositionOffset : false;
 
-		bool bRequireRecreatingRayTracingGeometry = (LODIndex != RayTracingGeometry.LODIndex || bHiddenMaterialVisibilityDirtyForRayTracing);
+		bool bRequireRecreatingRayTracingGeometry = LODIndex != RayTracingGeometry.LODIndex
+			|| bHiddenMaterialVisibilityDirtyForRayTracing
+			|| RayTracingGeometry.Initializer.Segments.Num() == 0;
 		if (!bRequireRecreatingRayTracingGeometry)
 		{
 			for (FRayTracingGeometrySegment& Segment : RayTracingGeometry.Initializer.Segments)
@@ -892,6 +894,8 @@ void FSkeletalMeshObjectGPUSkin::UpdateRayTracingGeometry(FSkeletalMeshLODRender
 		}
 		else if (!bAnySegmentUsesWorldPositionOffset)
 		{
+			check(LODModel.RenderSections.Num() == RayTracingGeometry.Initializer.Segments.Num());
+
 			// Refit BLAS with new vertex buffer data
 			for (int32 SectionIndex = 0; SectionIndex < LODModel.RenderSections.Num(); ++SectionIndex)
 			{
