Make deformer settings object transient.
They are recreated on component instantiation.

diff --git a/Engine/Source/Runtime/Engine/Classes/Components/SkinnedMeshComponent.h b/Engine/Source/Runtime/Engine/Classes/Components/SkinnedMeshComponent.h
--- a/Engine/Source/Runtime/Engine/Classes/Components/SkinnedMeshComponent.h
+++ b/Engine/Source/Runtime/Engine/Classes/Components/SkinnedMeshComponent.h
@@ -284,7 +284,7 @@ public:
 	UPROPERTY(Transient, BlueprintReadOnly, Category = "Deformer")
 	TObjectPtr<UMeshDeformerInstance> MeshDeformerInstance;
 
-	UPROPERTY(VisibleAnywhere, BlueprintReadOnly, DisplayName = "Settings", Category = "Deformer", meta=(EditInline, EditCondition="MeshDeformerInstanceSettings!=nullptr", HideEditConditionToggle, EditConditionHides))
+	UPROPERTY(Transient, VisibleAnywhere, BlueprintReadOnly, DisplayName = "Settings", Category = "Deformer", meta=(EditInline, EditCondition="MeshDeformerInstanceSettings!=nullptr", HideEditConditionToggle, EditConditionHides))
 	TObjectPtr<UMeshDeformerInstanceSettings> MeshDeformerInstanceSettings;
 
 	/** const getters for previous transform idea */
