Changed ObjectRadius node behavior to return the radius of the sphere that encompasses the instance local bounds extent instead of the world bounds extent. This makes the bounds wrap the instance only, instead of the whole Instance Static Mesh group.

diff --git a/Engine/Source/Runtime/Engine/Public/PrimitiveUniformShaderParameters.h b/Engine/Source/Runtime/Engine/Public/PrimitiveUniformShaderParameters.h
--- a/Engine/Source/Runtime/Engine/Public/PrimitiveUniformShaderParameters.h
+++ b/Engine/Source/Runtime/Engine/Public/PrimitiveUniformShaderParameters.h
@@ -206,10 +206,9 @@ public:
 	inline FPrimitiveUniformShaderParametersBuilder& WorldBounds(const FBoxSphereBounds& InWorldBounds)
 	{
 		AbsoluteObjectWorldPosition = InWorldBounds.Origin;
-		ObjectRadius = InWorldBounds.SphereRadius;
-		Parameters.ObjectBoundsX = InWorldBounds.BoxExtent.X;
-		Parameters.ObjectBoundsY = InWorldBounds.BoxExtent.Y;
-		Parameters.ObjectBoundsZ = InWorldBounds.BoxExtent.Z;
+		Parameters.ObjectBoundsX	= InWorldBounds.BoxExtent.X;
+		Parameters.ObjectBoundsY	= InWorldBounds.BoxExtent.Y;
+		Parameters.ObjectBoundsZ	= InWorldBounds.BoxExtent.Z;
 		return *this;
 	}
 
@@ -245,6 +244,7 @@ public:
 		bHasInstanceLocalBounds = true;
 		Parameters.InstanceLocalBoundsCenter = InInstanceLocalBounds.GetCenter();
 		Parameters.InstanceLocalBoundsExtent = InInstanceLocalBounds.GetExtent();
+		ObjectRadius						 = InInstanceLocalBounds.GetExtent().Length();
 		return *this;
 	}
 
