Strata - fixed AO decals

diff --git a/Engine/Shaders/Private/DecalCommon.ush b/Engine/Shaders/Private/DecalCommon.ush
--- a/Engine/Shaders/Private/DecalCommon.ush
+++ b/Engine/Shaders/Private/DecalCommon.ush
@@ -32,7 +32,11 @@ float DecalCommonOutput(
 #endif
 #elif DECAL_RENDERSTAGE == DECAL_RENDERSTAGE_AO
 	// Replace color with AO for AO decals.
-	Color = lerp(1, StrataGetAO(InHeader), DBuffer.Coverage);
+#if STRATA_INLINE_SHADING
+	Color = lerp(1, InHeader.IrradianceAO.MaterialAO, DBuffer.Coverage);
+#else
+#error Strata - Decals must be evaluated inline.
+#endif
 #endif
 
 #if DECAL_RENDERTARGETMODE == DECAL_RENDERTARGETMODE_SCENECOLOR
diff --git a/Engine/Source/Runtime/Engine/Private/Materials/Material.cpp b/Engine/Source/Runtime/Engine/Private/Materials/Material.cpp
--- a/Engine/Source/Runtime/Engine/Private/Materials/Material.cpp
+++ b/Engine/Source/Runtime/Engine/Private/Materials/Material.cpp
@@ -6156,7 +6156,8 @@ static bool IsPropertyActive_Internal(EMaterialProperty InProperty,
 	{
 		if (bStrataEnabled)
 		{
-			return InProperty == MP_FrontMaterial;
+			return InProperty == MP_FrontMaterial
+				|| InProperty == MP_AmbientOcclusion;
 		}
 		else if (InProperty >= MP_CustomizedUVs0 && InProperty <= MP_CustomizedUVs7)
 		{
diff --git a/Engine/Source/Runtime/Engine/Private/Materials/MaterialShared.cpp b/Engine/Source/Runtime/Engine/Private/Materials/MaterialShared.cpp
--- a/Engine/Source/Runtime/Engine/Private/Materials/MaterialShared.cpp
+++ b/Engine/Source/Runtime/Engine/Private/Materials/MaterialShared.cpp
@@ -1848,6 +1848,11 @@ bool FMaterialResource::HasMaterialPropertyConnected(EMaterialProperty In) const
 
 	if (Engine_IsStrataEnabled() && bIsStrataSupportedDomain)
 	{
+		if (In == MP_AmbientOcclusion)
+		{
+			// AO is specified on the root node so use the regular accessor.
+			return Material->HasAmbientOcclusionConnected();
+		}
 		// Strata material traversal is cached as this is an expensive operation
 		return FStrataMaterialInfo::HasPropertyConnected(Material->CachedConnectedInputs, In);
 	}
