World Partition: include non-colliding components in the bounding box for location volumes, so they display properly in the world partition editor window.

#rnx

diff --git a/Engine/Source/Runtime/Engine/Private/WorldPartition/LoaderAdapter/LoaderAdapterActor.cpp b/Engine/Source/Runtime/Engine/Private/WorldPartition/LoaderAdapter/LoaderAdapterActor.cpp
--- a/Engine/Source/Runtime/Engine/Private/WorldPartition/LoaderAdapter/LoaderAdapterActor.cpp
+++ b/Engine/Source/Runtime/Engine/Private/WorldPartition/LoaderAdapter/LoaderAdapterActor.cpp
@@ -13,7 +13,7 @@ FLoaderAdapterActor::FLoaderAdapterActor(AActor* InActor)
 
 TOptional<FBox> FLoaderAdapterActor::GetBoundingBox() const
 {
-	return Actor->GetComponentsBoundingBox();
+	return Actor->GetComponentsBoundingBox(/*bNonColliding*/true);
 }
 
 TOptional<FString> FLoaderAdapterActor::GetLabel() const
