When fixing up SceneComponent hierarchy on component destruction, do not attach to-be-destroyed child components to new parents, and do not attach living child components to to-be-destroyed parents.

diff --git a/Engine/Source/Runtime/Engine/Private/Components/SceneComponent.cpp b/Engine/Source/Runtime/Engine/Private/Components/SceneComponent.cpp
--- a/Engine/Source/Runtime/Engine/Private/Components/SceneComponent.cpp
+++ b/Engine/Source/Runtime/Engine/Private/Components/SceneComponent.cpp
@@ -1080,9 +1080,14 @@ void USceneComponent::OnComponentDestroyed(bool bDestroyingHierarchy)
 	// If we're just destroying for the exit purge don't bother with any of this
 	if (!GExitPurge && !bComputeBoundsOnceForGame)
 	{
-		// If we're destroying the hierarchy we only have to make sure that we detach children from other Actor's
+		// If we're destroying the hierarchy we only have to make sure that we detach children from other Actors
 		AActor* MyOwner = GetOwner();
 
+		// Do not involve objects which will be destroyed in hierarchy fixups
+PRAGMA_DISABLE_DEPRECATION_WARNINGS
+		constexpr EInternalObjectFlags SkipFlags = EInternalObjectFlags::PendingKill | EInternalObjectFlags::Garbage | EInternalObjectFlags::Unreachable;
+PRAGMA_ENABLE_DEPRECATION_WARNINGS
+
 		if (bDestroyingHierarchy)
 		{
 			// We'll lazily determine if we were attached in any way to a component from another Actor so that any of our children
@@ -1105,24 +1110,29 @@ void USceneComponent::OnComponentDestroyed(bool bDestroyingHierarchy)
 					{
 						if (Child->GetAttachParent() == this)
 						{
-							if (!bExternalAttachParentDetermined)
+							bool bNeedsDetach = true;
+							// If this child is going to be destroyed just detach it and don't find a new parent
+							if (!Child->HasAnyInternalFlags(SkipFlags))
 							{
-								ExternalAttachParent = GetAttachParent();
-								while (ExternalAttachParent)
+								if (!bExternalAttachParentDetermined)
 								{
-									if (ExternalAttachParent->GetOwner() != MyOwner)
+									ExternalAttachParent = GetAttachParent();
+									while (ExternalAttachParent)
 									{
-										break;
+										// Only attach to a parent which will not soon be destroyed
+										if (!ExternalAttachParent->HasAnyInternalFlags(SkipFlags) && ExternalAttachParent->GetOwner() != MyOwner)
+										{
+											break;
+										}
+										ExternalAttachParent = ExternalAttachParent->GetAttachParent();
 									}
-									ExternalAttachParent = ExternalAttachParent->GetAttachParent();
+									bExternalAttachParentDetermined = true;
 								}
-								bExternalAttachParentDetermined = true;
-							}
 
-							bool bNeedsDetach = true;
-							if (ExternalAttachParent)
-							{
-								bNeedsDetach = (Child->AttachToComponent(ExternalAttachParent, FAttachmentTransformRules::KeepWorldTransform) == false);
+								if (ExternalAttachParent)
+								{
+									bNeedsDetach = (Child->AttachToComponent(ExternalAttachParent, FAttachmentTransformRules::KeepWorldTransform) == false);
+								}
 							}
 							if (bNeedsDetach)
 							{
@@ -1177,12 +1187,23 @@ void USceneComponent::OnComponentDestroyed(bool bDestroyingHierarchy)
 				{
 					if (Child->GetAttachParent())
 					{
-						if (Child->GetAttachParent() == this)
+						// If the child is also being destroyed during GC, don't reattach it to anything
+						if (Child->HasAnyInternalFlags(SkipFlags))
+						{
+							Child->DetachFromComponent(FDetachmentTransformRules::KeepWorldTransform);
+						}
+						else if (Child->GetAttachParent() == this)
 						{
 							bool bNeedsDetach = true;
-							if (GetAttachParent())
+							USceneComponent* NewParent = GetAttachParent();
+							// Walk up the hierarchy until we find a valid parent which is not marked for destruction by gameplay or GC 
+							while (NewParent && NewParent->HasAnyInternalFlags(SkipFlags))
+							{
+								NewParent = NewParent->GetAttachParent();
+							}
+							if (NewParent)
 							{
-								bNeedsDetach = (Child->AttachToComponent(GetAttachParent(), FAttachmentTransformRules::KeepWorldTransform) == false);
+								bNeedsDetach = (Child->AttachToComponent(NewParent, FAttachmentTransformRules::KeepWorldTransform) == false);
 							}
 							if (bNeedsDetach)
 							{
