Constrained async uniform expression cache updates to happen during a scene render. This avoids multiple race conditions that can occur due to various material parmeter changes.

diff --git a/Engine/Source/Runtime/Engine/Private/Materials/MaterialShared.cpp b/Engine/Source/Runtime/Engine/Private/Materials/MaterialShared.cpp
--- a/Engine/Source/Runtime/Engine/Private/Materials/MaterialShared.cpp
+++ b/Engine/Source/Runtime/Engine/Private/Materials/MaterialShared.cpp
@@ -3577,6 +3577,144 @@ FUniformExpressionCache::~FUniformExpressionCache()
 }
 PRAGMA_ENABLE_DEPRECATION_WARNINGS
 
+class FUniformExpressionCacheAsyncUpdateTask
+{
+public:
+	void Begin()
+	{
+		check(!Task);
+		check(!bEnabled);
+		bEnabled = true;
+	}
+
+	void End()
+	{
+		check(bEnabled);
+		bEnabled = false;
+		Wait();
+	}
+
+	bool IsEnabled() const
+	{
+		return bEnabled && GUniformExpressionCacheAsyncUpdates > 0 && !GRHICommandList.Bypass();
+	}
+
+	void SetTask(const FGraphEventRef& InTask)
+	{
+		check(!Task);
+		check(IsEnabled());
+		Task = InTask;
+	}
+
+	void Wait()
+	{
+		if (Task)
+		{
+			Task->Wait();
+			Task = nullptr;
+		}
+	}
+
+private:
+	FGraphEventRef Task;
+	bool bEnabled = false;
+
+} GUniformExpressionCacheAsyncUpdateTask;
+
+FUniformExpressionCacheAsyncUpdateScope::FUniformExpressionCacheAsyncUpdateScope()
+{
+	ENQUEUE_RENDER_COMMAND(BeginAsyncUniformExpressionCacheUpdates)(
+		[](FRHICommandList&)
+	{
+		GUniformExpressionCacheAsyncUpdateTask.Begin();
+		FMaterialRenderProxy::UpdateDeferredCachedUniformExpressions();
+	});
+}
+
+FUniformExpressionCacheAsyncUpdateScope::~FUniformExpressionCacheAsyncUpdateScope()
+{
+	ENQUEUE_RENDER_COMMAND(EndAsyncUniformExpressionCacheUpdates)(
+		[](FRHICommandList&)
+	{
+		GUniformExpressionCacheAsyncUpdateTask.End();
+	});
+}
+
+void FUniformExpressionCacheAsyncUpdateScope::WaitForTask()
+{
+	GUniformExpressionCacheAsyncUpdateTask.Wait();
+}
+
+class FUniformExpressionCacheAsyncUpdater
+{
+public:
+	void Add(FUniformExpressionCache* UniformExpressionCache, const FUniformExpressionSet* UniformExpressionSet, const FRHIUniformBufferLayout* UniformBufferLayout, const FMaterialRenderContext& Context)
+	{
+		Items.Emplace(UniformExpressionCache, UniformExpressionSet, UniformBufferLayout, Context);
+	}
+
+	void Update(FRHICommandListImmediate& RHICmdListImmediate)
+	{
+		if (Items.IsEmpty())
+		{
+			return;
+		}
+
+		GUniformExpressionCacheAsyncUpdateTask.Wait();
+
+		FRHICommandList* RHICmdList = new FRHICommandList(FRHIGPUMask::All());
+
+		FGraphEventRef Event = FFunctionGraphTask::CreateAndDispatchWhenReady(
+			[Items = MoveTemp(Items), RHICmdList]
+		{
+			TRACE_CPUPROFILER_EVENT_SCOPE(FUniformExpressionCacheAsyncUpdater::Update);
+			FTaskTagScope Scope(ETaskTag::EParallelRenderingThread);
+			FMemMark Mark(FMemStack::Get());
+
+			for (const FItem& Item : Items)
+			{
+				uint8* TempBuffer = FMemStack::Get().PushBytes(Item.UniformBufferLayout->ConstantBufferSize, SHADER_PARAMETER_STRUCT_ALIGNMENT);
+
+				FMaterialRenderContext Context(Item.MaterialRenderProxy, *Item.Material, Item.bShowSelection);
+
+				Item.UniformExpressionSet->FillUniformBuffer(Context, Item.AllocatedVTs, Item.UniformBufferLayout, TempBuffer, Item.UniformBufferLayout->ConstantBufferSize);
+
+				RHICmdList->UpdateUniformBuffer(Item.UniformBuffer, TempBuffer);
+			}
+
+		}, TStatId(), nullptr, ENamedThreads::AnyHiPriThreadHiPriTask);
+
+		RHICmdListImmediate.QueueRenderThreadCommandListSubmit(Event, RHICmdList);
+
+		GUniformExpressionCacheAsyncUpdateTask.SetTask(Event);
+	}
+
+private:
+	struct FItem
+	{
+		FItem() = default;
+		FItem(FUniformExpressionCache* InUniformExpressionCache, const FUniformExpressionSet* InUniformExpressionSet, const FRHIUniformBufferLayout* InUniformBufferLayout, const FMaterialRenderContext& Context)
+			: UniformBuffer(InUniformExpressionCache->UniformBuffer)
+			, AllocatedVTs(InUniformExpressionCache->AllocatedVTs)
+			, UniformExpressionSet(InUniformExpressionSet)
+			, UniformBufferLayout(InUniformBufferLayout)
+			, MaterialRenderProxy(Context.MaterialRenderProxy)
+			, Material(&Context.Material)
+			, bShowSelection(Context.bShowSelection)
+		{}
+
+		TRefCountPtr<FRHIUniformBuffer> UniformBuffer;
+		TArray<IAllocatedVirtualTexture*, FConcurrentLinearArrayAllocator> AllocatedVTs;
+		const FUniformExpressionSet* UniformExpressionSet = nullptr;
+		const FRHIUniformBufferLayout* UniformBufferLayout = nullptr;
+		const FMaterialRenderProxy* MaterialRenderProxy = nullptr;
+		const FMaterial* Material = nullptr;
+		bool bShowSelection = false;
+	};
+
+	TArray<FItem, FConcurrentLinearArrayAllocator> Items;
+};
+
 void FUniformExpressionCache::ResetAllocatedVTs()
 {
 	for (int32 i=0; i< OwnedAllocatedVTs.Num(); ++i)
@@ -3587,7 +3725,7 @@ void FUniformExpressionCache::ResetAllocatedVTs()
 	OwnedAllocatedVTs.Reset();
 }
 
-void FMaterialRenderProxy::EvaluateUniformExpressions(FUniformExpressionCache& OutUniformExpressionCache, const FMaterialRenderContext& Context, FUniformExpressionCacheUpdater* Updater) const
+void FMaterialRenderProxy::EvaluateUniformExpressions(FUniformExpressionCache& OutUniformExpressionCache, const FMaterialRenderContext& Context, FUniformExpressionCacheAsyncUpdater* Updater) const
 {
 	check(IsInRenderingThread());
 
@@ -3904,80 +4042,6 @@ const FMaterial& FMaterialRenderProxy::GetIncompleteMaterialWithFallback(ERHIFea
 	return *Material;
 }
 
-class FUniformExpressionCacheAsyncUpdateTask
-{
-public:
-	FUniformExpressionCacheAsyncUpdateTask()
-	{
-		DelegateHandle = FCoreRenderDelegates::OnFlushRenderingCommandsEnd.AddLambda([this]
-		{
-			Wait();
-		});
-	}
-
-	~FUniformExpressionCacheAsyncUpdateTask()
-	{
-		Wait();
-		FCoreRenderDelegates::OnFlushRenderingCommandsEnd.Remove(DelegateHandle);
-	}
-
-	void SetTask(const FGraphEventRef& InTask)
-	{
-		check(!Task);
-		Task = InTask;
-	}
-
-	void Wait()
-	{
-		if (Task)
-		{
-			Task->Wait();
-			Task = nullptr;
-		}
-	}
-
-private:
-	FGraphEventRef Task;
-	FDelegateHandle DelegateHandle;
-
-} GUniformExpressionCacheAsyncUpdateTask;
-
-void FUniformExpressionCacheUpdater::Update(FRHICommandListImmediate& RHICmdListImmediate)
-{
-	if (Items.IsEmpty())
-	{
-		return;
-	}
-
-	GUniformExpressionCacheAsyncUpdateTask.Wait();
-
-	FRHICommandList* RHICmdList = new FRHICommandList(FRHIGPUMask::All());
-
-	FGraphEventRef Event = FFunctionGraphTask::CreateAndDispatchWhenReady(
-		[Items = MoveTemp(Items), RHICmdList]
-	{
-		TRACE_CPUPROFILER_EVENT_SCOPE(FUniformExpressionCacheUpdater::Update);
-		FTaskTagScope Scope(ETaskTag::EParallelRenderingThread);
-		FMemMark Mark(FMemStack::Get());
-
-		for (const FItem& Item : Items)
-		{
-			uint8* TempBuffer = FMemStack::Get().PushBytes(Item.UniformBufferLayout->ConstantBufferSize, SHADER_PARAMETER_STRUCT_ALIGNMENT);
-
-			FMaterialRenderContext Context(Item.MaterialRenderProxy, *Item.Material, Item.bShowSelection);
-
-			Item.UniformExpressionSet->FillUniformBuffer(Context, Item.AllocatedVTs, Item.UniformBufferLayout, TempBuffer, Item.UniformBufferLayout->ConstantBufferSize);
-
-			RHICmdList->UpdateUniformBuffer(Item.UniformBuffer, TempBuffer);
-		}
-
-	}, TStatId(), nullptr, ENamedThreads::AnyHiPriThreadHiPriTask);
-
-	RHICmdListImmediate.QueueRenderThreadCommandListSubmit(Event, RHICmdList);
-
-	GUniformExpressionCacheAsyncUpdateTask.SetTask(Event);
-}
-
 void FMaterialRenderProxy::UpdateDeferredCachedUniformExpressions()
 {
 	LLM_SCOPE(ELLMTag::Materials);
@@ -3988,8 +4052,8 @@ void FMaterialRenderProxy::UpdateDeferredCachedUniformExpressions()
 	CSV_SCOPED_TIMING_STAT_EXCLUSIVE(Material_UpdateDeferredCachedUniformExpressions);
 	QUICK_SCOPE_CYCLE_COUNTER(STAT_UpdateDeferredCachedUniformExpressions);
 
-	FUniformExpressionCacheUpdater Updater;
-	FUniformExpressionCacheUpdater* UpdaterIfEnabled = !GRHICommandList.Bypass() && GUniformExpressionCacheAsyncUpdates > 0 ? &Updater : nullptr;
+	FUniformExpressionCacheAsyncUpdater Updater;
+	FUniformExpressionCacheAsyncUpdater* UpdaterIfEnabled = GUniformExpressionCacheAsyncUpdateTask.IsEnabled() ? &Updater : nullptr;
 
 	for (TSet<FMaterialRenderProxy*>::TConstIterator It(DeferredUniformExpressionCacheRequests); It; ++It)
 	{
diff --git a/Engine/Source/Runtime/Engine/Public/MaterialShared.h b/Engine/Source/Runtime/Engine/Public/MaterialShared.h
--- a/Engine/Source/Runtime/Engine/Public/MaterialShared.h
+++ b/Engine/Source/Runtime/Engine/Public/MaterialShared.h
@@ -2475,41 +2475,21 @@ private:
 };
 
 class USubsurfaceProfile;
+class FUniformExpressionCacheAsyncUpdater;
 
-class ENGINE_API FUniformExpressionCacheUpdater
+/** Defines a scope to update deferred uniform expression caches using an async task to fill uniform buffers. The scope
+ *  attempts to launch async updates immediately, but further async updates can launch within the scope. Only one async
+ *  task is launched at a time though. The async task is synced when the scope destructs. The scope enqueues render commands
+ *  so it can be used on the game or render threads.
+ */
+class ENGINE_API FUniformExpressionCacheAsyncUpdateScope
 {
 public:
-	void Add(FUniformExpressionCache* UniformExpressionCache, const FUniformExpressionSet* UniformExpressionSet, const FRHIUniformBufferLayout* UniformBufferLayout, const FMaterialRenderContext& Context)
-	{
-		Items.Emplace(UniformExpressionCache, UniformExpressionSet, UniformBufferLayout, Context);
-	}
-
-	void Update(FRHICommandListImmediate& RHICmdList);
-
-private:
-	struct FItem
-	{
-		FItem() = default;
-		FItem(FUniformExpressionCache* InUniformExpressionCache, const FUniformExpressionSet* InUniformExpressionSet, const FRHIUniformBufferLayout* InUniformBufferLayout, const FMaterialRenderContext& Context)
-			: UniformBuffer(InUniformExpressionCache->UniformBuffer)
-			, AllocatedVTs(InUniformExpressionCache->AllocatedVTs)
-			, UniformExpressionSet(InUniformExpressionSet)
-			, UniformBufferLayout(InUniformBufferLayout)
-			, MaterialRenderProxy(Context.MaterialRenderProxy)
-			, Material(&Context.Material)
-			, bShowSelection(Context.bShowSelection)
-		{}
-
-		TRefCountPtr<FRHIUniformBuffer> UniformBuffer;
-		TArray<IAllocatedVirtualTexture*, FConcurrentLinearArrayAllocator> AllocatedVTs;
-		const FUniformExpressionSet* UniformExpressionSet = nullptr;
-		const FRHIUniformBufferLayout* UniformBufferLayout = nullptr;
-		const FMaterialRenderProxy* MaterialRenderProxy = nullptr;
-		const FMaterial* Material = nullptr;
-		bool bShowSelection = false;
-	};
+	FUniformExpressionCacheAsyncUpdateScope();
+	~FUniformExpressionCacheAsyncUpdateScope();
 
-	TArray<FItem, FConcurrentLinearArrayAllocator> Items;
+	/** Call if a wait is required within the scope. */
+	static void WaitForTask();
 };
 
 /**
@@ -2542,7 +2522,7 @@ public:
 	 * @param OutUniformExpressionCache - The uniform expression cache to build.
 	 * @param MaterialRenderContext - The context for which to cache expressions.
 	 */
-	void EvaluateUniformExpressions(FUniformExpressionCache& OutUniformExpressionCache, const FMaterialRenderContext& Context, FUniformExpressionCacheUpdater* Updater = nullptr) const;
+	void EvaluateUniformExpressions(FUniformExpressionCache& OutUniformExpressionCache, const FMaterialRenderContext& Context, FUniformExpressionCacheAsyncUpdater* Updater = nullptr) const;
 
 	/**
 	 * Caches uniform expressions for efficient runtime evaluation.
diff --git a/Engine/Source/Runtime/Renderer/Private/ReflectionEnvironmentCapture.cpp b/Engine/Source/Runtime/Renderer/Private/ReflectionEnvironmentCapture.cpp
--- a/Engine/Source/Runtime/Renderer/Private/ReflectionEnvironmentCapture.cpp
+++ b/Engine/Source/Runtime/Renderer/Private/ReflectionEnvironmentCapture.cpp
@@ -477,6 +477,8 @@ void CaptureSceneToScratchCubemap(
 
 	const ERHIFeatureLevel::Type FeatureLevel = SceneRenderer->FeatureLevel;
 
+	FUniformExpressionCacheAsyncUpdateScope AsyncUpdateScope;
+
 	// update any resources that needed a deferred update
 	FDeferredUpdateResource::UpdateResources(RHICmdList);
 
@@ -485,8 +487,6 @@ void CaptureSceneToScratchCubemap(
 	// We need to execute the pre-render view extensions before we do any view dependent work.
 	FSceneRenderer::ViewExtensionPreRender_RenderThread(GraphBuilder, SceneRenderer);
 
-	FMaterialRenderProxy::UpdateDeferredCachedUniformExpressions();
-
 	{
 		RDG_EVENT_SCOPE(GraphBuilder, "CubeMapCapture");
 
diff --git a/Engine/Source/Runtime/Renderer/Private/SceneCaptureRendering.cpp b/Engine/Source/Runtime/Renderer/Private/SceneCaptureRendering.cpp
--- a/Engine/Source/Runtime/Renderer/Private/SceneCaptureRendering.cpp
+++ b/Engine/Source/Runtime/Renderer/Private/SceneCaptureRendering.cpp
@@ -507,7 +507,7 @@ static void UpdateSceneCaptureContent_RenderThread(
 	bool bClearRenderTarget,
 	bool bOrthographicCamera)
 {
-	FMaterialRenderProxy::UpdateDeferredCachedUniformExpressions();
+	FUniformExpressionCacheAsyncUpdateScope AsyncUpdateScope;
 
 	switch (SceneRenderer->Scene->GetShadingPath())
 	{
diff --git a/Engine/Source/Runtime/Renderer/Private/SceneRendering.cpp b/Engine/Source/Runtime/Renderer/Private/SceneRendering.cpp
--- a/Engine/Source/Runtime/Renderer/Private/SceneRendering.cpp
+++ b/Engine/Source/Runtime/Renderer/Private/SceneRendering.cpp
@@ -4178,6 +4178,9 @@ static FSceneRenderCleanUpState GSceneRenderCleanUpState;
 
 void FSceneRenderer::RenderThreadEnd(FRHICommandListImmediate& RHICmdList, const TArray<FSceneRenderer*>& SceneRenderers)
 {
+	// We need to sync async uniform expression cache updates since we're about to start deleting material proxies.
+	FUniformExpressionCacheAsyncUpdateScope::WaitForTask();
+
 	check(GSceneRenderCleanUpState.Renderers.IsEmpty());
 
 	const ESceneRenderCleanUpMode SceneRenderCleanUpMode = GetSceneRenderCleanUpMode();
@@ -4556,11 +4559,7 @@ void FRendererModule::BeginRenderingViewFamilies(FCanvas* Canvas, TArrayView<FSc
 			FThreadIdleStats::BeginCriticalPath();
 		});
 
-	ENQUEUE_RENDER_COMMAND(UpdateDeferredCachedUniformExpressions)(
-		[](FRHICommandList& RHICmdList)
-		{
-			FMaterialRenderProxy::UpdateDeferredCachedUniformExpressions();
-		});
+	FUniformExpressionCacheAsyncUpdateScope AsyncUpdateScope;
 
 	ENQUEUE_RENDER_COMMAND(UpdateFastVRamConfig)(
 		[](FRHICommandList& RHICmdList)
