Potential fix for UpdateDeferredCachedUniformExpressions crash.  Added improved thread safety to flag that tracks whether the material instance is in use by render thread.  Flag to track in-flight resource destroy is now set when the resource is created, rather than transiently when the destroy is kicked off.  And a separate flag tracks whether there is an in-flight CacheUniformExpressions operation, and removes any that are active when a material instance is destroyed.

#rnx

diff --git a/Engine/Source/Runtime/Engine/Classes/Materials/MaterialInstance.h b/Engine/Source/Runtime/Engine/Classes/Materials/MaterialInstance.h
--- a/Engine/Source/Runtime/Engine/Classes/Materials/MaterialInstance.h
+++ b/Engine/Source/Runtime/Engine/Classes/Materials/MaterialInstance.h
@@ -430,6 +430,15 @@ enum class EMaterialInstanceClearParameterFlag
 };
 ENUM_CLASS_FLAGS(EMaterialInstanceClearParameterFlag);
 
+enum class EMaterialInstanceUsedByRTFlag : uint32
+{
+	None = 0u,
+	ResourceCreate = (1u << 0),
+	CacheUniformExpressions = (1u << 1),
+
+	All = ResourceCreate | CacheUniformExpressions,
+};
+
 #if WITH_EDITORONLY_DATA
 class FMaterialInstanceParameterUpdateContext
 {
@@ -666,8 +675,8 @@ private:
 	/** Material resources being cached for cooking. */
 	TMap<const class ITargetPlatform*, TArray<FMaterialResource*>> CachedMaterialResourcesForCooking;
 #endif
-	/** Flag used to guarantee that the RT is finished using various resources in this UMaterial before cleanup. */
-	FThreadSafeBool ReleasedByRT;
+	/** Thread-safe flags used to track whether this instance is still in use by the render thread (EMaterialInstanceUsedByRTFlag) */
+	mutable std::atomic<uint32> UsedByRT;
 
 public:
 	virtual ENGINE_API ~UMaterialInstance() {};
@@ -898,6 +907,10 @@ public:
 #endif
 	ENGINE_API virtual bool IsComplete() const override;
 
+	/** Tracking of in-flight uniform expression cache update operations for the material instance, for thread safety destroying the resource. */
+	void StartCacheUniformExpressions() const;
+	void FinishCacheUniformExpressions() const;
+
 protected:
 
 	/**
diff --git a/Engine/Source/Runtime/Engine/Private/Materials/MaterialInstance.cpp b/Engine/Source/Runtime/Engine/Private/Materials/MaterialInstance.cpp
--- a/Engine/Source/Runtime/Engine/Private/Materials/MaterialInstance.cpp
+++ b/Engine/Source/Runtime/Engine/Private/Materials/MaterialInstance.cpp
@@ -60,6 +60,26 @@ DECLARE_CYCLE_STAT(TEXT("MaterialInstance CopyUniformParamsInternal"), STAT_Mate
 
 const FMaterialInstanceCachedData FMaterialInstanceCachedData::EmptyData{};
 
+void UMaterialInstance::StartCacheUniformExpressions() const
+{
+	UsedByRT |= (uint32)EMaterialInstanceUsedByRTFlag::CacheUniformExpressions;
+}
+
+void UMaterialInstance::FinishCacheUniformExpressions() const
+{
+	UsedByRT &= ~(uint32)EMaterialInstanceUsedByRTFlag::CacheUniformExpressions;
+}
+
+void FMaterialInstanceResource::StartCacheUniformExpressions() const
+{
+	Owner->StartCacheUniformExpressions();
+}
+
+void FMaterialInstanceResource::FinishCacheUniformExpressions() const
+{
+	Owner->FinishCacheUniformExpressions();
+}
+
 /**
  * Cache uniform expressions for the given material.
  * @param MaterialInstance - The material instance for which to cache uniform expressions.
@@ -68,6 +88,7 @@ void CacheMaterialInstanceUniformExpressions(const UMaterialInstance* MaterialIn
 {
 	if (MaterialInstance->Resource)
 	{
+		MaterialInstance->StartCacheUniformExpressions();
 		MaterialInstance->Resource->CacheUniformExpressions_GameThread(bRecreateUniformBuffer);
 	}
 }
@@ -367,6 +388,8 @@ void GameThread_UpdateMIParameter(const UMaterialInstance* Instance, const Param
 {
 	if (FApp::CanEverRender())
 	{
+		Instance->StartCacheUniformExpressions();
+
 		const UMaterial* Material = Instance->GetMaterial_Concurrent();
 		if (Material != nullptr)
 		{
@@ -635,7 +658,7 @@ bool UMaterialInstance::UpdateParameters()
 
 UMaterialInstance::UMaterialInstance(const FObjectInitializer& ObjectInitializer)
 	: Super(ObjectInitializer)
-	, ReleasedByRT(true)
+	, UsedByRT((uint32)EMaterialInstanceUsedByRTFlag::None)
 {
 	bHasStaticPermutationResource = false;
 	bLoadedCachedData = false;
@@ -660,6 +683,7 @@ void UMaterialInstance::PostInitProperties()
 	if(!HasAnyFlags(RF_ClassDefaultObject))
 	{
 		Resource = new FMaterialInstanceResource(this);
+		UsedByRT |= (uint32)EMaterialInstanceUsedByRTFlag::ResourceCreate;
 	}
 }
 
@@ -2892,12 +2916,10 @@ void UMaterialInstance::BeginDestroy()
 
 	if (Resource || ResourcesToDestroy.Num() > 0)
 	{
-		ReleasedByRT = false;
-
 		FMaterialRenderProxy* LocalResource = Resource;
-		FThreadSafeBool* Released = &ReleasedByRT;
+		std::atomic<uint32>* Used = &UsedByRT;
 		ENQUEUE_RENDER_COMMAND(BeginDestroyCommand)(
-		[ResourcesToDestroy = MoveTemp(ResourcesToDestroy), LocalResource, Released](FRHICommandListImmediate& RHICmdList)
+		[ResourcesToDestroy = MoveTemp(ResourcesToDestroy), LocalResource, Used](FRHICommandListImmediate& RHICmdList)
 		{
 			if (LocalResource)
 			{
@@ -2910,7 +2932,15 @@ void UMaterialInstance::BeginDestroy()
 				CurrentResource->PrepareDestroy_RenderThread();
 			}
 
-			*Released = true;
+			// Clear flag set when Resource was created
+			*Used &= ~(uint32)EMaterialInstanceUsedByRTFlag::ResourceCreate;
+
+			// And remove from deferred uniform expression cache queue if it's in that
+			if (*Used & (uint32)EMaterialInstanceUsedByRTFlag::CacheUniformExpressions)
+			{
+				LocalResource->CancelCacheUniformExpressions();
+				*Used &= ~(uint32)EMaterialInstanceUsedByRTFlag::CacheUniformExpressions;
+			}
 		});		
 	}
 }
@@ -2919,7 +2949,7 @@ bool UMaterialInstance::IsReadyForFinishDestroy()
 {
 	bool bIsReady = Super::IsReadyForFinishDestroy();
 
-	return bIsReady && ReleasedByRT;
+	return bIsReady && UsedByRT == 0;
 }
 
 void UMaterialInstance::FinishDestroy()
diff --git a/Engine/Source/Runtime/Engine/Private/Materials/MaterialInstanceSupport.h b/Engine/Source/Runtime/Engine/Private/Materials/MaterialInstanceSupport.h
--- a/Engine/Source/Runtime/Engine/Private/Materials/MaterialInstanceSupport.h
+++ b/Engine/Source/Runtime/Engine/Private/Materials/MaterialInstanceSupport.h
@@ -316,6 +316,9 @@ private:
 	static bool IsValidParameterValue(const UTexture* Value) { return Value != nullptr; }
 	static bool IsValidParameterValue(const URuntimeVirtualTexture* Value) { return Value != nullptr; }
 
+	virtual void StartCacheUniformExpressions() const override;
+	virtual void FinishCacheUniformExpressions() const override;
+
 	// Validation logic that exhaustively searches the array, to compare with what the hash table logic finds
 	#define HASHED_MATERIAL_PARAMETER_MAP_VALIDATE 0
 
diff --git a/Engine/Source/Runtime/Engine/Private/Materials/MaterialShared.cpp b/Engine/Source/Runtime/Engine/Private/Materials/MaterialShared.cpp
--- a/Engine/Source/Runtime/Engine/Private/Materials/MaterialShared.cpp
+++ b/Engine/Source/Runtime/Engine/Private/Materials/MaterialShared.cpp
@@ -3827,6 +3827,7 @@ void FMaterialRenderProxy::CacheUniformExpressions(bool bRecreateUniformBuffer)
 	{
 		UE_LOG(LogMaterial, Fatal, TEXT("Cannot queue the Expression Cache for Material %s when it is about to be deleted"), *MaterialName);
 	}
+	StartCacheUniformExpressions();
 	DeferredUniformExpressionCacheRequests.Add(this);
 
 	InvalidateUniformExpressionCache(bRecreateUniformBuffer);
@@ -3937,6 +3938,11 @@ void FMaterialRenderProxy::InitDynamicRHI()
 #endif // WITH_EDITOR
 }
 
+void FMaterialRenderProxy::CancelCacheUniformExpressions()
+{
+	DeferredUniformExpressionCacheRequests.Remove(this);
+}
+
 void FMaterialRenderProxy::ReleaseDynamicRHI()
 {
 #if WITH_EDITOR
@@ -3947,7 +3953,11 @@ void FMaterialRenderProxy::ReleaseDynamicRHI()
 	}
 #endif // WITH_EDITOR
 
-	DeferredUniformExpressionCacheRequests.Remove(this);
+	if (DeferredUniformExpressionCacheRequests.Remove(this))
+	{
+		// Notify that we're finished with this inflight cache request, because the object is being released
+		FinishCacheUniformExpressions();
+	}
 
 	InvalidateUniformExpressionCache(true);
 
@@ -4074,6 +4084,8 @@ void FMaterialRenderProxy::UpdateDeferredCachedUniformExpressions()
 				MaterialProxy->EvaluateUniformExpressions(MaterialProxy->UniformExpressionCache[(int32)InFeatureLevel], MaterialRenderContext, UpdaterIfEnabled);
 			}
 		});
+
+		MaterialProxy->FinishCacheUniformExpressions();
 	}
 
 	if (UpdaterIfEnabled)
diff --git a/Engine/Source/Runtime/Engine/Public/MaterialShared.h b/Engine/Source/Runtime/Engine/Public/MaterialShared.h
--- a/Engine/Source/Runtime/Engine/Public/MaterialShared.h
+++ b/Engine/Source/Runtime/Engine/Public/MaterialShared.h
@@ -2529,6 +2529,9 @@ public:
 	 */
 	void CacheUniformExpressions(bool bRecreateUniformBuffer);
 
+	/** Cancels an in-flight cache operation. */
+	void CancelCacheUniformExpressions();
+
 	/**
 	 * Enqueues a rendering command to cache uniform expressions for efficient runtime evaluation.
 	 * bRecreateUniformBuffer - whether to recreate the material uniform buffer.  
@@ -2629,6 +2632,9 @@ private:
 	IAllocatedVirtualTexture* GetPreallocatedVTStack(const FMaterialRenderContext& Context, const FUniformExpressionSet& UniformExpressionSet, const FMaterialVirtualTextureStack& VTStack) const;
 	IAllocatedVirtualTexture* AllocateVTStack(const FMaterialRenderContext& Context, const FUniformExpressionSet& UniformExpressionSet, const FMaterialVirtualTextureStack& VTStack) const;
 
+	virtual void StartCacheUniformExpressions() const {}
+	virtual void FinishCacheUniformExpressions() const {}
+
 	/** 0 if not set, game thread pointer, do not dereference, only for comparison */
 	const USubsurfaceProfile* SubsurfaceProfileRT;
 	FString MaterialName;
