Skeletal Mesh:
- Fix unwanted casts from SkinnedMeshComponent to SkeletalMeshComponent where other SkinnedMeshComponent derived classes might be used, usually around access to the LeaderPoseComponent.
- Refactor USkeleton to take USkinnedAsset for its functions' parameters (only leaves out Preview Mesh related functions, since they are virtual functions implementing the IInterface_PreviewMeshProvider interface).

diff --git a/Engine/Plugins/VirtualProduction/Takes/Source/TakeTrackRecorders/Private/MovieScene3DTransformTrackRecorder.cpp b/Engine/Plugins/VirtualProduction/Takes/Source/TakeTrackRecorders/Private/MovieScene3DTransformTrackRecorder.cpp
--- a/Engine/Plugins/VirtualProduction/Takes/Source/TakeTrackRecorders/Private/MovieScene3DTransformTrackRecorder.cpp
+++ b/Engine/Plugins/VirtualProduction/Takes/Source/TakeTrackRecorders/Private/MovieScene3DTransformTrackRecorder.cpp
@@ -489,11 +489,11 @@ void UMovieScene3DTransformTrackRecorder::PostProcessAnimationData(UMovieSceneAn
 
 			}
 			// Search for the Root Bone in the Skeleton
-			const USkeletalMesh* SkeletalMesh = Cast<USkeletalMeshComponent>(SkeletalMeshComponent->LeaderPoseComponent) ?
-				Cast<USkeletalMeshComponent>(SkeletalMeshComponent->LeaderPoseComponent)->GetSkeletalMeshAsset() : 
+			const USkinnedAsset* SkinnedAsset = SkeletalMeshComponent->LeaderPoseComponent != nullptr ?
+				SkeletalMeshComponent->LeaderPoseComponent->GetSkinnedAsset() : 
 				SkeletalMeshComponent->GetSkeletalMeshAsset();
 			const UAnimSequence* AnimSequence = AnimTrackRecorder->GetAnimSequence();
-			if (AnimSequence && SkeletalMesh && AnimSequence->GetSkeleton())
+			if (AnimSequence && SkinnedAsset && AnimSequence->GetSkeleton())
 			{
 				// Find the root bone
 				int32 RootIndex = INDEX_NONE;
@@ -506,8 +506,8 @@ void UMovieScene3DTransformTrackRecorder::PostProcessAnimationData(UMovieSceneAn
 					const int32 BoneTreeIndex = AnimationTrack.BoneTreeIndex;
 					if (BoneTreeIndex != INDEX_NONE)
 					{
-						const int32 BoneIndex = AnimSkeleton->GetMeshBoneIndexFromSkeletonBoneIndex(SkeletalMesh, BoneTreeIndex);
-						const int32 ParentIndex = SkeletalMesh->GetRefSkeleton().GetParentIndex(BoneIndex);
+						const int32 BoneIndex = AnimSkeleton->GetMeshBoneIndexFromSkeletonBoneIndex(SkinnedAsset, BoneTreeIndex);
+						const int32 ParentIndex = SkinnedAsset->GetRefSkeleton().GetParentIndex(BoneIndex);
 						if (ParentIndex == INDEX_NONE)
 						{
 							// We've found the root (root bones do not have a valid parent)
diff --git a/Engine/Plugins/VirtualProduction/Takes/Source/TakeTrackRecorders/Private/MovieSceneAnimationTrackRecorder.cpp b/Engine/Plugins/VirtualProduction/Takes/Source/TakeTrackRecorders/Private/MovieSceneAnimationTrackRecorder.cpp
--- a/Engine/Plugins/VirtualProduction/Takes/Source/TakeTrackRecorders/Private/MovieSceneAnimationTrackRecorder.cpp
+++ b/Engine/Plugins/VirtualProduction/Takes/Source/TakeTrackRecorders/Private/MovieSceneAnimationTrackRecorder.cpp
@@ -158,9 +158,9 @@ void UMovieSceneAnimationTrackRecorder::CreateTrackImpl()
 			{
 				// verify if this bone exists in skeleton
 				const int32 BoneTreeIndex = AnimSkeleton->GetSkeletonBoneIndexFromMeshBoneIndex(
-					Cast<USkeletalMeshComponent>(SkeletalMeshComponent->LeaderPoseComponent) ? 
-					Cast<USkeletalMeshComponent>(SkeletalMeshComponent->LeaderPoseComponent)->GetSkeletalMeshAsset() :
-					SkeletalMeshComponent->GetSkeletalMeshAsset(), BoneIndex);
+					SkeletalMeshComponent->LeaderPoseComponent != nullptr ? 
+					SkeletalMeshComponent->LeaderPoseComponent->GetSkinnedAsset() :
+					SkeletalMeshComponent->GetSkinnedAsset(), BoneIndex);
 				if (BoneTreeIndex != INDEX_NONE)
 				{
 					// add tracks for the bone existing
diff --git a/Engine/Source/Editor/DetailCustomizations/Private/SkinnedMeshComponentDetails.cpp b/Engine/Source/Editor/DetailCustomizations/Private/SkinnedMeshComponentDetails.cpp
--- a/Engine/Source/Editor/DetailCustomizations/Private/SkinnedMeshComponentDetails.cpp
+++ b/Engine/Source/Editor/DetailCustomizations/Private/SkinnedMeshComponentDetails.cpp
@@ -6,7 +6,7 @@
 #include "Widgets/Text/STextBlock.h"
 #include "Widgets/Input/SEditableTextBox.h"
 #include "GameFramework/Actor.h"
-#include "Components/SkeletalMeshComponent.h"
+#include "Components/SkinnedMeshComponent.h"
 #include "Editor.h"
 #include "PropertyHandle.h"
 #include "DetailLayoutBuilder.h"
@@ -181,9 +181,9 @@ void FSkinnedMeshComponentDetails::PopulateSkinWeightProfileNames()
 	SkinWeightProfileNames.Add(MakeShared<FName>(DefaultProfileName));
 
 	// Retrieve all possible skin weight profiles from the component
-	if (USkeletalMeshComponent* Component = Cast<USkeletalMeshComponent>(WeakSkinnedMeshComponent.Get()))
+	if (USkinnedMeshComponent* Component = WeakSkinnedMeshComponent.Get())
 	{
-		if (USkeletalMesh* Mesh = Component->GetSkeletalMeshAsset())
+		if (USkeletalMesh* Mesh = Cast<USkeletalMesh>(Component->GetSkinnedAsset()))
 		{			
 			for (const FSkinWeightProfileInfo& Profile : Mesh->GetSkinWeightProfiles())
 			{
diff --git a/Engine/Source/Editor/SequenceRecorder/Private/AnimationRecorder.cpp b/Engine/Source/Editor/SequenceRecorder/Private/AnimationRecorder.cpp
--- a/Engine/Source/Editor/SequenceRecorder/Private/AnimationRecorder.cpp
+++ b/Engine/Source/Editor/SequenceRecorder/Private/AnimationRecorder.cpp
@@ -271,9 +271,9 @@ void FAnimationRecorder::StartRecord(USkeletalMeshComponent* Component, UAnimSeq
 	{
 		// verify if this bone exists in skeleton
 		const int32 BoneTreeIndex = AnimSkeleton->GetSkeletonBoneIndexFromMeshBoneIndex(
-			Cast<USkeletalMeshComponent>(Component->LeaderPoseComponent) ?
-			Cast<USkeletalMeshComponent>(Component->LeaderPoseComponent)->GetSkeletalMeshAsset() :
-			Component->GetSkeletalMeshAsset(), BoneIndex);
+			Component->LeaderPoseComponent != nullptr ?
+			Component->LeaderPoseComponent->GetSkinnedAsset() :
+			Component->GetSkinnedAsset(), BoneIndex);
 		if (BoneTreeIndex != INDEX_NONE)
 		{
 			// add tracks for the bone existing
@@ -656,9 +656,9 @@ void FAnimationRecorder::ProcessRecordedTimes(UAnimSequence* AnimSequence, USkel
 		for (int32 BoneIndex = 0; BoneIndex < SpaceBases->Num(); ++BoneIndex)
 		{
 			const int32 BoneTreeIndex = AnimSkeleton->GetSkeletonBoneIndexFromMeshBoneIndex(
-				Cast<USkeletalMeshComponent>(SkeletalMeshComponent->LeaderPoseComponent) ?
-				Cast<USkeletalMeshComponent>(SkeletalMeshComponent->LeaderPoseComponent)->GetSkeletalMeshAsset() :
-				SkeletalMeshComponent->GetSkeletalMeshAsset(), BoneIndex);
+				SkeletalMeshComponent->LeaderPoseComponent != nullptr ?
+				SkeletalMeshComponent->LeaderPoseComponent->GetSkinnedAsset() :
+				SkeletalMeshComponent->GetSkinnedAsset(), BoneIndex);
 			if (BoneTreeIndex != INDEX_NONE)
 			{
 				FName BoneTreeName = AnimSkeleton->GetReferenceSkeleton().GetBoneName(BoneTreeIndex);
@@ -680,9 +680,9 @@ void FAnimationRecorder::ProcessRecordedTimes(UAnimSequence* AnimSequence, USkel
 	{
 		// verify if this bone exists in skeleton
 		const int32 BoneTreeIndex = AnimSkeleton->GetSkeletonBoneIndexFromMeshBoneIndex(
-			Cast<USkeletalMeshComponent>(SkeletalMeshComponent->LeaderPoseComponent) ?
-			Cast<USkeletalMeshComponent>(SkeletalMeshComponent->LeaderPoseComponent)->GetSkeletalMeshAsset() :
-			SkeletalMeshComponent->GetSkeletalMeshAsset(), BoneIndex);
+			SkeletalMeshComponent->LeaderPoseComponent != nullptr ?
+			SkeletalMeshComponent->LeaderPoseComponent->GetSkinnedAsset() :
+			SkeletalMeshComponent->GetSkinnedAsset(), BoneIndex);
 		if (BoneTreeIndex != INDEX_NONE)
 		{
 			// add tracks for the bone existing
@@ -835,10 +835,10 @@ bool FAnimationRecorder::Record(USkeletalMeshComponent* Component, FTransform co
 	if (ensure(AnimationObject))
 	{
 		IAnimationDataController& Controller = AnimationObject->GetController();
-		USkeletalMesh* SkeletalMesh =
-			Cast<USkeletalMeshComponent>(Component->LeaderPoseComponent) ?
-			Cast<USkeletalMeshComponent>(Component->LeaderPoseComponent)->GetSkeletalMeshAsset() :
-			Component->GetSkeletalMeshAsset();
+		USkinnedAsset* SkinnedAsset =
+			Component->LeaderPoseComponent != nullptr ?
+			Component->LeaderPoseComponent->GetSkinnedAsset() :
+			Component->GetSkinnedAsset();
 
 		const TArray<FBoneAnimationTrack>& BoneAnimationTracks = AnimationObject->GetDataModel()->GetBoneAnimationTracks();
 
@@ -854,8 +854,8 @@ bool FAnimationRecorder::Record(USkeletalMeshComponent* Component, FTransform co
 				const int32 BoneTreeIndex = AnimationTrack.BoneTreeIndex;
 				if (BoneTreeIndex != INDEX_NONE)
 				{
-					const int32 BoneIndex = AnimSkeleton->GetMeshBoneIndexFromSkeletonBoneIndex(SkeletalMesh, BoneTreeIndex);
-					const int32 ParentIndex = SkeletalMesh->GetRefSkeleton().GetParentIndex(BoneIndex);
+					const int32 BoneIndex = AnimSkeleton->GetMeshBoneIndexFromSkeletonBoneIndex(SkinnedAsset, BoneTreeIndex);
+					const int32 ParentIndex = SkinnedAsset->GetRefSkeleton().GetParentIndex(BoneIndex);
 					const FTransform LocalTransform = SpacesBases[BoneIndex];
 					if (ParentIndex == INDEX_NONE)
 					{
@@ -895,8 +895,8 @@ bool FAnimationRecorder::Record(USkeletalMeshComponent* Component, FTransform co
 			const int32 BoneTreeIndex = AnimationTrack.BoneTreeIndex;
 			if (BoneTreeIndex != INDEX_NONE)
 			{
-				const int32 BoneIndex = AnimSkeleton->GetMeshBoneIndexFromSkeletonBoneIndex(SkeletalMesh, BoneTreeIndex);
-				const int32 ParentIndex = SkeletalMesh->GetRefSkeleton().GetParentIndex(BoneIndex);
+				const int32 BoneIndex = AnimSkeleton->GetMeshBoneIndexFromSkeletonBoneIndex(SkinnedAsset, BoneTreeIndex);
+				const int32 ParentIndex = SkinnedAsset->GetRefSkeleton().GetParentIndex(BoneIndex);
 
 				if (bRecordTransforms)
 				{
diff --git a/Engine/Source/Editor/SequenceRecorder/Private/Sections/MovieScene3DTransformSectionRecorder.cpp b/Engine/Source/Editor/SequenceRecorder/Private/Sections/MovieScene3DTransformSectionRecorder.cpp
--- a/Engine/Source/Editor/SequenceRecorder/Private/Sections/MovieScene3DTransformSectionRecorder.cpp
+++ b/Engine/Source/Editor/SequenceRecorder/Private/Sections/MovieScene3DTransformSectionRecorder.cpp
@@ -132,11 +132,11 @@ void FMovieScene3DTransformSectionRecorder::FinalizeSection(float CurrentTime)
 		USkeletalMeshComponent* SkeletalMeshComponent = AnimRecorder->GetSkeletalMeshComponent();
 		if (SkeletalMeshComponent)
 		{
-			USkeletalMesh* SkeletalMesh =
-				Cast<USkeletalMeshComponent>(SkeletalMeshComponent->LeaderPoseComponent) ?
-				Cast<USkeletalMeshComponent>(SkeletalMeshComponent->LeaderPoseComponent)->GetSkeletalMeshAsset() :
-				SkeletalMeshComponent->GetSkeletalMeshAsset();
-			if (AnimSequence && SkeletalMesh)
+			USkinnedAsset* SkinnedAsset =
+				SkeletalMeshComponent->LeaderPoseComponent != nullptr ?
+				SkeletalMeshComponent->LeaderPoseComponent->GetSkinnedAsset() :
+				SkeletalMeshComponent->GetSkinnedAsset();
+			if (AnimSequence && SkinnedAsset)
 			{
 				// find the root bone
 				int32 RootIndex = INDEX_NONE;
@@ -149,8 +149,8 @@ void FMovieScene3DTransformSectionRecorder::FinalizeSection(float CurrentTime)
 					const int32 BoneTreeIndex = AnimationTrack.BoneTreeIndex;
 					if (BoneTreeIndex != INDEX_NONE)
 					{
-						const int32 BoneIndex = AnimSkeleton->GetMeshBoneIndexFromSkeletonBoneIndex(SkeletalMesh, BoneTreeIndex);
-						const int32 ParentIndex = SkeletalMesh->GetRefSkeleton().GetParentIndex(BoneIndex);
+						const int32 BoneIndex = AnimSkeleton->GetMeshBoneIndexFromSkeletonBoneIndex(SkinnedAsset, BoneTreeIndex);
+						const int32 ParentIndex = SkinnedAsset->GetRefSkeleton().GetParentIndex(BoneIndex);
 						if (ParentIndex == INDEX_NONE)
 						{
 							// found root
diff --git a/Engine/Source/Runtime/Engine/Classes/Animation/Skeleton.h b/Engine/Source/Runtime/Engine/Classes/Animation/Skeleton.h
--- a/Engine/Source/Runtime/Engine/Classes/Animation/Skeleton.h
+++ b/Engine/Source/Runtime/Engine/Classes/Animation/Skeleton.h
@@ -741,7 +741,11 @@ public:
 	/** Non-serialised cache of linkups between different skeletal meshes and this Skeleton. */
 	TArray<struct FSkeletonToMeshLinkup> LinkupCache;
 
-	/** Runtime built mapping table between SkeletalMeshes, and LinkupCache array indices. */
+private:
+	/** Runtime built mapping table between SkinnedAssets, and LinkupCache array indices. */
+	TMap<TWeakObjectPtr<USkinnedAsset>, int32> SkinnedAsset2LinkupCache;
+public:
+	UE_DEPRECATED(5.1, "Public access to this member variable is deprecated.")
 	TMap<TWeakObjectPtr<USkeletalMesh>, int32> SkelMesh2LinkupCache;
 
 	/** A cached soft object pointer of this skeleton. This is done for performance reasons when searching for compatible skeletons when using IsCompatible(Skeleton). */
@@ -889,42 +893,42 @@ public:
 	 */
 
 	/**
-	 * Verify to see if we can match this skeleton with the provided SkeletalMesh.
+	 * Verify to see if we can match this skeleton with the provided SkinnedAsset.
 	 * 
 	 * Returns true 
 	 *		- if bone hierarchy matches (at least needs to have matching parent) 
 	 *		- and if parent chain matches - meaning if bone tree has A->B->C and if ref pose has A->C, it will fail
 	 *		- and if more than 50 % of bones matches
 	 *  
-	 * @param	InSkelMesh	SkeletalMesh to compare the Skeleton against.
+	 * @param	InSkinnedAsset	SkinnedAsset to compare the Skeleton against.
 	 * @param   bDoParentChainCheck When true (the default) this method also compares if chains match with the parent. 
 	 * 
-	 * @return				true if animation set can play on supplied SkeletalMesh, false if not.
+	 * @return				true if animation set can play on supplied SkinnedAsset, false if not.
 	 */
-	ENGINE_API bool IsCompatibleMesh(const USkeletalMesh* InSkelMesh, bool bDoParentChainCheck=true) const;
+	ENGINE_API bool IsCompatibleMesh(const USkinnedAsset* InSkinnedAsset, bool bDoParentChainCheck=true) const;
 
 	/** Clears all cache data **/
 	ENGINE_API void ClearCacheData();
 
 	/** 
-	 * Find a mesh linkup table (mapping of skeleton bone tree indices to refpose indices) for a particular SkeletalMesh
+	 * Find a mesh linkup table (mapping of skeleton bone tree indices to refpose indices) for a particular SkinnedAsset
 	 * If one does not already exist, create it now.
 	 */
-	ENGINE_API int32 GetMeshLinkupIndex(const USkeletalMesh* InSkelMesh);
+	ENGINE_API int32 GetMeshLinkupIndex(const USkinnedAsset* InSkinnedAsset);
 
 	/** 
-	 * Merge Bones (RequiredBones from InSkelMesh) to BoneTrees if not exists
+	 * Merge Bones (RequiredBones from InSkinnedAsset) to BoneTrees if not exists
 	 * 
 	 * Note that this bonetree can't ever clear up because doing so will corrupt all animation data that was imported based on this
 	 * If nothing exists, it will build new bone tree 
 	 * 
-	 * @param InSkelMesh			: Mesh to build from. 
-	 * @param RequiredRefBones		: RequiredBones are subset of list of bones (index to InSkelMesh->RefSkeleton)
+	 * @param InSkinnedAsset		: Mesh to build from. 
+	 * @param RequiredRefBones		: RequiredBones are subset of list of bones (index to InSkinnedAsset->RefSkeleton)
 									Most of cases, you don't like to add all bones to skeleton, so you'll have choice of cull out some
 	 * 
 	 * @return true if success
 	 */
-	ENGINE_API bool MergeBonesToBoneTree(const USkeletalMesh* InSkeletalMesh, const TArray<int32> &RequiredRefBones);
+	ENGINE_API bool MergeBonesToBoneTree(const USkinnedAsset* InSkinnedAsset, const TArray<int32> &RequiredRefBones);
 
 	/** 
 	 * Merge all Bones to BoneTrees if not exists
@@ -932,22 +936,22 @@ public:
 	 * Note that this bonetree can't ever clear up because doing so will corrupt all animation data that was imported based on this
 	 * If nothing exists, it will build new bone tree 
 	 * 
-	 * @param InSkelMesh			: Mesh to build from. 
+	 * @param InSkinnedAsset		: Mesh to build from. 
 	 * 
 	 * @return true if success
 	 */
-	ENGINE_API bool MergeAllBonesToBoneTree(const USkeletalMesh* InSkelMesh);
+	ENGINE_API bool MergeAllBonesToBoneTree(const USkinnedAsset* InSkinnedAsset);
 
 	/** 
 	 * Merge has failed, then Recreate BoneTree
 	 * 
 	 * This will invalidate all animations that were linked before, but this is needed 
 	 * 
-	 * @param InSkelMesh			: Mesh to build from. 
+	 * @param InSkinnedAsset		: Mesh to build from. 
 	 * 
 	 * @return true if success
 	 */
-	ENGINE_API bool RecreateBoneTree(USkeletalMesh* InSkelMesh);
+	ENGINE_API bool RecreateBoneTree(USkinnedAsset* InSkinnedAsset);
 
 	// @todo document
 	const TArray<FTransform>& GetRefLocalPoses( FName RetargetSource = NAME_None ) const 
@@ -966,11 +970,11 @@ public:
 
 #if WITH_EDITORONLY_DATA
 	/**
-	 * Find a retarget source for a particular skel mesh.
-	 * @param	InMesh	The skeletal mesh to find a source for
+	 * Find a retarget source for a particular mesh.
+	 * @param	InSkinnedAsset	The skinned asset mesh to find a source for
 	 * @return NAME_None if a retarget source was not found, or a valid name if it was
 	 */
-	ENGINE_API FName GetRetargetSourceForMesh(USkeletalMesh* InMesh) const;
+	ENGINE_API FName GetRetargetSourceForMesh(USkinnedAsset* InSkinnedAsset) const;
 #endif
 
 	/** 
@@ -985,19 +989,19 @@ public:
 
 	/** 
 	 * Get Bone Tree Index from Reference Bone Index
-	 * @param	InSkelMesh	SkeletalMesh for the ref bone idx
-	 * @param	InRefBoneIdx	Reference Bone Index to look for - index of USkeletalMesh.RefSkeleton
+	 * @param	InSkinnedAsset	SkinnedAsset for the ref bone idx
+	 * @param	InRefBoneIdx	Reference Bone Index to look for - index of USkinnedAsset.RefSkeleton
 	 * @return	Index of BoneTree Index
 	 */
-	ENGINE_API int32 GetSkeletonBoneIndexFromMeshBoneIndex(const USkeletalMesh* InSkelMesh, const int32 MeshBoneIndex);
+	ENGINE_API int32 GetSkeletonBoneIndexFromMeshBoneIndex(const USkinnedAsset* InSkinnedAsset, const int32 MeshBoneIndex);
 
 	/** 
 	 * Get Reference Bone Index from Bone Tree Index
-	 * @param	InSkelMesh	SkeletalMesh for the ref bone idx
+	 * @param	InSkinnedAsset	SkinnedAsset for the ref bone idx
 	 * @param	InBoneTreeIdx	Bone Tree Index to look for - index of USkeleton.BoneTree
 	 * @return	Index of BoneTree Index
 	 */
-	ENGINE_API int32 GetMeshBoneIndexFromSkeletonBoneIndex(const USkeletalMesh* InSkelMesh, const int32 SkeletonBoneIndex);
+	ENGINE_API int32 GetMeshBoneIndexFromSkeletonBoneIndex(const USkinnedAsset* InSkinnedAsset, const int32 SkeletonBoneIndex);
 
 	EBoneTranslationRetargetingMode::Type GetBoneTranslationRetargetingMode(const int32 BoneTreeIdx) const
 	{
@@ -1011,16 +1015,16 @@ public:
 	/** 
 	 * Rebuild Look up between SkelMesh to BoneTree - this should only get called when SkelMesh is re-imported or so, where the mapping may be no longer valid
 	 *
-	 * @param	InSkelMesh	: SkeletalMesh to build look up for
+	 * @param	InSkinnedAsset	: SkinnedAsset to build look up for
 	 */
-	ENGINE_API void RebuildLinkup(const USkeletalMesh* InSkelMesh);
+	ENGINE_API void RebuildLinkup(const USkinnedAsset* InSkinnedAsset);
 
 	/**
 	 * Remove Link up cache for the SkelMesh
 	 *
-	 * @param	InSkelMesh	: SkeletalMesh to remove linkup cache for 
+	 * @param	InSkinnedAsset	: SkinnedAsset to remove linkup cache for 
 	 */
-	void RemoveLinkup(const USkeletalMesh* InSkelMesh);
+	void RemoveLinkup(const USkinnedAsset* InSkinnedAsset);
 
 	ENGINE_API void SetBoneTranslationRetargetingMode(const int32 BoneIndex, EBoneTranslationRetargetingMode::Type NewRetargetingMode, bool bChildrenToo=false);
 
@@ -1031,15 +1035,15 @@ public:
 	ENGINE_API virtual void Serialize(FArchive& Ar) override;
 
 	/** 
-	 * Create RefLocalPoses from InSkelMesh. Note InSkelMesh cannot be null and this function will assert if it is.
+	 * Create RefLocalPoses from InSkinnedAsset. Note InSkinnedAsset cannot be null and this function will assert if it is.
 	 * 
 	 * If bClearAll is false, it will overwrite ref pose of bones that are found in InSkelMesh
 	 * If bClearAll is true, it will reset all Reference Poses 
-	 * Note that this means it will remove transforms of extra bones that might not be found in this skeletalmesh
+	 * Note that this means it will remove transforms of extra bones that might not be found in this InSkinnedAsset
 	 *
-	 * @return true if successful. false if skeletalmesh wasn't compatible with the bone hierarchy
+	 * @return true if successful. false if InSkinnedAsset wasn't compatible with the bone hierarchy
 	 */
-	ENGINE_API void UpdateReferencePoseFromMesh(const USkeletalMesh* InSkelMesh);
+	ENGINE_API void UpdateReferencePoseFromMesh(const USkinnedAsset* InSkinnedAsset);
 
 #if WITH_EDITORONLY_DATA
 	/**
@@ -1051,25 +1055,25 @@ public:
 #endif
 protected:
 	/** 
-	 * Check if Parent Chain Matches between BoneTree, and SkelMesh 
-	 * Meaning if BoneTree has A->B->C (top to bottom) and if SkelMesh has A->C
+	 * Check if Parent Chain Matches between BoneTree, and SkinnedAsset 
+	 * Meaning if BoneTree has A->B->C (top to bottom) and if SkinnedAsset has A->C
 	 * It will fail since it's missing B
 	 * We ensure this chain matches to play animation properly
 	 *
 	 * @param StartBoneIndex	: BoneTreeIndex to start from in BoneTree 
-	 * @param InSkelMesh		: SkeletalMesh to compare
+	 * @param InSkinnedAsset	: InSkinnedAsset to compare
 	 *
 	 * @return true if matches till root. false if not. 
 	 */
-	bool DoesParentChainMatch(int32 StartBoneTreeIndex, const USkeletalMesh* InSkelMesh) const;
+	bool DoesParentChainMatch(int32 StartBoneTreeIndex, const USkinnedAsset* InSkinnedAsset) const;
 
-	/** 
-	 * Build Look up between SkelMesh to BoneTree
+	/**
+	 * Build Look up between SkinnedAsset to BoneTree
 	 *
-	 * @param	InSkelMesh	: SkeletalMesh to build look up for
-	 * @return	Index of LinkupCache that this SkelMesh is linked to 
+	 * @param	InSkinnedAsset	: SkinnedAsset to build look up for
+	 * @return	Index of LinkupCache that this SkelMesh is linked to
 	 */
-	int32 BuildLinkup(const USkeletalMesh* InSkelMesh);
+	int32 BuildLinkup(const USkinnedAsset* InSkinnedAsset);
 
 #if WITH_EDITORONLY_DATA
 	/**
@@ -1080,12 +1084,12 @@ protected:
 	/**
 	 * Create Reference Skeleton From the given Mesh 
 	 * 
-	 * @param InSkeletonMesh	SkeletalMesh that this Skeleton is based on
+	 * @param InSkinnedAsset	SkinnedAsset that this Skeleton is based on
 	 * @param RequiredRefBones	List of required bones to create skeleton from
 	 *
 	 * @return true if successful
 	 */
-	bool CreateReferenceSkeletonFromMesh(const USkeletalMesh* InSkeletalMesh, const TArray<int32> & RequiredRefBones);
+	bool CreateReferenceSkeletonFromMesh(const USkinnedAsset* InSkinnedAsset, const TArray<int32> & RequiredRefBones);
 
 #if WITH_EDITOR
 	DECLARE_MULTICAST_DELEGATE( FOnSkeletonHierarchyChangedMulticaster );
diff --git a/Engine/Source/Runtime/Engine/Private/Animation/Skeleton.cpp b/Engine/Source/Runtime/Engine/Private/Animation/Skeleton.cpp
--- a/Engine/Source/Runtime/Engine/Private/Animation/Skeleton.cpp
+++ b/Engine/Source/Runtime/Engine/Private/Animation/Skeleton.cpp
@@ -908,10 +908,10 @@ void USkeleton::RemoveSkeletonRemapping(const USkeleton* SourceSkeleton)
 	}
 }
 
-bool USkeleton::DoesParentChainMatch(int32 StartBoneIndex, const USkeletalMesh* InSkelMesh) const
+bool USkeleton::DoesParentChainMatch(int32 StartBoneIndex, const USkinnedAsset* InSkinnedAsset) const
 {
 	const FReferenceSkeleton& SkeletonRefSkel = ReferenceSkeleton;
-	const FReferenceSkeleton& MeshRefSkel = InSkelMesh->GetRefSkeleton();
+	const FReferenceSkeleton& MeshRefSkel = InSkinnedAsset->GetRefSkeleton();
 
 	// if start is root bone
 	if ( StartBoneIndex == 0 )
@@ -953,13 +953,13 @@ bool USkeleton::DoesParentChainMatch(int32 StartBoneIndex, const USkeletalMesh*
 	return true;
 }
 
-bool USkeleton::IsCompatibleMesh(const USkeletalMesh* InSkelMesh, bool bDoParentChainCheck) const
+bool USkeleton::IsCompatibleMesh(const USkinnedAsset* InSkinnedAsset, bool bDoParentChainCheck) const
 {
 	// at least % of bone should match 
 	int32 NumOfBoneMatches = 0;
 
 	const FReferenceSkeleton& SkeletonRefSkel = ReferenceSkeleton;
-	const FReferenceSkeleton& MeshRefSkel = InSkelMesh->GetRefSkeleton();
+	const FReferenceSkeleton& MeshRefSkel = InSkinnedAsset->GetRefSkeleton();
 	const int32 NumBones = MeshRefSkel.GetRawBoneNum();
 
 	// first ensure the parent exists for each bone
@@ -975,7 +975,7 @@ bool USkeleton::IsCompatibleMesh(const USkeletalMesh* InSkelMesh, bool bDoParent
 			++NumOfBoneMatches;
 
 			// follow the parent chain to verify the chain is same
-			if(bDoParentChainCheck && !DoesParentChainMatch(SkeletonBoneIndex, InSkelMesh))
+			if(bDoParentChainCheck && !DoesParentChainMatch(SkeletonBoneIndex, InSkinnedAsset))
 			{
 				UE_LOG(LogAnimation, Verbose, TEXT("%s : Hierarchy does not match."), *MeshBoneName.ToString());
 				return false;
@@ -1015,7 +1015,7 @@ bool USkeleton::IsCompatibleMesh(const USkeletalMesh* InSkelMesh, bool bDoParent
 			}
 
 			// second follow the parent chain to verify the chain is same
-			if (bDoParentChainCheck && !DoesParentChainMatch(SkeletonBoneIndex, InSkelMesh))
+			if (bDoParentChainCheck && !DoesParentChainMatch(SkeletonBoneIndex, InSkinnedAsset))
 			{
 				UE_LOG(LogAnimation, Verbose, TEXT("%s : Hierarchy does not match."), *MeshBoneName.ToString());
 				return false;
@@ -1033,21 +1033,24 @@ void USkeleton::ClearCacheData()
 {
 	FScopeLock ScopeLock(&LinkupCacheLock);
 	LinkupCache.Empty();
+	PRAGMA_DISABLE_DEPRECATION_WARNINGS
 	SkelMesh2LinkupCache.Empty();
+	PRAGMA_ENABLE_DEPRECATION_WARNINGS
+	SkinnedAsset2LinkupCache.Empty();
 }
 
-int32 USkeleton::GetMeshLinkupIndex(const USkeletalMesh* InSkelMesh)
+int32 USkeleton::GetMeshLinkupIndex(const USkinnedAsset* InSkinnedAsset)
 {
 	int32* IndexPtr = nullptr;
 	{
 		FScopeLock ScopeLock(&LinkupCacheLock);
-		IndexPtr = SkelMesh2LinkupCache.Find(MakeWeakObjectPtr(const_cast<USkeletalMesh*>(InSkelMesh)));
+		IndexPtr = SkinnedAsset2LinkupCache.Find(MakeWeakObjectPtr(const_cast<USkinnedAsset*>(InSkinnedAsset)));
 	}
 	int32 LinkupIndex = INDEX_NONE;
 
 	if ( IndexPtr == NULL )
 	{
-		LinkupIndex = BuildLinkup(InSkelMesh);
+		LinkupIndex = BuildLinkup(InSkinnedAsset);
 	}
 	else
 	{
@@ -1063,16 +1066,22 @@ int32 USkeleton::GetMeshLinkupIndex(const USkeletalMesh* InSkelMesh)
 	return LinkupIndex;
 }
 
-void USkeleton::RemoveLinkup(const USkeletalMesh* InSkelMesh)
+void USkeleton::RemoveLinkup(const USkinnedAsset* InSkinnedAsset)
 {
 	FScopeLock ScopeLock(&LinkupCacheLock);
-	SkelMesh2LinkupCache.Remove(MakeWeakObjectPtr(const_cast<USkeletalMesh*>(InSkelMesh)));
+	PRAGMA_DISABLE_DEPRECATION_WARNINGS
+	if (const USkeletalMesh* const InSkelMesh = Cast<const USkeletalMesh>(InSkinnedAsset))
+	{
+		SkelMesh2LinkupCache.Remove(MakeWeakObjectPtr(const_cast<USkeletalMesh*>(InSkelMesh)));
+	}
+	PRAGMA_ENABLE_DEPRECATION_WARNINGS
+	SkinnedAsset2LinkupCache.Remove(MakeWeakObjectPtr(const_cast<USkinnedAsset*>(InSkinnedAsset)));
 }
 
-int32 USkeleton::BuildLinkup(const USkeletalMesh* InSkelMesh)
+int32 USkeleton::BuildLinkup(const USkinnedAsset* InSkinnedAsset)
 {
 	const FReferenceSkeleton& SkeletonRefSkel = ReferenceSkeleton;
-	const FReferenceSkeleton& MeshRefSkel = InSkelMesh->GetRefSkeleton();
+	const FReferenceSkeleton& MeshRefSkel = InSkinnedAsset->GetRefSkeleton();
 
 	// @todoanim : need to refresh NULL SkeletalMeshes from Cache
 	// since now they're autoweak pointer, they will go away if not used
@@ -1103,12 +1112,12 @@ int32 USkeleton::BuildLinkup(const USkeletalMesh* InSkelMesh)
 			Message->AddToken(FTextToken::Create(LOCTEXT("SkeletonBuildLinkupMissingBones1", "The Skeleton ")));
 			Message->AddToken(FAssetNameToken::Create(GetPathName(), FText::FromString( GetNameSafe(this) ) ));
 			Message->AddToken(FTextToken::Create(LOCTEXT("SkeletonBuildLinkupMissingBones2", " is missing bones that SkeletalMesh ")));
-			Message->AddToken(FAssetNameToken::Create(InSkelMesh->GetPathName(), FText::FromString( GetNameSafe(InSkelMesh) )));
+			Message->AddToken(FAssetNameToken::Create(InSkinnedAsset->GetPathName(), FText::FromString( GetNameSafe(InSkinnedAsset) )));
 			Message->AddToken(FTextToken::Create(LOCTEXT("SkeletonBuildLinkupMissingBones3", "  needs. They will be added now. Please save the Skeleton!")));
 			LoadErrors.Open();
 
 			// Re-add all SkelMesh bones to the Skeleton.
-			MergeAllBonesToBoneTree(InSkelMesh);
+			MergeAllBonesToBoneTree(InSkinnedAsset);
 
 			// Fix missing bone.
 			SkeletonBoneIndex = SkeletonRefSkel.FindBoneIndex(MeshBoneName);
@@ -1116,7 +1125,7 @@ int32 USkeleton::BuildLinkup(const USkeletalMesh* InSkelMesh)
 #else
 		// If we're not in editor, we still want to know which skeleton is missing a bone.
 		ensureMsgf(SkeletonBoneIndex != INDEX_NONE, TEXT("USkeleton::BuildLinkup: The Skeleton %s, is missing bones that SkeletalMesh %s needs. MeshBoneName %s"),
-				*GetNameSafe(this), *GetNameSafe(InSkelMesh), *MeshBoneName.ToString());
+				*GetNameSafe(this), *GetNameSafe(InSkinnedAsset), *MeshBoneName.ToString());
 #endif
 
 		NewMeshLinkup.MeshToSkeletonTable[MeshBoneIndex] = SkeletonBoneIndex;
@@ -1138,64 +1147,70 @@ int32 USkeleton::BuildLinkup(const USkeletalMesh* InSkelMesh)
 		FScopeLock ScopeLock(&LinkupCacheLock);
 		NewIndex = LinkupCache.Add(NewMeshLinkup);
 		check(NewIndex != INDEX_NONE);
-		SkelMesh2LinkupCache.Add(MakeWeakObjectPtr(const_cast<USkeletalMesh*>(InSkelMesh)), NewIndex);
+		PRAGMA_DISABLE_DEPRECATION_WARNINGS
+		if (const USkeletalMesh* const InSkelMesh = Cast<const USkeletalMesh>(InSkinnedAsset))
+		{
+			SkelMesh2LinkupCache.Add(MakeWeakObjectPtr(const_cast<USkeletalMesh*>(InSkelMesh)), NewIndex);
+		}
+		PRAGMA_ENABLE_DEPRECATION_WARNINGS
+		SkinnedAsset2LinkupCache.Add(MakeWeakObjectPtr(const_cast<USkinnedAsset*>(InSkinnedAsset)), NewIndex);
 	}
 	return NewIndex;
 }
 
 
-void USkeleton::RebuildLinkup(const USkeletalMesh* InSkelMesh)
+void USkeleton::RebuildLinkup(const USkinnedAsset* InSkinnedAsset)
 {
 	// remove the key
-	RemoveLinkup(InSkelMesh);
+	RemoveLinkup(InSkinnedAsset);
 	// build new one
-	BuildLinkup(InSkelMesh);
+	BuildLinkup(InSkinnedAsset);
 }
 
-void USkeleton::UpdateReferencePoseFromMesh(const USkeletalMesh* InSkelMesh)
+void USkeleton::UpdateReferencePoseFromMesh(const USkinnedAsset* InSkinnedAsset)
 {
-	check(InSkelMesh);
+	check(InSkinnedAsset);
 	
 	FReferenceSkeletonModifier RefSkelModifier(ReferenceSkeleton, this);
 
 	for (int32 BoneIndex = 0; BoneIndex < ReferenceSkeleton.GetRawBoneNum(); BoneIndex++)
 	{
 		// find index from ref pose array
-		const int32 MeshBoneIndex = InSkelMesh->GetRefSkeleton().FindRawBoneIndex(ReferenceSkeleton.GetBoneName(BoneIndex));
+		const int32 MeshBoneIndex = InSkinnedAsset->GetRefSkeleton().FindRawBoneIndex(ReferenceSkeleton.GetBoneName(BoneIndex));
 		if (MeshBoneIndex != INDEX_NONE)
 		{
-			RefSkelModifier.UpdateRefPoseTransform(BoneIndex, InSkelMesh->GetRefSkeleton().GetRefBonePose()[MeshBoneIndex]);
+			RefSkelModifier.UpdateRefPoseTransform(BoneIndex, InSkinnedAsset->GetRefSkeleton().GetRefBonePose()[MeshBoneIndex]);
 		}
 	}
 
 	MarkPackageDirty();
 }
 
-bool USkeleton::RecreateBoneTree(USkeletalMesh* InSkelMesh)
+bool USkeleton::RecreateBoneTree(USkinnedAsset* InSkinnedAsset)
 {
-	if( InSkelMesh )
+	if( InSkinnedAsset )
 	{
 		// regenerate Guid
 		RegenerateGuid();	
 		BoneTree.Empty();
 		ReferenceSkeleton.Empty();
 
-		return MergeAllBonesToBoneTree(InSkelMesh);
+		return MergeAllBonesToBoneTree(InSkinnedAsset);
 	}
 
 	return false;
 }
 
-bool USkeleton::MergeAllBonesToBoneTree(const USkeletalMesh* InSkelMesh)
+bool USkeleton::MergeAllBonesToBoneTree(const USkinnedAsset* InSkinnedAsset)
 {
-	if( InSkelMesh )
+	if( InSkinnedAsset )
 	{
 		TArray<int32> RequiredBoneIndices;
 
 		// for now add all in this case. 
-		RequiredBoneIndices.AddUninitialized(InSkelMesh->GetRefSkeleton().GetRawBoneNum());
+		RequiredBoneIndices.AddUninitialized(InSkinnedAsset->GetRefSkeleton().GetRawBoneNum());
 		// gather bone list
-		for (int32 I=0; I<InSkelMesh->GetRefSkeleton().GetRawBoneNum(); ++I)
+		for (int32 I=0; I<InSkinnedAsset->GetRefSkeleton().GetRawBoneNum(); ++I)
 		{
 			RequiredBoneIndices[I] = I;
 		}
@@ -1203,18 +1218,18 @@ bool USkeleton::MergeAllBonesToBoneTree(const USkeletalMesh* InSkelMesh)
 		if( RequiredBoneIndices.Num() > 0 )
 		{
 			// merge bones to the selected skeleton
-			return MergeBonesToBoneTree( InSkelMesh, RequiredBoneIndices );
+			return MergeBonesToBoneTree( InSkinnedAsset, RequiredBoneIndices );
 		}
 	}
 
 	return false;
 }
 
-bool USkeleton::CreateReferenceSkeletonFromMesh(const USkeletalMesh* InSkeletalMesh, const TArray<int32> & RequiredRefBones)
+bool USkeleton::CreateReferenceSkeletonFromMesh(const USkinnedAsset* InSkinnedAsset, const TArray<int32> & RequiredRefBones)
 {
 	// Filter list, we only want bones that have their parents present in this array.
 	TArray<int32> FilteredRequiredBones; 
-	FAnimationRuntime::ExcludeBonesWithNoParents(RequiredRefBones, InSkeletalMesh->GetRefSkeleton(), FilteredRequiredBones);
+	FAnimationRuntime::ExcludeBonesWithNoParents(RequiredRefBones, InSkinnedAsset->GetRefSkeleton(), FilteredRequiredBones);
 
 	FReferenceSkeletonModifier RefSkelModifier(ReferenceSkeleton, this);
 
@@ -1229,7 +1244,7 @@ bool USkeleton::CreateReferenceSkeletonFromMesh(const USkeletalMesh* InSkeletalM
 		{
 			const int32& BoneIndex = FilteredRequiredBones[Index];
 
-			FMeshBoneInfo NewMeshBoneInfo = InSkeletalMesh->GetRefSkeleton().GetRefBoneInfo()[BoneIndex];
+			FMeshBoneInfo NewMeshBoneInfo = InSkinnedAsset->GetRefSkeleton().GetRefBoneInfo()[BoneIndex];
 			// Fix up ParentIndex for our new Skeleton.
 			if( BoneIndex == 0 )
 			{
@@ -1237,11 +1252,11 @@ bool USkeleton::CreateReferenceSkeletonFromMesh(const USkeletalMesh* InSkeletalM
 			}
 			else
 			{
-				const int32 ParentIndex = InSkeletalMesh->GetRefSkeleton().GetParentIndex(BoneIndex);
-				const FName ParentName = InSkeletalMesh->GetRefSkeleton().GetBoneName(ParentIndex);
+				const int32 ParentIndex = InSkinnedAsset->GetRefSkeleton().GetParentIndex(BoneIndex);
+				const FName ParentName = InSkinnedAsset->GetRefSkeleton().GetBoneName(ParentIndex);
 				NewMeshBoneInfo.ParentIndex = ReferenceSkeleton.FindRawBoneIndex(ParentName);
 			}
-			RefSkelModifier.Add(NewMeshBoneInfo, InSkeletalMesh->GetRefSkeleton().GetRefBonePose()[BoneIndex]);
+			RefSkelModifier.Add(NewMeshBoneInfo, InSkinnedAsset->GetRefSkeleton().GetRefBonePose()[BoneIndex]);
 		}
 
 		return true;
@@ -1251,7 +1266,7 @@ bool USkeleton::CreateReferenceSkeletonFromMesh(const USkeletalMesh* InSkeletalM
 }
 
 
-bool USkeleton::MergeBonesToBoneTree(const USkeletalMesh* InSkeletalMesh, const TArray<int32> & RequiredRefBones)
+bool USkeleton::MergeBonesToBoneTree(const USkinnedAsset* InSkinnedAsset, const TArray<int32> & RequiredRefBones)
 {
 	// see if it needs all animation data to remap - only happens when bone structure CHANGED - added
 	bool bSuccess = false;
@@ -1262,29 +1277,29 @@ bool USkeleton::MergeBonesToBoneTree(const USkeletalMesh* InSkeletalMesh, const
 	// if it's first time
 	if( BoneTree.Num() == 0 )
 	{
-		bSuccess = CreateReferenceSkeletonFromMesh(InSkeletalMesh, RequiredRefBones);
+		bSuccess = CreateReferenceSkeletonFromMesh(InSkinnedAsset, RequiredRefBones);
 		bShouldHandleHierarchyChange = true;
 	}
 	else
 	{
 		// can we play? - hierarchy matches
-		if( IsCompatibleMesh(InSkeletalMesh) )
+		if( IsCompatibleMesh(InSkinnedAsset) )
 		{
 			// Exclude bones who do not have a parent.
 			TArray<int32> FilteredRequiredBones;
-			FAnimationRuntime::ExcludeBonesWithNoParents(RequiredRefBones, InSkeletalMesh->GetRefSkeleton(), FilteredRequiredBones);
+			FAnimationRuntime::ExcludeBonesWithNoParents(RequiredRefBones, InSkinnedAsset->GetRefSkeleton(), FilteredRequiredBones);
 
 			FReferenceSkeletonModifier RefSkelModifier(ReferenceSkeleton, this);
 
 			for (int32 Index=0; Index<FilteredRequiredBones.Num(); Index++)
 			{
 				const int32& MeshBoneIndex = FilteredRequiredBones[Index];
-				const int32& SkeletonBoneIndex = ReferenceSkeleton.FindRawBoneIndex(InSkeletalMesh->GetRefSkeleton().GetBoneName(MeshBoneIndex));
+				const int32& SkeletonBoneIndex = ReferenceSkeleton.FindRawBoneIndex(InSkinnedAsset->GetRefSkeleton().GetBoneName(MeshBoneIndex));
 				
 				// Bone doesn't already exist. Add it.
 				if( SkeletonBoneIndex == INDEX_NONE )
 				{
-					FMeshBoneInfo NewMeshBoneInfo = InSkeletalMesh->GetRefSkeleton().GetRefBoneInfo()[MeshBoneIndex];
+					FMeshBoneInfo NewMeshBoneInfo = InSkinnedAsset->GetRefSkeleton().GetRefBoneInfo()[MeshBoneIndex];
 					// Fix up ParentIndex for our new Skeleton.
 					if( ReferenceSkeleton.GetRawBoneNum() == 0 )
 					{
@@ -1292,10 +1307,10 @@ bool USkeleton::MergeBonesToBoneTree(const USkeletalMesh* InSkeletalMesh, const
 					}
 					else
 					{
-						NewMeshBoneInfo.ParentIndex = ReferenceSkeleton.FindRawBoneIndex(InSkeletalMesh->GetRefSkeleton().GetBoneName(InSkeletalMesh->GetRefSkeleton().GetParentIndex(MeshBoneIndex)));
+						NewMeshBoneInfo.ParentIndex = ReferenceSkeleton.FindRawBoneIndex(InSkinnedAsset->GetRefSkeleton().GetBoneName(InSkinnedAsset->GetRefSkeleton().GetParentIndex(MeshBoneIndex)));
 					}
 
-					RefSkelModifier.Add(NewMeshBoneInfo, InSkeletalMesh->GetRefSkeleton().GetRefBonePose()[MeshBoneIndex]);
+					RefSkelModifier.Add(NewMeshBoneInfo, InSkinnedAsset->GetRefSkeleton().GetRefBonePose()[MeshBoneIndex]);
 					BoneTree.AddZeroed(1);
 					bShouldHandleHierarchyChange = true;
 				}
@@ -1336,9 +1351,9 @@ void USkeleton::SetBoneTranslationRetargetingMode(const int32 BoneIndex, EBoneTr
 
 #if WITH_EDITORONLY_DATA
 
-FName USkeleton::GetRetargetSourceForMesh(USkeletalMesh* InMesh) const
+FName USkeleton::GetRetargetSourceForMesh(USkinnedAsset* InSkinnedAsset) const
 {
-	FSoftObjectPath MeshPath(InMesh);
+	FSoftObjectPath MeshPath(InSkinnedAsset);
 	for(const TPair<FName, FReferencePose>& AnimRetargetSource : AnimRetargetSources)
 	{
 		if(AnimRetargetSource.Value.SourceReferenceMesh.ToSoftObjectPath() == MeshPath)
@@ -1372,10 +1387,10 @@ int32 USkeleton::GetRawAnimationTrackIndex(const int32 InSkeletonBoneIndex, cons
 	return INDEX_NONE;
 }
 
-int32 USkeleton::GetSkeletonBoneIndexFromMeshBoneIndex(const USkeletalMesh* InSkelMesh, const int32 MeshBoneIndex)
+int32 USkeleton::GetSkeletonBoneIndexFromMeshBoneIndex(const USkinnedAsset* InSkinnedAsset, const int32 MeshBoneIndex)
 {
 	check(MeshBoneIndex != INDEX_NONE);
-	const int32 LinkupCacheIdx = GetMeshLinkupIndex(InSkelMesh);
+	const int32 LinkupCacheIdx = GetMeshLinkupIndex(InSkinnedAsset);
 
 	//LinkupCache lock scope
 	{
@@ -1386,10 +1401,10 @@ int32 USkeleton::GetSkeletonBoneIndexFromMeshBoneIndex(const USkeletalMesh* InSk
 }
 
 
-int32 USkeleton::GetMeshBoneIndexFromSkeletonBoneIndex(const USkeletalMesh* InSkelMesh, const int32 SkeletonBoneIndex)
+int32 USkeleton::GetMeshBoneIndexFromSkeletonBoneIndex(const USkinnedAsset* InSkinnedAsset, const int32 SkeletonBoneIndex)
 {
 	check(SkeletonBoneIndex != INDEX_NONE);
-	const int32 LinkupCacheIdx = GetMeshLinkupIndex(InSkelMesh);
+	const int32 LinkupCacheIdx = GetMeshLinkupIndex(InSkinnedAsset);
 
 	//LinkupCache lock scope
 	{
diff --git a/Engine/Source/Runtime/Engine/Private/ObjectCacheContext.cpp b/Engine/Source/Runtime/Engine/Private/ObjectCacheContext.cpp
--- a/Engine/Source/Runtime/Engine/Private/ObjectCacheContext.cpp
+++ b/Engine/Source/Runtime/Engine/Private/ObjectCacheContext.cpp
@@ -564,11 +564,6 @@ TObjectCacheIterator<UMaterialInterface> FObjectCacheContext::GetMaterialsAffect
 	return TObjectCacheIterator<UMaterialInterface>(ComputeResult.Array());
 }
 
-TObjectCacheIterator<USkinnedMeshComponent> FObjectCacheContext::GetSkinnedMeshComponents(USkeletalMesh* InSkeletalMesh)
-{
-	return GetSkinnedMeshComponents(Cast<USkinnedAsset>(InSkeletalMesh));
-}
-
 TObjectCacheIterator<USkinnedMeshComponent> FObjectCacheContext::GetSkinnedMeshComponents(USkinnedAsset* InSkinnedAsset)
 {
 	if (!SkinnedAssetToComponents.IsSet())
diff --git a/Engine/Source/Runtime/Engine/Public/ObjectCacheContext.h b/Engine/Source/Runtime/Engine/Public/ObjectCacheContext.h
--- a/Engine/Source/Runtime/Engine/Public/ObjectCacheContext.h
+++ b/Engine/Source/Runtime/Engine/Public/ObjectCacheContext.h
@@ -70,8 +70,6 @@ public:
 	TObjectCacheIterator<UStaticMeshComponent>  GetStaticMeshComponents();
 	TObjectCacheIterator<USkinnedMeshComponent> GetSkinnedMeshComponents();
 	TObjectCacheIterator<USkinnedMeshComponent> GetSkinnedMeshComponents(USkinnedAsset* InSkinnedAsset);
-	UE_DEPRECATED(5.1, "Use GetSkinnedMeshComponents(USkinnedAsset* InSkinnedAsset) instead")
-	TObjectCacheIterator<USkinnedMeshComponent> GetSkinnedMeshComponents(USkeletalMesh* InSkeletalMesh);
 	TObjectCacheIterator<UStaticMeshComponent>  GetStaticMeshComponents(UStaticMesh* InStaticMesh);
 	TObjectCacheIterator<UMaterialInterface>    GetMaterialsAffectedByTexture(UTexture* InTexture);
 	TObjectCacheIterator<UPrimitiveComponent>   GetPrimitivesAffectedByMaterial(UMaterialInterface* InMaterial);
