Preload animation compression settings before we load any game content in editor builds

Fixes warnings about calling LoadObject in PostLoad of anim sequences

diff --git a/Engine/Source/Runtime/Engine/Private/Animation/AnimationUtils.cpp b/Engine/Source/Runtime/Engine/Private/Animation/AnimationUtils.cpp
--- a/Engine/Source/Runtime/Engine/Private/Animation/AnimationUtils.cpp
+++ b/Engine/Source/Runtime/Engine/Private/Animation/AnimationUtils.cpp
@@ -874,6 +874,8 @@ static void EnsureDependenciesAreLoaded(UObject* Object)
 	}
 }
 
+
+
 UObject* GetDefaultAnimationCompressionSettings(const TCHAR* IniValueName)
 {
 	FConfigSection* AnimDefaultObjectSettingsSection = GConfig->GetSectionPrivate(TEXT("Animation.DefaultObjectSettings"), false, true, GEngineIni);
@@ -901,6 +903,17 @@ UObject* GetDefaultAnimationCompressionSettings(const TCHAR* IniValueName)
 	return DefaultCompressionSettings;
 }
 
+#if WITH_EDITOR
+
+void FAnimationUtils::PreloadCompressionSettings()
+{
+	GetDefaultAnimationBoneCompressionSettings();
+	GetDefaultAnimationRecorderBoneCompressionSettings();
+	GetDefaultAnimationCurveCompressionSettings();
+}
+
+#endif
+
 UAnimBoneCompressionSettings* FAnimationUtils::GetDefaultAnimationBoneCompressionSettings()
 {
 	if (DefaultBoneCompressionSettings == nullptr)
diff --git a/Engine/Source/Runtime/Engine/Public/AnimationUtils.h b/Engine/Source/Runtime/Engine/Public/AnimationUtils.h
--- a/Engine/Source/Runtime/Engine/Public/AnimationUtils.h
+++ b/Engine/Source/Runtime/Engine/Public/AnimationUtils.h
@@ -145,6 +145,11 @@ public:
 		const FVector& ScaleNudge,
 		TArray<FAnimPerturbationError>& InducedErrors);
 
+#if WITH_EDITOR
+	/** Preload the bone, curve and recorder compression settings */
+	ENGINE_API static void PreloadCompressionSettings();
+#endif
+	
 	/** Returns the default animation bone compression settings, can never by null. */
 	ENGINE_API static UAnimBoneCompressionSettings* GetDefaultAnimationBoneCompressionSettings();
 
diff --git a/Engine/Source/Runtime/Launch/Private/LaunchEngineLoop.cpp b/Engine/Source/Runtime/Launch/Private/LaunchEngineLoop.cpp
--- a/Engine/Source/Runtime/Launch/Private/LaunchEngineLoop.cpp
+++ b/Engine/Source/Runtime/Launch/Private/LaunchEngineLoop.cpp
@@ -110,6 +110,7 @@
 	#include "Serialization/BulkDataRegistry.h"
 	#include "ShaderCompiler.h"
 	#include "Virtualization/VirtualizationSystem.h"
+	#include "AnimationUtils.h"
 
 	#if PLATFORM_WINDOWS
 		#include "Windows/AllowWindowsPlatformTypes.h"
@@ -3526,8 +3527,11 @@ int32 FEngineLoop::PreInitPostStartupScreen(const TCHAR* CmdLine)
 #if WITH_EDITOR
 	// We need to mount the shared resources for templates (if there are any) before we try and load and game classes
 	FUnrealEdMisc::Get().MountTemplateSharedPaths();
-#endif
 
+	// We need to load any animation compression settings before we load game classes
+	FAnimationUtils::PreloadCompressionSettings();
+#endif
+	
 	{
 		SCOPED_BOOT_TIMING("LoadStartupModules");
 		if (!LoadStartupModules())
