LevelInstance: Avoid getting bounds from LevelPackage on OFPA levels

diff --git a/Engine/Source/Runtime/Engine/Private/WorldPartition/LevelInstance/LevelInstanceActorDesc.cpp b/Engine/Source/Runtime/Engine/Private/WorldPartition/LevelInstance/LevelInstanceActorDesc.cpp
--- a/Engine/Source/Runtime/Engine/Private/WorldPartition/LevelInstance/LevelInstanceActorDesc.cpp
+++ b/Engine/Source/Runtime/Engine/Private/WorldPartition/LevelInstance/LevelInstanceActorDesc.cpp
@@ -225,10 +225,13 @@ void FLevelInstanceActorDesc::Serialize(FArchive& Ar)
 		const ILevelInstanceInterface* LevelInstanceCDO = CastChecked<ILevelInstanceInterface>(CDO);
 		if (!LevelPackage.IsNone() && (LevelInstanceCDO->IsLoadingEnabled() || bFixupOldVersion))
 		{
-			FBox OutBounds;
-			if (ULevelInstanceSubsystem::GetLevelInstanceBoundsFromPackage(LevelInstanceTransform, LevelPackage, OutBounds))
+			if (!IsContainerInstance())
 			{
-				OutBounds.GetCenterAndExtents(BoundsLocation, BoundsExtent);
+				FBox OutBounds;
+				if (ULevelInstanceSubsystem::GetLevelInstanceBoundsFromPackage(LevelInstanceTransform, LevelPackage, OutBounds))
+				{
+					OutBounds.GetCenterAndExtents(BoundsLocation, BoundsExtent);
+				}
 			}
 		}
 	}
