Removed the EngineSessionManager private class that was only used to report the'Engine.AbnormalShutdown'. According to analysts, this event is not used anywhere (and likely unreliable). For the Editor, the functionality was superseeded by the SummarySession which contains the Editor shutdown type.

diff --git a/Engine/Source/Editor/UnrealEd/Public/EditorAnalytics.h b/Engine/Source/Editor/UnrealEd/Public/EditorAnalytics.h
--- a/Engine/Source/Editor/UnrealEd/Public/EditorAnalytics.h
+++ b/Engine/Source/Editor/UnrealEd/Public/EditorAnalytics.h
@@ -4,7 +4,6 @@
 
 #include "CoreMinimal.h"
 
-class FEngineSessionManager;
 struct FAnalyticsEventAttribute;
 
 namespace EAnalyticsErrorCodes
@@ -118,7 +117,4 @@ public:
 
 	/* Determine whether the error code should cause the reporter to display the message in a dialog. */
 	UNREALED_API static bool ShouldElevateMessageThroughDialog(const int32 ErrorCode);
-
-private:
-	static TSharedPtr<class FEngineSessionManager> SessionManager;
 };
diff --git a/Engine/Source/Runtime/Engine/Private/EngineAnalytics.cpp b/Engine/Source/Runtime/Engine/Private/EngineAnalytics.cpp
--- a/Engine/Source/Runtime/Engine/Private/EngineAnalytics.cpp
+++ b/Engine/Source/Runtime/Engine/Private/EngineAnalytics.cpp
@@ -13,7 +13,6 @@
 #include "IAnalyticsProviderET.h"
 #include "AnalyticsET.h"
 #include "GeneralProjectSettings.h"
-#include "EngineSessionManager.h"
 #include "Misc/EngineVersion.h"
 #include "RHI.h"
 #include "GenericPlatform/GenericPlatformCrashContext.h"
@@ -29,7 +28,6 @@
 
 bool FEngineAnalytics::bIsInitialized;
 TSharedPtr<IAnalyticsProviderET> FEngineAnalytics::Analytics;
-TSharedPtr<FEngineSessionManager> FEngineAnalytics::SessionManager;
 
 #if WITH_EDITOR
 static TUniquePtr<FAnalyticsSessionSummaryManager> AnalyticsSessionSummaryManager;
@@ -153,13 +151,6 @@ void FEngineAnalytics::Initialize()
 			bIsInitialized = true;
 		}
 
-		// Create the session manager singleton
-		if (!SessionManager.IsValid())
-		{
-			SessionManager = MakeShared<FEngineSessionManager>(EEngineSessionManagerMode::Editor);
-			SessionManager->Initialize();
-		}
-
 #if WITH_EDITOR
 		if (!AnalyticsSessionSummaryManager)
 		{
@@ -191,13 +182,6 @@ void FEngineAnalytics::Initialize()
 
 void FEngineAnalytics::Shutdown(bool bIsEngineShutdown)
 {
-	// Destroy the session manager singleton if it exists
-	if (SessionManager.IsValid() && bIsEngineShutdown)
-	{
-		SessionManager->Shutdown();
-		SessionManager.Reset();
-	}
-
 #if WITH_EDITOR
 	OnShutdownEngineAnalytics.Broadcast();
 
@@ -233,11 +217,6 @@ void FEngineAnalytics::Tick(float DeltaTime)
 {
 	QUICK_SCOPE_CYCLE_COUNTER(STAT_FEngineAnalytics_Tick);
 
-	if (SessionManager.IsValid())
-	{
-		SessionManager->Tick(DeltaTime);
-	}
-
 #if WITH_EDITOR
 	if (EditorAnalyticSessionSummary)
 	{
diff --git a/Engine/Source/Runtime/Engine/Private/EngineSessionManager.cpp b/Engine/Source/Runtime/Engine/Private/EngineSessionManager.cpp
deleted file mode 100644
diff --git a/Engine/Source/Runtime/Engine/Private/EngineSessionManager.h b/Engine/Source/Runtime/Engine/Private/EngineSessionManager.h
deleted file mode 100644
diff --git a/Engine/Source/Runtime/Engine/Public/EngineAnalytics.h b/Engine/Source/Runtime/Engine/Public/EngineAnalytics.h
--- a/Engine/Source/Runtime/Engine/Public/EngineAnalytics.h
+++ b/Engine/Source/Runtime/Engine/Public/EngineAnalytics.h
@@ -59,6 +59,5 @@ public:
 private:
 	static bool bIsInitialized;
 	static ENGINE_API TSharedPtr<IAnalyticsProviderET> Analytics;
-	static TSharedPtr<FEngineSessionManager> SessionManager;
 };
 
