Add a BP callable getter for UStaticMesh::bAutoComputeLODScreenSize

diff --git a/Engine/Source/Runtime/Engine/Classes/Engine/StaticMesh.h b/Engine/Source/Runtime/Engine/Classes/Engine/StaticMesh.h
--- a/Engine/Source/Runtime/Engine/Classes/Engine/StaticMesh.h
+++ b/Engine/Source/Runtime/Engine/Classes/Engine/StaticMesh.h
@@ -731,6 +731,14 @@ public:
 
 #endif // #if WITH_EDITORONLY_DATA
 
+#if WITH_EDITOR
+	UFUNCTION(BlueprintCallable, Category = StaticMesh)
+	bool IsLODScreenSizeAutoComputed() const
+	{
+		return bAutoComputeLODScreenSize;
+	}
+#endif
+
 	/** Check the QualitLevel property is enabled for MinLod. */
 	bool IsMinLodQualityLevelEnable() const;
 
