Don't copy virtual texture data if sizes don't match the expected values.
It's not clear if/how that happens, and it might be in shipping builds only where the checks were compiled out.

diff --git a/Engine/Source/Runtime/Engine/Private/VT/VirtualTextureTranscodeCache.cpp b/Engine/Source/Runtime/Engine/Private/VT/VirtualTextureTranscodeCache.cpp
--- a/Engine/Source/Runtime/Engine/Private/VT/VirtualTextureTranscodeCache.cpp
+++ b/Engine/Source/Runtime/Engine/Private/VT/VirtualTextureTranscodeCache.cpp
@@ -152,19 +152,25 @@ struct FTranscodeTask
 			case EVirtualTextureCodec::RawGPU:
 				if (StagingBufferForLayer.Stride == PackedStride)
 				{
-					check(TileLayerSize <= StagingBufferForLayer.MemorySize);
-					Params.Data->CopyTo(StagingBufferForLayer.Memory, DataOffset, TileLayerSize);
+					bDecodeResult = bDecodeResult && ensure(TileLayerSize <= StagingBufferForLayer.MemorySize);
+					if (bDecodeResult)
+					{
+						Params.Data->CopyTo(StagingBufferForLayer.Memory, DataOffset, TileLayerSize);
+					}
 				}
 				else
 				{
-					check(PackedStride <= StagingBufferForLayer.Stride);
-					check(TileLayerSize <= PackedOutputSize);
-					check(TileHeightInBlocks * StagingBufferForLayer.Stride <= StagingBufferForLayer.MemorySize);
-					TempBuffer.SetNumUninitialized(PackedOutputSize, false);
-					Params.Data->CopyTo(TempBuffer.GetData(), DataOffset, TileLayerSize);
-					for (uint32 y = 0; y < TileHeightInBlocks; ++y)
+					bDecodeResult = bDecodeResult && ensure(PackedStride <= StagingBufferForLayer.Stride);
+					bDecodeResult = bDecodeResult && ensure(TileLayerSize <= PackedOutputSize);
+					bDecodeResult = bDecodeResult && ensure(TileHeightInBlocks * StagingBufferForLayer.Stride <= StagingBufferForLayer.MemorySize);
+					if (bDecodeResult)
 					{
-						FMemory::Memcpy((uint8*)StagingBufferForLayer.Memory + y * StagingBufferForLayer.Stride, TempBuffer.GetData() + y * PackedStride, PackedStride);
+						TempBuffer.SetNumUninitialized(PackedOutputSize, false);
+						Params.Data->CopyTo(TempBuffer.GetData(), DataOffset, TileLayerSize);
+						for (uint32 y = 0; y < TileHeightInBlocks; ++y)
+						{
+							FMemory::Memcpy((uint8*)StagingBufferForLayer.Memory + y * StagingBufferForLayer.Stride, TempBuffer.GetData() + y * PackedStride, PackedStride);
+						}
 					}
 				}
 				break;
@@ -177,21 +183,24 @@ struct FTranscodeTask
 				if (StagingBufferForLayer.Stride == PackedStride)
 				{
 					// output buffer is tightly packed, can decompress directly
-					check(PackedOutputSize <= StagingBufferForLayer.MemorySize);
-					bDecodeResult = FCompression::UncompressMemoryStream(NAME_Zlib, StagingBufferForLayer.Memory, PackedOutputSize, Params.Data, DataOffset, TileLayerSize);
+					bDecodeResult = bDecodeResult && ensure(PackedOutputSize <= StagingBufferForLayer.MemorySize);
+					bDecodeResult = bDecodeResult && FCompression::UncompressMemoryStream(NAME_Zlib, StagingBufferForLayer.Memory, PackedOutputSize, Params.Data, DataOffset, TileLayerSize);
 				}
 				else
 				{
 					// output buffer has per-scanline padding, need to decompress to temp buffer, then copy line-by-line
-					check(PackedStride <= StagingBufferForLayer.Stride);
-					TempBuffer.SetNumUninitialized(PackedOutputSize, false);
-					bDecodeResult = FCompression::UncompressMemoryStream(NAME_Zlib, TempBuffer.GetData(), PackedOutputSize, Params.Data, DataOffset, TileLayerSize);
+					bDecodeResult = bDecodeResult && ensure(PackedStride <= StagingBufferForLayer.Stride);
 					if (bDecodeResult)
 					{
-						check(TileHeightInBlocks * StagingBufferForLayer.Stride <= StagingBufferForLayer.MemorySize);
-						for (uint32 y = 0; y < TileHeightInBlocks; ++y)
+						TempBuffer.SetNumUninitialized(PackedOutputSize, false);
+						bDecodeResult = FCompression::UncompressMemoryStream(NAME_Zlib, TempBuffer.GetData(), PackedOutputSize, Params.Data, DataOffset, TileLayerSize);
+						if (bDecodeResult)
 						{
-							FMemory::Memcpy((uint8*)StagingBufferForLayer.Memory + y * StagingBufferForLayer.Stride, TempBuffer.GetData() + y * PackedStride, PackedStride);
+							check(TileHeightInBlocks * StagingBufferForLayer.Stride <= StagingBufferForLayer.MemorySize);
+							for (uint32 y = 0; y < TileHeightInBlocks; ++y)
+							{
+								FMemory::Memcpy((uint8*)StagingBufferForLayer.Memory + y * StagingBufferForLayer.Stride, TempBuffer.GetData() + y * PackedStride, PackedStride);
+							}
 						}
 					}
 				}
