Fixes for mesh-space blending in blendspaces. Adds the ability to blend regular animations in mesh space too.

diff --git a/Engine/Source/Runtime/Engine/Classes/Animation/BlendSpace.h b/Engine/Source/Runtime/Engine/Classes/Animation/BlendSpace.h
--- a/Engine/Source/Runtime/Engine/Classes/Animation/BlendSpace.h
+++ b/Engine/Source/Runtime/Engine/Classes/Animation/BlendSpace.h
@@ -690,12 +690,14 @@ private:
 	bool UpdateBlendSamples_Internal(const FVector& InBlendSpacePosition, float InDeltaTime, TArray<FBlendSampleData>& InOutOldSampleDataList, TArray<FBlendSampleData>& InOutSampleDataCache, int32& InOutCachedTriangulationIndex) const;
 
 public:
-	/**
-	* When you use per bone sample smoothing, this makes blending happen in mesh space and can happen if this 
-	* contains additive animation samples. This is more performance intensive than blending in local space.
-	*/
+#if WITH_EDITORONLY_DATA
+	UE_DEPRECATED(5.1, "This property is deprecated. Please use/see bContainsRotationOffsetMeshSpaceSamples instead")
+	bool bRotationBlendInMeshSpace_DEPRECATED;
+#endif
+
+	/** Indicates whether any samples have the flag to apply rotation offsets in mesh space */
 	UPROPERTY()
-	bool bRotationBlendInMeshSpace;
+	bool bContainsRotationOffsetMeshSpaceSamples;
 
 	/** Input Smoothing parameters for each input axis */
 	UPROPERTY(EditAnywhere, Category = InputInterpolation)
@@ -738,6 +740,22 @@ public:
 	UPROPERTY(EditAnywhere, Category = SampleSmoothing, meta = (DisplayName = "Smoothing"))
 	bool bTargetWeightInterpolationEaseInOut = true;
 
+	/**
+	 * If set then blending is performed in mesh space if there are per-bone sample smoothing overrides.
+	 * 
+	 * Note that mesh space blending is significantly more expensive (slower) than normal blending when the 
+	 * samples are regular animations (i.e. not additive animations that are already set to apply in mesh 
+	 * space), and is typically only useful if you want some parts of the skeleton to achieve a pose 
+	 * in mesh space faster or slower than others - for example to make the head move faster than the 
+	 * body/arms when aiming, so the character looks at the target slightly before aiming at it.
+	 * 
+	 * Note also that blend space assets with additive/mesh space samples will always blend in mesh space, and 
+	 * also that enabling this option with blend space graphs producing additive/mesh space samples may cause
+	 * undesired results.
+	 */
+	UPROPERTY(EditAnywhere, Category = SampleSmoothing)
+	bool bAllowMeshSpaceBlending = false;
+
 	/** 
 	* The default looping behavior of this blend space.
 	* Asset players can override this
diff --git a/Engine/Source/Runtime/Engine/Private/Animation/AimOffsetBlendSpace.cpp b/Engine/Source/Runtime/Engine/Private/Animation/AimOffsetBlendSpace.cpp
--- a/Engine/Source/Runtime/Engine/Private/Animation/AimOffsetBlendSpace.cpp
+++ b/Engine/Source/Runtime/Engine/Private/Animation/AimOffsetBlendSpace.cpp
@@ -9,7 +9,6 @@
 UAimOffsetBlendSpace::UAimOffsetBlendSpace(const FObjectInitializer& ObjectInitializer)
 	: Super(ObjectInitializer)
 {
-	bRotationBlendInMeshSpace = true;
 }
 
 bool UAimOffsetBlendSpace::IsValidAdditiveType(EAdditiveAnimationType AdditiveType) const
diff --git a/Engine/Source/Runtime/Engine/Private/Animation/AimOffsetBlendSpace1D.cpp b/Engine/Source/Runtime/Engine/Private/Animation/AimOffsetBlendSpace1D.cpp
--- a/Engine/Source/Runtime/Engine/Private/Animation/AimOffsetBlendSpace1D.cpp
+++ b/Engine/Source/Runtime/Engine/Private/Animation/AimOffsetBlendSpace1D.cpp
@@ -9,7 +9,6 @@
 UAimOffsetBlendSpace1D::UAimOffsetBlendSpace1D(const FObjectInitializer& ObjectInitializer)
 	: Super(ObjectInitializer)
 {
-	bRotationBlendInMeshSpace = true;
 }
 
 bool UAimOffsetBlendSpace1D::IsValidAdditiveType(EAdditiveAnimationType AdditiveType) const
diff --git a/Engine/Source/Runtime/Engine/Private/Animation/AnimSingleNodeInstanceProxy.cpp b/Engine/Source/Runtime/Engine/Private/Animation/AnimSingleNodeInstanceProxy.cpp
--- a/Engine/Source/Runtime/Engine/Private/Animation/AnimSingleNodeInstanceProxy.cpp
+++ b/Engine/Source/Runtime/Engine/Private/Animation/AnimSingleNodeInstanceProxy.cpp
@@ -200,7 +200,8 @@ void FAnimSingleNodeInstanceProxy::InternalBlendSpaceEvaluatePose(class UBlendSp
 		FAnimationPoseData AdditiveAnimationPoseData = { AdditivePose, AdditiveCurve, AdditiveAttributes };
 		BlendSpace->GetAnimationPose(BlendSampleDataCache, ExtractionContext, AdditiveAnimationPoseData);
 
-		enum EAdditiveAnimationType AdditiveType = BlendSpace->bRotationBlendInMeshSpace? AAT_RotationOffsetMeshSpace : AAT_LocalSpaceBase;
+		enum EAdditiveAnimationType AdditiveType = 
+			BlendSpace->bContainsRotationOffsetMeshSpaceSamples ? AAT_RotationOffsetMeshSpace : AAT_LocalSpaceBase;
 		FAnimationRuntime::AccumulateAdditivePose(AnimationPoseData, AdditiveAnimationPoseData, 1.f, AdditiveType);
 	}
 	else
diff --git a/Engine/Source/Runtime/Engine/Private/Animation/AnimationRuntime.cpp b/Engine/Source/Runtime/Engine/Private/Animation/AnimationRuntime.cpp
--- a/Engine/Source/Runtime/Engine/Private/Animation/AnimationRuntime.cpp
+++ b/Engine/Source/Runtime/Engine/Private/Animation/AnimationRuntime.cpp
@@ -632,7 +632,6 @@ void FAnimationRuntime::BlendPosesTogetherPerBone(TArrayView<const FCompactPose>
 	}
 
 	BlendPosePerBone<ETransformBlendMode::Overwrite>(PerBoneIndices, BlendSampleDataCache[0], OutPose, SourcePoses[0]);
-
 	for (int32 i = 1; i < SourcePoses.Num(); ++i)
 	{
 		BlendPosePerBone<ETransformBlendMode::Accumulate>(PerBoneIndices, BlendSampleDataCache[i], OutPose, SourcePoses[i]);
@@ -720,27 +719,34 @@ void FAnimationRuntime::BlendPosesTogetherPerBone(TArrayView<const FCompactPose>
 }
 
 void FAnimationRuntime::BlendPosesTogetherPerBoneInMeshSpace(
-	const TArrayView<FCompactPose> SourcePoses,
-	const TArrayView<const FBlendedCurve> SourceCurves,
-	const UBlendSpace* BlendSpace,
+	const TArrayView<FCompactPose>           SourcePoses,
+	const TArrayView<const FBlendedCurve>    SourceCurves,
+	const UBlendSpace*                       BlendSpace,
 	const TArrayView<const FBlendSampleData> BlendSampleDataCache,
-	/*out*/ FCompactPose& ResultPose,
-	/*out*/ FBlendedCurve& ResultCurve)
+	/*out*/ FCompactPose&                    ResultPose,
+	/*out*/ FBlendedCurve&                   ResultCurve)
 {
 	UE::Anim::FStackAttributeContainer TempAttributes;
 	FAnimationPoseData AnimationPoseData = { ResultPose, ResultCurve, TempAttributes };
 
-	BlendPosesTogetherPerBoneInMeshSpace(SourcePoses, SourceCurves, {}, BlendSpace, BlendSampleDataCache, AnimationPoseData);
+	BlendPosesTogetherPerBoneInMeshSpace(
+		SourcePoses, SourceCurves, {}, BlendSpace, BlendSampleDataCache, AnimationPoseData);
 }
 
-void FAnimationRuntime::BlendPosesTogetherPerBoneInMeshSpace(TArrayView<FCompactPose> SourcePoses, TArrayView<const FBlendedCurve> SourceCurves, TArrayView<const UE::Anim::FStackAttributeContainer> SourceAttributes, const UBlendSpace* BlendSpace, TArrayView<const FBlendSampleData> BlendSampleDataCache, /*out*/ FAnimationPoseData& OutAnimationPoseData)
+void FAnimationRuntime::BlendPosesTogetherPerBoneInMeshSpace(
+	TArrayView<FCompactPose>                             SourcePoses, 
+	TArrayView<const FBlendedCurve>                      SourceCurves, 
+	TArrayView<const UE::Anim::FStackAttributeContainer> SourceAttributes, 
+	const UBlendSpace*                                   BlendSpace, 
+	TArrayView<const FBlendSampleData>                   BlendSampleDataCache, 
+	/*out*/ FAnimationPoseData&                          OutAnimationPoseData)
 {
-	FQuat NewRotation;
-	USkeleton* Skeleton = BlendSpace->GetSkeleton();
-
 	FCompactPose& OutPose = OutAnimationPoseData.GetPose();
 
-	// all this is going to do is to convert SourcePoses.Rotation to be mesh space, and then once it goes through BlendPosesTogetherPerBone, convert back to local
+	// Convert the sources poses into mesh space, and then once they have gone through
+	// BlendPosesTogetherPerBone, convert back to local space
+
+	// Convert SourcePoses.Rotation to be mesh space
 	for (FCompactPose& Pose : SourcePoses)
 	{
 		for (const FCompactPoseBoneIndex BoneIndex : Pose.ForEachBoneIndex())
@@ -748,23 +754,17 @@ void FAnimationRuntime::BlendPosesTogetherPerBoneInMeshSpace(TArrayView<FCompact
 			const FCompactPoseBoneIndex ParentIndex = Pose.GetParentBoneIndex(BoneIndex);
 			if (ParentIndex != INDEX_NONE)
 			{
-				NewRotation = Pose[ParentIndex].GetRotation() * Pose[BoneIndex].GetRotation();
+				FQuat NewRotation = Pose[ParentIndex].GetRotation() * Pose[BoneIndex].GetRotation();
 				NewRotation.Normalize();
+				Pose[BoneIndex].SetRotation(NewRotation);
 			}
-			else
-			{
-				NewRotation = Pose[BoneIndex].GetRotation();
-			}
-
-			// now copy back to SourcePoses
-			Pose[BoneIndex].SetRotation(NewRotation);
 		}
 	}
 
 	// now we have mesh space rotation, call BlendPosesTogetherPerBone
 	BlendPosesTogetherPerBone(SourcePoses, SourceCurves, SourceAttributes, BlendSpace, BlendSampleDataCache, OutAnimationPoseData);
 
-	// now result atoms has the output with mesh space rotation. Convert back to local space, start from back
+	// Now OutPose has the output with mesh space rotation. Convert back to local space, start from back
 	for (const FCompactPoseBoneIndex BoneIndex : OutPose.ForEachBoneIndexReverse())
 	{
 		const FCompactPoseBoneIndex ParentIndex = OutPose.GetParentBoneIndex(BoneIndex);
diff --git a/Engine/Source/Runtime/Engine/Private/Animation/BlendSpace.cpp b/Engine/Source/Runtime/Engine/Private/Animation/BlendSpace.cpp
--- a/Engine/Source/Runtime/Engine/Private/Animation/BlendSpace.cpp
+++ b/Engine/Source/Runtime/Engine/Private/Animation/BlendSpace.cpp
@@ -900,18 +900,32 @@ void UBlendSpace::GetAnimationPose_Internal(TArray<FBlendSampleData>& BlendSampl
 
 	if (PerBoneBlend.Num() > 0)
 	{
-		if (bRotationBlendInMeshSpace)
+		bool bValidAdditive = IsValidAdditive();
+
+		if (bAllowMeshSpaceBlending && !bContainsRotationOffsetMeshSpaceSamples)
 		{
-			FAnimationRuntime::BlendPosesTogetherPerBoneInMeshSpace(ChildrenPosesView, ChildrenCurves, ChildrenAttributes, this, BlendSampleDataCache, OutAnimationPoseData);
+			// Why blend in mesh space when there are per-bone smoothing settings? Because then we
+			// can blend between two aim poses (for example), with the hands moving faster towards
+			// the target than the spine. This results in nice organic looking movement, rather than
+			// everything moving at the same rate. However, note that if the samples contain mesh-space
+			// rotations then the regular blend will already happen in mesh space automatically.
+			FAnimationRuntime::BlendPosesTogetherPerBoneInMeshSpace(
+				ChildrenPosesView, ChildrenCurves, ChildrenAttributes,
+				this, BlendSampleDataCache, OutAnimationPoseData);
 		}
 		else
 		{
-			FAnimationRuntime::BlendPosesTogetherPerBone(ChildrenPosesView, ChildrenCurves, ChildrenAttributes, this, BlendSampleDataCache, OutAnimationPoseData);
+			FAnimationRuntime::BlendPosesTogetherPerBone(
+				ChildrenPosesView, ChildrenCurves, ChildrenAttributes,
+				this, BlendSampleDataCache, OutAnimationPoseData);
 		}
 	}
 	else
 	{
-		FAnimationRuntime::BlendPosesTogether(ChildrenPosesView, ChildrenCurves, ChildrenAttributes, ChildrenWeights, OutAnimationPoseData);
+		// We could allow mesh space blending here, when there are no per-bone smoothing settings. However, it's
+		// unlikely that it would actually provide a benefit, but is a lot more expensive.
+		FAnimationRuntime::BlendPosesTogether(
+			ChildrenPosesView, ChildrenCurves, ChildrenAttributes, ChildrenWeights, OutAnimationPoseData);
 	}
 
 	// Once all the accumulation and blending has been done, normalize rotations.
@@ -1221,15 +1235,7 @@ void UBlendSpace::ValidateSampleData()
 		}
 	}
 
-	// Set the space for blending
-	if (IsAsset())
-	{
-		bRotationBlendInMeshSpace = ContainsMatchingSamples(AAT_RotationOffsetMeshSpace);
-	}
-	else
-	{
-		bRotationBlendInMeshSpace = IsValidAdditive();
-	}
+	bContainsRotationOffsetMeshSpaceSamples = ContainsMatchingSamples(AAT_RotationOffsetMeshSpace);
 
 	SampleIndexWithMarkers = bAllMarkerPatternsMatch ? SampleWithMarkers : INDEX_NONE;
 
diff --git a/Engine/Source/Runtime/Engine/Public/AnimationRuntime.h b/Engine/Source/Runtime/Engine/Public/AnimationRuntime.h
--- a/Engine/Source/Runtime/Engine/Public/AnimationRuntime.h
+++ b/Engine/Source/Runtime/Engine/Public/AnimationRuntime.h
@@ -256,7 +256,8 @@ public:
 
 	/**
 	* Blends together a set of poses, each with a given weight.
-	* This function is for BlendSpace per bone blending. BlendSampleDataCache contains the weight information and is indexed using BlendSampleDataCacheIndices, to prevent caller having to supply an ordered array 
+	* This function is for BlendSpace per bone blending. BlendSampleDataCache contains the weight information 
+	* and is indexed using BlendSampleDataCacheIndices, to prevent caller having to supply an ordered array 
 	*
 	* This blends in local space
 	*
@@ -282,29 +283,42 @@ public:
 		/*out*/ FAnimationPoseData& OutAnimationPoseData);
 
 	/**
-	* Blends together a set of poses, each with a given weight.
-	* This function is for BlendSpace per bone blending. BlendSampleDataCache contains the weight information
+	* Blends together a set of local space (not mesh space) poses, each with a given weight.
+	* This function is for BlendSpace per bone blending. BlendSampleDataCache contains the
+	* weight information
 	*
-	* This blends rotation in mesh space, so performance intensive.
+	* This blends the rotations in mesh space, so is performance intensive as it requires a
+	* conversion from joint rotations into mesh space, then the blend, then a conversion back.
 	*
-	* @param	ResultPose		Output pose of relative bone transforms.
+	* @param OutResultPose Output blended pose of relative bone transforms.
+	* @param OutResultCurve Output blended curves
 	*/
 	UE_DEPRECATED(4.26, "Use BlendPosesTogetherPerBone with other signature")
 	static void BlendPosesTogetherPerBoneInMeshSpace(
-		TArrayView<FCompactPose> SourcePoses,
-		TArrayView<const FBlendedCurve> SourceCurves,
-		const UBlendSpace* BlendSpace,
+		TArrayView<FCompactPose>           SourcePoses,
+		TArrayView<const FBlendedCurve>    SourceCurves,
+		const UBlendSpace*                 BlendSpace,
 		TArrayView<const FBlendSampleData> BlendSampleDataCache,
-		/*out*/ FCompactPose& ResultPose,
-		/*out*/ FBlendedCurve& ResultCurve);
+		FCompactPose&                      OutResultPose,
+		FBlendedCurve&                     OutResultCurve);
 
+	/**
+	* Blends together a set of local space (not mesh space) poses, each with a given weight.
+	* This function is for BlendSpace per bone blending. BlendSampleDataCache contains the 
+	* weight information
+	*
+	* This blends the rotations in mesh space, so is performance intensive as it requires a 
+	* conversion from joint rotations into mesh space, then the blend, then a conversion back.
+	*
+	* @param OutAnimationPoseData Output pose, curves and attributes.
+	*/
 	static void BlendPosesTogetherPerBoneInMeshSpace(
-		TArrayView< FCompactPose> SourcePoses,
-		TArrayView<const FBlendedCurve> SourceCurves,
+		TArrayView<FCompactPose>                             SourcePoses,
+		TArrayView<const FBlendedCurve>                      SourceCurves,
 		TArrayView<const UE::Anim::FStackAttributeContainer> SourceAttributes,	
-		const UBlendSpace* BlendSpace,
-		TArrayView<const FBlendSampleData> BlendSampleDataCache,
-		/*out*/ FAnimationPoseData& OutAnimationPoseData);
+		const UBlendSpace*                                   BlendSpace,
+		TArrayView<const FBlendSampleData>                   BlendSampleDataCache,
+		FAnimationPoseData&                                  OutAnimationPoseData);
 
 	/** Blending flags for BlendPosesPerBoneFilter */
 	enum class EBlendPosesPerBoneFilterFlags : uint32
