Test if the world is valid before running deferred actor construction scripts, a level might have been made invisible or flagged for unload prior to the actor being allowed to run its scripts
When considering compilation for game specifically, force compile static meshes of unregistered component. checks for async compilation for game are usually done prior to component registration, this will also prevent unnecessary deferring of construction scripts

diff --git a/Engine/Source/Runtime/Engine/Private/ActorDeferredScriptManager.cpp b/Engine/Source/Runtime/Engine/Private/ActorDeferredScriptManager.cpp
--- a/Engine/Source/Runtime/Engine/Private/ActorDeferredScriptManager.cpp
+++ b/Engine/Source/Runtime/Engine/Private/ActorDeferredScriptManager.cpp
@@ -74,9 +74,9 @@ void FActorDeferredScriptManager::ProcessAsyncTasks(bool bLimitExecutionTime)
 		GIsEditorLoadingPackage = true;
 		while (!PendingConstructionScriptActors.IsEmpty() && bHasTimeLeft)
 		{
-			TWeakObjectPtr<AActor> WeakActor = PendingConstructionScriptActors.First();
-			PendingConstructionScriptActors.PopFirst();
-			if (AActor* Actor = WeakActor.Get())
+			TWeakObjectPtr<AActor> WeakActor = PendingConstructionScriptActors.Last();
+			PendingConstructionScriptActors.PopLast();
+			if (AActor* Actor = WeakActor.Get(); Actor && Actor->GetWorld())
 			{
 				Actor->RerunConstructionScripts();
 			}
diff --git a/Engine/Source/Runtime/Engine/Private/StaticMeshCompiler.cpp b/Engine/Source/Runtime/Engine/Private/StaticMeshCompiler.cpp
--- a/Engine/Source/Runtime/Engine/Private/StaticMeshCompiler.cpp
+++ b/Engine/Source/Runtime/Engine/Private/StaticMeshCompiler.cpp
@@ -482,8 +482,7 @@ void FStaticMeshCompilingManager::FinishCompilationsForGame()
 				{
 					for (const UStaticMeshComponent* Component : ObjectCacheScope.GetContext().GetStaticMeshComponents(StaticMesh))
 					{
-						if (Component->IsRegistered() &&
-							PIEWorlds.Contains(Component->GetWorld()) &&
+						if (PIEWorlds.Contains(Component->GetWorld()) &&
 							(PlayInEditorMode == 0 || Component->GetCollisionEnabled() != ECollisionEnabled::NoCollision || Component->IsNavigationRelevant() || Component->bAlwaysCreatePhysicsState || Component->CanCharacterStepUpOn != ECB_No))
 						{
 							if (PlayInEditorMode == 2)
