Ensure shaders are cached before drawing a tile with a material.

- In the editor with ODSC shader maps could be incomplete resulting in us drawing a tiling with the default material (which never gets updated)
- If we detect the material is incomplete, we will synchronously compile shaders so they are ready before drawing.
- Display a slow task if possible (not PIE)

diff --git a/Engine/Source/Runtime/Engine/Private/KismetRenderingLibrary.cpp b/Engine/Source/Runtime/Engine/Private/KismetRenderingLibrary.cpp
--- a/Engine/Source/Runtime/Engine/Private/KismetRenderingLibrary.cpp
+++ b/Engine/Source/Runtime/Engine/Private/KismetRenderingLibrary.cpp
@@ -3,6 +3,7 @@
 #include "Kismet/KismetRenderingLibrary.h"
 #include "HAL/FileManager.h"
 #include "Misc/Paths.h"
+#include "Misc/ScopedSlowTask.h"
 #include "Serialization/BufferArchive.h"
 #include "EngineGlobals.h"
 #include "RenderingThread.h"
@@ -149,6 +150,8 @@ void UKismetRenderingLibrary::ResizeRenderTarget2D(UTextureRenderTarget2D* Textu
 
 void UKismetRenderingLibrary::DrawMaterialToRenderTarget(UObject* WorldContextObject, UTextureRenderTarget2D* TextureRenderTarget, UMaterialInterface* Material)
 {
+	TRACE_CPUPROFILER_EVENT_SCOPE(DrawMaterialToRenderTarget);
+
 	if (!FApp::CanEverRender())
 	{
 		// Returning early to avoid warnings about missing resources that are expected when CanEverRender is false.
@@ -175,6 +178,13 @@ void UKismetRenderingLibrary::DrawMaterialToRenderTarget(UObject* WorldContextOb
 	}
 	else
 	{
+		if (!Material->IsComplete())
+		{
+			FScopedSlowTask SlowTask(0.0f, NSLOCTEXT("KismetRenderingLibrary", "CacheShaders", "Caching Shaders..."));
+			SlowTask.MakeDialog();
+			Material->CacheShaders(EMaterialShaderPrecompileMode::Synchronous);
+		}
+
 		World->FlushDeferredParameterCollectionInstanceUpdates();
 
 		FTextureRenderTargetResource* RenderTargetResource = TextureRenderTarget->GameThread_GetRenderTargetResource();
