Fix: BlendSpaces: Per-bone overrides for sample smoothing appear to have no effect on lower body

diff --git a/Engine/Source/Runtime/Engine/Private/Animation/BlendSpace.cpp b/Engine/Source/Runtime/Engine/Private/Animation/BlendSpace.cpp
--- a/Engine/Source/Runtime/Engine/Private/Animation/BlendSpace.cpp
+++ b/Engine/Source/Runtime/Engine/Private/Animation/BlendSpace.cpp
@@ -733,7 +733,10 @@ TSharedPtr<IInterpolationIndexProvider::FPerBoneInterpolationData> UBlendSpace::
 	return MakeShareable<FSortedPerBoneInterpolationData>(Data);
 }
 
-int32 UBlendSpace::GetPerBoneInterpolationIndex(const FCompactPoseBoneIndex& InCompactPoseBoneIndex, const FBoneContainer& RequiredBones, const IInterpolationIndexProvider::FPerBoneInterpolationData* Data) const
+int32 UBlendSpace::GetPerBoneInterpolationIndex(
+	const FCompactPoseBoneIndex&                                  InCompactPoseBoneIndex, 
+	const FBoneContainer&                                         RequiredBones, 
+	const IInterpolationIndexProvider::FPerBoneInterpolationData* Data) const
 {
 	if (!ensure(Data))
 	{
@@ -746,18 +749,17 @@ int32 UBlendSpace::GetPerBoneInterpolationIndex(const FCompactPoseBoneIndex& InC
 		const FPerBoneInterpolation& PerBoneInterpolation = SortedData[Iter].PerBoneBlend;
 		int32 OriginalIndex = SortedData[Iter].OriginalIndex;
 		const FBoneReference& SmoothedBone = PerBoneInterpolation.BoneReference;
-		if (SmoothedBone.IsValidToEvaluate(RequiredBones))
+		FCompactPoseBoneIndex SmoothedBoneCompactPoseIndex = 
+			RequiredBones.GetCompactPoseIndexFromSkeletonPoseIndex(SmoothedBone.GetSkeletonPoseIndex(RequiredBones));
+		if (SmoothedBoneCompactPoseIndex == InCompactPoseBoneIndex)
 		{
-			FCompactPoseBoneIndex SmoothedBonePoseIndex = RequiredBones.GetCompactPoseIndexFromSkeletonPoseIndex(SmoothedBone.GetSkeletonPoseIndex(RequiredBones));
-			if (SmoothedBonePoseIndex == InCompactPoseBoneIndex)
-			{
-				return OriginalIndex;
-			}
-			// Returns true if BoneIndex is a child of SmoothedBonePoseIndex. Searches up from BoneIndex
-			if (RequiredBones.BoneIsChildOf(InCompactPoseBoneIndex, SmoothedBonePoseIndex))
-			{
-				return OriginalIndex;
-			}
+			return OriginalIndex;
+		}
+		// BoneIsChildOf returns true if InCompactPoseBoneIndex is a child of SmoothedBoneCompactPoseIndex. 
+		if (SmoothedBoneCompactPoseIndex != INDEX_NONE && 
+			RequiredBones.BoneIsChildOf(InCompactPoseBoneIndex, SmoothedBoneCompactPoseIndex))
+		{
+			return OriginalIndex;
 		}
 	}
 	return INDEX_NONE;
