Added previous position and play rate of the current root motion montage to FSaveMove_Character

diff --git a/Engine/Source/Runtime/Engine/Classes/GameFramework/CharacterMovementComponent.h b/Engine/Source/Runtime/Engine/Classes/GameFramework/CharacterMovementComponent.h
--- a/Engine/Source/Runtime/Engine/Classes/GameFramework/CharacterMovementComponent.h
+++ b/Engine/Source/Runtime/Engine/Classes/GameFramework/CharacterMovementComponent.h
@@ -2849,6 +2849,8 @@ public:
 
 	TWeakObjectPtr<class UAnimMontage> RootMotionMontage;
 	float RootMotionTrackPosition;
+	float RootMotionPreviousTrackPosition;
+	float RootMotionPlayRateWithScale;
 	FRootMotionMovementParams RootMotionMovement;
 
 	FRootMotionSourceGroup SavedRootMotion;
diff --git a/Engine/Source/Runtime/Engine/Private/Components/CharacterMovementComponent.cpp b/Engine/Source/Runtime/Engine/Private/Components/CharacterMovementComponent.cpp
--- a/Engine/Source/Runtime/Engine/Private/Components/CharacterMovementComponent.cpp
+++ b/Engine/Source/Runtime/Engine/Private/Components/CharacterMovementComponent.cpp
@@ -11574,6 +11574,8 @@ void FSavedMove_Character::Clear()
 
 	RootMotionMontage = NULL;
 	RootMotionTrackPosition = 0.f;
+	RootMotionPreviousTrackPosition = 0.f;
+	RootMotionPlayRateWithScale = 1.f;
 	RootMotionMovement.Clear();
 
 	SavedRootMotion.Clear();
@@ -11730,6 +11732,8 @@ void FSavedMove_Character::PostUpdate(ACharacter* Character, FSavedMove_Characte
 			{
 				RootMotionMontage = RootMotionMontageInstance->Montage;
 				RootMotionTrackPosition = RootMotionMontageInstance->GetPosition();
+				RootMotionPreviousTrackPosition = RootMotionMontageInstance->GetPreviousPosition();
+				RootMotionPlayRateWithScale = RootMotionMontageInstance->GetPlayRate() * RootMotionMontage->RateScale;
 				RootMotionMovement = Character->ClientRootMotionParams;
 			}
 
@@ -12119,6 +12123,8 @@ void FSavedMove_Character::PrepMoveFor(ACharacter* Character)
 			{
 				RootMotionMovement.Clear();
 				RootMotionTrackPosition = RootMotionMontageInstance->GetPosition();
+				RootMotionPreviousTrackPosition = RootMotionTrackPosition;
+				RootMotionPlayRateWithScale = RootMotionMontageInstance->GetPlayRate() * RootMotionMontage->RateScale;
 				RootMotionMontageInstance->SimulateAdvance(DeltaTime, RootMotionTrackPosition, RootMotionMovement);
 				RootMotionMontageInstance->SetPosition(RootMotionTrackPosition);
 				RootMotionMovement.ScaleRootMotionTranslation(Character->GetAnimRootMotionTranslationScale());
