Tidy up mesh deformer abstract interface.
#rnx

diff --git a/Engine/Source/Runtime/Engine/Classes/Animation/MeshDeformer.h b/Engine/Source/Runtime/Engine/Classes/Animation/MeshDeformer.h
--- a/Engine/Source/Runtime/Engine/Classes/Animation/MeshDeformer.h
+++ b/Engine/Source/Runtime/Engine/Classes/Animation/MeshDeformer.h
@@ -23,10 +23,7 @@ public:
 	/** Create persistent settings for the mesh deformer instance */ 
 	virtual UMeshDeformerInstanceSettings* CreateSettingsInstance(
 		UMeshComponent* InMeshComponent
-		)
-	{
-		return nullptr;
-	}
+		) PURE_VIRTUAL(, return nullptr;);
 	
 	/** Create an instance of the mesh deformer. */
 	virtual UMeshDeformerInstance* CreateInstance(
