Iris
- Fixed issue where NetConnection->ViewTarget was not properly updated prior to replicating data

diff --git a/Engine/Source/Runtime/Engine/Private/NetDriver.cpp b/Engine/Source/Runtime/Engine/Private/NetDriver.cpp
--- a/Engine/Source/Runtime/Engine/Private/NetDriver.cpp
+++ b/Engine/Source/Runtime/Engine/Private/NetDriver.cpp
@@ -6026,13 +6026,14 @@ void UNetDriver::UpdateReplicationViews() const
 	}
 
 	TObjectPtr<UNetConnection> const* Connections = (ServerConnection != nullptr ? &ServerConnection : ClientConnections.GetData());
+	const bool bUpdateConnectionViewTarget = ServerConnection != nullptr;
 	const uint32 ConnectionCount = (ServerConnection != nullptr ? 1U : uint32(ClientConnections.Num()));
-	TArray<const UNetConnection*, TInlineAllocator<4>> SubConnections;
+	TArray<UNetConnection*, TInlineAllocator<4>> SubConnections;
 	for (uint32 ConnIt = 0, ConnEndIt = ConnectionCount; ConnIt != ConnEndIt; ++ConnIt)
 	{
 		ReplicationView.Views.Reset();
 
-		const UNetConnection* Conn = Connections[ConnIt];
+		UNetConnection* Conn = Connections[ConnIt];
 		SubConnections.Reset();
 		SubConnections.Add(Conn);
 		if (Conn->Children.Num() > 0)
@@ -6040,8 +6041,20 @@ void UNetDriver::UpdateReplicationViews() const
 			SubConnections.Insert(reinterpret_cast<UNetConnection*const*>(Conn->Children.GetData()), Conn->Children.Num(), 1);
 		}
 
-		for (const UNetConnection* SubConn : SubConnections)
+		// As we no longer call ServerReplicateActors we must make sure to update the ViewTarget
+		if (bUpdateConnectionViewTarget)
 		{
+			Conn->ViewTarget = Conn->PlayerController ? Conn->PlayerController->GetViewTarget() : ToRawPtr(Conn->OwningActor);
+		}
+
+		for (UNetConnection* SubConn : SubConnections)
+		{
+			// See comment above
+			if (bUpdateConnectionViewTarget)
+			{
+				SubConn->ViewTarget = SubConn->PlayerController ? SubConn->PlayerController->GetViewTarget() : ToRawPtr(SubConn->OwningActor);
+			}
+
 			const AActor* ViewTarget = Conn->ViewTarget;
 			const APlayerController* ViewingController = Conn->PlayerController;
 			if (ViewTarget == nullptr && ViewingController == nullptr)
