Scene capture:  fix for crash bug in scene renderer when interactively modifying render target size as a slider, due to size of the resource mismatching with the size stored in the object.  We need to always update the resource to match the object, for render targets, even for interactive edits.

#rnx

diff --git a/Engine/Source/Runtime/Engine/Private/Texture.cpp b/Engine/Source/Runtime/Engine/Private/Texture.cpp
--- a/Engine/Source/Runtime/Engine/Private/Texture.cpp
+++ b/Engine/Source/Runtime/Engine/Private/Texture.cpp
@@ -667,9 +667,12 @@ void UTexture::PostEditChangeProperty(FPropertyChangedEvent& PropertyChangedEven
 		}
 	}
 
-	// Don't update the texture resource if we've turned "DeferCompression" on, as this 
-	// would cause it to immediately update as an uncompressed texture
-	if( !DeferCompressionWasEnabled && (PropertyChangedEvent.ChangeType & EPropertyChangeType::Interactive) == 0 )
+	// Don't update the texture resource if we've turned "DeferCompression" on, as this would cause
+	// it to immediately update as an uncompressed texture.  If it's a render target, we always need
+	// to update the resource, to avoid an assert when rendering to it, due to a mismatch between the
+	// render target and scene render.
+	if( !DeferCompressionWasEnabled &&
+		((PropertyChangedEvent.ChangeType & EPropertyChangeType::Interactive) == 0 || GetTextureClass() == ETextureClass::RenderTarget))
 	{
 		// Update the texture resource. This will recache derived data if necessary
 		// which may involve recompressing the texture.
diff --git a/Engine/Source/Runtime/Engine/Private/TextureRenderTarget2D.cpp b/Engine/Source/Runtime/Engine/Private/TextureRenderTarget2D.cpp
--- a/Engine/Source/Runtime/Engine/Private/TextureRenderTarget2D.cpp
+++ b/Engine/Source/Runtime/Engine/Private/TextureRenderTarget2D.cpp
@@ -14,6 +14,9 @@
 #include "UObject/RenderingObjectVersion.h"
 #include "GenerateMips.h"
 #include "RenderGraphUtils.h"
+#if WITH_EDITOR
+#include "Components/SceneCaptureComponent2D.h"
+#endif
 
 int32 GTextureRenderTarget2DMaxSizeX = 999999999;
 int32 GTextureRenderTarget2DMaxSizeY = 999999999;
@@ -209,6 +212,23 @@ void UTextureRenderTarget2D::PostEditChangeProperty(FPropertyChangedEvent& Prope
 
     // SRGB may have been changed by Super, reset it since we prefer to honor explicit user choice
 	SRGB = IsSRGB();
+
+	// If this is an interactive size edit, notify any scene capture component that points to this render target that it needs to refresh.
+	// During interactive edits, time is paused, so the Tick function which normally handles capturing isn't called.
+	static const FName SizeXName = GET_MEMBER_NAME_CHECKED(UTextureRenderTarget2D, SizeX);
+	static const FName SizeYName = GET_MEMBER_NAME_CHECKED(UTextureRenderTarget2D, SizeY);
+
+	if ((PropertyChangedEvent.ChangeType & EPropertyChangeType::Interactive) && (PropertyChangedEvent.GetPropertyName() == SizeXName || PropertyChangedEvent.GetPropertyName() == SizeYName))
+	{
+		for (TObjectIterator<USceneCaptureComponent2D> It; It; ++It)
+		{
+			USceneCaptureComponent2D* SceneCaptureComponent = *It;
+			if (SceneCaptureComponent->TextureTarget == this)
+			{
+				SceneCaptureComponent->CaptureSceneDeferred();
+			}
+		}
+	}
 }
 #endif // WITH_EDITOR
 
diff --git a/Engine/Source/Runtime/Engine/Private/TextureRenderTargetCube.cpp b/Engine/Source/Runtime/Engine/Private/TextureRenderTargetCube.cpp
--- a/Engine/Source/Runtime/Engine/Private/TextureRenderTargetCube.cpp
+++ b/Engine/Source/Runtime/Engine/Private/TextureRenderTargetCube.cpp
@@ -12,6 +12,9 @@
 #include "DeviceProfiles/DeviceProfileManager.h"
 #include "Engine/TextureCube.h"
 #include "ClearQuad.h"
+#if WITH_EDITOR
+#include "Components/SceneCaptureComponentCube.h"
+#endif
 
 /*-----------------------------------------------------------------------------
 	UTextureRenderTargetCube
@@ -108,6 +111,22 @@ void UTextureRenderTargetCube::PostEditChangeProperty(FPropertyChangedEvent& Pro
 	SizeX = FMath::Clamp<int32>(SizeX - (SizeX % GPixelFormats[Format].BlockSizeX),1,MaxSize);
 
 	Super::PostEditChangeProperty(PropertyChangedEvent);
+
+	// If this is an interactive size edit, notify any scene capture component that points to this render target that it needs to refresh.
+	// During interactive edits, time is paused, so the Tick function which normally handles capturing isn't called.
+	static const FName SizeXName = GET_MEMBER_NAME_CHECKED(UTextureRenderTargetCube, SizeX);
+
+	if ((PropertyChangedEvent.ChangeType & EPropertyChangeType::Interactive) && PropertyChangedEvent.GetPropertyName() == SizeXName)
+	{
+		for (TObjectIterator<USceneCaptureComponentCube> It; It; ++It)
+		{
+			USceneCaptureComponentCube* SceneCaptureComponent = *It;
+			if (SceneCaptureComponent->TextureTarget == this)
+			{
+				SceneCaptureComponent->CaptureSceneDeferred();
+			}
+		}
+	}
 }
 #endif // WITH_EDITOR
 
