Remove unused bHasMotionBlurVelocityMeshes
#rnx

diff --git a/Engine/Source/Runtime/Engine/Classes/Components/PrimitiveComponent.h b/Engine/Source/Runtime/Engine/Classes/Components/PrimitiveComponent.h
--- a/Engine/Source/Runtime/Engine/Classes/Components/PrimitiveComponent.h
+++ b/Engine/Source/Runtime/Engine/Classes/Components/PrimitiveComponent.h
@@ -363,10 +363,6 @@ public:
 	UPROPERTY(EditAnywhere, AdvancedDisplay, BlueprintReadOnly, Category=LOD)
 	uint8 bAllowCullDistanceVolume:1;
 
-	/** True if the primitive has motion blur velocity meshes */
-	UPROPERTY()
-	uint8 bHasMotionBlurVelocityMeshes:1;
-	
 	/** If true, this component will be visible in reflection captures. */
 	UPROPERTY(EditAnywhere, AdvancedDisplay, BlueprintReadOnly, Category = Rendering)
 	uint8 bVisibleInReflectionCaptures:1;
diff --git a/Engine/Source/Runtime/Engine/Private/PrimitiveSceneProxy.cpp b/Engine/Source/Runtime/Engine/Private/PrimitiveSceneProxy.cpp
--- a/Engine/Source/Runtime/Engine/Private/PrimitiveSceneProxy.cpp
+++ b/Engine/Source/Runtime/Engine/Private/PrimitiveSceneProxy.cpp
@@ -181,7 +181,6 @@ FPrimitiveSceneProxy::FPrimitiveSceneProxy(const UPrimitiveComponent* InComponen
 ,	bLevelInstanceEditingState(InComponent->GetLevelInstanceEditingState())
 ,	bHovered(false)
 ,	bUseViewOwnerDepthPriorityGroup(InComponent->bUseViewOwnerDepthPriorityGroup)
-,	bHasMotionBlurVelocityMeshes(InComponent->bHasMotionBlurVelocityMeshes)
 ,	StaticDepthPriorityGroup((uint8)InComponent->GetStaticDepthPriorityGroup())
 ,	ViewOwnerDepthPriorityGroup(InComponent->ViewOwnerDepthPriorityGroup)
 ,	bStaticLighting(InComponent->HasStaticLighting())
diff --git a/Engine/Source/Runtime/Engine/Public/PrimitiveSceneProxy.h b/Engine/Source/Runtime/Engine/Public/PrimitiveSceneProxy.h
--- a/Engine/Source/Runtime/Engine/Public/PrimitiveSceneProxy.h
+++ b/Engine/Source/Runtime/Engine/Public/PrimitiveSceneProxy.h
@@ -564,7 +564,6 @@ public:
 	inline int32 GetVisibilityId() const { return VisibilityId; }
 	inline int16 GetTranslucencySortPriority() const { return TranslucencySortPriority; }
 	inline float GetTranslucencySortDistanceOffset() const { return TranslucencySortDistanceOffset; }
-	inline bool HasMotionBlurVelocityMeshes() const { return bHasMotionBlurVelocityMeshes; }
 
 	inline int32 GetVirtualTextureLodBias() const { return VirtualTextureLodBias; }
 	inline int32 GetVirtualTextureCullMips() const { return VirtualTextureCullMips; }
@@ -1086,9 +1085,6 @@ private:
 	/** true if ViewOwnerDepthPriorityGroup should be used. */
 	uint8 bUseViewOwnerDepthPriorityGroup : 1;
 
-	/** true if the primitive has motion blur velocity meshes */
-	uint8 bHasMotionBlurVelocityMeshes : 1;
-
 	/** DPG this prim belongs to. */
 	uint8 StaticDepthPriorityGroup : SDPG_NumBits;
 
