Non-unity build fix (added include for TObjectIterator)

#rnx

diff --git a/Engine/Source/Runtime/Engine/Private/TextureRenderTarget2D.cpp b/Engine/Source/Runtime/Engine/Private/TextureRenderTarget2D.cpp
--- a/Engine/Source/Runtime/Engine/Private/TextureRenderTarget2D.cpp
+++ b/Engine/Source/Runtime/Engine/Private/TextureRenderTarget2D.cpp
@@ -16,6 +16,7 @@
 #include "RenderGraphUtils.h"
 #if WITH_EDITOR
 #include "Components/SceneCaptureComponent2D.h"
+#include "UObject/UObjectIterator.h"
 #endif
 
 int32 GTextureRenderTarget2DMaxSizeX = 999999999;
diff --git a/Engine/Source/Runtime/Engine/Private/TextureRenderTargetCube.cpp b/Engine/Source/Runtime/Engine/Private/TextureRenderTargetCube.cpp
--- a/Engine/Source/Runtime/Engine/Private/TextureRenderTargetCube.cpp
+++ b/Engine/Source/Runtime/Engine/Private/TextureRenderTargetCube.cpp
@@ -14,6 +14,7 @@
 #include "ClearQuad.h"
 #if WITH_EDITOR
 #include "Components/SceneCaptureComponentCube.h"
+#include "UObject/UObjectIterator.h"
 #endif
 
 /*-----------------------------------------------------------------------------
