Unshelved from pending changelist '21318030':

PR #9472: Expose APlayerState::IsSpectator to blueprints (Contributed by [REDACTED_PERSON])

diff --git a/Engine/Source/Runtime/Engine/Classes/GameFramework/PlayerState.h b/Engine/Source/Runtime/Engine/Classes/GameFramework/PlayerState.h
--- a/Engine/Source/Runtime/Engine/Classes/GameFramework/PlayerState.h
+++ b/Engine/Source/Runtime/Engine/Classes/GameFramework/PlayerState.h
@@ -74,17 +74,17 @@ class ENGINE_API APlayerState : public AInfo
 
 	/** Player's current score. */
 	UE_DEPRECATED(4.25, "This member will be made private. Use GetScore or SetScore instead.")
-	UPROPERTY(ReplicatedUsing=OnRep_Score, BlueprintReadOnly, Category=PlayerState)
+	UPROPERTY(ReplicatedUsing=OnRep_Score, Category=PlayerState, BlueprintGetter=GetScore)
 	float Score;
 
 	/** Unique net id number. Actual value varies based on current online subsystem, use it only as a guaranteed unique number per player. */
 	UE_DEPRECATED(4.25, "This member will be made private. Use GetPlayerId or SetPlayerId instead.")
-	UPROPERTY(ReplicatedUsing=OnRep_PlayerId, BlueprintReadOnly, Category=PlayerState)
+	UPROPERTY(ReplicatedUsing=OnRep_PlayerId, Category=PlayerState, BlueprintGetter=GetPlayerId)
 	int32 PlayerId;
 
 private:
 	/** Replicated compressed ping for this player (holds ping in msec divided by 4) */
-	UPROPERTY(Replicated, BlueprintReadOnly, Category=PlayerState, meta=(AllowPrivateAccess))
+	UPROPERTY(Replicated, Category=PlayerState, BlueprintGetter=GetCompressedPing, meta=(AllowPrivateAccess))
 	uint8 CompressedPing;
 
 	/** The current PingBucket index that is being filled */
@@ -102,7 +102,7 @@ private:
 public:
 	/** Whether this player is currently a spectator */
 	UE_DEPRECATED(4.25, "This member will be made private. Use IsSpectator or SetIsSpectator instead.")
-	UPROPERTY(Replicated, BlueprintReadOnly, Category=PlayerState)
+	UPROPERTY(Replicated, Category=PlayerState, BlueprintGetter=IsSpectator)
 	uint8 bIsSpectator:1;
 
 	/** Whether this player can only ever be a spectator */
@@ -112,7 +112,7 @@ public:
 
 	/** True if this PlayerState is associated with an AIController */
 	UE_DEPRECATED(4.25, "This member will be made private. Use IsABot or SetIsABot instead.")
-	UPROPERTY(Replicated, BlueprintReadOnly, Category=PlayerState)
+	UPROPERTY(Replicated, Category=PlayerState, BlueprintGetter=IsABot)
 	uint8 bIsABot:1;
 
 	/** client side flag - whether this player has been welcomed or not (player entered message) */
@@ -352,6 +352,7 @@ PRAGMA_DISABLE_DEPRECATION_WARNINGS
 public:
 
 	/** Gets the literal value of Score. */
+	UFUNCTION(BlueprintGetter)
 	float GetScore() const
 	{
 		return Score;
@@ -361,6 +362,7 @@ public:
 	void SetScore(const float NewScore);
 
 	/** Gets the literal value of PlayerId. */
+	UFUNCTION(BlueprintGetter)
 	int32 GetPlayerId() const
 	{
 		return PlayerId;
@@ -370,6 +372,7 @@ public:
 	void SetPlayerId(const int32 NewId);
 
 	/** Gets the literal value of the compressed Ping value (Ping = PingInMS / 4). */
+	UFUNCTION(BlueprintGetter)
 	uint8 GetCompressedPing() const
 	{
 		return CompressedPing;
@@ -395,6 +398,7 @@ public:
 	float GetPingInMilliseconds() const;
 
 	/** Gets the literal value of bIsSpectator. */
+	UFUNCTION(BlueprintGetter)
 	bool IsSpectator() const
 	{
 		return bIsSpectator;
@@ -414,6 +418,7 @@ public:
 	void SetIsOnlyASpectator(const bool bNewSpectator);
 
 	/** Gets the literal value of bIsABot. */
+	UFUNCTION(BlueprintGetter)
 	bool IsABot() const
 	{
 		return bIsABot;
@@ -497,4 +502,4 @@ private:
 		PlayerState->SetPawnPrivate(Pawn);
 		PlayerState->OnPawnSet.Broadcast(PlayerState, PlayerState->PawnPrivate, OldPawn); 
 	}
-};
\ No newline at end of file
+};
