Fix memory stomp when copying replicated arrays with "net.WithArrayOnRepFix" enabled.

diff --git a/Engine/Source/Runtime/Engine/Private/RepLayout.cpp b/Engine/Source/Runtime/Engine/Private/RepLayout.cpp
--- a/Engine/Source/Runtime/Engine/Private/RepLayout.cpp
+++ b/Engine/Source/Runtime/Engine/Private/RepLayout.cpp
@@ -3442,7 +3442,7 @@ static FGuidReferencesMap* PrepReceivedArray(
 		// in game logic.
 		if (GbWithArrayOnRepFix && ShadowArray->Num() != DataArray->Num())
 		{
-			Parent.Property->CopyCompleteValue((uint8*)ShadowArray, (uint8*)DataArray);
+			Cmd.Property->CopyCompleteValue((uint8*)ShadowArray, (uint8*)DataArray);
 		}
 
 		*OutShadowBaseData = ShadowArray->GetData();
