Removed some log spam from attempting to load a class from an empty string

#rnx

diff --git a/Engine/Plugins/Developer/Concert/ConcertSync/ConcertSyncClient/Source/ConcertSyncClient/Private/ConcertClientPresenceManager.cpp b/Engine/Plugins/Developer/Concert/ConcertSync/ConcertSyncClient/Source/ConcertSyncClient/Private/ConcertClientPresenceManager.cpp
--- a/Engine/Plugins/Developer/Concert/ConcertSync/ConcertSyncClient/Source/ConcertSyncClient/Private/ConcertClientPresenceManager.cpp
+++ b/Engine/Plugins/Developer/Concert/ConcertSync/ConcertSyncClient/Source/ConcertSyncClient/Private/ConcertClientPresenceManager.cpp
@@ -428,7 +428,7 @@ AConcertClientPresenceActor* FConcertClientPresenceManager::SpawnPresenceActor(c
 	if (!VRDevice.IsNone())
 	{
 		UClass*& AvatarClass = OthersAvatarClasses.FindOrAdd(InClientInfo.VRAvatarActorClass, nullptr);
-		if (!AvatarClass)
+		if (!AvatarClass && !InClientInfo.VRAvatarActorClass.IsEmpty())
 		{
 			AvatarClass = LoadObject<UClass>(nullptr, *InClientInfo.VRAvatarActorClass);
 		}
@@ -437,7 +437,7 @@ AConcertClientPresenceActor* FConcertClientPresenceManager::SpawnPresenceActor(c
 	else
 	{
 		UClass*& AvatarClass = OthersAvatarClasses.FindOrAdd(InClientInfo.DesktopAvatarActorClass, nullptr);
-		if (!AvatarClass)
+		if (!AvatarClass && !InClientInfo.DesktopAvatarActorClass.IsEmpty())
 		{
 			AvatarClass = LoadObject<UClass>(nullptr, *InClientInfo.DesktopAvatarActorClass);
 		}
diff --git a/Engine/Source/Runtime/Engine/Private/ComponentInstanceDataCache.cpp b/Engine/Source/Runtime/Engine/Private/ComponentInstanceDataCache.cpp
--- a/Engine/Source/Runtime/Engine/Private/ComponentInstanceDataCache.cpp
+++ b/Engine/Source/Runtime/Engine/Private/ComponentInstanceDataCache.cpp
@@ -364,19 +364,22 @@ bool FDataCacheDuplicatedObjectData::Serialize(FArchive& Ar)
 	if (Ar.IsLoading())
 	{
 		DuplicatedObject = nullptr;
-		// Get the object class
-		if (UClass* ObjectClass = LoadObject<UClass>(nullptr, *ObjectClassPath))
+		if (!ObjectClassPath.IsEmpty())
 		{
-			// Get the object outer
-			if (UObject* FoundOuter = StaticFindObject(UObject::StaticClass(), nullptr, *ObjectOuterPath))
+			// Get the object class
+			if (UClass* ObjectClass = LoadObject<UClass>(nullptr, *ObjectClassPath))
 			{
-				// Create the duplicated object
-				DuplicatedObject = NewObject<UObject>(FoundOuter, ObjectClass, *ObjectName.ToString(), (EObjectFlags)ObjectPersistentFlags);
+				// Get the object outer
+				if (UObject* FoundOuter = StaticFindObject(UObject::StaticClass(), nullptr, *ObjectOuterPath))
+				{
+					// Create the duplicated object
+					DuplicatedObject = NewObject<UObject>(FoundOuter, ObjectClass, *ObjectName.ToString(), (EObjectFlags)ObjectPersistentFlags);
 
-				// Deserialize the duplicated object
-				FMemoryReader MemAr(ObjectData);
-				FObjectAndNameAsStringProxyArchive Reader(MemAr, false);
-				ObjectClass->SerializeTaggedProperties(Reader, (uint8*)DuplicatedObject, ObjectClass, nullptr);
+					// Deserialize the duplicated object
+					FMemoryReader MemAr(ObjectData);
+					FObjectAndNameAsStringProxyArchive Reader(MemAr, false);
+					ObjectClass->SerializeTaggedProperties(Reader, (uint8*)DuplicatedObject, ObjectClass, nullptr);
+				}
 			}
 		}
 	}
