World Partition Validations: implemented new error checks for level instances (invalid world asset, world asset not using external actors, incompatible world asset, etc).

#rnx

diff --git a/Engine/Plugins/Editor/DataValidation/Source/DataValidation/Private/WorldPartitionChangelistValidator.cpp b/Engine/Plugins/Editor/DataValidation/Source/DataValidation/Private/WorldPartitionChangelistValidator.cpp
--- a/Engine/Plugins/Editor/DataValidation/Source/DataValidation/Private/WorldPartitionChangelistValidator.cpp
+++ b/Engine/Plugins/Editor/DataValidation/Source/DataValidation/Private/WorldPartitionChangelistValidator.cpp
@@ -362,4 +362,19 @@ void UWorldPartitionChangelistValidator::OnActorNeedsResave(const FWorldPartitio
 	// Changelist validation already ensures that dirty actors must be part of the changelist
 }
 
+void UWorldPartitionChangelistValidator::OnLevelInstanceInvalidWorldAsset(const FWorldPartitionActorDescView& ActorDescView, FName WorldAsset, ELevelInstanceInvalidReason Reason)
+{
+	if (Filter(ActorDescView))
+	{
+		if (Reason == ELevelInstanceInvalidReason::WorldAssetNotFound)
+		{
+			FText CurrentError = FText::Format(LOCTEXT("DataValidation.Changelist.WorldPartition.LevelInstanceInvalidWorldAsset", "Level instance {0} has an invalid world asset {1}."),
+				FText::FromString(GetPrettyPackageName(ActorDescView)), 
+				FText::FromName(WorldAsset));
+
+			Errors->Add(CurrentError);
+		}
+	}
+}
+
 #undef LOCTEXT_NAMESPACE
diff --git a/Engine/Plugins/Editor/DataValidation/Source/DataValidation/Private/WorldPartitionChangelistValidator.h b/Engine/Plugins/Editor/DataValidation/Source/DataValidation/Private/WorldPartitionChangelistValidator.h
--- a/Engine/Plugins/Editor/DataValidation/Source/DataValidation/Private/WorldPartitionChangelistValidator.h
+++ b/Engine/Plugins/Editor/DataValidation/Source/DataValidation/Private/WorldPartitionChangelistValidator.h
@@ -45,5 +45,6 @@ protected:
 	virtual void OnDataLayerHierarchyTypeMismatch(const UDataLayerInstance* DataLayerInstance, const UDataLayerInstance* Parent) override;
 	virtual void OnDataLayerAssetConflict(const UDataLayerInstanceWithAsset* DataLayerInstance, const UDataLayerInstanceWithAsset* ConflictingDataLayerInstance) override;
 	virtual void OnActorNeedsResave(const FWorldPartitionActorDescView& ActorDescView) override;
+	virtual void OnLevelInstanceInvalidWorldAsset(const FWorldPartitionActorDescView& ActorDescView, FName WorldAsset, ELevelInstanceInvalidReason Reason) override;
 };
 
diff --git a/Engine/Source/Runtime/Engine/Private/WorldPartition/ErrorHandling/WorldPartitionStreamingGenerationLogErrorHandler.cpp b/Engine/Source/Runtime/Engine/Private/WorldPartition/ErrorHandling/WorldPartitionStreamingGenerationLogErrorHandler.cpp
--- a/Engine/Source/Runtime/Engine/Private/WorldPartition/ErrorHandling/WorldPartitionStreamingGenerationLogErrorHandler.cpp
+++ b/Engine/Source/Runtime/Engine/Private/WorldPartition/ErrorHandling/WorldPartitionStreamingGenerationLogErrorHandler.cpp
@@ -58,4 +58,23 @@ void FStreamingGenerationLogErrorHandler::OnActorNeedsResave(const FWorldPartiti
 {
 	UE_LOG(LogWorldPartition, Log, TEXT("Actor %s needs to be resaved"), *ActorDescView.GetActorLabelOrName().ToString());
 }
+
+void FStreamingGenerationLogErrorHandler::OnLevelInstanceInvalidWorldAsset(const FWorldPartitionActorDescView& ActorDescView, FName WorldAsset, ELevelInstanceInvalidReason Reason)
+{
+	const FSoftObjectPath ActorPath(ActorDescView.GetActorPath());
+	const FString ActorName = ActorPath.GetLongPackageName() + TEXT(".") + ActorDescView.GetActorLabelOrName().ToString();
+
+	switch (Reason)
+	{
+	case ELevelInstanceInvalidReason::WorldAssetNotFound:
+		UE_LOG(LogWorldPartition, Log, TEXT("Actor %s has an invalid world asset '%s'"), *ActorName, *WorldAsset.ToString());
+		break;
+	case ELevelInstanceInvalidReason::WorldAssetNotUsingExternalActors:
+		UE_LOG(LogWorldPartition, Log, TEXT("Actor %s world asset '%s' is not using external actors"), *ActorName, *WorldAsset.ToString());
+		break;
+	case ELevelInstanceInvalidReason::WorldAssetImcompatiblePartitioned:
+		UE_LOG(LogWorldPartition, Log, TEXT("Actor %s world asset '%s' is partitioned but not marked as compatible"), *ActorName, *WorldAsset.ToString());
+		break;
+	};
+}
 #endif
diff --git a/Engine/Source/Runtime/Engine/Private/WorldPartition/ErrorHandling/WorldPartitionStreamingGenerationTokenizedMessageErrorHandler.cpp b/Engine/Source/Runtime/Engine/Private/WorldPartition/ErrorHandling/WorldPartitionStreamingGenerationTokenizedMessageErrorHandler.cpp
--- a/Engine/Source/Runtime/Engine/Private/WorldPartition/ErrorHandling/WorldPartitionStreamingGenerationTokenizedMessageErrorHandler.cpp
+++ b/Engine/Source/Runtime/Engine/Private/WorldPartition/ErrorHandling/WorldPartitionStreamingGenerationTokenizedMessageErrorHandler.cpp
@@ -142,6 +142,39 @@ void ITokenizedMessageErrorHandler::OnActorNeedsResave(const FWorldPartitionActo
 	HandleTokenizedMessage(MoveTemp(Message));
 }
 
+void ITokenizedMessageErrorHandler::OnLevelInstanceInvalidWorldAsset(const FWorldPartitionActorDescView& ActorDescView, FName WorldAsset, ELevelInstanceInvalidReason Reason)
+{
+	FSoftObjectPath ActorPath(ActorDescView.GetActorPath());
+	EMessageSeverity::Type MessageSeverity = EMessageSeverity::Info;
+	FText ReasonText;
+
+	switch (Reason)
+	{
+	case ELevelInstanceInvalidReason::WorldAssetNotFound:
+		MessageSeverity = EMessageSeverity::Error;
+		ReasonText = LOCTEXT("TokenMessage_WorldPartition_HasInvalidWorldAsset", "has an invalid world asset");
+		break;
+	case ELevelInstanceInvalidReason::WorldAssetNotUsingExternalActors:
+		MessageSeverity = EMessageSeverity::Error;
+		ReasonText = LOCTEXT("TokenMessage_WorldPartition_WorldAssetIsNotUsingExternalActors", "is not using external actors");
+		break;
+	case ELevelInstanceInvalidReason::WorldAssetImcompatiblePartitioned:
+		MessageSeverity = EMessageSeverity::Error;
+		ReasonText = LOCTEXT("TokenMessage_WorldPartition_WorldAssetIsPartitionedIncompatible", "is partitioned but not marked as compatible");
+		break;
+	};
+
+	TSharedRef<FTokenizedMessage> Message = FTokenizedMessage::Create(MessageSeverity);
+	Message->AddToken(FTextToken::Create(LOCTEXT("TokenMessage_WorldPartition_Actor", "Level Instance")))
+		->AddToken(FAssetNameToken::Create(ActorPath.GetLongPackageName() + TEXT(".") + ActorDescView.GetActorLabelOrName().ToString()))
+		->AddToken(FTextToken::Create(ReasonText))
+		->AddToken(FAssetNameToken::Create(WorldAsset.ToString()));
+
+	AddAdditionalNameToken(Message, FName(TEXT("WorldPartition_LevelInstanceInvalidWorldAsset_CheckForErrors")));
+
+	HandleTokenizedMessage(MoveTemp(Message));
+}
+
 #undef LOCTEXT_NAMESPACE
 
 #endif
diff --git a/Engine/Source/Runtime/Engine/Private/WorldPartition/LevelInstance/LevelInstanceActorDesc.cpp b/Engine/Source/Runtime/Engine/Private/WorldPartition/LevelInstance/LevelInstanceActorDesc.cpp
--- a/Engine/Source/Runtime/Engine/Private/WorldPartition/LevelInstance/LevelInstanceActorDesc.cpp
+++ b/Engine/Source/Runtime/Engine/Private/WorldPartition/LevelInstance/LevelInstanceActorDesc.cpp
@@ -1,7 +1,6 @@
 // Copyright Epic Games, Inc. All Rights Reserved.
 
 #include "WorldPartition/LevelInstance/LevelInstanceActorDesc.h"
-#include "WorldPartition/ActorDescContainer.h"
 
 #if WITH_EDITOR
 #include "Engine/World.h"
@@ -9,8 +8,14 @@
 #include "LevelInstance/LevelInstanceInterface.h"
 #include "LevelInstance/LevelInstanceSubsystem.h"
 #include "WorldPartition/WorldPartition.h"
+#include "WorldPartition/ActorDescContainer.h"
+#include "WorldPartition/ErrorHandling/WorldPartitionStreamingGenerationErrorHandler.h"
 #include "UObject/UE5ReleaseStreamObjectVersion.h"
 #include "UObject/UE5MainStreamObjectVersion.h"
+#include "Logging/MessageLog.h"
+#include "Logging/TokenizedMessage.h"
+#include "Misc/MapErrors.h"
+#include "Misc/UObjectToken.h"
 
 static int32 GLevelInstanceDebugForceLevelStreaming = 0;
 static FAutoConsoleVariableRef CVarForceLevelStreaming(
@@ -173,6 +178,31 @@ bool FLevelInstanceActorDesc::GetContainerInstance(const UActorDescContainer*& O
 	return false;
 }
 
+void FLevelInstanceActorDesc::CheckForErrors(IStreamingGenerationErrorHandler* ErrorHandler) const
+{
+	FWorldPartitionActorDesc::CheckForErrors(ErrorHandler);
+
+	FPackagePath WorldAssetPath;
+	if (!FPackagePath::TryFromPackageName(LevelPackage, WorldAssetPath) || !FPackageName::DoesPackageExist(WorldAssetPath))
+	{
+		ErrorHandler->OnLevelInstanceInvalidWorldAsset(this, LevelPackage, IStreamingGenerationErrorHandler::ELevelInstanceInvalidReason::WorldAssetNotFound);
+	}
+	else if (!ULevel::GetIsLevelUsingExternalActorsFromPackage(LevelPackage))
+	{
+		if (DesiredRuntimeBehavior != ELevelInstanceRuntimeBehavior::LevelStreaming)
+		{
+			ErrorHandler->OnLevelInstanceInvalidWorldAsset(this, LevelPackage, IStreamingGenerationErrorHandler::ELevelInstanceInvalidReason::WorldAssetNotUsingExternalActors);
+		}
+	}
+	else if (ULevel::GetIsLevelPartitionedFromPackage(LevelPackage))
+	{
+		if ((DesiredRuntimeBehavior != ELevelInstanceRuntimeBehavior::Partitioned) || !ULevel::GetPartitionedLevelCanBeUsedByLevelInstanceFromPackage(LevelPackage))
+		{
+			ErrorHandler->OnLevelInstanceInvalidWorldAsset(this, LevelPackage, IStreamingGenerationErrorHandler::ELevelInstanceInvalidReason::WorldAssetImcompatiblePartitioned);
+		}
+	}
+}
+
 void FLevelInstanceActorDesc::TransferFrom(const FWorldPartitionActorDesc* From)
 {
 	FWorldPartitionActorDesc::TransferFrom(From);
@@ -237,4 +267,4 @@ void FLevelInstanceActorDesc::Serialize(FArchive& Ar)
 	}
 }
 
-#endif
+#endif
\ No newline at end of file
diff --git a/Engine/Source/Runtime/Engine/Private/WorldPartition/WorldPartitionActorDesc.cpp b/Engine/Source/Runtime/Engine/Private/WorldPartition/WorldPartitionActorDesc.cpp
--- a/Engine/Source/Runtime/Engine/Private/WorldPartition/WorldPartitionActorDesc.cpp
+++ b/Engine/Source/Runtime/Engine/Private/WorldPartition/WorldPartitionActorDesc.cpp
@@ -22,6 +22,7 @@
 #include "WorldPartition/DataLayer/DataLayerSubsystem.h"
 #include "WorldPartition/DataLayer/DataLayerAsset.h"
 #include "WorldPartition/DataLayer/DataLayerUtils.h"
+#include "WorldPartition/ErrorHandling/WorldPartitionStreamingGenerationErrorHandler.h"
 #include "Engine/Public/ActorReferencesUtils.h"
 
 TMap<TSubclassOf<AActor>, FWorldPartitionActorDesc::FActorDescDeprecator> FWorldPartitionActorDesc::Deprecators;
@@ -453,6 +454,14 @@ FName FWorldPartitionActorDesc::GetDisplayClassName() const
 	return BaseClass.IsNone() ? GetCleanClassName(NativeClass) : GetCleanClassName(BaseClass);
 }
 
+void FWorldPartitionActorDesc::CheckForErrors(IStreamingGenerationErrorHandler* ErrorHandler) const
+{
+	if (IsResaveNeeded())
+	{
+		ErrorHandler->OnActorNeedsResave(this);
+	}
+}
+
 bool FWorldPartitionActorDesc::IsLoaded(bool bEvenIfPendingKill) const
 {
 	if (ActorPtr.IsExplicitlyNull() || ActorPtr.IsStale())
diff --git a/Engine/Source/Runtime/Engine/Private/WorldPartition/WorldPartitionActorDescView.cpp b/Engine/Source/Runtime/Engine/Private/WorldPartition/WorldPartitionActorDescView.cpp
--- a/Engine/Source/Runtime/Engine/Private/WorldPartition/WorldPartitionActorDescView.cpp
+++ b/Engine/Source/Runtime/Engine/Private/WorldPartition/WorldPartitionActorDescView.cpp
@@ -150,6 +150,11 @@ bool FWorldPartitionActorDescView::GetContainerInstance(const UActorDescContaine
 	return ActorDesc->GetContainerInstance(OutLevelContainer, OutLevelTransform, OutClusterMode);
 }
 
+void FWorldPartitionActorDescView::CheckForErrors(IStreamingGenerationErrorHandler* ErrorHandler) const
+{
+	ActorDesc->CheckForErrors(ErrorHandler);
+}
+
 FName FWorldPartitionActorDescView::GetActorLabelOrName() const
 {
 	return ActorDesc->GetActorLabelOrName();
@@ -188,11 +193,6 @@ void FWorldPartitionActorDescView::SetRuntimeReferences(TArray<FGuid>& InRuntime
 	RuntimeReferences = InRuntimeReferences;
 }
 
-bool FWorldPartitionActorDescView::IsResaveNeeded() const
-{
-	return ActorDesc->IsResaveNeeded();
-}
-
 AActor* FWorldPartitionActorDescView::GetActor() const
 {
 	return ActorDesc->GetActor();
diff --git a/Engine/Source/Runtime/Engine/Private/WorldPartition/WorldPartitionStreamingGeneration.cpp b/Engine/Source/Runtime/Engine/Private/WorldPartition/WorldPartitionStreamingGeneration.cpp
--- a/Engine/Source/Runtime/Engine/Private/WorldPartition/WorldPartitionStreamingGeneration.cpp
+++ b/Engine/Source/Runtime/Engine/Private/WorldPartition/WorldPartitionStreamingGeneration.cpp
@@ -536,6 +536,12 @@ class FWorldPartitionStreamingGenerator
 			}
 		}
 
+		// Route standard CheckForErrors calls which should not modify actor descriptors in any ways
+		ContainerDescriptor.ActorDescViewMap.ForEachActorDescView([this](FWorldPartitionActorDescView& ActorDescView)
+		{
+			ActorDescView.CheckForErrors(ErrorHandler);
+		});
+
 		// Perform various adjustements based on validations and report errors
 		//
 		// The first validation pass is used to report errors, subsequent passes are used to make corrections to the FWorldPartitionActorDescView
@@ -709,15 +715,6 @@ class FWorldPartitionStreamingGenerator
 				}
 			});		
 		}
-
-		// Report actors that need to be resaved
-		ContainerDescriptor.ActorDescViewMap.ForEachActorDescView([this](FWorldPartitionActorDescView& ActorDescView)
-		{
-			if (ActorDescView.IsResaveNeeded())
-			{
-				ErrorHandler->OnActorNeedsResave(ActorDescView);
-			}
-		});
 	}
 
 	/** 
diff --git a/Engine/Source/Runtime/Engine/Public/WorldPartition/ErrorHandling/WorldPartitionStreamingGenerationErrorHandler.h b/Engine/Source/Runtime/Engine/Public/WorldPartition/ErrorHandling/WorldPartitionStreamingGenerationErrorHandler.h
--- a/Engine/Source/Runtime/Engine/Public/WorldPartition/ErrorHandling/WorldPartitionStreamingGenerationErrorHandler.h
+++ b/Engine/Source/Runtime/Engine/Public/WorldPartition/ErrorHandling/WorldPartitionStreamingGenerationErrorHandler.h
@@ -62,5 +62,17 @@ public:
 	 * Called when an actor needs to be resaved.
 	 */
 	virtual void OnActorNeedsResave(const FWorldPartitionActorDescView& ActorDescView) = 0;
+
+	/**
+	 * Called when a level instance actor has errrors
+	 */
+	enum class ELevelInstanceInvalidReason
+	{
+		WorldAssetNotFound,
+		WorldAssetNotUsingExternalActors,
+		WorldAssetImcompatiblePartitioned
+	};
+
+	virtual void OnLevelInstanceInvalidWorldAsset(const FWorldPartitionActorDescView& ActorDescView, FName WorldAsset, ELevelInstanceInvalidReason Reason) = 0;
 };
 #endif
diff --git a/Engine/Source/Runtime/Engine/Public/WorldPartition/ErrorHandling/WorldPartitionStreamingGenerationLogErrorHandler.h b/Engine/Source/Runtime/Engine/Public/WorldPartition/ErrorHandling/WorldPartitionStreamingGenerationLogErrorHandler.h
--- a/Engine/Source/Runtime/Engine/Public/WorldPartition/ErrorHandling/WorldPartitionStreamingGenerationLogErrorHandler.h
+++ b/Engine/Source/Runtime/Engine/Public/WorldPartition/ErrorHandling/WorldPartitionStreamingGenerationLogErrorHandler.h
@@ -17,5 +17,6 @@ public:
 	virtual void OnDataLayerHierarchyTypeMismatch(const UDataLayerInstance* DataLayerInstance, const UDataLayerInstance* Parent) override;
 	virtual void OnDataLayerAssetConflict(const UDataLayerInstanceWithAsset* DataLayerInstance, const UDataLayerInstanceWithAsset* ConflictingDataLayerInstance) override;
 	virtual void OnActorNeedsResave(const FWorldPartitionActorDescView& ActorDescView) override;
+	virtual void OnLevelInstanceInvalidWorldAsset(const FWorldPartitionActorDescView& ActorDescView, FName WorldAsset, ELevelInstanceInvalidReason Reason) override;
 };
 #endif
diff --git a/Engine/Source/Runtime/Engine/Public/WorldPartition/ErrorHandling/WorldPartitionStreamingGenerationNullErrorHandler.h b/Engine/Source/Runtime/Engine/Public/WorldPartition/ErrorHandling/WorldPartitionStreamingGenerationNullErrorHandler.h
--- a/Engine/Source/Runtime/Engine/Public/WorldPartition/ErrorHandling/WorldPartitionStreamingGenerationNullErrorHandler.h
+++ b/Engine/Source/Runtime/Engine/Public/WorldPartition/ErrorHandling/WorldPartitionStreamingGenerationNullErrorHandler.h
@@ -17,5 +17,6 @@ public:
 	virtual void OnDataLayerHierarchyTypeMismatch(const UDataLayerInstance* DataLayerInstance, const UDataLayerInstance* Parent) override {};
 	virtual void OnDataLayerAssetConflict(const UDataLayerInstanceWithAsset* DataLayerInstance, const UDataLayerInstanceWithAsset* ConflictingDataLayerInstance) override {};
 	virtual void OnActorNeedsResave(const FWorldPartitionActorDescView& ActorDescView) override {};
+	virtual void OnLevelInstanceInvalidWorldAsset(const FWorldPartitionActorDescView& ActorDescView, FName WorldAsset, ELevelInstanceInvalidReason Reason) override {};
 };
 #endif
diff --git a/Engine/Source/Runtime/Engine/Public/WorldPartition/ErrorHandling/WorldPartitionStreamingGenerationTokenizedMessageErrorHandler.h b/Engine/Source/Runtime/Engine/Public/WorldPartition/ErrorHandling/WorldPartitionStreamingGenerationTokenizedMessageErrorHandler.h
--- a/Engine/Source/Runtime/Engine/Public/WorldPartition/ErrorHandling/WorldPartitionStreamingGenerationTokenizedMessageErrorHandler.h
+++ b/Engine/Source/Runtime/Engine/Public/WorldPartition/ErrorHandling/WorldPartitionStreamingGenerationTokenizedMessageErrorHandler.h
@@ -18,6 +18,7 @@ class ENGINE_API ITokenizedMessageErrorHandler : public IStreamingGenerationErro
 	virtual void OnDataLayerHierarchyTypeMismatch(const UDataLayerInstance* DataLayerInstance, const UDataLayerInstance* Parent) override;
 	virtual void OnDataLayerAssetConflict(const UDataLayerInstanceWithAsset* DataLayerInstance, const UDataLayerInstanceWithAsset* ConflictingDataLayerInstance) override;
 	virtual void OnActorNeedsResave(const FWorldPartitionActorDescView& ActorDescView) override;
+	virtual void OnLevelInstanceInvalidWorldAsset(const FWorldPartitionActorDescView& ActorDescView, FName WorldAsset, ELevelInstanceInvalidReason Reason) override;
 
 protected:
 	virtual void AddAdditionalNameToken(TSharedRef<FTokenizedMessage>& InMessage, const FName& InErrorName) {}
diff --git a/Engine/Source/Runtime/Engine/Public/WorldPartition/LevelInstance/LevelInstanceActorDesc.h b/Engine/Source/Runtime/Engine/Public/WorldPartition/LevelInstance/LevelInstanceActorDesc.h
--- a/Engine/Source/Runtime/Engine/Public/WorldPartition/LevelInstance/LevelInstanceActorDesc.h
+++ b/Engine/Source/Runtime/Engine/Public/WorldPartition/LevelInstance/LevelInstanceActorDesc.h
@@ -12,6 +12,7 @@
 
 class ULevelInstanceSubsystem;
 class UActorDescContainer;
+class IStreamingGenerationErrorHandler;
 enum class ELevelInstanceRuntimeBehavior : uint8;
 
 /**
@@ -26,6 +27,7 @@ public:
 	virtual bool IsContainerInstance() const override;
 	virtual FName GetLevelPackage() const override { return LevelPackage; }
 	virtual bool GetContainerInstance(const UActorDescContainer*& OutLevelContainer, FTransform& OutLevelTransform, EContainerClusterMode& OutClusterMode) const override;
+	virtual void CheckForErrors(IStreamingGenerationErrorHandler* ErrorHandler) const override;
 
 protected:
 	virtual void Init(const AActor* InActor) override;
diff --git a/Engine/Source/Runtime/Engine/Public/WorldPartition/WorldPartitionActorDesc.h b/Engine/Source/Runtime/Engine/Public/WorldPartition/WorldPartitionActorDesc.h
--- a/Engine/Source/Runtime/Engine/Public/WorldPartition/WorldPartitionActorDesc.h
+++ b/Engine/Source/Runtime/Engine/Public/WorldPartition/WorldPartitionActorDesc.h
@@ -22,6 +22,7 @@ struct FWorldPartitionActorDescInitData
 
 class AActor;
 class UActorDescContainer;
+class IStreamingGenerationErrorHandler;
 struct FActorContainerID;
 
 enum class EContainerClusterMode : uint8
@@ -115,6 +116,7 @@ public:
 	virtual const FGuid& GetSceneOutlinerParent() const { return GetParentActor(); }
 	virtual bool IsResaveNeeded() const { return false; }
 	virtual bool IsRuntimeRelevant(const FActorContainerID& InContainerID) const { return true; }
+	virtual void CheckForErrors(IStreamingGenerationErrorHandler* ErrorHandler) const;
 
 	bool operator==(const FWorldPartitionActorDesc& Other) const
 	{
diff --git a/Engine/Source/Runtime/Engine/Public/WorldPartition/WorldPartitionActorDescView.h b/Engine/Source/Runtime/Engine/Public/WorldPartition/WorldPartitionActorDescView.h
--- a/Engine/Source/Runtime/Engine/Public/WorldPartition/WorldPartitionActorDescView.h
+++ b/Engine/Source/Runtime/Engine/Public/WorldPartition/WorldPartitionActorDescView.h
@@ -7,6 +7,7 @@
 #if WITH_EDITOR
 class AActor;
 class FWorldPartitionActorDesc;
+class IStreamingGenerationErrorHandler;
 class UActorDescContainer;
 class UWorldPartition;
 enum class EContainerClusterMode : uint8;
@@ -50,6 +51,8 @@ public:
 	bool IsContainerInstance() const;
 	bool GetContainerInstance(const UActorDescContainer*& OutLevelContainer, FTransform& OutLevelTransform, EContainerClusterMode& OutClusterMode) const;
 
+	void CheckForErrors(IStreamingGenerationErrorHandler* ErrorHandler) const;
+
 	FName GetActorLabelOrName() const;
 
 	void SetForcedNonSpatiallyLoaded();
@@ -58,8 +61,6 @@ public:
 	void SetRuntimeDataLayers(TArray<FName>& InRuntimeDataLayers);
 	void SetRuntimeReferences(TArray<FGuid>& InRuntimeReferences);
 
-	bool IsResaveNeeded() const;
-
 	AActor* GetActor() const;
 
 	bool operator==(const FWorldPartitionActorDescView& Other) const
