Disabled async uniform expression cache updates until further testing.

diff --git a/Engine/Source/Runtime/Engine/Private/Materials/MaterialShared.cpp b/Engine/Source/Runtime/Engine/Private/Materials/MaterialShared.cpp
--- a/Engine/Source/Runtime/Engine/Private/Materials/MaterialShared.cpp
+++ b/Engine/Source/Runtime/Engine/Private/Materials/MaterialShared.cpp
@@ -86,7 +86,7 @@ FAutoConsoleVariableRef CVarDeferUniformExpressionCaching(
 	ECVF_RenderThreadSafe
 	);
 
-int32 GUniformExpressionCacheAsyncUpdates = 1;
+int32 GUniformExpressionCacheAsyncUpdates = 0;
 FAutoConsoleVariableRef CVarUniformExpressionCacheAsyncUpdates(
 	TEXT("r.UniformExpressionCacheAsyncUpdates"),
 	GUniformExpressionCacheAsyncUpdates,
@@ -3582,14 +3582,11 @@ class FUniformExpressionCacheAsyncUpdateTask
 public:
 	void Begin()
 	{
-		check(!Task);
-		check(!bEnabled);
 		bEnabled = true;
 	}
 
 	void End()
 	{
-		check(bEnabled);
 		bEnabled = false;
 		Wait();
 	}
