Pasting an actor with Data Layers in a level that doesn't support DataLayers will now clear the Data Layers from the actor.

diff --git a/Engine/Source/Runtime/Engine/Private/Actor.cpp b/Engine/Source/Runtime/Engine/Private/Actor.cpp
--- a/Engine/Source/Runtime/Engine/Private/Actor.cpp
+++ b/Engine/Source/Runtime/Engine/Private/Actor.cpp
@@ -3482,6 +3482,8 @@ void AActor::PostEditImport()
 {
 	Super::PostEditImport();
 
+	FixupDataLayers();
+
 	DispatchOnComponentsCreated(this);
 
 	ULevel* Level = GetLevel();
diff --git a/Engine/Source/Runtime/Engine/Private/ActorEditor.cpp b/Engine/Source/Runtime/Engine/Private/ActorEditor.cpp
--- a/Engine/Source/Runtime/Engine/Private/ActorEditor.cpp
+++ b/Engine/Source/Runtime/Engine/Private/ActorEditor.cpp
@@ -1551,7 +1551,10 @@ bool AActor::IsPropertyChangedAffectingDataLayers(FPropertyChangedEvent& Propert
 
 bool AActor::SupportsDataLayer() const
 {
-	return (ActorTypeSupportsDataLayer() &&
+	ULevel* Level = GetLevel();
+	const bool bIsLevelNotPartitioned = Level ? !Level->bIsPartitioned : false;
+	return (!bIsLevelNotPartitioned &&
+			ActorTypeSupportsDataLayer() &&
 			!FActorEditorUtils::IsABuilderBrush(this) &&
 			!GetClass()->GetDefaultObject<AActor>()->bHiddenEd);
 }
