Detect at runtime streamed-in sub-levels with WorldDataLayers, avoid registering it and clear its DataLayerInstances to disable runtime DataLayers for this unsupported case.

diff --git a/Engine/Source/Runtime/Engine/Private/WorldPartition/DataLayer/DataLayerSubsystem.cpp b/Engine/Source/Runtime/Engine/Private/WorldPartition/DataLayer/DataLayerSubsystem.cpp
--- a/Engine/Source/Runtime/Engine/Private/WorldPartition/DataLayer/DataLayerSubsystem.cpp
+++ b/Engine/Source/Runtime/Engine/Private/WorldPartition/DataLayer/DataLayerSubsystem.cpp
@@ -226,6 +226,12 @@ TSet<FName> UDataLayerSubsystem::GetEffectiveLoadedDataLayerNames() const
 
 void UDataLayerSubsystem::RegisterWorldDataLayer(AWorldDataLayers* WorldDataLayers)
 {
+	// Don't register WorldDataLayers in game if it's not relevant at runtime
+	if (GetWorld()->IsGameWorld() && WorldDataLayers && !WorldDataLayers->IsRuntimeRelevant())
+	{
+		return;
+	}
+
 	if (WorldDataLayerCollection.RegisterWorldDataLayer(WorldDataLayers))
 	{
 #if WITH_EDITOR
diff --git a/Engine/Source/Runtime/Engine/Private/WorldPartition/DataLayer/WorldDataLayers.cpp b/Engine/Source/Runtime/Engine/Private/WorldPartition/DataLayer/WorldDataLayers.cpp
--- a/Engine/Source/Runtime/Engine/Private/WorldPartition/DataLayer/WorldDataLayers.cpp
+++ b/Engine/Source/Runtime/Engine/Private/WorldPartition/DataLayer/WorldDataLayers.cpp
@@ -89,7 +89,7 @@ void AWorldDataLayers::InitializeDataLayerRuntimeStates()
 {
 	check(ActiveDataLayerNames.IsEmpty() && LoadedDataLayerNames.IsEmpty());
 
-	if (GetWorld()->IsGameWorld())
+	if (GetWorld()->IsGameWorld() && IsRuntimeRelevant())
 	{
 		ForEachDataLayer([this](class UDataLayerInstance* DataLayer)
 		{
@@ -122,7 +122,7 @@ void AWorldDataLayers::InitializeDataLayerRuntimeStates()
 		RepEffectiveActiveDataLayerNames = EffectiveActiveDataLayerNames.Array();
 		RepEffectiveLoadedDataLayerNames = EffectiveLoadedDataLayerNames.Array();
 
-		UE_LOG(LogWorldPartition, Log, TEXT("Initial Data Layer Effective States Activated(%s) Loaded(%s)"), *JoinDataLayerShortNamesFromInstanceNames(this, RepEffectiveActiveDataLayerNames), *JoinDataLayerShortNamesFromInstanceNames(this, RepEffectiveLoadedDataLayerNames));
+		UE_CLOG(RepEffectiveActiveDataLayerNames.Num() || RepEffectiveLoadedDataLayerNames.Num(), LogWorldPartition, Log, TEXT("Initial Data Layer Effective States Activated(%s) Loaded(%s)"), *JoinDataLayerShortNamesFromInstanceNames(this, RepEffectiveActiveDataLayerNames), *JoinDataLayerShortNamesFromInstanceNames(this, RepEffectiveLoadedDataLayerNames));
 	}
 }
 
@@ -770,6 +770,14 @@ void AWorldDataLayers::PostLoad()
 		RemoveDataLayers(EditorDataLayers);
 	}
 
+	// Sub-WorldDataLayers are not supported at runtime, clear data layers
+	if (GetWorld()->IsGameWorld() && !IsRuntimeRelevant() && !IsEmpty())
+	{
+		UE_LOG(LogWorldPartition, Warning, TEXT("WorldDataLayers %s is not runtime relevant, Data Layers of this level will be ignored."), *GetPathName());
+		TArray<UDataLayerInstance*> AllDataLayers = DataLayerInstances.Array();
+		RemoveDataLayers(AllDataLayers);
+	}
+
 	// Setup defaults before overriding with user settings
 	for (UDataLayerInstance* DataLayer : DataLayerInstances)
 	{
@@ -838,6 +846,11 @@ void AWorldDataLayers::PostUnregisterAllComponents()
 	Super::PostUnregisterAllComponents();
 }
 
+bool AWorldDataLayers::IsRuntimeRelevant() const
+{
+	return !IsSubWorldDataLayers();
+}
+
 bool AWorldDataLayers::IsSubWorldDataLayers() const
 {
 	check(GetWorld());
diff --git a/Engine/Source/Runtime/Engine/Public/WorldPartition/DataLayer/WorldDataLayers.h b/Engine/Source/Runtime/Engine/Public/WorldPartition/DataLayer/WorldDataLayers.h
--- a/Engine/Source/Runtime/Engine/Public/WorldPartition/DataLayer/WorldDataLayers.h
+++ b/Engine/Source/Runtime/Engine/Public/WorldPartition/DataLayer/WorldDataLayers.h
@@ -89,6 +89,7 @@ public:
 #endif
 
 	bool IsSubWorldDataLayers() const;
+	bool IsRuntimeRelevant() const;
 	static FName GetWorldPartionWorldDataLayersName() { return FName(TEXT("WorldDataLayers")); } // reserved for ULevel::WorldDataLayers
 	
 	void DumpDataLayers(FOutputDevice& OutputDevice) const;
