World Partition Validations: added helper to retrieve the full actor name using actor label when possible and updated all error handlers to use that function instead.

#rnx

diff --git a/Engine/Source/Runtime/Engine/Private/WorldPartition/ErrorHandling/WorldPartitionStreamingGenerationErrorHandler.cpp b/Engine/Source/Runtime/Engine/Private/WorldPartition/ErrorHandling/WorldPartitionStreamingGenerationErrorHandler.cpp
new file mode 100644
--- /dev/null
+++ b/Engine/Source/Runtime/Engine/Private/WorldPartition/ErrorHandling/WorldPartitionStreamingGenerationErrorHandler.cpp
@@ -0,0 +1,13 @@
+// Copyright Epic Games, Inc. All Rights Reserved.
+
+#if WITH_EDITOR
+#include "WorldPartition/ErrorHandling/WorldPartitionStreamingGenerationErrorHandler.h"
+#include "UObject/SoftObjectPath.h"
+
+// Helpers
+FString IStreamingGenerationErrorHandler::GetFullActorName(const FWorldPartitionActorDescView& ActorDescView)
+{
+	const FSoftObjectPath ActorPath(ActorDescView.GetActorPath());
+	return ActorDescView.GetActorLabelOrName().ToString() + TEXT(" (") + ActorPath.GetLongPackageName() + TEXT(")");
+}
+#endif
diff --git a/Engine/Source/Runtime/Engine/Private/WorldPartition/ErrorHandling/WorldPartitionStreamingGenerationLogErrorHandler.cpp b/Engine/Source/Runtime/Engine/Private/WorldPartition/ErrorHandling/WorldPartitionStreamingGenerationLogErrorHandler.cpp
--- a/Engine/Source/Runtime/Engine/Private/WorldPartition/ErrorHandling/WorldPartitionStreamingGenerationLogErrorHandler.cpp
+++ b/Engine/Source/Runtime/Engine/Private/WorldPartition/ErrorHandling/WorldPartitionStreamingGenerationLogErrorHandler.cpp
@@ -8,7 +8,7 @@
 
 void FStreamingGenerationLogErrorHandler::OnInvalidReference(const FWorldPartitionActorDescView& ActorDescView, const FGuid& ReferenceGuid)
 {
-	UE_LOG(LogWorldPartition, Log, TEXT("Actor %s have missing reference to %s"), *ActorDescView.GetActorLabelOrName().ToString(), *ReferenceGuid.ToString());
+	UE_LOG(LogWorldPartition, Log, TEXT("Actor %s have missing reference to %s"), *GetFullActorName(ActorDescView), *ReferenceGuid.ToString());
 }
 
 void FStreamingGenerationLogErrorHandler::OnInvalidReferenceGridPlacement(const FWorldPartitionActorDescView& ActorDescView, const FWorldPartitionActorDescView& ReferenceActorDescView)
@@ -16,27 +16,27 @@ void FStreamingGenerationLogErrorHandler::OnInvalidReferenceGridPlacement(const
 	const FString SpatiallyLoadedActor(TEXT("Spatially loaded actor"));
 	const FString NonSpatiallyLoadedActor(TEXT("Non-spatially loaded loaded actor"));
 
-	UE_LOG(LogWorldPartition, Log, TEXT("%s %s reference %s %s"), ActorDescView.GetIsSpatiallyLoaded() ? *SpatiallyLoadedActor : *NonSpatiallyLoadedActor, *ActorDescView.GetActorLabelOrName().ToString(), ReferenceActorDescView.GetIsSpatiallyLoaded() ? *SpatiallyLoadedActor : *NonSpatiallyLoadedActor, *ReferenceActorDescView.GetActorLabelOrName().ToString());
+	UE_LOG(LogWorldPartition, Log, TEXT("%s %s reference %s %s"), ActorDescView.GetIsSpatiallyLoaded() ? *SpatiallyLoadedActor : *NonSpatiallyLoadedActor, *GetFullActorName(ActorDescView), ReferenceActorDescView.GetIsSpatiallyLoaded() ? *SpatiallyLoadedActor : *NonSpatiallyLoadedActor, *GetFullActorName(ReferenceActorDescView));
 }
 
 void FStreamingGenerationLogErrorHandler::OnInvalidReferenceDataLayers(const FWorldPartitionActorDescView& ActorDescView, const FWorldPartitionActorDescView& ReferenceActorDescView)
 {
-	UE_LOG(LogWorldPartition, Log, TEXT("Actor %s references an actor in a different set of runtime data layers %s"), *ActorDescView.GetActorLabelOrName().ToString(), *ReferenceActorDescView.GetActorLabelOrName().ToString());
+	UE_LOG(LogWorldPartition, Log, TEXT("Actor %s references an actor in a different set of runtime data layers %s"), *GetFullActorName(ActorDescView), *GetFullActorName(ReferenceActorDescView));
 }
 
 void FStreamingGenerationLogErrorHandler::OnInvalidReferenceRuntimeGrid(const FWorldPartitionActorDescView& ActorDescView, const FWorldPartitionActorDescView& ReferenceActorDescView)
 {
-	UE_LOG(LogWorldPartition, Log, TEXT("Actor %s references an actor in a different runtime grid %s"), *ActorDescView.GetActorLabelOrName().ToString(), *ReferenceActorDescView.GetActorLabelOrName().ToString());
+	UE_LOG(LogWorldPartition, Log, TEXT("Actor %s references an actor in a different runtime grid %s"), *GetFullActorName(ActorDescView), *GetFullActorName(ReferenceActorDescView));
 }
 
 void FStreamingGenerationLogErrorHandler::OnInvalidReferenceLevelScriptStreamed(const FWorldPartitionActorDescView& ActorDescView)
 {
-	UE_LOG(LogWorldPartition, Log, TEXT("Level Script Blueprint references streamed actor %s"), *ActorDescView.GetActorLabelOrName().ToString());
+	UE_LOG(LogWorldPartition, Log, TEXT("Level Script Blueprint references streamed actor %s"), *GetFullActorName(ActorDescView));
 }
 
 void FStreamingGenerationLogErrorHandler::OnInvalidReferenceLevelScriptDataLayers(const FWorldPartitionActorDescView& ActorDescView)
 {
-	UE_LOG(LogWorldPartition, Log, TEXT("Level Script Blueprint references streamed actor %s with a non empty set of data layers"), *ActorDescView.GetActorLabelOrName().ToString());
+	UE_LOG(LogWorldPartition, Log, TEXT("Level Script Blueprint references streamed actor %s with a non empty set of data layers"), *GetFullActorName(ActorDescView));
 }
 
 void FStreamingGenerationLogErrorHandler::OnInvalidReferenceDataLayerAsset(const UDataLayerInstanceWithAsset* DataLayerInstance)
@@ -56,13 +56,12 @@ void FStreamingGenerationLogErrorHandler::OnDataLayerAssetConflict(const UDataLa
 
 void FStreamingGenerationLogErrorHandler::OnActorNeedsResave(const FWorldPartitionActorDescView& ActorDescView)
 {
-	UE_LOG(LogWorldPartition, Log, TEXT("Actor %s needs to be resaved"), *ActorDescView.GetActorLabelOrName().ToString());
+	UE_LOG(LogWorldPartition, Log, TEXT("Actor %s needs to be resaved"), *GetFullActorName(ActorDescView));
 }
 
 void FStreamingGenerationLogErrorHandler::OnLevelInstanceInvalidWorldAsset(const FWorldPartitionActorDescView& ActorDescView, FName WorldAsset, ELevelInstanceInvalidReason Reason)
 {
-	const FSoftObjectPath ActorPath(ActorDescView.GetActorPath());
-	const FString ActorName = ActorPath.GetLongPackageName() + TEXT(".") + ActorDescView.GetActorLabelOrName().ToString();
+	const FString ActorName = GetFullActorName(ActorDescView);
 
 	switch (Reason)
 	{
diff --git a/Engine/Source/Runtime/Engine/Private/WorldPartition/ErrorHandling/WorldPartitionStreamingGenerationTokenizedMessageErrorHandler.cpp b/Engine/Source/Runtime/Engine/Private/WorldPartition/ErrorHandling/WorldPartitionStreamingGenerationTokenizedMessageErrorHandler.cpp
--- a/Engine/Source/Runtime/Engine/Private/WorldPartition/ErrorHandling/WorldPartitionStreamingGenerationTokenizedMessageErrorHandler.cpp
+++ b/Engine/Source/Runtime/Engine/Private/WorldPartition/ErrorHandling/WorldPartitionStreamingGenerationTokenizedMessageErrorHandler.cpp
@@ -14,7 +14,7 @@ void ITokenizedMessageErrorHandler::OnInvalidReference(const FWorldPartitionActo
 {
 	TSharedRef<FTokenizedMessage> Message = FTokenizedMessage::Create(EMessageSeverity::Warning);
 	Message->AddToken(FTextToken::Create(LOCTEXT("TokenMessage_WorldPartition_Actor", "Actor")))
-		->AddToken(FActorToken::Create(ActorDescView.GetActorPath().ToString(), ActorDescView.GetGuid(), FText::FromName(ActorDescView.GetActorLabelOrName())))
+		->AddToken(FActorToken::Create(ActorDescView.GetActorPath().ToString(), ActorDescView.GetGuid(), FText::FromString(GetFullActorName(ActorDescView))))
 		->AddToken(FTextToken::Create(LOCTEXT("TokenMessage_WorldPartition_HaveMissingRefsTo", "have missing references to")))
 		->AddToken(FTextToken::Create(FText::FromString(ReferenceGuid.ToString())));
 
@@ -30,7 +30,7 @@ void ITokenizedMessageErrorHandler::OnInvalidReferenceGridPlacement(const FWorld
 
 	TSharedRef<FTokenizedMessage> Message = FTokenizedMessage::Create(EMessageSeverity::Error);
 	Message->AddToken(FTextToken::Create(ActorDescView.GetIsSpatiallyLoaded() ? SpatiallyLoadedActor : NonSpatiallyLoadedActor))
-		->AddToken(FActorToken::Create(ActorDescView.GetActorPath().ToString(), ActorDescView.GetGuid(), FText::FromName(ActorDescView.GetActorLabelOrName())))
+		->AddToken(FActorToken::Create(ActorDescView.GetActorPath().ToString(), ActorDescView.GetGuid(), FText::FromString(GetFullActorName(ActorDescView))))
 		->AddToken(FTextToken::Create(LOCTEXT("TokenMessage_WorldPartition_References", "references")))
 		->AddToken(FTextToken::Create(ReferenceActorDescView.GetIsSpatiallyLoaded() ? SpatiallyLoadedActor : NonSpatiallyLoadedActor))
 		->AddToken(FActorToken::Create(ReferenceActorDescView.GetActorPath().ToString(), ReferenceActorDescView.GetGuid(), FText::FromName(ReferenceActorDescView.GetActorLabelOrName())));
@@ -44,7 +44,7 @@ void ITokenizedMessageErrorHandler::OnInvalidReferenceDataLayers(const FWorldPar
 {
 	TSharedRef<FTokenizedMessage> Message = FTokenizedMessage::Create(EMessageSeverity::Error);
 	Message->AddToken(FTextToken::Create(LOCTEXT("TokenMessage_WorldPartition_Actor", "Actor")))
-		->AddToken(FActorToken::Create(ActorDescView.GetActorPath().ToString(), ActorDescView.GetGuid(), FText::FromName(ActorDescView.GetActorLabelOrName())))
+		->AddToken(FActorToken::Create(ActorDescView.GetActorPath().ToString(), ActorDescView.GetGuid(),  FText::FromString(GetFullActorName(ActorDescView))))
 		->AddToken(FTextToken::Create(LOCTEXT("TokenMessage_WorldPartition_ReferenceActorInOtherDataLayers", "references an actor in a different set of runtime data layers")))
 		->AddToken(FActorToken::Create(ReferenceActorDescView.GetActorPath().ToString(), ReferenceActorDescView.GetGuid(), FText::FromName(ReferenceActorDescView.GetActorLabelOrName())));
 
@@ -57,7 +57,7 @@ void ITokenizedMessageErrorHandler::OnInvalidReferenceRuntimeGrid(const FWorldPa
 {
 	TSharedRef<FTokenizedMessage> Message = FTokenizedMessage::Create(EMessageSeverity::Error);
 	Message->AddToken(FTextToken::Create(LOCTEXT("TokenMessage_WorldPartition_Actor", "Actor")))
-		->AddToken(FActorToken::Create(ActorDescView.GetActorPath().ToString(), ActorDescView.GetGuid(), FText::FromName(ActorDescView.GetActorLabelOrName())))
+		->AddToken(FActorToken::Create(ActorDescView.GetActorPath().ToString(), ActorDescView.GetGuid(), FText::FromString(GetFullActorName(ActorDescView))))
 		->AddToken(FTextToken::Create(LOCTEXT("TokenMessage_WorldPartition_ReferenceActorInOtherRuntimeGrid", "references an actor in a different runtime grid")))
 		->AddToken(FActorToken::Create(ReferenceActorDescView.GetActorPath().ToString(), ReferenceActorDescView.GetGuid(), FText::FromName(ReferenceActorDescView.GetActorLabelOrName())));
 
@@ -70,7 +70,7 @@ void ITokenizedMessageErrorHandler::OnInvalidReferenceLevelScriptStreamed(const
 {
 	TSharedRef<FTokenizedMessage> Message = FTokenizedMessage::Create(EMessageSeverity::Error);
 	Message->AddToken(FTextToken::Create(LOCTEXT("TokenMessage_WorldPartition_LevelScriptBlueprintStreamedActorReference", "Level Script Blueprint references streamed actor")))
-		->AddToken(FActorToken::Create(ActorDescView.GetActorPath().ToString(), ActorDescView.GetGuid(), FText::FromName(ActorDescView.GetActorLabelOrName())));
+		->AddToken(FActorToken::Create(ActorDescView.GetActorPath().ToString(), ActorDescView.GetGuid(), FText::FromString(GetFullActorName(ActorDescView))));
 
 	AddAdditionalNameToken(Message, FName(TEXT("WorldPartition_LevelScriptBlueprintRefefenceStreamed_CheckForErrors")));
 
@@ -82,7 +82,7 @@ void ITokenizedMessageErrorHandler::OnInvalidReferenceLevelScriptDataLayers(cons
 {
 	TSharedRef<FTokenizedMessage> Message = FTokenizedMessage::Create(EMessageSeverity::Error);
 	Message->AddToken(FTextToken::Create(LOCTEXT("TokenMessage_WorldPartition_LevelScriptBlueprintActorReference", "Level Script Blueprint references actor")))
-		->AddToken(FActorToken::Create(ActorDescView.GetActorPath().ToString(), ActorDescView.GetGuid(), FText::FromName(ActorDescView.GetActorLabelOrName())))
+		->AddToken(FActorToken::Create(ActorDescView.GetActorPath().ToString(), ActorDescView.GetGuid(), FText::FromString(GetFullActorName(ActorDescView))))
 		->AddToken(FTextToken::Create(LOCTEXT("TokenMessage_WorldPartition_LevelScriptBlueprintDataLayerReference", "with a non empty set of data layers")));
 
 	AddAdditionalNameToken(Message, FName(TEXT("WorldPartition_LevelScriptBlueprintRefefenceDataLayer_CheckForErrors")));
@@ -135,7 +135,7 @@ void ITokenizedMessageErrorHandler::OnActorNeedsResave(const FWorldPartitionActo
 {
 	TSharedRef<FTokenizedMessage> Message = FTokenizedMessage::Create(EMessageSeverity::Warning);
 	Message->AddToken(FTextToken::Create(LOCTEXT("TokenMessage_WorldPartition_ActorNeedsResave", "Actor needs resave")))
-		->AddToken(FActorToken::Create(ActorDescView.GetActorPath().ToString(), ActorDescView.GetGuid(), FText::FromName(ActorDescView.GetActorLabelOrName())));
+		->AddToken(FActorToken::Create(ActorDescView.GetActorPath().ToString(), ActorDescView.GetGuid(), FText::FromString(GetFullActorName(ActorDescView))));
 
 	AddAdditionalNameToken(Message, FName(TEXT("WorldPartition_ActorNeedsResave_CheckForErrors")));
 
@@ -166,7 +166,7 @@ void ITokenizedMessageErrorHandler::OnLevelInstanceInvalidWorldAsset(const FWorl
 
 	TSharedRef<FTokenizedMessage> Message = FTokenizedMessage::Create(MessageSeverity);
 	Message->AddToken(FTextToken::Create(LOCTEXT("TokenMessage_WorldPartition_Actor", "Level Instance")))
-		->AddToken(FAssetNameToken::Create(ActorPath.GetLongPackageName() + TEXT(".") + ActorDescView.GetActorLabelOrName().ToString()))
+		->AddToken(FTextToken::Create(FText::FromString(GetFullActorName(ActorDescView))))
 		->AddToken(FTextToken::Create(ReasonText))
 		->AddToken(FAssetNameToken::Create(WorldAsset.ToString()));
 
diff --git a/Engine/Source/Runtime/Engine/Public/WorldPartition/ErrorHandling/WorldPartitionStreamingGenerationErrorHandler.h b/Engine/Source/Runtime/Engine/Public/WorldPartition/ErrorHandling/WorldPartitionStreamingGenerationErrorHandler.h
--- a/Engine/Source/Runtime/Engine/Public/WorldPartition/ErrorHandling/WorldPartitionStreamingGenerationErrorHandler.h
+++ b/Engine/Source/Runtime/Engine/Public/WorldPartition/ErrorHandling/WorldPartitionStreamingGenerationErrorHandler.h
@@ -74,5 +74,8 @@ public:
 	};
 
 	virtual void OnLevelInstanceInvalidWorldAsset(const FWorldPartitionActorDescView& ActorDescView, FName WorldAsset, ELevelInstanceInvalidReason Reason) = 0;
+
+	// Helpers
+	static FString GetFullActorName(const FWorldPartitionActorDescView& ActorDescView);
 };
 #endif
