LLM Scopes added to SkinnedMeshComponent to track LOD vertex overrides.

diff --git a/Engine/Source/Runtime/Engine/Private/Components/SkinnedMeshComponent.cpp b/Engine/Source/Runtime/Engine/Private/Components/SkinnedMeshComponent.cpp
--- a/Engine/Source/Runtime/Engine/Private/Components/SkinnedMeshComponent.cpp
+++ b/Engine/Source/Runtime/Engine/Private/Components/SkinnedMeshComponent.cpp
@@ -3779,6 +3779,8 @@ void USkinnedMeshComponent::SetVertexColorOverride_LinearColor(int32 LODIndex, c
 
 void USkinnedMeshComponent::SetVertexColorOverride(int32 LODIndex, const TArray<FColor>& VertexColors)
 {
+	LLM_SCOPE_BYNAME(TEXT("SkeletalMesh/VertexColorOverride"));
+
 	InitLODInfos();
 
 	FSkeletalMeshRenderData* SkelMeshRenderData = GetSkeletalMeshRenderData();
@@ -3972,6 +3974,8 @@ void CreateSkinWeightsArray(
 
 void USkinnedMeshComponent::SetSkinWeightOverride(int32 LODIndex, const TArray<FSkelMeshSkinWeightInfo>& SkinWeights)
 {
+	LLM_SCOPE_BYNAME(TEXT("SkeletalMesh/SkinWeightOverride"));
+
 	InitLODInfos();
 
 	FSkeletalMeshRenderData* SkelMeshRenderData = GetSkeletalMeshRenderData();
