AInfo: Fix RootComponent/SpriteComponent pointing to RF_PendingKill components

diff --git a/Engine/Source/Runtime/Engine/Classes/GameFramework/Info.h b/Engine/Source/Runtime/Engine/Classes/GameFramework/Info.h
--- a/Engine/Source/Runtime/Engine/Classes/GameFramework/Info.h
+++ b/Engine/Source/Runtime/Engine/Classes/GameFramework/Info.h
@@ -21,6 +21,11 @@ class AInfo : public AActor
 {
 	GENERATED_UCLASS_BODY()
 
+public:
+#if WITH_EDITOR
+	ENGINE_API virtual void PostLoad() override;
+#endif
+
 private:
 #if WITH_EDITOR
 	virtual bool ActorTypeSupportsDataLayer() const override { return false; }
diff --git a/Engine/Source/Runtime/Engine/Private/Info.cpp b/Engine/Source/Runtime/Engine/Private/Info.cpp
--- a/Engine/Source/Runtime/Engine/Private/Info.cpp
+++ b/Engine/Source/Runtime/Engine/Private/Info.cpp
@@ -51,3 +51,24 @@ AInfo::AInfo(const FObjectInitializer& ObjectInitializer)
 /** Returns SpriteComponent subobject **/
 UBillboardComponent* AInfo::GetSpriteComponent() const { return SpriteComponent; }
 #endif
+
+#if WITH_EDITOR
+void AInfo::PostLoad()
+{
+	Super::PostLoad();
+
+	// Fixes Actors that were saved prior to their constructor using ObjectInitializer.DoNotCreateDefaultSubobject(TEXT("Sprite"))
+	// Those Components are MarkForGarbage in their UActorComponent::PostInitProperties and so we need to clear the Component pointers
+	if (RootComponent != nullptr && !IsValid(RootComponent))
+	{
+		RootComponent = nullptr;
+	}
+
+#if WITH_EDITORONLY_DATA
+	if (SpriteComponent != nullptr && !IsValid(SpriteComponent))
+	{
+		SpriteComponent = nullptr;
+	}
+#endif
+}
+#endif
\ No newline at end of file
