- Fixed assignment of NetTraceId/GameInstanceId when using NetTrace
- Fixed typo that inverted intent of logic updating ViewTargets when running iris

diff --git a/Engine/Source/Runtime/Engine/Classes/Engine/NetDriver.h b/Engine/Source/Runtime/Engine/Classes/Engine/NetDriver.h
--- a/Engine/Source/Runtime/Engine/Classes/Engine/NetDriver.h
+++ b/Engine/Source/Runtime/Engine/Classes/Engine/NetDriver.h
@@ -1944,6 +1944,8 @@ private:
 	void SendClientMoveAdjustments();
 #endif
 
+	void InitNetTraceId();
+
 	UPROPERTY(transient)
 	TObjectPtr<UReplicationDriver> ReplicationDriver;
 
@@ -1979,7 +1981,7 @@ private:
 	FNetConnectionIdHandler ConnectionIdHandler;
 
 	/** Unique id used by NetTrace to identify driver */
-	uint32 NetTraceId = 0;
+	uint32 NetTraceId;
 
 	/** Stat tracking for the total number of out of order packets lost */
 	int32 TotalOutOfOrderPacketsLost = 0;
diff --git a/Engine/Source/Runtime/Engine/Private/NetDriver.cpp b/Engine/Source/Runtime/Engine/Private/NetDriver.cpp
--- a/Engine/Source/Runtime/Engine/Private/NetDriver.cpp
+++ b/Engine/Source/Runtime/Engine/Private/NetDriver.cpp
@@ -476,6 +476,7 @@ UNetDriver::UNetDriver(const FObjectInitializer& ObjectInitializer)
 ,	DuplicateLevelID(INDEX_NONE)
 ,	PacketLossBurstEndTime(-1.0f)
 ,	OutTotalNotifiedPackets(0u)
+,	NetTraceId(NetTraceInvalidGameInstanceId)
 {
 	UpdateDelayRandomStream.Initialize(FApp::bUseFixedSeed ? GetFName() : NAME_None);
 }
@@ -611,9 +612,10 @@ void UNetDriver::LoadChannelDefinitions()
 void UNetDriver::NotifyGameInstanceUpdated()
 {
 #if UE_NET_TRACE_ENABLED
-	if (GetWorld() && GetWorld()->WorldType)
+	const UWorld* LocalWorld = GetWorld();
+	if (GetNetTraceId() != NetTraceInvalidGameInstanceId && LocalWorld && LocalWorld->WorldType)
 	{
-		FString InstanceName = FString::Printf(TEXT("%s (%s)"), *GetNameSafe(GetWorld()->GetGameInstance()), LexToString(GetWorld()->WorldType.GetValue()));
+		FString InstanceName = FString::Printf(TEXT("%s (%s)"), *GetNameSafe(LocalWorld->GetGameInstance()), LexToString(LocalWorld->WorldType.GetValue()));
 		UE_NET_TRACE_UPDATE_INSTANCE(GetNetTraceId(), IsServer(), *InstanceName);
 	}
 #endif
@@ -1357,9 +1359,6 @@ bool UNetDriver::InitBase(bool bInitAsClient, FNetworkNotify* InNotify, const FU
 
 	Notify = InNotify;
 
-	// If we are not using Iris, we use the UniqueId to identify the NetDriver.
-	NetTraceId = GetUniqueID();
-
 #if UE_WITH_IRIS
 	if (IsUsingIrisReplication() && !ReplicationSystem)
 	{
@@ -1389,6 +1388,11 @@ bool UNetDriver::InitBase(bool bInitAsClient, FNetworkNotify* InNotify, const FU
 		}
 	}
 #endif // UE_WITH_IRIS
+	
+	InitNetTraceId();
+
+	// NotifyGameInstanceUpdate might have been called prior to setting up the NetTraceId so we call it again
+	NotifyGameInstanceUpdated();
 
 	if (!bInitAsClient)
 	{
@@ -6026,7 +6030,8 @@ void UNetDriver::UpdateReplicationViews() const
 	}
 
 	TObjectPtr<UNetConnection> const* Connections = (ServerConnection != nullptr ? &ServerConnection : ClientConnections.GetData());
-	const bool bUpdateConnectionViewTarget = ServerConnection != nullptr;
+	const bool bUpdateConnectionViewTarget = ServerConnection == nullptr;
+
 	const uint32 ConnectionCount = (ServerConnection != nullptr ? 1U : uint32(ClientConnections.Num()));
 	TArray<UNetConnection*, TInlineAllocator<4>> SubConnections;
 	for (uint32 ConnIt = 0, ConnEndIt = ConnectionCount; ConnIt != ConnEndIt; ++ConnIt)
@@ -6118,6 +6123,21 @@ void UNetDriver::SendClientMoveAdjustments()
 }
 #endif // UE_WITH_IRIS
 
+void UNetDriver::InitNetTraceId()
+{
+	if (IsUsingIrisReplication())
+	{
+		return;
+	}
+
+	if (NetTraceId == NetTraceInvalidGameInstanceId)
+	{
+		// We just need to make sure that all active NetDriver`s running in the same game instance uses different NetTraceId`s.
+		static uint8 CurrentNetTraceId = 255;
+		NetTraceId = ++CurrentNetTraceId;
+	}
+}
+
 #if NET_DEBUG_RELEVANT_ACTORS
 static void	DumpRelevantActors( UWorld* InWorld )
 {
diff --git a/Engine/Source/Runtime/Net/Core/Public/Net/Core/Trace/NetTraceConfig.h b/Engine/Source/Runtime/Net/Core/Public/Net/Core/Trace/NetTraceConfig.h
--- a/Engine/Source/Runtime/Net/Core/Public/Net/Core/Trace/NetTraceConfig.h
+++ b/Engine/Source/Runtime/Net/Core/Public/Net/Core/Trace/NetTraceConfig.h
@@ -26,6 +26,8 @@ namespace ENetTraceVerbosity
 	};
 }
 
+const uint32 NetTraceInvalidGameInstanceId = ~0U;
+
 #if UE_NET_TRACE_ENABLED
 #	ifndef UE_NET_TRACE_COMPILETIME_VERBOSITY
 #		if !(UE_BUILD_SHIPPING || UE_BUILD_TEST)
