Fix localization warnings.

#rnx

diff --git a/Engine/Source/Runtime/Engine/Private/WorldPartition/ErrorHandling/WorldPartitionStreamingGenerationTokenizedMessageErrorHandler.cpp b/Engine/Source/Runtime/Engine/Private/WorldPartition/ErrorHandling/WorldPartitionStreamingGenerationTokenizedMessageErrorHandler.cpp
--- a/Engine/Source/Runtime/Engine/Private/WorldPartition/ErrorHandling/WorldPartitionStreamingGenerationTokenizedMessageErrorHandler.cpp
+++ b/Engine/Source/Runtime/Engine/Private/WorldPartition/ErrorHandling/WorldPartitionStreamingGenerationTokenizedMessageErrorHandler.cpp
@@ -165,7 +165,7 @@ void ITokenizedMessageErrorHandler::OnLevelInstanceInvalidWorldAsset(const FWorl
 	};
 
 	TSharedRef<FTokenizedMessage> Message = FTokenizedMessage::Create(MessageSeverity);
-	Message->AddToken(FTextToken::Create(LOCTEXT("TokenMessage_WorldPartition_Actor", "Level Instance")))
+	Message->AddToken(FTextToken::Create(LOCTEXT("TokenMessage_WorldPartition_LevelInstance", "Level Instance")))
 		->AddToken(FTextToken::Create(FText::FromString(GetFullActorName(ActorDescView))))
 		->AddToken(FTextToken::Create(ReasonText))
 		->AddToken(FAssetNameToken::Create(WorldAsset.ToString()));
