When recursing through material graph for a given property only recurse through that property pin on a SetMaterialAttributes node.

diff --git a/Engine/Source/Runtime/Engine/Private/Materials/Material.cpp b/Engine/Source/Runtime/Engine/Private/Materials/Material.cpp
--- a/Engine/Source/Runtime/Engine/Private/Materials/Material.cpp
+++ b/Engine/Source/Runtime/Engine/Private/Materials/Material.cpp
@@ -5559,6 +5559,7 @@ bool UMaterial::RecursiveGetExpressionChain(
 	UMaterialExpressionQualitySwitch* QualitySwitchExp;
 	UMaterialExpressionShadingPathSwitch* ShadingPathSwitchExp;
 	UMaterialExpressionMakeMaterialAttributes* MakeMaterialAttributesExp;
+	UMaterialExpressionSetMaterialAttributes* SetMaterialAttributesExp;
 	UMaterialExpressionShaderStageSwitch* ShaderStageSwitchExp;
 
 	if (InFeatureLevel != ERHIFeatureLevel::Num && (FeatureLevelSwitchExp = Cast<UMaterialExpressionFeatureLevelSwitch>(InExpression)) != nullptr)
@@ -5661,6 +5662,27 @@ bool UMaterial::RecursiveGetExpressionChain(
 			InputsFrequency.Init(InShaderFrequency, Inputs.Num());
 		}
 	}
+	else if ((SetMaterialAttributesExp = Cast<UMaterialExpressionSetMaterialAttributes>(InExpression)) != nullptr)
+	{
+		if (InProperty != MP_MAX)
+		{
+			for (int32 AttributeIndex = 0; AttributeIndex < SetMaterialAttributesExp->AttributeSetTypes.Num(); ++AttributeIndex)
+			{
+				if (FMaterialAttributeDefinitionMap::GetProperty(SetMaterialAttributesExp->AttributeSetTypes[AttributeIndex]) == InProperty)
+				{
+					Inputs.Add(SetMaterialAttributesExp->GetInput(AttributeIndex + 1)); // Need +1 for MaterialAttributes input pin.
+					InputsFrequency.Add(InShaderFrequency);
+					break;
+				}
+			}
+		}
+		else
+		{
+			// Follow all properties.
+			Inputs = InExpression->GetInputs();
+			InputsFrequency.Init(InShaderFrequency, Inputs.Num());
+		}
+	}
 	else
 	{
 		Inputs = InExpression->GetInputs();
