Added warning for default physics collision handler to catch invalid actor references in pending collision lists. Mitigated ensure during collision handler process.

diff --git a/Engine/Source/Runtime/Engine/Private/PhysicsEngine/PhysicsCollisionHandler.cpp b/Engine/Source/Runtime/Engine/Private/PhysicsEngine/PhysicsCollisionHandler.cpp
--- a/Engine/Source/Runtime/Engine/Private/PhysicsEngine/PhysicsCollisionHandler.cpp
+++ b/Engine/Source/Runtime/Engine/Private/PhysicsEngine/PhysicsCollisionHandler.cpp
@@ -4,6 +4,8 @@
 #include "Kismet/GameplayStatics.h"
 #include "PhysicsPublic.h"
 #include "PhysicsEngine/BodyInstance.h"
+#include "Components/PrimitiveComponent.h"
+
 UPhysicsCollisionHandler::UPhysicsCollisionHandler(const FObjectInitializer& ObjectInitializer)
 	: Super(ObjectInitializer)
 {
@@ -19,6 +21,28 @@ void UPhysicsCollisionHandler::DefaultHandleCollision_AssumesLocked(const FRigid
 
 	if(BodyInst0 && BodyInst1)
 	{
+		const FPhysicsActorHandle Handle0 = BodyInst0->GetPhysicsActorHandle();
+		const FPhysicsActorHandle Handle1 = BodyInst1->GetPhysicsActorHandle();
+
+		if(!FPhysicsInterface::IsValid(Handle0) || !FPhysicsInterface::IsValid(Handle1))
+		{
+			const UPrimitiveComponent* Comp0 = BodyInst0->OwnerComponent.Get();
+			const UPrimitiveComponent* Comp1 = BodyInst1->OwnerComponent.Get();
+
+			if(Comp0 && Comp1)
+			{
+				UE_LOG(LogChaos, Warning, TEXT("Collision handler encountered invalid actor\n\t%s:%s [%s] - %s\n\t%s:%s [%s] - %s"),
+					   *Comp0->GetOwner()->GetName(), *Comp0->GetName(), *Comp0->GetClass()->GetName(), FPhysicsInterface::IsValid(Handle0) ? TEXT("Valid") : TEXT("Invalid"),
+					   *Comp1->GetOwner()->GetName(), *Comp1->GetName(), *Comp1->GetClass()->GetName(), FPhysicsInterface::IsValid(Handle1) ? TEXT("Valid") : TEXT("Invalid"));
+			}
+			else
+			{
+				UE_LOG(LogChaos, Warning, TEXT("Collision handler encountered invalid actors and could not recover components"));
+			}
+
+			return;
+		}
+
 		// Find relative velocity.
 		FVector Velocity0 = BodyInst0->GetUnrealWorldVelocityAtPoint_AssumesLocked(ContactInfo.ContactPosition);
 		FVector AngularVel0 = FMath::RadiansToDegrees(BodyInst0->GetUnrealWorldAngularVelocityInRadians_AssumesLocked());
