[Bug Fix] Actors attached to a static mesh socket inside a blueprint actor not updating when the mesh is moved.

UDN: https://epicgames.lightning.force.com/lightning/r/Case/5004z00001fDKxqAAG/view

diff --git a/Engine/Source/Runtime/Engine/Private/ActorConstruction.cpp b/Engine/Source/Runtime/Engine/Private/ActorConstruction.cpp
--- a/Engine/Source/Runtime/Engine/Private/ActorConstruction.cpp
+++ b/Engine/Source/Runtime/Engine/Private/ActorConstruction.cpp
@@ -304,6 +304,7 @@ void AActor::RerunConstructionScripts()
 			FName AttachedToSocket;
 			bool bSetRelativeTransform;
 			FTransform RelativeTransform;
+			FName AttachParentName;
 		};
 
 		// Save info about attached actors
@@ -388,6 +389,10 @@ void AActor::RerunConstructionScripts()
 						Info.AttachedActor = AttachedActor;
 						Info.AttachedToSocket = EachRoot->GetAttachSocketName();
 						Info.bSetRelativeTransform = false;
+						if (EachRoot->GetAttachParent() != RootComponent)
+						{
+							Info.AttachParentName = EachRoot->GetAttachParent()->GetFName();
+						}
 						AttachedActorInfos.Add(Info);
 
 						// Now detach it
@@ -550,7 +555,33 @@ void AActor::RerunConstructionScripts()
 				USceneComponent* ChildRoot = Info.AttachedActor->GetRootComponent();
 				if (ChildRoot && ChildRoot->GetAttachParent() != RootComponent)
 				{
-					ChildRoot->AttachToComponent(RootComponent, FAttachmentTransformRules::KeepWorldTransform, Info.AttachedToSocket);
+					if (Info.AttachParentName != NAME_None)
+					{
+						TArray<USceneComponent*> ChildComponents;
+						RootComponent->GetChildrenComponents(true, ChildComponents);
+						for (USceneComponent* Child : ChildComponents)
+						{
+							if (Child->GetFName() == Info.AttachParentName)
+							{
+								ChildRoot->AttachToComponent(Child, FAttachmentTransformRules::SnapToTargetIncludingScale, Info.AttachedToSocket);
+								break;
+							}
+						}
+						// if we couldn't find component by name, attach it to root and log a warning
+						if (!ChildRoot->GetAttachParent())
+						{
+							UE_LOG(LogBlueprint, Warning,
+								TEXT("Couldn't find a component named \'%s\' when reattaching. Attaching to root component instead."),
+								*Info.AttachParentName.ToString()
+							);
+							ChildRoot->AttachToComponent(RootComponent, FAttachmentTransformRules::KeepWorldTransform, Info.AttachedToSocket);
+						}
+					}
+					else
+					{
+						ChildRoot->AttachToComponent(RootComponent, FAttachmentTransformRules::KeepWorldTransform, Info.AttachedToSocket);
+					}
+
 					if (Info.bSetRelativeTransform)
 					{
 						ChildRoot->SetRelativeTransform(Info.RelativeTransform);
