UVolumeTexture: Add texture address mode member, controlling all three axes at once.

The wrap default (as set in UTexture and FStreamableTextureResource) is maintained. This allows us to use the texture's SamplerStateRHI where previously we only had control over the filtering mode (which was insufficient).

diff --git a/Engine/Source/Runtime/Engine/Classes/Engine/VolumeTexture.h b/Engine/Source/Runtime/Engine/Classes/Engine/VolumeTexture.h
--- a/Engine/Source/Runtime/Engine/Classes/Engine/VolumeTexture.h
+++ b/Engine/Source/Runtime/Engine/Classes/Engine/VolumeTexture.h
@@ -51,6 +51,10 @@ public:
 	int32 Source2DTileSizeY;
 #endif
 
+	/** The addressing mode to use for the X, Y and Z axis. */
+	UPROPERTY(EditAnywhere, BlueprintReadWrite, Category = Texture, meta = (DisplayName = "Tiling Method"), AssetRegistrySearchable, AdvancedDisplay)
+	TEnumAsByte<enum TextureAddress> AddressMode = TA_Wrap;
+
 	ENGINE_API bool UpdateSourceFromSourceTexture();
 	
 	/**
@@ -85,6 +89,9 @@ public:
 	virtual float GetSurfaceHeight() const override { return static_cast<float>(GetSizeY()); }
 	virtual float GetSurfaceDepth() const override { return static_cast<float>(GetSizeZ()); }
 	virtual uint32 GetSurfaceArraySize() const override { return 0; }
+	virtual TextureAddress GetTextureAddressX() const override { return AddressMode; }
+	virtual TextureAddress GetTextureAddressY() const override { return AddressMode; }
+	virtual TextureAddress GetTextureAddressZ() const override { return AddressMode; }
 	virtual FTextureResource* CreateResource() override;
 #if WITH_EDITOR
 	ENGINE_API void SetDefaultSource2DTileSize();
diff --git a/Engine/Source/Runtime/Engine/Private/Rendering/StreamableTextureResource.cpp b/Engine/Source/Runtime/Engine/Private/Rendering/StreamableTextureResource.cpp
--- a/Engine/Source/Runtime/Engine/Private/Rendering/StreamableTextureResource.cpp
+++ b/Engine/Source/Runtime/Engine/Private/Rendering/StreamableTextureResource.cpp
@@ -99,6 +99,9 @@ FStreamableTextureResource::FStreamableTextureResource(UTexture* InOwner, const
 	bGreyScaleFormat = (PixelFormat == PF_G8) || (PixelFormat == PF_BC4);
 
 	Filter = (ESamplerFilter)UDeviceProfileManager::Get().GetActiveProfile()->GetTextureLODSettings()->GetSamplerFilter(InOwner);
+	AddressU = InOwner->GetTextureAddressX() == TA_Wrap ? AM_Wrap : (InOwner->GetTextureAddressX() == TA_Clamp ? AM_Clamp : AM_Mirror);
+	AddressV = InOwner->GetTextureAddressY() == TA_Wrap ? AM_Wrap : (InOwner->GetTextureAddressY() == TA_Clamp ? AM_Clamp : AM_Mirror);
+	AddressW = InOwner->GetTextureAddressZ() == TA_Wrap ? AM_Wrap : (InOwner->GetTextureAddressZ() == TA_Clamp ? AM_Clamp : AM_Mirror);
 
 	// Get the biggest mips size, might be different from the actual resolution (depending on NumOfResidentLODs).
 	const FTexture2DMipMap& Mip0 = PlatformData->Mips[State.AssetLODBias];
diff --git a/Engine/Source/Runtime/Engine/Private/VolumeTexture.cpp b/Engine/Source/Runtime/Engine/Private/VolumeTexture.cpp
--- a/Engine/Source/Runtime/Engine/Private/VolumeTexture.cpp
+++ b/Engine/Source/Runtime/Engine/Private/VolumeTexture.cpp
@@ -594,6 +594,11 @@ void UVolumeTexture::PostEditChangeProperty(FPropertyChangedEvent& PropertyChang
 		{
 			UpdateSourceFromSourceTexture();
 		}
+		// Potentially update sampler state in the materials
+		if (PropertyName == GET_MEMBER_NAME_CHECKED(UVolumeTexture, AddressMode))
+		{
+			NotifyMaterials();
+		}
 	}
 		
 	Super::PostEditChangeProperty(PropertyChangedEvent);
