Wrapping the check for the SkipReplication CVar in a static so games that extend PhysicsReplication can more easily check it.

diff --git a/Engine/Source/Runtime/Engine/Private/PhysicsEngine/PhysicsReplication.cpp b/Engine/Source/Runtime/Engine/Private/PhysicsEngine/PhysicsReplication.cpp
--- a/Engine/Source/Runtime/Engine/Private/PhysicsEngine/PhysicsReplication.cpp
+++ b/Engine/Source/Runtime/Engine/Private/PhysicsEngine/PhysicsReplication.cpp
@@ -23,7 +23,7 @@
 #include "PBDRigidsSolver.h"
 #include "Chaos/PBDRigidsEvolutionGBF.h"
 
-namespace CharacterMovementCVars
+namespace PhysicsReplicationCVars
 {
 	extern int32 NetShowCorrections;
 	extern float NetCorrectionLifetime;
@@ -134,7 +134,7 @@ FPhysicsReplication::~FPhysicsReplication()
 
 bool FPhysicsReplication::ApplyRigidBodyState(float DeltaSeconds, FBodyInstance* BI, FReplicatedPhysicsTarget& PhysicsTarget, const FRigidBodyErrorCorrection& ErrorCorrection, const float InPingSecondsOneWay, int32 LocalFrame, int32 NumPredictedFrames)
 {
-	if (CharacterMovementCVars::SkipPhysicsReplication)
+	if (ShouldSkipPhysicsReplication())
 	{
 		return false;
 	}
@@ -179,7 +179,7 @@ bool FPhysicsReplication::ApplyRigidBodyState(float DeltaSeconds, FBodyInstance*
 
 bool FPhysicsReplication::ApplyRigidBodyState(float DeltaSeconds, FBodyInstance* BI, FReplicatedPhysicsTarget& PhysicsTarget, const FRigidBodyErrorCorrection& ErrorCorrection, const float PingSecondsOneWay, bool* bDidHardSnap)
 {
-	if (CharacterMovementCVars::SkipPhysicsReplication)
+	if (ShouldSkipPhysicsReplication())
 	{
 		return false;
 	}
@@ -242,19 +242,19 @@ bool FPhysicsReplication::ApplyRigidBodyState(float DeltaSeconds, FBodyInstance*
 	}
 
 	// Grab configuration variables from engine config or from CVars if overriding is turned on.
-	const float NetPingExtrapolation = CharacterMovementCVars::NetPingExtrapolation >= 0.0f ? CharacterMovementCVars::NetPingExtrapolation : ErrorCorrection.PingExtrapolation;
-	const float NetPingLimit = CharacterMovementCVars::NetPingLimit > 0.0f ? CharacterMovementCVars::NetPingLimit : ErrorCorrection.PingLimit;
-	const float ErrorPerLinearDiff = CharacterMovementCVars::ErrorPerLinearDifference >= 0.0f ? CharacterMovementCVars::ErrorPerLinearDifference : ErrorCorrection.ErrorPerLinearDifference;
-	const float ErrorPerAngularDiff = CharacterMovementCVars::ErrorPerAngularDifference >= 0.0f ? CharacterMovementCVars::ErrorPerAngularDifference : ErrorCorrection.ErrorPerAngularDifference;
-	const float MaxRestoredStateError = CharacterMovementCVars::MaxRestoredStateError >= 0.0f ? CharacterMovementCVars::MaxRestoredStateError : ErrorCorrection.MaxRestoredStateError;
-	const float ErrorAccumulationSeconds = CharacterMovementCVars::ErrorAccumulationSeconds >= 0.0f ? CharacterMovementCVars::ErrorAccumulationSeconds : ErrorCorrection.ErrorAccumulationSeconds;
-	const float ErrorAccumulationDistanceSq = CharacterMovementCVars::ErrorAccumulationDistanceSq >= 0.0f ? CharacterMovementCVars::ErrorAccumulationDistanceSq : ErrorCorrection.ErrorAccumulationDistanceSq;
-	const float ErrorAccumulationSimilarity = CharacterMovementCVars::ErrorAccumulationSimilarity >= 0.0f ? CharacterMovementCVars::ErrorAccumulationSimilarity : ErrorCorrection.ErrorAccumulationSimilarity;
-	const float PositionLerp = CharacterMovementCVars::PositionLerp >= 0.0f ? CharacterMovementCVars::PositionLerp : ErrorCorrection.PositionLerp;
-	const float LinearVelocityCoefficient = CharacterMovementCVars::LinearVelocityCoefficient >= 0.0f ? CharacterMovementCVars::LinearVelocityCoefficient : ErrorCorrection.LinearVelocityCoefficient;
-	const float AngleLerp = CharacterMovementCVars::AngleLerp >= 0.0f ? CharacterMovementCVars::AngleLerp : ErrorCorrection.AngleLerp;
-	const float AngularVelocityCoefficient = CharacterMovementCVars::AngularVelocityCoefficient >= 0.0f ? CharacterMovementCVars::AngularVelocityCoefficient : ErrorCorrection.AngularVelocityCoefficient;
-	const float MaxLinearHardSnapDistance = CharacterMovementCVars::MaxLinearHardSnapDistance >= 0.f ? CharacterMovementCVars::MaxLinearHardSnapDistance : ErrorCorrection.MaxLinearHardSnapDistance;
+	const float NetPingExtrapolation = PhysicsReplicationCVars::NetPingExtrapolation >= 0.0f ? PhysicsReplicationCVars::NetPingExtrapolation : ErrorCorrection.PingExtrapolation;
+	const float NetPingLimit = PhysicsReplicationCVars::NetPingLimit > 0.0f ? PhysicsReplicationCVars::NetPingLimit : ErrorCorrection.PingLimit;
+	const float ErrorPerLinearDiff = PhysicsReplicationCVars::ErrorPerLinearDifference >= 0.0f ? PhysicsReplicationCVars::ErrorPerLinearDifference : ErrorCorrection.ErrorPerLinearDifference;
+	const float ErrorPerAngularDiff = PhysicsReplicationCVars::ErrorPerAngularDifference >= 0.0f ? PhysicsReplicationCVars::ErrorPerAngularDifference : ErrorCorrection.ErrorPerAngularDifference;
+	const float MaxRestoredStateError = PhysicsReplicationCVars::MaxRestoredStateError >= 0.0f ? PhysicsReplicationCVars::MaxRestoredStateError : ErrorCorrection.MaxRestoredStateError;
+	const float ErrorAccumulationSeconds = PhysicsReplicationCVars::ErrorAccumulationSeconds >= 0.0f ? PhysicsReplicationCVars::ErrorAccumulationSeconds : ErrorCorrection.ErrorAccumulationSeconds;
+	const float ErrorAccumulationDistanceSq = PhysicsReplicationCVars::ErrorAccumulationDistanceSq >= 0.0f ? PhysicsReplicationCVars::ErrorAccumulationDistanceSq : ErrorCorrection.ErrorAccumulationDistanceSq;
+	const float ErrorAccumulationSimilarity = PhysicsReplicationCVars::ErrorAccumulationSimilarity >= 0.0f ? PhysicsReplicationCVars::ErrorAccumulationSimilarity : ErrorCorrection.ErrorAccumulationSimilarity;
+	const float PositionLerp = PhysicsReplicationCVars::PositionLerp >= 0.0f ? PhysicsReplicationCVars::PositionLerp : ErrorCorrection.PositionLerp;
+	const float LinearVelocityCoefficient = PhysicsReplicationCVars::LinearVelocityCoefficient >= 0.0f ? PhysicsReplicationCVars::LinearVelocityCoefficient : ErrorCorrection.LinearVelocityCoefficient;
+	const float AngleLerp = PhysicsReplicationCVars::AngleLerp >= 0.0f ? PhysicsReplicationCVars::AngleLerp : ErrorCorrection.AngleLerp;
+	const float AngularVelocityCoefficient = PhysicsReplicationCVars::AngularVelocityCoefficient >= 0.0f ? PhysicsReplicationCVars::AngularVelocityCoefficient : ErrorCorrection.AngularVelocityCoefficient;
+	const float MaxLinearHardSnapDistance = PhysicsReplicationCVars::MaxLinearHardSnapDistance >= 0.f ? PhysicsReplicationCVars::MaxLinearHardSnapDistance : ErrorCorrection.MaxLinearHardSnapDistance;
 
 	// Get Current state
 	FRigidBodyState CurrentState;
@@ -346,7 +346,7 @@ bool FPhysicsReplication::ApplyRigidBodyState(float DeltaSeconds, FBodyInstance*
 		const bool bHardSnap =
 			LinDiffSize > MaxLinearHardSnapDistance ||
 			PhysicsTarget.AccumulatedErrorSeconds > ErrorAccumulationSeconds ||
-			CharacterMovementCVars::AlwaysHardSnap;
+			PhysicsReplicationCVars::AlwaysHardSnap;
 
 		const FTransform IdealWorldTM(TargetQuat, TargetPos);
 
@@ -416,7 +416,7 @@ bool FPhysicsReplication::ApplyRigidBodyState(float DeltaSeconds, FBodyInstance*
 
 		// Should we show the async part?
 #if !UE_BUILD_SHIPPING
-		if (CharacterMovementCVars::NetShowCorrections != 0)
+		if (PhysicsReplicationCVars::NetShowCorrections != 0)
 		{
 			PhysicsTarget.ErrorHistory.bAutoAdjustMinMax = false;
 			PhysicsTarget.ErrorHistory.MinValue = 0.0f;
@@ -425,7 +425,7 @@ bool FPhysicsReplication::ApplyRigidBodyState(float DeltaSeconds, FBodyInstance*
 			if (UWorld* OwningWorld = GetOwningWorld())
 			{
 				FColor Color = FColor::White;
-				DrawDebugDirectionalArrow(OwningWorld, CurrentState.Position, TargetPos, 5.0f, Color, true, CharacterMovementCVars::NetCorrectionLifetime, 0, 1.5f);
+				DrawDebugDirectionalArrow(OwningWorld, CurrentState.Position, TargetPos, 5.0f, Color, true, PhysicsReplicationCVars::NetCorrectionLifetime, 0, 1.5f);
 				DrawDebugFloatHistory(*OwningWorld, PhysicsTarget.ErrorHistory, CurrentState.Position + FVector(0.0f, 0.0f, 100.0f), FVector2D(100.0f, 50.0f), FColor::White, false, 0, -1);
 			}
 		}
@@ -590,6 +590,11 @@ void FPhysicsReplication::OnTick(float DeltaSeconds, TMap<TWeakObjectPtr<UPrimit
 	CurAsyncData = nullptr;
 }
 
+bool FPhysicsReplication::ShouldSkipPhysicsReplication()
+{
+	return (PhysicsReplicationCVars::SkipPhysicsReplication != 0);
+}
+
 void FPhysicsReplication::Tick(float DeltaSeconds)
 {
 	OnTick(DeltaSeconds, ComponentToTargets);
@@ -610,10 +615,10 @@ FPhysicsReplication::FPhysicsReplication(FPhysScene* InPhysicsScene)
 void FPhysicsReplication::PrepareAsyncData_External(const FRigidBodyErrorCorrection& ErrorCorrection)
 {
 	//todo move this logic into a common function?
-	const float PositionLerp = CharacterMovementCVars::PositionLerp >= 0.0f ? CharacterMovementCVars::PositionLerp : ErrorCorrection.PositionLerp;
-	const float LinearVelocityCoefficient = CharacterMovementCVars::LinearVelocityCoefficient >= 0.0f ? CharacterMovementCVars::LinearVelocityCoefficient : ErrorCorrection.LinearVelocityCoefficient;
-	const float AngleLerp = CharacterMovementCVars::AngleLerp >= 0.0f ? CharacterMovementCVars::AngleLerp : ErrorCorrection.AngleLerp;
-	const float AngularVelocityCoefficient = CharacterMovementCVars::AngularVelocityCoefficient >= 0.0f ? CharacterMovementCVars::AngularVelocityCoefficient : ErrorCorrection.AngularVelocityCoefficient;
+	const float PositionLerp = PhysicsReplicationCVars::PositionLerp >= 0.0f ? PhysicsReplicationCVars::PositionLerp : ErrorCorrection.PositionLerp;
+	const float LinearVelocityCoefficient = PhysicsReplicationCVars::LinearVelocityCoefficient >= 0.0f ? PhysicsReplicationCVars::LinearVelocityCoefficient : ErrorCorrection.LinearVelocityCoefficient;
+	const float AngleLerp = PhysicsReplicationCVars::AngleLerp >= 0.0f ? PhysicsReplicationCVars::AngleLerp : ErrorCorrection.AngleLerp;
+	const float AngularVelocityCoefficient = PhysicsReplicationCVars::AngularVelocityCoefficient >= 0.0f ? PhysicsReplicationCVars::AngularVelocityCoefficient : ErrorCorrection.AngularVelocityCoefficient;
 
 	CurAsyncData = AsyncCallback->GetProducerInputData_External();
 	CurAsyncData->ErrorCorrection.PositionLerp = PositionLerp;
diff --git a/Engine/Source/Runtime/Engine/Public/PhysicsReplication.h b/Engine/Source/Runtime/Engine/Public/PhysicsReplication.h
--- a/Engine/Source/Runtime/Engine/Public/PhysicsReplication.h
+++ b/Engine/Source/Runtime/Engine/Public/PhysicsReplication.h
@@ -92,6 +92,9 @@ public:
 	FPhysicsReplication(FPhysScene* PhysScene);
 	virtual ~FPhysicsReplication();
 
+	/** Helper method so the Skip Replication CVar can be check elsewhere (including game extensions to this class) */
+	static bool ShouldSkipPhysicsReplication();
+
 	/** Tick and update all body states according to replicated targets */
 	void Tick(float DeltaSeconds);
 
