Inline small functions that are called often

diff --git a/Engine/Source/Runtime/Engine/Private/Streaming/TextureInstanceView.cpp b/Engine/Source/Runtime/Engine/Private/Streaming/TextureInstanceView.cpp
--- a/Engine/Source/Runtime/Engine/Private/Streaming/TextureInstanceView.cpp
+++ b/Engine/Source/Runtime/Engine/Private/Streaming/TextureInstanceView.cpp
@@ -620,18 +620,3 @@ void FRenderAssetInstanceAsyncView::GetRenderAssetScreenSize(
 		}
 	}
 }
-
-bool FRenderAssetInstanceAsyncView::HasRenderAssetReferences(const UStreamableRenderAsset* InAsset) const
-{
-	return View.IsValid() && (bool)View->GetElementIterator(InAsset);
-}
-
-bool FRenderAssetInstanceAsyncView::HasComponentWithForcedLOD(const UStreamableRenderAsset* InAsset) const
-{
-	return View.IsValid() && View->HasComponentWithForcedLOD(InAsset);
-}
-
-bool FRenderAssetInstanceAsyncView::HasAnyComponentWithForcedLOD() const
-{
-	return View.IsValid() && View->HasAnyComponentWithForcedLOD();
-}
diff --git a/Engine/Source/Runtime/Engine/Private/Streaming/TextureInstanceView.h b/Engine/Source/Runtime/Engine/Private/Streaming/TextureInstanceView.h
--- a/Engine/Source/Runtime/Engine/Private/Streaming/TextureInstanceView.h
+++ b/Engine/Source/Runtime/Engine/Private/Streaming/TextureInstanceView.h
@@ -287,11 +287,20 @@ public:
 		int32& MaxNumForcedLODs,
 		const TCHAR* LogPrefix) const;
 
-	bool HasRenderAssetReferences(const UStreamableRenderAsset* InAsset) const;
+	FORCEINLINE bool HasRenderAssetReferences(const UStreamableRenderAsset* InAsset) const
+	{
+		return View.IsValid() && (bool)View->GetElementIterator(InAsset);
+	}
 
-	bool HasComponentWithForcedLOD(const UStreamableRenderAsset* InAsset) const;
+	FORCEINLINE bool HasComponentWithForcedLOD(const UStreamableRenderAsset* InAsset) const
+	{
+		return View.IsValid() && View->HasComponentWithForcedLOD(InAsset);
+	}
 
-	bool HasAnyComponentWithForcedLOD() const;
+	FORCEINLINE bool HasAnyComponentWithForcedLOD() const
+	{
+		return View.IsValid() && View->HasAnyComponentWithForcedLOD();
+	}
 
 	// Release the data now as this is expensive.
 	void OnTaskDone() { BoundsViewInfo.Empty(); }
