Fixed async uniform expression cache updates. Added support for nesting use of the async update scope. Added a wait when uniform expression caches are invalidated.

diff --git a/Engine/Source/Runtime/Engine/Private/Materials/MaterialShared.cpp b/Engine/Source/Runtime/Engine/Private/Materials/MaterialShared.cpp
--- a/Engine/Source/Runtime/Engine/Private/Materials/MaterialShared.cpp
+++ b/Engine/Source/Runtime/Engine/Private/Materials/MaterialShared.cpp
@@ -86,7 +86,7 @@ FAutoConsoleVariableRef CVarDeferUniformExpressionCaching(
 	ECVF_RenderThreadSafe
 	);
 
-int32 GUniformExpressionCacheAsyncUpdates = 0;
+int32 GUniformExpressionCacheAsyncUpdates = 1;
 FAutoConsoleVariableRef CVarUniformExpressionCacheAsyncUpdates(
 	TEXT("r.UniformExpressionCacheAsyncUpdates"),
 	GUniformExpressionCacheAsyncUpdates,
@@ -3582,18 +3582,21 @@ class FUniformExpressionCacheAsyncUpdateTask
 public:
 	void Begin()
 	{
-		bEnabled = true;
+		ReferenceCount++;
 	}
 
 	void End()
 	{
-		bEnabled = false;
-		Wait();
+		check(ReferenceCount > 0);
+		if (--ReferenceCount == 0)
+		{
+			Wait();
+		}
 	}
 
 	bool IsEnabled() const
 	{
-		return bEnabled && GUniformExpressionCacheAsyncUpdates > 0 && !GRHICommandList.Bypass();
+		return ReferenceCount > 0 && GUniformExpressionCacheAsyncUpdates > 0 && !GRHICommandList.Bypass();
 	}
 
 	void SetTask(const FGraphEventRef& InTask)
@@ -3614,7 +3617,7 @@ public:
 
 private:
 	FGraphEventRef Task;
-	bool bEnabled = false;
+	int32 ReferenceCount = 0;
 
 } GUniformExpressionCacheAsyncUpdateTask;
 
@@ -3855,6 +3858,7 @@ void FMaterialRenderProxy::CacheUniformExpressions_GameThread(bool bRecreateUnif
 void FMaterialRenderProxy::InvalidateUniformExpressionCache(bool bRecreateUniformBuffer)
 {
 	check(IsInRenderingThread());
+	GUniformExpressionCacheAsyncUpdateTask.Wait();
 
 #if WITH_EDITOR
 	FStaticLightingSystemInterface::OnMaterialInvalidated.Broadcast(this);
@@ -3909,6 +3913,12 @@ FMaterialRenderProxy::FMaterialRenderProxy(FString InMaterialName)
 
 FMaterialRenderProxy::~FMaterialRenderProxy()
 {
+	// We only wait on deletions happening on the render thread. Async deletions can happen during scene render shutdown and those are waited on explicitly.
+	if (IsInRenderingThread())
+	{
+		GUniformExpressionCacheAsyncUpdateTask.Wait();
+	}
+
 	if(IsInitialized())
 	{
 		check(IsInRenderingThread());
