[Engine Runtime]
* UnrealEngine - Moved UE_SCOPED_ENGINE_ACTIVITY for better coverage of map loading

diff --git a/Engine/Source/Runtime/Engine/Private/UnrealEngine.cpp b/Engine/Source/Runtime/Engine/Private/UnrealEngine.cpp
--- a/Engine/Source/Runtime/Engine/Private/UnrealEngine.cpp
+++ b/Engine/Source/Runtime/Engine/Private/UnrealEngine.cpp
@@ -14248,6 +14248,7 @@ void UEngine::TickWorldTravel(FWorldContext& Context, float DeltaSeconds)
 
 bool UEngine::LoadMap( FWorldContext& WorldContext, FURL URL, class UPendingNetGame* Pending, FString& Error )
 {
+	UE_SCOPED_ENGINE_ACTIVITY(TEXT("Loading Map %s"), *URL.Map);
 	TRACE_LOADTIME_REQUEST_GROUP_SCOPE(TEXT("LoadMap - %s"), *URL.Map);
 	STAT_ADD_CUSTOMMESSAGE_NAME( STAT_NamedMarker, *(FString( TEXT( "LoadMap - " ) + URL.Map )) );
 	TRACE_BOOKMARK(TEXT("LoadMap - %s"), *URL.Map);
@@ -14559,7 +14560,6 @@ bool UEngine::LoadMap( FWorldContext& WorldContext, FURL URL, class UPendingNetG
 		// If the level isn't already in memory, load level from disk
 		if (WorldPackage == nullptr)
 		{
-			UE_SCOPED_ENGINE_ACTIVITY(TEXT("Loading World %s"), *URL.Map);
 			WorldPackage = LoadPackage(nullptr, *URL.Map, (WorldContext.WorldType == EWorldType::PIE ? LOAD_PackageForPIE : LOAD_None));
 		}
 
