Minor cleanup

diff --git a/Engine/Source/Runtime/Engine/Public/StaticMeshResources.h b/Engine/Source/Runtime/Engine/Public/StaticMeshResources.h
--- a/Engine/Source/Runtime/Engine/Public/StaticMeshResources.h
+++ b/Engine/Source/Runtime/Engine/Public/StaticMeshResources.h
@@ -931,11 +931,11 @@ protected:
 		{
 			/** Default constructor. */
 			FSectionInfo()
-				: Material(NULL)
-#if WITH_EDITOR
+				: Material(nullptr)
+			#if WITH_EDITOR
 				, bSelected(false)
-				, HitProxy(NULL)
-#endif
+				, HitProxy(nullptr)
+			#endif
 				, FirstPreCulledIndex(0)
 				, NumPreCulledTriangles(-1)
 			{}
@@ -943,18 +943,18 @@ protected:
 			/** The material with which to render this section. */
 			UMaterialInterface* Material;
 
-#if WITH_EDITOR
+		#if WITH_EDITOR
 			/** True if this section should be rendered as selected (editor only). */
 			bool bSelected;
 
 			/** The editor needs to be able to individual sub-mesh hit detection, so we store a hit proxy on each mesh. */
 			HHitProxy* HitProxy;
-#endif
+		#endif
 
-#if WITH_EDITORONLY_DATA
+		#if WITH_EDITORONLY_DATA
 			// The material index from the component. Used by the texture streaming accuracy viewmodes.
 			int32 MaterialIndex;
-#endif
+		#endif
 
 			int32 FirstPreCulledIndex;
 			int32 NumPreCulledTriangles;
