Fix ActorChannel compil issue

diff --git a/Engine/Source/Runtime/Engine/Classes/Engine/ActorChannel.h b/Engine/Source/Runtime/Engine/Classes/Engine/ActorChannel.h
--- a/Engine/Source/Runtime/Engine/Classes/Engine/ActorChannel.h
+++ b/Engine/Source/Runtime/Engine/Classes/Engine/ActorChannel.h
@@ -171,7 +171,7 @@ public:
 	int64 ReplicateActor();
 
 	/** Tells if the actor is ready to be replicated since he is BeginPlay or inside BeginPlay */
-	bool IsActorReadyForReplication() const { return (Actor ? (Actor->HasActorBegunPlay() || Actor->IsActorBeginningPlay()) : false); }
+	bool IsActorReadyForReplication() const;
 
 	/**
 	 * Set this channel's actor to the given actor.
diff --git a/Engine/Source/Runtime/Engine/Private/DataChannel.cpp b/Engine/Source/Runtime/Engine/Private/DataChannel.cpp
--- a/Engine/Source/Runtime/Engine/Private/DataChannel.cpp
+++ b/Engine/Source/Runtime/Engine/Private/DataChannel.cpp
@@ -3960,6 +3960,11 @@ void UActorChannel::PrepareForRemoteFunction(UObject* TargetObj)
 	}
 }
 
+bool UActorChannel::IsActorReadyForReplication() const
+{ 
+	return (Actor ? (Actor->HasActorBegunPlay() || Actor->IsActorBeginningPlay()) : false); 
+}
+
 void UActorChannel::QueueRemoteFunctionBunch( UObject* CallTarget, UFunction* Func, FOutBunch &Bunch )
 {
 	FindOrCreateReplicator(CallTarget).Get().QueueRemoteFunctionBunch( Func, Bunch );
