Removed csv stat for Exclusive/GameThread/ReplicateActor and added Exclusive/GameThread/ReplayReplicateActors. ReplayReplicateActors and ServerReplicateActors should add up to the total ReplicateActor time (plus some other overhead).

diff --git a/Engine/Source/Runtime/Engine/Private/DataChannel.cpp b/Engine/Source/Runtime/Engine/Private/DataChannel.cpp
--- a/Engine/Source/Runtime/Engine/Private/DataChannel.cpp
+++ b/Engine/Source/Runtime/Engine/Private/DataChannel.cpp
@@ -3196,8 +3196,6 @@ int64 UActorChannel::ReplicateActor()
 #endif
 
 	const bool bReplay = Connection->IsReplay();
-	CSV_SCOPED_TIMING_STAT_EXCLUSIVE_CONDITIONAL(ReplicateActor, !bReplay);
-
 	const bool bEnableScopedCycleCounter = !bReplay && GReplicateActorTimingEnabled;
 	FSimpleScopeSecondsCounter ScopedSecondsCounter(GReplicateActorTimeSeconds, bEnableScopedCycleCounter);
 
diff --git a/Engine/Source/Runtime/Engine/Private/DemoNetDriver.cpp b/Engine/Source/Runtime/Engine/Private/DemoNetDriver.cpp
--- a/Engine/Source/Runtime/Engine/Private/DemoNetDriver.cpp
+++ b/Engine/Source/Runtime/Engine/Private/DemoNetDriver.cpp
@@ -2138,6 +2138,7 @@ void UDemoNetDriver::TickDemoRecordFrame(float DeltaSeconds)
 
 	{
 		DECLARE_SCOPE_CYCLE_COUNTER(TEXT("Replay actor replication time"), STAT_ReplayReplicateActors, STATGROUP_Net);
+		CSV_SCOPED_TIMING_STAT_EXCLUSIVE(ReplayReplicateActors);
 
 		if (HasLevelStreamingFixes())
 		{
