For components which won't be uninitialized, detach cross-level parents/children on end play

diff --git a/Engine/Source/Runtime/Engine/Classes/Components/SceneComponent.h b/Engine/Source/Runtime/Engine/Classes/Components/SceneComponent.h
--- a/Engine/Source/Runtime/Engine/Classes/Components/SceneComponent.h
+++ b/Engine/Source/Runtime/Engine/Classes/Components/SceneComponent.h
@@ -905,6 +905,7 @@ public:
 	virtual void GetLifetimeReplicatedProps(TArray<FLifetimeProperty>& OutLifetimeProps) const override;
 	virtual void OnRegister() override;
 	virtual void OnUnregister() override;
+	virtual void EndPlay(EEndPlayReason::Type Reason) override;
 	virtual bool ShouldCreateRenderState() const override { return true; }
 	virtual void UpdateComponentToWorld(EUpdateTransformFlags UpdateTransformFlags = EUpdateTransformFlags::None, ETeleportType Teleport = ETeleportType::None) override final
 	{
diff --git a/Engine/Source/Runtime/Engine/Private/Components/SceneComponent.cpp b/Engine/Source/Runtime/Engine/Private/Components/SceneComponent.cpp
--- a/Engine/Source/Runtime/Engine/Private/Components/SceneComponent.cpp
+++ b/Engine/Source/Runtime/Engine/Private/Components/SceneComponent.cpp
@@ -750,6 +750,36 @@ void USceneComponent::OnUnregister()
 	Super::OnUnregister();
 }
 
+void USceneComponent::EndPlay(EEndPlayReason::Type Reason)
+{
+	Super::EndPlay(Reason);
+
+	if (Reason == EEndPlayReason::RemovedFromWorld && !HasBeenInitialized())
+	{
+		// Detach components which are in different streaming levels so that this level can be properly garbage collected.
+		// Note that we explicitly want to check the outer hierarchy and not the owning package because we want references that participate in GC.
+		UObject* Outermost = GetOutermostObject();
+		if (AttachParent && AttachParent->GetOutermostObject() != Outermost)
+		{
+			DetachFromComponent(FDetachmentTransformRules::KeepWorldTransform);
+		}
+
+		TInlineComponentArray<USceneComponent*> ChildrenToDetach;
+		for (int32 i = 0; i < AttachChildren.Num(); ++i)
+		{
+			if (USceneComponent* AttachChild = AttachChildren[i].Get(); AttachChild && AttachChild->GetOutermostObject() != Outermost)
+			{
+				ChildrenToDetach.Add(AttachChild);
+			}
+		}
+
+		for (USceneComponent* Child : ChildrenToDetach)
+		{
+			Child->DetachFromComponent(FDetachmentTransformRules::KeepWorldTransform);
+		}
+	}
+}
+
 void USceneComponent::PropagateTransformUpdate(bool bTransformChanged, EUpdateTransformFlags UpdateTransformFlags, ETeleportType Teleport)
 {
 	//QUICK_SCOPE_CYCLE_COUNTER(STAT_USceneComponent_PropagateTransformUpdate);
