Fixed a warning saying: "Invalid Simulate Options: Body (Character0.Arrow) is set to simulate physics but Collision Enabled is incompatible".

diff --git a/Engine/Source/Runtime/Engine/Classes/GameFramework/Character.h b/Engine/Source/Runtime/Engine/Classes/GameFramework/Character.h
--- a/Engine/Source/Runtime/Engine/Classes/GameFramework/Character.h
+++ b/Engine/Source/Runtime/Engine/Classes/GameFramework/Character.h
@@ -584,6 +584,10 @@ public:
 	/** Incremented every time there is an Actor overlap event (start or stop) on this actor. */
 	uint32 NumActorOverlapEventsCounter;
 
+	//~ Begin UObject Interface.
+	virtual void PostLoad() override;
+	//~ End UObject Interface
+
 	//~ Begin AActor Interface.
 	virtual void BeginPlay() override;
 	virtual void ClearCrossLevelReferences() override;
diff --git a/Engine/Source/Runtime/Engine/Private/Character.cpp b/Engine/Source/Runtime/Engine/Private/Character.cpp
--- a/Engine/Source/Runtime/Engine/Private/Character.cpp
+++ b/Engine/Source/Runtime/Engine/Private/Character.cpp
@@ -79,6 +79,7 @@ ACharacter::ACharacter(const FObjectInitializer& ObjectInitializer)
 		ArrowComponent->SpriteInfo.DisplayName = ConstructorStatics.NAME_Characters;
 		ArrowComponent->SetupAttachment(CapsuleComponent);
 		ArrowComponent->bIsScreenSizeScaled = true;
+		ArrowComponent->SetSimulatePhysics(false);
 	}
 #endif // WITH_EDITORONLY_DATA
 
@@ -144,6 +145,18 @@ void ACharacter::PostInitializeComponents()
 	}
 }
 
+void ACharacter::PostLoad()
+{
+	Super::PostLoad();
+
+#if WITH_EDITORONLY_DATA
+	if (ArrowComponent)
+	{
+		ArrowComponent->SetSimulatePhysics(false);
+	}
+#endif // WITH_EDITORONLY_DATA
+}
+
 void ACharacter::BeginPlay()
 {
 	Super::BeginPlay();
