Fixed issue regarding the refraction attribute when Make/Set MaterialAttributes nodes are used.

Issue was caused by error in logic that detects whether the material applies distortion. Old logic was checking only for terminal output node to write to the Refraction attribute, while the node that write to the Refraction attribute might be somewhere in the subgraph that eventaully leads to that terminal node. The solution was to look traverse the entire subgraph looking for any node that sets the Refraction attribute.

diff --git a/Engine/Source/Runtime/Engine/Private/Materials/Material.cpp b/Engine/Source/Runtime/Engine/Private/Materials/Material.cpp
--- a/Engine/Source/Runtime/Engine/Private/Materials/Material.cpp
+++ b/Engine/Source/Runtime/Engine/Private/Materials/Material.cpp
@@ -4209,6 +4209,7 @@ void UMaterial::PostEditChangePropertyInternal(FPropertyChangedEvent& PropertyCh
 	// check for distortion in material 
 	{
 		bUsesDistortion = false;
+
 		// check for a distortion value
 		if (EditorOnly->Refraction.Expression
 			|| (EditorOnly->Refraction.UseConstant && FMath::Abs(EditorOnly->Refraction.Constant - 1.0f) >= UE_KINDA_SMALL_NUMBER))
@@ -4219,24 +4220,33 @@ void UMaterial::PostEditChangePropertyInternal(FPropertyChangedEvent& PropertyCh
 		// check the material attributes for refraction expressions as well
 		if (EditorOnly->MaterialAttributes.Expression)
 		{
-			// handle make attribute expressions
-			UMaterialExpressionMakeMaterialAttributes * MakeAttributeExpression = Cast<UMaterialExpressionMakeMaterialAttributes>(EditorOnly->MaterialAttributes.Expression);
-			if (MakeAttributeExpression && MakeAttributeExpression->Refraction.Expression)
-			{
-				bUsesDistortion = true;
-			}
+			// The node that possibly sets the refraction attribute can be anywhere in the sub-graph that eventually connects
+			// to the material output expression. We need to look for any node of type [Make/Set]MaterialAttributes in this subgraph
+			// and check whether it writes to the Refraction attribute.
+			TArray<UMaterialExpression*> Expressions;
+			EditorOnly->MaterialAttributes.Expression->GetAllInputExpressions(Expressions);
 
-			// handle set attribute expressions
-			UMaterialExpressionSetMaterialAttributes * SetAttributeExpression = Cast<UMaterialExpressionSetMaterialAttributes>(EditorOnly->MaterialAttributes.Expression);
-			if (SetAttributeExpression)
+			for (int32 ExprIndex = 0; ExprIndex < Expressions.Num() && !bUsesDistortion; ++ExprIndex)
 			{
-				for (int32 Index = 0; Index < SetAttributeExpression->Inputs.Num(); Index++)
+				// Case when node is a MakeMaterialAttributes
+				UMaterialExpressionMakeMaterialAttributes* MakeAttributeExpression = Cast<UMaterialExpressionMakeMaterialAttributes>(Expressions[ExprIndex]);
+				if (MakeAttributeExpression && MakeAttributeExpression->Refraction.Expression)
+				{
+					bUsesDistortion = true;
+				}
+
+				// Case when node is a SetMaterialAttributes
+				UMaterialExpressionSetMaterialAttributes* SetAttributeExpression = Cast<UMaterialExpressionSetMaterialAttributes>(Expressions[ExprIndex]);
+				if (!bUsesDistortion && SetAttributeExpression)
 				{
-					FExpressionInput & Input = SetAttributeExpression->Inputs[Index];
-					FName InputName = SetAttributeExpression->GetInputName(Index);
-					if (InputName == TEXT("Refraction"))
+					for (int32 InputIndex = 0; InputIndex < SetAttributeExpression->Inputs.Num(); InputIndex++)
 					{
-						bUsesDistortion = true;
+						FName InputName = SetAttributeExpression->GetInputName(InputIndex);
+						if (InputName == TEXT("Refraction"))
+						{
+							bUsesDistortion = true;
+							break;
+						}
 					}
 				}
 			}
