Set r.Mobile.DisableVertexFog=1 as default and remove this setting from project settings

diff --git a/Engine/Source/Runtime/Engine/Classes/Engine/RendererSettings.h b/Engine/Source/Runtime/Engine/Classes/Engine/RendererSettings.h
--- a/Engine/Source/Runtime/Engine/Classes/Engine/RendererSettings.h
+++ b/Engine/Source/Runtime/Engine/Classes/Engine/RendererSettings.h
@@ -335,12 +335,6 @@ class ENGINE_API URendererSettings : public UDeveloperSettings
 		ToolTip = "Whether to use the Reflection Capture Compression or not for mobile. It will use ETC2 format to do the compression."))
 	uint32 bReflectionCaptureCompression : 1;
 
-	UPROPERTY(config, EditAnywhere, Category = Mobile, meta = (
-		ConsoleVariable = "r.Mobile.DisableVertexFog", DisplayName = "Disable vertex fogging in mobile shaders",
-		ToolTip = "If true, vertex fog will be omitted from the most of the mobile base pass shaders. Instead, fog will be applied in a separate pass and only when scene has a fog component.",
-		ConfigRestartRequired = true))
-	uint32 bMobileDisableVertexFog : 1;
-
 	UPROPERTY(config, EditAnywhere, Category = Materials, meta = (
 		ConsoleVariable = "r.DiscardUnusedQuality", DisplayName = "Game Discards Unused Material Quality Levels",
 		ToolTip = "When running in game mode, whether to keep shaders for all quality levels in memory or only those needed for the current quality level.\nUnchecked: Keep all quality levels in memory allowing a runtime quality level change. (default)\nChecked: Discard unused quality levels when loading content for the game, saving some memory."))
diff --git a/Engine/Source/Runtime/Renderer/Private/MobileBasePassRendering.cpp b/Engine/Source/Runtime/Renderer/Private/MobileBasePassRendering.cpp
--- a/Engine/Source/Runtime/Renderer/Private/MobileBasePassRendering.cpp
+++ b/Engine/Source/Runtime/Renderer/Private/MobileBasePassRendering.cpp
@@ -24,8 +24,8 @@
 // Changing this causes a full shader recompile
 static TAutoConsoleVariable<int32> CVarMobileDisableVertexFog(
 	TEXT("r.Mobile.DisableVertexFog"),
-	0,
-	TEXT("Set to 1 to disable vertex fogging in all mobile shaders."),
+	1,
+	TEXT("If true, vertex fog will be omitted from the most of the mobile base pass shaders. Instead, fog will be applied in a separate pass and only when scene has a fog component."),
 	ECVF_ReadOnly | ECVF_RenderThreadSafe);
 
 static TAutoConsoleVariable<int32> CVarMobileEnableMovableSpotLightShadows(
