Animation: Add Setter function for deprecated SkeletalMesh property.

The SkeletalMesh property was deprecated and the SetSkeletalMesh function was removed. This broke Sequencer evaluation of animated SkeletalMesh properties because it could no longer find the SetSkeletalMesh function. This fixes the issue by using the new Setter functionality. A new deprecated function called SetSkeletalMesh_DEPRECATED has been added and is the Setter function for the deprecated SkeletalMesh property.

The SetSkeletalMesh_DEPRECATED function can be removed when the SkeletalMesh property is removed.

diff --git a/Engine/Source/Runtime/Engine/Classes/Components/SkinnedMeshComponent.h b/Engine/Source/Runtime/Engine/Classes/Components/SkinnedMeshComponent.h
--- a/Engine/Source/Runtime/Engine/Classes/Components/SkinnedMeshComponent.h
+++ b/Engine/Source/Runtime/Engine/Classes/Components/SkinnedMeshComponent.h
@@ -243,7 +243,7 @@ class ENGINE_API USkinnedMeshComponent : public UMeshComponent, public ILODSyncI
 
 	/** The skeletal mesh used by this component. */
 	UE_DEPRECATED(5.1, "Replaced by SkinnedAsset. Use GetSkinnedAsset()/SetSkinnedAsset() instead, or GetSkeletalMeshAsset/SetSkeletalMeshAsset() when called from a USkeletalMeshComponent.")
-	UPROPERTY(EditAnywhere, BlueprintGetter = GetSkeletalMesh_DEPRECATED, Category = "Mesh|SkeletalAsset", meta = (DisallowedClasses = "/Script/ApexDestruction.DestructibleMesh", DeprecatedProperty, DeprecationMessage = "Use USkeletalMeshComponent::GetSkeletalMeshAsset() or GetSkinnedAsset() instead."))
+	UPROPERTY(EditAnywhere, Setter = SetSkeletalMesh_DEPRECATED, BlueprintGetter = GetSkeletalMesh_DEPRECATED, Category = "Mesh|SkeletalAsset", meta = (DisallowedClasses = "/Script/ApexDestruction.DestructibleMesh", DeprecatedProperty, DeprecationMessage = "Use USkeletalMeshComponent::GetSkeletalMeshAsset() or GetSkinnedAsset() instead."))
 	TObjectPtr<class USkeletalMesh> SkeletalMesh;
 
 private:
@@ -905,6 +905,10 @@ public:
 	UFUNCTION(BlueprintPure, Category = "Components|SkinnedMesh", DisplayName = "Get Skeletal Mesh", meta = (DeprecatedFunction, DeprecationMessage = "Use USkeletalMeshComponent::GetSkeletalMeshAsset() or GetSkinnedAsset() instead."))
 	class USkeletalMesh* GetSkeletalMesh_DEPRECATED() const;
 
+	UE_DEPRECATED(5.1, "Use USkeletalMeshComponent::SetSkinnedAssetAndUpdate() instead.")
+	UFUNCTION(BlueprintCallable, Category = "Components|SkinnedMesh", DisplayName = "Set Skeletal Mesh")
+	void SetSkeletalMesh_DEPRECATED(class USkeletalMesh* NewMesh) { SetSkinnedAssetAndUpdate(Cast<USkinnedAsset>(NewMesh)); }
+
 	/**
 	 * Change the SkinnedAsset that is rendered for this Component. Will re-initialize the animation tree etc.
 	 *
