Shared Linear Texture Encoding [DDC2] -- Separates out the platform tiling code to allow for texture builds to reuse the linear texture that gets encodes with windows.

Primary Components:
    DDC2 texture build input/output revamped to support future FDerivedData references where the bulk data remains in the Zen server during cooks.
    Tiling code is separated out in to ITextureTiler interface so that tiling can be done without calling FChildTextureFormat::CompressImage[Ex]
    Adding mip dimensions to GetMipCountForBuildSettings

diff --git a/Engine/Source/Developer/TextureBuild/Private/TextureBuildFunction.cpp b/Engine/Source/Developer/TextureBuild/Private/TextureBuildFunction.cpp
--- a/Engine/Source/Developer/TextureBuild/Private/TextureBuildFunction.cpp
+++ b/Engine/Source/Developer/TextureBuild/Private/TextureBuildFunction.cpp
@@ -19,6 +19,7 @@
 #include "Serialization/CompactBinaryWriter.h"
 #include "Serialization/FileRegions.h"
 #include "Serialization/MemoryWriter.h"
+#include "TextureBuildUtilities.h"
 #include "TextureCompressorModule.h"
 #include "TextureFormatManager.h"
 
@@ -28,7 +29,7 @@ DEFINE_LOG_CATEGORY_STATIC(LogTextureBuildFunction, Log, All);
 // MUST have a corresponding change to this version. Individual texture formats have a version to
 // change that is specific to the format. A merge conflict affecting the version MUST be resolved
 // by generating a new version.
-static const FGuid TextureDerivedDataVersion(TEXT("19ffd7ab-ea5c-4a18-b6f4-6b0b5eedd606"));
+static const FGuid TextureDerivedDataVersion(TEXT("9C3A0F67-3973-4F4D-85FB-35B44FEA8089"));
 
 
 static void ReadCbField(FCbFieldView Field, bool& OutValue) { OutValue = Field.AsBool(OutValue); }
@@ -137,7 +138,7 @@ static FTextureBuildSettings ReadBuildSettingsFromCompactBinary(const FCbObjectV
 	ReadCbField(Object["TopMipSize"], BuildSettings.TopMipSize);
 	ReadCbField(Object["VolumeSizeZ"], BuildSettings.VolumeSizeZ);
 	ReadCbField(Object["ArraySlices"], BuildSettings.ArraySlices);
-	BuildSettings.bStreamable = Object["bStreamable"].AsBool(BuildSettings.bStreamable);
+	BuildSettings.bStreamable_Unused = Object["bStreamable"].AsBool(BuildSettings.bStreamable_Unused);
 	BuildSettings.bVirtualStreamable = Object["bVirtualStreamable"].AsBool(BuildSettings.bVirtualStreamable);
 	BuildSettings.bChromaKeyTexture = Object["bChromaKeyTexture"].AsBool(BuildSettings.bChromaKeyTexture);
 	ReadCbField(Object["PowerOfTwoMode"], BuildSettings.PowerOfTwoMode);
@@ -161,11 +162,6 @@ static FTextureBuildSettings ReadBuildSettingsFromCompactBinary(const FCbObjectV
 	return BuildSettings;
 }
 
-static void ReadOutputSettingsFromCompactBinary(const FCbObjectView& Object, int32& NumInlineMips)
-{
-	NumInlineMips = Object["NumInlineMips"].AsInt32();
-}
-
 static ERawImageFormat::Type ComputeRawImageFormat(ETextureSourceFormat SourceFormat)
 {
 	return FImageCoreUtils::ConvertToRawImageFormat(SourceFormat);
@@ -318,20 +314,32 @@ void FTextureBuildFunction::Build(UE::DerivedData::FBuildContext& Context) const
 	const FTextureBuildSettings BuildSettings = ReadBuildSettingsFromCompactBinary(Settings["Build"].AsObjectView());
 	
 	const uint16 RequiredTextureFormatVersion = Settings["FormatVersion"].AsUInt16();
+	const ITextureFormat* TextureFormat;
 	if (ITextureFormatManagerModule* TFM = GetTextureFormatManager())
 	{
-		const ITextureFormat* TextureFormat = TFM->FindTextureFormat(BuildSettings.TextureFormatName);
-		const uint16 CurrentTextureFormatVersion = TextureFormat ? TextureFormat->GetVersion(BuildSettings.TextureFormatName, &BuildSettings) : 0;
-		if (CurrentTextureFormatVersion != RequiredTextureFormatVersion)
-		{
-			UE_LOG(LogTextureBuildFunction, Error, TEXT("%s has version %hu when version %hu is required."),
-				*BuildSettings.TextureFormatName.ToString(), CurrentTextureFormatVersion, RequiredTextureFormatVersion);;
-			return;
-		}
+		TextureFormat = TFM->FindTextureFormat(BuildSettings.TextureFormatName);
+	}
+	else
+	{
+		UE_LOG(LogTextureBuildFunction, Error, TEXT("TextureFormatManager not found!"));
+		return;
+	}
+
+	const uint16 CurrentTextureFormatVersion = TextureFormat ? TextureFormat->GetVersion(BuildSettings.TextureFormatName, &BuildSettings) : 0;
+	if (CurrentTextureFormatVersion != RequiredTextureFormatVersion)
+	{
+		UE_LOG(LogTextureBuildFunction, Error, TEXT("%s has version %hu when version %hu is required."),
+			*BuildSettings.TextureFormatName.ToString(), CurrentTextureFormatVersion, RequiredTextureFormatVersion);;
+		return;
+	}
+	
+	FTextureEngineParameters EngineParameters;
+	if (UE::TextureBuildUtilities::TextureEngineParameters::FromCompactBinary(EngineParameters, Context.FindConstant(UTF8TEXTVIEW("EngineParameters"))) == false)
+	{
+		UE_LOG(LogTextureBuildFunction, Error, TEXT("Engine parameters are not available."));
+		return;
 	}
 
-	int32 NumInlineMips;
-	ReadOutputSettingsFromCompactBinary(Settings["Output"].AsObjectView(), NumInlineMips);
 
 	TArray<FImage> SourceMips;
 	if (!TryReadTextureSourceFromCompactBinary(Settings["Source"], Context,BuildSettings, SourceMips))
@@ -348,9 +356,13 @@ void FTextureBuildFunction::Build(UE::DerivedData::FBuildContext& Context) const
 
 	UE_LOG(LogTextureBuildFunction, Display, TEXT("Compressing %d source mip(s) (%dx%d) to %s..."), SourceMips.Num(), SourceMips[0].SizeX, SourceMips[0].SizeY, *BuildSettings.TextureFormatName.ToString());
 
+	ITextureCompressorModule& TextureCompressorModule = FModuleManager::GetModuleChecked<ITextureCompressorModule>(TEXTURE_COMPRESSOR_MODULENAME);
+	
 	TArray<FCompressedImage2D> CompressedMips;
 	uint32 NumMipsInTail;
 	uint32 ExtData;
+	bool bImageHasAlpha = false;
+	
 	
 	/*
 	for( const FImage & Image : SourceMips )
@@ -360,106 +372,212 @@ void FTextureBuildFunction::Build(UE::DerivedData::FBuildContext& Context) const
 	}
 	*/
 
-	bool bBuildSucceeded = FModuleManager::GetModuleChecked<ITextureCompressorModule>(TEXTURE_COMPRESSOR_MODULENAME).BuildTexture(
+	bool bBuildSucceeded = TextureCompressorModule.BuildTexture(
 		SourceMips,
 		AssociatedNormalSourceMips,
 		BuildSettings,
 		Context.GetName(),
 		CompressedMips,
 		NumMipsInTail,
-		ExtData);
+		ExtData,
+		&bImageHasAlpha);
 	if (!bBuildSucceeded)
 	{
 		return;
 	}
 	check(CompressedMips.Num() > 0);
 
-	int32 MipCount = CompressedMips.Num();
-	const bool bForceAllMipsToBeInlined = BuildSettings.bCubemap || (BuildSettings.bVolume && !BuildSettings.bStreamable) || (BuildSettings.bTextureArray && !BuildSettings.bStreamable);
-	const int32 FirstInlineMip = bForceAllMipsToBeInlined ? 0 : FMath::Max(0, MipCount - FMath::Max(NumInlineMips, (int32)NumMipsInTail));
+
+	FEncodedTextureDescription TextureDescription;
 
 	{
-		FCbWriter DescriptionWriter;
-		DescriptionWriter.BeginObject();
-
-		DescriptionWriter.BeginArray(ANSITEXTVIEW("Size"));
-		DescriptionWriter.AddInteger(CompressedMips[0].SizeX);
-		DescriptionWriter.AddInteger(CompressedMips[0].SizeY);
-		// this is brittle ; CompressedMips[] should report NumSlices itself
-		const int32 NumSlices = (BuildSettings.bVolume || BuildSettings.bTextureArray) ? CompressedMips[0].SizeZ : BuildSettings.bCubemap ? 6 : 1;
-		DescriptionWriter.AddInteger(NumSlices);
-		DescriptionWriter.EndArray();
-
-		DescriptionWriter.AddString(ANSITEXTVIEW("PixelFormat"), GetPixelFormatString((EPixelFormat)CompressedMips[0].PixelFormat));
-		DescriptionWriter.AddBool(ANSITEXTVIEW("bCubeMap"), BuildSettings.bCubemap);
-		DescriptionWriter.AddInteger(ANSITEXTVIEW("ExtData"), ExtData);
-		DescriptionWriter.AddInteger(ANSITEXTVIEW("NumMips"), MipCount);
-		DescriptionWriter.AddInteger(ANSITEXTVIEW("NumStreamingMips"), FirstInlineMip);
-		DescriptionWriter.AddInteger(ANSITEXTVIEW("NumMipsInTail"), NumMipsInTail);
-
-		DescriptionWriter.BeginArray(ANSITEXTVIEW("Mips"));
-		int64 MipOffset = 0;
-		for (int32 MipIndex = 0; MipIndex < MipCount; ++MipIndex)
+		int32 CalculatedMip0SizeX = 0, CalculatedMip0SizeY = 0, CalculatedMip0NumSlices = 0;
+		int32 CalculatedMipCount = TextureCompressorModule.GetMipCountForBuildSettings(SourceMips[0].SizeX, SourceMips[0].SizeY, SourceMips[0].NumSlices, SourceMips.Num(), BuildSettings, CalculatedMip0SizeX, CalculatedMip0SizeY, CalculatedMip0NumSlices);
+		BuildSettings.GetEncodedTextureDescription(&TextureDescription, TextureFormat, CalculatedMip0SizeX, CalculatedMip0SizeY, CalculatedMip0NumSlices, CalculatedMipCount, bImageHasAlpha);
+	}
+
+	FEncodedTextureExtendedData ExtendedData;
+
+	// ExtendedData is only really useful for textures that have a post build step for tiling,
+	// however it's possible that we ran the old build process where the tiling occurs as part
+	// of the BuildTexture->CompressImage step via child texture formats. In that case, we've already
+	// tiled and we need to pass the data back out. Otherwise, this gets ignored and the tiling step
+	// regenerates it.
+	{
+		ExtendedData.NumMipsInTail = NumMipsInTail;
+		ExtendedData.ExtData = ExtData;
+
+		int32 EncodedMipCount = TextureDescription.GetNumEncodedMips(&ExtendedData);
+		ExtendedData.MipSizesInBytes.AddUninitialized(EncodedMipCount);
+		for (int32 MipIndex = 0; MipIndex < EncodedMipCount; MipIndex++)
 		{
-			FCompressedImage2D& CompressedMip = CompressedMips[MipIndex];
-
-			DescriptionWriter.BeginObject();
-			
-			DescriptionWriter.BeginArray(ANSITEXTVIEW("Size"));
-			DescriptionWriter.AddInteger(CompressedMip.SizeX);
-			DescriptionWriter.AddInteger(CompressedMip.SizeY);
-			DescriptionWriter.AddInteger(CompressedMip.SizeZ);
-			DescriptionWriter.EndArray();
-
-			/*
-			UE_LOG(LogInit,Display,TEXT("FTextureBuildFunction: CompressedMip: %d = %dx%dx%d"),MipIndex,
-				CompressedMip.SizeX,CompressedMip.SizeY,CompressedMip.SizeZ);
-			*/
-
-			const bool bIsInlineMip = MipIndex >= FirstInlineMip;
-
-			EFileRegionType FileRegion = FFileRegion::SelectType(EPixelFormat(CompressedMip.PixelFormat));
-			DescriptionWriter.AddInteger(ANSITEXTVIEW("FileRegion"), static_cast<int32>(FileRegion));
-			if (bIsInlineMip)
-			{
-				DescriptionWriter.AddInteger(ANSITEXTVIEW("MipOffset"), MipOffset);
-			}
-			DescriptionWriter.AddInteger(ANSITEXTVIEW("NumBytes"), CompressedMip.RawData.Num());
+			ExtendedData.MipSizesInBytes[MipIndex] = CompressedMips[MipIndex].RawData.Num();
+		}
+	}
+
+	// Long term, this will be supplied to the build and this would only be called to verify.
+	int32 NumStreamingMips = TextureDescription.GetNumStreamingMips(&ExtendedData, EngineParameters);
+
+	// \todo save bImageHasAlphaChannel back out so that if textures get edited after a build then saved,
+	// we can capture the update.
+
+	
+	{
+		Context.AddValue(UE::DerivedData::FValueId::FromName(ANSITEXTVIEW("EncodedTextureDescription")), UE::TextureBuildUtilities::EncodedTextureDescription::ToCompactBinary(TextureDescription));
+		Context.AddValue(UE::DerivedData::FValueId::FromName(ANSITEXTVIEW("EncodedTextureExtendedData")), UE::TextureBuildUtilities::EncodedTextureExtendedData::ToCompactBinary(ExtendedData));
 
-			DescriptionWriter.EndObject();
+		// Streaming mips
+		for (int32 MipIndex = 0; MipIndex < NumStreamingMips; ++MipIndex)
+		{
+			TAnsiStringBuilder<16> MipName;
+			MipName << ANSITEXTVIEW("Mip") << MipIndex;
+
+			FSharedBuffer MipData = MakeSharedBufferFromArray(MoveTemp(CompressedMips[MipIndex].RawData));
+			Context.AddValue(UE::DerivedData::FValueId::FromName(MipName), MipData);
+		}
 
-			if (bIsInlineMip)
+		// Mip tail
+		TArray<FSharedBuffer> MipTailComponents;
+		for (int32 MipIndex = NumStreamingMips; MipIndex < TextureDescription.NumMips; ++MipIndex)
+		{
+			FSharedBuffer MipData = MakeSharedBufferFromArray(MoveTemp(CompressedMips[MipIndex].RawData));
+			MipTailComponents.Add(MipData);
+		}
+		FCompositeBuffer MipTail(MipTailComponents);
+		if (MipTail.GetSize() > 0)
+		{
+			Context.AddValue(UE::DerivedData::FValueId::FromName(ANSITEXTVIEW("MipTail")), MipTail);
+		}
+	}
+}
+
+void GenericTextureTilingBuildFunction(UE::DerivedData::FBuildContext& Context, const ITextureTiler* Tiler)
+{
+	// The texture description is either passed as a constant or as an output from the other build ("build input").
+	FEncodedTextureDescription TextureDescription;
+	{
+		FCbObject TextureDescriptionCb = Context.FindConstant(UTF8TEXTVIEW("EncodedTextureDescriptionConstant"));
+		if (!TextureDescriptionCb)
+		{
+			FSharedBuffer RawTextureDescription = Context.FindInput(UTF8TEXTVIEW("EncodedTextureDescriptionInput"));
+			if (!RawTextureDescription)
 			{
-				MipOffset += CompressedMip.RawData.Num();
+				return;
 			}
+			TextureDescriptionCb = FCbObject(RawTextureDescription);
+		}
+		UE::TextureBuildUtilities::EncodedTextureDescription::FromCompactBinary(TextureDescription, TextureDescriptionCb);
+	}
+
+	// The extended data is either passed as a constant, but is not output from the linear build - it's
+	// our job to make it.
+	FEncodedTextureExtendedData TextureExtendedData;
+	{
+		FCbObject TextureExtendedDataCb = Context.FindConstant(UTF8TEXTVIEW("EncodedTextureExtendedDataConstant"));
+		if (TextureExtendedDataCb)
+		{
+			UE::TextureBuildUtilities::EncodedTextureExtendedData::FromCompactBinary(TextureExtendedData, TextureExtendedDataCb);
 		}
-		DescriptionWriter.EndArray();
+		else
+		{
+			TextureExtendedData = Tiler->GetExtendedDataForTexture(TextureDescription);
+		}
+	}
+
+	FTextureEngineParameters EngineParameters;
+	{
+		FCbObject EngineParametersCb = Context.FindConstant(UTF8TEXTVIEW("EngineParameters"));
+		UE::TextureBuildUtilities::TextureEngineParameters::FromCompactBinary(EngineParameters, EngineParametersCb);
+	}
+
+	//
+	// Careful - the linear build might have a different streaming mip count than we output due to mip tail
+	// packing.
+	//
+	int32 InputTextureNumStreamingMips = TextureDescription.GetNumStreamingMips(nullptr, EngineParameters);
+	int32 OutputTextureNumStreamingMips = TextureDescription.GetNumStreamingMips(&TextureExtendedData, EngineParameters);
 
-		DescriptionWriter.EndObject();
-		FCbObject DescriptionObject = DescriptionWriter.Save().AsObject();
-		Context.AddValue(UE::DerivedData::FValueId::FromName(ANSITEXTVIEW("Description")), DescriptionObject);
+	FSharedBuffer InputTextureMipTailData;
+	if (TextureDescription.NumMips > InputTextureNumStreamingMips)
+	{
+		InputTextureMipTailData = Context.FindInput(UTF8TEXTVIEW("MipTail"));
 	}
 
-	// Streaming mips
-	for (int32 MipIndex = 0; MipIndex < FirstInlineMip; ++MipIndex)
+	// We might be packing several mips in to a single tiled mip at the end, so we need to have all the buffers available
+	// to potentially pass to ProcessMipLevel. Can do this on demand so that the highest mip level isn't in memory for the entire
+	// mip chain... however all the time is also spent on it and it's only +33% size for the entire chain, so not really worth.
+	TArray<FSharedBuffer> InputTextureMipBuffers;
+	TArray<FMemoryView> InputTextureMipViews;
+
+	uint64 CurrentMipTailOffset = 0;
+	for (int32 MipIndex = 0; MipIndex < TextureDescription.NumMips; MipIndex++)
 	{
-		TAnsiStringBuilder<16> MipName;
-		MipName << ANSITEXTVIEW("Mip") << MipIndex;
+		FMemoryView SourceMipView;
+		if (MipIndex >= InputTextureNumStreamingMips)
+		{
+			// Mip tail.
+			uint64 SourceMipSize = TextureDescription.GetMipSizeInBytes(MipIndex);
+			SourceMipView = InputTextureMipTailData.GetView().Mid(CurrentMipTailOffset, SourceMipSize);
+			CurrentMipTailOffset += SourceMipSize;
+		}
+		else
+		{
+			TUtf8StringBuilder<10> StreamingMipName;
+			StreamingMipName << "Mip" << MipIndex;
 
-		FSharedBuffer MipData = MakeSharedBufferFromArray(MoveTemp(CompressedMips[MipIndex].RawData));
-		Context.AddValue(UE::DerivedData::FValueId::FromName(MipName), MipData);
+			FSharedBuffer SourceData = Context.FindInput(StreamingMipName);
+			SourceMipView = SourceData.GetView();
+			InputTextureMipBuffers.Add(SourceData);
+		}
+		InputTextureMipViews.Add(SourceMipView);
 	}
 
-	// Mip tail
-	TArray<FSharedBuffer> MipTailComponents;
-	for (int32 MipIndex = FirstInlineMip; MipIndex < MipCount; ++MipIndex)
+	// If the platform packs mip tails, we need to pass all the relevant mip buffers at once.
+	int32 FirstMipTailIndex = TextureDescription.NumMips - 1;
+	int32 MipTailCount = 1;
+	if (TextureExtendedData.NumMipsInTail > 1)
 	{
-		FSharedBuffer MipData = MakeSharedBufferFromArray(MoveTemp(CompressedMips[MipIndex].RawData));
-		MipTailComponents.Add(MipData);
+		MipTailCount = TextureExtendedData.NumMipsInTail;
+		FirstMipTailIndex = TextureDescription.NumMips - MipTailCount;
 	}
-	FCompositeBuffer MipTail(MipTailComponents);
+
+
+	// Process the mips
+	TArray<FSharedBuffer> MipTailBuffers;
+	for (int32 MipIndex = 0; MipIndex < FirstMipTailIndex + 1; MipIndex++)
+	{
+		TUtf8StringBuilder<10> StreamingMipName;
+		StreamingMipName << "Mip" << MipIndex;
+
+		TArrayView<FMemoryView> MipsForLevel = MakeArrayView(InputTextureMipViews.GetData() + MipIndex, 1);
+		if (MipIndex == FirstMipTailIndex)
+		{
+			MipsForLevel = MakeArrayView(InputTextureMipViews.GetData() + MipIndex, MipTailCount);
+		}
+		FSharedBuffer MipData = Tiler->ProcessMipLevel(TextureDescription, TextureExtendedData, MipsForLevel, MipIndex);
+
+		// Make sure we got the size we advertised prior to the build. If this ever fires then we
+		// have a critical mismatch!
+		check(TextureExtendedData.MipSizesInBytes[MipIndex] == MipData.GetSize());
+
+		// Save the data to the output.
+		if (MipIndex < OutputTextureNumStreamingMips)
+		{
+			Context.AddValue(UE::DerivedData::FValueId::FromName(StreamingMipName), MipData);
+		}
+		else
+		{
+			MipTailBuffers.Add(MipData);
+		}
+	} // end for each mip
+
+	// The mip tail is a bunch of mips all together in one "Value", so assemble them here.
+	FCompositeBuffer MipTail(MipTailBuffers);
 	if (MipTail.GetSize() > 0)
 	{
-		Context.AddValue(UE::DerivedData::FValueId::FromName(ANSITEXTVIEW("MipTail")), MipTail);
+		Context.AddValue(UE::DerivedData::FValueId::FromName(UTF8TEXTVIEW("MipTail")), MipTail);
 	}
+
+	Context.AddValue(UE::DerivedData::FValueId::FromName(UTF8TEXTVIEW("EncodedTextureDescription")), UE::TextureBuildUtilities::EncodedTextureDescription::ToCompactBinary(TextureDescription));
+	Context.AddValue(UE::DerivedData::FValueId::FromName(UTF8TEXTVIEW("EncodedTextureExtendedData")), UE::TextureBuildUtilities::EncodedTextureExtendedData::ToCompactBinary(TextureExtendedData));
 }
diff --git a/Engine/Source/Developer/TextureBuild/Public/TextureBuildFunction.h b/Engine/Source/Developer/TextureBuild/Public/TextureBuildFunction.h
--- a/Engine/Source/Developer/TextureBuild/Public/TextureBuildFunction.h
+++ b/Engine/Source/Developer/TextureBuild/Public/TextureBuildFunction.h
@@ -5,9 +5,12 @@
 #include "CoreTypes.h"
 #include "DerivedDataBuildFunction.h"
 #include "DerivedDataBuildVersion.h"
+#include "DerivedDataCacheKey.h"
+#include "DerivedDataSharedString.h"
 #include "Misc/Guid.h"
 
 class ITextureFormat;
+class ITextureTiler;
 namespace UE { namespace DerivedData { class FBuildVersionBuilder; } }
 
 class FTextureBuildFunction : public UE::DerivedData::IBuildFunction
@@ -20,3 +23,36 @@ public:
 protected:
 	virtual void GetVersion(UE::DerivedData::FBuildVersionBuilder& Builder, ITextureFormat*& OutTextureFormatVersioning) const = 0;
 };
+
+/**
+*	This function does the meat of breaking out the inputs from the build context and handing them
+*	to the tiler, then packing them back up for the build process.
+*/
+TEXTUREBUILD_API void GenericTextureTilingBuildFunction(UE::DerivedData::FBuildContext& Context, const ITextureTiler* Tiler);
+
+
+/**
+*	This build function expects an implementation of ITextureTiler as its template
+*	and looks a bit awkward because IBuildFunction is required to not have any state,
+*	so we can't put the instance of the ITextureTiler on our object - hence the odd
+*	statics.
+*/
+template <class ITextureTilerObject>
+class FGenericTextureTilingBuildFunction : public UE::DerivedData::IBuildFunction
+{
+	virtual void Build(UE::DerivedData::FBuildContext& Context) const
+	{
+		ITextureTilerObject Tiler;
+		GenericTextureTilingBuildFunction(Context, &Tiler);
+	}
+	virtual void Configure(UE::DerivedData::FBuildConfigContext& Context) const
+	{
+		Context.SetCacheBucket(UE::DerivedData::FCacheBucket(UTF8TEXTVIEW("TiledTextures")));
+	}
+	virtual FGuid GetVersion() const final { return ITextureTilerObject::GetBuildFunctionVersionGuid(); }
+	const UE::DerivedData::FUtf8SharedString& GetName() const final
+	{
+		static UE::DerivedData::FUtf8SharedString Name(ITextureTilerObject::GetBuildFunctionNameStatic());
+		return Name;
+	}
+};
diff --git a/Engine/Source/Developer/TextureBuild/TextureBuild.Build.cs b/Engine/Source/Developer/TextureBuild/TextureBuild.Build.cs
--- a/Engine/Source/Developer/TextureBuild/TextureBuild.Build.cs
+++ b/Engine/Source/Developer/TextureBuild/TextureBuild.Build.cs
@@ -18,6 +18,7 @@ public class TextureBuild : ModuleRules
 			"DerivedDataCache",
 			"ImageCore",
 			"ImageWrapper",
+			"TextureBuildUtilities",
 			"TextureCompressor",
 			"TextureFormat",
 		});
diff --git a/Engine/Source/Developer/TextureBuildUtilities/Private/TextureBuildUtilities.cpp b/Engine/Source/Developer/TextureBuildUtilities/Private/TextureBuildUtilities.cpp
--- a/Engine/Source/Developer/TextureBuildUtilities/Private/TextureBuildUtilities.cpp
+++ b/Engine/Source/Developer/TextureBuildUtilities/Private/TextureBuildUtilities.cpp
@@ -2,6 +2,8 @@
 
 #include "TextureBuildUtilities.h"
 #include "Misc/DataDrivenPlatformInfoRegistry.h"
+#include "Serialization/CompactBinary.h"
+#include "Serialization/CompactBinaryWriter.h"
 
 
 namespace UE
@@ -90,5 +92,123 @@ TEXTUREBUILDUTILITIES_API ERawImageFormat::Type GetVirtualTextureBuildIntermedia
 	}
 }
 
+#ifndef TEXT_TO_ENUM
+#define TEXT_TO_ENUM(eVal, txt)		if (txt.Compare(#eVal) == 0)	return eVal;
+#endif
+
+static EPixelFormat GetPixelFormatFromUtf8(const FUtf8StringView& InPixelFormatStr)
+{
+#define TEXT_TO_PIXELFORMAT(f) TEXT_TO_ENUM(f, InPixelFormatStr);
+	FOREACH_ENUM_EPIXELFORMAT(TEXT_TO_PIXELFORMAT)
+#undef TEXT_TO_PIXELFORMAT
+		return PF_Unknown;
+}
+
+
+namespace EncodedTextureExtendedData
+{
+	FCbObject ToCompactBinary(const FEncodedTextureExtendedData& InExtendedData)
+	{
+		FCbWriter Writer;
+		Writer.BeginObject();
+		Writer.AddInteger("NumMipsInTail", InExtendedData.NumMipsInTail);
+		Writer.AddInteger("ExtData", InExtendedData.ExtData);
+		Writer.BeginArray("MipSizes");
+		for (uint64 MipSize : InExtendedData.MipSizesInBytes)
+		{
+			Writer.AddInteger(MipSize);
+		}
+		Writer.EndArray();
+		Writer.EndObject();
+		return Writer.Save().AsObject();
+	}
+
+	bool FromCompactBinary(FEncodedTextureExtendedData& OutExtendedData, FCbObject InCbObject)
+	{
+		OutExtendedData.ExtData = InCbObject["ExtData"].AsUInt32();
+		OutExtendedData.NumMipsInTail = InCbObject["NumMipsInTail"].AsInt32();
+
+		FCbArrayView MipArrayView = InCbObject["MipSizes"].AsArrayView();
+		for (FCbFieldView MipFieldView : MipArrayView)
+		{
+			OutExtendedData.MipSizesInBytes.Add(MipFieldView.AsUInt64());
+		}
+		return true;
+	}
+} // namespace EncodedTextureExtendedData
+
+namespace EncodedTextureDescription
+{
+	FCbObject ToCompactBinary(const FEncodedTextureDescription& InDescription)
+	{
+		FCbWriter Writer;
+		Writer.BeginObject();
+		Writer.AddInteger("TopMipSizeX", InDescription.TopMipSizeX);
+		Writer.AddInteger("TopMipSizeY", InDescription.TopMipSizeY);
+		Writer.AddInteger("TopMipVolumeSizeZ", InDescription.TopMipVolumeSizeZ);
+		Writer.AddInteger("ArraySlices", InDescription.ArraySlices);
+		Writer.AddString("PixelFormat", GetPixelFormatString(InDescription.PixelFormat));
+		Writer.AddInteger("NumMips", InDescription.NumMips);
+		Writer.AddBool("bCubeMap", InDescription.bCubeMap);
+		Writer.AddBool("bTextureArray", InDescription.bTextureArray);
+		Writer.AddBool("bVolumeTexture", InDescription.bVolumeTexture);
+		Writer.EndObject();
+		return Writer.Save().AsObject();
+	}
+
+	bool FromCompactBinary(FEncodedTextureDescription& OutDescription, FCbObject InCbObject)
+	{
+		OutDescription.TopMipSizeX = InCbObject["TopMipSizeX"].AsInt32();
+		OutDescription.TopMipSizeY = InCbObject["TopMipSizeY"].AsInt32();
+		OutDescription.TopMipVolumeSizeZ = InCbObject["TopMipVolumeSizeZ"].AsInt32();
+		OutDescription.ArraySlices = InCbObject["ArraySlices"].AsInt32();
+		OutDescription.PixelFormat = GetPixelFormatFromUtf8(InCbObject["PixelFormat"].AsString());
+		OutDescription.NumMips = (uint8)InCbObject["NumMips"].AsInt32();
+		OutDescription.bCubeMap = InCbObject["bCubeMap"].AsBool();
+		OutDescription.bTextureArray = InCbObject["bTextureArray"].AsBool();
+		OutDescription.bVolumeTexture = InCbObject["bVolumeTexture"].AsBool();
+		return true;
+	}
+} // namespace EncodedTextureDescription
+
+
+
+namespace TextureEngineParameters
+{
+	FCbObject ToCompactBinaryWithDefaults(const FTextureEngineParameters& InEngineParameters)
+	{
+		FTextureEngineParameters Defaults;
+
+		FCbWriter Writer;
+		Writer.BeginObject();
+		if (InEngineParameters.bEngineSupportsTexture2DArrayStreaming != Defaults.bEngineSupportsTexture2DArrayStreaming)
+		{
+			Writer.AddBool("bEngineSupportsTexture2DArrayStreaming", InEngineParameters.bEngineSupportsTexture2DArrayStreaming);
+		}
+		if (InEngineParameters.bEngineSupportsVolumeTextureStreaming != Defaults.bEngineSupportsVolumeTextureStreaming)
+		{
+			Writer.AddBool("bEngineSupportsVolumeTextureStreaming", InEngineParameters.bEngineSupportsVolumeTextureStreaming);
+		}
+		if (InEngineParameters.NumInlineDerivedMips != Defaults.NumInlineDerivedMips)
+		{
+			Writer.AddInteger("NumInlineDerivedMips", InEngineParameters.NumInlineDerivedMips);
+		}
+		Writer.EndObject();
+		return Writer.Save().AsObject();
+	}
+
+	bool FromCompactBinary(FTextureEngineParameters& OutEngineParameters, FCbObject InCbObject)
+	{
+		OutEngineParameters = FTextureEngineParameters(); // init to defaults
+
+		OutEngineParameters.NumInlineDerivedMips = InCbObject["NumInlineDerivedMips"].AsInt32(OutEngineParameters.NumInlineDerivedMips);
+		OutEngineParameters.bEngineSupportsTexture2DArrayStreaming = InCbObject["bEngineSupportsTexture2DArrayStreaming"].AsBool(OutEngineParameters.bEngineSupportsTexture2DArrayStreaming);
+		OutEngineParameters.bEngineSupportsVolumeTextureStreaming = InCbObject["bEngineSupportsVolumeTextureStreaming"].AsBool(OutEngineParameters.bEngineSupportsVolumeTextureStreaming);
+		return true;
+	}
+} // namespace EncodedTextureDescription
+
+
+
 } // namespace
 }
\ No newline at end of file
diff --git a/Engine/Source/Developer/TextureBuildUtilities/Public/TextureBuildUtilities.h b/Engine/Source/Developer/TextureBuildUtilities/Public/TextureBuildUtilities.h
--- a/Engine/Source/Developer/TextureBuildUtilities/Public/TextureBuildUtilities.h
+++ b/Engine/Source/Developer/TextureBuildUtilities/Public/TextureBuildUtilities.h
@@ -5,6 +5,7 @@
 #include "CoreMinimal.h"
 #include "ImageCore.h"
 #include "TextureCompressorModule.h" // for FTextureBuildSettings
+#include "Interfaces/ITextureFormat.h"
 
 /***
 
@@ -21,6 +22,24 @@ namespace UE
 namespace TextureBuildUtilities
 {
 
+namespace EncodedTextureExtendedData
+{
+	TEXTUREBUILDUTILITIES_API FCbObject ToCompactBinary(const FEncodedTextureExtendedData& InExtendedData);
+	TEXTUREBUILDUTILITIES_API bool FromCompactBinary(FEncodedTextureExtendedData& OutExtendedData, FCbObject InCbObject);
+}
+
+namespace EncodedTextureDescription
+{
+	TEXTUREBUILDUTILITIES_API FCbObject ToCompactBinary(const FEncodedTextureDescription& InDescription);
+	TEXTUREBUILDUTILITIES_API bool FromCompactBinary(FEncodedTextureDescription& OutDescription, FCbObject InCbObject);
+}
+
+namespace TextureEngineParameters
+{
+	TEXTUREBUILDUTILITIES_API FCbObject ToCompactBinaryWithDefaults(const FTextureEngineParameters& InEngineParameters);
+	TEXTUREBUILDUTILITIES_API bool FromCompactBinary(FTextureEngineParameters& OutEngineParameters, FCbObject InCbObject);
+}
+
 TEXTUREBUILDUTILITIES_API bool TextureFormatIsHdr(FName const& InName);
 
 // Removes platform and other custom prefixes from the name.
@@ -36,5 +55,5 @@ FORCEINLINE const FName TextureFormatRemovePrefixFromName(FName const& InName)
 
 TEXTUREBUILDUTILITIES_API ERawImageFormat::Type GetVirtualTextureBuildIntermediateFormat(const FTextureBuildSettings& BuildSettings);
 
-}
-}
+} // namespace TextureBuildUtilities
+} // namespace UE
diff --git a/Engine/Source/Developer/TextureCompressor/Private/TextureCompressorModule.cpp b/Engine/Source/Developer/TextureCompressor/Private/TextureCompressorModule.cpp
--- a/Engine/Source/Developer/TextureCompressor/Private/TextureCompressorModule.cpp
+++ b/Engine/Source/Developer/TextureCompressor/Private/TextureCompressorModule.cpp
@@ -2,6 +2,7 @@
 
 #include "TextureCompressorModule.h"
 #include "Math/RandomStream.h"
+#include "ChildTextureFormat.h"
 #include "Containers/IndirectArray.h"
 #include "Stats/Stats.h"
 #include "Async/AsyncWork.h"
@@ -1261,33 +1262,52 @@ struct FTextureDownscaleSettings
 	float Downscale;
 	uint8 DownscaleOptions;
 	bool UseNewMipFilter;
-};
 
-static void DownscaleImage(const FImage& SrcImage, FImage& DstImage, const FTextureDownscaleSettings& Settings)
-{
-	if (Settings.Downscale <= 1.f)
+	FTextureDownscaleSettings(const FTextureBuildSettings& BuildSettings)
 	{
-		return;
+		Downscale = BuildSettings.Downscale;
+		DownscaleOptions = BuildSettings.DownscaleOptions;
+		BlockSize = 4;
+		UseNewMipFilter = BuildSettings.bUseNewMipFilter;
 	}
-	
-	TRACE_CPUPROFILER_EVENT_SCOPE(Texture.DownscaleImage);
+};
 
+static float GetDownscaleFinalSizeAndClampedDownscale(int32 SrcImageWidth, int32 SrcImageHeight,  const FTextureDownscaleSettings& Settings, int32& OutWidth, int32& OutHeight)
+{
+	check(Settings.Downscale > 1.0f); // must be already handled.
 	float Downscale = FMath::Clamp(Settings.Downscale, 1.f, 8.f);
 	// note: more accurate would be to use FMath::Max(1, FMath::RoundToInt(SrcImage.SizeX / Downscale))
-	int32 FinalSizeX = FMath::CeilToInt(SrcImage.SizeX / Downscale);
-	int32 FinalSizeY = FMath::CeilToInt(SrcImage.SizeY / Downscale);
+	int32 FinalSizeX = FMath::CeilToInt(SrcImageWidth / Downscale);
+	int32 FinalSizeY = FMath::CeilToInt(SrcImageHeight / Downscale);
 
 	// compute final size respecting image block size
 	if (Settings.BlockSize > 1
-		&& SrcImage.SizeX % Settings.BlockSize == 0 
-		&& SrcImage.SizeY % Settings.BlockSize == 0)
+		&& SrcImageWidth % Settings.BlockSize == 0
+		&& SrcImageHeight % Settings.BlockSize == 0)
 	{
 		// the following code finds non-zero dimensions of the scaled image which preserve both aspect ratio and block alignment, and are also the closest to the requested dimensions
-		int32 ScalingGridSizeX = SrcImage.SizeX / FMath::GreatestCommonDivisor(SrcImage.SizeX, SrcImage.SizeY) * Settings.BlockSize;
+		int32 ScalingGridSizeX = SrcImageWidth / FMath::GreatestCommonDivisor(SrcImageWidth, SrcImageHeight) * Settings.BlockSize;
 		// note: more accurate would be to use (SrcImage.SizeX / Downscale) instead of FinalSizeX here
 		FinalSizeX = FMath::Max(ScalingGridSizeX, FMath::GridSnap(FinalSizeX, ScalingGridSizeX));
-		FinalSizeY = (int64)FinalSizeX * SrcImage.SizeY / SrcImage.SizeX;
+		FinalSizeY = (int64)FinalSizeX * SrcImageHeight / SrcImageWidth;
+	}
+
+	OutWidth = FinalSizeX;
+	OutHeight = FinalSizeY;
+	return Downscale;
+}
+
+static void DownscaleImage(const FImage& SrcImage, FImage& DstImage, const FTextureDownscaleSettings& Settings)
+{
+	if (Settings.Downscale <= 1.f)
+	{
+		return;
 	}
+	
+	TRACE_CPUPROFILER_EVENT_SCOPE(Texture.DownscaleImage);
+
+	int32 FinalSizeX = 0, FinalSizeY =0;
+	float Downscale = GetDownscaleFinalSizeAndClampedDownscale(SrcImage.SizeX, SrcImage.SizeY, Settings, FinalSizeX, FinalSizeY);
 
 	Downscale = (float)SrcImage.SizeX / FinalSizeX;
 		
@@ -2773,7 +2793,7 @@ static void ApplyCompositeTexture(FImage& RoughnessSourceMips, const FImage& Nor
 	Image Compression.
 ------------------------------------------------------------------------------*/
 
-void FTextureBuildSettings::GetEncodedTextureDescription(FEncodedTextureDescription* OutTextureDescription, const ITextureFormat* InTextureFormat, int32 InMipCount, bool bInImageHasAlphaChannel) const
+void FTextureBuildSettings::GetEncodedTextureDescription(FEncodedTextureDescription* OutTextureDescription, const ITextureFormat* InTextureFormat, int32 InEncodedMip0SizeX, int32 InEncodedMip0SizeY, int32 InEncodedMip0NumSlices, int32 InMipCount, bool bInImageHasAlphaChannel) const
 {
 	FEncodedTextureDescription& TextureDescription = *OutTextureDescription;
 	TextureDescription = FEncodedTextureDescription();
@@ -2782,10 +2802,25 @@ void FTextureBuildSettings::GetEncodedTextureDescription(FEncodedTextureDescript
 	TextureDescription.bVolumeTexture = bVolume;
 	TextureDescription.NumMips = InMipCount;
 	TextureDescription.PixelFormat = InTextureFormat->GetEncodedPixelFormat(*this, bInImageHasAlphaChannel);
-	TextureDescription.TopMipSizeX = TopMipSize.X;
-	TextureDescription.TopMipSizeY = TopMipSize.Y;
-	TextureDescription.TopMipVolumeSizeZ = VolumeSizeZ;
-	TextureDescription.ArraySlices = ArraySlices;
+
+	TextureDescription.TopMipSizeX = InEncodedMip0SizeX;
+	TextureDescription.TopMipSizeY = InEncodedMip0SizeY;
+	TextureDescription.TopMipVolumeSizeZ = bVolume ? InEncodedMip0NumSlices : 1;
+	if (bTextureArray)
+	{
+		if (bCubemap)
+		{
+			TextureDescription.ArraySlices = InEncodedMip0NumSlices / 6;
+		}
+		else
+		{
+			TextureDescription.ArraySlices = InEncodedMip0NumSlices;
+		}
+	}
+	else
+	{
+		TextureDescription.ArraySlices = 1;
+	}
 }
 
 // compress mip-maps in InMipChain and add mips to Texture, might alter the source content
@@ -2805,10 +2840,9 @@ static bool CompressMipChain(
 	// Determine ExtData (platform specific data) and NumMipsInTail.
 	// ExtData gets passed 
 	FEncodedTextureExtendedData ExtendedData;
+	FEncodedTextureDescription TextureDescription;
+	Settings.GetEncodedTextureDescription(&TextureDescription, TextureFormat, MipChain[0].SizeX, MipChain[0].SizeY, MipChain[0].NumSlices, MipChain.Num(), bImageHasAlphaChannel);
 	{
-		FEncodedTextureDescription TextureDescription;
-		Settings.GetEncodedTextureDescription(&TextureDescription, TextureFormat, MipChain.Num(), bImageHasAlphaChannel);
-
 		ExtendedData = TextureFormat->GetExtendedDataForTexture(TextureDescription);
 		OutNumMipsInTail = ExtendedData.NumMipsInTail;
 		OutExtData = ExtendedData.ExtData;
@@ -2907,6 +2941,7 @@ static bool CompressMipChain(
 		bCompressionSucceeded = ProcessMips(0, FirstMipTailIndex + 1);
 	}
 
+	// Fill out the dimensions for the packed mip tail, should we have one
 	for (int32 MipIndex = FirstMipTailIndex + 1; MipIndex < MipCount; ++MipIndex)
 	{
 		FCompressedImage2D& PrevMip = OutMips[MipIndex - 1];
@@ -2953,7 +2988,7 @@ static void NormalizeMip(FImage& InOutMip)
 
 
 // Special case for TMGS_LeaveExistingMips
-static int32 GetMipCountForLeaveExistingMips(int32 InMip0SizeX, int32 InMip0SizeY, int32 InExistingMipCount, uint32 InMaxTexture2DResolution)
+static int32 GetMipCountForLeaveExistingMips(int32 InMip0SizeX, int32 InMip0SizeY, int32 InExistingMipCount, uint32 InMaxTexture2DResolution, int32& OutMip0SizeX, int32& OutMip0SizeY)
 {
 	int32 i = 0;
 	for (; i < InExistingMipCount; i++)
@@ -2964,6 +2999,8 @@ static int32 GetMipCountForLeaveExistingMips(int32 InMip0SizeX, int32 InMip0Size
 		if (MipSizeX <= InMaxTexture2DResolution &&
 			MipSizeY <= InMaxTexture2DResolution)
 		{
+			OutMip0SizeX = MipSizeX;
+			OutMip0SizeY = MipSizeY;
 			return InExistingMipCount - i;
 		}
 	}
@@ -3034,16 +3071,17 @@ public:
 	{
 	}
 
-	virtual int32 GetMipCountForBuildSettings(int32 InMip0SizeX, int32 InMip0SizeY, int32 InMip0NumSlices, int32 InExistingMipCount, const FTextureBuildSettings& BuildSettings) const override
+	virtual int32 GetMipCountForBuildSettings(
+		int32 InMip0SizeX, int32 InMip0SizeY, int32 InMip0NumSlices, 
+		int32 InExistingMipCount, 
+		const FTextureBuildSettings& BuildSettings, 
+		int32& OutMip0SizeX, int32& OutMip0SizeY, int32& OutMip0NumSlices) const override
 	{
 		if (BuildSettings.MipGenSettings == TMGS_LeaveExistingMips)
 		{
 			// Since we can't generate, we only have to limit to MaxTextureSize
-			return GetMipCountForLeaveExistingMips(InMip0SizeX, InMip0SizeY, InExistingMipCount, BuildSettings.MaxTextureResolution);
-		}
-		else if (BuildSettings.MipGenSettings == TMGS_NoMipmaps)
-		{
-			return 1;
+			OutMip0NumSlices = InMip0NumSlices; // At the moment, when importing a volume texture, only the 2d dimensions are checked against max resolution.
+			return GetMipCountForLeaveExistingMips(InMip0SizeX, InMip0SizeY, InExistingMipCount, BuildSettings.MaxTextureResolution, OutMip0SizeX, OutMip0SizeY);
 		}
 
 		// AFAICT LatLongCubeMaps don't do any of this - pow2 is broken with them but it runs, and max texture stuff
@@ -3097,19 +3135,52 @@ public:
 					break;
 				}
 			}
+
+			if (BuildSettings.Downscale > 1.0f)
+			{
+				int32 DownscaledSizeX = 0, DownscaledSizeY = 0;
+				GetDownscaleFinalSizeAndClampedDownscale(BaseSizeX, BaseSizeY, FTextureDownscaleSettings(BuildSettings), DownscaledSizeX, DownscaledSizeY);
+
+				if (BuildSettings.bVolume)
+				{
+					UE_LOG(LogTextureCompressor, Error, TEXT("Downscaling volumes not yet supported - should have been handled in GetTextureBuildSettings!"));
+				}
+				check(BuildSettings.bVolume == false);
+
+				BaseSizeX = DownscaledSizeX;
+				BaseSizeY = DownscaledSizeY;
+			}
+
+			// Volumes are the only thing where num slices changes.
+			if (BuildSettings.bVolume == false)
+			{
+				OutMip0NumSlices = InMip0NumSlices;
+			}
+			else
+			{
+				OutMip0NumSlices = BaseSizeZ;
+			}
+		}
+		else
+		{
+			uint32 LongLatCubemapExtents = ComputeLongLatCubemapExtents(BaseSizeX, BuildSettings.MaxTextureResolution);
+			BaseSizeX = LongLatCubemapExtents;
+			BaseSizeY = LongLatCubemapExtents;
+			OutMip0NumSlices = 6;
 		}
 
 		// At this point we have a base mip size that is valid.
+		OutMip0SizeX = BaseSizeX;
+		OutMip0SizeY = BaseSizeY;
 
-		// Relevant BuildSettings.MipGenSettings have been handled at top of function.
+		if (BuildSettings.MipGenSettings == TMGS_NoMipmaps)
+		{
+			return 1;
+		}
 
 		// NumOutputMips is the number of mips that would be made if you made a full mip chain
 		//  eg. 256 makes 9 mips , 300 also makes 9 mips
 		uint32 MaxMipDimension = FMath::Max3(BaseSizeX, BaseSizeY, BaseSizeZ);
-		if (BuildSettings.bLongLatSource)
-		{
-			MaxMipDimension = ComputeLongLatCubemapExtents(BaseSizeX, BuildSettings.MaxTextureResolution);
-		}
 		return 1 + FMath::FloorLog2(MaxMipDimension);
 	}
 
@@ -3120,9 +3191,10 @@ public:
 		FStringView DebugTexturePathName,
 		TArray<FCompressedImage2D>& OutTextureMips,
 		uint32& OutNumMipsInTail,
-		uint32& OutExtData
+		uint32& OutExtData,
+		bool* bOutImageHasAlpha
 	)
-	{
+	{		
 		//TRACE_CPUPROFILER_EVENT_SCOPE(Texture.BuildTexture);
 
 		const ITextureFormat* TextureFormat = nullptr;
@@ -3222,6 +3294,11 @@ public:
 		{
 			BuildSettings.ArraySlices = 1;
 		}
+
+		if (bOutImageHasAlpha)
+		{
+			*bOutImageHasAlpha = bImageHasAlphaChannel;
+		}
 		
 		return CompressMipChain(TextureFormat, IntermediateMipChain, BuildSettings, bImageHasAlphaChannel, DebugTexturePathName,
 					OutTextureMips, OutNumMipsInTail, OutExtData);
@@ -3556,14 +3633,8 @@ private:
 			}
 
 			if (BuildSettings.Downscale > 1.f)
-			{
-				FTextureDownscaleSettings DownscaleSettings;
-				DownscaleSettings.Downscale = BuildSettings.Downscale;
-				DownscaleSettings.DownscaleOptions = BuildSettings.DownscaleOptions;
-				DownscaleSettings.BlockSize = 4;
-				DownscaleSettings.UseNewMipFilter = BuildSettings.bUseNewMipFilter;
-		
-				DownscaleImage(*Mip, *Mip, DownscaleSettings);
+			{		
+				DownscaleImage(*Mip, *Mip, FTextureDownscaleSettings(BuildSettings));
 			}
 
 			if (BuildSettings.bHasColorSpaceDefinition)
@@ -3623,10 +3694,19 @@ private:
 		}
 		check(OutMipChain.Num() == NumOutputMips);
 
-		int32 CalculatedMipCount = GetMipCountForBuildSettings(InSourceMipChain[0].SizeX, InSourceMipChain[0].SizeY, InSourceMipChain[0].NumSlices, InSourceMipChain.Num(), BuildSettings);
-		if (CalculatedMipCount != NumOutputMips)
-		{
-			UE_LOG(LogTextureCompressor, Error, TEXT("Texture %.*s generated %d mips when GetMipCountForBuildSettings expected %d!"), DebugTexturePathName.Len(), DebugTexturePathName.GetData(), NumOutputMips, CalculatedMipCount);
+		int32 CalculatedMip0SizeX, CalculatedMip0SizeY, CalculatedMip0NumSlices;
+		int32 CalculatedMipCount = GetMipCountForBuildSettings(
+			InSourceMipChain[0].SizeX, InSourceMipChain[0].SizeY, InSourceMipChain[0].NumSlices, InSourceMipChain.Num(), BuildSettings,
+			CalculatedMip0SizeX, CalculatedMip0SizeY, CalculatedMip0NumSlices);
+		if (CalculatedMipCount != NumOutputMips ||
+			CalculatedMip0SizeX != OutMipChain[0].SizeX ||
+			CalculatedMip0SizeY != OutMipChain[0].SizeY ||
+			CalculatedMip0NumSlices != OutMipChain[0].NumSlices)
+		{
+			UE_LOG(LogTextureCompressor, Error, TEXT("Texture %.*s generated unexpected mip chain: GetMipCountForBuildSettings expected %d mips, %dx%dx%d, got %d mips, %dx%dx%d!"), 
+				DebugTexturePathName.Len(), DebugTexturePathName.GetData(), 
+				CalculatedMipCount, CalculatedMip0SizeX, CalculatedMip0SizeY, CalculatedMip0NumSlices,
+				NumOutputMips, OutMipChain[0].SizeX, OutMipChain[0].SizeY, OutMipChain[0].NumSlices);
 		}
 
 		// Apply post-mip generation adjustments.
diff --git a/Engine/Source/Developer/TextureCompressor/Public/TextureCompressorModule.h b/Engine/Source/Developer/TextureCompressor/Public/TextureCompressorModule.h
--- a/Engine/Source/Developer/TextureCompressor/Public/TextureCompressorModule.h
+++ b/Engine/Source/Developer/TextureCompressor/Public/TextureCompressorModule.h
@@ -10,6 +10,8 @@
 
 struct FEncodedTextureDescription;
 class ITextureFormat;
+class ITextureTiler;
+struct FTextureEngineParameters;
 
 /**
  * Compressed image data.
@@ -165,8 +167,11 @@ struct FTextureBuildSettings
 	mutable int32 VolumeSizeZ;
 	/** The array texture's top mip size Z without LODBias applied */
 	mutable int32 ArraySlices;
-	/** Can the texture be streamed */
-	uint32 bStreamable : 1;
+	/** Can the texture be streamed. This is deprecated because it was used in a single place in a single 
+	*	platform for something that handled an edge case that never happened. That code is removes so this is
+	*	never touched other than saving it, and it'll be removed soon.
+	*/
+	uint32 bStreamable_Unused : 1;
 	/** Is the texture streamed using the VT system */
 	uint32 bVirtualStreamable : 1;
 	/** Whether to chroma key the image, replacing any pixels that match ChromaKeyColor with transparent black */
@@ -216,6 +221,9 @@ struct FTextureBuildSettings
 	// ETextureEncodeSpeed, either Final or Fast.
 	uint8 RepresentsEncodeSpeedNoSend;
 
+	// If the target format is a tiled format and can leverage reusing the linear encoding, this is not nullptr.
+	const ITextureTiler* Tiler = nullptr;
+
 	static constexpr uint32 MaxTextureResolutionDefault = TNumericLimits<uint32>::Max();
 
 	/** Default settings. */
@@ -262,7 +270,7 @@ struct FTextureBuildSettings
 		, TopMipSize(0, 0)
 		, VolumeSizeZ(0)
 		, ArraySlices(0)
-		, bStreamable(false)
+		, bStreamable_Unused(false)
 		, bVirtualStreamable(false)
 		, bChromaKeyTexture(false)
 		, PowerOfTwoMode(0 /*ETexturePowerOfTwoSetting::None*/)
@@ -310,7 +318,7 @@ struct FTextureBuildSettings
 	* 
 	* At the moment, this only generates valid results _after_ some setup code in FTextureCompressorModule::BuildTexture()
 	*/
-	void GetEncodedTextureDescription(FEncodedTextureDescription* OutTextureDescription, const ITextureFormat* InTextureFormat, int32 InMipCount, bool bInImageHasAlphaChannel) const;
+	TEXTURECOMPRESSOR_API void GetEncodedTextureDescription(FEncodedTextureDescription* OutTextureDescription, const ITextureFormat* InTextureFormat, int32 InEncodedMip0SizeX, int32 InEncodedMip0SizeY, int32 InEncodedMip0NumSlices, int32 InMipCount, bool bInImageHasAlphaChannel) const;
 };
 
 /**
@@ -337,7 +345,8 @@ public:
 		FStringView DebugTexturePathName,
 		TArray<FCompressedImage2D>& OutTextureMips,
 		uint32& OutNumMipsInTail,
-		uint32& OutExtData
+		uint32& OutExtData,
+		bool* bOutImageHasAlpha // If desired, this will report whether the mip processing determined an alpha channel is necessary in the encoded texture.		
 		) = 0;
 
 	
@@ -380,7 +389,10 @@ public:
 	TEXTURECOMPRESSOR_API static void GenerateAngularFilteredMips(TArray<FImage>& InOutMipChain, int32 NumMips, uint32 DiffuseConvolveMipLevel);
 
 	/**
-	* Returns the number of mips that the given texture will generate with the given build settings.
+	* Returns the number of mips that the given texture will generate with the given build settings, as well as the size of the top mip.
 	*/
-	virtual int32 GetMipCountForBuildSettings(int32 InMip0SizeX, int32 InMip0SizeY, int32 InMip0NumSlices, int32 InExistingMipCount, const FTextureBuildSettings& InBuildSettings) const =0;
+	virtual int32 GetMipCountForBuildSettings(
+		int32 InMip0SizeX, int32 InMip0SizeY, int32 InMip0NumSlices, 
+		int32 InExistingMipCount, const FTextureBuildSettings& InBuildSettings, 
+		int32& OutMip0SizeX, int32& OutMip0SizeY, int32& OutMip0NumSlices) const =0;
 };
diff --git a/Engine/Source/Developer/TextureFormat/Public/ChildTextureFormat.h b/Engine/Source/Developer/TextureFormat/Public/ChildTextureFormat.h
--- a/Engine/Source/Developer/TextureFormat/Public/ChildTextureFormat.h
+++ b/Engine/Source/Developer/TextureFormat/Public/ChildTextureFormat.h
@@ -13,7 +13,17 @@
 
 /**
  * Version of ITextureFormat that handles a child texture format that is used as a "post-process" after compressing textures, useful for
- * several platforms that need to modify already compressed texture data for optimal data
+ * several platforms that need to modify already compressed texture data for optimal data.
+ * 
+ * There are 3 ways a texture can get processed:
+ *		-- ITextureFormat::IsTiling returns true. This should only be for the old Oodle RAD Game Tools supplied texture encoder 
+ *			and is being deprecated in favor of the new built-in Oodle texture encoder (TextureFormatOodle). In this case,
+ *			the child format passes the images to the encoder to split apart prior to passing to the platform tools.
+ *		-- FChildTextureFormat::GetTiler returns nullptr. This is being phased out because the linear version of the input texture
+ *			gets encoded for each platform with a post process, wasting cook time. In this case the child format encodes
+ *			the "base" texture and then tiles it using the platform tools.
+ *		-- FChildTextureFormat::GetTiler returns a valid pointer. This is the most recent design. In this case, the child format 
+ *			is used as a holder of metadata and doesn't actually do any work (CompressImage class of functions not called).
  */
 class FChildTextureFormat : public ITextureFormat
 {
@@ -25,6 +35,43 @@ public:
 
 	}
 
+	virtual const FChildTextureFormat* GetChildFormat() const override final
+	{
+		return this;
+	}
+
+	/**
+	*	Return the texture tiler for this platform. If present, then the texture build process
+	*	will try to use the cached linear encoding for the texture as input rather than rebuilding
+	*	the linear texture prior to tiling. Confusingly, this can not happen if the base texture
+	*	formate returns true for SupportsTiling, as that refers to the texture encoder, not the format.
+	*/
+	virtual const ITextureTiler* GetTiler() const
+	{
+		return nullptr;
+	}
+
+	FName GetBaseFormatName(FName ChildFormatName) const
+	{
+		return FName(*(ChildFormatName.ToString().Replace(*FormatPrefix, TEXT(""))));
+	}
+
+	/**
+	 * Given a platform specific format name, get the texture format object that will do the encoding.
+	 */
+	const ITextureFormat* GetBaseFormatObject(FName FormatName) const
+	{
+		FName BaseFormatName = GetBaseFormatName(FormatName);
+
+		ITextureFormatManagerModule& TFM = GetTextureFormatManagerRef();
+		const ITextureFormat* FormatObject = TFM.FindTextureFormat(BaseFormatName);
+
+		checkf(FormatObject != nullptr, TEXT("Bad FormatName %s passed to FChildTextureFormat::GetBaseFormatObject()"));
+
+		return FormatObject;
+	}
+
+
 protected:
 
 	void AddBaseTextureFormatModules(const TCHAR* ModuleNameWildcard)
@@ -61,10 +108,6 @@ protected:
 		}
 	}
 
-	FName GetBaseFormatName(FName PlatformName) const
-	{
-		return FName(*(PlatformName.ToString().Replace(*FormatPrefix, TEXT(""))));
-	}
 
 	FCbObjectView GetBaseFormatConfigOverride(const FCbObjectView& ObjView) const
 	{
@@ -79,21 +122,6 @@ protected:
 		return BaseSettings;
 	}
 
-	/**
-	 * Given a platform specific format name, get the parent texture format object
-	 */
-	const ITextureFormat* GetBaseFormatObject(FName FormatName) const
-	{
-		FName BaseFormatName = GetBaseFormatName(FormatName);
-
-		ITextureFormatManagerModule& TFM = GetTextureFormatManagerRef();
-		const ITextureFormat* FormatObject = TFM.FindTextureFormat(BaseFormatName);
-
-		checkf(FormatObject != nullptr, TEXT("Bad FormatName %s passed to FChildTextureFormat::GetBaseFormatObject()"));
-
-		return FormatObject;
-	}
-
 	/**
 	 * The final version is a combination of parent and child formats, 8 bits for each
 	 */
diff --git a/Engine/Source/Developer/TextureFormat/Public/Interfaces/ITextureFormat.h b/Engine/Source/Developer/TextureFormat/Public/Interfaces/ITextureFormat.h
--- a/Engine/Source/Developer/TextureFormat/Public/Interfaces/ITextureFormat.h
+++ b/Engine/Source/Developer/TextureFormat/Public/Interfaces/ITextureFormat.h
@@ -7,6 +7,7 @@
 #include "Serialization/CompactBinary.h"
 
 struct FTextureBuildSettings;
+class FChildTextureFormat;
 
 /**
  * Structure for texture format compressor capabilities.
@@ -26,19 +27,102 @@ struct FTextureFormatCompressorCaps
 	uint32 ExtData_DEPRECATED;
 };
 
+/**
+* Holds various engine configuration parameters that can affect the output of a build but
+* should generally be constant across all texture builds. These are sourced from CVars
+* and enums/defines that aren't necessarily visible in all modules.
+* 
+* This structure serializes to compact binary only writing if the values are not default,
+* so changing the default initialization without changing the texture build version/guid
+* can result in build mismatch.
+* 
+* Created via GenerateTextureEngineParameters() in TextureDerivedDataTask.cpp.
+*/
+struct FTextureEngineParameters
+{
+	bool bEngineSupportsVolumeTextureStreaming = true;			// GEngineSupportsVolumeTextureStreaming
+	bool bEngineSupportsTexture2DArrayStreaming = true;			// GEngineSupportsTexture2DArrayStreaming
+	int32 NumInlineDerivedMips = 7;								// NUM_INLINE_DERIVED_MIPS
+};
+
+static bool GetStreamingDisabledForNonVirtualTextureProperties(bool bInCubeMap, bool bInVolumeTexture, bool bInTextureArray, const FTextureEngineParameters& InEngineParameters)
+{
+	if (bInCubeMap)
+	{
+		return true;
+	}
+	if (bInVolumeTexture && InEngineParameters.bEngineSupportsVolumeTextureStreaming == false)
+	{
+		return true;
+	}
+	if (bInTextureArray && InEngineParameters.bEngineSupportsTexture2DArrayStreaming == false)
+	{
+		return true;
+	}
+	return false;
+}
+
+
+
 // Extra data for an encoded texture. For "normal" textures (i.e. linear, without a packed mip tail), this must be
 // all zeroes.
 struct FEncodedTextureExtendedData
 {
-	uint32 NumMipsInTail = 0;
+	int32 NumMipsInTail = 0;
 	uint32 ExtData = 0;
+
+	// With packing/tiling, mip sizes are not trivially computable. Not that these sizes must NOT be
+	// used for mips prior to tiling. For those, FEncodedTextureDescription::GetMipSizeInBytes().
+	TArray<uint64, TInlineAllocator<15 /*MAX_TEXTURE_MIP_COUNT*/>> MipSizesInBytes;
 };
 
+
+/**
+* Calculate the number of streaming mips for the given set of texture properties. This must work off
+* of properties that can (eventually) be calculated without running a full texture build.
+* 
+* Texture mips are split in to two large groups: streaming and non-streaming (aka "inline"). Note
+* that "inline" is sometimes used as a verb to mean "load off of disk and place in our bulk data".
+* "Inline" textures are loaded with the texture asset, and streaming textures are loaded on demand.
+* Generally, 7 of the smallest mips are inlined, however some platforms pack a lot of mips in to a single
+* allocation ("packed mip tail" = NumMipsInTail). Those mips must all be inlined.
+* 
+* InNumMips			    The total mips that the texture contains.
+* InExtendedData	    If the texture is being built for a platform that provides extended data, pass it
+*					    here. For platforms that don't need it (i.e. PC), this should be nullptr.
+* InEngineParameters    Holds a vairety of engine configuration constants, create with GenerateEngineParameters()
+*
+*/
+static int32 GetNumStreamingMipsDirect(int32 InNumMips, bool bInCubeMap, bool bInVolumeTexture, bool bInTextureArray, const FEncodedTextureExtendedData* InExtendedData, const FTextureEngineParameters& InEngineParameters)
+{
+	bool bAllowStreaming = true;
+	{
+		const bool bDisableStreaming = GetStreamingDisabledForNonVirtualTextureProperties(bInCubeMap, bInVolumeTexture, bInTextureArray, InEngineParameters);
+		if (bDisableStreaming)
+			bAllowStreaming = false;
+	}
+
+	int32 NumStreamingMips = 0;
+	if (bAllowStreaming)
+	{
+		// Some platforms pack several mips in to a single entry. If this is the case,
+		// those must be non-streaming.
+		int32 NumMipsInTail = 0;
+		if (InExtendedData)
+		{
+			NumMipsInTail = InExtendedData->NumMipsInTail;
+		}
+
+		int32 NumInlineMips = FMath::Max(NumMipsInTail, InEngineParameters.NumInlineDerivedMips);
+		NumStreamingMips = FMath::Max(0, InNumMips - NumInlineMips);
+	}
+	return NumStreamingMips;
+}
+
 // Everything necessary to know the memory layout for an encoded untiled unpacked texture (i.e. enough information
 // to describe the texture entirely to a PC hardware API).
 // Once a texture gets tiled or gets a packed mip tail, FEncodedTextureEncodedData is additionally
 // required to know the memory layout.
-// This doesn't hold UE specific texture layout information such as NumStreamingMips.
 struct FEncodedTextureDescription
 {	
 	int32 TopMipSizeX;
@@ -50,13 +134,28 @@ struct FEncodedTextureDescription
 	bool bCubeMap;
 	bool bTextureArray;
 	bool bVolumeTexture;
+	
+
+	bool operator==(const FEncodedTextureDescription& OtherTextureDescription) const
+	{
+		return TopMipSizeX == OtherTextureDescription.TopMipSizeX &&
+			TopMipSizeY == OtherTextureDescription.TopMipSizeY &&
+			TopMipVolumeSizeZ == OtherTextureDescription.TopMipVolumeSizeZ &&
+			ArraySlices == OtherTextureDescription.ArraySlices &&
+			PixelFormat == OtherTextureDescription.PixelFormat &&
+			NumMips == OtherTextureDescription.NumMips &&
+			bCubeMap == OtherTextureDescription.bCubeMap &&
+			bTextureArray == OtherTextureDescription.bTextureArray &&
+			bVolumeTexture == OtherTextureDescription.bVolumeTexture;
+	}
 
-	// Convert to the typical "slices" count for the texture.
-	// InMipIndex only matters for volume textures.
-	int32 GetNumSlices(int32 InMipIndex) const
+	// Returns the slice count for usage cases/platform that expect slice count to include
+	// volume texture depth. InMipIndex only affects volume textures.
+	int32 GetNumSlices_WithDepth(int32 InMipIndex) const
 	{
 		if (bVolumeTexture)
 		{
+			check(bTextureArray == false && bCubeMap == false);
 			check(InMipIndex < NumMips);
 			return FMath::Max(TopMipVolumeSizeZ >> InMipIndex, 1);
 		}
@@ -69,6 +168,133 @@ struct FEncodedTextureDescription
 		}
 		return Slices;
 	}
+
+	// Returns the slice count for usage cases/platforms that expect slice count to only include
+	// cubemap/array slices.
+	int32 GetNumSlices_NoDepth() const
+	{
+		if (bVolumeTexture)
+		{
+			check(bTextureArray == false && bCubeMap == false);
+			return 1; // no such thing as a cube volume, or a volume array.
+		}
+
+		check((bTextureArray && ArraySlices > 1) || (!bTextureArray && ArraySlices == 1));
+		int32 Slices = ArraySlices;
+		if (bCubeMap)
+		{
+			Slices *= 6;
+		}
+		return Slices;
+	}
+
+	FIntVector3 GetMipDimensions(int32 InMipIndex) const
+	{
+		FIntVector3 Results;
+		Results.X = FMath::Max(TopMipSizeX >> InMipIndex, 1);
+		Results.Y = FMath::Max(TopMipSizeY >> InMipIndex, 1);
+		Results.Z = 1;
+		if (bVolumeTexture)
+		{
+			Results.Z = FMath::Max(TopMipVolumeSizeZ >> InMipIndex, 1);
+		}
+		return Results;
+	}
+
+	uint64 GetMipSizeInBytes(int32 InMipIndex) const
+	{
+		FIntVector3 MipDims = GetMipDimensions(InMipIndex);
+		uint64 SliceByteCount = GPixelFormats[PixelFormat].Get2DImageSizeInBytes(MipDims.X, MipDims.Y);
+		return SliceByteCount * GetNumSlices_WithDepth(InMipIndex);
+	}
+
+	int32 GetNumStreamingMips(const FEncodedTextureExtendedData* InExtendedData, const FTextureEngineParameters& InEngineParameters) const
+	{
+		return GetNumStreamingMipsDirect(NumMips, bCubeMap, bVolumeTexture, bTextureArray, InExtendedData, InEngineParameters);
+	}
+
+	// Convenience function for iterating over the encoded mips when you need to know how many mips are represented. Use as:
+	//
+	//	for (int32 EncodedMipIndex = 0; EncodedMipIndex < OutMipTailIndex + 1; EncodedMipIndex++)
+	//	{
+	//		int32 MipsRepresentedThisIndex = EncodedMipIndex == OutMipTailIndex ? OutMipsInTail : 1;
+	//	}
+	//
+	// This handles mip chains whether or not they have packed mip tails.
+	// Note GetNumEncodedMips() == OutMipTailIndex + 1
+	//
+	void GetEncodedMipIterators(const FEncodedTextureExtendedData* InExtendedData, int32& OutMipTailIndex, int32& OutMipsInTail) const
+	{
+		OutMipTailIndex = NumMips - 1;
+		OutMipsInTail = 1;
+		if (InExtendedData && InExtendedData->NumMipsInTail > 1)
+		{
+			OutMipsInTail = InExtendedData->NumMipsInTail;
+			OutMipTailIndex = NumMips - OutMipsInTail;
+		}
+	}
+
+	// Returns the number of mips that actually carry bulk data for this texture. Nominally the number of total mips,
+	// however some platforms have packed mip tails, which means they still have the total number of mips, but the last
+	// several are all bundled together for memory savings.
+	int32 GetNumEncodedMips(const FEncodedTextureExtendedData* InExtendedData) const
+	{
+		if (InExtendedData &&
+			InExtendedData->NumMipsInTail > 1)
+		{
+			return NumMips - InExtendedData->NumMipsInTail + 1;
+		}
+		return NumMips;
+	}
+
+	FEncodedTextureDescription GetDescriptionForMipLevel(const FEncodedTextureExtendedData* InExtendedData, int32 InMipIndex) const
+	{
+		FEncodedTextureDescription MipTextureDescription = *this;
+		FIntVector3 TailFirstMipDims = GetMipDimensions(InMipIndex);
+		MipTextureDescription.TopMipSizeX = TailFirstMipDims.X;
+		MipTextureDescription.TopMipSizeY = TailFirstMipDims.Y;
+		MipTextureDescription.TopMipVolumeSizeZ = TailFirstMipDims.Z;
+		MipTextureDescription.NumMips = 1;
+		if (InExtendedData && InExtendedData->NumMipsInTail && InMipIndex >= NumMips - InExtendedData->NumMipsInTail)
+		{
+			// we must only ever get the first mip tail index!
+			check(InMipIndex == NumMips - InExtendedData->NumMipsInTail);
+
+			// We want the layout for the entire tail.
+			MipTextureDescription.NumMips = InExtendedData->NumMipsInTail;
+		}		
+		return MipTextureDescription;
+	}
+};
+
+/**
+*	Interface for platform formats that consume a linear, unpacked texture that an be built on
+*	a host platform (e.g. windows) and then tile/pack it as necessary.
+*/
+class ITextureTiler
+{
+public:
+	/**
+	*	The generic texture tiling build function expects the following functions to exist that
+	*	do what they say on the tin.
+	* 
+	*	static const FUtf8StringView GetBuildFunctionNameStatic()
+	*	static FGuid GetBuildFunctionVersionGuid()
+	*/
+	
+	/**
+	* Generate and return any out-of-band data that needs to be saved for a given encoded texture description.
+	*/
+	virtual FEncodedTextureExtendedData GetExtendedDataForTexture(const FEncodedTextureDescription& InTextureDescription) const = 0;
+
+	virtual const FUtf8StringView GetBuildFunctionName() const = 0;
+
+	/**
+		InLinearSurfaces must have the necessary input mips for the mip level - i.e. for a packed mip tail,
+		InMipIndex is the index of the top mip of the tail, and InLinearSurfaces must have all the source mips
+		for the entire tail.
+	*/
+	virtual FSharedBuffer ProcessMipLevel(const FEncodedTextureDescription& InTextureDescription, const FEncodedTextureExtendedData& InExtendedData, TArrayView<FMemoryView> InLinearSurfaces, int32 InMipIndex) const = 0;
 };
 
 /**
@@ -277,7 +503,7 @@ public:
 	 * @returns true if tiling is supported, false if it must be done by the caller
 	 *
 	 */
-	virtual bool SupportsTiling(const FTextureBuildSettings& BuildSettings) const
+	virtual bool SupportsTiling() const
 	{
 		return false;
 	}
@@ -357,6 +583,11 @@ public:
 	{
 		return FString();
 	}
+
+	virtual const FChildTextureFormat* GetChildFormat() const
+	{
+		return nullptr;
+	}
 	
 	/**
 	 * Identify the latest sdk version for this texture encoder
diff --git a/Engine/Source/Runtime/Core/Private/HAL/ConsoleManager.cpp b/Engine/Source/Runtime/Core/Private/HAL/ConsoleManager.cpp
--- a/Engine/Source/Runtime/Core/Private/HAL/ConsoleManager.cpp
+++ b/Engine/Source/Runtime/Core/Private/HAL/ConsoleManager.cpp
@@ -3138,6 +3138,12 @@ static TAutoConsoleVariable<float> CVarSkeletalMeshLODRadiusScale(
 	TEXT("Scale factor for the screen radius used in computing discrete LOD for skeletal meshes. (0.25-1)"),
 	ECVF_Scalability | ECVF_RenderThreadSafe);
 
+static TAutoConsoleVariable<int32> CVarSharedLinearTextureEncoding(
+	TEXT("r.SharedLinearTextureEncoding"),
+	1,
+	TEXT("If set to 1, textures for platforms that tile will reuse a host linear texture instead of reencoding."),
+	ECVF_RenderThreadSafe);
+
 static TAutoConsoleVariable<int32> CVarPreTileTextures(
 	TEXT("r.PreTileTextures"),
 	1,
diff --git a/Engine/Source/Runtime/Engine/Classes/Engine/Texture.h b/Engine/Source/Runtime/Engine/Classes/Engine/Texture.h
--- a/Engine/Source/Runtime/Engine/Classes/Engine/Texture.h
+++ b/Engine/Source/Runtime/Engine/Classes/Engine/Texture.h
@@ -828,18 +828,19 @@ struct FTexturePlatformData
 
 	struct FStructuredDerivedDataKey
 	{
+		FIoHash TilingBuildDefinitionKey; // All zeroes if the derived data didn't have a child build.
 		FIoHash BuildDefinitionKey;
 		FGuid SourceGuid;
 		FGuid CompositeSourceGuid;
 
 		bool operator==(const FStructuredDerivedDataKey& Other) const
 		{
-			return BuildDefinitionKey == Other.BuildDefinitionKey && SourceGuid == Other.SourceGuid && CompositeSourceGuid == Other.CompositeSourceGuid;
+			return TilingBuildDefinitionKey == Other.TilingBuildDefinitionKey && BuildDefinitionKey == Other.BuildDefinitionKey && SourceGuid == Other.SourceGuid && CompositeSourceGuid == Other.CompositeSourceGuid;
 		}
 
 		bool operator!=(const FStructuredDerivedDataKey& Other) const
 		{
-			return BuildDefinitionKey != Other.BuildDefinitionKey || SourceGuid != Other.SourceGuid || CompositeSourceGuid != Other.CompositeSourceGuid;
+			return TilingBuildDefinitionKey != Other.TilingBuildDefinitionKey || BuildDefinitionKey != Other.BuildDefinitionKey || SourceGuid != Other.SourceGuid || CompositeSourceGuid != Other.CompositeSourceGuid;
 		}
 	};
 
@@ -902,6 +903,12 @@ public:
 	 */
 	void SerializeCooked(FArchive& Ar, class UTexture* Owner, bool bStreamable);
 	
+	
+	inline void SetPackedData(int32 InNumSlices, bool bInHasOptData, bool bInCubeMap)
+	{
+		PackedData = (InNumSlices & BitMask_NumSlices) | (bInCubeMap ? BitMask_CubeMap : 0) | (bInHasOptData ? BitMask_HasOptData : 0);
+	}
+
 	inline bool GetHasOptData() const
 	{
 		return (PackedData & BitMask_HasOptData) == BitMask_HasOptData;
@@ -947,7 +954,6 @@ public:
 	}
 
 #if WITH_EDITOR
-	static bool IsUsingNewDerivedData();
 	bool IsAsyncWorkComplete() const;
 
 	// Compresses the texture using the given compressor and adds the result to the DDC.
diff --git a/Engine/Source/Runtime/Engine/Engine.Build.cs b/Engine/Source/Runtime/Engine/Engine.Build.cs
--- a/Engine/Source/Runtime/Engine/Engine.Build.cs
+++ b/Engine/Source/Runtime/Engine/Engine.Build.cs
@@ -355,6 +355,7 @@ public class Engine : ModuleRules
 			);
 
 			PrivateDependencyModuleNames.Add("DerivedDataCache");
+			PrivateDependencyModuleNames.Add("TextureCompressor");
 
 			PrivateIncludePathModuleNames.Add("TextureCompressor");
 			PrivateIncludePaths.Add("Developer/TextureCompressor/Public");
diff --git a/Engine/Source/Runtime/Engine/Private/Texture.cpp b/Engine/Source/Runtime/Engine/Private/Texture.cpp
--- a/Engine/Source/Runtime/Engine/Private/Texture.cpp
+++ b/Engine/Source/Runtime/Engine/Private/Texture.cpp
@@ -1327,6 +1327,10 @@ bool UTexture::IsCandidateForTextureStreaming(const ITargetPlatform* InTargetPla
 		MipGenSettings != TMGS_NoMipmaps)
 	{
 		// If bCookedIsStreamable flag was previously computed, use it.
+		// This is computed _after_ the derived data is cached. It's only Reset() and used in the SerializeCookedPlatformData
+		// path.. but not the BeginCacheForCookedPlatformData path. Need to know what happens if you do
+		// one, then the other, as the flag will be set. Also note that this is an OR when multiple platforms
+		// are cooked... so this can be true for a non-streaming platform if it gets cooked with a streaming platform!
 		if (bCookedIsStreamable.IsSet())
 		{
 			return *bCookedIsStreamable;
diff --git a/Engine/Source/Runtime/Engine/Private/TextureDerivedData.cpp b/Engine/Source/Runtime/Engine/Private/TextureDerivedData.cpp
--- a/Engine/Source/Runtime/Engine/Private/TextureDerivedData.cpp
+++ b/Engine/Source/Runtime/Engine/Private/TextureDerivedData.cpp
@@ -34,6 +34,9 @@
 
 #if WITH_EDITOR
 
+#include "ChildTextureFormat.h"
+#include "ColorSpace.h"
+#include "Compression/OodleDataCompressionUtil.h"
 #include "DerivedDataCache.h"
 #include "DerivedDataCacheInterface.h"
 #include "DerivedDataRequestOwner.h"
@@ -47,8 +50,6 @@
 #include "VT/LightmapVirtualTexture.h"
 #include "TextureCompiler.h"
 #include "TextureEncodingSettings.h"
-#include "ColorSpace.h"
-#include "Compression/OodleDataCompressionUtil.h"
 
 
 /*------------------------------------------------------------------------------
@@ -71,6 +72,26 @@
 #define TEXTURE_VT_DERIVEDDATA_VER	TEXT("7C16439390E24F1F9468894FB4D4BC54")
 
 
+static bool IsUsingNewDerivedData()
+{
+	struct FTextureDerivedDataSetting
+	{
+		FTextureDerivedDataSetting()
+		{
+			bUseNewDerivedData = FParse::Param(FCommandLine::Get(), TEXT("DDC2AsyncTextureBuilds")) || FParse::Param(FCommandLine::Get(), TEXT("DDC2TextureBuilds"));
+			if (!bUseNewDerivedData)
+			{
+				GConfig->GetBool(TEXT("TextureBuild"), TEXT("NewTextureBuilds"), bUseNewDerivedData, GEditorIni);
+			}
+			UE_CLOG(bUseNewDerivedData, LogTexture, Log, TEXT("Using new texture derived data builds."));
+		}
+		bool bUseNewDerivedData;
+	};
+	static const FTextureDerivedDataSetting TextureDerivedDataSetting;
+	return TextureDerivedDataSetting.bUseNewDerivedData;
+}
+
+
 #if ENABLE_COOK_STATS
 namespace TextureCookStats
 {
@@ -587,7 +608,28 @@ static void FinalizeBuildSettingsForLayer(
 					OutBuildResultMetadata->bIsValid = true;
 					OutBuildResultMetadata->bSupportsEncodeSpeed = bSupportsEncodeSpeed;
 				}
-			}
+			
+				if (IsUsingNewDerivedData()) // shared linear encoding only hooked up to ddc2 atm.
+				{
+					static auto CVarSharedLinearTextureEncoding = IConsoleManager::Get().FindTConsoleVariableDataInt(TEXT("r.SharedLinearTextureEncoding"));
+					if (CVarSharedLinearTextureEncoding->GetValueOnAnyThread())
+					{
+						// Shared linear encoding can only work if the base texture format does not expect to
+						// do the tiling itself (SupportsTiling == false).
+						const FChildTextureFormat* ChildTextureFormat = TextureFormat->GetChildFormat();
+						if (ChildTextureFormat && ChildTextureFormat->GetBaseFormatObject(OutSettings.TextureFormatName)->SupportsTiling() == false)
+						{
+							OutSettings.Tiler = ChildTextureFormat->GetTiler();
+							if (OutSettings.Tiler)
+							{
+								// Since the tiler will do the encoding, skip the child texture format entirely and just look for the
+								// base texture and then run the tiler after.
+								OutSettings.TextureFormatName = ChildTextureFormat->GetBaseFormatName(OutSettings.TextureFormatName);
+							}
+						}
+					} // end if enabled
+				} // end if ddc2
+			} // end if texture format found.
 		}
 	}
 
@@ -875,7 +917,6 @@ static void GetTextureBuildSettings(
 	OutBuildSettings.CompositePower = Texture.CompositePower;
 	OutBuildSettings.LODBias = TextureLODSettings.CalculateLODBias(SourceSize.X, SourceSize.Y, Texture.MaxTextureSize, Texture.LODGroup, Texture.LODBias, Texture.NumCinematicMipLevels, Texture.MipGenSettings, bVirtualTextureStreaming);
 	OutBuildSettings.LODBiasWithCinematicMips = TextureLODSettings.CalculateLODBias(SourceSize.X, SourceSize.Y, Texture.MaxTextureSize, Texture.LODGroup, Texture.LODBias, 0, Texture.MipGenSettings, bVirtualTextureStreaming);
-	OutBuildSettings.bStreamable = GetTextureIsStreamableOnPlatform(Texture, TargetPlatform);
 	OutBuildSettings.bVirtualStreamable = bVirtualTextureStreaming;
 	OutBuildSettings.PowerOfTwoMode = Texture.PowerOfTwoMode;
 	OutBuildSettings.PaddingColor = Texture.PaddingColor;
@@ -1320,25 +1361,6 @@ void FTexturePlatformData::CancelCache()
 	}
 }
 
-bool FTexturePlatformData::IsUsingNewDerivedData()
-{
-	struct FTextureDerivedDataSetting
-	{
-		FTextureDerivedDataSetting()
-		{
-			bUseNewDerivedData = FParse::Param(FCommandLine::Get(), TEXT("DDC2AsyncTextureBuilds")) || FParse::Param(FCommandLine::Get(), TEXT("DDC2TextureBuilds"));
-			if (!bUseNewDerivedData)
-			{
-				GConfig->GetBool(TEXT("TextureBuild"), TEXT("NewTextureBuilds"), bUseNewDerivedData, GEditorIni);
-			}
-			UE_CLOG(bUseNewDerivedData, LogTexture, Log, TEXT("Using new texture derived data builds."));
-		}
-		bool bUseNewDerivedData;
-	};
-	static const FTextureDerivedDataSetting TextureDerivedDataSetting;
-	return TextureDerivedDataSetting.bUseNewDerivedData;
-}
-
 bool FTexturePlatformData::IsAsyncWorkComplete() const
 {
 	return !AsyncTask || AsyncTask->Poll();
@@ -2804,7 +2826,7 @@ void UTexture::CachePlatformData(bool bAsyncCache, bool bAllowAsyncBuild, bool b
 
 				// Check if our keys match. If we have two, they both have to match, otherwise a change that only affects one
 				// might not cause a rebuild, leading to confusion in the texture editor.
-				if (FTexturePlatformData::IsUsingNewDerivedData() && (Source.GetNumLayers() == 1) && !BuildSettingsFetchOrBuild[0].bVirtualStreamable)
+				if (IsUsingNewDerivedData() && (Source.GetNumLayers() == 1) && !BuildSettingsFetchOrBuild[0].bVirtualStreamable)
 				{
 					// DDC2 version
 					using namespace UE::DerivedData;
diff --git a/Engine/Source/Runtime/Engine/Private/TextureDerivedDataBuildUtils.cpp b/Engine/Source/Runtime/Engine/Private/TextureDerivedDataBuildUtils.cpp
--- a/Engine/Source/Runtime/Engine/Private/TextureDerivedDataBuildUtils.cpp
+++ b/Engine/Source/Runtime/Engine/Private/TextureDerivedDataBuildUtils.cpp
@@ -3,6 +3,7 @@
 #include "TextureDerivedDataBuildUtils.h"
 
 #if WITH_EDITOR
+#include "ChildTextureFormat.h"
 #include "DerivedDataBuild.h"
 #include "DerivedDataBuildFunctionRegistry.h"
 #include "DerivedDataSharedString.h"
@@ -162,7 +163,7 @@ static void WriteBuildSettings(FCbWriter& Writer, const FTextureBuildSettings& B
 	WriteCbFieldWithDefault(Writer, "TopMipSize", BuildSettings.TopMipSize, DefaultSettings.TopMipSize);
 	WriteCbFieldWithDefault(Writer, "VolumeSizeZ", BuildSettings.VolumeSizeZ, DefaultSettings.VolumeSizeZ);
 	WriteCbFieldWithDefault(Writer, "ArraySlices", BuildSettings.ArraySlices, DefaultSettings.ArraySlices);
-	WriteCbFieldWithDefault<bool>(Writer, "bStreamable", BuildSettings.bStreamable, DefaultSettings.bStreamable);
+	WriteCbFieldWithDefault<bool>(Writer, "bStreamable", BuildSettings.bStreamable_Unused, DefaultSettings.bStreamable_Unused);
 	WriteCbFieldWithDefault<bool>(Writer, "bVirtualStreamable", BuildSettings.bVirtualStreamable, DefaultSettings.bVirtualStreamable);
 	WriteCbFieldWithDefault<bool>(Writer, "bChromaKeyTexture", BuildSettings.bChromaKeyTexture, DefaultSettings.bChromaKeyTexture);
 	WriteCbFieldWithDefault(Writer, "PowerOfTwoMode", BuildSettings.PowerOfTwoMode, DefaultSettings.PowerOfTwoMode);
@@ -197,15 +198,6 @@ static void WriteBuildSettings(FCbWriter& Writer, const FTextureBuildSettings& B
 	Writer.EndObject();
 }
 
-static void WriteOutputSettings(FCbWriter& Writer, int32 NumInlineMips)
-{
-	Writer.BeginObject();
-
-	Writer.AddInteger("NumInlineMips", NumInlineMips);
-
-	Writer.EndObject();
-}
-
 static void WriteSource(FCbWriter& Writer, const UTexture& Texture, int32 LayerIndex, const FTextureBuildSettings& BuildSettings)
 {
 	const FTextureSource& Source = Texture.Source;
@@ -293,7 +285,7 @@ UE::DerivedData::FUtf8SharedString FindTextureBuildFunction(const FName TextureF
 	return Function;
 }
 
-FCbObject SaveTextureBuildSettings(const UTexture& Texture, const FTextureBuildSettings& BuildSettings, int32 LayerIndex, int32 NumInlineMips, bool bUseCompositeTexture, int64 RequiredMemoryEstimate)
+FCbObject SaveTextureBuildSettings(const UTexture& Texture, const FTextureBuildSettings& BuildSettings, int32 LayerIndex, bool bUseCompositeTexture, int64 RequiredMemoryEstimate)
 {
 	const ITextureFormat* TextureFormat = nullptr;
 	if (ITextureFormatManagerModule* TFM = GetTextureFormatManager())
@@ -322,9 +314,6 @@ FCbObject SaveTextureBuildSettings(const UTexture& Texture, const FTextureBuildS
 	Writer.SetName("Build");
 	WriteBuildSettings(Writer, BuildSettings, TextureFormat);
 
-	Writer.SetName("Output");
-	WriteOutputSettings(Writer, NumInlineMips);
-
 	Writer.SetName("Source");
 	WriteSource(Writer, Texture, LayerIndex, BuildSettings);
 
diff --git a/Engine/Source/Runtime/Engine/Private/TextureDerivedDataBuildUtils.h b/Engine/Source/Runtime/Engine/Private/TextureDerivedDataBuildUtils.h
--- a/Engine/Source/Runtime/Engine/Private/TextureDerivedDataBuildUtils.h
+++ b/Engine/Source/Runtime/Engine/Private/TextureDerivedDataBuildUtils.h
@@ -13,6 +13,6 @@ class UTexture;
 struct FTextureBuildSettings;
 
 UE::DerivedData::FUtf8SharedString FindTextureBuildFunction(FName TextureFormatName);
-FCbObject SaveTextureBuildSettings(const UTexture& Texture, const FTextureBuildSettings& BuildSettings, int32 LayerIndex, int32 NumInlineMips, bool bUseCompositeTexture, int64 RequiredMemoryEstimate);
+FCbObject SaveTextureBuildSettings(const UTexture& Texture, const FTextureBuildSettings& BuildSettings, int32 LayerIndex, bool bUseCompositeTexture, int64 RequiredMemoryEstimate);
 
 #endif // WITH_EDITOR
diff --git a/Engine/Source/Runtime/Engine/Private/TextureDerivedDataTask.cpp b/Engine/Source/Runtime/Engine/Private/TextureDerivedDataTask.cpp
--- a/Engine/Source/Runtime/Engine/Private/TextureDerivedDataTask.cpp
+++ b/Engine/Source/Runtime/Engine/Private/TextureDerivedDataTask.cpp
@@ -25,6 +25,7 @@
 
 #include "Algo/Accumulate.h"
 #include "Algo/AnyOf.h"
+#include "ChildTextureFormat.h"
 #include "DerivedDataBuild.h"
 #include "DerivedDataBuildAction.h"
 #include "DerivedDataBuildInputResolver.h"
@@ -51,7 +52,7 @@
 #include "Misc/ScopeExit.h"
 #include "ProfilingDebugging/CookStats.h"
 #include "Serialization/BulkDataRegistry.h"
-#include "TextureCompressorModule.h"
+#include "Serialization/CompactBinaryWriter.h"
 #include "TextureBuildUtilities.h"
 #include "TextureDerivedDataBuildUtils.h"
 #include "TextureEncodingSettings.h"
@@ -74,8 +75,24 @@ static TAutoConsoleVariable<int32> CVarVTValidateCompressionOnSave(
 	TEXT("This is slow, but allows debugging corrupt VT data")
 );
 
+static TAutoConsoleVariable<int32> CVarForceRetileTextures(
+	TEXT("r.ForceRetileTextures"),
+	0,
+	TEXT("If r.SharedLinearTextureEncoding is enabled, this will force the tiling build step to rebuild,")
+	TEXT("however the linear texture is allowed to fetch from cache.")
+);
+
 void GetTextureDerivedDataKeyFromSuffix(const FString& KeySuffix, FString& OutKey);
 
+static FTextureEngineParameters GenerateTextureEngineParameters()
+{
+	FTextureEngineParameters EngineParameters;
+	EngineParameters.bEngineSupportsTexture2DArrayStreaming = GSupportsTexture2DArrayStreaming;
+	EngineParameters.bEngineSupportsVolumeTextureStreaming = GSupportsVolumeTextureStreaming;
+	EngineParameters.NumInlineDerivedMips = NUM_INLINE_DERIVED_MIPS;
+	return EngineParameters;
+}
+
 class FTextureStatusMessageContext : public FScopedSlowTask
 {
 public:
@@ -306,11 +323,127 @@ void FTextureSourceData::GetAsyncSourceMips(IImageWrapperModule* InImageWrapper)
 	}
 }
 
+
+// When texture streaming is disabled, all of the mips are packed into a single FBulkData/FDerivedData
+// and "inlined", meaning they are saved and loaded as part of the serialized asset data.
+static bool GetBuildSettingsDisablesStreaming(const FTextureBuildSettings& InBuildSettings, const FTextureEngineParameters& InEngineParameters)
+{
+	if (InBuildSettings.bVirtualStreamable)
+	{
+		// Only basic 2d textures can be virtual streamable.
+		return InBuildSettings.bCubemap || InBuildSettings.bVolume || InBuildSettings.bTextureArray;
+	}
+	else
+	{
+		return GetStreamingDisabledForNonVirtualTextureProperties(InBuildSettings.bCubemap, InBuildSettings.bVolume, InBuildSettings.bTextureArray, InEngineParameters);
+	}
+}
+
 namespace UE::TextureDerivedData
 {
 
 using namespace UE::DerivedData;
 
+// Handle converting the tiling build's inputs from names to concrete values.
+class FTilingTextureBuildInputResolver final : public IBuildInputResolver
+{
+public:
+	explicit FTilingTextureBuildInputResolver(UE::DerivedData::FBuildSession& InParentBuild_Session, UE::DerivedData::FBuildDefinition& InParentBuild_Definition, UE::DerivedData::FBuildPolicy& InParentBuild_Policy)
+		: ParentBuild_Session(InParentBuild_Session),
+		ParentBuild_Definition(InParentBuild_Definition),
+		ParentBuild_Policy(InParentBuild_Policy)
+	{
+	}
+
+	UE::DerivedData::FBuildSession& ParentBuild_Session;
+	UE::DerivedData::FBuildDefinition ParentBuild_Definition;
+	UE::DerivedData::FBuildPolicy ParentBuild_Policy;
+
+	UE::DerivedData::FOptionalBuildOutput ParentBuild_Output;
+	bool bParentBuild_HitCache = false;
+
+	// Convert from named keys to hash/size pairs. There is no expectation that the results are
+	// ready when this function returns - InResolvedCallback is called when the results arrive.
+	void ResolveInputMeta(
+		const FBuildDefinition& InDefinition,
+		IRequestOwner& InOwner,
+		FOnBuildInputMetaResolved&& InResolvedCallback) final
+	{
+		// This should only ever be called once, so the build should never be ready.
+		check(!ParentBuild_Output);
+
+		// Kick the parent build so we can get access to the results. When we get the results,
+		// do the actual input resolution.
+		ParentBuild_Session.Build(ParentBuild_Definition, {}, ParentBuild_Policy, InOwner,
+			[this, InResolvedCallback = MoveTemp(InResolvedCallback), &InDefinition](UE::DerivedData::FBuildCompleteParams&& InParams)
+		{
+			this->ParentBuild_Output = MoveTemp(InParams.Output);
+
+			if (InParams.Status == UE::DerivedData::EStatus::Canceled)
+			{
+				return;
+			}
+
+			this->bParentBuild_HitCache = EnumHasAnyFlags(InParams.BuildStatus, EBuildStatus::CacheQueryHit);
+
+			TArray<UE::DerivedData::FBuildInputMetaByKey, TInlineAllocator<8>> Inputs;
+
+			UE::DerivedData::EStatus Status = UE::DerivedData::EStatus::Ok;
+			InDefinition.IterateInputBuilds([this, &Status, &Inputs](FUtf8StringView InOurKey, const UE::DerivedData::FBuildValueKey& InBuildValueKey)
+			{
+				if (InBuildValueKey.BuildKey != this->ParentBuild_Definition.GetKey())
+				{
+					Status = UE::DerivedData::EStatus::Error;
+					return;
+				}
+				const UE::DerivedData::FValueWithId& ParentBuildValue = this->ParentBuild_Output.Get().GetValue(InBuildValueKey.Id);
+				if (ParentBuildValue.IsNull())
+				{
+					Status = UE::DerivedData::EStatus::Error;
+					return;
+				}
+				Inputs.Add({InOurKey, ParentBuildValue.GetRawHash(), ParentBuildValue.GetRawSize()});
+			});
+
+			if (Status != UE::DerivedData::EStatus::Ok)
+			{
+				return InResolvedCallback({{}, Status});
+			}
+			return InResolvedCallback({Inputs, UE::DerivedData::EStatus::Ok});
+		});
+	} // end ResolvedInputMeta
+
+	void ResolveInputData(
+		const FBuildDefinition & Definition,
+		IRequestOwner & Owner,
+		FOnBuildInputDataResolved && OnResolved,
+		FBuildInputFilter && Filter) final
+	{
+		EStatus Status = EStatus::Ok;
+		TArray<FBuildInputDataByKey, TInlineAllocator<8>> Inputs;
+		Definition.IterateInputBuilds([this, &Filter, &Status, &Inputs](FUtf8StringView InOurKey, const UE::DerivedData::FBuildValueKey& InBuildValueKey)
+		{
+			if (Filter && Filter(InOurKey) == false)
+			{
+				return;
+			}
+
+			const UE::DerivedData::FValueWithId& ParentBuildValue = ParentBuild_Output.Get().GetValue(InBuildValueKey.Id);
+			if (!ParentBuildValue || !ParentBuildValue.HasData())
+			{
+				Status = UE::DerivedData::EStatus::Error;
+				return;
+			}
+			
+			Inputs.Add({InOurKey, ParentBuildValue.GetData()});
+		});
+
+		OnResolved({ Inputs, Status });
+	}
+private:
+
+};
+
 class FTextureBuildInputResolver final : public IBuildInputResolver
 {
 public:
@@ -486,11 +619,6 @@ void FTextureCacheDerivedDataWorker::BuildTexture(TArray<FTextureBuildSettings>&
 				*TexturePathName, TextureData.Layers.Num());
 		}
 
-		int32 PredictedNumMips;
-		{
-			const TArray<FImage>& SourceMips = TextureData.Blocks[0].MipsPerLayer[0];
-			PredictedNumMips = Compressor->GetMipCountForBuildSettings(SourceMips[0].SizeX, SourceMips[0].SizeY, SourceMips[0].NumSlices, SourceMips.Num(), InBuildSettingsPerLayer[0]);
-		}
 
 		check(DerivedData->Mips.Num() == 0);
 		DerivedData->SizeX = 0;
@@ -502,6 +630,7 @@ void FTextureCacheDerivedDataWorker::BuildTexture(TArray<FTextureBuildSettings>&
 		FOptTexturePlatformData OptData;
 
 		// Compress the texture by calling texture compressor directly.
+		bool bImageHasAlphaChannel = false;
 		TArray<FCompressedImage2D> CompressedMips;
 		if (Compressor->BuildTexture(TextureData.Blocks[0].MipsPerLayer[0],
 			((bool)Texture.CompositeTexture && CompositeTextureData.Blocks.Num() && CompositeTextureData.Blocks[0].MipsPerLayer.Num()) ? CompositeTextureData.Blocks[0].MipsPerLayer[0] : TArray<FImage>(),
@@ -509,7 +638,8 @@ void FTextureCacheDerivedDataWorker::BuildTexture(TArray<FTextureBuildSettings>&
 			TexturePathName,
 			CompressedMips,
 			OptData.NumMipsInTail,
-			OptData.ExtData))
+			OptData.ExtData,
+			&bImageHasAlphaChannel))
 		{
 			check(CompressedMips.Num());
 
@@ -580,9 +710,6 @@ void FTextureCacheDerivedDataWorker::BuildTexture(TArray<FTextureBuildSettings>&
 
 			DerivedData->SetOptData(OptData);
 
-			// verify PredictedNumMips matches
-			check (PredictedNumMips == CompressedMips.Num());
-
 			// Store it in the cache.
 			// @todo: This will remove the streaming bulk data, which we immediately reload below!
 			// Should ideally avoid this redundant work, but it only happens when we actually have 
@@ -599,11 +726,16 @@ void FTextureCacheDerivedDataWorker::BuildTexture(TArray<FTextureBuildSettings>&
 
 		if (DerivedData->Mips.Num())
 		{
+			bSucceeded = true;
+
 			const bool bInlineMips = EnumHasAnyFlags(CacheFlags, ETextureCacheFlags::InlineMips);
-			bSucceeded = !bInlineMips || DerivedData->TryInlineMipData(InBuildSettingsPerLayer[0].LODBiasWithCinematicMips, TexturePathName);
-			if (!bSucceeded)
+			if (bInlineMips) // Note that mips are inlined when cooking.
 			{
-				UE_LOG(LogTexture, Display, TEXT("Failed to put and then read back mipmap data from DDC for %s"), *TexturePathName);
+				bSucceeded = DerivedData->TryInlineMipData(InBuildSettingsPerLayer[0].LODBiasWithCinematicMips, TexturePathName);
+				if (bSucceeded == false)
+				{
+					UE_LOG(LogTexture, Display, TEXT("Failed to put and then read back mipmap data from DDC for %s"), *TexturePathName);
+				}
 			}
 		}
 		else
@@ -1538,6 +1670,342 @@ void FTextureCacheDerivedDataWorker::Finalize()
 	}
 }
 
+
+
+struct FBuildResultOptions
+{
+	bool bLoadStreamingMips;
+	int32 FirstStreamingMipToLoad;
+};
+
+static bool UnpackPlatformDataFromBuild(FTexturePlatformData& OutPlatformData, UE::DerivedData::FBuildCompleteParams&& InBuildCompleteParams, FBuildResultOptions InBuildResultOptions)
+{
+	using namespace UE::DerivedData;
+	UE::DerivedData::FBuildOutput& BuildOutput = InBuildCompleteParams.Output;
+
+	// We take this as a build output, however in ideal (future) situations, this is generated prior to build launch and
+	// just routed through the build. Since we currently handle several varying situations, we just always consume it from
+	// the build no matter where it came from. (Both TextureDescription and ExtendedData)
+	FEncodedTextureDescription EncodedTextureDescription;
+	{
+		const FValueWithId& Value = BuildOutput.GetValue(FValueId::FromName(ANSITEXTVIEW("EncodedTextureDescription")));
+		UE::TextureBuildUtilities::EncodedTextureDescription::FromCompactBinary(EncodedTextureDescription, FCbObject(Value.GetData().Decompress()));
+	}
+
+	FEncodedTextureExtendedData EncodedTextureExtendedData;
+	{
+		const FValueWithId& Value = BuildOutput.GetValue(FValueId::FromName(ANSITEXTVIEW("EncodedTextureExtendedData")));
+		UE::TextureBuildUtilities::EncodedTextureExtendedData::FromCompactBinary(EncodedTextureExtendedData, FCbObject(Value.GetData().Decompress()));
+	}
+
+	// consider putting this in the build output so that it's only ever polled in one place.
+	const FTextureEngineParameters EngineParameters = GenerateTextureEngineParameters();
+	int32 NumStreamingMips = EncodedTextureDescription.GetNumStreamingMips(&EncodedTextureExtendedData, EngineParameters);
+	int32 NumEncodedMips = EncodedTextureDescription.GetNumEncodedMips(&EncodedTextureExtendedData);
+	check(NumEncodedMips >= NumStreamingMips);
+
+	//
+	// Mips are split up:
+	//	Streaming mips are all stored independently under value FTexturePlatformData::MakeMipId(MipIndex);
+	//	Nonstreaming ("inlined") mips are stored in one buffer under value "MipTail". To disentangle the separate mips
+	//	we need their size. If EncodedTextureExtendedData is present, we have the ::MipSizesInBytes field. If that isn't
+	//	present, then we know it's nothing fancy and we can directly compute mip sizes.
+	//
+	TArray<uint64, TInlineAllocator<MAX_TEXTURE_MIP_COUNT>> MipSizesFromBuild;
+	if (EncodedTextureExtendedData.MipSizesInBytes.Num() == 0)
+	{
+		// Normal mip sizes - generate from texture description
+		MipSizesFromBuild.AddUninitialized(EncodedTextureDescription.NumMips);
+		for (int32 MipIndex = 0; MipIndex < EncodedTextureDescription.NumMips; MipIndex++)
+		{
+			MipSizesFromBuild[MipIndex] = EncodedTextureDescription.GetMipSizeInBytes(MipIndex);
+		}
+	}
+
+	// For the moment, we don't necessarily have the sizes from the extended data, in that case we must have them from the build.
+	TArray<uint64, TInlineAllocator<MAX_TEXTURE_MIP_COUNT>>& MipSizesInBytes = EncodedTextureExtendedData.MipSizesInBytes.Num() ? EncodedTextureExtendedData.MipSizesInBytes : MipSizesFromBuild;
+	check(MipSizesInBytes.Num() == NumEncodedMips);
+
+
+	//
+	//
+	// We have all the metadata we need, we can grab the data
+	//
+	//
+	OutPlatformData.PixelFormat = EncodedTextureDescription.PixelFormat;
+	OutPlatformData.SizeX = EncodedTextureDescription.TopMipSizeX;
+	OutPlatformData.SizeY = EncodedTextureDescription.TopMipSizeY;
+	OutPlatformData.OptData.NumMipsInTail = EncodedTextureExtendedData.NumMipsInTail;
+	OutPlatformData.OptData.ExtData = EncodedTextureExtendedData.ExtData;
+	{
+		const bool bHasOptData = (EncodedTextureExtendedData.NumMipsInTail != 0) || (EncodedTextureExtendedData.ExtData != 0);
+		OutPlatformData.SetPackedData(EncodedTextureDescription.GetNumSlices_WithDepth(0), bHasOptData, EncodedTextureDescription.bCubeMap);
+	}
+	OutPlatformData.Mips.Empty(EncodedTextureDescription.NumMips);
+	EFileRegionType FileRegion = FFileRegion::SelectType(EncodedTextureDescription.PixelFormat);
+
+
+	FSharedBuffer MipTailData;
+	if (EncodedTextureDescription.NumMips > NumStreamingMips)
+	{
+		const FValueWithId& MipTailValue = BuildOutput.GetValue(FValueId::FromName(ANSITEXTVIEW("MipTail")));
+		if (!MipTailValue)
+		{
+			UE_LOG(LogTexture, Error, TEXT("Missing texture mip tail for build of '%s' by %s."), *BuildOutput.GetName(), *WriteToString<32>(BuildOutput.GetFunction()));
+			return false;
+		}
+		MipTailData = MipTailValue.GetData().Decompress();
+	}
+
+
+
+	uint64 CurrentMipTailOffset = 0;
+	for (int32 MipIndex = 0; MipIndex < EncodedTextureDescription.NumMips; MipIndex++)
+	{
+		FTexture2DMipMap* NewMip = new FTexture2DMipMap();
+		OutPlatformData.Mips.Add(NewMip);
+
+		NewMip->FileRegionType = FileRegion;
+		FIntVector3 MipDims = EncodedTextureDescription.GetMipDimensions(MipIndex);
+		NewMip->SizeX = MipDims.X;
+		NewMip->SizeY = MipDims.Y;
+		NewMip->SizeZ = MipDims.Z;
+		if (EncodedTextureDescription.bTextureArray)
+		{
+			// FTexture2DMipMap expects SizeZ to be the array count, potentially with cubemap slices.
+			NewMip->SizeZ = EncodedTextureDescription.ArraySlices;
+
+			if (EncodedTextureDescription.bCubeMap)
+			{
+				NewMip->SizeZ *= 6; 
+			}
+		}
+
+		if (MipIndex >= NumEncodedMips)
+		{
+			// Packed mip tail data is inside the outermost mip for the pack, so we don't have
+			// any bulk data to pull out.
+			continue;
+		}
+
+		if (MipIndex >= NumStreamingMips)
+		{			
+			// This mip is packed inside a single buffer. This is distinct from a "packed mip tail", but might
+			// coincidentally be the same. All mips past NumStreamingMips need to be copied into the bulk data
+			// and are always resident in memory with the texture.
+
+			uint64 MipSizeInBytes = MipSizesInBytes[MipIndex];
+			FMemoryView MipView = MipTailData.GetView().Mid(CurrentMipTailOffset, MipSizeInBytes);
+			CurrentMipTailOffset += MipSizeInBytes;
+
+			NewMip->BulkData.Lock(LOCK_READ_WRITE);
+			void* MipAllocData = NewMip->BulkData.Realloc(int64(MipSizeInBytes));
+			MakeMemoryView(MipAllocData, MipSizeInBytes).CopyFrom(MipView);
+			NewMip->BulkData.Unlock();
+		}
+		else
+		{
+			const FValueId MipId = FTexturePlatformData::MakeMipId(MipIndex);
+			const FValueWithId& MipValue = BuildOutput.GetValue(MipId);
+			if (!MipValue)
+			{
+				UE_LOG(LogTexture, Error, TEXT("Missing streaming texture mip %d for build of '%s' by %s."), MipIndex, *BuildOutput.GetName(), *WriteToString<32>(BuildOutput.GetFunction()));
+				return false;
+			}
+
+			// Did whoever launched the build want the streaming mips in memory?
+			if (InBuildResultOptions.bLoadStreamingMips && (MipIndex >= InBuildResultOptions.FirstStreamingMipToLoad))
+			{
+				NewMip->BulkData.Lock(LOCK_READ_WRITE);
+				const uint64 MipSize = MipValue.GetRawSize();
+				void* MipData = NewMip->BulkData.Realloc(IntCastChecked<int64>(MipSize));
+				ON_SCOPE_EXIT{ NewMip->BulkData.Unlock(); };
+				if (!MipValue.GetData().TryDecompressTo(MakeMemoryView(MipData, MipSize)))
+				{
+					UE_LOG(LogTexture, Error, TEXT("Failed to decompress streaming texture mip %d for build of '%s' by %s."), MipIndex, *BuildOutput.GetName(), *WriteToString<32>(BuildOutput.GetFunction()));
+					return false;
+				}
+			}
+
+			FSharedString MipName(WriteToString<256>(BuildOutput.GetName(), TEXT(" [MIP "), MipIndex, TEXT("]")));
+			NewMip->DerivedData = UE::FDerivedData(MoveTemp(MipName), InBuildCompleteParams.CacheKey, MipId);
+			PRAGMA_DISABLE_DEPRECATION_WARNINGS;
+			NewMip->bPagedToDerivedData = true;
+			PRAGMA_ENABLE_DEPRECATION_WARNINGS;
+		}
+	}
+	return true;
+}
+
+void HandleBuildOutputThenUnpack(FTexturePlatformData& OutPlatformData, UE::DerivedData::FBuildCompleteParams&& InBuildCompleteParams, FBuildResultOptions InBuildResultOptions)
+{
+	using namespace UE::DerivedData;
+
+	const FBuildOutput& Output = InBuildCompleteParams.Output;
+	const FSharedString& Name = Output.GetName();
+	const FUtf8SharedString& Function = Output.GetFunction();
+
+	for (const FBuildOutputMessage& Message : Output.GetMessages())
+	{
+		switch (Message.Level)
+		{
+		case EBuildOutputMessageLevel::Error:
+			UE_LOG(LogTexture, Warning, TEXT("[Error] %s (Build of '%s' by %s.)"),
+				*WriteToString<256>(Message.Message), *Name, *WriteToString<32>(Function));
+			break;
+		case EBuildOutputMessageLevel::Warning:
+			UE_LOG(LogTexture, Warning, TEXT("%s (Build of '%s' by %s.)"),
+				*WriteToString<256>(Message.Message), *Name, *WriteToString<32>(Function));
+			break;
+		case EBuildOutputMessageLevel::Display:
+			UE_LOG(LogTexture, Display, TEXT("%s (Build of '%s' by %s.)"),
+				*WriteToString<256>(Message.Message), *Name, *WriteToString<32>(Function));
+			break;
+		default:
+			checkNoEntry();
+			break;
+		}
+	}
+
+	for (const FBuildOutputLog& Log : Output.GetLogs())
+	{
+		switch (Log.Level)
+		{
+		case EBuildOutputLogLevel::Error:
+			UE_LOG(LogTexture, Warning, TEXT("[Error] %s: %s (Build of '%s' by %s.)"),
+				*WriteToString<64>(Log.Category), *WriteToString<256>(Log.Message),
+				*Name, *WriteToString<32>(Function));
+			break;
+		case EBuildOutputLogLevel::Warning:
+			UE_LOG(LogTexture, Warning, TEXT("%s: %s (Build of '%s' by %s.)"),
+				*WriteToString<64>(Log.Category), *WriteToString<256>(Log.Message),
+				*Name, *WriteToString<32>(Function));
+			break;
+		default:
+			checkNoEntry();
+			break;
+		}
+	}
+
+	if (Output.HasError())
+	{
+		UE_LOG(LogTexture, Warning, TEXT("Failed to build derived data for build of '%s' by %s."),
+			*Name, *WriteToString<32>(Function));
+		return;
+	}
+
+	UnpackPlatformDataFromBuild(OutPlatformData, MoveTemp(InBuildCompleteParams), InBuildResultOptions);
+}
+
+struct FBuildResults
+{
+	FTexturePlatformData& PlatformData;
+	bool bCacheHit = false;
+	uint64 BuildOutputSize = 0;
+
+	FBuildResults(FTexturePlatformData& InPlatformData) : PlatformData(InPlatformData) {}
+};
+
+static void GetBuildResultsFromCompleteParams(
+	FBuildResults& OutBuildResults, 
+	FBuildResultOptions InBuildResultOptions, 
+	UE::DerivedData::FBuildCompleteParams&& InBuildCompleteParams,
+	UE::TextureDerivedData::FTilingTextureBuildInputResolver* InTilingTextureBuildInputResolver // if we are completing a build that had a child build, we need this to get metadata.
+)
+{
+	using namespace UE::DerivedData;
+	OutBuildResults.PlatformData.DerivedDataKey.Emplace<FCacheKeyProxy>(InBuildCompleteParams.CacheKey);
+
+	if (InTilingTextureBuildInputResolver)
+	{
+		if (InTilingTextureBuildInputResolver->bParentBuild_HitCache &&
+			EnumHasAnyFlags(InBuildCompleteParams.BuildStatus, EBuildStatus::CacheQueryHit))
+		{
+			OutBuildResults.bCacheHit = true;
+		}
+	}
+	else
+	{
+		if (EnumHasAnyFlags(InBuildCompleteParams.BuildStatus, EBuildStatus::CacheQueryHit))
+		{
+			OutBuildResults.bCacheHit = true;
+		}
+	}
+
+	OutBuildResults.BuildOutputSize = Algo::TransformAccumulate(InBuildCompleteParams.Output.GetValues(),
+		[](const FValue& Value) { return Value.GetData().GetRawSize(); }, uint64(0));
+	if (InBuildCompleteParams.Status != EStatus::Canceled)
+	{
+		HandleBuildOutputThenUnpack(OutBuildResults.PlatformData, MoveTemp(InBuildCompleteParams), InBuildResultOptions);
+	}
+}
+
+
+struct FBuildInfo
+{
+	UE::DerivedData::FBuildSession& BuildSession;
+	UE::DerivedData::FBuildDefinition BuildDefinition;
+	UE::DerivedData::FBuildPolicy BuildPolicy;
+	TOptional<FTexturePlatformData::FTextureEncodeResultMetadata> ResultMetadata;
+	UE::TextureDerivedData::FTilingTextureBuildInputResolver* TilingInputResolver = nullptr; // nullptr if no tiling build.
+
+	explicit FBuildInfo(UE::DerivedData::FBuildSession& InBuildSession, UE::DerivedData::FBuildDefinition InBuildDefinition, UE::DerivedData::FBuildPolicy InBuildPolicy, TOptional<FTexturePlatformData::FTextureEncodeResultMetadata> InResultMetadata)
+		: BuildSession(InBuildSession)
+		, BuildDefinition(InBuildDefinition)
+		, BuildPolicy(InBuildPolicy)
+		, ResultMetadata(InResultMetadata)
+	{}
+};
+
+static void LaunchBuildWithFallback(
+	FBuildResults& OutBuildResults,
+	FBuildResultOptions InBuildResultOptions,
+	FBuildInfo& InInitialBuild, 
+	TOptional<FBuildInfo> InFallbackBuild, 
+	UE::DerivedData::FRequestOwner& InRequestOwner	// Owner must be valid for the duration of the build.
+)
+{
+	using namespace UE::DerivedData;
+
+	if (InInitialBuild.ResultMetadata.IsSet())
+	{
+		OutBuildResults.PlatformData.ResultMetadata = *InInitialBuild.ResultMetadata;
+	}
+
+	InInitialBuild.BuildSession.Build(InInitialBuild.BuildDefinition, {}, InInitialBuild.BuildPolicy, InRequestOwner,
+		[
+			FallbackBuild = MoveTemp(InFallbackBuild),
+			RequestOwner = &InRequestOwner,
+			OutBuildResults = &OutBuildResults,
+			BuildResultOptions = InBuildResultOptions,
+			InitialBuildTilingInputResolver = InInitialBuild.TilingInputResolver
+		](FBuildCompleteParams&& Params) mutable
+	{
+		if (Params.Status == EStatus::Error &&
+			FallbackBuild.IsSet())
+		{
+			if (FallbackBuild->ResultMetadata.IsSet())
+			{
+				OutBuildResults->PlatformData.ResultMetadata = *FallbackBuild->ResultMetadata;
+			}
+			FallbackBuild->BuildSession.Build(FallbackBuild->BuildDefinition, {}, FallbackBuild->BuildPolicy, *RequestOwner,
+				[
+					OutBuildResults = OutBuildResults,
+					BuildResultOptions = BuildResultOptions,
+					FallbackBuildTilingInputResolver = FallbackBuild->TilingInputResolver
+				](FBuildCompleteParams&& Params) mutable
+			{
+				GetBuildResultsFromCompleteParams(*OutBuildResults, BuildResultOptions, MoveTemp(Params), FallbackBuildTilingInputResolver);
+			});
+		}
+		else
+		{			
+			GetBuildResultsFromCompleteParams(*OutBuildResults, BuildResultOptions, MoveTemp(Params), InitialBuildTilingInputResolver);
+		}
+	});
+};
+
+
 //
 // DDC2 texture fetch/build task.
 //
@@ -1548,33 +2016,31 @@ public:
 		UTexture& Texture,
 		FTexturePlatformData& InDerivedData,
 		const UE::DerivedData::FUtf8SharedString& FunctionName,
+		const UE::DerivedData::FUtf8SharedString& TilingFunctionName,
 		const FTextureBuildSettings* InSettingsFetchFirst, // can be nullptr
 		const FTextureBuildSettings& InSettingsFetchOrBuild,
 		const FTexturePlatformData::FTextureEncodeResultMetadata* InFetchFirstMetadata, // can be nullptr
 		const FTexturePlatformData::FTextureEncodeResultMetadata* InFetchOrBuildMetadata, // can be nullptr
 		EQueuedWorkPriority InPriority,
 		ETextureCacheFlags Flags)
-		: DerivedData(InDerivedData)
+		: BuildResults(InDerivedData)
 		, Priority(InPriority)
-		, bCacheHit(false)
-		, bInlineMips(EnumHasAnyFlags(Flags, ETextureCacheFlags::InlineMips))
-		, FirstMipToLoad(InSettingsFetchOrBuild.LODBiasWithCinematicMips)
 		, InputResolver(Texture)
 	{
-		using namespace UE::DerivedData;
-
 		static bool bLoadedModules = LoadModules();
 
-		FSharedString TexturePath;
+		BuildResultOptions.bLoadStreamingMips = EnumHasAnyFlags(Flags, ETextureCacheFlags::InlineMips);
+		BuildResultOptions.FirstStreamingMipToLoad = InSettingsFetchOrBuild.LODBiasWithCinematicMips;
+
+		// Can't fetch first if we are rebuilding.
+		if (InSettingsFetchFirst &&
+			EnumHasAnyFlags(Flags, ETextureCacheFlags::ForceRebuild))
 		{
-			TStringBuilder<256> TexturePathBuilder;
-			Texture.GetPathName(nullptr, TexturePathBuilder);
-			TexturePath = TexturePathBuilder.ToView();
+			InSettingsFetchFirst = nullptr;
+			InFetchFirstMetadata = nullptr;
 		}
 
-		IBuild& Build = GetBuild();
-		IBuildInputResolver* GlobalResolver = GetGlobalBuildInputResolver();
-		BuildSession = Build.CreateSession(TexturePath, GlobalResolver ? GlobalResolver : &InputResolver);
+		using namespace UE::DerivedData;
 
 		EPriority OwnerPriority = EnumHasAnyFlags(Flags, ETextureCacheFlags::Async) ? ConvertPriority(Priority) : UE::DerivedData::EPriority::Blocking;
 		Owner.Emplace(OwnerPriority);
@@ -1589,12 +2055,20 @@ public:
 		check( Texture.Source.GetNumLayers() == 1 );
 		int64 RequiredMemoryEstimate = GetBuildRequiredMemoryEstimate(&Texture,&InSettingsFetchOrBuild);
 
+		// Debug string.
+		FSharedString TexturePath;
+		{
+			TStringBuilder<256> TexturePathBuilder;
+			Texture.GetPathName(nullptr, TexturePathBuilder);
+			TexturePath = TexturePathBuilder.ToView();
+		}
+
 		TOptional<FTextureStatusMessageContext> StatusMessage;
 		if (IsInGameThread() && OwnerPriority == EPriority::Blocking)
 		{
-			// this gets sent whether or not we are building the texture, and is a rare edge case for UI feedback.
+			// this gets set whether or not we are building the texture, and is a rare edge case for UI feedback.
 			// We don't actually know whether we're using fetchfirst or actually building, so if we have two keys,
-			// we can assume we're FinalIfAvailable.
+			// we just assume we're FinalIfAvailable.
 			ETextureEncodeSpeed EncodeSpeed = (ETextureEncodeSpeed)InSettingsFetchOrBuild.RepresentsEncodeSpeedNoSend;
 			if (InSettingsFetchFirst)
 			{
@@ -1605,65 +2079,146 @@ public:
 		}
 		
 
-		if (InSettingsFetchFirst &&
-			EnumHasAnyFlags(Flags, ETextureCacheFlags::ForceRebuild))
+		TOptional<FTexturePlatformData::FTextureEncodeResultMetadata> FetchFirstResultMetadata;
+		if (InFetchFirstMetadata)
 		{
-			// Can't fetch first if we are rebuilding.
-			InSettingsFetchFirst = nullptr;
+			FetchFirstResultMetadata = *InFetchFirstMetadata;
 		}
 
-		FBuildDefinition FetchOrBuildDefinition = CreateDefinition(Build, Texture, TexturePath, FunctionName, InSettingsFetchOrBuild, bUseCompositeTexture, RequiredMemoryEstimate);
-		DerivedData.FetchOrBuildDerivedDataKey.Emplace<FTexturePlatformData::FStructuredDerivedDataKey>(GetKey(FetchOrBuildDefinition, Texture, bUseCompositeTexture));
+		TOptional<FTexturePlatformData::FTextureEncodeResultMetadata> FetchOrBuildResultMetadata;
+		if (InFetchOrBuildMetadata)
+		{
+			FetchOrBuildResultMetadata = *InFetchOrBuildMetadata;
+		}
 
-		bool bBuildKicked = false;
+		// Description and MipTail should always cache. Everything else (i.e. Mip# i.e. streaming mips) should skip data
+		// when we are not inlining.
+		FBuildPolicy FetchFirstBuildPolicy = FetchFirst_CreateBuildPolicy(BuildResultOptions);
+		FBuildPolicy FetchOrBuildPolicy = FetchOrBuild_CreateBuildPolicy(Flags, BuildResultOptions);
+
+		//
+		// Set up the build
+		//
+		IBuild& Build = GetBuild();
+		{
+			IBuildInputResolver* GlobalResolver = GetGlobalBuildInputResolver();
+			BuildSession = Build.CreateSession(TexturePath, GlobalResolver ? GlobalResolver : &InputResolver);
+		}
+
+		/////////////////////////////////////////////////////////////////////////////////////
+
+		//
+		// We have a 2x2 possibility space, and this is the cleanest way to break that down:
+		// FetchFirst exists yes/no
+		// Tiling build exists yes/no
+		//
 		if (InSettingsFetchFirst)
 		{
-			// If the keys are the same, ignore fetch first.
-			FBuildDefinition FetchDefinition = CreateDefinition(Build, Texture, TexturePath, FunctionName, *InSettingsFetchFirst, bUseCompositeTexture, RequiredMemoryEstimate);
-			DerivedData.FetchFirstDerivedDataKey.Emplace<FTexturePlatformData::FStructuredDerivedDataKey>(GetKey(FetchDefinition, Texture, bUseCompositeTexture));
-			if (FetchDefinition.GetKey() != FetchOrBuildDefinition.GetKey())
+			if (TilingFunctionName.Len())
 			{
-				bBuildKicked = true;
-				if (InFetchFirstMetadata)
+				//
+				// FetchFirst_Tiling builds FetchFirst_Parent, if we miss in either place
+				// we go to FetchOrBuild_Tiling which builds FetchOrBuild_Parent.
+				//
+				FBuildDefinition FetchOrBuild_ParentDefinition = CreateDefinition(Build, Texture, TexturePath, FunctionName, InSettingsFetchOrBuild, bUseCompositeTexture, RequiredMemoryEstimate);
+				FBuildDefinition FetchOrBuild_TilingDefinition = CreateTilingDefinition(Build, &Texture, InSettingsFetchOrBuild, nullptr, nullptr, FetchOrBuild_ParentDefinition, TexturePath, TilingFunctionName);
+
+				FBuildDefinition FetchFirst_ParentDefinition = CreateDefinition(Build, Texture, TexturePath, FunctionName, *InSettingsFetchFirst, bUseCompositeTexture, RequiredMemoryEstimate);
+				FBuildDefinition FetchFirst_TilingDefinition = CreateTilingDefinition(Build, &Texture, *InSettingsFetchFirst, nullptr, nullptr, FetchFirst_ParentDefinition, TexturePath, TilingFunctionName);
+
+				BuildResults.PlatformData.FetchOrBuildDerivedDataKey.Emplace<FTexturePlatformData::FStructuredDerivedDataKey>(GetKey(FetchOrBuild_ParentDefinition, &FetchOrBuild_TilingDefinition, Texture, bUseCompositeTexture));
+				BuildResults.PlatformData.FetchFirstDerivedDataKey.Emplace<FTexturePlatformData::FStructuredDerivedDataKey>(GetKey(FetchFirst_ParentDefinition, &FetchFirst_TilingDefinition, Texture, bUseCompositeTexture));
+
+				FetchOrBuild_ChildInputResolver.Emplace(BuildSession.Get(), FetchOrBuild_ParentDefinition, FetchOrBuildPolicy);
+				FetchOrBuild_ChildBuildSession = Build.CreateSession(TexturePath, FetchOrBuild_ChildInputResolver.GetPtrOrNull());
+
+				// Must be done after the policy is copied in to the child input resolver
+				if (CVarForceRetileTextures.GetValueOnAnyThread())
 				{
-					this->DerivedData.ResultMetadata = *InFetchFirstMetadata;
+					FetchOrBuildPolicy = EBuildPolicy::Build;
 				}
 
-				FBuildPolicyBuilder BuildPolicyBuilder(bInlineMips ? EBuildPolicy::Cache : (EBuildPolicy::CacheQuery | EBuildPolicy::SkipData));
-				if (!bInlineMips)
+				FBuildInfo FetchOrBuildInfo(FetchOrBuild_ChildBuildSession.GetValue(), FetchOrBuild_TilingDefinition, FetchOrBuildPolicy, FetchOrBuildResultMetadata);
+				FetchOrBuildInfo.TilingInputResolver = FetchOrBuild_ChildInputResolver.GetPtrOrNull();
+
+				if (FetchOrBuild_ParentDefinition.GetKey() == FetchFirst_ParentDefinition.GetKey() &&
+					FetchOrBuild_TilingDefinition.GetKey() == FetchFirst_TilingDefinition.GetKey())
 				{
-					BuildPolicyBuilder.AddValuePolicy(FValueId::FromName(ANSITEXTVIEW("Description")), EBuildPolicy::Cache);
-					BuildPolicyBuilder.AddValuePolicy(FValueId::FromName(ANSITEXTVIEW("MipTail")), EBuildPolicy::Cache);
+					// Same definition, just do FetchOrBuild.
+					LaunchBuildWithFallback(BuildResults, BuildResultOptions, FetchOrBuildInfo, {}, *Owner);
 				}
-
-				BuildSession.Get().Build(FetchDefinition, {}, BuildPolicyBuilder.Build(), *Owner,
-					[this,
-					 FetchOrBuildDefinition = MoveTemp(FetchOrBuildDefinition),
-					 Flags,
-					 FetchOrBuildMetadata = InFetchOrBuildMetadata ? *InFetchOrBuildMetadata : FTexturePlatformData::FTextureEncodeResultMetadata()
-					](FBuildCompleteParams&& Params)
+				else
+				{	
+					FetchFirst_ChildInputResolver.Emplace(BuildSession.Get(), FetchFirst_ParentDefinition, FetchFirstBuildPolicy);
+					FetchFirst_ChildBuildSession = Build.CreateSession(TexturePath, FetchFirst_ChildInputResolver.GetPtrOrNull());
+					// Must be done after the policy is copied in to the child input resolver
+					if (CVarForceRetileTextures.GetValueOnAnyThread())
 					{
-						switch (Params.Status)
-						{
-						default:
-						case EStatus::Ok:
-							return EndBuild(MoveTemp(Params));
-						case EStatus::Error:
-							this->DerivedData.ResultMetadata = FetchOrBuildMetadata;
-							return BeginBuild(FetchOrBuildDefinition, Flags);
-						}
-					});
+						FetchFirstBuildPolicy = EBuildPolicy::Build;
+					}
+
+					FBuildInfo FetchFirstInfo(FetchFirst_ChildBuildSession.GetValue(), FetchFirst_TilingDefinition, FetchFirstBuildPolicy, FetchFirstResultMetadata);
+					FetchFirstInfo.TilingInputResolver = FetchFirst_ChildInputResolver.GetPtrOrNull();
+
+					LaunchBuildWithFallback(BuildResults, BuildResultOptions, FetchFirstInfo, FetchOrBuildInfo, *Owner);
+				}
 			}
-		}
+			else
+			{
+				// FetchFirst runs, if we miss we build FetchOrBuild
+				FBuildDefinition FetchOrBuild_Definition = CreateDefinition(Build, Texture, TexturePath, FunctionName, InSettingsFetchOrBuild, bUseCompositeTexture, RequiredMemoryEstimate);
+				FBuildDefinition FetchFirst_Definition = CreateDefinition(Build, Texture, TexturePath, FunctionName, *InSettingsFetchFirst, bUseCompositeTexture, RequiredMemoryEstimate);
+
+				BuildResults.PlatformData.FetchOrBuildDerivedDataKey.Emplace<FTexturePlatformData::FStructuredDerivedDataKey>(GetKey(FetchOrBuild_Definition, nullptr, Texture, bUseCompositeTexture));
+				BuildResults.PlatformData.FetchFirstDerivedDataKey.Emplace<FTexturePlatformData::FStructuredDerivedDataKey>(GetKey(FetchFirst_Definition, nullptr, Texture, bUseCompositeTexture));
 
-		if (bBuildKicked == false)
+				FBuildInfo FetchOrBuildInfo(BuildSession.Get(), FetchOrBuild_Definition, FetchOrBuildPolicy, FetchOrBuildResultMetadata);
+				if (FetchOrBuild_Definition.GetKey() == FetchFirst_Definition.GetKey())
+				{
+					LaunchBuildWithFallback(BuildResults, BuildResultOptions, FetchOrBuildInfo, {}, *Owner);
+				}
+				else
+				{
+					FBuildInfo FetchFirstInfo(BuildSession.Get(), FetchFirst_Definition, FetchFirstBuildPolicy, FetchFirstResultMetadata);
+					LaunchBuildWithFallback(BuildResults, BuildResultOptions, FetchFirstInfo, FetchOrBuildInfo, *Owner);
+				}
+			}
+		}
+		else
 		{
-			// we didn't use the fetch first path.
-			if (InFetchOrBuildMetadata)
+			if (TilingFunctionName.Len())
 			{
-				this->DerivedData.ResultMetadata = *InFetchOrBuildMetadata;
+				// Tiling runs which builds Parent.
+				FBuildDefinition FetchOrBuild_ParentDefinition = CreateDefinition(Build, Texture, TexturePath, FunctionName, InSettingsFetchOrBuild, bUseCompositeTexture, RequiredMemoryEstimate);
+				FBuildDefinition FetchOrBuild_TilingDefinition = CreateTilingDefinition(Build, &Texture, InSettingsFetchOrBuild, nullptr, nullptr, FetchOrBuild_ParentDefinition, TexturePath, TilingFunctionName);
+
+				BuildResults.PlatformData.FetchOrBuildDerivedDataKey.Emplace<FTexturePlatformData::FStructuredDerivedDataKey>(GetKey(FetchOrBuild_ParentDefinition, &FetchOrBuild_TilingDefinition, Texture, bUseCompositeTexture));
+
+				FetchOrBuild_ChildInputResolver.Emplace(BuildSession.Get(), FetchOrBuild_ParentDefinition, FetchOrBuildPolicy);
+				// Must be done after the policy is copied in to the child input resolver
+				if (CVarForceRetileTextures.GetValueOnAnyThread())
+				{
+					FetchOrBuildPolicy = EBuildPolicy::Build;
+				}
+
+				FetchOrBuild_ChildBuildSession = Build.CreateSession(TexturePath, FetchOrBuild_ChildInputResolver.GetPtrOrNull());
+
+				FBuildInfo FetchOrBuildInfo(FetchOrBuild_ChildBuildSession.GetValue(), FetchOrBuild_TilingDefinition, FetchOrBuildPolicy, FetchOrBuildResultMetadata);
+				FetchOrBuildInfo.TilingInputResolver = FetchOrBuild_ChildInputResolver.GetPtrOrNull();
+
+				LaunchBuildWithFallback(BuildResults, BuildResultOptions, FetchOrBuildInfo, {}, *Owner);
+			}
+			else
+			{
+				// we just directly build the single definition.
+				FBuildDefinition FetchOrBuild_Definition = CreateDefinition(Build, Texture, TexturePath, FunctionName, InSettingsFetchOrBuild, bUseCompositeTexture, RequiredMemoryEstimate);
+
+				BuildResults.PlatformData.FetchOrBuildDerivedDataKey.Emplace<FTexturePlatformData::FStructuredDerivedDataKey>(GetKey(FetchOrBuild_Definition, nullptr, Texture, bUseCompositeTexture));
+
+				FBuildInfo FetchOrBuildInfo(BuildSession.Get(), FetchOrBuild_Definition, FetchOrBuildPolicy, FetchOrBuildResultMetadata);
+
+				LaunchBuildWithFallback(BuildResults, BuildResultOptions, FetchOrBuildInfo, {}, *Owner);
 			}
-			BeginBuild(FetchOrBuildDefinition, Flags);
 		}
 
 		if (StatusMessage.IsSet())
@@ -1682,8 +2237,8 @@ public:
 		const int64 RequiredMemoryEstimate)
 	{
 		UE::DerivedData::FBuildDefinitionBuilder DefinitionBuilder = Build.CreateDefinition(TexturePath, FunctionName);
-		DefinitionBuilder.AddConstant(UTF8TEXTVIEW("Settings"),
-			SaveTextureBuildSettings(Texture, Settings, 0, NUM_INLINE_DERIVED_MIPS, bUseCompositeTexture, RequiredMemoryEstimate));
+		DefinitionBuilder.AddConstant(UTF8TEXTVIEW("EngineParameters"), UE::TextureBuildUtilities::TextureEngineParameters::ToCompactBinaryWithDefaults(GenerateTextureEngineParameters()));
+		DefinitionBuilder.AddConstant(UTF8TEXTVIEW("Settings"), SaveTextureBuildSettings(Texture, Settings, 0, bUseCompositeTexture, RequiredMemoryEstimate));
 		DefinitionBuilder.AddInputBulkData(UTF8TEXTVIEW("Source"), Texture.Source.GetPersistentId());
 		if (Texture.CompositeTexture && bUseCompositeTexture)
 		{
@@ -1693,48 +2248,53 @@ public:
 	}
 
 private:
-	void BeginBuild(const UE::DerivedData::FBuildDefinition& Definition, ETextureCacheFlags Flags)
+
+	static UE::DerivedData::FBuildPolicy FetchFirst_CreateBuildPolicy(FBuildResultOptions InBuildResultOptions)
 	{
 		using namespace UE::DerivedData;
 
-		FBuildPolicy BuildPolicy;
-		if (EnumHasAnyFlags(Flags, ETextureCacheFlags::ForceRebuild))
+		if (InBuildResultOptions.bLoadStreamingMips)
 		{
-			BuildPolicy = EBuildPolicy::Default & ~EBuildPolicy::CacheQuery;
-		}
-		else if (bInlineMips)
-		{
-			BuildPolicy = EBuildPolicy::Default;
+			// We want all of the output values.
+			return EBuildPolicy::Cache;
 		}
 		else
 		{
-			FBuildPolicyBuilder BuildPolicyBuilder(EBuildPolicy::Build | EBuildPolicy::CacheQuery | EBuildPolicy::CacheStoreOnBuild | EBuildPolicy::SkipData);
-			BuildPolicyBuilder.AddValuePolicy(FValueId::FromName(ANSITEXTVIEW("Description")), EBuildPolicy::Default);
-			BuildPolicyBuilder.AddValuePolicy(FValueId::FromName(ANSITEXTVIEW("MipTail")), EBuildPolicy::Default);
-			BuildPolicy = BuildPolicyBuilder.Build();
-		}
-
-		BuildSession.Get().Build(Definition, {}, BuildPolicy, *Owner,
-			[this](FBuildCompleteParams&& Params) { EndBuild(MoveTemp(Params)); });
+			// We only want the metadata and the mip tail.
+			FBuildPolicyBuilder FetchFirstBuildPolicyBuilder(EBuildPolicy::CacheQuery | EBuildPolicy::SkipData);
+			FetchFirstBuildPolicyBuilder.AddValuePolicy(FValueId::FromName(ANSITEXTVIEW("EncodedTextureDescription")), EBuildPolicy::Cache);
+			FetchFirstBuildPolicyBuilder.AddValuePolicy(FValueId::FromName(ANSITEXTVIEW("EncodedTextureExtendedData")), EBuildPolicy::Cache);
+			FetchFirstBuildPolicyBuilder.AddValuePolicy(FValueId::FromName(ANSITEXTVIEW("MipTail")), EBuildPolicy::Cache);
+			return FetchFirstBuildPolicyBuilder.Build();
+		}		
 	}
 
-	void EndBuild(UE::DerivedData::FBuildCompleteParams&& Params)
+	static UE::DerivedData::FBuildPolicy FetchOrBuild_CreateBuildPolicy(ETextureCacheFlags Flags, FBuildResultOptions InBuildResultOptions)
 	{
 		using namespace UE::DerivedData;
-		DerivedData.DerivedDataKey.Emplace<FCacheKeyProxy>(Params.CacheKey);
-		bCacheHit = EnumHasAnyFlags(Params.BuildStatus, EBuildStatus::CacheQueryHit);
-		BuildOutputSize = Algo::TransformAccumulate(Params.Output.GetValues(),
-			[](const FValue& Value) { return Value.GetData().GetRawSize(); }, uint64(0));
-		if (Params.Status != EStatus::Canceled)
+
+		if (EnumHasAnyFlags(Flags, ETextureCacheFlags::ForceRebuild))
+		{
+			return EBuildPolicy::Default & ~EBuildPolicy::CacheQuery;
+		}
+		else if (InBuildResultOptions.bLoadStreamingMips)
 		{
-			WriteDerivedData(MoveTemp(Params));
+			return EBuildPolicy::Default;
+		}
+		else
+		{
+			FBuildPolicyBuilder BuildPolicyBuilder(EBuildPolicy::Build | EBuildPolicy::CacheQuery | EBuildPolicy::CacheStoreOnBuild | EBuildPolicy::SkipData);
+			BuildPolicyBuilder.AddValuePolicy(FValueId::FromName(ANSITEXTVIEW("EncodedTextureDescription")), EBuildPolicy::Cache);
+			BuildPolicyBuilder.AddValuePolicy(FValueId::FromName(ANSITEXTVIEW("EncodedTextureExtendedData")), EBuildPolicy::Cache);
+			BuildPolicyBuilder.AddValuePolicy(FValueId::FromName(ANSITEXTVIEW("MipTail")), EBuildPolicy::Default);
+			return BuildPolicyBuilder.Build();
 		}
 	}
 
 	void Finalize(bool& bOutFoundInCache, uint64& OutProcessedByteCount) final
 	{
-		bOutFoundInCache = bCacheHit;
-		OutProcessedByteCount = BuildOutputSize;
+		bOutFoundInCache = BuildResults.bCacheHit;
+		OutProcessedByteCount = BuildResults.BuildOutputSize;
 	}
 public:
 
@@ -1892,10 +2452,14 @@ public:
 		return true;
 	}
 
-	static FTexturePlatformData::FStructuredDerivedDataKey GetKey(const UE::DerivedData::FBuildDefinition& BuildDefinition, const UTexture& Texture, bool bUseCompositeTexture)
+	static FTexturePlatformData::FStructuredDerivedDataKey GetKey(const UE::DerivedData::FBuildDefinition& BuildDefinition, const UE::DerivedData::FBuildDefinition* TilingBuildDefinitionKey, const UTexture& Texture, bool bUseCompositeTexture)
 	{
 		// DDC2 Key SerializeForKey is here!
 		FTexturePlatformData::FStructuredDerivedDataKey Key;
+		if (TilingBuildDefinitionKey != nullptr)
+		{
+			Key.TilingBuildDefinitionKey = TilingBuildDefinitionKey->GetKey().Hash;
+		}
 		Key.BuildDefinitionKey = BuildDefinition.GetKey().Hash;
 		Key.SourceGuid = Texture.Source.GetId();
 		if (bUseCompositeTexture && Texture.CompositeTexture)
@@ -1905,189 +2469,63 @@ public:
 		return Key;
 	}
 
-private:
-
-	static bool DeserializeTextureFromValues(FTexturePlatformData& DerivedData, UE::DerivedData::FBuildCompleteParams&& Params, int32 FirstMipToLoad, bool bInlineMips)
-	{
-		using namespace UE;
-		using namespace UE::DerivedData;
-
-		const FBuildOutput& Output = Params.Output;
-		const FValueWithId& Value = Output.GetValue(FValueId::FromName(ANSITEXTVIEW("Description")));
-		if (!Value)
-		{
-			UE_LOG(LogTexture, Error, TEXT("Missing texture description for build of '%s' by %s."),
-				*Output.GetName(), *WriteToString<32>(Output.GetFunction()));
-			return false;
-		}
-
-		FCbObject TextureDescription(Value.GetData().Decompress());
-
-		FCbFieldViewIterator SizeIt = TextureDescription[ANSITEXTVIEW("Size")].AsArrayView().CreateViewIterator();
-		DerivedData.SizeX = SizeIt++->AsInt32();
-		DerivedData.SizeY = SizeIt++->AsInt32();
-		int32 NumSlices = SizeIt++->AsInt32();
-
-		UEnum* PixelFormatEnum = UTexture::GetPixelFormatEnum();
-		FUtf8StringView PixelFormatStringView = TextureDescription[ANSITEXTVIEW("PixelFormat")].AsString();
-		FName PixelFormatName(PixelFormatStringView.Len(), PixelFormatStringView.GetData());
-		DerivedData.PixelFormat = (EPixelFormat)PixelFormatEnum->GetValueByName(PixelFormatName);
+	static UE::DerivedData::FBuildDefinition CreateTilingDefinition(
+		UE::DerivedData::IBuild& InBuild,
+		UTexture* InTexture,
+		const FTextureBuildSettings& InBuildSettings,
+		FEncodedTextureDescription* InTextureDescription, // only valid if our textures can generate this pre build
+		FEncodedTextureExtendedData* InTextureExtendedData, // only valid if our textures can generate this pre build
+		const UE::DerivedData::FBuildDefinition& InParentBuildDefinition,
+		const UE::DerivedData::FSharedString& InDefinitionDebugName,
+		const UE::DerivedData::FUtf8SharedString& InBuildFunctionName
+	)
+	{			
+		const FTextureEngineParameters EngineParameters = GenerateTextureEngineParameters();
 
-		const bool bCubeMap = TextureDescription[ANSITEXTVIEW("bCubeMap")].AsBool();
-		DerivedData.OptData.ExtData = TextureDescription[ANSITEXTVIEW("ExtData")].AsUInt32();
-		DerivedData.OptData.NumMipsInTail = TextureDescription[ANSITEXTVIEW("NumMipsInTail")].AsUInt32();
-		const bool bHasOptData = (DerivedData.OptData.NumMipsInTail != 0) || (DerivedData.OptData.ExtData != 0);
-		static constexpr uint32 BitMask_CubeMap = 1u << 31u;
-		static constexpr uint32 BitMask_HasOptData = 1u << 30u;
-		static constexpr uint32 BitMask_NumSlices = BitMask_HasOptData - 1u;
-		DerivedData.PackedData = (NumSlices & BitMask_NumSlices) | (bCubeMap ? BitMask_CubeMap : 0) | (bHasOptData ? BitMask_HasOptData : 0);
+		//
+		// We always consume an unpacked texture (i.e. extended data == nullptr)
+		//
+		const FTextureSource& Source = InTexture->Source;
+		int32 InputTextureMip0SizeX, InputTextureMip0SizeY, InputTextureMip0NumSlices;
+		int32 InputTextureNumMips = TextureCompressorModule->GetMipCountForBuildSettings(Source.GetSizeX(), Source.GetSizeY(), Source.GetNumSlices(), Source.GetNumMips(), InBuildSettings, InputTextureMip0SizeX, InputTextureMip0SizeY, InputTextureMip0NumSlices);
+		int32 InputTextureNumStreamingMips = GetNumStreamingMipsDirect(InputTextureNumMips, InBuildSettings.bCubemap, InBuildSettings.bVolume, InBuildSettings.bTextureArray, nullptr, EngineParameters);
 
-		int32 NumMips = TextureDescription[ANSITEXTVIEW("NumMips")].AsInt32();
-		int32 NumStreamingMips = TextureDescription[ANSITEXTVIEW("NumStreamingMips")].AsInt32();
+		//
+		// A child definition consumes a parent definition and swizzles it. However it
+		// needs to know ahead of time the total mip count and the streaming mip count
+		//
+		const UE::DerivedData::FBuildKey InParentBuildKey = InParentBuildDefinition.GetKey();
 
-		FCbArrayView MipArrayView = TextureDescription[ANSITEXTVIEW("Mips")].AsArrayView();
-		if (NumMips != MipArrayView.Num())
-		{
-			UE_LOG(LogTexture, Error, TEXT("Mismatched mip quantity (%d and %d) for build of '%s' by %s."),
-				NumMips, MipArrayView.Num(), *Output.GetName(), *WriteToString<32>(Output.GetFunction()));
-			return false;
-		}
-		check(NumMips >= (int32)DerivedData.OptData.NumMipsInTail);
-		check(NumMips >= NumStreamingMips);
+		UE::DerivedData::FBuildDefinitionBuilder DefinitionBuilder = InBuild.CreateDefinition(InDefinitionDebugName, InBuildFunctionName);
 
-		FSharedBuffer MipTailData;
-		if (NumMips > NumStreamingMips)
+		// The tiler needs the description, which either comes from us (new style) or the linear build (old style)
+		if (InTextureDescription == nullptr)
 		{
-			const FValueWithId& MipTailValue = Output.GetValue(FValueId::FromName(ANSITEXTVIEW("MipTail")));
-			if (!MipTailValue)
-			{
-				UE_LOG(LogTexture, Error, TEXT("Missing texture mip tail for build of '%s' by %s."),
-					*Output.GetName(), *WriteToString<32>(Output.GetFunction()));
-				return false;
-			}
-			MipTailData = MipTailValue.GetData().Decompress();
-		}
-
-		int32 MipIndex = 0;
-		DerivedData.Mips.Empty(NumMips);
-		for (FCbFieldView MipFieldView : MipArrayView)
-		{
-			FCbObjectView MipObjectView = MipFieldView.AsObjectView();
-			TUniquePtr<FTexture2DMipMap> NewMip = MakeUnique<FTexture2DMipMap>();
-
-			FCbFieldViewIterator MipSizeIt = MipObjectView[ANSITEXTVIEW("Size")].AsArrayView().CreateViewIterator();
-			NewMip->SizeX = MipSizeIt++->AsInt32();
-			NewMip->SizeY = MipSizeIt++->AsInt32();
-			NewMip->SizeZ = MipSizeIt++->AsInt32();
-			NewMip->FileRegionType = static_cast<EFileRegionType>(MipObjectView[ANSITEXTVIEW("FileRegion")].AsInt32());
-			
-			if (MipIndex >= NumStreamingMips)
-			{
-				uint64 MipSize = MipObjectView[ANSITEXTVIEW("NumBytes")].AsUInt64();
-				FMemoryView MipView = MipTailData.GetView().Mid(MipObjectView[ANSITEXTVIEW("MipOffset")].AsUInt64(), MipSize);
-
-				NewMip->BulkData.Lock(LOCK_READ_WRITE);
-				void* MipAllocData = NewMip->BulkData.Realloc(int64(MipSize));
-				MakeMemoryView(MipAllocData, MipSize).CopyFrom(MipView);
-				NewMip->BulkData.Unlock();
-			}
-			else
-			{
-				const FValueId MipId = FTexturePlatformData::MakeMipId(MipIndex);
-				const FValueWithId& MipValue = Output.GetValue(MipId);
-				if (!MipValue)
-				{
-					UE_LOG(LogTexture, Error, TEXT("Missing streaming texture mip %d for build of '%s' by %s."),
-						MipIndex, *Output.GetName(), *WriteToString<32>(Output.GetFunction()));
-					return false;
-				}
-
-				if (bInlineMips && (MipIndex >= FirstMipToLoad))
-				{
-					NewMip->BulkData.Lock(LOCK_READ_WRITE);
-					const uint64 MipSize = MipValue.GetRawSize();
-					void* MipData = NewMip->BulkData.Realloc(IntCastChecked<int64>(MipSize));
-					ON_SCOPE_EXIT { NewMip->BulkData.Unlock(); };
-					if (!MipValue.GetData().TryDecompressTo(MakeMemoryView(MipData, MipSize)))
-					{
-						UE_LOG(LogTexture, Error, TEXT("Failed to decompress streaming texture mip %d for build of '%s' by %s."),
-							MipIndex, *Output.GetName(), *WriteToString<32>(Output.GetFunction()));
-						return false;
-					}
-				}
-
-				FSharedString MipName(WriteToString<256>(Output.GetName(), TEXT(" [MIP "), MipIndex, TEXT("]")));
-				NewMip->DerivedData = FDerivedData(MoveTemp(MipName), Params.CacheKey, MipId);
-				PRAGMA_DISABLE_DEPRECATION_WARNINGS;
-				NewMip->bPagedToDerivedData = true;
-				PRAGMA_ENABLE_DEPRECATION_WARNINGS;
-			}
-
-			DerivedData.Mips.Add(NewMip.Release());
-			++MipIndex;
+			// old style
+			DefinitionBuilder.AddInputBuild(UTF8TEXTVIEW("EncodedTextureDescriptionInput"), { InParentBuildKey, UE::DerivedData::FValueId::FromName(UTF8TEXTVIEW("EncodedTextureDescription")) });			
+			// The tiling build can generate the extended data.
 		}
-
-		return true;
-	}
-
-	void WriteDerivedData(UE::DerivedData::FBuildCompleteParams&& Params)
-	{
-		using namespace UE::DerivedData;
-
-		const FBuildOutput& Output = Params.Output;
-		const FSharedString& Name = Output.GetName();
-		const FUtf8SharedString& Function = Output.GetFunction();
-
-		for (const FBuildOutputMessage& Message : Output.GetMessages())
+		else
 		{
-			switch (Message.Level)
-			{
-			case EBuildOutputMessageLevel::Error:
-				UE_LOG(LogTexture, Warning, TEXT("[Error] %s (Build of '%s' by %s.)"),
-					*WriteToString<256>(Message.Message), *Name, *WriteToString<32>(Function));
-				break;
-			case EBuildOutputMessageLevel::Warning:
-				UE_LOG(LogTexture, Warning, TEXT("%s (Build of '%s' by %s.)"),
-					*WriteToString<256>(Message.Message), *Name, *WriteToString<32>(Function));
-				break;
-			case EBuildOutputMessageLevel::Display:
-				UE_LOG(LogTexture, Display, TEXT("%s (Build of '%s' by %s.)"),
-					*WriteToString<256>(Message.Message), *Name, *WriteToString<32>(Function));
-				break;
-			default:
-				checkNoEntry();
-				break;
-			}
+			// new style - we want to provide everything so that the only outputs are bulk data that we can
+			// just hold references to.
+			DefinitionBuilder.AddConstant(UTF8TEXTVIEW("EncodedTextureDescriptionConstant"), UE::TextureBuildUtilities::EncodedTextureDescription::ToCompactBinary(*InTextureDescription));
+			DefinitionBuilder.AddConstant(UTF8TEXTVIEW("EncodedTextureExtendedDataConstant"), UE::TextureBuildUtilities::EncodedTextureExtendedData::ToCompactBinary(*InTextureExtendedData));
 		}
-
-		for (const FBuildOutputLog& Log : Output.GetLogs())
+		
+		if (InputTextureNumMips > InputTextureNumStreamingMips)
 		{
-			switch (Log.Level)
-			{
-			case EBuildOutputLogLevel::Error:
-				UE_LOG(LogTexture, Warning, TEXT("[Error] %s: %s (Build of '%s' by %s.)"),
-					*WriteToString<64>(Log.Category), *WriteToString<256>(Log.Message),
-					*Name, *WriteToString<32>(Function));
-				break;
-			case EBuildOutputLogLevel::Warning:
-				UE_LOG(LogTexture, Warning, TEXT("%s: %s (Build of '%s' by %s.)"),
-					*WriteToString<64>(Log.Category), *WriteToString<256>(Log.Message),
-					*Name, *WriteToString<32>(Function));
-				break;
-			default:
-				checkNoEntry();
-				break;
-			}
+			DefinitionBuilder.AddInputBuild(UTF8TEXTVIEW("MipTail"), { InParentBuildKey, UE::DerivedData::FValueId::FromName(UTF8TEXTVIEW("MipTail")) });
 		}
 
-		if (Output.HasError())
+		for (int32 MipIndex = 0; MipIndex < InputTextureNumStreamingMips; MipIndex++)
 		{
-			UE_LOG(LogTexture, Warning, TEXT("Failed to build derived data for build of '%s' by %s."),
-				*Name, *WriteToString<32>(Function));
-			return;
+			TUtf8StringBuilder<10> MipName;
+			MipName << "Mip" << MipIndex;
+			DefinitionBuilder.AddInputBuild(MipName, { InParentBuildKey, UE::DerivedData::FValueId::FromName(MipName) });
 		}
 
-		DeserializeTextureFromValues(DerivedData, MoveTemp(Params), FirstMipToLoad, bInlineMips);
+		return DefinitionBuilder.Build();
 	}
 
 	static UE::DerivedData::EPriority ConvertPriority(EQueuedWorkPriority SourcePriority)
@@ -2119,25 +2557,37 @@ private:
 		}
 	}
 
+	static ITextureCompressorModule* TextureCompressorModule;
 	static bool LoadModules()
 	{
 		FModuleManager::LoadModuleChecked<IImageWrapperModule>(FName("ImageWrapper"));
-		FModuleManager::LoadModuleChecked<ITextureCompressorModule>(TEXTURE_COMPRESSOR_MODULENAME);
+		TextureCompressorModule = &FModuleManager::LoadModuleChecked<ITextureCompressorModule>(TEXTURE_COMPRESSOR_MODULENAME);
 		return true;
 	}
+	
+	// Stuff that we get as a result of the build.
+	FBuildResults BuildResults;
+
+	// Controls for what optional build outputs we want.
+	FBuildResultOptions BuildResultOptions;
 
-	FTexturePlatformData& DerivedData;
+	// Build bureaucracy
 	TOptional<UE::DerivedData::FRequestOwner> Owner;
-	UE::DerivedData::FOptionalBuildSession BuildSession;
 	EQueuedWorkPriority Priority;
-	bool bCacheHit;
-	bool bInlineMips;
-	int32 FirstMipToLoad;
-	uint64 BuildOutputSize = 0;
+
+	UE::DerivedData::FOptionalBuildSession BuildSession;
 	UE::TextureDerivedData::FTextureBuildInputResolver InputResolver;
+	
+	TOptional<UE::TextureDerivedData::FBuildSession> FetchOrBuild_ChildBuildSession;
+	TOptional<UE::TextureDerivedData::FBuildSession> FetchFirst_ChildBuildSession;
+	TOptional<UE::TextureDerivedData::FTilingTextureBuildInputResolver> FetchOrBuild_ChildInputResolver;
+	TOptional<UE::TextureDerivedData::FTilingTextureBuildInputResolver> FetchFirst_ChildInputResolver;
+
 	FRWLock Lock;
 }; // end DDC2 fetch/build task (FTextureBuildTask)
 
+/* static */ ITextureCompressorModule* FTextureBuildTask::TextureCompressorModule = 0;
+
 FTextureAsyncCacheDerivedDataTask* CreateTextureBuildTask(
 	UTexture& Texture,
 	FTexturePlatformData& DerivedData,
@@ -2149,9 +2599,16 @@ FTextureAsyncCacheDerivedDataTask* CreateTextureBuildTask(
 	ETextureCacheFlags Flags)
 {
 	using namespace UE::DerivedData;
+	FUtf8SharedString TilingFunctionName;
+	if (SettingsFetchOrBuild.Tiler)
+	{
+		TilingFunctionName = SettingsFetchOrBuild.Tiler->GetBuildFunctionName();
+	}
+
 	if (FUtf8SharedString FunctionName = FindTextureBuildFunction(SettingsFetchOrBuild.TextureFormatName); !FunctionName.IsEmpty())
 	{
-		return new FTextureBuildTask(Texture, DerivedData, FunctionName, SettingsFetch, SettingsFetchOrBuild, FetchMetadata, FetchOrBuildMetadata, Priority, Flags);
+		return new FTextureBuildTask(Texture, DerivedData, FunctionName, TilingFunctionName, SettingsFetch, SettingsFetchOrBuild, FetchMetadata, FetchOrBuildMetadata, Priority, Flags);
+
 	}
 	return nullptr;
 }
@@ -2163,6 +2620,11 @@ FTexturePlatformData::FStructuredDerivedDataKey CreateTextureDerivedDataKey(
 {
 	using namespace UE::DerivedData;
 
+	FUtf8SharedString TilingFunctionName;
+	if (Settings.Tiler)
+	{
+		TilingFunctionName = Settings.Tiler->GetBuildFunctionName();
+	}
 	if (FUtf8SharedString FunctionName = FindTextureBuildFunction(Settings.TextureFormatName); !FunctionName.IsEmpty())
 	{
 		IBuild& Build = GetBuild();
@@ -2182,8 +2644,13 @@ FTexturePlatformData::FStructuredDerivedDataKey CreateTextureDerivedDataKey(
 			int64 RequiredMemoryEstimate = GetBuildRequiredMemoryEstimate(&Texture,&Settings);
 
 			FBuildDefinition Definition = FTextureBuildTask::CreateDefinition(Build, Texture, TexturePath.ToView(), FunctionName, Settings, bUseCompositeTexture, RequiredMemoryEstimate);
+			TOptional<FBuildDefinition> TilingDefinition;
+			if (TilingFunctionName.IsEmpty() == false)
+			{
+				TilingDefinition.Emplace(FTextureBuildTask::CreateTilingDefinition(Build, &Texture, Settings, nullptr, nullptr, Definition, TexturePath.ToView(), TilingFunctionName));
+			}
 
-			return FTextureBuildTask::GetKey(Definition, Texture, bUseCompositeTexture);
+			return FTextureBuildTask::GetKey(Definition, TilingDefinition.GetPtrOrNull(), Texture, bUseCompositeTexture);
 		}
 	}
 	return FTexturePlatformData::FStructuredDerivedDataKey();
diff --git a/Engine/Source/Runtime/Engine/Private/TextureDerivedDataTask.h b/Engine/Source/Runtime/Engine/Private/TextureDerivedDataTask.h
--- a/Engine/Source/Runtime/Engine/Private/TextureDerivedDataTask.h
+++ b/Engine/Source/Runtime/Engine/Private/TextureDerivedDataTask.h
@@ -24,7 +24,10 @@
 
 enum
 {
-	/** The number of mips to store inline. */
+	/**
+	*	The number of mips to store "inline". This is the number of mips, counting from the smallest, that are
+	*	always loaded and never streamed.
+	*/
 	NUM_INLINE_DERIVED_MIPS = 7,
 };
 
diff --git a/Engine/Source/Runtime/Engine/Private/VT/VirtualTextureDataBuilder.cpp b/Engine/Source/Runtime/Engine/Private/VT/VirtualTextureDataBuilder.cpp
--- a/Engine/Source/Runtime/Engine/Private/VT/VirtualTextureDataBuilder.cpp
+++ b/Engine/Source/Runtime/Engine/Private/VT/VirtualTextureDataBuilder.cpp
@@ -669,7 +669,7 @@ void FVirtualTextureDataBuilder::BuildTiles(const TArray<FVTSourceTileEntry>& Ti
 			TArray<FImage> EmptyList;
 			uint32 NumMipsInTail, ExtData;
 			// this is the Build for Tiles to do the encode to GPU formats, with no processing
-			if (!ensure(Compressor->BuildTexture(TileImages, EmptyList, TBSettings, DebugTilePathName, CompressedMip, NumMipsInTail, ExtData)))
+			if (!ensure(Compressor->BuildTexture(TileImages, EmptyList, TBSettings, DebugTilePathName, CompressedMip, NumMipsInTail, ExtData, nullptr)))
 			{
 				bCompressionError = true;
 			}
@@ -924,7 +924,7 @@ void FVirtualTextureDataBuilder::BuildSourcePixels(const FTextureSourceData& Sou
 			if (LocalBlockSizeScale == 1)
 			{
 				uint32 NumMipsInTail, ExtData;
-				bBuildTextureResult = Compressor->BuildTexture(SourceMips, *CompositeSourceMips, TBSettings, CurDebugTexturePathName, CompressedMips, NumMipsInTail, ExtData);
+				bBuildTextureResult = Compressor->BuildTexture(SourceMips, *CompositeSourceMips, TBSettings, CurDebugTexturePathName, CompressedMips, NumMipsInTail, ExtData, nullptr);
 			}
 			else
 			{
@@ -954,7 +954,7 @@ void FVirtualTextureDataBuilder::BuildSourcePixels(const FTextureSourceData& Sou
 				TBSettings.bUseLegacyGamma = false;
 
 				uint32 NumMipsInTail, ExtData;
-				bBuildTextureResult = Compressor->BuildTexture(ScaledSourceMips, ScaledCompositeMips, TBSettings, CurDebugTexturePathName, CompressedMips, NumMipsInTail, ExtData);
+				bBuildTextureResult = Compressor->BuildTexture(ScaledSourceMips, ScaledCompositeMips, TBSettings, CurDebugTexturePathName, CompressedMips, NumMipsInTail, ExtData, nullptr);
 			}
 
 			check(bBuildTextureResult);
@@ -1098,7 +1098,7 @@ void FVirtualTextureDataBuilder::BuildSourcePixels(const FTextureSourceData& Sou
 			TArray<FCompressedImage2D> CompressedMips;
 			uint32 NumMipsInTail, ExtData;
 			// this is a Build to uncompressed, to apply processing
-			if (!Compressor->BuildTexture(MiptailInputImages, EmptyImageArray, TBSettings, CurDebugTexturePathName, CompressedMips, NumMipsInTail, ExtData))
+			if (!Compressor->BuildTexture(MiptailInputImages, EmptyImageArray, TBSettings, CurDebugTexturePathName, CompressedMips, NumMipsInTail, ExtData, nullptr))
 			{
 				check(false);
 			}
