Reusability improvements for the ComponentInstanceDataCache

- Split component source info from FActorComponentInstanceData into its own struct so it can be re-used as a key for per-component diffing.
- Added FComponentInstanceDataCache::GetComponentHierarchy util, to get the same component hierarchy data that's used by ApplyToActor.

#rnx

diff --git a/Engine/Source/Runtime/Engine/Private/ComponentInstanceDataCache.cpp b/Engine/Source/Runtime/Engine/Private/ComponentInstanceDataCache.cpp
--- a/Engine/Source/Runtime/Engine/Private/ComponentInstanceDataCache.cpp
+++ b/Engine/Source/Runtime/Engine/Private/ComponentInstanceDataCache.cpp
@@ -430,6 +430,53 @@ void FActorInstanceData::ApplyToActor(AActor* Actor, const ECacheApplyPhase Cach
 	}
 }
 
+
+FActorComponentInstanceSourceInfo::FActorComponentInstanceSourceInfo(const UActorComponent* SourceComponent)
+	: FActorComponentInstanceSourceInfo(SourceComponent->GetArchetype(), SourceComponent->CreationMethod, SourceComponent->GetUCSSerializationIndex())
+{
+}
+
+FActorComponentInstanceSourceInfo::FActorComponentInstanceSourceInfo(TObjectPtr<const UObject> InSourceComponentTemplate, EComponentCreationMethod InSourceComponentCreationMethod, int32 InSourceComponentTypeSerializedIndex)
+	: SourceComponentTemplate(InSourceComponentTemplate)
+	, SourceComponentCreationMethod(InSourceComponentCreationMethod)
+	, SourceComponentTypeSerializedIndex(InSourceComponentTypeSerializedIndex)
+{
+}
+
+bool FActorComponentInstanceSourceInfo::MatchesComponent(const UActorComponent* Component) const
+{
+	return MatchesComponent(Component, Component->GetArchetype());
+}
+
+bool FActorComponentInstanceSourceInfo::MatchesComponent(const UActorComponent* Component, const UObject* ComponentTemplate) const
+{
+	if (SourceComponentTemplate
+		&& SourceComponentTemplate->GetClass() == ComponentTemplate->GetClass()
+		&& (SourceComponentTemplate == ComponentTemplate || (GIsReinstancing && SourceComponentTemplate->GetFName() == ComponentTemplate->GetFName()))
+		&& SourceComponentCreationMethod == Component->CreationMethod
+		)
+	{
+		if (SourceComponentCreationMethod != EComponentCreationMethod::UserConstructionScript)
+		{
+			return true;
+		}
+		
+		if (SourceComponentTypeSerializedIndex >= 0)
+		{
+			check(SourceComponentCreationMethod == EComponentCreationMethod::UserConstructionScript);
+			return SourceComponentTypeSerializedIndex == Component->GetUCSSerializationIndex();
+		}
+	}
+
+	return false;
+}
+
+void FActorComponentInstanceSourceInfo::AddReferencedObjects(FReferenceCollector& Collector)
+{
+	Collector.AddReferencedObject(SourceComponentTemplate);
+}
+
+
 FActorComponentInstanceData::FActorComponentInstanceData()
 	: SourceComponentTemplate(nullptr)
 	, SourceComponentCreationMethod(EComponentCreationMethod::Native)
@@ -441,7 +488,6 @@ FActorComponentInstanceData::FActorComponentInstanceData(const UActorComponent*
 	check(SourceComponent);
 	SourceComponentTemplate = SourceComponent->GetArchetype();
 	SourceComponentCreationMethod = SourceComponent->CreationMethod;
-	SourceComponentTypeSerializedIndex = INDEX_NONE;
 	SourceComponentTypeSerializedIndex = SourceComponent->GetUCSSerializationIndex();
 
 	if (SourceComponent->IsEditableWhenInherited())
@@ -474,23 +520,9 @@ FActorComponentInstanceData::FActorComponentInstanceData(const UActorComponent*
 	}
 }
 
-bool FActorComponentInstanceData::MatchesComponent(const UActorComponent* Component, const UObject* ComponentTemplate, const TMap<UActorComponent*, const UObject*>& ComponentToArchetypeMap) const
+bool FActorComponentInstanceData::MatchesComponent(const UActorComponent* Component, const UObject* ComponentTemplate) const
 {
-	bool bMatches = false;
-	if (   Component
-		&& Component->CreationMethod == SourceComponentCreationMethod
-		&& (ComponentTemplate == SourceComponentTemplate || (GIsReinstancing && ComponentTemplate->GetFName() == SourceComponentTemplate->GetFName())))
-	{
-		if (SourceComponentCreationMethod != EComponentCreationMethod::UserConstructionScript)
-		{
-			bMatches = true;
-		}
-		else if (SourceComponentTypeSerializedIndex >= 0)
-		{
-			bMatches = (SourceComponentTypeSerializedIndex == Component->GetUCSSerializationIndex());
-		}
-	}
-	return bMatches;
+	return FActorComponentInstanceSourceInfo(SourceComponentTemplate, SourceComponentCreationMethod, SourceComponentTypeSerializedIndex).MatchesComponent(Component, ComponentTemplate);
 }
 
 void FActorComponentInstanceData::ApplyToComponent(UActorComponent* Component, const ECacheApplyPhase CacheApplyPhase)
@@ -631,98 +663,89 @@ void FComponentInstanceDataCache::Serialize(FArchive& Ar)
 	Ar << InstanceComponentTransformToRootMap;
 }
 
-void FComponentInstanceDataCache::ApplyToActor(AActor* Actor, const ECacheApplyPhase CacheApplyPhase) const
+void FComponentInstanceDataCache::GetComponentHierarchy(const AActor* Actor, TArray<UActorComponent*, TInlineAllocator<NumInlinedActorComponents>>& OutComponentHierarchy)
 {
-	if (Actor != nullptr)
-	{
-		const bool bIsChildActor = Actor->IsChildActor();
+	// We want to apply instance data from the root node down to ensure changes such as transforms 
+	// propagate correctly so we will build the components list in a breadth-first manner.
+	OutComponentHierarchy.Reset(Actor->GetComponents().Num());
 
-		// We want to apply instance data from the root node down to ensure changes such as transforms 
-		// propagate correctly so we will build the components list in a breadth-first manner.
-		TInlineComponentArray<UActorComponent*> Components;
-		Components.Reserve(Actor->GetComponents().Num());
+	auto AddComponentHierarchy = [Actor, &OutComponentHierarchy](USceneComponent* Component)
+	{
+		int32 FirstProcessIndex = OutComponentHierarchy.Num();
 
-		auto AddComponentHierarchy = [Actor, &Components](USceneComponent* Component)
-		{
-			int32 FirstProcessIndex = Components.Num();
+		// Add this to our list and make it our starting node
+		OutComponentHierarchy.Add(Component);
 
-			// Add this to our list and make it our starting node
-			Components.Add(Component);
+		int32 CompsToProcess = 1;
 
-			int32 CompsToProcess = 1;
+		while (CompsToProcess)
+		{
+			// track how many elements were here
+			const int32 StartingProcessedCount = OutComponentHierarchy.Num();
 
-			while (CompsToProcess)
+			// process the currently unprocessed elements
+			for (int32 ProcessIndex = 0; ProcessIndex < CompsToProcess; ++ProcessIndex)
 			{
-				// track how many elements were here
-				const int32 StartingProcessedCount = Components.Num();
+				USceneComponent* SceneComponent = CastChecked<USceneComponent>(OutComponentHierarchy[FirstProcessIndex + ProcessIndex]);
 
-				// process the currently unprocessed elements
-				for (int32 ProcessIndex = 0; ProcessIndex < CompsToProcess; ++ProcessIndex)
+				// add all children to the end of the array
+				for (int32 ChildIndex = 0; ChildIndex < SceneComponent->GetNumChildrenComponents(); ++ChildIndex)
 				{
-					USceneComponent* SceneComponent = CastChecked<USceneComponent>(Components[FirstProcessIndex + ProcessIndex]);
-
-					// add all children to the end of the array
-					for (int32 ChildIndex = 0; ChildIndex < SceneComponent->GetNumChildrenComponents(); ++ChildIndex)
+					if (USceneComponent* ChildComponent = SceneComponent->GetChildComponent(ChildIndex))
 					{
-						if (USceneComponent* ChildComponent = SceneComponent->GetChildComponent(ChildIndex))
+						// We don't want to recurse in to child actors (or any other attached actor) when applying the instance cache, 
+						// components within a child actor are handled by applying the instance data to the child actor component
+						if (ChildComponent->GetOwner() == Actor)
 						{
-							// We don't want to recurse in to child actors (or any other attached actor) when applying the instance cache, 
-							// components within a child actor are handled by applying the instance data to the child actor component
-							if (ChildComponent->GetOwner() == Actor)
-							{
-							Components.Add(ChildComponent);
+							OutComponentHierarchy.Add(ChildComponent);
 						}
 					}
 				}
-				}
-
-				// next loop start with the nodes we just added
-				FirstProcessIndex = StartingProcessedCount;
-				CompsToProcess = Components.Num() - StartingProcessedCount;
 			}
-		};
 
-		if (USceneComponent* RootComponent = Actor->GetRootComponent())
-		{
-			AddComponentHierarchy(RootComponent);
+			// next loop start with the nodes we just added
+			FirstProcessIndex = StartingProcessedCount;
+			CompsToProcess = OutComponentHierarchy.Num() - StartingProcessedCount;
 		}
+	};
 
-		for (UActorComponent* Component : Actor->GetComponents())
+	if (USceneComponent* RootComponent = Actor->GetRootComponent())
+	{
+		AddComponentHierarchy(RootComponent);
+	}
+
+	for (UActorComponent* Component : Actor->GetComponents())
+	{
+		if (USceneComponent* SceneComponent = Cast<USceneComponent>(Component))
 		{
-			if (USceneComponent* SceneComponent = Cast<USceneComponent>(Component))
+			if (SceneComponent != Actor->GetRootComponent())
 			{
-				if (SceneComponent != Actor->GetRootComponent())
-				{
-					// Scene components that aren't attached to the root component hierarchy won't already have been processed so process them now.
-					// * If there is an unattached scene component 
-					// * If there is a scene component attached to another Actor's hierarchy
-					// * If the scene is not registered (likely because bAutoRegister is false or component is marked pending kill), then we may not have successfully attached to our parent and properly been handled
+				// Scene components that aren't attached to the root component hierarchy won't already have been processed so process them now.
+				// * If there is an unattached scene component 
+				// * If there is a scene component attached to another Actor's hierarchy
+				// * If the scene is not registered (likely because bAutoRegister is false or component is marked pending kill), then we may not have successfully attached to our parent and properly been handled
 				USceneComponent* ParentComponent = SceneComponent->GetAttachParent();
-					if (   (ParentComponent == nullptr)
-					    || (ParentComponent->GetOwner() != Actor) 					
-					    || (!SceneComponent->IsRegistered() && !ParentComponent->GetAttachChildren().Contains(SceneComponent)))
+				if ((ParentComponent == nullptr)
+					|| (ParentComponent->GetOwner() != Actor)
+					|| (!SceneComponent->IsRegistered() && !ParentComponent->GetAttachChildren().Contains(SceneComponent)))
 				{
 					AddComponentHierarchy(SceneComponent);
 				}
 			}
-			}
-			else if (Component)
-			{
-				Components.Add(Component);
-			}
 		}
-
-		// Cache all archetype objects
-		TMap<UActorComponent*, const UObject*> ComponentToArchetypeMap;
-		ComponentToArchetypeMap.Reserve(Components.Num());
-
-		for (UActorComponent* ComponentInstance : Components)
+		else if (Component)
 		{
-			if (bIsChildActor || ComponentInstance->IsCreatedByConstructionScript())
-			{
-				ComponentToArchetypeMap.Add(ComponentInstance, ComponentInstance->GetArchetype());
-			}
+			OutComponentHierarchy.Add(Component);
 		}
+	}
+}
+
+void FComponentInstanceDataCache::ApplyToActor(AActor* Actor, const ECacheApplyPhase CacheApplyPhase) const
+{
+	if (Actor != nullptr)
+	{
+		TInlineComponentArray<UActorComponent*> Components;
+		GetComponentHierarchy(Actor, Components);
 
 		// Apply per-instance data.
 		TBitArray<> ComponentInstanceDataToConsider(true, ComponentsInstanceData.Num());
@@ -734,18 +757,18 @@ void FComponentInstanceDataCache::ApplyToActor(AActor* Actor, const ECacheApplyP
 			}
 		}
 
+		const bool bIsChildActor = Actor->IsChildActor();
 		for (UActorComponent* ComponentInstance : Components)
 		{
 			if (ComponentInstance && (bIsChildActor || ComponentInstance->IsCreatedByConstructionScript())) // Only try and apply data to 'created by construction script' components
 			{
 				// Cache template here to avoid redundant calls in the loop below
-				const UObject* ComponentTemplate = ComponentToArchetypeMap.FindChecked(ComponentInstance);
+				const UObject* ComponentTemplate = ComponentInstance->GetArchetype();
 
 				for (TConstSetBitIterator<> ComponentInstanceDataIt(ComponentInstanceDataToConsider); ComponentInstanceDataIt; ++ComponentInstanceDataIt)
 				{
 					const TStructOnScope<FActorComponentInstanceData>& ComponentInstanceData = ComponentsInstanceData[ComponentInstanceDataIt.GetIndex()];
-					if (   ComponentInstanceData->GetComponentClass() == ComponentTemplate->GetClass() // filter on class early to avoid unnecessary virtual and expensive tests
-					    && ComponentInstanceData->MatchesComponent(ComponentInstance, ComponentTemplate, ComponentToArchetypeMap))
+					if (ComponentInstanceData->MatchesComponent(ComponentInstance, ComponentTemplate))
 					{
 						ComponentInstanceData->ApplyToComponent(ComponentInstance, CacheApplyPhase);
 						ComponentInstanceDataToConsider[ComponentInstanceDataIt.GetIndex()] = false;
diff --git a/Engine/Source/Runtime/Engine/Public/ComponentInstanceDataCache.h b/Engine/Source/Runtime/Engine/Public/ComponentInstanceDataCache.h
--- a/Engine/Source/Runtime/Engine/Public/ComponentInstanceDataCache.h
+++ b/Engine/Source/Runtime/Engine/Public/ComponentInstanceDataCache.h
@@ -6,6 +6,7 @@
 #include "Templates/SubclassOf.h"
 #include "UObject/Object.h"
 #include "UObject/StructOnScope.h"
+#include "Engine/EngineTypes.h"
 #include "ComponentInstanceDataCache.generated.h"
 
 class AActor;
@@ -115,6 +116,37 @@ private:
 	TArray<FName> ReferencedNames;
 };
 
+USTRUCT()
+struct ENGINE_API FActorComponentInstanceSourceInfo
+{
+	GENERATED_BODY()
+
+public:
+	FActorComponentInstanceSourceInfo() = default;
+	explicit FActorComponentInstanceSourceInfo(const UActorComponent* SourceComponent);
+	FActorComponentInstanceSourceInfo(TObjectPtr<const UObject> InSourceComponentTemplate, EComponentCreationMethod InSourceComponentCreationMethod, int32 InSourceComponentTypeSerializedIndex);
+
+	/** Determines whether this component instance data matches the component */
+	bool MatchesComponent(const UActorComponent* Component) const;
+	bool MatchesComponent(const UActorComponent* Component, const UObject* ComponentTemplate) const;
+
+	void AddReferencedObjects(FReferenceCollector& Collector);
+
+private:
+	/** The template used to create the source component */
+	UPROPERTY()
+	TObjectPtr<const UObject> SourceComponentTemplate = nullptr;
+
+	/** The method that was used to create the source component */
+	UPROPERTY()
+	EComponentCreationMethod SourceComponentCreationMethod = EComponentCreationMethod::Native;
+
+	/** The index of the source component in its owner's serialized array
+	when filtered to just that component type */
+	UPROPERTY()
+	int32 SourceComponentTypeSerializedIndex = INDEX_NONE;
+};
+
 /** Base class for component instance cached data of a particular type. */
 USTRUCT()
 struct ENGINE_API FActorComponentInstanceData : public FInstanceCacheDataBase
@@ -125,7 +157,7 @@ public:
 	FActorComponentInstanceData(const UActorComponent* SourceComponent);
 
 	/** Determines whether this component instance data matches the component */
-	bool MatchesComponent(const UActorComponent* Component, const UObject* ComponentTemplate, const TMap<UActorComponent*, const UObject*>& ComponentToArchetypeMap) const;
+	bool MatchesComponent(const UActorComponent* Component, const UObject* ComponentTemplate) const;
 
 	/** Determines if any instance data was actually saved. */
 	virtual bool ContainsData() const { return SavedProperties.Num() > 0; }
@@ -217,6 +249,8 @@ public:
 
 	void AddReferencedObjects(FReferenceCollector& Collector);
 
+	static void GetComponentHierarchy(const AActor* Actor, TArray<UActorComponent*, TInlineAllocator<NumInlinedActorComponents>>& OutComponentHierarchy);
+
 private:
 	// called during de-serialization to copy serialized properties over existing component instance data and keep non UPROPERTY data intact
 	void CopySerializableProperties(TArray<TStructOnScope<FActorComponentInstanceData>> InComponentsInstanceData);
