Set velocity relevance correctly for Material Billboards.

diff --git a/Engine/Source/Runtime/Engine/Private/Components/MaterialBillboardComponent.cpp b/Engine/Source/Runtime/Engine/Private/Components/MaterialBillboardComponent.cpp
--- a/Engine/Source/Runtime/Engine/Private/Components/MaterialBillboardComponent.cpp
+++ b/Engine/Source/Runtime/Engine/Private/Components/MaterialBillboardComponent.cpp
@@ -269,6 +269,7 @@ public:
 		Result.bDynamicRelevance = true;
 		Result.bShadowRelevance = IsShadowCast(View);
 		MaterialRelevance.SetPrimitiveViewRelevance(Result);
+		Result.bVelocityRelevance = DrawsVelocity() && Result.bOpaque && Result.bRenderInMainPass;
 		return Result;
 	}
 	virtual bool CanBeOccluded() const override { return !MaterialRelevance.bDisableDepthTest; }
