Map a debug input device to a local player if they are created with the DebugCreatePlayer command.

#rnx

diff --git a/Engine/Source/Runtime/Engine/Private/GameInstance.cpp b/Engine/Source/Runtime/Engine/Private/GameInstance.cpp
--- a/Engine/Source/Runtime/Engine/Private/GameInstance.cpp
+++ b/Engine/Source/Runtime/Engine/Private/GameInstance.cpp
@@ -938,6 +938,19 @@ void UGameInstance::DebugCreatePlayer(int32 ControllerId)
 	{
 		UE_LOG(LogPlayerManagement, Error, TEXT("Failed to DebugCreatePlayer: %s"), *Error);
 	}
+	else
+	{
+		FPlatformUserId UserId = FGenericPlatformMisc::GetPlatformUserForUserIndex(ControllerId);
+		FInputDeviceId InputDevice = INPUTDEVICEID_NONE;
+		IPlatformInputDeviceMapper& DeviceMapper = IPlatformInputDeviceMapper::Get();
+		DeviceMapper.RemapControllerIdToPlatformUserAndDevice(ControllerId, UserId, InputDevice);
+	
+		// If the input device that was created hasn't been mapped yet, we shuold map a dummy input device to it
+		if (!DeviceMapper.GetUserForInputDevice(InputDevice).IsValid())
+		{
+			DeviceMapper.Internal_MapInputDeviceToUser(InputDevice, UserId, EInputDeviceConnectionState::Connected);
+		}
+	}
 #endif
 }
 
