PR #9484: Fixed typos in "toggle screen messages" outputs (Contributed by [REDACTED_PERSON])

diff --git a/Engine/Source/Runtime/Engine/Private/UnrealEngine.cpp b/Engine/Source/Runtime/Engine/Private/UnrealEngine.cpp
--- a/Engine/Source/Runtime/Engine/Private/UnrealEngine.cpp
+++ b/Engine/Source/Runtime/Engine/Private/UnrealEngine.cpp
@@ -9475,7 +9475,7 @@ bool UEngine::HandleDisableAllScreenMessagesCommand( const TCHAR* Cmd, FOutputDe
 bool UEngine::HandleEnableAllScreenMessagesCommand( const TCHAR* Cmd, FOutputDevice& Ar )
 {
 	GAreScreenMessagesEnabled = true;
-	UE_LOG(LogEngine, Log, TEXT("Onscreen warngins/messages are now ENABLED"));
+	UE_LOG(LogEngine, Log, TEXT("Onscreen warnings/messages are now ENABLED"));
 	return true;
 }
 
@@ -9483,7 +9483,7 @@ bool UEngine::HandleEnableAllScreenMessagesCommand( const TCHAR* Cmd, FOutputDev
 bool UEngine::HandleToggleAllScreenMessagesCommand( const TCHAR* Cmd, FOutputDevice& Ar )
 {
 	GAreScreenMessagesEnabled = !GAreScreenMessagesEnabled;
-	UE_LOG(LogEngine, Log, TEXT("Onscreen warngins/messages are now %s"),
+	UE_LOG(LogEngine, Log, TEXT("Onscreen warnings/messages are now %s"),
 		GAreScreenMessagesEnabled ? TEXT("ENABLED") : TEXT("DISABLED"));
 	return true;
 }
