Enforce RoughDiffuse to be compiled out on platform which disable it explicetly from settings.

diff --git a/Engine/Source/Runtime/Engine/Private/ShaderCompiler/ShaderCompiler.cpp b/Engine/Source/Runtime/Engine/Private/ShaderCompiler/ShaderCompiler.cpp
--- a/Engine/Source/Runtime/Engine/Private/ShaderCompiler/ShaderCompiler.cpp
+++ b/Engine/Source/Runtime/Engine/Private/ShaderCompiler/ShaderCompiler.cpp
@@ -6003,10 +6003,18 @@ void GlobalBeginCompileShader(
 	}
 
 	{
-		static IConsoleVariable* CVar = IConsoleManager::Get().FindConsoleVariable(TEXT("r.Strata"));
-		const bool bStrata = CVar && CVar->GetInt() != 0;
+		static IConsoleVariable* CVarStrata = IConsoleManager::Get().FindConsoleVariable(TEXT("r.Strata"));
+		const bool bStrata = CVarStrata && CVarStrata->GetInt() != 0;
 		Input.Environment.SetDefine(TEXT("STRATA_ENABLED"), bStrata ? 1 : 0);
 
+		// Force rough diffuse to be disable on platform which explicitly disable it from their settings
+		if (bStrata && IsConsolePlatform(Target.GetPlatform()))
+		{
+			static FShaderPlatformCachedIniValue<int32> CVarStrataRoughDiffuse(TEXT("r.Strata.RoughDiffuse"));
+			const bool bStrataRoughDiffuse = CVarStrataRoughDiffuse.Get(Target.GetPlatform()) != 0;
+			Input.Environment.SetDefine(TEXT("STRATA_DIFFUSE_CHAN"), bStrataRoughDiffuse ? 1 : 0);
+		}
+
 		static IConsoleVariable* CVarBackCompatibility = IConsoleManager::Get().FindConsoleVariable(TEXT("r.StrataBackCompatibility"));
 		const bool bStrataBackCompatibility = CVarBackCompatibility && CVarBackCompatibility->GetInt() > 0;
 		Input.Environment.SetDefine(TEXT("PROJECT_STRATA_BACKCOMPATIBILITY"), bStrataBackCompatibility ? 1 : 0);
