Adding the ability to swap audio spatialization plugins at runtime via new FAudioDevice methods and console commands.  Added UI-friendly functions to get/set via FName.

diff --git a/Engine/Plugins/Runtime/Oculus/OculusAudio/Source/OculusAudio/Private/OculusAudioContextManager.cpp b/Engine/Plugins/Runtime/Oculus/OculusAudio/Source/OculusAudio/Private/OculusAudioContextManager.cpp
--- a/Engine/Plugins/Runtime/Oculus/OculusAudio/Source/OculusAudio/Private/OculusAudioContextManager.cpp
+++ b/Engine/Plugins/Runtime/Oculus/OculusAudio/Source/OculusAudio/Private/OculusAudioContextManager.cpp
@@ -32,17 +32,18 @@ void FOculusAudioContextManager::OnListenerInitialize(FAudioDevice* AudioDevice,
 void FOculusAudioContextManager::OnListenerShutdown(FAudioDevice* AudioDevice)
 {
 	FOculusAudioPlugin* Plugin = &FModuleManager::GetModuleChecked<FOculusAudioPlugin>("OculusAudio");
-	FString OculusSpatializerPluginName = Plugin->GetSpatializationPluginFactory()->GetDisplayName();
-	FString CurrentSpatializerPluginName = AudioPluginUtilities::GetDesiredPluginName(EAudioPlugin::SPATIALIZATION);
+	const FString OculusSpatializerPluginName = Plugin->GetSpatializationPluginFactory()->GetDisplayName();
+	const FString CurrentSpatializerPluginName = AudioPluginUtilities::GetDesiredPluginName(EAudioPlugin::SPATIALIZATION);
 	if (CurrentSpatializerPluginName.Equals(OculusSpatializerPluginName)) // we have a match!
 	{
 		OculusAudioSpatializationAudioMixer* Spatializer =
-			static_cast<OculusAudioSpatializationAudioMixer*>(AudioDevice->SpatializationPluginInterface.Get());
+			static_cast<OculusAudioSpatializationAudioMixer*>(AudioDevice->GetSpatializationPluginInterface().Get());
+
 		Spatializer->ClearContext();
 	}
 
-	FString OculusReverbPluginName = Plugin->GetReverbPluginFactory()->GetDisplayName();
-	FString CurrentReverbPluginName = AudioPluginUtilities::GetDesiredPluginName(EAudioPlugin::REVERB);
+	const FString OculusReverbPluginName = Plugin->GetReverbPluginFactory()->GetDisplayName();
+	const FString CurrentReverbPluginName = AudioPluginUtilities::GetDesiredPluginName(EAudioPlugin::REVERB);
 	if (CurrentReverbPluginName.Equals(OculusReverbPluginName))
 	{
 		OculusAudioReverb* Reverb = static_cast<OculusAudioReverb*>(AudioDevice->ReverbPluginInterface.Get());
diff --git a/Engine/Plugins/Runtime/ResonanceAudio/Source/ResonanceAudio/Private/ResonanceAudioPluginListener.cpp b/Engine/Plugins/Runtime/ResonanceAudio/Source/ResonanceAudio/Private/ResonanceAudioPluginListener.cpp
--- a/Engine/Plugins/Runtime/ResonanceAudio/Source/ResonanceAudio/Private/ResonanceAudioPluginListener.cpp
+++ b/Engine/Plugins/Runtime/ResonanceAudio/Source/ResonanceAudio/Private/ResonanceAudioPluginListener.cpp
@@ -71,7 +71,7 @@ namespace ResonanceAudio
 		}
 
 		ReverbPtr = (FResonanceAudioReverb*)AudioDevice->ReverbPluginInterface.Get();
-		SpatializationPtr = (FResonanceAudioSpatialization*)AudioDevice->SpatializationPluginInterface.Get();
+		SpatializationPtr = (FResonanceAudioSpatialization*)AudioDevice->GetSpatializationPluginInterface().Get();
 
 		// Make sure that Reverb *AND* spatialization plugins are enabled.
 		if (ReverbPtr == nullptr || SpatializationPtr == nullptr)
diff --git a/Engine/Source/Runtime/AudioMixer/Private/AudioMixerDevice.cpp b/Engine/Source/Runtime/AudioMixer/Private/AudioMixerDevice.cpp
--- a/Engine/Source/Runtime/AudioMixer/Private/AudioMixerDevice.cpp
+++ b/Engine/Source/Runtime/AudioMixer/Private/AudioMixerDevice.cpp
@@ -429,7 +429,7 @@ namespace Audio
 				AudioClock = 0.0;
 				AudioClockDelta = (double)OpenStreamParams.NumFrames / OpenStreamParams.SampleRate;
 
-				FAudioPluginInitializationParams PluginInitializationParams;
+
 				PluginInitializationParams.NumSources = SourceManagerInitParams.NumSources;
 				PluginInitializationParams.SampleRate = SampleRate;
 				PluginInitializationParams.BufferLength = OpenStreamParams.NumFrames;
@@ -439,10 +439,11 @@ namespace Audio
 					LLM_SCOPE(ELLMTag::AudioMixerPlugins);
 
 					// Initialize any plugins if they exist
-					if (SpatializationPluginInterface.IsValid())
-					{
-						SpatializationPluginInterface->Initialize(PluginInitializationParams);
-					}
+
+					// spatialization
+					FAudioDevice::FAudioSpatializationInterfaceInfo PluginInfo = GetCurrentSpatializationPluginInterfaceInfo();
+					PluginInfo.SpatializationPlugin->Initialize(PluginInitializationParams);
+					PluginInfo.bIsInitialized = true;
 
 					if (OcclusionInterface.IsValid())
 					{
diff --git a/Engine/Source/Runtime/AudioMixer/Private/AudioMixerSource.cpp b/Engine/Source/Runtime/AudioMixer/Private/AudioMixerSource.cpp
--- a/Engine/Source/Runtime/AudioMixer/Private/AudioMixerSource.cpp
+++ b/Engine/Source/Runtime/AudioMixer/Private/AudioMixerSource.cpp
@@ -379,7 +379,9 @@ namespace Audio
 			InitParams.NumInputFrames = NumFrames;
 			InitParams.SourceVoice = MixerSourceVoice;
 			InitParams.bUseHRTFSpatialization = UseObjectBasedSpatialization();
-			InitParams.bIsExternalSend = MixerDevice->bSpatializationIsExternalSend;
+
+			// in this file once spat override is implemented
+			InitParams.bIsExternalSend = MixerDevice->GetCurrentSpatializationPluginInterfaceInfo().bSpatializationIsExternalSend;
 			InitParams.bIsSoundfield = WaveInstance->bIsAmbisonics && (WaveData->NumChannels == 4);
 
 			FActiveSound* ActiveSound = WaveInstance->ActiveSound;
@@ -471,7 +473,7 @@ namespace Audio
 	
 			// If we're spatializing using HRTF and its an external send, don't need to setup a default/base submix send to master or EQ submix
 			// We'll only be using non-default submix sends (e.g. reverb).
-			if (!(InitParams.bUseHRTFSpatialization && MixerDevice->bSpatializationIsExternalSend))
+			if (!(InitParams.bUseHRTFSpatialization && InitParams.bIsExternalSend))
 			{
 				FMixerSubmixWeakPtr SubmixPtr;
 				// If a sound specifies a base submix manually, always use that
@@ -1729,7 +1731,7 @@ namespace Audio
 
 	bool FMixerSource::UseObjectBasedSpatialization() const
 	{
-		return (Buffer->NumChannels <= MixerDevice->MaxChannelsSupportedBySpatializationPlugin &&
+		return (Buffer->NumChannels <= MixerDevice->GetCurrentSpatializationPluginInterfaceInfo().MaxChannelsSupportedBySpatializationPlugin &&
 				AudioDevice->IsSpatializationPluginEnabled() &&
 				WaveInstance->SpatializationMethod == ESoundSpatializationAlgorithm::SPATIALIZATION_HRTF);
 	}
@@ -1752,7 +1754,7 @@ namespace Audio
 
 	bool FMixerSource::UseSpatializationPlugin() const
 	{
-		return (Buffer->NumChannels <= MixerDevice->MaxChannelsSupportedBySpatializationPlugin) &&
+		return (Buffer->NumChannels <= MixerDevice->GetCurrentSpatializationPluginInterfaceInfo().MaxChannelsSupportedBySpatializationPlugin) &&
 			AudioDevice->IsSpatializationPluginEnabled() &&
 			WaveInstance->SpatializationPluginSettings != nullptr;
 	}
diff --git a/Engine/Source/Runtime/AudioMixer/Private/AudioMixerSourceManager.cpp b/Engine/Source/Runtime/AudioMixer/Private/AudioMixerSourceManager.cpp
--- a/Engine/Source/Runtime/AudioMixer/Private/AudioMixerSourceManager.cpp
+++ b/Engine/Source/Runtime/AudioMixer/Private/AudioMixerSourceManager.cpp
@@ -165,7 +165,6 @@ namespace Audio
 		, NumSourceWorkers(4)
 		, bInitialized(false)
 		, bUsingSpatializationPlugin(false)
-		, MaxChannelsSupportedBySpatializationPlugin(1)
 	{
 		// Get a manual resetable event
 		const bool bIsManualReset = true;
@@ -337,11 +336,11 @@ namespace Audio
 		NumSourceWorkers = SourceWorkers.Num();
 
 		// Cache the spatialization plugin
-		SpatializationPlugin = MixerDevice->SpatializationPluginInterface;
-		if (SpatializationPlugin.IsValid())
+		SpatialInterfaceInfo = MixerDevice->GetCurrentSpatializationPluginInterfaceInfo();
+		const auto& SpatializationPlugin = SpatialInterfaceInfo.SpatializationPlugin;
+		if (SpatialInterfaceInfo.SpatializationPlugin.IsValid())
 		{
 			bUsingSpatializationPlugin = true;
-			MaxChannelsSupportedBySpatializationPlugin = MixerDevice->MaxChannelsSupportedBySpatializationPlugin;
 		}
 
 		// Spam command queue with nops.
@@ -529,7 +528,7 @@ namespace Audio
 		{
 			AUDIO_MIXER_CHECK(bUsingSpatializationPlugin);
 			LLM_SCOPE(ELLMTag::AudioMixerPlugins);
-			SpatializationPlugin->OnReleaseSource(SourceId);
+			SpatialInterfaceInfo.SpatializationPlugin->OnReleaseSource(SourceId);
 		}
 
 		if (SourceInfo.bUseOcclusionPlugin)
@@ -844,7 +843,17 @@ namespace Audio
 			{
 				AUDIO_MIXER_CHECK(bUsingSpatializationPlugin);
 				LLM_SCOPE(ELLMTag::AudioMixerPlugins);
-				SpatializationPlugin->OnInitSource(SourceId, InitParams.AudioComponentUserID, InitParams.SpatializationPluginSettings);
+
+				// re-cache the spatialization plugin in case it changed
+				bUsingSpatializationPlugin = false;
+				SpatialInterfaceInfo = MixerDevice->GetCurrentSpatializationPluginInterfaceInfo();
+				const auto& SpatializationPlugin = SpatialInterfaceInfo.SpatializationPlugin;
+				if (SpatialInterfaceInfo.SpatializationPlugin.IsValid())
+				{
+					bUsingSpatializationPlugin = true;
+				}
+
+				SpatialInterfaceInfo.SpatializationPlugin->OnInitSource(SourceId, InitParams.AudioComponentUserID, InitParams.SpatializationPluginSettings);
 			}
 
 			// Create the occlusion plugin source effect
@@ -2221,8 +2230,8 @@ namespace Audio
 			CSV_SCOPED_TIMING_STAT(Audio, HRTF);
 			SCOPE_CYCLE_COUNTER(STAT_AudioMixerHRTF);
 
-			AUDIO_MIXER_CHECK(SpatializationPlugin.IsValid());
-			AUDIO_MIXER_CHECK(SourceInfo.NumInputChannels <= MaxChannelsSupportedBySpatializationPlugin);
+			AUDIO_MIXER_CHECK(SpatialInterfaceInfo.SpatializationPlugin.IsValid());
+			AUDIO_MIXER_CHECK(SourceInfo.NumInputChannels <= SpatialInterfaceInfo.MaxChannelsSupportedBySpatializationPlugin);
 
 			FAudioPluginSourceInputData AudioPluginInputData;
 			AudioPluginInputData.AudioBuffer = &SourceInfo.SourceBuffer;
@@ -2230,7 +2239,7 @@ namespace Audio
 			AudioPluginInputData.SourceId = SourceId;
 			AudioPluginInputData.SpatializationParams = &SourceInfo.SpatParams;
 
-			if (!MixerDevice->bSpatializationIsExternalSend)
+			if (!SpatialInterfaceInfo.bSpatializationIsExternalSend)
 			{
 				SourceInfo.AudioPluginOutputData.AudioBuffer.Reset();
 				SourceInfo.AudioPluginOutputData.AudioBuffer.AddZeroed(2 * NumOutputFrames);
@@ -2238,7 +2247,7 @@ namespace Audio
 
 			{
 				LLM_SCOPE(ELLMTag::AudioMixerPlugins);
-				SpatializationPlugin->ProcessAudio(AudioPluginInputData, SourceInfo.AudioPluginOutputData);
+				SpatialInterfaceInfo.SpatializationPlugin->ProcessAudio(AudioPluginInputData, SourceInfo.AudioPluginOutputData);
 			}
 
 			// If this is an external send, we treat this source audio as if it was still a mono source
@@ -2246,7 +2255,7 @@ namespace Audio
 			// sent to submixes (e.g. reverb) panned and mixed down. Certain submixes will want this spatial 
 			// information in addition to the external send. We've already bypassed adding this source
 			// to a base submix (e.g. master/eq, etc)
-			if (MixerDevice->bSpatializationIsExternalSend)
+			if (SpatialInterfaceInfo.bSpatializationIsExternalSend)
 			{
 				// Otherwise our pre- and post-effect channels are the same as the input channels
 				SourceInfo.NumPostEffectChannels = SourceInfo.NumInputChannels;
@@ -2519,7 +2528,7 @@ namespace Audio
 				}
 			}
 
-			if (SourceInfo.IsRenderingToSubmixes() || MixerDevice->bSpatializationIsExternalSend)
+			if (SourceInfo.IsRenderingToSubmixes() || SpatialInterfaceInfo.bSpatializationIsExternalSend)
 			{
 				// Apply distance attenuation
 				ApplyDistanceAttenuation(SourceInfo, NumSamples);
@@ -2896,9 +2905,9 @@ namespace Audio
 		// Let the plugin know we finished processing all sources
 		if (bUsingSpatializationPlugin)
 		{
-			AUDIO_MIXER_CHECK(SpatializationPlugin.IsValid());
+			AUDIO_MIXER_CHECK(SpatialInterfaceInfo.SpatializationPlugin.IsValid());
 			LLM_SCOPE(ELLMTag::AudioMixerPlugins);
-			SpatializationPlugin->OnAllSourcesProcessed();
+			SpatialInterfaceInfo.SpatializationPlugin->OnAllSourcesProcessed();
 		}
 
 		// Update the game thread copy of source doneness
diff --git a/Engine/Source/Runtime/AudioMixer/Private/AudioMixerSourceManager.h b/Engine/Source/Runtime/AudioMixer/Private/AudioMixerSourceManager.h
--- a/Engine/Source/Runtime/AudioMixer/Private/AudioMixerSourceManager.h
+++ b/Engine/Source/Runtime/AudioMixer/Private/AudioMixerSourceManager.h
@@ -328,7 +328,7 @@ namespace Audio
 		FMixerDevice* MixerDevice;
 
 		// Cached ptr to an optional spatialization plugin
-		TAudioSpatializationPtr SpatializationPlugin;
+		FAudioDevice::FAudioSpatializationInterfaceInfo SpatialInterfaceInfo;
 
 		// Array of pointers to game thread audio source objects
 		TArray<FMixerSourceVoice*> MixerSources;
@@ -548,7 +548,6 @@ namespace Audio
 
 		uint8 bInitialized : 1;
 		uint8 bUsingSpatializationPlugin : 1;
-		int32 MaxChannelsSupportedBySpatializationPlugin;
 
 		// Set to true when the audio source manager should pump the command queue
 		FThreadSafeBool bPumpQueue;
diff --git a/Engine/Source/Runtime/Engine/Classes/Kismet/GameplayStatics.h b/Engine/Source/Runtime/Engine/Classes/Kismet/GameplayStatics.h
--- a/Engine/Source/Runtime/Engine/Classes/Kismet/GameplayStatics.h
+++ b/Engine/Source/Runtime/Engine/Classes/Kismet/GameplayStatics.h
@@ -905,6 +905,18 @@ public:
 	UFUNCTION(BlueprintCallable, Category = "Audio")
 	static void PrimeSound(USoundBase* InSound);
 
+	/** Get list of available Audio Spatialization Plugin names */
+	UFUNCTION(BlueprintCallable, Category = "Audio", meta = (WorldContext = "WorldContextObject"))
+	static TArray<FName> GetAvailableSpatialPluginNames(const UObject* WorldContextObject);
+
+	/** Get currently active Audio Spatialization Plugin name */
+	UFUNCTION(BlueprintCallable, Category = "Audio", meta = (WorldContext = "WorldContextObject"))
+	static FName GetActiveSpatialPluginName(const UObject* WorldContextObject);
+
+	/** Get list of available Audio Spatialization Plugins */
+	UFUNCTION(BlueprintCallable, Category = "Audio", meta = (WorldContext = "WorldContextObject"))
+	static bool SetActiveSpatialPluginByName(const UObject* WorldContextObject, FName InPluginName);
+
 	/** Primes the sound waves in the given USoundClass, caching the first chunk of streamed audio. */
 	UFUNCTION(BlueprintCallable, Category = "Audio")
 	static void PrimeAllSoundsInSoundClass(class USoundClass* InSoundClass);
diff --git a/Engine/Source/Runtime/Engine/Private/ActiveSound.cpp b/Engine/Source/Runtime/Engine/Private/ActiveSound.cpp
--- a/Engine/Source/Runtime/Engine/Private/ActiveSound.cpp
+++ b/Engine/Source/Runtime/Engine/Private/ActiveSound.cpp
@@ -1840,7 +1840,7 @@ void FActiveSound::UpdateAttenuation(float DeltaTime, FSoundParseParameters& Par
 			ParseParams.SpatializationMethod = Settings->SpatializationAlgorithm;
 		}
 
-		ParseParams.bSpatializationIsExternalSend = AudioDevice->bSpatializationIsExternalSend;
+		ParseParams.bSpatializationIsExternalSend = AudioDevice->GetCurrentSpatializationPluginInterfaceInfo().bSpatializationIsExternalSend;
 	}
 
 	// If not overriding from a node, set focus data
diff --git a/Engine/Source/Runtime/Engine/Private/AudioDevice.cpp b/Engine/Source/Runtime/Engine/Private/AudioDevice.cpp
--- a/Engine/Source/Runtime/Engine/Private/AudioDevice.cpp
+++ b/Engine/Source/Runtime/Engine/Private/AudioDevice.cpp
@@ -180,6 +180,61 @@ FAutoConsoleVariableRef CVarFlushAudioRenderThreadOnGC(
 	TEXT("When set to 1, every time the GC runs, we flush all pending audio render thread commands.\n"),
 	ECVF_Default);
 
+static FAutoConsoleCommandWithWorld GListAvailableSpatialPluginsCommand(
+	TEXT("au.spatialization.ListAvailableSpatialPlugins"),
+	TEXT("This will output a list of currently availible/active spatialization plugins"),
+	FConsoleCommandWithWorldDelegate::CreateStatic(
+		[](UWorld* InWorld)
+	{
+			if(InWorld)
+			{
+				if (FAudioDeviceHandle AudioDevice = InWorld->GetAudioDevice())
+				{
+					const TArray<FName> PluginNames = AudioDevice->GetAvailableSpatializationPluginNames();
+					FName ActivePluginName = AudioDevice->GetCurrentSpatializationPluginInterfaceInfo().PluginName;
+
+					if(GEngine)
+					{
+						for (FName PluginName : PluginNames)
+						{
+							const bool bIsActive = PluginName == ActivePluginName;
+							GEngine->AddOnScreenDebugMessage(-1, 15.0f, FColor::Yellow
+								, PluginName.ToString() + (bIsActive? '*' : ' '));
+						}
+						GEngine->AddOnScreenDebugMessage(-1, 15.0f, FColor::Yellow
+							, FString::Printf(TEXT("-----------------------------------------")));
+						GEngine->AddOnScreenDebugMessage(-1, 15.0f, FColor::Yellow
+							, FString::Printf(TEXT("Available spatial plugins [count = %i] (* = Active)"), PluginNames.Num()));
+					}
+				}
+			}
+	})
+);
+
+static FAutoConsoleCommandWithWorldArgsAndOutputDevice GSetCurrentSpatialPluginCommand(
+	TEXT("au.spatialization.SetCurrentSpatialPlugin"),
+	TEXT("Attempt to swap to the named spatialization plugin (au.spatialization.ListAvailableSpatialPlugins to see what is availible)"),
+	FConsoleCommandWithWorldArgsAndOutputDeviceDelegate::CreateStatic(
+		[](const TArray<FString>& PluginNameArray, UWorld* InWorld, FOutputDevice&)
+		{
+			FString PluginName;
+			for(const FString& Substring : PluginNameArray)
+			{
+				PluginName += Substring;
+				PluginName += ' ';
+			}
+			PluginName.TrimStartAndEndInline();
+
+			if(InWorld)
+			{
+				if (FAudioDeviceHandle AudioDevice = InWorld->GetAudioDevice())
+				{
+					AudioDevice->SetCurrentSpatializationPlugin(FName(PluginName));
+				}
+			}
+		})
+	);
+
 namespace Audio
 {
 	ICompressedAudioInfo* CreateSoundAssetDecoder(const FName& InRuntimeFormat)
@@ -289,7 +344,6 @@ FAudioDevice::FAudioDevice()
 	, SampleRate(0)
 	, NumPrecacheFrames(MONO_PCM_BUFFER_SAMPLES)
 	, DeviceID(static_cast<Audio::FDeviceId>(INDEX_NONE))
-	, SpatializationPluginInterface(nullptr)
 	, SourceDataOverridePluginInterface(nullptr)
 	, ReverbPluginInterface(nullptr)
 	, OcclusionInterface(nullptr)
@@ -320,10 +374,8 @@ FAudioDevice::FAudioDevice()
 	, bIsStoppingVoicesEnabled(false)
 	, bIsBakedAnalysisEnabled(false)
 	, bAudioMixerModuleLoaded(false)
-	, bSpatializationIsExternalSend(false)
 	, bOcclusionIsExternalSend(false)
 	, bReverbIsExternalSend(false)
-	, MaxChannelsSupportedBySpatializationPlugin(1)
 	, bStartupSoundsPreCached(false)
 	, bSpatializationInterfaceEnabled(false)
 	, bOcclusionInterfaceEnabled(false)
@@ -490,27 +542,46 @@ bool FAudioDevice::Init(Audio::FDeviceId InDeviceID, int32 InMaxSources, int32 I
 	// Cache any plugin settings objects we have loaded
 	UpdateAudioPluginSettingsObjectCache();
 
-	//Get the requested spatialization plugin and set it up.
+	//Get the requested default spatialization plugin and set it up.
 	IAudioSpatializationFactory* SpatializationPluginFactory = AudioPluginUtilities::GetDesiredSpatializationPlugin();
 	if (SpatializationPluginFactory != nullptr)
 	{
-		SpatializationPluginInterface = SpatializationPluginFactory->CreateNewSpatializationPlugin(this);
+		// cache the name of our default Spatialization plugin
+		CurrentSpatializationPluginInterfaceName = AudioPluginUtilities::GetDesiredSpatializationPluginName();
+
+		// Cache info for the available spatial plugins
+		TArray<IAudioSpatializationFactory*> SpatializerFactories = AudioPluginUtilities::GetSpatialPluginArray();
+		for(const auto& Factory : SpatializerFactories)
+		{
+			if(!ensure(Factory))
+			{
+				continue;  // null factory ptr in the modular features array
+			}
+
+			const FName PluginName = FName(Factory->GetDisplayName());
+			FAudioSpatializationInterfaceInfo SpatPluginInfo(PluginName, /* InAudioDevice */ this, Factory);
+
+			SpatializationInterfaces.Add(SpatPluginInfo);
+		}
+
 		if (!IsAudioMixerEnabled())
 		{
 			//Set up initialization parameters for system level effect plugins:
-			FAudioPluginInitializationParams PluginInitializationParams;
 			PluginInitializationParams.SampleRate = SampleRate;
-				PluginInitializationParams.NumSources = GetMaxSources();
+			PluginInitializationParams.NumSources = GetMaxSources();
 			PluginInitializationParams.BufferLength = PlatformSettings.CallbackBufferFrameSize;
 			PluginInitializationParams.AudioDevicePtr = this;
 
-			SpatializationPluginInterface->Initialize(PluginInitializationParams);
+			const TAudioSpatializationPtr PluginInterfacePtr = GetSpatializationPluginInterface();
+			if(ensure(PluginInterfacePtr))
+			{
+				GetSpatializationPluginInterface()->Initialize(PluginInitializationParams);
+			}
 		}
 
 		bSpatializationInterfaceEnabled = true;
-		bSpatializationIsExternalSend = SpatializationPluginFactory->IsExternalSend();
-		MaxChannelsSupportedBySpatializationPlugin = SpatializationPluginFactory->GetMaxSupportedChannels();
-		UE_LOG(LogAudio, Display, TEXT("Audio Spatialization Plugin: %s is external send: %d"), *(SpatializationPluginFactory->GetDisplayName()), bSpatializationIsExternalSend);
+
+		UE_LOG(LogAudio, Display, TEXT("Audio Spatialization Plugin: %s is external send: %d"), *(SpatializationPluginFactory->GetDisplayName()), SpatializationPluginFactory->IsExternalSend());
 	}
 	else
 	{
@@ -560,9 +631,8 @@ bool FAudioDevice::Init(Audio::FDeviceId InDeviceID, int32 InMaxSources, int32 I
 		ModulationInterface = ModulationPluginFactory->CreateNewModulationPlugin(this);
 
 		//Set up initialization parameters for system level effect plugins:
-		FAudioPluginInitializationParams PluginInitializationParams;
 		PluginInitializationParams.SampleRate = SampleRate;
-			PluginInitializationParams.NumSources = GetMaxSources();
+		PluginInitializationParams.NumSources = GetMaxSources();
 		PluginInitializationParams.BufferLength = PlatformSettings.CallbackBufferFrameSize;
 		PluginInitializationParams.AudioDevicePtr = this;
 		ModulationInterface->Initialize(PluginInitializationParams);
@@ -651,6 +721,66 @@ float FAudioDevice::GetLowPassFilterResonance() const
 	return 0.9f;
 }
 
+
+bool FAudioDevice::SetCurrentSpatializationPlugin(FName PluginName)
+{
+	if(PluginName == CurrentSpatializationPluginInterfaceName)
+	{
+		return true; // no need to do anything.
+	}
+
+	StopAllSounds(false);
+
+	for (FAudioSpatializationInterfaceInfo& Plugin : SpatializationInterfaces)
+	{
+		if(Plugin.PluginName == PluginName)
+		{
+			CurrentSpatializationPluginInterfaceName = PluginName;
+			CurrentSpatializationInterfaceInfoPtr = &Plugin;
+
+			if(!Plugin.bIsInitialized && Plugin.SpatializationPlugin)
+			{
+				Plugin.SpatializationPlugin->Initialize(PluginInitializationParams);
+				Plugin.bIsInitialized = true;
+			}
+
+			return true;
+		}
+	}
+
+	return false;
+}
+
+TArray<FName> FAudioDevice::GetAvailableSpatializationPluginNames() const
+{
+	TArray<FName> SpatPluginNames;
+	for(const auto& Plugin : SpatializationInterfaces)
+	{
+		SpatPluginNames.Add(Plugin.PluginName);
+	}
+
+	return SpatPluginNames;
+}
+
+FAudioDevice::FAudioSpatializationInterfaceInfo FAudioDevice::GetCurrentSpatializationPluginInterfaceInfo()
+{
+	// See if we need to update the ptr to the current info
+	if(!CurrentSpatializationInterfaceInfoPtr || CurrentSpatializationInterfaceInfoPtr->PluginName != CurrentSpatializationPluginInterfaceName)
+	{
+		CurrentSpatializationInterfaceInfoPtr = nullptr;
+		for (FAudioSpatializationInterfaceInfo& Plugin : SpatializationInterfaces)
+		{
+			if(Plugin.PluginName == CurrentSpatializationPluginInterfaceName)
+			{
+				CurrentSpatializationInterfaceInfoPtr = &Plugin;
+				break;
+			}
+		}
+	}
+
+	return *CurrentSpatializationInterfaceInfoPtr;
+}
+
 void FAudioDevice::PrecacheStartupSounds()
 {
 	// Iterate over all already loaded sounds and precache them. This relies on Super::Init in derived classes to be called last.
@@ -815,13 +945,21 @@ void FAudioDevice::Teardown()
 
 	LLM_SCOPE(ELLMTag::AudioMixerPlugins);
 
-	if (SpatializationPluginInterface.IsValid())
+	// shutdown our active audio plugins
+
+	// (spatial)
+	for(const auto& Plugin : SpatializationInterfaces)
 	{
-		SpatializationPluginInterface->Shutdown();
-		SpatializationPluginInterface.Reset();
-		bSpatializationInterfaceEnabled = false;
+		if(TAudioSpatializationPtr SpatialPluginPtr = Plugin.SpatializationPlugin)
+		{
+			SpatialPluginPtr->Shutdown();
+		}
 	}
 
+	SpatializationInterfaces.Empty();
+	bSpatializationInterfaceEnabled = false;
+
+	// (reverb)
 	if (ReverbPluginInterface.IsValid())
 	{
 		ReverbPluginInterface->Shutdown();
@@ -829,12 +967,14 @@ void FAudioDevice::Teardown()
 		bReverbInterfaceEnabled = false;
 	}
 
+	// (source data override)
 	if (SourceDataOverridePluginInterface.IsValid())
 	{
 		SourceDataOverridePluginInterface.Reset();
 		bSourceDataOverrideInterfaceEnabled = false;
 	}
 
+	// (Occlusion)
 	if (OcclusionInterface.IsValid())
 	{
 		OcclusionInterface->Shutdown();
@@ -842,6 +982,7 @@ void FAudioDevice::Teardown()
 		bOcclusionInterfaceEnabled = false;
 	}
 
+	// (modulation)
 	ModulationInterface.Reset();
 	bModulationInterfaceEnabled = false;
 
@@ -7119,6 +7260,24 @@ void FAudioDevice::SetAudioStateProperty(const FName& PropertyName, const FName&
 	Property = PropertyValue;
 }
 
+
+FAudioDevice::FAudioSpatializationInterfaceInfo::FAudioSpatializationInterfaceInfo(FName InPluginName, FAudioDevice* InAudioDevice, IAudioSpatializationFactory* InAudioSpatializationFactoryPtr)
+			: PluginName(InPluginName)
+{
+	if(!ensure(InAudioDevice && InAudioSpatializationFactoryPtr))
+	{
+		return;
+	}
+
+	// use the factory to create a PluginInterface
+	SpatializationPlugin = InAudioSpatializationFactoryPtr->CreateNewSpatializationPlugin(InAudioDevice);
+
+	// cache metadata from the incoming plugin interface 
+	bSpatializationIsExternalSend = InAudioSpatializationFactoryPtr->IsExternalSend();
+	MaxChannelsSupportedBySpatializationPlugin = InAudioSpatializationFactoryPtr->GetMaxSupportedChannels();
+}
+
+
 bool FAudioDevice::ShouldUseAttenuation(const UWorld* World) const
 {
 	// We use attenuation settings:
diff --git a/Engine/Source/Runtime/Engine/Private/AudioPluginUtilities.cpp b/Engine/Source/Runtime/Engine/Private/AudioPluginUtilities.cpp
--- a/Engine/Source/Runtime/Engine/Private/AudioPluginUtilities.cpp
+++ b/Engine/Source/Runtime/Engine/Private/AudioPluginUtilities.cpp
@@ -46,6 +46,17 @@ FORCEINLINE const TCHAR* GetPluginConfigName(EAudioPlugin PluginType)
 /************************************************************************/
 /* Plugin Utilities                                                     */
 /************************************************************************/
+FName AudioPluginUtilities::GetDesiredSpatializationPluginName()
+{
+	return FName(GetDesiredPluginName(EAudioPlugin::SPATIALIZATION));
+}
+
+TArray<IAudioSpatializationFactory*> AudioPluginUtilities::GetSpatialPluginArray()
+{
+	IModularFeatures::FScopedLockModularFeatureList ScopedLockModularFeatureList;
+	return IModularFeatures::Get().GetModularFeatureImplementations<IAudioSpatializationFactory>(IAudioSpatializationFactory::GetModularFeatureName());
+}
+
 IAudioSpatializationFactory* AudioPluginUtilities::GetDesiredSpatializationPlugin()
 {
 	FString DesiredSpatializationPlugin = GetDesiredPluginName(EAudioPlugin::SPATIALIZATION);
diff --git a/Engine/Source/Runtime/Engine/Private/GameplayStatics.cpp b/Engine/Source/Runtime/Engine/Private/GameplayStatics.cpp
--- a/Engine/Source/Runtime/Engine/Private/GameplayStatics.cpp
+++ b/Engine/Source/Runtime/Engine/Private/GameplayStatics.cpp
@@ -1928,6 +1928,45 @@ void UGameplayStatics::PrimeSound(USoundBase* InSound)
 	}
 }
 
+TArray<FName> UGameplayStatics::GetAvailableSpatialPluginNames(const UObject* WorldContextObject)
+{
+	if(const UWorld* ThisWorld = GEngine->GetWorldFromContextObject(WorldContextObject, EGetWorldErrorMode::LogAndReturnNull))
+	{
+		if (FAudioDeviceHandle AudioDevice = ThisWorld->GetAudioDevice())
+		{
+			return AudioDevice->GetAvailableSpatializationPluginNames();
+		}
+	}
+
+	return {};
+}
+
+FName UGameplayStatics::GetActiveSpatialPluginName(const UObject* WorldContextObject)
+{
+	if(const UWorld* ThisWorld = GEngine->GetWorldFromContextObject(WorldContextObject, EGetWorldErrorMode::LogAndReturnNull))
+	{
+		if (FAudioDeviceHandle AudioDevice = ThisWorld->GetAudioDevice())
+		{
+			return AudioDevice->GetCurrentSpatializationPluginInterfaceInfo().PluginName;
+		}
+	}
+
+	return {};
+}
+
+bool UGameplayStatics::SetActiveSpatialPluginByName(const UObject* WorldContextObject, FName InPluginName)
+{
+	if(const UWorld* ThisWorld = GEngine->GetWorldFromContextObject(WorldContextObject, EGetWorldErrorMode::LogAndReturnNull))
+	{
+		if (FAudioDeviceHandle AudioDevice = ThisWorld->GetAudioDevice())
+		{
+			return AudioDevice->SetCurrentSpatializationPlugin(InPluginName);
+		}
+	}
+
+	return {};
+}
+
 void UGameplayStatics::PrimeAllSoundsInSoundClass(class USoundClass* InSoundClass)
 {
 	for (TObjectIterator<USoundWave> Itr; Itr; ++Itr)
diff --git a/Engine/Source/Runtime/Engine/Public/AudioDevice.h b/Engine/Source/Runtime/Engine/Public/AudioDevice.h
--- a/Engine/Source/Runtime/Engine/Public/AudioDevice.h
+++ b/Engine/Source/Runtime/Engine/Public/AudioDevice.h
@@ -515,9 +515,9 @@ public:
 	 */
 	FAudioDevice();
 
-	virtual ~FAudioDevice()
-	{
-	}
+	PRAGMA_DISABLE_DEPRECATION_WARNINGS // supress deprecation warning in default dtor
+	virtual ~FAudioDevice() = default;
+	PRAGMA_ENABLE_DEPRECATION_WARNINGS
 
 	/** Returns an array of available audio devices names for the platform */
 	virtual void GetAudioDeviceList(TArray<FString>& OutAudioDeviceNames) const
@@ -1329,12 +1329,22 @@ public:
 		return bSpatializationInterfaceEnabled;
 	}
 
+	/** Set the current Spatialization plugin (by name) */
+	bool SetCurrentSpatializationPlugin(FName PluginName);
+
+	/** Get the current list of available spatialization plugins */
+	TArray<FName> GetAvailableSpatializationPluginNames() const;
+
 	/** Return the spatialization plugin interface. */
 	TAudioSpatializationPtr GetSpatializationPluginInterface()
 	{
-		return SpatializationPluginInterface;
+		return GetCurrentSpatializationPluginInterfaceInfo().SpatializationPlugin;
 	}
 
+	/** Return the spatialization plugin interface info (requested by name). */
+	struct FAudioSpatializationInterfaceInfo;
+	FAudioSpatializationInterfaceInfo GetCurrentSpatializationPluginInterfaceInfo();
+
 	/** Whether or not there's a modulation plugin enabled. */
 	bool IsModulationPluginEnabled() const
 	{
@@ -2003,9 +2013,32 @@ public:
 	/** The handle for this audio device used in the audio device manager. */
 	Audio::FDeviceId DeviceID;
 
+	struct FAudioSpatializationInterfaceInfo
+	{
+		// ctors
+		FAudioSpatializationInterfaceInfo() = default;
+		FAudioSpatializationInterfaceInfo(FName InPluginName, FAudioDevice* InAudioDevice, IAudioSpatializationFactory* InAudioSpatializationFactoryPtr);
+
+		FName PluginName;
+		TAudioSpatializationPtr SpatializationPlugin = nullptr;
+		int32 MaxChannelsSupportedBySpatializationPlugin = 1;
+		uint8 bSpatializationIsExternalSend:1;
+		uint8 bIsInitialized:1;
+	};
+
+	UE_DEPRECATED(5.1, "Do not access this member directly, it is not used. Call GetSpatializationPluginInterface() instead.")
+	TAudioSpatializationPtr SpatializationPluginInterface = nullptr;
+
+protected:
 	/** 3rd party audio spatialization interface. */
-	TAudioSpatializationPtr SpatializationPluginInterface;
+	FName CurrentSpatializationPluginInterfaceName;
+	TArray<FAudioSpatializationInterfaceInfo> SpatializationInterfaces;
+	FAudioSpatializationInterfaceInfo* CurrentSpatializationInterfaceInfoPtr = nullptr;
+
+	/** Cached parameters passed to the initialization of various audio plugins */
+	FAudioPluginInitializationParams PluginInitializationParams;
 
+public:
 	/** 3rd party source data override interface. */
 	TAudioSourceDataOverridePtr SourceDataOverridePluginInterface;
 
@@ -2145,14 +2178,21 @@ public:
 	/** Whether or not the audio mixer module is being used by this device. */
 	uint8 bAudioMixerModuleLoaded : 1;
 
-	/** Whether of not various audio plugin interfaces are external sends. */
-	uint8 bSpatializationIsExternalSend:1;
+	/** Whether or not various audio plugin interfaces are external sends. */
 	uint8 bOcclusionIsExternalSend:1;
 	uint8 bReverbIsExternalSend:1;
 
+// deprecate these as they have been moved to the info struct.
+
+	UE_DEPRECATED(5.1, "This member is no longer  in use. Use the return value of GetCurrentSpatializationPluginInterfaceInfo()")
+	uint8 bSpatializationIsExternalSend:1;
+
 	/** Max amount of channels a source can be to be spatialized by our active spatialization plugin. */
+	UE_DEPRECATED(5.1, "This member is no longer  in use. Use the return value of GetCurrentSpatializationPluginInterfaceInfo()")
 	int32 MaxChannelsSupportedBySpatializationPlugin;
 
+
+
 private:
 	/** True once the startup sounds have been precached */
 	uint8 bStartupSoundsPreCached:1;
diff --git a/Engine/Source/Runtime/Engine/Public/AudioPluginUtilities.h b/Engine/Source/Runtime/Engine/Public/AudioPluginUtilities.h
--- a/Engine/Source/Runtime/Engine/Public/AudioPluginUtilities.h
+++ b/Engine/Source/Runtime/Engine/Public/AudioPluginUtilities.h
@@ -16,6 +16,8 @@ struct ENGINE_API AudioPluginUtilities
 	 * 
 	 * if no matching plugin is found, nullptr is returned.
 	 */
+	static FName GetDesiredSpatializationPluginName();
+	static TArray<IAudioSpatializationFactory*> GetSpatialPluginArray();
 	static IAudioSpatializationFactory* GetDesiredSpatializationPlugin();
 	static IAudioSourceDataOverrideFactory* GetDesiredSourceDataOverridePlugin();
 	static IAudioReverbFactory* GetDesiredReverbPlugin();
