When creating a MID material for a decal and a MID has already been set, then create the MID based on the current MID parent.
We can't create a MID from another MID, and we don't store the original parent in the decal component itself.

diff --git a/Engine/Source/Runtime/Engine/Private/Components/DecalComponent.cpp b/Engine/Source/Runtime/Engine/Private/Components/DecalComponent.cpp
--- a/Engine/Source/Runtime/Engine/Private/Components/DecalComponent.cpp
+++ b/Engine/Source/Runtime/Engine/Private/Components/DecalComponent.cpp
@@ -263,13 +263,21 @@ class UMaterialInterface* UDecalComponent::GetDecalMaterial() const
 
 class UMaterialInstanceDynamic* UDecalComponent::CreateDynamicMaterialInstance()
 {
+	UMaterialInterface* CurrentMaterial = DecalMaterial;
+	
+	// If we already set a MID, then we need to create based on its parent.
+	if (UMaterialInstanceDynamic* CurrentMaterialMID = Cast<UMaterialInstanceDynamic>(CurrentMaterial))
+	{
+		CurrentMaterial = CurrentMaterialMID->Parent;
+	}
+
 	// Create the MID
-	UMaterialInstanceDynamic* Instance = UMaterialInstanceDynamic::Create(DecalMaterial, this);
+	UMaterialInstanceDynamic* NewMaterialInstance = UMaterialInstanceDynamic::Create(CurrentMaterial, this);
 
-	// Assign it, once parent is set
-	SetDecalMaterial(Instance);
+	// Assign the MID
+	SetDecalMaterial(NewMaterialInstance);
 
-	return Instance;
+	return NewMaterialInstance;
 }
 
 void UDecalComponent::GetUsedMaterials( TArray<UMaterialInterface*>& OutMaterials, bool bGetDebugMaterials ) const
