Undo //UE5/Main/Engine/Source/Runtime/Engine/Public/... changelist 21425691

It was causing UE.TargetAutomation(RunTest=Rendering) to fail.

diff --git a/Engine/Source/Runtime/Engine/Public/PrimitiveUniformShaderParameters.h b/Engine/Source/Runtime/Engine/Public/PrimitiveUniformShaderParameters.h
--- a/Engine/Source/Runtime/Engine/Public/PrimitiveUniformShaderParameters.h
+++ b/Engine/Source/Runtime/Engine/Public/PrimitiveUniformShaderParameters.h
@@ -206,9 +206,10 @@ public:
 	inline FPrimitiveUniformShaderParametersBuilder& WorldBounds(const FBoxSphereBounds& InWorldBounds)
 	{
 		AbsoluteObjectWorldPosition = InWorldBounds.Origin;
-		Parameters.ObjectBoundsX	= InWorldBounds.BoxExtent.X;
-		Parameters.ObjectBoundsY	= InWorldBounds.BoxExtent.Y;
-		Parameters.ObjectBoundsZ	= InWorldBounds.BoxExtent.Z;
+		ObjectRadius = InWorldBounds.SphereRadius;
+		Parameters.ObjectBoundsX = InWorldBounds.BoxExtent.X;
+		Parameters.ObjectBoundsY = InWorldBounds.BoxExtent.Y;
+		Parameters.ObjectBoundsZ = InWorldBounds.BoxExtent.Z;
 		return *this;
 	}
 
@@ -244,7 +245,6 @@ public:
 		bHasInstanceLocalBounds = true;
 		Parameters.InstanceLocalBoundsCenter = InInstanceLocalBounds.GetCenter();
 		Parameters.InstanceLocalBoundsExtent = InInstanceLocalBounds.GetExtent();
-		ObjectRadius						 = InInstanceLocalBounds.GetExtent().Length();
 		return *this;
 	}
 
