Remove unnecessary warning from GetBuildRequiredMemoryEstimate() which is called for VT textures assets even when we don't cook them due to project settings.

diff --git a/Engine/Source/Runtime/Engine/Private/TextureDerivedDataTask.cpp b/Engine/Source/Runtime/Engine/Private/TextureDerivedDataTask.cpp
--- a/Engine/Source/Runtime/Engine/Private/TextureDerivedDataTask.cpp
+++ b/Engine/Source/Runtime/Engine/Private/TextureDerivedDataTask.cpp
@@ -1001,9 +1001,7 @@ static int64 GetBuildRequiredMemoryEstimate(UTexture* InTexture,
 		
 		if ( Source.GetNumBlocks() != 1 || Source.GetNumLayers() != 1 )
 		{
-			UE_LOG(LogTexture, Warning, TEXT("Texture has no source data: %s"), *(InTexture->GetPathName()));
-
-			return -1; /* Unknown */
+			return -1; /* Requires VT enabled. */
 		}
 
 		// Compute the memory it should take to uncompress the bulkdata in memory
