Sequencer: Fix SequencerTrackClass path

Actor - /Script/MovieSceneTracks.MovieSceneVisibilityTrack
SceneComponent - /Script/MovieSceneTracks.MovieSceneVisibilityTrack

diff --git a/Engine/Source/Runtime/Engine/Classes/Components/SceneComponent.h b/Engine/Source/Runtime/Engine/Classes/Components/SceneComponent.h
--- a/Engine/Source/Runtime/Engine/Classes/Components/SceneComponent.h
+++ b/Engine/Source/Runtime/Engine/Classes/Components/SceneComponent.h
@@ -210,7 +210,7 @@ private:
 
 public:
 	/** Whether to hide the primitive in game, if the primitive is Visible. */
-	UPROPERTY(Interp, EditAnywhere, BlueprintReadOnly, Category=Rendering, meta=(SequencerTrackClass = "MovieSceneVisibilityTrack"))
+	UPROPERTY(Interp, EditAnywhere, BlueprintReadOnly, Category=Rendering, meta=(SequencerTrackClass = "/Script/MovieSceneTracks.MovieSceneVisibilityTrack"))
 	uint8 bHiddenInGame:1;
 
 	/** If true, a change in the bounds of the component will call trigger a streaming data rebuild */
diff --git a/Engine/Source/Runtime/Engine/Classes/GameFramework/Actor.h b/Engine/Source/Runtime/Engine/Classes/GameFramework/Actor.h
--- a/Engine/Source/Runtime/Engine/Classes/GameFramework/Actor.h
+++ b/Engine/Source/Runtime/Engine/Classes/GameFramework/Actor.h
@@ -239,7 +239,7 @@ private:
 	 * Allows us to only see this Actor in the Editor, and not in the actual game.
 	 * @see SetActorHiddenInGame()
 	 */
-	UPROPERTY(Interp, EditAnywhere, Category=Rendering, BlueprintReadOnly, Replicated, meta=(AllowPrivateAccess="true", DisplayName="Actor Hidden In Game", SequencerTrackClass="MovieSceneVisibilityTrack"))
+	UPROPERTY(Interp, EditAnywhere, Category=Rendering, BlueprintReadOnly, Replicated, meta=(AllowPrivateAccess="true", DisplayName="Actor Hidden In Game", SequencerTrackClass="/Script/MovieSceneTracks.MovieSceneVisibilityTrack"))
 	uint8 bHidden:1;
 
 	UPROPERTY(Replicated)
