- Fix incorrect calculation for GPU scene payload sizes

diff --git a/Engine/Source/Runtime/Engine/Public/InstanceUniformShaderParameters.h b/Engine/Source/Runtime/Engine/Public/InstanceUniformShaderParameters.h
--- a/Engine/Source/Runtime/Engine/Public/InstanceUniformShaderParameters.h
+++ b/Engine/Source/Runtime/Engine/Public/InstanceUniformShaderParameters.h
@@ -36,8 +36,6 @@
 
 #define INVALID_LAST_UPDATE_FRAME 0xFFFFFFFFu
 
-#define INSTANCE_SCENE_DATA_COMPRESSED_TRANSFORMS	1
-
 // TODO: Rename to FInstanceSceneData
 struct FPrimitiveInstance
 {
diff --git a/Engine/Source/Runtime/Engine/Public/MeshBatch.h b/Engine/Source/Runtime/Engine/Public/MeshBatch.h
--- a/Engine/Source/Runtime/Engine/Public/MeshBatch.h
+++ b/Engine/Source/Runtime/Engine/Public/MeshBatch.h
@@ -117,11 +117,14 @@ struct FMeshBatchDynamicPrimitiveData
 
 		if (PayloadDataFlags & INSTANCE_SCENE_DATA_FLAG_HAS_DYNAMIC_DATA)
 		{
-#if INSTANCE_SCENE_DATA_COMPRESSED_TRANSFORMS
-			Total += 2;
-#else
-			Total += 3;
-#endif
+			if (FDataDrivenShaderPlatformInfo::GetSupportSceneDataCompressedTransforms(GMaxRHIShaderPlatform))
+			{
+				Total += 2;
+			}
+			else
+			{
+				Total += 3;
+			}
 		}
 
 		if (PayloadDataFlags & INSTANCE_SCENE_DATA_FLAG_HAS_LIGHTSHADOW_UV_BIAS)
