Nanite override materials implementation - allows Nanite specific materials to be used on a mesh instead of the regular mesh materials. Enabled by default based on cvar r.Nanite.MaterialOverrides. The same static mesh asset is used, with Nanite either building off of LOD0 (and decimating LOD0 to be used as a fallback mesh for non-Nanite), or using the hi-res mesh LOD path to add Nanite data as a sidecar to the asset without touching LOD0...N

diff --git a/Engine/Source/Runtime/Core/Private/Misc/MapErrors.cpp b/Engine/Source/Runtime/Core/Private/Misc/MapErrors.cpp
--- a/Engine/Source/Runtime/Core/Private/Misc/MapErrors.cpp
+++ b/Engine/Source/Runtime/Core/Private/Misc/MapErrors.cpp
@@ -55,6 +55,7 @@ FLazyName FMapErrors::StaticMeshNull(TEXT("StaticMeshNull"));
 FLazyName FMapErrors::StaticMeshComponent(TEXT("StaticMeshComponent"));
 FLazyName FMapErrors::SimpleCollisionButNonUniformScale(TEXT("SimpleCollisionButNonUniformScale"));
 FLazyName FMapErrors::MoreMaterialsThanReferenced(TEXT("MoreMaterialsThanReferenced"));
+FLazyName FMapErrors::MoreNaniteMaterialsThanReferenced(TEXT("MoreNaniteMaterialsThanReferenced"));
 FLazyName FMapErrors::ElementsWithZeroTriangles(TEXT("ElementsWithZeroTriangles"));
 FLazyName FMapErrors::LevelStreamingVolume(TEXT("LevelStreamingVolume"));
 FLazyName FMapErrors::NoLevelsAssociated(TEXT("NoLevelsAssociated"));
diff --git a/Engine/Source/Runtime/Core/Public/Misc/MapErrors.h b/Engine/Source/Runtime/Core/Public/Misc/MapErrors.h
--- a/Engine/Source/Runtime/Core/Public/Misc/MapErrors.h
+++ b/Engine/Source/Runtime/Core/Public/Misc/MapErrors.h
@@ -194,6 +194,9 @@ struct CORE_API FMapErrors
 	/** More overridden materials {Count} on static mesh component than are referenced {Count} in source mesh {StaticMesh}: **TODO** */
     static FLazyName MoreMaterialsThanReferenced;
 
+	/** More Nanite overridden materials {Count} on static mesh component than are referenced {Count} in source mesh {StaticMesh}: **TODO** */
+	static FLazyName MoreNaniteMaterialsThanReferenced;
+
 	 /** {Count} element(s) with zero triangles in static mesh {StaticMesh}: **TODO** */
     static FLazyName ElementsWithZeroTriangles;
   
diff --git a/Engine/Source/Runtime/Engine/Classes/Components/MeshComponent.h b/Engine/Source/Runtime/Engine/Classes/Components/MeshComponent.h
--- a/Engine/Source/Runtime/Engine/Classes/Components/MeshComponent.h
+++ b/Engine/Source/Runtime/Engine/Classes/Components/MeshComponent.h
@@ -26,6 +26,10 @@ public:
 	/** Per-Component material overrides.  These must NOT be set directly or a race condition can occur between GC and the rendering thread. */
 	UPROPERTY(EditAnywhere, AdvancedDisplay, Category=Rendering, Meta=(ToolTip="Material overrides."))
 	TArray<TObjectPtr<class UMaterialInterface>> OverrideMaterials;
+
+	/** Per-Component Nanite specific material overrides.  These must NOT be set directly or a race condition can occur between GC and the rendering thread. */
+	UPROPERTY(EditAnywhere, AdvancedDisplay, Category = Rendering, Meta = (ToolTip = "Nanite Material overrides."))
+	TArray<TObjectPtr<class UMaterialInterface>> NaniteOverrideMaterials;
 	
 	UFUNCTION(BlueprintCallable, Category="Rendering|Material")
 	virtual TArray<class UMaterialInterface*> GetMaterials() const;
@@ -39,7 +43,10 @@ public:
 	UFUNCTION(BlueprintCallable, Category = "Rendering|Material")
 	virtual bool IsMaterialSlotNameValid(FName MaterialSlotName) const;
 
-	/** Returns override Materials count */
+	/** Determines if NaniteOverrideMaterials array is used instead of OverrideMaterials */
+	virtual bool UseNaniteOverrideMaterials() const { return false; }
+
+	/** Returns override materials count */
 	virtual int32 GetNumOverrideMaterials() const;
 
 	/** Translucent material to blend on top of this mesh. Mesh will be rendered twice - once with a base material and once with overlay material */
diff --git a/Engine/Source/Runtime/Engine/Classes/Components/StaticMeshComponent.h b/Engine/Source/Runtime/Engine/Classes/Components/StaticMeshComponent.h
--- a/Engine/Source/Runtime/Engine/Classes/Components/StaticMeshComponent.h
+++ b/Engine/Source/Runtime/Engine/Classes/Components/StaticMeshComponent.h
@@ -439,6 +439,9 @@ public:
 	 */
 	virtual bool HasValidNaniteData() const;
 
+	/** Determines if NaniteOverrideMaterials array is used instead of OverrideMaterials */
+	virtual bool UseNaniteOverrideMaterials() const override;
+
 	UFUNCTION(BlueprintCallable, Category="Rendering|LOD")
 	void SetForcedLodModel(int32 NewForcedLodModel);
 
diff --git a/Engine/Source/Runtime/Engine/Private/Components/MeshComponent.cpp b/Engine/Source/Runtime/Engine/Private/Components/MeshComponent.cpp
--- a/Engine/Source/Runtime/Engine/Private/Components/MeshComponent.cpp
+++ b/Engine/Source/Runtime/Engine/Private/Components/MeshComponent.cpp
@@ -26,7 +26,11 @@ UMeshComponent::UMeshComponent(const FObjectInitializer& ObjectInitializer)
 
 UMaterialInterface* UMeshComponent::GetMaterial(int32 ElementIndex) const
 {
-	if (OverrideMaterials.IsValidIndex(ElementIndex))
+	if (NaniteOverrideMaterials.IsValidIndex(ElementIndex) && UseNaniteOverrideMaterials())
+	{
+		return NaniteOverrideMaterials[ElementIndex];
+	}
+	else if (OverrideMaterials.IsValidIndex(ElementIndex))
 	{
 		return OverrideMaterials[ElementIndex];
 	}
@@ -40,23 +44,25 @@ void UMeshComponent::SetMaterial(int32 ElementIndex, UMaterialInterface* Materia
 {
 	if (ElementIndex >= 0)
 	{
-		if (OverrideMaterials.IsValidIndex(ElementIndex) && (OverrideMaterials[ElementIndex] == Material))
+		TArray<TObjectPtr<UMaterialInterface>>& OverrideMaterialsRef = UseNaniteOverrideMaterials() ? NaniteOverrideMaterials : OverrideMaterials;
+
+		if (OverrideMaterialsRef.IsValidIndex(ElementIndex) && (OverrideMaterialsRef[ElementIndex] == Material))
 		{
 			// Do nothing, the material is already set
 		}
 		else
 		{
 			// Grow the array if the new index is too large
-			if (OverrideMaterials.Num() <= ElementIndex)
+			if (OverrideMaterialsRef.Num() <= ElementIndex)
 			{
-				OverrideMaterials.AddZeroed(ElementIndex + 1 - OverrideMaterials.Num());
+				OverrideMaterialsRef.AddZeroed(ElementIndex + 1 - OverrideMaterialsRef.Num());
 			}
 			
 			// Check if we are setting a dynamic instance of the original material, or replacing a nullptr material  (if not we should dirty the material parameter name cache)
 			if (Material != nullptr)
 			{
 				UMaterialInstanceDynamic* DynamicMaterial = Cast<UMaterialInstanceDynamic>(Material);
-				if (!((DynamicMaterial != nullptr && DynamicMaterial->Parent == OverrideMaterials[ElementIndex]) || OverrideMaterials[ElementIndex] == nullptr))
+				if (!((DynamicMaterial != nullptr && DynamicMaterial->Parent == OverrideMaterialsRef[ElementIndex]) || OverrideMaterialsRef[ElementIndex] == nullptr))
 				{
 					// Mark cached material parameter names dirty
 					MarkCachedMaterialParameterNameIndicesDirty();
@@ -64,7 +70,7 @@ void UMeshComponent::SetMaterial(int32 ElementIndex, UMaterialInterface* Materia
 			}	
 
 			// Set the material and invalidate things
-			OverrideMaterials[ElementIndex] = Material;
+			OverrideMaterialsRef[ElementIndex] = Material;
 			MarkRenderStateDirty();
 			// If MarkRenderStateDirty didn't notify the streamer, do it now
 			if (!bIgnoreStreamingManagerUpdate && OwnerLevelHasRegisteredStaticComponentsInStreamingManager(GetOwner()))
@@ -130,6 +136,11 @@ FMaterialRelevance UMeshComponent::GetMaterialRelevance(ERHIFeatureLevel::Type I
 
 int32 UMeshComponent::GetNumOverrideMaterials() const
 {
+	if (UseNaniteOverrideMaterials())
+	{
+		return FMath::Max(NaniteOverrideMaterials.Num(), OverrideMaterials.Num());
+	}
+
 	return OverrideMaterials.Num();
 }
 
@@ -140,7 +151,8 @@ void UMeshComponent::PostEditChangeChainProperty(FPropertyChangedChainEvent& Pro
 
 	if (PropertyChangedEvent.Property != nullptr)
 	{
-		if (PropertyChangedEvent.Property->GetFName() == GET_MEMBER_NAME_STRING_CHECKED(UMeshComponent, OverrideMaterials))
+		if (PropertyChangedEvent.Property->GetFName() == GET_MEMBER_NAME_STRING_CHECKED(UMeshComponent, OverrideMaterials) ||
+			PropertyChangedEvent.Property->GetFName() == GET_MEMBER_NAME_STRING_CHECKED(UMeshComponent, NaniteOverrideMaterials))
 		{
 			CleanUpOverrideMaterials();
 		}
@@ -151,15 +163,24 @@ void UMeshComponent::CleanUpOverrideMaterials()
 {
 	bool bUpdated = false;
 
-	//We have to remove material override Ids that are bigger then the material list
-	if (GetNumOverrideMaterials() > GetNumMaterials())
+	// We have to remove material override Ids that are bigger then the material list
+	if (OverrideMaterials.Num() > GetNumMaterials())
 	{
 		//Remove the override material id that are superior to the static mesh materials number
-		int32 RemoveCount = GetNumOverrideMaterials() - GetNumMaterials();
+		int32 RemoveCount = OverrideMaterials.Num() - GetNumMaterials();
 		OverrideMaterials.RemoveAt(GetNumMaterials(), RemoveCount);
 		bUpdated = true;
 	}
 
+	// We have to remove Nanite material override Ids that are bigger then the material list
+	if (NaniteOverrideMaterials.Num() > GetNumMaterials())
+	{
+		// Remove the override material id that are superior to the static mesh materials number
+		int32 RemoveCount = NaniteOverrideMaterials.Num() - GetNumMaterials();
+		NaniteOverrideMaterials.RemoveAt(GetNumMaterials(), RemoveCount);
+		bUpdated = true;
+	}
+
 	if (bUpdated)
 	{
 		MarkRenderStateDirty();
@@ -169,9 +190,10 @@ void UMeshComponent::CleanUpOverrideMaterials()
 
 void UMeshComponent::EmptyOverrideMaterials()
 {
-	if (OverrideMaterials.Num())
+	if (OverrideMaterials.Num() || NaniteOverrideMaterials.Num())
 	{
 		OverrideMaterials.Reset();
+		NaniteOverrideMaterials.Reset();
 		MarkRenderStateDirty();
 	}
 }
@@ -292,19 +314,19 @@ TArray<class UMaterialInterface*> UMeshComponent::GetMaterials() const
 
 int32 UMeshComponent::GetMaterialIndex(FName MaterialSlotName) const
 {
-	//This function should be override
-	return -1;
+	// This function should be overridden
+	return INDEX_NONE;
 }
 
 TArray<FName> UMeshComponent::GetMaterialSlotNames() const
 {
-	//This function should be override
+	// This function should be overridden
 	return TArray<FName>();
 }
 
 bool UMeshComponent::IsMaterialSlotNameValid(FName MaterialSlotName) const
 {
-	//This function should be override
+	// This function should be overridden
 	return false;
 }
 
@@ -312,7 +334,7 @@ void UMeshComponent::SetScalarParameterValueOnMaterials(const FName ParameterNam
 {
 	if (!bEnableMaterialParameterCaching)
 	{
-		const TArray<UMaterialInterface*> MaterialInterfaces = GetMaterials();		
+		const TArray<UMaterialInterface*> MaterialInterfaces = GetMaterials();
 		for (int32 MaterialIndex = 0; MaterialIndex < MaterialInterfaces.Num(); ++MaterialIndex)
 		{
 			UMaterialInterface* MaterialInterface = MaterialInterfaces[MaterialIndex];
diff --git a/Engine/Source/Runtime/Engine/Private/Components/StaticMeshComponent.cpp b/Engine/Source/Runtime/Engine/Private/Components/StaticMeshComponent.cpp
--- a/Engine/Source/Runtime/Engine/Private/Components/StaticMeshComponent.cpp
+++ b/Engine/Source/Runtime/Engine/Private/Components/StaticMeshComponent.cpp
@@ -55,6 +55,7 @@
 DECLARE_MEMORY_STAT( TEXT( "StaticMesh VxColor Inst Mem" ), STAT_InstVertexColorMemory, STATGROUP_MemoryStaticMesh );
 DECLARE_MEMORY_STAT( TEXT( "StaticMesh PreCulled Index Memory" ), STAT_StaticMeshPreCulledIndexMemory, STATGROUP_MemoryStaticMesh );
 
+extern int32 GEnableNaniteMaterialOverrides;
 
 FStaticMeshComponentInstanceData::FStaticMeshComponentInstanceData(const UStaticMeshComponent* SourceComponent)
 	: FPrimitiveComponentInstanceData(SourceComponent)
@@ -504,6 +505,19 @@ void UStaticMeshComponent::CheckForErrors()
 				->AddToken(FTextToken::Create(FText::Format(LOCTEXT( "MapCheck_Message_MoreMaterialsThanReferenced", "More overridden materials ({OverridenCount}) on static mesh component than are referenced ({ReferencedCount}) in source mesh '{MeshName}'" ), Arguments ) ))
 				->AddToken(FMapErrorToken::Create(FMapErrors::MoreMaterialsThanReferenced));
 		}
+
+		if (NaniteOverrideMaterials.Num() > GetStaticMesh()->GetStaticMaterials().Num())
+		{
+			FFormatNamedArguments Arguments;
+			Arguments.Add(TEXT("OverridenCount"), NaniteOverrideMaterials.Num());
+			Arguments.Add(TEXT("ReferencedCount"), GetStaticMesh()->GetStaticMaterials().Num());
+			Arguments.Add(TEXT("MeshName"), FText::FromString(GetStaticMesh()->GetName()));
+			FMessageLog("MapCheck").Warning()
+				->AddToken(FUObjectToken::Create(Owner))
+				->AddToken(FTextToken::Create(FText::Format(LOCTEXT( "MapCheck_Message_MoreNaniteMaterialsThanReferenced", "More Nanite overridden materials ({OverridenCount}) on static mesh component than are referenced ({ReferencedCount}) in source mesh '{MeshName}'" ), Arguments ) ))
+				->AddToken(FMapErrorToken::Create(FMapErrors::MoreNaniteMaterialsThanReferenced));
+		}
+
 		if (ZeroTriangleElements > 0)
 		{
 			FFormatNamedArguments Arguments;
@@ -1643,7 +1657,7 @@ void UStaticMeshComponent::UpdatePreCulledData(int32 LODIndex, const TArray<uint
 
 bool UStaticMeshComponent::ShouldCreateNaniteProxy() const
 {
-	if (bDisallowNanite)
+	if (bDisallowNanite || GetScene() == nullptr)
 	{
 		// Regardless of the static mesh asset supporting Nanite, this component does not want Nanite to be used
 		return false;
@@ -1848,17 +1862,18 @@ void UStaticMeshComponent::PostEditChangeProperty(FPropertyChangedEvent& Propert
 			// Broadcast that the static mesh has changed
 			OnStaticMeshChangedEvent.Broadcast(this);
 
-			// If the staticmesh changed, then the component needs a texture streaming rebuild.
+			// If the static mesh changed, then the component needs a texture streaming rebuild.
 			StreamingTextureData.Empty();
 			
-			if (OverrideMaterials.Num())
+			if (OverrideMaterials.Num() || NaniteOverrideMaterials.Num())
 			{
 				// Static mesh was switched so we should clean up the override materials
 				CleanUpOverrideMaterials();
 			}
 		}
 
-		if (PropertyThatChanged->GetFName() == GET_MEMBER_NAME_CHECKED(UStaticMeshComponent, OverrideMaterials))
+		if (PropertyThatChanged->GetFName() == GET_MEMBER_NAME_CHECKED(UStaticMeshComponent, OverrideMaterials) ||
+			PropertyThatChanged->GetFName() == GET_MEMBER_NAME_CHECKED(UStaticMeshComponent, NaniteOverrideMaterials))
 		{
 			// If the owning actor is part of a cluster flag it as dirty
 			IHierarchicalLODUtilitiesModule& Module = FModuleManager::LoadModuleChecked<IHierarchicalLODUtilitiesModule>("HierarchicalLODUtilities");
@@ -2181,6 +2196,12 @@ bool UStaticMeshComponent::HasValidNaniteData() const
 	return NaniteResources ? NaniteResources->PageStreamingStates.Num() > 0 : false;
 }
 
+bool UStaticMeshComponent::UseNaniteOverrideMaterials() const
+{
+	// Check for valid data on this SMC and support for Nanite material overrides
+	return ShouldCreateNaniteProxy() && GEnableNaniteMaterialOverrides != 0;
+}
+
 void UStaticMeshComponent::SetForcedLodModel(int32 NewForcedLodModel)
 {
 	if (ForcedLodModel != NewForcedLodModel)
@@ -2568,8 +2589,13 @@ bool UStaticMeshComponent::IsMaterialSlotNameValid(FName MaterialSlotName) const
 
 UMaterialInterface* UStaticMeshComponent::GetMaterial(int32 MaterialIndex) const
 {
+	// If we have Nanite specific overrides, use that
+	if (NaniteOverrideMaterials.IsValidIndex(MaterialIndex) && NaniteOverrideMaterials[MaterialIndex] && UseNaniteOverrideMaterials())
+	{
+		return NaniteOverrideMaterials[MaterialIndex];
+	}
 	// If we have a base materials array, use that
-	if(OverrideMaterials.IsValidIndex(MaterialIndex) && OverrideMaterials[MaterialIndex])
+	else if (OverrideMaterials.IsValidIndex(MaterialIndex) && OverrideMaterials[MaterialIndex])
 	{
 		return OverrideMaterials[MaterialIndex];
 	}
@@ -2624,8 +2650,19 @@ void UStaticMeshComponent::GetUsedMaterials(TArray<UMaterialInterface*>& OutMate
 #if WITH_EDITOR
 bool UStaticMeshComponent::GetMaterialPropertyPath(int32 ElementIndex, UObject*& OutOwner, FString& OutPropertyPath, FProperty*& OutProperty)
 {
-	if(OverrideMaterials.IsValidIndex(ElementIndex))
-	{				
+	if (NaniteOverrideMaterials.IsValidIndex(ElementIndex) && UseNaniteOverrideMaterials())
+	{
+		OutOwner = this;
+		OutPropertyPath = FString::Printf(TEXT("%s[%d]"), GET_MEMBER_NAME_STRING_CHECKED(UMeshComponent, NaniteOverrideMaterials), ElementIndex);
+		if (FArrayProperty* ArrayProperty = CastField<FArrayProperty>(UMeshComponent::StaticClass()->FindPropertyByName(GET_MEMBER_NAME_CHECKED(UMeshComponent, NaniteOverrideMaterials))))
+		{
+			OutProperty = ArrayProperty->Inner;
+		}
+
+		return true;
+	}
+	else if (OverrideMaterials.IsValidIndex(ElementIndex))
+	{
 		OutOwner = this;
 		OutPropertyPath = FString::Printf(TEXT("%s[%d]"), GET_MEMBER_NAME_STRING_CHECKED(UMeshComponent, OverrideMaterials), ElementIndex);
 		if (FArrayProperty* ArrayProperty = CastField<FArrayProperty>(UMeshComponent::StaticClass()->FindPropertyByName(GET_MEMBER_NAME_CHECKED(UMeshComponent, OverrideMaterials))))
@@ -2635,6 +2672,7 @@ bool UStaticMeshComponent::GetMaterialPropertyPath(int32 ElementIndex, UObject*&
 
 		return true;
 	}
+
 	if (GetStaticMesh())
 	{
 		OutOwner = GetStaticMesh();
diff --git a/Engine/Source/Runtime/Engine/Private/StaticMeshRender.cpp b/Engine/Source/Runtime/Engine/Private/StaticMeshRender.cpp
--- a/Engine/Source/Runtime/Engine/Private/StaticMeshRender.cpp
+++ b/Engine/Source/Runtime/Engine/Private/StaticMeshRender.cpp
@@ -111,6 +111,19 @@ FAutoConsoleVariableRef CVarRenderNaniteMeshes(
 	ECVF_Scalability | ECVF_RenderThreadSafe
 );
 
+// TODO: Should move this outside of SM, since Nanite can be used for multiple primitive types
+int32 GEnableNaniteMaterialOverrides = 1;
+FAutoConsoleVariableRef CVarEnableNaniteNaterialOverrides(
+	TEXT("r.Nanite.MaterialOverrides"),
+	GEnableNaniteMaterialOverrides,
+	TEXT("Enable support for Nanite specific material overrides."),
+	FConsoleVariableDelegate::CreateLambda([](IConsoleVariable* InVariable)
+	{
+		FGlobalComponentRecreateRenderStateContext Context;
+	}),
+	ECVF_Scalability | ECVF_RenderThreadSafe
+);
+
 int32 GNaniteProxyRenderMode = 0;
 FAutoConsoleVariableRef CVarNaniteProxyRenderMode(
 	TEXT("r.Nanite.ProxyRenderMode"),
diff --git a/Engine/Source/Runtime/Foliage/Private/InstancedFoliage.cpp b/Engine/Source/Runtime/Foliage/Private/InstancedFoliage.cpp
--- a/Engine/Source/Runtime/Foliage/Private/InstancedFoliage.cpp
+++ b/Engine/Source/Runtime/Foliage/Private/InstancedFoliage.cpp
@@ -1703,6 +1703,27 @@ void FFoliageStaticMesh::UpdateComponentSettings(const UFoliageType_InstancedSta
 					}
 				}
 			}
+
+			// Check Nanite override materials
+			if (Component->NaniteOverrideMaterials.Num() != FoliageType_ISM->NaniteOverrideMaterials.Num())
+			{
+				Component->NaniteOverrideMaterials = FoliageType_ISM->NaniteOverrideMaterials;
+				bNeedsMarkRenderStateDirty = true;
+				bNeedsInvalidateLightingCache = true;
+			}
+			else
+			{
+				for (int32 Index = 0; Index < FoliageType_ISM->NaniteOverrideMaterials.Num(); Index++)
+				{
+					if (Component->NaniteOverrideMaterials[Index] != FoliageType_ISM->NaniteOverrideMaterials[Index])
+					{
+						Component->NaniteOverrideMaterials = FoliageType_ISM->NaniteOverrideMaterials;
+						bNeedsMarkRenderStateDirty = true;
+						bNeedsInvalidateLightingCache = true;
+						break;
+					}
+				}
+			}
 		}
 
 		Component->BodyInstance.CopyBodyInstancePropertiesFrom(&FoliageType->BodyInstance);
diff --git a/Engine/Source/Runtime/Foliage/Public/FoliageType_InstancedStaticMesh.h b/Engine/Source/Runtime/Foliage/Public/FoliageType_InstancedStaticMesh.h
--- a/Engine/Source/Runtime/Foliage/Public/FoliageType_InstancedStaticMesh.h
+++ b/Engine/Source/Runtime/Foliage/Public/FoliageType_InstancedStaticMesh.h
@@ -19,6 +19,9 @@ class UFoliageType_InstancedStaticMesh : public UFoliageType
 
 	UPROPERTY(EditAnywhere, BlueprintReadWrite, AdvancedDisplay, Category = Mesh, Meta = (ToolTip = "Material overrides for foliage instances."))
 	TArray<TObjectPtr<class UMaterialInterface>> OverrideMaterials;
+
+	UPROPERTY(EditAnywhere, BlueprintReadWrite, AdvancedDisplay, Category = Mesh, Meta = (ToolTip = "Nanite material overrides for foliage instances."))
+	TArray<TObjectPtr<class UMaterialInterface>> NaniteOverrideMaterials;
 		
 	/** The component class to use for foliage instances. 
 	  * You can make a Blueprint subclass of FoliageInstancedStaticMeshComponent to implement custom behavior and assign that class here. */
diff --git a/Engine/Source/Runtime/Landscape/Classes/LandscapeGrassType.h b/Engine/Source/Runtime/Landscape/Classes/LandscapeGrassType.h
--- a/Engine/Source/Runtime/Landscape/Classes/LandscapeGrassType.h
+++ b/Engine/Source/Runtime/Landscape/Classes/LandscapeGrassType.h
@@ -34,6 +34,9 @@ struct FGrassVariety
 	UPROPERTY(EditAnywhere, Category=Grass, meta = (ToolTip = "Material Overrides."))
 	TArray<TObjectPtr<class UMaterialInterface>> OverrideMaterials;
 
+	UPROPERTY(EditAnywhere, Category=Grass, meta = (ToolTip = "Nanite Material Overrides."))
+	TArray<TObjectPtr<class UMaterialInterface>> NaniteOverrideMaterials;
+
 	/* Instances per 10 square meters. */
 	UPROPERTY(EditAnywhere, Category=Grass, meta = (UIMin = 0, ClampMin = 0, UIMax = 1000, ClampMax = 1000))
 	FPerPlatformFloat GrassDensity;
diff --git a/Engine/Source/Runtime/Landscape/Private/LandscapeGrass.cpp b/Engine/Source/Runtime/Landscape/Private/LandscapeGrass.cpp
--- a/Engine/Source/Runtime/Landscape/Private/LandscapeGrass.cpp
+++ b/Engine/Source/Runtime/Landscape/Private/LandscapeGrass.cpp
@@ -3035,6 +3035,7 @@ void ALandscapeProxy::UpdateGrass(const TArray<FVector>& Cameras, int32& InOutNu
 										GrassInstancedStaticMeshComponent->bCastDynamicShadow = GrassVariety.bCastDynamicShadow && !bDisableDynamicShadows;
 										GrassInstancedStaticMeshComponent->bCastContactShadow = GrassVariety.bCastContactShadow && !bDisableDynamicShadows;
 										GrassInstancedStaticMeshComponent->OverrideMaterials = GrassVariety.OverrideMaterials;
+										GrassInstancedStaticMeshComponent->NaniteOverrideMaterials = GrassVariety.NaniteOverrideMaterials;
 										GrassInstancedStaticMeshComponent->bEvaluateWorldPositionOffset = true;
 
 										const FMeshMapBuildData* MeshMapBuildData = Component->GetMeshMapBuildData();
