Removed csv stats related to deleted subobject tests. Not actively used anymore.

diff --git a/Engine/Source/Runtime/Engine/Private/DataChannel.cpp b/Engine/Source/Runtime/Engine/Private/DataChannel.cpp
--- a/Engine/Source/Runtime/Engine/Private/DataChannel.cpp
+++ b/Engine/Source/Runtime/Engine/Private/DataChannel.cpp
@@ -50,8 +50,6 @@ DECLARE_CYCLE_STAT(TEXT("ActorChan_FindOrCreateRep"), Stat_ActorChanFindOrCreate
 DECLARE_LLM_MEMORY_STAT(TEXT("NetChannel"), STAT_NetChannelLLM, STATGROUP_LLMFULL);
 LLM_DEFINE_TAG(NetChannel, NAME_None, TEXT("Networking"), GET_STATFNAME(STAT_NetChannelLLM), GET_STATFNAME(STAT_NetworkingSummaryLLM));
 
-CSV_DECLARE_CATEGORY_EXTERN(Replication);
-
 // Enable the code allowing to validate replicate subobjects when converting from the legacy method to the explicit registration list
 #ifndef SUBOBJECT_TRANSITION_VALIDATION
 	#define SUBOBJECT_TRANSITION_VALIDATION !UE_BUILD_SHIPPING
@@ -3432,10 +3430,6 @@ int64 UActorChannel::ReplicateActor()
 		{
 			QUICK_SCOPE_CYCLE_COUNTER(STAT_NetDeletedSubObjects);
 
-#if !UE_BUILD_SHIPPING
-			CSV_CUSTOM_STAT(Replication, DeletedSubobjectTests, ReplicationMap.Num(), ECsvCustomStatOp::Accumulate);
-#endif
-
 			// Look for deleted subobjects
 			FObjectReplicator* LocalActorReplicator = ActorReplicator.Get();
 			for (auto RepComp = ReplicationMap.CreateIterator(); RepComp; ++RepComp)
@@ -3446,9 +3440,6 @@ int64 UActorChannel::ReplicateActor()
 				{
 					if (LocalReplicator->ObjectNetGUID.IsValid())
 					{
-#if !UE_BUILD_SHIPPING
-						CSV_CUSTOM_STAT(Replication, DeletedSubobjects, 1, ECsvCustomStatOp::Accumulate);
-#endif
 						UE_NET_TRACE_SCOPE(ContentBlockForSubObjectDelete, Bunch, GetTraceCollector(Bunch), ENetTraceVerbosity::Trace);
 						UE_NET_TRACE_OBJECT_SCOPE(LocalReplicator->ObjectNetGUID, Bunch, GetTraceCollector(Bunch), ENetTraceVerbosity::Trace);
 
