Moved csv stat RepNotifies closer to the actual processing. This removes PostRepNotifies from the stat, but we rarely have an implementation for that function, so not much is lost.

diff --git a/Engine/Source/Runtime/Engine/Private/DataReplication.cpp b/Engine/Source/Runtime/Engine/Private/DataReplication.cpp
--- a/Engine/Source/Runtime/Engine/Private/DataReplication.cpp
+++ b/Engine/Source/Runtime/Engine/Private/DataReplication.cpp
@@ -2182,7 +2182,6 @@ void FObjectReplicator::CallRepNotifies(bool bSkipIfChannelHasQueuedBunches)
 {
 	// This logic is mostly a copy of FRepLayout::CallRepNotifies, and they should be merged.
 
-	CSV_SCOPED_TIMING_STAT_EXCLUSIVE(RepNotifies);
 	UObject* Object = GetObject();
 
 	if (!IsValid(Object))
diff --git a/Engine/Source/Runtime/Engine/Private/RepLayout.cpp b/Engine/Source/Runtime/Engine/Private/RepLayout.cpp
--- a/Engine/Source/Runtime/Engine/Private/RepLayout.cpp
+++ b/Engine/Source/Runtime/Engine/Private/RepLayout.cpp
@@ -4590,6 +4590,8 @@ void FRepLayout::CallRepNotifies(FReceivingRepState* RepState, UObject* Object)
 		return;
 	}
 
+	CSV_SCOPED_TIMING_STAT_EXCLUSIVE(RepNotifies);
+
 	FRepShadowDataBuffer ShadowData(RepState->StaticBuffer.GetData());
 	FRepObjectDataBuffer ObjectData(Object);
 
