Added stateless handshake level check of Network CL compatibility.

diff --git a/Engine/Source/Runtime/Engine/Private/PacketHandlers/StatelessConnectHandlerComponent.cpp b/Engine/Source/Runtime/Engine/Private/PacketHandlers/StatelessConnectHandlerComponent.cpp
--- a/Engine/Source/Runtime/Engine/Private/PacketHandlers/StatelessConnectHandlerComponent.cpp
+++ b/Engine/Source/Runtime/Engine/Private/PacketHandlers/StatelessConnectHandlerComponent.cpp
@@ -7,6 +7,7 @@
 #include "Misc/SecureHash.h"
 #include "Engine/NetConnection.h"
 #include "Net/Core/Misc/PacketAudit.h"
+#include "Misc/NetworkVersion.h"
 
 
 DEFINE_LOG_CATEGORY(LogHandshake);
@@ -37,19 +38,19 @@ DEFINE_LOG_CATEGORY(LogHandshake);
  * and the server responding with a unique 'Cookie' value, which the client has to respond with.
  *
  * Client - Initial Connect:
- * [?:MagicHeader][HandshakeBit][RestartHandshakeBit][8:MinVersion][8:CurVersion][8:HandshakePacketType][8:SentPacketCount][SecretIdBit][28:PacketSizeFiller][AlignPad][?:RandomData]
+ * [?:MagicHeader][HandshakeBit][RestartHandshakeBit][8:MinVersion][8:CurVersion][8:HandshakePacketType][8:SentPacketCount][32:NetworkVersion][SecretIdBit][28:PacketSizeFiller][AlignPad][?:RandomData]
  *													--->
  *															Server - Stateless Handshake Challenge:
- *															[?:MagicHeader][HandshakeBit][RestartHandshakeBit][8:MinVersion][8:CurVersion][8:HandshakePacketType][8:SentPacketCount][SecretIdBit][8:Timestamp][20:Cookie][AlignPad][?:RandomData]
+ *															[?:MagicHeader][HandshakeBit][RestartHandshakeBit][8:MinVersion][8:CurVersion][8:HandshakePacketType][8:SentPacketCount][32:NetworkVersion][SecretIdBit][8:Timestamp][20:Cookie][AlignPad][?:RandomData]
  *													<---
  * Client - Stateless Challenge Response:
- * [?:MagicHeader][HandshakeBit][RestartHandshakeBit][8:MinVersion][8:CurVersion][8:HandshakePacketType][8:SentPacketCount][SecretIdBit][8:Timestamp][20:Cookie][AlignPad][?:RandomData]
+ * [?:MagicHeader][HandshakeBit][RestartHandshakeBit][8:MinVersion][8:CurVersion][8:HandshakePacketType][8:SentPacketCount][32:NetworkVersion][SecretIdBit][8:Timestamp][20:Cookie][AlignPad][?:RandomData]
  *													--->
  *															Server:
  *															Ignore, or create UNetConnection.
  *
  *															Server - Stateless Handshake Ack
- *															[?:MagicHeader][HandshakeBit][RestartHandshakeBit][8:MinVersion][8:CurVersion][8:HandshakePacketType][8:SentPacketCount][SecretIdBit][8:Timestamp][20:Cookie][AlignPad][?:RandomData]
+ *															[?:MagicHeader][HandshakeBit][RestartHandshakeBit][8:MinVersion][8:CurVersion][8:HandshakePacketType][8:SentPacketCount][32:NetworkVersion][SecretIdBit][8:Timestamp][20:Cookie][AlignPad][?:RandomData]
  *													<---
  * Client:
  * Handshake Complete.
@@ -61,14 +62,14 @@ DEFINE_LOG_CATEGORY(LogHandshake);
  * so the protocol has been crafted so the server sends only a minimal (1 byte) response, to minimize DRDoS reflection amplification.
  *
  *															Server - Restart Handshake Request:
- *															[?:MagicHeader][HandshakeBit][RestartHandshakeBit][8:MinVersion][8:CurVersion][8:HandshakePacketType][8:SentPacketCount][AlignPad][?:RandomData]
+ *															[?:MagicHeader][HandshakeBit][RestartHandshakeBit][8:MinVersion][8:CurVersion][8:HandshakePacketType][8:SentPacketCount][32:NetworkVersion][AlignPad][?:RandomData]
  *													<--
  * Client -  Initial Connect (as above)
  *													-->
  *															Server -  Stateless Handshake Challenge (as above)
  *													<--
  * Client - Stateless Challenge Response + Original Cookie
- * [?:MagicHeader][HandshakeBit][RestartHandshakeBit][8:MinVersion][8:CurVersion][8:HandshakePacketType][8:SentPacketCount][SecretIdBit][8:Timestamp][20:Cookie][20:OriginalCookie][AlignPad][?:RandomData]
+ * [?:MagicHeader][HandshakeBit][RestartHandshakeBit][8:MinVersion][8:CurVersion][8:HandshakePacketType][8:SentPacketCount][SecretIdBit][8:Timestamp][20:Cookie][20:OriginalCookie][32:NetworkVersion][AlignPad][?:RandomData]
  *													-->
  *															Server:
  *															Ignore, or restore UNetConnection.
@@ -88,6 +89,7 @@ DEFINE_LOG_CATEGORY(LogHandshake);
  *	- CurVersion:			The currently active protocol version used by the remote side (determines received packet format)
  *	- HandshakePacketType:	Number indicating the type of handshake packet, based on EHandshakePacketType
  *	- SentPacketCount:		The number of handshake packets sent - for packet-analysis/debugging purposes
+ *	- NetworkVersion:		The Network CL version, according to FNetworkVersion::GetLocalNetworkVersion
  *	- Timestamp:			Server timestamp, from the moment the handshake challenge was sent.
  *	- Cookie:				Cookie generated by the server, which the client must reply with.
  *	- AlignPad:				Handshake packets and PacketHandler's in general, require complex padding of packets. See ParseHandshakePacket.
@@ -160,9 +162,9 @@ DEFINE_LOG_CATEGORY(LogHandshake);
  * Defines
  */
 
-#define HANDSHAKE_PACKET_SIZE_BITS				259
-#define RESTART_HANDSHAKE_PACKET_SIZE_BITS		34
-#define RESTART_RESPONSE_SIZE_BITS				419
+#define HANDSHAKE_PACKET_SIZE_BITS				291
+#define RESTART_HANDSHAKE_PACKET_SIZE_BITS		66
+#define RESTART_RESPONSE_SIZE_BITS				451
 
 
 // The number of seconds between secret value updates, and the random variance applied to this
@@ -224,6 +226,15 @@ namespace UE::Net
 		0,
 		TEXT("Whether or not to (clientside) perform randomized falling-back to previous versions of the handshake protocol, upon failure."));
 
+	static int32 GHandshakeEnforceNetworkCLVersion = 0;
+
+	FAutoConsoleVariableRef CVarNetHandshakeEnforceNetworkCLVersion(
+		TEXT("net.HandshakeEnforceNetworkCLVersion"),
+		GHandshakeEnforceNetworkCLVersion,
+		TEXT("Whether or not the stateless handshake should enforce the Network CL version, instead of the higher level netcode. ")
+		TEXT("(does NOT result in a graceful rejection, packet is dropped without response - higher level netcode would have sent upgrade message)"));
+
+
 	/** The base amount of random data to add to handshake packets */
 	static constexpr int32 BaseRandomDataLengthBytes		= 16;
 
@@ -239,6 +250,16 @@ namespace UE::Net
 
 	/** RESTART_RESPONSE_SIZE_BITS for EHandshakeVersion::Original */
 	static constexpr int32 OriginalRestartResponseSizeBits = 387;
+
+
+	/** HANDSHAKE_PACKET_SIZE_BITS for EHandshakeVersion::Randomized */
+	static constexpr int32 VerRandomizedHandshakePacketSizeBits = 259;
+
+	/** RESTART_HANDSHAKE_PACKET_SIZE_BITS for EHandshakeVersion::Randomized */
+	static constexpr int32 VerRandomizedRestartHandshakePacketSizeBits = 34;
+
+	/** RESTART_RESPONSE_SIZE_BITS for EHandshakeVersion::Randomized */
+	static constexpr int32 VerRandomizedRestartResponseSizeBits = 419;
 }
 
 
@@ -355,7 +376,7 @@ void StatelessConnectHandlerComponent::SendInitialPacket(EHandshakeVersion Hands
 			if (HandshakeVersion >= EHandshakeVersion::Randomized)
 			{
 				uint8 MinVersion = MinSupportedHandshakeVersion;
-				uint8 CurVersion = CurrentHandshakeVersion;
+				uint8 CurVersion = static_cast<uint8>(HandshakeVersion);
 				uint8 HandshakePacketType = static_cast<uint8>(EHandshakePacketType::InitialPacket);
 
 				InitialPacket << MinVersion;
@@ -364,6 +385,13 @@ void StatelessConnectHandlerComponent::SendInitialPacket(EHandshakeVersion Hands
 				InitialPacket << SentHandshakePacketCount;
 			}
 
+			if (HandshakeVersion >= EHandshakeVersion::NetCLVersion)
+			{
+				uint32 LocalNetworkVersion = FNetworkVersion::GetLocalNetworkVersion();
+
+				InitialPacket << LocalNetworkVersion;
+			}
+
 			InitialPacket.WriteBit(SecretIdPad);
 
 			FMemory::Memzero(PacketSizeFiller, UE_ARRAY_COUNT(PacketSizeFiller));
@@ -442,6 +470,13 @@ void StatelessConnectHandlerComponent::SendConnectChallenge(TSharedPtr<const FIn
 			ChallengePacket << ClientSentHandshakePacketCount;
 		}
 
+		if (HandshakeVersion >= EHandshakeVersion::NetCLVersion)
+		{
+			uint32 LocalNetworkVersion = FNetworkVersion::GetLocalNetworkVersion();
+
+			ChallengePacket << LocalNetworkVersion;
+		}
+
 		ChallengePacket.WriteBit(ActiveSecret);
 
 		ChallengePacket << Timestamp;
@@ -536,6 +571,13 @@ void StatelessConnectHandlerComponent::SendChallengeResponse(EHandshakeVersion H
 			ResponsePacket << SentHandshakePacketCount;
 		}
 
+		if (HandshakeVersion >= EHandshakeVersion::NetCLVersion)
+		{
+			uint32 LocalNetworkVersion = FNetworkVersion::GetLocalNetworkVersion();
+
+			ResponsePacket << LocalNetworkVersion;
+		}
+
 		ResponsePacket.WriteBit(InSecretId);
 
 		ResponsePacket << InTimestamp;
@@ -627,6 +669,13 @@ void StatelessConnectHandlerComponent::SendChallengeAck(TSharedPtr<const FIntern
 			AckPacket << ClientSentHandshakePacketCount;
 		}
 
+		if (HandshakeVersion >= EHandshakeVersion::NetCLVersion)
+		{
+			uint32 LocalNetworkVersion = FNetworkVersion::GetLocalNetworkVersion();
+
+			AckPacket << LocalNetworkVersion;
+		}
+
 		AckPacket.WriteBit(bHandshakePacket);	// ActiveSecret
 
 		AckPacket << Timestamp;
@@ -712,6 +761,13 @@ void StatelessConnectHandlerComponent::SendRestartHandshakeRequest(const TShared
 			RestartPacket << SentHandshakePacketCount;
 		}
 
+		if (HandshakeVersion >= EHandshakeVersion::NetCLVersion)
+		{
+			uint32 LocalNetworkVersion = FNetworkVersion::GetLocalNetworkVersion();
+
+			RestartPacket << LocalNetworkVersion;
+		}
+
 #if !UE_BUILD_SHIPPING
 		FDDoSDetection* DDoS = Handler->GetDDoS();
 
@@ -774,6 +830,11 @@ void StatelessConnectHandlerComponent::CapHandshakePacket(FBitWriter& HandshakeP
 		check(NumBits == OriginalHandshakePacketSizeBits || NumBits == OriginalRestartHandshakePacketSizeBits ||
 				NumBits == OriginalRestartResponseSizeBits);
 	}
+	else if (HandshakeVersion == EHandshakeVersion::Randomized)
+	{
+		check(NumBits == VerRandomizedHandshakePacketSizeBits || NumBits == VerRandomizedRestartHandshakePacketSizeBits ||
+				NumBits == VerRandomizedRestartResponseSizeBits);
+	}
 	else
 	{
 		check(NumBits == HANDSHAKE_PACKET_SIZE_BITS || NumBits == RESTART_HANDSHAKE_PACKET_SIZE_BITS || NumBits == RESTART_RESPONSE_SIZE_BITS);
@@ -784,8 +845,14 @@ void StatelessConnectHandlerComponent::CapHandshakePacket(FBitWriter& HandshakeP
 
 	if (HandshakeVersion >= EHandshakeVersion::Randomized)
 	{
-		const int32 RandomDataLengthBytes = BaseRandomDataLengthBytes - FMath::RandRange(0, RandomDataLengthVarianceBytes);
+		int32 RandomDataLengthBytes = BaseRandomDataLengthBytes - FMath::RandRange(0, RandomDataLengthVarianceBytes);
 
+		// In versions that must stay compatible with the original protocol, make sure there isn't a size collision with the original restart response
+		if (NumBits + (RandomDataLengthBytes * 8) == OriginalRestartResponseSizeBits)
+		{
+			RandomDataLengthBytes = FMath::Max(0, RandomDataLengthBytes - 1);
+		}
+		
 		for (int32 RandIdx=0; RandIdx<RandomDataLengthBytes; RandIdx++)
 		{
 			uint8 RandVal = FMath::Rand() % 255;
@@ -1104,22 +1171,53 @@ void StatelessConnectHandlerComponent::IncomingConnectionless(FIncomingPacketRef
 
 		if (bHandshakePacket)
 		{
-			bool bValidVersion = false;
+			bool bValidHandshakeVersion = false;
 			EHandshakeVersion TargetVersion = static_cast<EHandshakeVersion>(CurrentHandshakeVersion);
+			const uint8 RemoteMinVersionUint8 = static_cast<uint8>(HandshakeData.RemoteMinVersion);
 
-			if (static_cast<uint8>(HandshakeData.RemoteMinVersion) >= MinSupportedHandshakeVersion &&
-				HandshakeData.RemoteMinVersion < TargetVersion)
+			if (RemoteMinVersionUint8 >= MinSupportedHandshakeVersion && HandshakeData.RemoteMinVersion <= TargetVersion
+				&& HandshakeData.RemoteMinVersion <= HandshakeData.RemoteCurVersion)
 			{
 				if (HandshakeData.RemoteCurVersion <= TargetVersion)
 				{
 					TargetVersion = HandshakeData.RemoteCurVersion;
 				}
 
-				bValidVersion = true;
+				bValidHandshakeVersion = true;
+			}
+
+			const bool bCheckNetVersion = !UE_BUILD_SHIPPING || GHandshakeEnforceNetworkCLVersion;
+			bool bValidNetVersion = true;
+
+			if (bCheckNetVersion && HandshakeData.RemoteCurVersion >= EHandshakeVersion::NetCLVersion)
+			{
+				const uint32 LocalNetworkVersion = FNetworkVersion::GetLocalNetworkVersion();
+
+				if (!FNetworkVersion::IsNetworkCompatible(LocalNetworkVersion, HandshakeData.RemoteNetworkVersion))
+				{
+					bValidNetVersion = false;
+				}
+			}
+
+#if !UE_BUILD_SHIPPING
+			if (!bValidHandshakeVersion || !bValidNetVersion)
+			{
+				FDDoSDetection* DDoS = Handler->GetDDoS();
+				const uint32 LocalNetworkVersion = FNetworkVersion::GetLocalNetworkVersion();
+
+				UE_CLOG((DDoS == nullptr || !DDoS->CheckLogRestrictions()), LogHandshake, Log,
+						TEXT("IncomingConnectionless: Incompatible version. bValidHandshakeVersion: %i, bValidNetVersion: %i, ")
+						TEXT("GHandshakeEnforceNetworkCLVersion: %i, RemoteMinVersion: %i, RemoteCurVersion: %i, MinSupportedHandshakeVersion: %i, ")
+						TEXT("CurrentHandshakeVersion: %i, RemoteNetworkVersion: %i, LocalNetworkVersion: %i"),
+						(int32)bValidHandshakeVersion, (int32)bValidNetVersion, GHandshakeEnforceNetworkCLVersion, RemoteMinVersionUint8,
+						static_cast<uint8>(HandshakeData.RemoteCurVersion), MinSupportedHandshakeVersion, CurrentHandshakeVersion,
+						HandshakeData.RemoteNetworkVersion, LocalNetworkVersion);
 			}
+#endif
 
+			const bool bPassedNetVersionConditions = bValidNetVersion || !GHandshakeEnforceNetworkCLVersion;
 
-			if (Handler->Mode == Handler::Mode::Server)
+			if (Handler->Mode == Handler::Mode::Server && bValidHandshakeVersion && bPassedNetVersionConditions)
 			{
 				const bool bInitialConnect = HandshakeData.HandshakePacketType == EHandshakePacketType::InitialPacket &&
 												HandshakeData.Timestamp == 0.0;
@@ -1219,12 +1317,14 @@ bool StatelessConnectHandlerComponent::ParseHandshakePacket(FBitReader& Packet,
 
 	// Ensure original packet sizes don't overlap with size range of new packet format - so that we can detect original version, based on size
 	{
-		static constexpr int32 MinHandshakePacketVariance = HANDSHAKE_PACKET_SIZE_BITS - RandomDataLengthVarianceBytes;
-		static constexpr int32 MaxHandshakePacketVariance = HANDSHAKE_PACKET_SIZE_BITS;
-		static constexpr int32 MinRestartHandshakePacketVariance = RESTART_HANDSHAKE_PACKET_SIZE_BITS - RandomDataLengthVarianceBytes;
-		static constexpr int32 MaxRestartHandshakePacketVariance = RESTART_HANDSHAKE_PACKET_SIZE_BITS;
-		static constexpr int32 MinRestartResponsePacketVariance = RESTART_RESPONSE_SIZE_BITS - RandomDataLengthVarianceBytes;
-		static constexpr int32 MaxRestartResponsePacketVariance = RESTART_RESPONSE_SIZE_BITS;
+		static constexpr int32 MinRandomBits = (BaseRandomDataLengthBytes - RandomDataLengthVarianceBytes) * 8;
+		static constexpr int32 MaxRandomBits = BaseRandomDataLengthBytes * 8;
+		static constexpr int32 MinHandshakePacketVariance = HANDSHAKE_PACKET_SIZE_BITS + MinRandomBits;
+		static constexpr int32 MaxHandshakePacketVariance = HANDSHAKE_PACKET_SIZE_BITS + MaxRandomBits;
+		static constexpr int32 MinRestartHandshakePacketVariance = RESTART_HANDSHAKE_PACKET_SIZE_BITS + MinRandomBits;
+		static constexpr int32 MaxRestartHandshakePacketVariance = RESTART_HANDSHAKE_PACKET_SIZE_BITS + MaxRandomBits;
+		static constexpr int32 MinRestartResponsePacketVariance = RESTART_RESPONSE_SIZE_BITS + MinRandomBits;
+		static constexpr int32 MaxRestartResponsePacketVariance = RESTART_RESPONSE_SIZE_BITS + MaxRandomBits;
 
 		static_assert(OriginalHandshakePacketSizeBits < MinHandshakePacketVariance || OriginalHandshakePacketSizeBits > MaxHandshakePacketVariance);
 		static_assert(OriginalHandshakePacketSizeBits < MinRestartHandshakePacketVariance ||
@@ -1237,11 +1337,15 @@ bool StatelessConnectHandlerComponent::ParseHandshakePacket(FBitReader& Packet,
 						OriginalRestartHandshakePacketSizeBits > MaxRestartHandshakePacketVariance);
 		static_assert(OriginalRestartHandshakePacketSizeBits < MinRestartResponsePacketVariance ||
 						OriginalRestartHandshakePacketSizeBits > MaxRestartResponsePacketVariance);
+
+		// Disable these static asserts, as the new protocol has been adjusted to never hit OriginalRestartResponseSizeBits
+#if 0
 		static_assert(OriginalRestartResponseSizeBits < MinHandshakePacketVariance || OriginalRestartResponseSizeBits > MaxHandshakePacketVariance);
 		static_assert(OriginalRestartResponseSizeBits < MinRestartHandshakePacketVariance ||
 						OriginalRestartResponseSizeBits > MaxRestartHandshakePacketVariance);
 		static_assert(OriginalRestartResponseSizeBits < MinRestartResponsePacketVariance ||
 						OriginalRestartResponseSizeBits > MaxRestartResponsePacketVariance);
+#endif
 	};
 
 	bool bValidPacket = false;
@@ -1255,14 +1359,18 @@ bool StatelessConnectHandlerComponent::ParseHandshakePacket(FBitReader& Packet,
 	}
 
 
-	const int32 HandshakePacketBitsLeft = BitsLeft - (HANDSHAKE_PACKET_SIZE_BITS - 1);
-	const int32 RestartHandshakePacketBitsLeft = BitsLeft - (RESTART_HANDSHAKE_PACKET_SIZE_BITS - 1);
-	const int32 RestartResponsePacketBitsLeft = BitsLeft - (RESTART_RESPONSE_SIZE_BITS - 1);
-	const int32 MinBitsLeft = (BaseRandomDataLengthBytes - RandomDataLengthVarianceBytes) * 8;
-	const int32 MaxBitsLeft = BaseRandomDataLengthBytes * 8;
-	const bool bMaybeHandshakePacketSize = HandshakePacketBitsLeft >= MinBitsLeft && HandshakePacketBitsLeft <= MaxBitsLeft;
-	const bool bMaybeRestartHandshakePacket = RestartHandshakePacketBitsLeft >= MinBitsLeft && RestartHandshakePacketBitsLeft <= MaxBitsLeft;
-	const bool bMaybeRestartResponsePacketSize = RestartResponsePacketBitsLeft >= MinBitsLeft && RestartResponsePacketBitsLeft <= MaxBitsLeft;
+	// Remaining bits, excluding packet sizes from different protocol versions
+	const int32 MinBitsLeftExclHandshake = BitsLeft - (HANDSHAKE_PACKET_SIZE_BITS - 1);
+	const int32 MaxBitsLeftExclHandshake = BitsLeft - (VerRandomizedHandshakePacketSizeBits - 1);
+	const int32 MinBitsLeftExclRestartHandshake = BitsLeft - (RESTART_HANDSHAKE_PACKET_SIZE_BITS - 1);
+	const int32 MaxBitsLeftExclRestartHandshake = BitsLeft - (VerRandomizedRestartHandshakePacketSizeBits - 1);
+	const int32 MinBitsLeftExclRestartResponse = BitsLeft - (RESTART_RESPONSE_SIZE_BITS - 1);
+	const int32 MaxBitsLeftExclRestartResponse = BitsLeft - (VerRandomizedRestartResponseSizeBits - 1);
+	const int32 MinRandomBits = (BaseRandomDataLengthBytes - RandomDataLengthVarianceBytes) * 8;
+	const int32 MaxRandomBits = BaseRandomDataLengthBytes * 8;
+	const bool bMaybeHandshakePacket = MaxBitsLeftExclHandshake >= MinRandomBits && MinBitsLeftExclHandshake <= MaxRandomBits;
+	const bool bMaybeRestartHandshakePacket = MaxBitsLeftExclRestartHandshake >= MinRandomBits && MinBitsLeftExclRestartHandshake <= MaxRandomBits;
+	const bool bMaybeRestartResponsePacket = MaxBitsLeftExclRestartResponse >= MinRandomBits && MinBitsLeftExclRestartResponse <= MaxRandomBits;
 
 	OutResult.bRestartHandshake = !!Packet.ReadBit();
 
@@ -1280,13 +1388,18 @@ bool StatelessConnectHandlerComponent::ParseHandshakePacket(FBitReader& Packet,
 	OutResult.RemoteCurVersion = static_cast<EHandshakeVersion>(RemoteCurVersion);
 	HandshakePacketTypeEnum = static_cast<EHandshakePacketType>(HandshakePacketType);
 
+	if (OutResult.RemoteCurVersion >= EHandshakeVersion::NetCLVersion)
+	{
+		Packet << OutResult.RemoteNetworkVersion;
+	}
+
 	// Only accept handshake packets of roughly the right size
-	const bool bHandshakePacket = bMaybeHandshakePacketSize && (HandshakePacketTypeEnum == EHandshakePacketType::InitialPacket ||
+	const bool bHandshakePacket = bMaybeHandshakePacket && (HandshakePacketTypeEnum == EHandshakePacketType::InitialPacket ||
 		HandshakePacketTypeEnum == EHandshakePacketType::Challenge || HandshakePacketTypeEnum == EHandshakePacketType::Response ||
 		HandshakePacketTypeEnum == EHandshakePacketType::Ack);
 
 	const bool bRestartHandshakePacket = bMaybeRestartHandshakePacket && HandshakePacketTypeEnum == EHandshakePacketType::RestartHandshake;
-	const bool bRestartResponsePacket = bMaybeRestartResponsePacketSize && HandshakePacketTypeEnum == EHandshakePacketType::RestartResponse;
+	const bool bRestartResponsePacket = bMaybeRestartResponsePacket && HandshakePacketTypeEnum == EHandshakePacketType::RestartResponse;
 
 	// Only accept handshake packets of precisely the right size
 	if (bHandshakePacket || bRestartResponsePacket)
diff --git a/Engine/Source/Runtime/Engine/Public/PacketHandlers/StatelessConnectHandlerComponent.h b/Engine/Source/Runtime/Engine/Public/PacketHandlers/StatelessConnectHandlerComponent.h
--- a/Engine/Source/Runtime/Engine/Public/PacketHandlers/StatelessConnectHandlerComponent.h
+++ b/Engine/Source/Runtime/Engine/Public/PacketHandlers/StatelessConnectHandlerComponent.h
@@ -35,8 +35,9 @@ namespace UE::Net
 	{
 		Original		= 0,	// The original/unversioned handshake protocol
 		Randomized		= 1,	// The version of the handshake protocol with randomization, versioning and debug/diagnostic tweaks
+		NetCLVersion	= 2,	// Added Network CL version, for optional extra-early client rejection
 
-		Latest			= Randomized
+		Latest			= NetCLVersion
 	};
 
 	/**
@@ -256,6 +257,9 @@ private:
 		/** The current handshake protocol version the remote side used for communication */
 		EHandshakeVersion RemoteCurVersion = EHandshakeVersion::Latest;
 
+		/** The Network CL version of the remote side */
+		uint32 RemoteNetworkVersion = 0;
+
 		/** The type of handshake packet */
 		UE::Net::EHandshakePacketType HandshakePacketType = UE::Net::EHandshakePacketType::InitialPacket;
 
