Add restriction option to custom HLSL expressions in user materials.

diff --git a/Engine/Source/Runtime/Engine/Private/Materials/HLSLMaterialTranslator.cpp b/Engine/Source/Runtime/Engine/Private/Materials/HLSLMaterialTranslator.cpp
--- a/Engine/Source/Runtime/Engine/Private/Materials/HLSLMaterialTranslator.cpp
+++ b/Engine/Source/Runtime/Engine/Private/Materials/HLSLMaterialTranslator.cpp
@@ -4,6 +4,26 @@
 =============================================================================*/
 #include "HLSLMaterialTranslator.h"
 #include "VT/VirtualTextureScalability.h"
+#include "Materials/MaterialExpressionCustom.h"
+#include "MaterialRestrictiveMode.h"
+
+bool IsExpressionClassPermitted(const UClass* const Class)
+{
+	if (Class)
+	{
+		if (Class == UMaterialExpressionCustom::StaticClass())
+		{
+			// If custom HLSL expressions are restricted, they can only be used in cooked content
+			return
+				Class->bCooked
+#if WITH_EDITORONLY_DATA
+				|| !FHLSLMaterialTranslator::InRestrictiveMode()
+#endif
+			;
+		}
+	}
+	return true;
+}
 
 #if WITH_EDITORONLY_DATA
 
@@ -2549,6 +2569,24 @@ FString FHLSLMaterialTranslator::GetParameterCodeRaw(int32 Index, const TCHAR* D
 	}
 }
 
+// Public interface for restrictive mode.
+void SetHLSLMaterialTranslatorRestrictiveMode(bool bRestrictiveMode)
+{
+	FHLSLMaterialTranslator::SetRestrictiveMode(bRestrictiveMode);
+}
+
+static bool GIsMaterialTranslatorInRestrictiveMode = false;
+
+void FHLSLMaterialTranslator::SetRestrictiveMode(bool bRestrictiveMode)
+{
+	GIsMaterialTranslatorInRestrictiveMode = bRestrictiveMode;
+}
+
+bool FHLSLMaterialTranslator::InRestrictiveMode()
+{
+	return GIsMaterialTranslatorInRestrictiveMode;
+}
+
 uint64 FHLSLMaterialTranslator::GetParameterHash(int32 Index)
 {
 	if (Index == INDEX_NONE)
@@ -11056,6 +11094,11 @@ int32 FHLSLMaterialTranslator::AccessMaterialAttribute(int32 CodeIndex, const FG
 
 int32 FHLSLMaterialTranslator::CustomExpression( class UMaterialExpressionCustom* Custom, int32 OutputIndex, TArray<int32>& CompiledInputs )
 {
+	if (!IsExpressionClassPermitted(Custom->GetClass()))
+	{
+		return Errorf(TEXT("Custom HLSL expressions are not permitted"));
+	}
+
 	const FMaterialCustomExpressionEntry* CustomEntry = nullptr;
 	for (const FMaterialCustomExpressionEntry& Entry : CustomExpressions)
 	{
diff --git a/Engine/Source/Runtime/Engine/Private/Materials/HLSLMaterialTranslator.h b/Engine/Source/Runtime/Engine/Private/Materials/HLSLMaterialTranslator.h
--- a/Engine/Source/Runtime/Engine/Private/Materials/HLSLMaterialTranslator.h
+++ b/Engine/Source/Runtime/Engine/Private/Materials/HLSLMaterialTranslator.h
@@ -54,6 +54,12 @@
 
 class Error;
 
+/**
+ * Returns whether the specified class of material expression is permitted.
+ * For instance, custom expressions are not permitted in certain UE editor configurations for client generated materials.
+ */
+bool IsExpressionClassPermitted(const UClass* const Class);
+
 #if WITH_EDITORONLY_DATA
 
 enum EMaterialExpressionVisitResult
@@ -514,6 +520,12 @@ protected:
 public:
 	// Must always be valid
 	virtual FString GetParameterCodeDeriv(int32 Index, ECompiledPartialDerivativeVariation Variation);
+
+	// Puts the translator into restrictive mode.
+	static void SetRestrictiveMode(bool bRestrictiveMode);
+	// Returns whether the material translator is in restrictive mode. If true, custom HLSL expressions are only allowed for cooked content.
+	static bool InRestrictiveMode();
+
 protected:
 
 	uint64 GetParameterHash(int32 Index);
diff --git a/Engine/Source/Runtime/Engine/Private/Materials/MaterialExpressions.cpp b/Engine/Source/Runtime/Engine/Private/Materials/MaterialExpressions.cpp
--- a/Engine/Source/Runtime/Engine/Private/Materials/MaterialExpressions.cpp
+++ b/Engine/Source/Runtime/Engine/Private/Materials/MaterialExpressions.cpp
@@ -24,6 +24,7 @@
 #include "Engine/Font.h"
 #include "MaterialShared.h"
 #include "MaterialExpressionIO.h"
+#include "Materials/HLSLMaterialTranslator.h"
 #include "Materials/MaterialExpression.h"
 #include "Materials/MaterialExpressionMaterialFunctionCall.h"
 #include "Materials/MaterialExpressionMaterialAttributeLayers.h"
@@ -311,6 +312,12 @@ bool Engine_IsStrataEnabled();
 /** Returns whether the given expression class is allowed. */
 bool IsAllowedExpressionType(const UClass* const Class, const bool bMaterialFunction)
 {
+	// Custom HLSL expressions are not allowed for client generated materials in certain UE editor configuration
+	if (!IsExpressionClassPermitted(Class))
+	{
+		return false;
+	}
+
 	static const auto AllowTextureArrayAssetCreationVar = IConsoleManager::Get().FindTConsoleVariableDataInt(TEXT("r.AllowTexture2DArrayCreation"));
 
 	// Exclude comments from the expression list, as well as the base parameter expression, as it should not be used directly
diff --git a/Engine/Source/Runtime/Engine/Public/MaterialRestrictiveMode.h b/Engine/Source/Runtime/Engine/Public/MaterialRestrictiveMode.h
new file mode 100644
--- /dev/null
+++ b/Engine/Source/Runtime/Engine/Public/MaterialRestrictiveMode.h
@@ -0,0 +1,11 @@
+// Copyright Epic Games, Inc. All Rights Reserved.
+#pragma once
+
+#include "CoreMinimal.h"
+
+#if WITH_EDITORONLY_DATA
+
+/// Puts the HLSL material translator into restrictive mode.
+ENGINE_API void SetHLSLMaterialTranslatorRestrictiveMode(bool bRestrictiveMode);
+
+#endif // WITH_EDITORONLY_DATA
