Allow transaction object diff generation to skip properties via a callback

#rnx

diff --git a/Engine/Source/Editor/UnrealEd/Private/EditorTransaction.cpp b/Engine/Source/Editor/UnrealEd/Private/EditorTransaction.cpp
--- a/Engine/Source/Editor/UnrealEd/Private/EditorTransaction.cpp
+++ b/Engine/Source/Editor/UnrealEd/Private/EditorTransaction.cpp
@@ -305,7 +305,10 @@ void FTransaction::FObjectRecord::Finalize( FTransaction* Owner, TSharedPtr<ITra
 			// have been tracking delta-changes between snapshots and this finalization will need to account for those changes too
 			if (DiffableObjectSnapshot)
 			{
-				UE::Transaction::DiffUtil::GenerateObjectDiff(*DiffableObjectSnapshot, CurrentDiffableObject, DeltaChange, /*bFullDiff*/false);
+				UE::Transaction::DiffUtil::FGenerateObjectDiffOptions DiffOptions;
+				DiffOptions.bFullDiff = false;
+
+				UE::Transaction::DiffUtil::GenerateObjectDiff(*DiffableObjectSnapshot, CurrentDiffableObject, DeltaChange, DiffOptions);
 			}
 		}
 
@@ -371,7 +374,12 @@ void FTransaction::FObjectRecord::Snapshot( FTransaction* Owner, TArrayView<cons
 		// Diff against the correct serialized data depending on whether we already had a snapshot
 		const UE::Transaction::FDiffableObject& InitialDiffableObject = DiffableObjectSnapshot ? *DiffableObjectSnapshot : *DiffableObject;
 		FTransactionObjectDeltaChange SnapshotDeltaChange;
-		UE::Transaction::DiffUtil::GenerateObjectDiff(InitialDiffableObject, CurrentDiffableObject, SnapshotDeltaChange, /*bFullDiff*/false);
+		{
+			UE::Transaction::DiffUtil::FGenerateObjectDiffOptions DiffOptions;
+			DiffOptions.bFullDiff = false;
+
+			UE::Transaction::DiffUtil::GenerateObjectDiff(InitialDiffableObject, CurrentDiffableObject, SnapshotDeltaChange, DiffOptions);
+		}
 
 		// Update the snapshot data for next time
 		if (!DiffableObjectSnapshot)
diff --git a/Engine/Source/Runtime/Engine/Private/TransactionCommon.cpp b/Engine/Source/Runtime/Engine/Private/TransactionCommon.cpp
--- a/Engine/Source/Runtime/Engine/Private/TransactionCommon.cpp
+++ b/Engine/Source/Runtime/Engine/Private/TransactionCommon.cpp
@@ -468,14 +468,14 @@ FDiffableObject GetDiffableObject(const UObject* Object, const EGetDiffableObjec
 	return DiffableObject;
 }
 
-FTransactionObjectDeltaChange GenerateObjectDiff(const FDiffableObject& OldDiffableObject, const FDiffableObject& NewDiffableObject, const bool bFullDiff)
+FTransactionObjectDeltaChange GenerateObjectDiff(const FDiffableObject& OldDiffableObject, const FDiffableObject& NewDiffableObject, const FGenerateObjectDiffOptions& DiffOptions)
 {
 	FTransactionObjectDeltaChange DeltaChange;
-	GenerateObjectDiff(OldDiffableObject, NewDiffableObject, DeltaChange, bFullDiff);
+	GenerateObjectDiff(OldDiffableObject, NewDiffableObject, DeltaChange, DiffOptions);
 	return DeltaChange;
 }
 
-void GenerateObjectDiff(const FDiffableObject& OldDiffableObject, const FDiffableObject& NewDiffableObject, FTransactionObjectDeltaChange& OutDeltaChange, const bool bFullDiff)
+void GenerateObjectDiff(const FDiffableObject& OldDiffableObject, const FDiffableObject& NewDiffableObject, FTransactionObjectDeltaChange& OutDeltaChange, const FGenerateObjectDiffOptions& DiffOptions)
 {
 	auto IsTaggedDataBlockIdentical = [&OldDiffableObject, &NewDiffableObject, &OutDeltaChange](const FName TaggedDataKey, const FSerializedTaggedData& OldSerializedTaggedData, const FSerializedTaggedData& NewSerializedTaggedData) -> bool
 	{
@@ -488,7 +488,7 @@ void GenerateObjectDiff(const FDiffableObject& OldDiffableObject, const FDiffabl
 		return bIsTaggedDataIdentical;
 	};
 
-	auto ShouldCompareTaggedData = [](const FName TaggedDataKey) -> bool
+	auto ShouldCompareTaggedData = [&DiffOptions](const FName TaggedDataKey) -> bool
 	{
 		if (TaggedDataKey.GetComparisonIndex() == TaggedDataKey_ScriptData.GetComparisonIndex())
 		{
@@ -496,6 +496,11 @@ void GenerateObjectDiff(const FDiffableObject& OldDiffableObject, const FDiffabl
 			return false;
 		}
 
+		if (DiffOptions.ShouldSkipProperty && DiffOptions.ShouldSkipProperty(TaggedDataKey))
+		{
+			return false;
+		}
+
 		return true;
 	};
 
@@ -504,7 +509,7 @@ void GenerateObjectDiff(const FDiffableObject& OldDiffableObject, const FDiffabl
 		return TaggedDataKey.GetComparisonIndex() == TaggedDataKey_UnknownData.GetComparisonIndex();
 	};
 
-	if (bFullDiff)
+	if (DiffOptions.bFullDiff)
 	{
 		OutDeltaChange.bHasNameChange |= OldDiffableObject.ObjectInfo.ObjectName != NewDiffableObject.ObjectInfo.ObjectName;
 		OutDeltaChange.bHasOuterChange |= OldDiffableObject.ObjectInfo.ObjectOuterPathName != NewDiffableObject.ObjectInfo.ObjectOuterPathName;
@@ -522,7 +527,7 @@ void GenerateObjectDiff(const FDiffableObject& OldDiffableObject, const FDiffabl
 		const FSerializedTaggedData* OldSerializedTaggedData = OldDiffableObject.SerializedTaggedData.Find(NewNamePropertyPair.Key);
 		if (ShouldCompareAsNonPropertyData(NewNamePropertyPair.Key))
 		{
-			if (bFullDiff && !OutDeltaChange.bHasNonPropertyChanges && (!OldSerializedTaggedData || !IsTaggedDataBlockIdentical(NewNamePropertyPair.Key, *OldSerializedTaggedData, NewNamePropertyPair.Value)))
+			if (DiffOptions.bFullDiff && !OutDeltaChange.bHasNonPropertyChanges && (!OldSerializedTaggedData || !IsTaggedDataBlockIdentical(NewNamePropertyPair.Key, *OldSerializedTaggedData, NewNamePropertyPair.Value)))
 			{
 				OutDeltaChange.bHasNonPropertyChanges = true;
 			}
@@ -534,14 +539,14 @@ void GenerateObjectDiff(const FDiffableObject& OldDiffableObject, const FDiffabl
 				OutDeltaChange.ChangedProperties.AddUnique(NewNamePropertyPair.Key);
 			}
 		}
-		else if (bFullDiff)
+		else if (DiffOptions.bFullDiff)
 		{
 			// Missing property, assume that the property changed
 			OutDeltaChange.ChangedProperties.AddUnique(NewNamePropertyPair.Key);
 		}
 	}
 
-	if (bFullDiff)
+	if (DiffOptions.bFullDiff)
 	{
 		for (const TPair<FName, FSerializedTaggedData>& OldNamePropertyPair : OldDiffableObject.SerializedTaggedData)
 		{
diff --git a/Engine/Source/Runtime/Engine/Public/TransactionCommon.h b/Engine/Source/Runtime/Engine/Public/TransactionCommon.h
--- a/Engine/Source/Runtime/Engine/Public/TransactionCommon.h
+++ b/Engine/Source/Runtime/Engine/Public/TransactionCommon.h
@@ -424,8 +424,14 @@ enum EGetDiffableObjectMode : uint8
 };
 ENGINE_API FDiffableObject GetDiffableObject(const UObject* Object, const EGetDiffableObjectMode Mode = EGetDiffableObjectMode::SerializeObject);
 
-ENGINE_API FTransactionObjectDeltaChange GenerateObjectDiff(const FDiffableObject& OldDiffableObject, const FDiffableObject& NewDiffableObject, const bool bFullDiff = true);
-ENGINE_API void GenerateObjectDiff(const FDiffableObject& OldDiffableObject, const FDiffableObject& NewDiffableObject, FTransactionObjectDeltaChange& OutDeltaChange, const bool bFullDiff = true);
+struct FGenerateObjectDiffOptions
+{
+	bool bFullDiff = true;
+	TFunction<bool(FName)> ShouldSkipProperty;
+};
+
+ENGINE_API FTransactionObjectDeltaChange GenerateObjectDiff(const FDiffableObject& OldDiffableObject, const FDiffableObject& NewDiffableObject, const FGenerateObjectDiffOptions& DiffOptions = FGenerateObjectDiffOptions());
+ENGINE_API void GenerateObjectDiff(const FDiffableObject& OldDiffableObject, const FDiffableObject& NewDiffableObject, FTransactionObjectDeltaChange& OutDeltaChange, const FGenerateObjectDiffOptions& DiffOptions = FGenerateObjectDiffOptions());
 
 } // namespace DiffUtil
 
