Animation: Fix CIS

diff --git a/Engine/Source/Runtime/Engine/Classes/Components/SkinnedMeshComponent.h b/Engine/Source/Runtime/Engine/Classes/Components/SkinnedMeshComponent.h
--- a/Engine/Source/Runtime/Engine/Classes/Components/SkinnedMeshComponent.h
+++ b/Engine/Source/Runtime/Engine/Classes/Components/SkinnedMeshComponent.h
@@ -8,6 +8,7 @@
 #include "UObject/ObjectMacros.h"
 #include "UObject/Object.h"
 #include "Engine/EngineTypes.h"
+#include "Engine/SkeletalMesh.h"
 #include "Components/SceneComponent.h"
 #include "Interfaces/Interface_AsyncCompilation.h"
 #include "Engine/TextureStreamingTypes.h"
