Add SkeletalMeshComponent function to access ConstraintProfile data - needed so PhysicsControlComponent can create controls from constraint profiles

diff --git a/Engine/Source/Runtime/Engine/Classes/Components/SkeletalMeshComponent.h b/Engine/Source/Runtime/Engine/Classes/Components/SkeletalMeshComponent.h
--- a/Engine/Source/Runtime/Engine/Classes/Components/SkeletalMeshComponent.h
+++ b/Engine/Source/Runtime/Engine/Classes/Components/SkeletalMeshComponent.h
@@ -2109,6 +2109,13 @@ public:
 	UFUNCTION(BlueprintCallable, Category = "Physics")
 	void SetConstraintProfileForAll(FName ProfileName, bool bDefaultIfNotFound = false);
 
+	/**
+	 * Gets the constraint profile properties that a joint drive would adopt if it were set to 
+	 * the given constraint profile. The default will be returned if an empty or invalid profile name 
+	 * is passed in. Returns true if the successful, or false if the joint can't be found.
+	 */
+	bool GetConstraintProfileProperties(FConstraintProfileProperties& OutProperties, FName JointName, FName ProfileName);
+
 	/** Enable or Disable AngularPositionDrive based on a list of bone names */
 	void SetNamedMotorsAngularPositionDrive(bool bEnableSwingDrive, bool bEnableTwistDrive, const TArray<FName>& BoneNames, bool bSetOtherBodiesToComplement = false);
 
diff --git a/Engine/Source/Runtime/Engine/Classes/PhysicsEngine/PhysicsConstraintTemplate.h b/Engine/Source/Runtime/Engine/Classes/PhysicsEngine/PhysicsConstraintTemplate.h
--- a/Engine/Source/Runtime/Engine/Classes/PhysicsEngine/PhysicsConstraintTemplate.h
+++ b/Engine/Source/Runtime/Engine/Classes/PhysicsEngine/PhysicsConstraintTemplate.h
@@ -99,6 +99,21 @@ class UPhysicsConstraintTemplate : public UObject
 		}
 	}
 
+	ENGINE_API const FConstraintProfileProperties& GetConstraintProfileProperties(FName ProfileName) const
+	{
+		if (ProfileName != NAME_None)
+		{
+			for (const FPhysicsConstraintProfileHandle& Handle : ProfileHandles)
+			{
+				if (Handle.ProfileName == ProfileName)
+				{
+					return Handle.ProfileProperties;
+				}
+			}
+		}
+		return DefaultProfile;
+	}
+
 #if WITH_EDITOR
 	/** Copy constraint instance into default profile.*/
 	ENGINE_API void SetDefaultProfile(FConstraintInstance& CI)
diff --git a/Engine/Source/Runtime/Engine/Private/SkeletalMeshComponentPhysics.cpp b/Engine/Source/Runtime/Engine/Private/SkeletalMeshComponentPhysics.cpp
--- a/Engine/Source/Runtime/Engine/Private/SkeletalMeshComponentPhysics.cpp
+++ b/Engine/Source/Runtime/Engine/Private/SkeletalMeshComponentPhysics.cpp
@@ -1303,6 +1303,24 @@ void USkeletalMeshComponent::SetConstraintProfileForAll(FName ProfileName, bool
 	}
 }
 
+bool USkeletalMeshComponent::GetConstraintProfileProperties(
+	FConstraintProfileProperties& OutProperties, FName JointName, FName ProfileName)
+{
+	if (UPhysicsAsset* const PhysicsAsset = GetPhysicsAsset())
+	{
+		for (int32 i = 0; i < Constraints.Num(); i++)
+		{
+			FConstraintInstance* ConstraintInstance = Constraints[i];
+			if (ConstraintInstance->JointName == JointName)
+			{
+				OutProperties = PhysicsAsset->ConstraintSetup[i]->GetConstraintProfileProperties(ProfileName);
+				return true;
+			}
+		}
+	}
+	return false;
+}
+
 void USkeletalMeshComponent::SetAllMotorsAngularDriveParams(float InSpring, float InDamping, float InForceLimit, bool bSkipCustomPhysicsType)
 {
 	UPhysicsAsset* const PhysicsAsset = GetPhysicsAsset();
