Strata - added fast and single path define so that they can be used for optimisation during base pass.

diff --git a/Engine/Source/Runtime/Engine/Private/Materials/HLSLMaterialTranslator.cpp b/Engine/Source/Runtime/Engine/Private/Materials/HLSLMaterialTranslator.cpp
--- a/Engine/Source/Runtime/Engine/Private/Materials/HLSLMaterialTranslator.cpp
+++ b/Engine/Source/Runtime/Engine/Private/Materials/HLSLMaterialTranslator.cpp
@@ -1932,6 +1932,9 @@ void FHLSLMaterialTranslator::GetMaterialEnvironment(EShaderPlatform InPlatform,
 		OutEnvironment.SetDefine(TEXT("STRATA_MATERIAL_EXPORT_CONTEXT"), (int32)GetStrataMaterialExportContext());
 		OutEnvironment.SetDefine(TEXT("STRATA_MATERIAL_EXPORT_LEGACY_BLEND_MODE"), (int32)GetStrataMaterialExportLegacyBlendMode());
 
+		OutEnvironment.SetDefine(TEXT("STRATA_SINGLEPATH"), bStrataMaterialIsSingle ? TEXT("1") : TEXT("0"));
+		OutEnvironment.SetDefine(TEXT("STRATA_FASTPATH"),   bStrataMaterialIsSingle ? TEXT("0") : (bStrataMaterialIsSimple ? TEXT("1") : TEXT("0")));
+
 		check(StrataMaterialRootOperator);
 		uint32 RootMaximumDistanceToLeaves = StrataMaterialRootOperator->MaxDistanceFromLeaves;
 		OutEnvironment.SetDefine(TEXT("MATERIAL_TOPOLOGY_TREE_MAX_DEPTH"), RootMaximumDistanceToLeaves);
