Added ShadingModels mask.
r.Mobile.ShadingModelsMask can be used to mask out specific shading models on mobile platforms.
Masked shading models will fallback to DefaultLit

diff --git a/Engine/Source/Runtime/Core/Private/HAL/ConsoleManager.cpp b/Engine/Source/Runtime/Core/Private/HAL/ConsoleManager.cpp
--- a/Engine/Source/Runtime/Core/Private/HAL/ConsoleManager.cpp
+++ b/Engine/Source/Runtime/Core/Private/HAL/ConsoleManager.cpp
@@ -2644,9 +2644,9 @@ static TAutoConsoleVariable<int32> CVarMobileAllowPerPixelShadingModels(
 );
 
 static TAutoConsoleVariable<int32> CVarMobileEnabledShadingModelsMask(
-	TEXT("r.Mobile.EnabledShadingModelsMask"),
-	0xFFFF,
-	TEXT("The mask of shader models that indicates which shading models could be enabled on mobile."),
+	TEXT("r.Mobile.ShadingModelsMask"),
+	0xFFFFFFFF,
+	TEXT("The mask that indicates which shading models are enabled on mobile platforms."),
 	ECVF_ReadOnly | ECVF_RenderThreadSafe
 );
 
diff --git a/Engine/Source/Runtime/Engine/Private/Materials/HLSLMaterialTranslator.cpp b/Engine/Source/Runtime/Engine/Private/Materials/HLSLMaterialTranslator.cpp
--- a/Engine/Source/Runtime/Engine/Private/Materials/HLSLMaterialTranslator.cpp
+++ b/Engine/Source/Runtime/Engine/Private/Materials/HLSLMaterialTranslator.cpp
@@ -9399,10 +9399,12 @@ int32 FHLSLMaterialTranslator::CustomPrimitiveData(int32 OutputIndex, EMaterialV
 int32 FHLSLMaterialTranslator::ShadingModel(EMaterialShadingModel InSelectedShadingModel)
 {
 	// If the shading model is masked out, fallback to default shading model
-	if ((EnabledShadingModelsMask(Platform) & (1 << (int32)InSelectedShadingModel)) == 0)
+	uint32 PlatformShadingModelsMask = GetPlatformShadingModelsMask(Platform);
+	if ((PlatformShadingModelsMask & (1u << (uint32)InSelectedShadingModel)) == 0)
 	{
 		InSelectedShadingModel = MSM_DefaultLit;
 	}
+	
 	ShadingModelsFromCompilation.AddShadingModel(InSelectedShadingModel);
 	return AddInlinedCodeChunk(MCT_ShadingModel, TEXT("%d"), InSelectedShadingModel);
 }
diff --git a/Engine/Source/Runtime/RenderCore/Private/RenderUtils.cpp b/Engine/Source/Runtime/RenderCore/Private/RenderUtils.cpp
--- a/Engine/Source/Runtime/RenderCore/Private/RenderUtils.cpp
+++ b/Engine/Source/Runtime/RenderCore/Private/RenderUtils.cpp
@@ -1117,15 +1117,14 @@ RENDERCORE_API bool AllowPerPixelShadingModels(const FStaticShaderPlatform Platf
 	return true;
 }
 
-RENDERCORE_API int32 EnabledShadingModelsMask(const FStaticShaderPlatform Platform)
+RENDERCORE_API uint32 GetPlatformShadingModelsMask(const FStaticShaderPlatform Platform)
 {
-	static const auto CVar = IConsoleManager::Get().FindTConsoleVariableDataInt(TEXT("r.Mobile.EnabledShadingModelsMask"));
-
-	if (IsMobilePlatform(Platform) && CVar)
+	if (IsMobilePlatform(Platform))
 	{
+		static const auto CVar = IConsoleManager::Get().FindTConsoleVariableDataInt(TEXT("r.Mobile.ShadingModelsMask"));
 		return CVar->GetValueOnAnyThread();
 	}
-	return 0xFFFF;
+	return 0xFFFFFFFF;
 }
 
 RENDERCORE_API uint64 GMobileAmbientOcclusionPlatformMask = 0;
diff --git a/Engine/Source/Runtime/RenderCore/Private/Shader.cpp b/Engine/Source/Runtime/RenderCore/Private/Shader.cpp
--- a/Engine/Source/Runtime/RenderCore/Private/Shader.cpp
+++ b/Engine/Source/Runtime/RenderCore/Private/Shader.cpp
@@ -1619,15 +1619,18 @@ void ShaderMapAppendKeyString(EShaderPlatform Platform, FString& KeyString)
 			KeyString += IsMobileDistanceFieldEnabled(Platform) ? TEXT("_MobSDF") : TEXT("");
 		}
 
-		{
-			KeyString +=  FString::Printf(TEXT("ESMM_%X"), EnabledShadingModelsMask(Platform));
-		}
 	}
 	else
 	{
 		KeyString += IsUsingEmulatedUniformBuffers(Platform) ? TEXT("_NoUB") : TEXT("");
 	}
 
+	uint32 PlatformShadingModelsMask = GetPlatformShadingModelsMask(Platform);
+	if (PlatformShadingModelsMask != 0xFFFFFFFF)
+	{
+		KeyString +=  FString::Printf(TEXT("SMM_%X"), PlatformShadingModelsMask);
+	}
+
 	const IShaderFormat* ShaderFormat = GetTargetPlatformManagerRef().FindShaderFormat(ShaderFormatName);
 	if (ShaderFormat)
 	{
diff --git a/Engine/Source/Runtime/RenderCore/Public/RenderUtils.h b/Engine/Source/Runtime/RenderCore/Public/RenderUtils.h
--- a/Engine/Source/Runtime/RenderCore/Public/RenderUtils.h
+++ b/Engine/Source/Runtime/RenderCore/Public/RenderUtils.h
@@ -496,7 +496,7 @@ RENDERCORE_API bool AllowPixelDepthOffset(const FStaticShaderPlatform Platform);
 
 RENDERCORE_API bool AllowPerPixelShadingModels(const FStaticShaderPlatform Platform);
 
-RENDERCORE_API int32 EnabledShadingModelsMask(const FStaticShaderPlatform Platform);
+RENDERCORE_API uint32 GetPlatformShadingModelsMask(const FStaticShaderPlatform Platform);
 
 RENDERCORE_API bool IsMobileAmbientOcclusionEnabled(const FStaticShaderPlatform Platform);
 
