Remove hotfix as the ddc keys have been bumped many times since then.

diff --git a/Engine/Source/Runtime/Engine/Private/TextureDerivedDataTask.cpp b/Engine/Source/Runtime/Engine/Private/TextureDerivedDataTask.cpp
--- a/Engine/Source/Runtime/Engine/Private/TextureDerivedDataTask.cpp
+++ b/Engine/Source/Runtime/Engine/Private/TextureDerivedDataTask.cpp
@@ -1523,33 +1523,6 @@ void FTextureCacheDerivedDataWorker::DoWork()
 				{
 					UE_LOG(LogTexture, Display, TEXT("Texture %s is missing derived mips. The texture will be rebuilt."), *TexturePathName);
 				}
-
-				if (bSucceeded && ActiveBuildSettingsPerLayer.Num() > 0)
-				{
-					// Bandaid code to prevent loading bad data from DDC
-					// Code inspired by the texture compressor module as a hot fix for the bad data that might have been push into the ddc in 4.23 or 4.24 
-					const bool bLongLatCubemap = DerivedData->IsCubemap() && DerivedData->GetNumSlices() == 1;
-					int32 MaximumNumberOfMipMaps = TNumericLimits<int32>::Max();
-					if (bLongLatCubemap)
-					{
-						MaximumNumberOfMipMaps = FMath::CeilLogTwo(FMath::Clamp<uint32>(uint32(1 << FMath::FloorLog2(DerivedData->SizeX / 2)), uint32(32), ActiveBuildSettingsPerLayer[0].MaxTextureResolution)) + 1;
-					}
-					else
-					{
-						// this looks wrong, should be floor
-						// that's okay because an over-estimate here doesn't hurt
-						MaximumNumberOfMipMaps = FMath::CeilLogTwo(FMath::Max3(DerivedData->SizeX, DerivedData->SizeY, ActiveBuildSettingsPerLayer[0].bVolume ? DerivedData->GetNumSlices() : 1)) + 1;
-					}
-
-					bSucceeded = DerivedData->Mips.Num() <= MaximumNumberOfMipMaps;
-
-					if (!bSucceeded)
-					{
-						UE_LOG(LogTexture, Warning, TEXT("The data retrieved from the derived data cache for the texture %s was invalid. ")
-							TEXT("The cached data has %d mips when a maximum of %d are expected. The texture will be rebuilt."),
-							*TexturePathName, DerivedData->Mips.Num(), MaximumNumberOfMipMaps);
-					}
-				}
 			}
 		}
 
