Add IWorldPartitionCookPackageGenerator to allow construct other than a WorldPartition to generate packages.
Add WorldPartitionCookPackageContext to keep track of which generator requested which package.
WorldPartition now cooks via the IWorldPartitionCookPackageGenerator

#test cooked AncientGame, CitySample, qaentry

diff --git a/Engine/Source/Runtime/Engine/Private/WorldPartition/Cook/WorldPartitionCookPackageContext.cpp b/Engine/Source/Runtime/Engine/Private/WorldPartition/Cook/WorldPartitionCookPackageContext.cpp
new file mode 100644
--- /dev/null
+++ b/Engine/Source/Runtime/Engine/Private/WorldPartition/Cook/WorldPartitionCookPackageContext.cpp
@@ -0,0 +1,81 @@
+#include "WorldPartition/Cook/WorldPartitionCookPackageContext.h"
+
+#if WITH_EDITOR
+
+#include "WorldPartition/WorldPartitionLog.h"
+#include "WorldPartition/Cook/WorldPartitionCookPackageGenerator.h"
+
+FWorldPartitionCookPackageContext::FWorldPartitionCookPackageContext()
+{
+}
+
+void FWorldPartitionCookPackageContext::RegisterPackageCookPackageGenerator(IWorldPartitionCookPackageGenerator* CookPackageGenerator)
+{
+	check(!CookPackageGenerators.Contains(CookPackageGenerator));
+	CookPackageGenerators.Add(CookPackageGenerator);
+	check(!PackagesToCookByGenerator.Contains(CookPackageGenerator));
+}
+
+const FWorldPartitionCookPackage* FWorldPartitionCookPackageContext::AddLevelStreamingPackageToGenerate(IWorldPartitionCookPackageGenerator* CookPackageGenerator, const FString& Root, const FString& RelativePath)
+{
+	return AddPackageToGenerateInternal(CookPackageGenerator, Root, RelativePath, FWorldPartitionCookPackage::EType::Level);
+}
+
+const FWorldPartitionCookPackage* FWorldPartitionCookPackageContext::AddGenericPackageToGenerate(IWorldPartitionCookPackageGenerator* CookPackageGenerator, const FString& Root, const FString& RelativePath)
+{
+	return AddPackageToGenerateInternal(CookPackageGenerator, Root, RelativePath, FWorldPartitionCookPackage::EType::Generic);
+}
+
+bool FWorldPartitionCookPackageContext::GetCookPackageGeneratorAndPackage(const FString& PackageRoot, const FString& PackageRelativePath, IWorldPartitionCookPackageGenerator*& CookPackageGenerator, FWorldPartitionCookPackage*& CookPackage)
+{
+	FWorldPartitionCookPackage::IDType PackageId = FWorldPartitionCookPackage::MakeCookPackageID(PackageRoot, PackageRelativePath);
+	if (IWorldPartitionCookPackageGenerator** GeneratorPtr = CookGeneratorByPackageId.Find(PackageId))
+	{
+		if (TUniquePtr<FWorldPartitionCookPackage>* PackagePtr = PackagesToCookById.Find(PackageId))
+		{
+			check((*PackagePtr)->Root.Equals(PackageRoot, ESearchCase::IgnoreCase) && (*PackagePtr)->RelativePath.Equals(PackageRelativePath, ESearchCase::IgnoreCase));
+
+			CookPackageGenerator = *GeneratorPtr;
+			CookPackage = (*PackagePtr).Get();
+			return true;
+		}
+	}
+
+	return false;
+}
+
+const FWorldPartitionCookPackage* FWorldPartitionCookPackageContext::AddPackageToGenerateInternal(IWorldPartitionCookPackageGenerator* CookPackageGenerator, const FString& Root, const FString& RelativePath, FWorldPartitionCookPackage::EType Type)
+{
+	if (CookPackageGenerators.Contains(CookPackageGenerator))
+	{
+		FWorldPartitionCookPackage::IDType PackageId = FWorldPartitionCookPackage::MakeCookPackageID(Root, RelativePath);
+		TUniquePtr<FWorldPartitionCookPackage>* ExistingPackage = PackagesToCookById.Find(PackageId);
+		if (ExistingPackage == nullptr)
+		{
+			TUniquePtr<FWorldPartitionCookPackage>& CookPackage = PackagesToCookById.Emplace(PackageId, MakeUnique<FWorldPartitionCookPackage>(Root, RelativePath, Type));
+			check(PackageId == CookPackage->PackageId);
+
+			CookGeneratorByPackageId.Add(PackageId, CookPackageGenerator);
+
+			TArray<FWorldPartitionCookPackage*>& PackagesToCookForHandler = PackagesToCookByGenerator.FindOrAdd(CookPackageGenerator);
+			PackagesToCookForHandler.Add(CookPackage.Get());
+
+			UE_LOG(LogWorldPartition, Verbose, TEXT("[Cook] Added Package %s with ID %llu in context"), *CookPackage->GetFullPath(), PackageId);
+
+			return CookPackage.Get();
+		}
+		else
+		{
+			UE_LOG(LogWorldPartition, Error, TEXT("[Cook] Trying to add package %s in context but there is already a package to generate with the same ID (%llu). Other package: %s Id %llu"),
+				*FWorldPartitionCookPackage::MakeFullPath(Root, RelativePath), PackageId, *(*ExistingPackage)->GetFullPath(), (*ExistingPackage)->PackageId);
+		}
+	}
+	else
+	{
+		UE_LOG(LogWorldPartition, Error, TEXT("[Cook] Trying to add package %s in context, but its generator is not registered."), *FWorldPartitionCookPackage::MakeFullPath(Root, RelativePath));
+	}
+
+	return nullptr;
+}
+
+#endif
\ No newline at end of file
diff --git a/Engine/Source/Runtime/Engine/Private/WorldPartition/Cook/WorldPartitionCookPackageContext.h b/Engine/Source/Runtime/Engine/Private/WorldPartition/Cook/WorldPartitionCookPackageContext.h
new file mode 100644
--- /dev/null
+++ b/Engine/Source/Runtime/Engine/Private/WorldPartition/Cook/WorldPartitionCookPackageContext.h
@@ -0,0 +1,45 @@
+// Copyright Epic Games, Inc. All Rights Reserved.
+
+#pragma once
+
+#if WITH_EDITOR
+
+#include "CoreMinimal.h"
+#include "WorldPartition/Cook/WorldPartitionCookPackageContextInterface.h"
+#include "WorldPartition/Cook/WorldPartitionCookPackage.h"
+
+class IWorldPartitionCookPackageGenerator;
+
+class FWorldPartitionCookPackageContext : public IWorldPartitionCookPackageContext
+{
+public:
+	FWorldPartitionCookPackageContext();
+
+	//~ Begin IWorldPartitionCookPackageContext Interface 
+	virtual void RegisterPackageCookPackageGenerator(IWorldPartitionCookPackageGenerator* CookPackageGenerator) override;
+	virtual void UnregisterPackageCookPackageGenerator(IWorldPartitionCookPackageGenerator* CookPackageGenerator) override { check(0); /*No use case*/ }
+
+	virtual const FWorldPartitionCookPackage* AddLevelStreamingPackageToGenerate(IWorldPartitionCookPackageGenerator* CookPackageGenerator, const FString& Root, const FString& RelativePath) override;
+	virtual const FWorldPartitionCookPackage* AddGenericPackageToGenerate(IWorldPartitionCookPackageGenerator* CookPackageGenerator, const FString& Root, const FString& RelativePath) override;
+	//~ End IWorldPartitionCookPackageContext Interface 
+
+	const TArray<FWorldPartitionCookPackage*>* GetCookPackages(const IWorldPartitionCookPackageGenerator* CookPackageGenerator) const { return PackagesToCookByGenerator.Find(CookPackageGenerator); }
+
+	bool GetCookPackageGeneratorAndPackage(const FString& PackageRoot, const FString& PackageRelativePath, IWorldPartitionCookPackageGenerator*& CookPackageGenerator, FWorldPartitionCookPackage*& CookPackage);
+
+	uint32 NumPackageToGenerate() const { return PackagesToCookById.Num(); }
+	uint32 NumGenerators() const { return CookPackageGenerators.Num(); }
+
+	TArray<IWorldPartitionCookPackageGenerator*>& GetCookPackageGenerators() { return CookPackageGenerators; }
+	const TArray<IWorldPartitionCookPackageGenerator*>& GetCookPackageGenerators() const { return CookPackageGenerators; }
+
+private:
+	const FWorldPartitionCookPackage* AddPackageToGenerateInternal(IWorldPartitionCookPackageGenerator* CookPackageGenerator, const FString& Root, const FString& RelativePath, FWorldPartitionCookPackage::EType Type);
+
+	TArray<IWorldPartitionCookPackageGenerator*> CookPackageGenerators;
+	TMap<FWorldPartitionCookPackage::IDType, TUniquePtr<FWorldPartitionCookPackage>> PackagesToCookById;
+	TMap<FWorldPartitionCookPackage::IDType, IWorldPartitionCookPackageGenerator*> CookGeneratorByPackageId;
+	TMap<IWorldPartitionCookPackageGenerator*, TArray<FWorldPartitionCookPackage*>> PackagesToCookByGenerator;
+};
+
+#endif
\ No newline at end of file
diff --git a/Engine/Source/Runtime/Engine/Private/WorldPartition/WorldPartition.cpp b/Engine/Source/Runtime/Engine/Private/WorldPartition/WorldPartition.cpp
--- a/Engine/Source/Runtime/Engine/Private/WorldPartition/WorldPartition.cpp
+++ b/Engine/Source/Runtime/Engine/Private/WorldPartition/WorldPartition.cpp
@@ -52,6 +52,8 @@
 #include "WorldPartition/LoaderAdapter/LoaderAdapterShape.h"
 #include "WorldPartition/WorldPartitionActorDescViewProxy.h"
 #include "WorldPartition/HLOD/HLODLayer.h"
+#include "WorldPartition/Cook/WorldPartitionCookPackage.h"
+#include "WorldPartition/Cook/WorldPartitionCookPackageContextInterface.h"
 #include "Modules/ModuleManager.h"
 #include "GameDelegates.h"
 #include "IDocumentation.h"
@@ -1223,16 +1225,38 @@ void UWorldPartition::DrawRuntimeHashPreview()
 	RuntimeHash->DrawPreview();
 }
 
-bool UWorldPartition::PopulateGeneratorPackageForCook(const TArray<ICookPackageSplitter::FGeneratedPackageForPreSave>& InGeneratedPackages, TArray<UPackage*>& OutModifiedPackages)
+void UWorldPartition::BeginCook(IWorldPartitionCookPackageContext& CookContext)
+{
+	CookContext.RegisterPackageCookPackageGenerator(this);
+}
+
+bool UWorldPartition::GatherPackagesToCook(IWorldPartitionCookPackageContext& CookContext)
+{
+	TArray<FString> PackagesToCook;
+	if (GenerateContainerStreaming(this, &PackagesToCook))
+	{
+		FString PackageName = GetPackage()->GetName();
+		for (const FString& PackageToCook : PackagesToCook)
+		{
+			CookContext.AddLevelStreamingPackageToGenerate(this, PackageName, PackageToCook);
+		}
+	
+		return true;
+	}
+
+	return false;
+}
+
+bool UWorldPartition::PopulateGeneratorPackageForCook(IWorldPartitionCookPackageContext& CookContext, const TArray<FWorldPartitionCookPackage*>& InPackagesToCook, TArray<UPackage*>& OutModifiedPackages)
 {
 	check(RuntimeHash);
-	return RuntimeHash->PopulateGeneratorPackageForCook(InGeneratedPackages, OutModifiedPackages);
+	return RuntimeHash->PopulateGeneratorPackageForCook(InPackagesToCook, OutModifiedPackages);
 }
 
-bool UWorldPartition::PopulateGeneratedPackageForCook(UPackage* InPackage, const FString& InPackageRelativePath, TArray<UPackage*>& OutModifiedPackages)
+bool UWorldPartition::PopulateGeneratedPackageForCook(IWorldPartitionCookPackageContext& CookContext, const FWorldPartitionCookPackage& InPackagesToCook, TArray<UPackage*>& OutModifiedPackages)
 {
 	check(RuntimeHash);
-	return RuntimeHash->PopulateGeneratedPackageForCook(InPackage, InPackageRelativePath, OutModifiedPackages);
+	return RuntimeHash->PopulateGeneratedPackageForCook(InPackagesToCook, OutModifiedPackages);
 }
 
 TArray<FBox> UWorldPartition::GetUserLoadedEditorRegions() const
diff --git a/Engine/Source/Runtime/Engine/Private/WorldPartition/WorldPartitionCookPackageSplitter.cpp b/Engine/Source/Runtime/Engine/Private/WorldPartition/WorldPartitionCookPackageSplitter.cpp
--- a/Engine/Source/Runtime/Engine/Private/WorldPartition/WorldPartitionCookPackageSplitter.cpp
+++ b/Engine/Source/Runtime/Engine/Private/WorldPartition/WorldPartitionCookPackageSplitter.cpp
@@ -11,6 +11,7 @@
 #include "WorldPartition/WorldPartitionRuntimeHash.h"
 #include "WorldPartition/WorldPartitionSubsystem.h"
 #include "WorldPartition/WorldPartition.h"
+#include "WorldPartition/WorldPartitionLog.h"
 #include "Editor.h"
 
 // Register FWorldPartitionCookPackageSplitter for UWorld class
@@ -101,20 +102,21 @@ TArray<ICookPackageSplitter::FGeneratedPackage> FWorldPartitionCookPackageSplitt
 	WorldPartition->Initialize(PartitionedWorld, FTransform::Identity);
 	bInitializedWorldPartition = true;
 
-	TArray<FString> WorldPartitionGeneratedPackages;
-	WorldPartition->GenerateStreaming(&WorldPartitionGeneratedPackages);
+	WorldPartition->BeginCook(CookContext);
 
-	TArray<ICookPackageSplitter::FGeneratedPackage> PackagesToGenerate;
-	PackagesToGenerate.Reserve(WorldPartitionGeneratedPackages.Num());
-	for (const FString& PackageName : WorldPartitionGeneratedPackages)
+	bool bIsSuccess = true;
+	for (IWorldPartitionCookPackageGenerator* CookPackageGenerator : CookContext.GetCookPackageGenerators())
 	{
-		ICookPackageSplitter::FGeneratedPackage& GeneratedPackage = PackagesToGenerate.Emplace_GetRef();
-		GeneratedPackage.RelativePath = PackageName;
-		GeneratedPackage.GeneratedRootPath = OwnerPackage->GetName();
-		// all packages we generate get a ULevel from CreateEmptyLevelForRuntimeCell and are hence maps
-		GeneratedPackage.SetCreateAsMap(true);
-		// @todo_ow: Set dependencies once we get iterative cooking working
+		bIsSuccess &= CookPackageGenerator->GatherPackagesToCook(CookContext);
 	}
+
+	UE_LOG(LogWorldPartition, Log, TEXT("[Cook] Gathered %u packages to generate from %u Generators."), CookContext.NumPackageToGenerate(), CookContext.NumGenerators());
+
+	TArray<ICookPackageSplitter::FGeneratedPackage> PackagesToGenerate;
+	BuildPackagesToGenerateList(PackagesToGenerate);
+
+	UE_LOG(LogWorldPartition, Log, TEXT("[Cook] Sending %u packages to be generated."), PackagesToGenerate.Num());
+
 	return PackagesToGenerate;
 }
 
@@ -122,23 +124,106 @@ bool FWorldPartitionCookPackageSplitter::PopulateGeneratedPackage(UPackage* Owne
 	const FGeneratedPackageForPopulate& GeneratedPackage, TArray<UObject*>& OutObjectsToMove,
 	TArray<UPackage*>& OutModifiedPackages)
 {
-	UWorld* PartitionedWorld = ValidateDataObject(OwnerObject);
-	UWorldPartition* WorldPartition = PartitionedWorld->PersistentLevel->GetWorldPartition();
-	return WorldPartition->PopulateGeneratedPackageForCook(GeneratedPackage.Package, GeneratedPackage.RelativePath, OutModifiedPackages);
+	UE_LOG(LogWorldPartition, Verbose, TEXT("[Cook][PopulateGeneratedPackage] Processing %s"), *FWorldPartitionCookPackage::MakeFullPath(GeneratedPackage.GeneratedRootPath, GeneratedPackage.RelativePath));
+
+	bool bIsSuccess = true;
+
+	IWorldPartitionCookPackageGenerator* CookPackageGenerator = nullptr;
+	FWorldPartitionCookPackage* CookPackage = nullptr;
+	TArray<UPackage*> ModifiedPackages;
+	if (CookContext.GetCookPackageGeneratorAndPackage(GeneratedPackage.GeneratedRootPath, GeneratedPackage.RelativePath, CookPackageGenerator, CookPackage))
+	{
+		bIsSuccess = CookPackageGenerator->PopulateGeneratedPackageForCook(CookContext, *CookPackage, ModifiedPackages);
+	}
+	else
+	{
+		UE_LOG(LogWorldPartition, Error, TEXT("[Cook][PopulateGeneratedPackage] Could not find WorldPartitionCookPackage for %s"), *FWorldPartitionCookPackage::MakeFullPath(GeneratedPackage.GeneratedRootPath, GeneratedPackage.RelativePath));
+		bIsSuccess = false;
+	}
+
+	UE_LOG(LogWorldPartition, Verbose, TEXT("[Cook][PopulateGeneratedPackage] Gathered %u modified packages for %s"), ModifiedPackages.Num() , *FWorldPartitionCookPackage::MakeFullPath(GeneratedPackage.GeneratedRootPath, GeneratedPackage.RelativePath));
+	OutModifiedPackages = MoveTemp(ModifiedPackages);
+
+	return bIsSuccess;
 }
 
 bool FWorldPartitionCookPackageSplitter::PopulateGeneratorPackage(UPackage* OwnerPackage, UObject* OwnerObject,
 	const TArray<ICookPackageSplitter::FGeneratedPackageForPreSave>& GeneratedPackages, TArray<UObject*>& OutObjectsToMove,
 	TArray<UPackage*>& OutModifiedPackages)
 {
-	UWorld* PartitionedWorld = ValidateDataObject(OwnerObject);
-	UWorldPartition* WorldPartition = PartitionedWorld->PersistentLevel->GetWorldPartition();
-	return WorldPartition->PopulateGeneratorPackageForCook(GeneratedPackages, OutModifiedPackages);
+	UE_LOG(LogWorldPartition, Log, TEXT("[Cook][PopulateGeneratorPackage] Processing %u packages"), GeneratedPackages.Num());
+
+	bool bIsSuccess = true;
+	if (GeneratedPackages.Num() != CookContext.NumPackageToGenerate())
+	{
+		UE_LOG(LogWorldPartition, Error, TEXT("[Cook][PopulateGeneratorPackage] Receieved %u generated packages. Was expecting %u"), GeneratedPackages.Num(), CookContext.NumPackageToGenerate());
+		bIsSuccess = false;
+	}
+
+
+	bIsSuccess &= MapGeneratePackageToCookPackage(GeneratedPackages);
+
+	TArray<UPackage*> ModifiedPackages;
+	for (IWorldPartitionCookPackageGenerator* CookPackageGenerator : CookContext.GetCookPackageGenerators())
+	{
+		if (const TArray<FWorldPartitionCookPackage*>* CookPackages = CookContext.GetCookPackages(CookPackageGenerator))
+		{
+			bIsSuccess &= CookPackageGenerator->PopulateGeneratorPackageForCook(CookContext, *CookPackages, ModifiedPackages);
+		}
+	}
+
+	UE_LOG(LogWorldPartition, Log, TEXT("[Cook][PopulateGeneratorPackage] Gathered %u modified packages"), ModifiedPackages.Num());
+	OutModifiedPackages = MoveTemp(ModifiedPackages);
+
+	return bIsSuccess;
 }
 
 void FWorldPartitionCookPackageSplitter::OnOwnerReloaded(UPackage* OwnerPackage, UObject* OwnerObject)
 {
 	// It should not be possible for the owner to reload due to garbage collection while we are active and keeping it referenced
-	check(!ReferencedWorld); 
+	check(!ReferencedWorld);
+}
+
+void FWorldPartitionCookPackageSplitter::BuildPackagesToGenerateList(TArray<ICookPackageSplitter::FGeneratedPackage>& PackagesToGenerate) const
+{
+	for (const IWorldPartitionCookPackageGenerator* CookPackageGenerator : CookContext.GetCookPackageGenerators())
+	{
+		if (const TArray<FWorldPartitionCookPackage*>* CookPackages = CookContext.GetCookPackages(CookPackageGenerator))
+		{
+			PackagesToGenerate.Reserve(CookPackages->Num());
+
+			for (const FWorldPartitionCookPackage* CookPackage : *CookPackages)
+			{
+				ICookPackageSplitter::FGeneratedPackage& GeneratedPackage = PackagesToGenerate.Emplace_GetRef();
+				GeneratedPackage.GeneratedRootPath = CookPackage->Root;
+				GeneratedPackage.RelativePath = CookPackage->RelativePath;
+
+				CookPackage->Type == FWorldPartitionCookPackage::EType::Level ? GeneratedPackage.SetCreateAsMap(true) : GeneratedPackage.SetCreateAsMap(false);
+
+				// @todo_ow: Set dependencies once we get iterative cooking working
+			}
+		}
+	}
+}
+
+bool FWorldPartitionCookPackageSplitter::MapGeneratePackageToCookPackage(const TArray<ICookPackageSplitter::FGeneratedPackageForPreSave>& GeneratedPackages)
+{
+	bool bIsSucces = true;
+	for (const ICookPackageSplitter::FGeneratedPackageForPreSave& GeneratedPackage : GeneratedPackages)
+	{
+		IWorldPartitionCookPackageGenerator* CookPackageGenerator = nullptr;
+		FWorldPartitionCookPackage* CookPackage = nullptr;
+		if (CookContext.GetCookPackageGeneratorAndPackage(GeneratedPackage.GeneratedRootPath, GeneratedPackage.RelativePath, CookPackageGenerator, CookPackage))
+		{
+			CookPackage->GeneratedPackage = GeneratedPackage.Package;
+		}
+		else
+		{
+			UE_LOG(LogWorldPartition, Error, TEXT("[Cook][MapGeneratePackageToCookPackage] Could not find WorldPartitionCookPackage for %s"), *FWorldPartitionCookPackage::MakeFullPath(GeneratedPackage.GeneratedRootPath, GeneratedPackage.RelativePath));
+			bIsSucces = false;
+		}
+	}
+
+	return bIsSucces;
 }
 #endif
diff --git a/Engine/Source/Runtime/Engine/Private/WorldPartition/WorldPartitionCookPackageSplitter.h b/Engine/Source/Runtime/Engine/Private/WorldPartition/WorldPartitionCookPackageSplitter.h
--- a/Engine/Source/Runtime/Engine/Private/WorldPartition/WorldPartitionCookPackageSplitter.h
+++ b/Engine/Source/Runtime/Engine/Private/WorldPartition/WorldPartitionCookPackageSplitter.h
@@ -7,6 +7,7 @@
 #include "CookPackageSplitter.h"
 #include "Engine/World.h"
 #include "UObject/WeakObjectPtrTemplates.h"
+#include "WorldPartition/Cook/WorldPartitionCookPackageContext.h"
 
 class FWorldPartitionCookPackageSplitter : public FGCObject, public ICookPackageSplitter
 {
@@ -35,9 +36,14 @@ private:
 
 	const UWorld* ValidateDataObject(const UObject* SplitData);
 	UWorld* ValidateDataObject(UObject* SplitData);
+
+	void BuildPackagesToGenerateList(TArray<ICookPackageSplitter::FGeneratedPackage>& PackagesToGenerate) const;
+	bool MapGeneratePackageToCookPackage(const TArray<ICookPackageSplitter::FGeneratedPackageForPreSave>& GeneratedPackages);
 	
 	TObjectPtr<UWorld> ReferencedWorld = nullptr;
 
+	FWorldPartitionCookPackageContext CookContext;
+
 	bool bInitializedWorldPartition = false;
 	bool bForceInitializedWorld = false;
 	bool bInitializedPhysicsSceneForSave = false;
diff --git a/Engine/Source/Runtime/Engine/Private/WorldPartition/WorldPartitionRuntimeSpatialHash.cpp b/Engine/Source/Runtime/Engine/Private/WorldPartition/WorldPartitionRuntimeSpatialHash.cpp
--- a/Engine/Source/Runtime/Engine/Private/WorldPartition/WorldPartitionRuntimeSpatialHash.cpp
+++ b/Engine/Source/Runtime/Engine/Private/WorldPartition/WorldPartitionRuntimeSpatialHash.cpp
@@ -52,6 +52,9 @@
 #include "Misc/UObjectToken.h"
 #include "Misc/MapErrors.h"
 
+#include "WorldPartition/Cook/WorldPartitionCookPackageContextInterface.h"
+#include "WorldPartition/Cook/WorldPartitionCookPackage.h"
+
 extern UNREALED_API class UEditorEngine* GEditor;
 #endif //WITH_EDITOR
 
@@ -1151,15 +1154,15 @@ bool UWorldPartitionRuntimeSpatialHash::CreateStreamingGrid(const FSpatialHashRu
 	return true;
 }
 
-bool UWorldPartitionRuntimeSpatialHash::PopulateGeneratedPackageForCook(UPackage* InPackage, const FString& InPackageRelativePath, TArray<UPackage*>& OutModifiedPackages)
+bool UWorldPartitionRuntimeSpatialHash::PopulateGeneratedPackageForCook(const FWorldPartitionCookPackage& InPackagesToCook, TArray<UPackage*>& OutModifiedPackages)
 {
 	OutModifiedPackages.Reset();
-	if (UWorldPartitionRuntimeCell** MatchingCell = PackagesToGenerateForCook.Find(InPackageRelativePath))
+	if (UWorldPartitionRuntimeCell** MatchingCell = PackagesToGenerateForCook.Find(InPackagesToCook.RelativePath))
 	{
 		UWorldPartitionRuntimeCell* Cell = *MatchingCell;
 		if (ensure(Cell))
 		{
-			return Cell->PopulateGeneratedPackageForCook(InPackage, OutModifiedPackages);
+			return Cell->PopulateGeneratedPackageForCook(InPackagesToCook.GeneratedPackage, OutModifiedPackages);
 		}
 	}
 	return false;
@@ -1180,7 +1183,7 @@ TArray<UWorldPartitionRuntimeCell*> UWorldPartitionRuntimeSpatialHash::GetAlways
 	return Result;
 }
 
-bool UWorldPartitionRuntimeSpatialHash::PopulateGeneratorPackageForCook(const TArray<ICookPackageSplitter::FGeneratedPackageForPreSave>& InGeneratedPackages, TArray<UPackage*>& OutModifiedPackages)
+bool UWorldPartitionRuntimeSpatialHash::PopulateGeneratorPackageForCook(const TArray<FWorldPartitionCookPackage*>& InPackagesToCook, TArray<UPackage*>& OutModifiedPackages)
 {
 	check(IsRunningCookCommandlet());
 
@@ -1194,11 +1197,11 @@ bool UWorldPartitionRuntimeSpatialHash::PopulateGeneratorPackageForCook(const TA
 		}
 	}
 
-	for (const ICookPackageSplitter::FGeneratedPackageForPreSave& GeneratedPackage : InGeneratedPackages)
+	for (const FWorldPartitionCookPackage* CookPackage : InPackagesToCook)
 	{
-		UWorldPartitionRuntimeCell** MatchingCell = PackagesToGenerateForCook.Find(GeneratedPackage.RelativePath);
+		UWorldPartitionRuntimeCell** MatchingCell = PackagesToGenerateForCook.Find(CookPackage->RelativePath);
 		UWorldPartitionRuntimeCell* Cell = MatchingCell ? *MatchingCell : nullptr;
-		if (!Cell || !Cell->PrepareCellForCook(GeneratedPackage.Package))
+		if (!Cell || !Cell->PrepareCellForCook(CookPackage->GeneratedPackage))
 		{
 			return false;
 		}
diff --git a/Engine/Source/Runtime/Engine/Private/WorldPartition/WorldPartitionStreamingGeneration.cpp b/Engine/Source/Runtime/Engine/Private/WorldPartition/WorldPartitionStreamingGeneration.cpp
--- a/Engine/Source/Runtime/Engine/Private/WorldPartition/WorldPartitionStreamingGeneration.cpp
+++ b/Engine/Source/Runtime/Engine/Private/WorldPartition/WorldPartitionStreamingGeneration.cpp
@@ -856,7 +856,12 @@ private:
 	TUniquePtr<FStreamingGenerationContext> StreamingGenerationContext;
 };
 
-bool UWorldPartition::GenerateStreaming(TArray<FString>* OutPackagesToGenerate)
+bool UWorldPartition::GenerateStreaming()
+{
+	return GenerateContainerStreaming(this);
+}
+
+bool UWorldPartition::GenerateContainerStreaming(const UActorDescContainer* InActorDescContainer, TArray<FString>* OutPackagesToGenerate /* = nullptr */)
 {
 	FActorDescList* ModifiedActorsDescList = nullptr;
 
diff --git a/Engine/Source/Runtime/Engine/Public/WorldPartition/Cook/WorldPartitionCookPackage.h b/Engine/Source/Runtime/Engine/Public/WorldPartition/Cook/WorldPartitionCookPackage.h
new file mode 100644
--- /dev/null
+++ b/Engine/Source/Runtime/Engine/Public/WorldPartition/Cook/WorldPartitionCookPackage.h
@@ -0,0 +1,79 @@
+// Copyright Epic Games, Inc. All Rights Reserved.
+#pragma once
+
+#if WITH_EDITOR
+
+#include "CoreMinimal.h"
+#include "Hash/CityHash.h"
+#include "Misc/Paths.h"
+#include "WorldPartition/WorldPartitionLog.h"
+
+struct FWorldPartitionCookPackage
+{
+	enum class EType
+	{
+		Unknown,
+		Level,
+		Generic
+	};
+
+	using IDType = uint64;
+	
+	static IDType MakeCookPackageID(const FString& InRoot, const FString& InRelativeFilename)
+	{
+		check(!InRoot.IsEmpty() && !InRelativeFilename.IsEmpty());
+		check(InRoot[0] == '/' && InRoot[InRoot.Len() - 1] != '/'); // Root is assumed to be in the format "/Root"
+		check(InRelativeFilename[0] == '/' && InRelativeFilename[InRelativeFilename.Len() - 1] != '/'); // RelativeFileName is assumed to be in the format "/Relativefilename
+
+		// Avoid doing string copies as this function is often called during cook when bridging between Cook code & WorldPartition code.
+		// Compute a hash for both InRoot & InRelativeFilename. Then combine them instead of creating a new fullpath string and computing the hash on it.
+		uint64 RootHash = CityHash64(reinterpret_cast<const char*>(*InRoot), InRoot.Len() * sizeof(TCHAR));
+		uint64 RelativePathHash = CityHash64(reinterpret_cast<const char*>(*InRelativeFilename), InRelativeFilename.Len() * sizeof(TCHAR));
+		return RootHash ^ RelativePathHash;
+	}
+
+	static FString MakeFullPath(const FString& InRoot, const FString& InRelativeFilename)
+	{
+		TStringBuilderWithBuffer<TCHAR, NAME_SIZE> FullPath;
+		FullPath += TEXT("/");
+		FullPath += InRoot;
+		FullPath += TEXT("/");
+		FullPath += InRelativeFilename;
+
+		return FPaths::RemoveDuplicateSlashes(*FullPath);
+	}
+
+	FWorldPartitionCookPackage(const FString& InRoot, const FString& InRelativePath, EType InType)
+		: Root(SanitizePathComponent(InRoot)),
+		RelativePath(SanitizePathComponent(InRelativePath)),
+		GeneratedPackage(nullptr),
+		PackageId(MakeCookPackageID(Root, RelativePath)),
+		Type(InType)
+	{
+	}
+
+	FString GetFullPath() const { return MakeFullPath(Root, RelativePath); }
+
+	const FString Root;
+	const FString RelativePath;
+	UPackage* GeneratedPackage;
+	const IDType PackageId;
+	const EType Type;
+
+private:
+	// PathComponents (Root & RelativePath members) need to follow the "/<PathComponent>" format for the PackageId computation to work.
+	FString SanitizePathComponent(const FString& Path) const
+	{
+		FString SanitizedPath = TEXT("/") + Path;
+		FPaths::RemoveDuplicateSlashes(SanitizedPath);
+		
+		if (SanitizedPath[SanitizedPath.Len() - 1] == '/')
+		{
+			SanitizedPath.RemoveAt(SanitizedPath.Len() - 1, 1);
+		}
+
+		return SanitizedPath;
+	}
+};
+
+#endif
diff --git a/Engine/Source/Runtime/Engine/Public/WorldPartition/Cook/WorldPartitionCookPackageContextInterface.h b/Engine/Source/Runtime/Engine/Public/WorldPartition/Cook/WorldPartitionCookPackageContextInterface.h
new file mode 100644
--- /dev/null
+++ b/Engine/Source/Runtime/Engine/Public/WorldPartition/Cook/WorldPartitionCookPackageContextInterface.h
@@ -0,0 +1,24 @@
+// Copyright Epic Games, Inc. All Rights Reserved.
+#pragma once
+
+#if WITH_EDITOR
+
+#include "CoreMinimal.h"
+
+class IWorldPartitionCookPackageGenerator;
+struct FWorldPartitionCookPackage;
+
+class IWorldPartitionCookPackageContext
+{
+
+public:
+	virtual ~IWorldPartitionCookPackageContext() {};
+
+	virtual void RegisterPackageCookPackageGenerator(IWorldPartitionCookPackageGenerator* Generator) = 0;
+	virtual void UnregisterPackageCookPackageGenerator(IWorldPartitionCookPackageGenerator* Generator) = 0;
+
+	virtual const FWorldPartitionCookPackage* AddLevelStreamingPackageToGenerate(IWorldPartitionCookPackageGenerator* Generator, const FString& Root, const FString& RelativePath) = 0;
+	virtual const FWorldPartitionCookPackage* AddGenericPackageToGenerate(IWorldPartitionCookPackageGenerator* Generator, const FString& Root, const FString& RelativePath) = 0;
+};
+
+#endif
\ No newline at end of file
diff --git a/Engine/Source/Runtime/Engine/Public/WorldPartition/Cook/WorldPartitionCookPackageGenerator.h b/Engine/Source/Runtime/Engine/Public/WorldPartition/Cook/WorldPartitionCookPackageGenerator.h
new file mode 100644
--- /dev/null
+++ b/Engine/Source/Runtime/Engine/Public/WorldPartition/Cook/WorldPartitionCookPackageGenerator.h
@@ -0,0 +1,23 @@
+// Copyright Epic Games, Inc. All Rights Reserved.
+#pragma once
+
+#include "CoreMinimal.h"
+
+#if WITH_EDITOR
+
+struct FWorldPartitionCookPackage;
+class IWorldPartitionCookPackageContext;
+
+#endif
+
+class IWorldPartitionCookPackageGenerator
+{
+#if WITH_EDITOR
+public:
+	virtual ~IWorldPartitionCookPackageGenerator() = default;
+
+	virtual bool GatherPackagesToCook(IWorldPartitionCookPackageContext& CookContext) = 0;
+	virtual bool PopulateGeneratorPackageForCook(IWorldPartitionCookPackageContext& CookContext, const TArray<FWorldPartitionCookPackage*>& InPackagesToCook, TArray<UPackage*>& OutModifiedPackages) = 0;
+	virtual bool PopulateGeneratedPackageForCook(IWorldPartitionCookPackageContext& CookContext, const FWorldPartitionCookPackage& InPackagesToCook, TArray<UPackage*>& OutModifiedPackages) = 0;
+#endif
+};
diff --git a/Engine/Source/Runtime/Engine/Public/WorldPartition/WorldPartition.h b/Engine/Source/Runtime/Engine/Public/WorldPartition/WorldPartition.h
--- a/Engine/Source/Runtime/Engine/Public/WorldPartition/WorldPartition.h
+++ b/Engine/Source/Runtime/Engine/Public/WorldPartition/WorldPartition.h
@@ -13,6 +13,7 @@
 #include "WorldPartition/WorldPartitionStreamingSource.h"
 #include "WorldPartition/WorldPartitionHandle.h"
 #include "WorldPartition/ActorDescContainer.h"
+#include "WorldPartition/Cook/WorldPartitionCookPackageGenerator.h"
 
 #if WITH_EDITOR
 #include "WorldPartition/WorldPartitionActorLoaderInterface.h"
@@ -73,7 +74,7 @@ public:
 #endif
 
 UCLASS(AutoExpandCategories=(WorldPartition))
-class ENGINE_API UWorldPartition final : public UActorDescContainer
+class ENGINE_API UWorldPartition final : public UActorDescContainer, public IWorldPartitionCookPackageGenerator
 {
 	GENERATED_UCLASS_BODY()
 
@@ -146,14 +147,22 @@ public:
 	FName GetWorldPartitionEditorName() const;
 
 	// Streaming generation
-	bool GenerateStreaming(TArray<FString>* OutPackagesToGenerate = nullptr);
+	bool GenerateStreaming();
+	bool GenerateContainerStreaming(const UActorDescContainer* ActorDescContainer, TArray<FString>* OutPackagesToGenerate = nullptr);
 	void FlushStreaming();
 
 	void RemapSoftObjectPath(FSoftObjectPath& ObjectPath);
 
-	// Cooking
-	bool PopulateGeneratorPackageForCook(const TArray<ICookPackageSplitter::FGeneratedPackageForPreSave>& InGeneratedPackages, TArray<UPackage*>& OutModifiedPackages);
-	bool PopulateGeneratedPackageForCook(UPackage* InPackage, const FString& InPackageRelativePath, TArray<UPackage*>& OutModifiedPackages);
+	// Begin Cooking
+	void BeginCook(IWorldPartitionCookPackageContext& CookContext);
+
+	//~ Begin IWorldPartitionCookPackageGenerator Interface 
+	virtual bool GatherPackagesToCook(IWorldPartitionCookPackageContext& CookContext) override;
+	virtual bool PopulateGeneratorPackageForCook(IWorldPartitionCookPackageContext& CookContext, const TArray<FWorldPartitionCookPackage*>& InPackagesToCook, TArray<UPackage*>& OutModifiedPackages) override;
+	virtual bool PopulateGeneratedPackageForCook(IWorldPartitionCookPackageContext& CookContext, const FWorldPartitionCookPackage& InPackagesToCool, TArray<UPackage*>& OutModifiedPackages) override;
+	//~ End IWorldPartitionCookPackageGenerator Interface 
+
+	// End Cooking
 
 	UE_DEPRECATED(5.1, "GetWorldBounds is deprecated, use GetEditorWorldBounds or GetRuntimeWorldBounds instead.")
 	FBox GetWorldBounds() const { return GetRuntimeWorldBounds(); }
diff --git a/Engine/Source/Runtime/Engine/Public/WorldPartition/WorldPartitionRuntimeHash.h b/Engine/Source/Runtime/Engine/Public/WorldPartition/WorldPartitionRuntimeHash.h
--- a/Engine/Source/Runtime/Engine/Public/WorldPartition/WorldPartitionRuntimeHash.h
+++ b/Engine/Source/Runtime/Engine/Public/WorldPartition/WorldPartitionRuntimeHash.h
@@ -36,8 +36,8 @@ class ENGINE_API UWorldPartitionRuntimeHash : public UObject
 
 #if WITH_EDITOR
 	virtual void SetDefaultValues() {}
-	virtual bool PopulateGeneratorPackageForCook(const TArray<ICookPackageSplitter::FGeneratedPackageForPreSave>& InGeneratedPackages, TArray<UPackage*>& OutModifiedPackages) { return false; }
-	virtual bool PopulateGeneratedPackageForCook(UPackage* InPackage, const FString& InPackageRelativePath, TArray<UPackage*>& OutModifiedPackages) { return false; }
+	virtual bool PopulateGeneratorPackageForCook(const TArray<FWorldPartitionCookPackage*>& PackagesToCook, TArray<UPackage*>& OutModifiedPackages) { return false; }
+	virtual bool PopulateGeneratedPackageForCook(const FWorldPartitionCookPackage& PackagesToCook, TArray<UPackage*>& OutModifiedPackages) { return false; }
 	virtual bool GenerateStreaming(class UWorldPartitionStreamingPolicy* StreamingPolicy, const IStreamingGenerationContext* StreamingGenerationContext, TArray<FString>* OutPackagesToGenerate) { return false; }
 	virtual void FlushStreaming() {}
 	virtual bool GenerateHLOD(ISourceControlHelper* SourceControlHelper, const IStreamingGenerationContext* StreamingGenerationContext, bool bCreateActorsOnly) { return false; }
diff --git a/Engine/Source/Runtime/Engine/Public/WorldPartition/WorldPartitionRuntimeSpatialHash.h b/Engine/Source/Runtime/Engine/Public/WorldPartition/WorldPartitionRuntimeSpatialHash.h
--- a/Engine/Source/Runtime/Engine/Public/WorldPartition/WorldPartitionRuntimeSpatialHash.h
+++ b/Engine/Source/Runtime/Engine/Public/WorldPartition/WorldPartitionRuntimeSpatialHash.h
@@ -232,8 +232,8 @@ public:
 
 #if WITH_EDITOR
 	virtual void SetDefaultValues() override;
-	virtual bool PopulateGeneratorPackageForCook(const TArray<ICookPackageSplitter::FGeneratedPackageForPreSave>& InGeneratedPackages, TArray<UPackage*>& OutModifiedPackages) override;
-	virtual bool PopulateGeneratedPackageForCook(UPackage* InPackage, const FString& InPackageRelativePath, TArray<UPackage*>& OutModifiedPackages) override;
+	virtual bool PopulateGeneratorPackageForCook(const TArray<FWorldPartitionCookPackage*>& PackagesToCook, TArray<UPackage*>& OutModifiedPackage) override;
+	virtual bool PopulateGeneratedPackageForCook(const FWorldPartitionCookPackage& PackageToCook, TArray<UPackage*>& OutModifiedPackages) override;
 	virtual void FlushStreaming() override;
 	virtual bool GenerateHLOD(ISourceControlHelper* SourceControlHelper, const IStreamingGenerationContext* StreamingGenerationContext, bool bCreateActorsOnly) override;
 	virtual void DrawPreview() const override;
