Add missing copyright boilerplate

diff --git a/Engine/Source/Runtime/Engine/Private/WorldPartition/Cook/WorldPartitionCookPackageContext.cpp b/Engine/Source/Runtime/Engine/Private/WorldPartition/Cook/WorldPartitionCookPackageContext.cpp
--- a/Engine/Source/Runtime/Engine/Private/WorldPartition/Cook/WorldPartitionCookPackageContext.cpp
+++ b/Engine/Source/Runtime/Engine/Private/WorldPartition/Cook/WorldPartitionCookPackageContext.cpp
@@ -1,3 +1,4 @@
+// Copyright Epic Games, Inc. All Rights Reserved.
 #include "WorldPartition/Cook/WorldPartitionCookPackageContext.h"
 
 #if WITH_EDITOR
