Fixed issue with PerInstanceCustomData and PerInstanceCustomData3Vector nodes occurring when the DefaultValue pin does not connect all the way to valid node inputs.

diff --git a/Engine/Source/Runtime/Engine/Private/Materials/HLSLMaterialTranslator.cpp b/Engine/Source/Runtime/Engine/Private/Materials/HLSLMaterialTranslator.cpp
--- a/Engine/Source/Runtime/Engine/Private/Materials/HLSLMaterialTranslator.cpp
+++ b/Engine/Source/Runtime/Engine/Private/Materials/HLSLMaterialTranslator.cpp
@@ -11603,18 +11603,14 @@ int32 FHLSLMaterialTranslator::PerInstanceFadeAmount()
  */
 int32 FHLSLMaterialTranslator::PerInstanceCustomData(int32 DataIndex, int32 DefaultValueIndex)
 {
-	if (ShaderFrequency != SF_Vertex)
-	{
-		//return NonVertexShaderExpressionError();
-		// Only works with Nanite - TODO: Edge case the error
-		MaterialCompilationOutput.bUsesPerInstanceCustomData = true;
-		return AddInlinedCodeChunkZeroDeriv(MCT_Float, TEXT("GetPerInstanceCustomData(Parameters, %d, %s)"), DataIndex, *GetParameterCode(DefaultValueIndex));
-	}
-	else
+	if (DefaultValueIndex == INDEX_NONE)
 	{
-		MaterialCompilationOutput.bUsesPerInstanceCustomData = true;
-		return AddInlinedCodeChunkZeroDeriv(MCT_Float, TEXT("GetPerInstanceCustomData(Parameters, %d, %s)"), DataIndex, *GetParameterCode(DefaultValueIndex));
+		return INDEX_NONE;
 	}
+
+	// The case where ShaderFrequency is not SF_Vertex only works with Nanite - TODO: Edge case the error
+	MaterialCompilationOutput.bUsesPerInstanceCustomData = true;
+	return AddInlinedCodeChunkZeroDeriv(MCT_Float, TEXT("GetPerInstanceCustomData(Parameters, %d, %s)"), DataIndex, *GetParameterCode(DefaultValueIndex));
 }
 
 /**
@@ -11625,18 +11621,14 @@ int32 FHLSLMaterialTranslator::PerInstanceCustomData(int32 DataIndex, int32 Defa
  */
 int32 FHLSLMaterialTranslator::PerInstanceCustomData3Vector(int32 DataIndex, int32 DefaultValueIndex)
 {
-	if (ShaderFrequency != SF_Vertex)
-	{
-		//return NonVertexShaderExpressionError();
-		// Only works with Nanite - TODO: Edge case the error
-		MaterialCompilationOutput.bUsesPerInstanceCustomData = true;
-		return AddInlinedCodeChunkZeroDeriv(MCT_Float3, TEXT("GetPerInstanceCustomData3Vector(Parameters, %d, %s)"), DataIndex, *GetParameterCode(DefaultValueIndex));
-	}
-	else
+	if (DefaultValueIndex == INDEX_NONE)
 	{
-		MaterialCompilationOutput.bUsesPerInstanceCustomData = true;
-		return AddInlinedCodeChunkZeroDeriv(MCT_Float3, TEXT("GetPerInstanceCustomData3Vector(Parameters, %d, %s)"), DataIndex, *GetParameterCode(DefaultValueIndex));
+		return INDEX_NONE;
 	}
+
+	// The case when ShaderFrequency is not SF_Vertex only works with Nanite - TODO: Edge case the error
+	MaterialCompilationOutput.bUsesPerInstanceCustomData = true;
+	return AddInlinedCodeChunkZeroDeriv(MCT_Float3, TEXT("GetPerInstanceCustomData3Vector(Parameters, %d, %s)"), DataIndex, *GetParameterCode(DefaultValueIndex));
 }
 
 /**
