Added ENGINE_API to FUniformExpressionSet::GetSummaryString

diff --git a/Engine/Source/Runtime/Engine/Public/MaterialShared.h b/Engine/Source/Runtime/Engine/Public/MaterialShared.h
--- a/Engine/Source/Runtime/Engine/Public/MaterialShared.h
+++ b/Engine/Source/Runtime/Engine/Public/MaterialShared.h
@@ -567,7 +567,7 @@ public:
 
 	bool IsEmpty() const;
 	bool operator==(const FUniformExpressionSet& ReferenceSet) const;
-	FString GetSummaryString() const;
+	ENGINE_API FString GetSummaryString() const;
 
 	FShaderParametersMetadata* CreateBufferStruct();
 
