World Partition: first implementation of server-side streaming. Essentially, the server will not load everything but stream all player sources.
Can be enabled via the wp.Runtime.EnableServerStreaming cvar: set to 1 to enable server streaming, set to 2 to only enable it in PIE.
Tested in Lyra.

#rnx

diff --git a/Engine/Source/Runtime/Engine/Private/WorldPartition/WorldPartition.cpp b/Engine/Source/Runtime/Engine/Private/WorldPartition/WorldPartition.cpp
--- a/Engine/Source/Runtime/Engine/Private/WorldPartition/WorldPartition.cpp
+++ b/Engine/Source/Runtime/Engine/Private/WorldPartition/WorldPartition.cpp
@@ -1018,6 +1018,13 @@ void UWorldPartition::Serialize(FArchive& Ar)
 	if (Ar.GetPortFlags() & PPF_DuplicateForPIE)
 	{
 		Ar << StreamingPolicy;
+
+#if WITH_EDITOR
+		if (Ar.IsLoading())
+		{
+			bIsPIE = true;
+		}
+#endif
 	}
 	else if (Ar.CustomVer(FUE5MainStreamObjectVersion::GUID) >= FUE5MainStreamObjectVersion::WorldPartitionSerializeStreamingPolicyOnCook)
 	{
diff --git a/Engine/Source/Runtime/Engine/Private/WorldPartition/WorldPartitionStreamingPolicy.cpp b/Engine/Source/Runtime/Engine/Private/WorldPartition/WorldPartitionStreamingPolicy.cpp
--- a/Engine/Source/Runtime/Engine/Private/WorldPartition/WorldPartitionStreamingPolicy.cpp
+++ b/Engine/Source/Runtime/Engine/Private/WorldPartition/WorldPartitionStreamingPolicy.cpp
@@ -58,6 +58,21 @@ static FAutoConsoleVariableRef CVarBlockOnSlowStreaming(
 	GBlockOnSlowStreaming,
 	TEXT("Set if streaming needs to block when to slow to catchup."));
 
+static int32 GEnableServerStreaming = 0;
+static FAutoConsoleVariableRef CVarEnableServerStreaming(
+	TEXT("wp.Runtime.EnableServerStreaming"),
+	GEnableServerStreaming,
+	TEXT("Set to 1 to enable server streaming, set to 2 to only enable it in PIE."));
+
+static inline bool IsServerStreamingEnabled(bool bIsPIE)
+{
+	switch (GEnableServerStreaming)
+	{
+	case 1:	return true;
+	case 2:	return bIsPIE;
+	}
+	return false;
+}
 
 UWorldPartitionStreamingPolicy::UWorldPartitionStreamingPolicy(const FObjectInitializer& ObjectInitializer)
 : Super(ObjectInitializer) 
@@ -95,19 +110,22 @@ void UWorldPartitionStreamingPolicy::UpdateStreamingSources()
 
 	const FTransform WorldToLocal = WorldPartition->GetInstanceTransform().Inverse();
 	UWorld* World = WorldPartition->GetWorld();
-	bool bUseReplaySources = false;
-	if (AWorldPartitionReplay::IsPlaybackEnabled(World))
-	{
-		bUseReplaySources = WorldPartition->Replay->GetReplayStreamingSources(StreamingSources);
-	}
 
-	const ENetMode NetMode = World->GetNetMode();
-	if (!bUseReplaySources && (NetMode == NM_Standalone || NetMode == NM_Client || AWorldPartitionReplay::IsRecordingEnabled(World)))
+#if WITH_EDITOR
+	const bool bIsPIE = WorldPartition->bIsPIE;
+#else
+	const bool bIsPIE = false;
+#endif
+	const bool bIsServerStreamingEnabled = IsServerStreamingEnabled(bIsPIE);
+
+	if (!AWorldPartitionReplay::IsPlaybackEnabled(World) || !WorldPartition->Replay->GetReplayStreamingSources(StreamingSources))
 	{
+		const ENetMode NetMode = World->GetNetMode();
+
 #if WITH_EDITOR
-		// We are in the SIE
 		if (GEnableSimulationStreamingSource && UWorldPartition::IsSimulating())
 		{
+			// We are in the SIE
 			const FVector ViewLocation = WorldToLocal.TransformPosition(GCurrentLevelEditingViewportClient->GetViewLocation());
 			const FRotator ViewRotation = WorldToLocal.TransformRotation(GCurrentLevelEditingViewportClient->GetViewRotation().Quaternion()).Rotator();
 			static const FName NAME_SIE(TEXT("SIE"));
@@ -115,41 +133,47 @@ void UWorldPartitionStreamingPolicy::UpdateStreamingSources()
 		}
 		else
 #endif
+		if (NetMode == NM_Standalone || NetMode == NM_Client || bIsServerStreamingEnabled || AWorldPartitionReplay::IsRecordingEnabled(World))
 		{
-			const int32 NumPlayers = GEngine->GetNumGamePlayers(World);
-			for (int32 PlayerIndex = 0; PlayerIndex < NumPlayers; ++PlayerIndex)
+			for (FConstPlayerControllerIterator It = World->GetPlayerControllerIterator(); It; ++It)
 			{
-				ULocalPlayer* Player = GEngine->GetGamePlayer(World, PlayerIndex);
-				if (Player && Player->PlayerController && Player->PlayerController->IsStreamingSourceEnabled())
+				if (APlayerController* PlayerController = It->Get())
 				{
-					FVector ViewLocation;
-					FRotator ViewRotation;
-					Player->PlayerController->GetPlayerViewPoint(ViewLocation, ViewRotation);
-
-					// Transform to Local
-					ViewLocation = WorldToLocal.TransformPosition(ViewLocation);
-					ViewRotation = WorldToLocal.TransformRotation(ViewRotation.Quaternion()).Rotator();
-					const EStreamingSourceTargetState TargetState = Player->PlayerController->StreamingSourceShouldActivate() ? EStreamingSourceTargetState::Activated : EStreamingSourceTargetState::Loaded;
-					const bool bBlockOnSlowLoading = Player->PlayerController->StreamingSourceShouldBlockOnSlowStreaming();
-					const EStreamingSourcePriority StreamingSourcePriority = Player->PlayerController->GetStreamingSourcePriority();
-					StreamingSources.Add(FWorldPartitionStreamingSource(Player->PlayerController->GetFName(), ViewLocation, ViewRotation, TargetState, bBlockOnSlowLoading, StreamingSourcePriority));
+					if (NetMode == NM_DedicatedServer || NetMode == NM_ListenServer || PlayerController->IsLocalController())
+					{
+						if (PlayerController->IsStreamingSourceEnabled())
+						{
+							FVector ViewLocation;
+							FRotator ViewRotation;
+							PlayerController->GetPlayerViewPoint(ViewLocation, ViewRotation);
+
+							// Transform to Local
+							ViewLocation = WorldToLocal.TransformPosition(ViewLocation);
+							ViewRotation = WorldToLocal.TransformRotation(ViewRotation.Quaternion()).Rotator();
+							const EStreamingSourceTargetState TargetState = PlayerController->StreamingSourceShouldActivate() ? EStreamingSourceTargetState::Activated : EStreamingSourceTargetState::Loaded;
+							const bool bBlockOnSlowLoading = PlayerController->StreamingSourceShouldBlockOnSlowStreaming();
+							const EStreamingSourcePriority StreamingSourcePriority = PlayerController->GetStreamingSourcePriority();
+							StreamingSources.Add(FWorldPartitionStreamingSource(PlayerController->GetFName(), ViewLocation, ViewRotation, TargetState, bBlockOnSlowLoading, StreamingSourcePriority));
+						}
+					}
 				}
 			}
-		}
 
-		UWorldPartitionSubsystem* WorldPartitionSubsystem = WorldPartition->GetWorld()->GetSubsystem<UWorldPartitionSubsystem>();
-		check(WorldPartitionSubsystem);
-		for (IWorldPartitionStreamingSourceProvider* StreamingSourceProvider : WorldPartitionSubsystem->GetStreamingSourceProviders())
-		{
-			FWorldPartitionStreamingSource StreamingSource;
-			// Default Streaming Source provider's priority to be less than those based on player controllers
-			StreamingSource.Priority = EStreamingSourcePriority::Low;
-			if (StreamingSourceProvider->GetStreamingSource(StreamingSource))
+			UWorldPartitionSubsystem* WorldPartitionSubsystem = WorldPartition->GetWorld()->GetSubsystem<UWorldPartitionSubsystem>();
+			check(WorldPartitionSubsystem);
+
+			for (IWorldPartitionStreamingSourceProvider* StreamingSourceProvider : WorldPartitionSubsystem->GetStreamingSourceProviders())
 			{
-				// Transform to Local
-				StreamingSource.Location = WorldToLocal.TransformPosition(StreamingSource.Location);
-				StreamingSource.Rotation = WorldToLocal.TransformRotation(StreamingSource.Rotation.Quaternion()).Rotator();
-				StreamingSources.Add(StreamingSource);
+				FWorldPartitionStreamingSource StreamingSource;
+				// Default Streaming Source provider's priority to be less than those based on player controllers
+				StreamingSource.Priority = EStreamingSourcePriority::Low;
+				if (StreamingSourceProvider->GetStreamingSource(StreamingSource))
+				{
+					// Transform to Local
+					StreamingSource.Location = WorldToLocal.TransformPosition(StreamingSource.Location);
+					StreamingSource.Rotation = WorldToLocal.TransformRotation(StreamingSource.Rotation.Quaternion()).Rotator();
+					StreamingSources.Add(StreamingSource);
+				}
 			}
 		}
 	}
@@ -233,9 +257,16 @@ void UWorldPartitionStreamingPolicy::UpdateStreamingState()
 	check(LoadStreamingCells.IsEmpty());
 
 	const ENetMode NetMode = World->GetNetMode();
+	const bool bClient = NetMode == NM_Standalone || NetMode == NM_Client;
+	const bool bServer = NetMode == NM_ListenServer || NetMode == NM_DedicatedServer;
+#if WITH_EDITOR
+	const bool bIsPIE = WorldPartition->bIsPIE;
+#else
+	const bool bIsPIE = false;
+#endif
+	const bool bIsServerStreamingEnabled = IsServerStreamingEnabled(bIsPIE);
 	
-	bool bClient = NetMode == NM_Standalone || NetMode == NM_Client || AWorldPartitionReplay::IsPlaybackEnabled(World);
-	if (bClient)
+	if (bClient || bIsServerStreamingEnabled || AWorldPartitionReplay::IsPlaybackEnabled(World))
 	{
 		// When world partition can't stream, all cells must be unloaded
 		if (WorldPartition->CanStream())
@@ -247,18 +278,18 @@ void UWorldPartitionStreamingPolicy::UpdateStreamingState()
 			LoadStreamingCells = LoadStreamingCells.Difference(ActivateStreamingCells);
 		}
 	}
-	else 
+
+	if (bServer)
 	{
 		// Server will activate all non data layer cells at first and then load/activate/unload data layer cells only when the data layer states change
-		if (DataLayersStatesServerEpoch == AWorldDataLayers::GetDataLayersStateEpoch())
+		if (!bIsServerStreamingEnabled && DataLayersStatesServerEpoch == AWorldDataLayers::GetDataLayersStateEpoch())
 		{
 			// Server as nothing to do early out
 			return; 
 		}
 
 		bUpdateEpoch = true;
-		const UDataLayerSubsystem* DataLayerSubsystem = WorldPartition->GetWorld()->GetSubsystem<UDataLayerSubsystem>();
-		
+
 		// Gather Client visible level names
 		if (const UNetDriver* NetDriver = World->GetNetDriver())
 		{
@@ -269,13 +300,18 @@ void UWorldPartitionStreamingPolicy::UpdateStreamingState()
 			}
 		}
 
-		// Non Data Layer Cells + Active Data Layers
-		WorldPartition->RuntimeHash->GetAllStreamingCells(ActivateStreamingCells, /*bAllDataLayers=*/ false, /*bDataLayersOnly=*/ false, DataLayerSubsystem->GetEffectiveActiveDataLayerNames());
-
-		// Loaded Data Layers Cells only
-		if (DataLayerSubsystem->GetEffectiveLoadedDataLayerNames().Num())
+		if (!bIsServerStreamingEnabled)
 		{
-			WorldPartition->RuntimeHash->GetAllStreamingCells(LoadStreamingCells, /*bAllDataLayers=*/ false, /*bDataLayersOnly=*/ true, DataLayerSubsystem->GetEffectiveLoadedDataLayerNames());
+			const UDataLayerSubsystem* DataLayerSubsystem = WorldPartition->GetWorld()->GetSubsystem<UDataLayerSubsystem>();
+
+			// Non Data Layer Cells + Active Data Layers
+			WorldPartition->RuntimeHash->GetAllStreamingCells(ActivateStreamingCells, /*bAllDataLayers=*/ false, /*bDataLayersOnly=*/ false, DataLayerSubsystem->GetEffectiveActiveDataLayerNames());
+
+			// Loaded Data Layers Cells only
+			if (DataLayerSubsystem->GetEffectiveLoadedDataLayerNames().Num())
+			{
+				WorldPartition->RuntimeHash->GetAllStreamingCells(LoadStreamingCells, /*bAllDataLayers=*/ false, /*bDataLayersOnly=*/ true, DataLayerSubsystem->GetEffectiveLoadedDataLayerNames());
+			}
 		}
 	}
 
@@ -321,12 +357,15 @@ void UWorldPartitionStreamingPolicy::UpdateStreamingState()
 		{
 			for (TSet<const UWorldPartitionRuntimeCell*>::TIterator It = Cells.CreateIterator(); It; ++It)
 			{
-				if (ClientVisibleLevelNames.Contains((*It)->GetLevel()->GetPackage()->GetFName()))
+				if (ULevel* Level = (*It)->GetLevel())
 				{
-					It.RemoveCurrent();
+					if (ClientVisibleLevelNames.Contains((*It)->GetLevel()->GetPackage()->GetFName()))
+					{
+						It.RemoveCurrent();
 
-					UE_CLOG(bUpdateEpoch, LogWorldPartition, Verbose, TEXT("Server epoch update delayed by client visibility"));
-					bUpdateEpoch = false;
+						UE_CLOG(bUpdateEpoch, LogWorldPartition, Verbose, TEXT("Server epoch update delayed by client visibility"));
+						bUpdateEpoch = false;
+					}
 				}
 			}
 		}
