Camera: Move TickComponent to Editor only since it's just updating the draw frustum

diff --git a/Engine/Source/Runtime/Engine/Private/Camera/CameraComponent.cpp b/Engine/Source/Runtime/Engine/Private/Camera/CameraComponent.cpp
--- a/Engine/Source/Runtime/Engine/Private/Camera/CameraComponent.cpp
+++ b/Engine/Source/Runtime/Engine/Private/Camera/CameraComponent.cpp
@@ -150,14 +150,14 @@ void UCameraComponent::OnRegister()
 	Super::OnRegister();
 }
 
+#if WITH_EDITOR
 void UCameraComponent::TickComponent(float DeltaTime, ELevelTick TickType, FActorComponentTickFunction* ThisTickFunction)
 {
-#if WITH_EDITORONLY_DATA	
 	UpdateDrawFrustum();
-#endif
 
 	Super::TickComponent(DeltaTime, TickType, ThisTickFunction);
 }
+#endif
 
 #if WITH_EDITORONLY_DATA
 
