Fix rewind debugger camera Fov issues

diff --git a/Engine/Source/Runtime/Engine/Private/ObjectTrace.cpp b/Engine/Source/Runtime/Engine/Private/ObjectTrace.cpp
--- a/Engine/Source/Runtime/Engine/Private/ObjectTrace.cpp
+++ b/Engine/Source/Runtime/Engine/Private/ObjectTrace.cpp
@@ -348,6 +348,9 @@ void FObjectTrace::OutputView(const UObject* InPlayer, const FSceneView* InView)
 	const FIntRect& ViewRect = InView->CameraConstrainedViewRect;
 	float AspectRatio = (float)ViewRect.Width()/(float)ViewRect.Height();
 
+	FMatrix ProjMatrix = InView->ViewMatrices.GetProjectionMatrix();
+	float Fov = atan(1.0 / ProjMatrix.M[0][0]) * 2.0 * 180.0 / UE_DOUBLE_PI;
+
 	UE_TRACE_LOG(Object, View, ObjectChannel)
 		<< View.Cycle(FPlatformTime::Cycles64())
 		<< View.PlayerId(GetObjectId(InPlayer))
@@ -357,7 +360,7 @@ void FObjectTrace::OutputView(const UObject* InPlayer, const FSceneView* InView)
 		<< View.Pitch(InView->ViewRotation.Pitch)
 		<< View.Yaw(InView->ViewRotation.Yaw)
 		<< View.Roll(InView->ViewRotation.Roll)
-		<< View.Fov(InView->FOV)
+		<< View.Fov(Fov)
 		<< View.AspectRatio(AspectRatio);
 }
 
