Move the CurrentFrame increment from FinishAsyncTrace to ResetAsyncTrace to fix a thread safety issue.

#localization none
#tests local, preflight

diff --git a/Engine/Source/Runtime/Engine/Private/Collision/WorldCollisionAsync.cpp b/Engine/Source/Runtime/Engine/Private/Collision/WorldCollisionAsync.cpp
--- a/Engine/Source/Runtime/Engine/Private/Collision/WorldCollisionAsync.cpp
+++ b/Engine/Source/Runtime/Engine/Private/Collision/WorldCollisionAsync.cpp
@@ -502,6 +502,16 @@ void UWorld::WaitForAllAsyncTraceTasks()
 
 void UWorld::ResetAsyncTrace()
 {
+	{
+		// increase buffer index to next one
+		++AsyncTraceState.CurrentFrame;
+		// set up new buffer to accept trace requests
+		AsyncTraceData& NewAsyncBuffer = AsyncTraceState.GetBufferForCurrentFrame();
+		NewAsyncBuffer.bAsyncAllowed = true;
+		NewAsyncBuffer.NumQueuedTraceData = 0;
+		NewAsyncBuffer.NumQueuedOverlapData = 0;
+	}
+
 	AsyncTraceData& DataBufferExecuted = AsyncTraceState.GetBufferForPreviousFrame();
 
 	// Wait for thread
@@ -529,15 +539,4 @@ void UWorld::FinishAsyncTrace()
 
 	// this flag only needed to know I can't accept any more new request in current frame
 	AsyncTraceState.GetBufferForCurrentFrame().bAsyncAllowed = false;
-
-	// increase buffer index to next one
-	++AsyncTraceState.CurrentFrame;
-
-	// set up new buffer to accept trace requests
-	AsyncTraceData& NewAsyncBuffer = AsyncTraceState.GetBufferForCurrentFrame();
-	NewAsyncBuffer.bAsyncAllowed = true;
-	NewAsyncBuffer.NumQueuedTraceData = 0;
-	NewAsyncBuffer.NumQueuedOverlapData = 0;
-
-}
-
+}
\ No newline at end of file
