Exceeding partial bunch max size is now an ensure

diff --git a/Engine/Source/Runtime/Engine/Private/DataChannel.cpp b/Engine/Source/Runtime/Engine/Private/DataChannel.cpp
--- a/Engine/Source/Runtime/Engine/Private/DataChannel.cpp
+++ b/Engine/Source/Runtime/Engine/Private/DataChannel.cpp
@@ -1120,7 +1120,7 @@ FPacketIdRange UChannel::SendBunch( FOutBunch* Bunch, bool Merge )
 		return FPacketIdRange(INDEX_NONE);
 	}
 
-	if (IsBunchTooLarge(Connection, Bunch))
+	if (!ensureMsgf(!IsBunchTooLarge(Connection, Bunch), TEXT("Attempted to send bunch exceeding max allowed size. BunchSize=%d, MaximumSize=%d Channel: %s"), Bunch->GetNumBytes(), NetMaxConstructedPartialBunchSizeBytes, *Describe()))
 	{
 		UE_LOG(LogNetPartialBunch, Error, TEXT("Attempted to send bunch exceeding max allowed size. BunchSize=%d, MaximumSize=%d Channel: %s"), Bunch->GetNumBytes(), NetMaxConstructedPartialBunchSizeBytes, *Describe());
 		Bunch->SetError();
