Update function name following review

diff --git a/Engine/Source/Runtime/Engine/Classes/Components/SkeletalMeshComponent.h b/Engine/Source/Runtime/Engine/Classes/Components/SkeletalMeshComponent.h
--- a/Engine/Source/Runtime/Engine/Classes/Components/SkeletalMeshComponent.h
+++ b/Engine/Source/Runtime/Engine/Classes/Components/SkeletalMeshComponent.h
@@ -2114,7 +2114,7 @@ public:
 	 * the given constraint profile. The default will be returned if an empty or invalid profile name 
 	 * is passed in. Returns true if the successful, or false if the joint can't be found.
 	 */
-	bool GetConstraintProfileProperties(FConstraintProfileProperties& OutProperties, FName JointName, FName ProfileName);
+	bool GetConstraintProfilePropertiesOrDefault(FConstraintProfileProperties& OutProperties, FName JointName, FName ProfileName);
 
 	/** Enable or Disable AngularPositionDrive based on a list of bone names */
 	void SetNamedMotorsAngularPositionDrive(bool bEnableSwingDrive, bool bEnableTwistDrive, const TArray<FName>& BoneNames, bool bSetOtherBodiesToComplement = false);
diff --git a/Engine/Source/Runtime/Engine/Classes/PhysicsEngine/PhysicsConstraintTemplate.h b/Engine/Source/Runtime/Engine/Classes/PhysicsEngine/PhysicsConstraintTemplate.h
--- a/Engine/Source/Runtime/Engine/Classes/PhysicsEngine/PhysicsConstraintTemplate.h
+++ b/Engine/Source/Runtime/Engine/Classes/PhysicsEngine/PhysicsConstraintTemplate.h
@@ -99,7 +99,11 @@ class UPhysicsConstraintTemplate : public UObject
 		}
 	}
 
-	ENGINE_API const FConstraintProfileProperties& GetConstraintProfileProperties(FName ProfileName) const
+	/** 
+	 * Returns the properties associated with the named constraint profile. If the profile isn't found 
+	 * it returns the default profile properties.
+	 */
+	ENGINE_API const FConstraintProfileProperties& GetConstraintProfilePropertiesOrDefault(FName ProfileName) const
 	{
 		if (ProfileName != NAME_None)
 		{
diff --git a/Engine/Source/Runtime/Engine/Private/SkeletalMeshComponentPhysics.cpp b/Engine/Source/Runtime/Engine/Private/SkeletalMeshComponentPhysics.cpp
--- a/Engine/Source/Runtime/Engine/Private/SkeletalMeshComponentPhysics.cpp
+++ b/Engine/Source/Runtime/Engine/Private/SkeletalMeshComponentPhysics.cpp
@@ -1303,7 +1303,7 @@ void USkeletalMeshComponent::SetConstraintProfileForAll(FName ProfileName, bool
 	}
 }
 
-bool USkeletalMeshComponent::GetConstraintProfileProperties(
+bool USkeletalMeshComponent::GetConstraintProfilePropertiesOrDefault(
 	FConstraintProfileProperties& OutProperties, FName JointName, FName ProfileName)
 {
 	if (UPhysicsAsset* const PhysicsAsset = GetPhysicsAsset())
@@ -1313,7 +1313,7 @@ bool USkeletalMeshComponent::GetConstraintProfileProperties(
 			FConstraintInstance* ConstraintInstance = Constraints[i];
 			if (ConstraintInstance->JointName == JointName)
 			{
-				OutProperties = PhysicsAsset->ConstraintSetup[i]->GetConstraintProfileProperties(ProfileName);
+				OutProperties = PhysicsAsset->ConstraintSetup[i]->GetConstraintProfilePropertiesOrDefault(ProfileName);
 				return true;
 			}
 		}
