Fixed player collision mode rendering of convex elements which was previously not taking element transforms into account

diff --git a/Engine/Source/Runtime/Engine/Private/PhysicsEngine/PhysDrawing.cpp b/Engine/Source/Runtime/Engine/Private/PhysicsEngine/PhysDrawing.cpp
--- a/Engine/Source/Runtime/Engine/Private/PhysicsEngine/PhysDrawing.cpp
+++ b/Engine/Source/Runtime/Engine/Private/PhysicsEngine/PhysDrawing.cpp
@@ -769,7 +769,7 @@ void FKConvexElem::AddCachedSolidConvexGeom(TArray<FDynamicMeshVertex>& VertexBu
 			{
 				const int32 Index = TriVertexIndex[TriVertCount];
 				FDynamicMeshVertex Vert;
-				Vert.Position = (FVector3f)VertexData[Index];
+				Vert.Position = (FVector3f)Transform.TransformPosition(VertexData[Index]);
 				Vert.Color = VertexColor;
 				Vert.SetTangents((FVector3f)TangentX, (FVector3f)TangentY, (FVector3f)TangentZ);
 				VertexBuffer.Add(Vert);
