Strata - added warning message when enabling strata.

diff --git a/Engine/Source/Runtime/Engine/Private/RendererSettings.cpp b/Engine/Source/Runtime/Engine/Private/RendererSettings.cpp
--- a/Engine/Source/Runtime/Engine/Private/RendererSettings.cpp
+++ b/Engine/Source/Runtime/Engine/Private/RendererSettings.cpp
@@ -237,11 +237,24 @@ void URendererSettings::PostEditChangeProperty(FPropertyChangedEvent& PropertyCh
 			}
 		}
 
+		if ((PropertyChangedEvent.Property->GetFName() == GET_MEMBER_NAME_CHECKED(URendererSettings, bEnableStrata)))
+		{
+			if (bEnableStrata)
+			{
+				FMessageDialog::Open(EAppMsgType::Ok, LOCTEXT("Strata Experimental", "Warning: Strata is experimental. Be aware that any materials saved when Strata is enabled won't be rendered correctly if Strata is disabled later on."));
+			}
+		}
+
 		if ((PropertyChangedEvent.Property->GetFName() == GET_MEMBER_NAME_CHECKED(URendererSettings, StrataBytePerPixel)))
 		{
+			const uint32 MinStrataBytePerPixel = 20u;
+			const uint32 MaxStrataBytePerPixel = 128u;
+			if (StrataBytePerPixel < MinStrataBytePerPixel || StrataBytePerPixel > MaxStrataBytePerPixel)
+			{
+				FMessageDialog::Open(EAppMsgType::Ok, LOCTEXT("Strata Byte Per Pixel", "Strata `byte-per-pixel` must be between 20 and 128."));
+			}
 			// We enforce at least 20 bytes per pixel because this is the minimal Strata GBuffer footprint of the simplest material.
-			StrataBytePerPixel = FMath::Clamp(StrataBytePerPixel, 20u, 128u);
-
+			StrataBytePerPixel = FMath::Clamp(StrataBytePerPixel, MinStrataBytePerPixel, MaxStrataBytePerPixel);
 		}
 
 		ExportValuesToConsoleVariables(PropertyChangedEvent.Property);
