World Partition - HLOD: Exclude some HLOD builder subclasses from editor dropdowns
* Custom HLOD builders can be assigned to HLOD layers
* Instancing, MeshApproximate, MeshMerge and MeshSimplify are already exposed through standard means
* WaterBodyHLODBuilder and LandscapeHLODBuilder are specific to their respective components

diff --git a/Engine/Plugins/Editor/WorldPartitionHLODUtilities/Source/Public/HLODBuilderInstancing.h b/Engine/Plugins/Editor/WorldPartitionHLODUtilities/Source/Public/HLODBuilderInstancing.h
--- a/Engine/Plugins/Editor/WorldPartitionHLODUtilities/Source/Public/HLODBuilderInstancing.h
+++ b/Engine/Plugins/Editor/WorldPartitionHLODUtilities/Source/Public/HLODBuilderInstancing.h
@@ -9,7 +9,7 @@
 /**
  * Build a AWorldPartitionHLOD whose components are ISMC
  */
-UCLASS()
+UCLASS(HideDropdown)
 class WORLDPARTITIONHLODUTILITIES_API UHLODBuilderInstancing : public UHLODBuilder
 {
 	 GENERATED_UCLASS_BODY()
diff --git a/Engine/Plugins/Editor/WorldPartitionHLODUtilities/Source/Public/HLODBuilderMeshApproximate.h b/Engine/Plugins/Editor/WorldPartitionHLODUtilities/Source/Public/HLODBuilderMeshApproximate.h
--- a/Engine/Plugins/Editor/WorldPartitionHLODUtilities/Source/Public/HLODBuilderMeshApproximate.h
+++ b/Engine/Plugins/Editor/WorldPartitionHLODUtilities/Source/Public/HLODBuilderMeshApproximate.h
@@ -30,7 +30,7 @@ class WORLDPARTITIONHLODUTILITIES_API UHLODBuilderMeshApproximateSettings : publ
 /**
  * Build an approximated mesh using geometry from the provided actors
  */
-UCLASS()
+UCLASS(HideDropdown)
 class WORLDPARTITIONHLODUTILITIES_API UHLODBuilderMeshApproximate : public UHLODBuilder
 {
 	GENERATED_UCLASS_BODY()
diff --git a/Engine/Plugins/Editor/WorldPartitionHLODUtilities/Source/Public/HLODBuilderMeshMerge.h b/Engine/Plugins/Editor/WorldPartitionHLODUtilities/Source/Public/HLODBuilderMeshMerge.h
--- a/Engine/Plugins/Editor/WorldPartitionHLODUtilities/Source/Public/HLODBuilderMeshMerge.h
+++ b/Engine/Plugins/Editor/WorldPartitionHLODUtilities/Source/Public/HLODBuilderMeshMerge.h
@@ -30,7 +30,7 @@ class WORLDPARTITIONHLODUTILITIES_API UHLODBuilderMeshMergeSettings : public UHL
 /**
  * Build a merged mesh using geometry from the provided actors
  */
-UCLASS()
+UCLASS(HideDropdown)
 class WORLDPARTITIONHLODUTILITIES_API UHLODBuilderMeshMerge : public UHLODBuilder
 {
 	GENERATED_UCLASS_BODY()
diff --git a/Engine/Plugins/Editor/WorldPartitionHLODUtilities/Source/Public/HLODBuilderMeshSimplify.h b/Engine/Plugins/Editor/WorldPartitionHLODUtilities/Source/Public/HLODBuilderMeshSimplify.h
--- a/Engine/Plugins/Editor/WorldPartitionHLODUtilities/Source/Public/HLODBuilderMeshSimplify.h
+++ b/Engine/Plugins/Editor/WorldPartitionHLODUtilities/Source/Public/HLODBuilderMeshSimplify.h
@@ -30,7 +30,7 @@ class WORLDPARTITIONHLODUTILITIES_API UHLODBuilderMeshSimplifySettings : public
 /**
  * Build a simplified mesh using geometry from the provided actors
  */
-UCLASS()
+UCLASS(HideDropdown)
 class WORLDPARTITIONHLODUTILITIES_API UHLODBuilderMeshSimplify : public UHLODBuilder
 {
 	GENERATED_UCLASS_BODY()
diff --git a/Engine/Plugins/Experimental/Water/Source/Runtime/Public/WaterBodyHLODBuilder.h b/Engine/Plugins/Experimental/Water/Source/Runtime/Public/WaterBodyHLODBuilder.h
--- a/Engine/Plugins/Experimental/Water/Source/Runtime/Public/WaterBodyHLODBuilder.h
+++ b/Engine/Plugins/Experimental/Water/Source/Runtime/Public/WaterBodyHLODBuilder.h
@@ -7,7 +7,7 @@
 #include "WorldPartition/HLOD/HLODBuilder.h"
 #include "WaterBodyHLODBuilder.generated.h"
 
-UCLASS()
+UCLASS(HideDropdown)
 class UWaterBodyHLODBuilder : public UHLODBuilder
 {
 	GENERATED_UCLASS_BODY()
diff --git a/Engine/Source/Runtime/Engine/Public/WorldPartition/HLOD/HLODLayer.h b/Engine/Source/Runtime/Engine/Public/WorldPartition/HLOD/HLODLayer.h
--- a/Engine/Source/Runtime/Engine/Public/WorldPartition/HLOD/HLODLayer.h
+++ b/Engine/Source/Runtime/Engine/Public/WorldPartition/HLOD/HLODLayer.h
@@ -82,7 +82,7 @@ private:
 	EHLODLayerType LayerType;
 
 	/** HLODBuilder class */
-	UPROPERTY(EditAnywhere, Config, Category=HLOD, meta = (EditConditionHides, EditCondition = "LayerType == EHLODLayerType::Custom"))
+	UPROPERTY(EditAnywhere, Config, Category=HLOD, meta = (DisplayName = "HLOD Builder Class", EditConditionHides, EditCondition = "LayerType == EHLODLayerType::Custom"))
 	TSubclassOf<UHLODBuilder> HLODBuilderClass;
 
 	UPROPERTY(VisibleAnywhere, Export, NoClear, Category=HLOD, meta = (EditInline, NoResetToDefault))
diff --git a/Engine/Source/Runtime/Landscape/Public/LandscapeHLODBuilder.h b/Engine/Source/Runtime/Landscape/Public/LandscapeHLODBuilder.h
--- a/Engine/Source/Runtime/Landscape/Public/LandscapeHLODBuilder.h
+++ b/Engine/Source/Runtime/Landscape/Public/LandscapeHLODBuilder.h
@@ -7,7 +7,7 @@
 #include "WorldPartition/HLOD/HLODBuilder.h"
 #include "LandscapeHLODBuilder.generated.h"
 
-UCLASS()
+UCLASS(HideDropdown)
 class LANDSCAPE_API ULandscapeHLODBuilder : public UHLODBuilder
 {
 	GENERATED_UCLASS_BODY()
