Camera Manager: Reset the PreviousViewTransform each frame that we update the view target. If your view target was a UCameraComponent this variable was already being reset (being read from the Motion Vector Simulation system), but if your view target was not a camera component this would not get reset.

This meant if you had a camera as your viewtarget (and the MVS system was being fed data by Sequencer) and then your viewtarget switched to something that was not a camera, it would persist the PreviousViewTransform leading to incorrect camera motion blur.

diff --git a/Engine/Source/Runtime/Engine/Private/PlayerCameraManager.cpp b/Engine/Source/Runtime/Engine/Private/PlayerCameraManager.cpp
--- a/Engine/Source/Runtime/Engine/Private/PlayerCameraManager.cpp
+++ b/Engine/Source/Runtime/Engine/Private/PlayerCameraManager.cpp
@@ -357,7 +357,7 @@ void APlayerCameraManager::UpdateViewTarget(FTViewTarget& OutVT, float DeltaTime
 	OutVT.POV.ProjectionMode = bIsOrthographic ? ECameraProjectionMode::Orthographic : ECameraProjectionMode::Perspective;
 	OutVT.POV.PostProcessSettings.SetBaseValues();
 	OutVT.POV.PostProcessBlendWeight = 1.0f;
-
+	OutVT.POV.PreviousViewTransform.Reset();
 
 	bool bDoNotApplyModifiers = false;
 
