World Partition - HLOD: Ensure that UChildActorComponents are marked as HLOD relevant if their child actors are HLOD relevant.
Marking them as HLOD relevant will force the loading of their child actors when processed by the HLOD builders.

diff --git a/Engine/Source/Runtime/Engine/Classes/Components/ChildActorComponent.h b/Engine/Source/Runtime/Engine/Classes/Components/ChildActorComponent.h
--- a/Engine/Source/Runtime/Engine/Classes/Components/ChildActorComponent.h
+++ b/Engine/Source/Runtime/Engine/Classes/Components/ChildActorComponent.h
@@ -171,6 +171,10 @@ public:
 	virtual void OnUnregister() override;
 	virtual TStructOnScope<FActorComponentInstanceData> GetComponentInstanceData() const override;
 	virtual void BeginPlay() override;
+	virtual bool IsHLODRelevant() const override;
+#if WITH_EDITOR
+	virtual TSubclassOf<class UHLODBuilder> GetCustomHLODBuilderClass() const override;
+#endif
 	//~ End ActorComponent Interface.
 
 	/** Apply the component instance data to the child actor component */
diff --git a/Engine/Source/Runtime/Engine/Private/Components/ChildActorComponent.cpp b/Engine/Source/Runtime/Engine/Private/Components/ChildActorComponent.cpp
--- a/Engine/Source/Runtime/Engine/Private/Components/ChildActorComponent.cpp
+++ b/Engine/Source/Runtime/Engine/Private/Components/ChildActorComponent.cpp
@@ -9,6 +9,9 @@
 #include "Engine/Engine.h"
 #include "Engine/DemoNetDriver.h"
 
+#if WITH_EDITOR
+#include "WorldPartition/HLOD/HLODBuilder.h"
+#endif
 #if UE_WITH_IRIS
 #include "Iris/ReplicationSystem/PropertyReplicationFragment.h"
 #endif // UE_WITH_IRIS
@@ -859,6 +862,12 @@ void UChildActorComponent::BeginPlay()
 	}
 }
 
+bool UChildActorComponent::IsHLODRelevant() const
+{
+	const bool bIsHLODRelevant = ChildActor && ChildActor->IsHLODRelevant();
+	return bIsHLODRelevant;
+}
+
 #if UE_WITH_IRIS
 void UChildActorComponent::RegisterReplicationFragments(UE::Net::FFragmentRegistrationContext& Context, UE::Net::EFragmentRegistrationFlags RegistrationFlags)
 {
@@ -871,6 +880,14 @@ void UChildActorComponent::RegisterReplicationFragments(UE::Net::FFragmentRegist
 
 
 #if WITH_EDITOR
+TSubclassOf<class UHLODBuilder> UChildActorComponent::GetCustomHLODBuilderClass() const
+{
+	// ChildActorComponents are only HLOD relevant so that their child actors are included in the HLOD generation.
+	// They don't provide any mesh/visual input, so we route them to be processed by the NullHLODBuilder which 
+	// ignores all components sent to it.
+	return UNullHLODBuilder::StaticClass();
+}
+
 void UChildActorComponent::SetEditorTreeViewVisualizationMode(EChildActorComponentTreeViewVisualizationMode InMode)
 {
 	Modify();
diff --git a/Engine/Source/Runtime/Engine/Private/WorldPartition/HLOD/HLODBuilder.cpp b/Engine/Source/Runtime/Engine/Private/WorldPartition/HLOD/HLODBuilder.cpp
--- a/Engine/Source/Runtime/Engine/Private/WorldPartition/HLOD/HLODBuilder.cpp
+++ b/Engine/Source/Runtime/Engine/Private/WorldPartition/HLOD/HLODBuilder.cpp
@@ -17,6 +17,11 @@ UHLODBuilder::UHLODBuilder(const FObjectInitializer& ObjectInitializer)
 {
 }
 
+UNullHLODBuilder::UNullHLODBuilder(const FObjectInitializer& ObjectInitializer)
+	: Super(ObjectInitializer)
+{
+}
+
 UHLODBuilderSettings::UHLODBuilderSettings(const FObjectInitializer& ObjectInitializer)
 	: Super(ObjectInitializer)
 {
diff --git a/Engine/Source/Runtime/Engine/Public/WorldPartition/HLOD/HLODBuilder.h b/Engine/Source/Runtime/Engine/Public/WorldPartition/HLOD/HLODBuilder.h
--- a/Engine/Source/Runtime/Engine/Public/WorldPartition/HLOD/HLODBuilder.h
+++ b/Engine/Source/Runtime/Engine/Public/WorldPartition/HLOD/HLODBuilder.h
@@ -133,3 +133,20 @@ protected:
 	TObjectPtr<const UHLODBuilderSettings> HLODBuilderSettings;
 #endif
 };
+
+
+/**
+ * Null HLOD builder that ignores it's input and generate no component.
+ */
+UCLASS(HideDropdown)
+class ENGINE_API UNullHLODBuilder : public UHLODBuilder
+{
+	GENERATED_UCLASS_BODY()
+
+#if WITH_EDITOR
+	virtual bool RequiresCompiledAssets() const { return false; }
+	virtual bool RequiresWarmup() const { return false; }
+	virtual uint32 ComputeHLODHash(const UActorComponent* InSourceComponent) const { return 0; }
+	virtual TArray<UActorComponent*> Build(const FHLODBuildContext& InHLODBuildContext, const TArray<UActorComponent*>& InSourceComponents) const { return {}; }
+#endif
+};
