Added transaction annotation diffing support for FActorTransactionAnnotation

This has the same rules as FComponentInstanceDataCache regarding the components that are diffed (child actor components, or any SCS/UCS components), and these components are paired for diffing using their FActorComponentInstanceSourceInfo to allow reinstanced components to be paired by checking they were created from the same source.

#rnx

diff --git a/Engine/Source/Runtime/Engine/Private/ActorEditor.cpp b/Engine/Source/Runtime/Engine/Private/ActorEditor.cpp
--- a/Engine/Source/Runtime/Engine/Private/ActorEditor.cpp
+++ b/Engine/Source/Runtime/Engine/Private/ActorEditor.cpp
@@ -351,6 +351,12 @@ void AActor::DebugShowOneComponentHierarchy( USceneComponent* SceneComp, int32&
 	}
 }
 
+FActorTransactionAnnotation::FDiffableComponentInfo::FDiffableComponentInfo(const UActorComponent* Component)
+	: ComponentSourceInfo(Component)
+	, DiffableComponent(UE::Transaction::DiffUtil::GetDiffableObject(Component, UE::Transaction::DiffUtil::SerializeProperties))
+{
+}
+
 TSharedRef<FActorTransactionAnnotation> FActorTransactionAnnotation::Create()
 {
 	return MakeShareable(new FActorTransactionAnnotation());
@@ -423,6 +429,22 @@ FActorTransactionAnnotation::FActorTransactionAnnotation(const AActor* InActor,
 	{
 		ActorTransactionAnnotationData.bRootComponentDataCached = false;
 	}
+
+	// This code also runs when reconstructing actors, so only cache the diff data when we're creating or applying a transaction
+	if (GUndo || GIsTransacting)
+	{
+		TInlineComponentArray<UActorComponent*> Components;
+		FComponentInstanceDataCache::GetComponentHierarchy(InActor, Components);
+
+		const bool bIsChildActor = InActor->IsChildActor();
+		for (UActorComponent* Component : Components)
+		{
+			if (Component && (bIsChildActor || Component->IsCreatedByConstructionScript()))
+			{
+				DiffableComponentInfos.Emplace(Component);
+			}
+		}
+	}
 }
 
 void FActorTransactionAnnotation::AddReferencedObjects(FReferenceCollector& Collector)
@@ -435,6 +457,149 @@ void FActorTransactionAnnotation::Serialize(FArchive& Ar)
 	Ar << ActorTransactionAnnotationData;
 }
 
+void FActorTransactionAnnotation::ComputeAdditionalObjectChanges(const ITransactionObjectAnnotation* OriginalAnnotation, TMap<UObject*, FTransactionObjectChange>& OutAdditionalObjectChanges)
+{
+	const FActorTransactionAnnotation* OriginalActorAnnotation = static_cast<const FActorTransactionAnnotation*>(OriginalAnnotation);
+
+	struct FDiffableComponentPair
+	{
+		UActorComponent* CurrentComponent = nullptr;
+		const UE::Transaction::FDiffableObject* OldDiffableComponent = nullptr;
+		const UE::Transaction::FDiffableObject* NewDiffableComponent = nullptr;
+	};
+
+	TArray<FDiffableComponentPair, TInlineAllocator<NumInlinedActorComponents>> DiffableComponentPairs;
+	if (const AActor* Actor = ActorTransactionAnnotationData.Actor.Get(/*bEvenIfPendingKill*/true))
+	{
+		TInlineComponentArray<UActorComponent*> Components;
+		FComponentInstanceDataCache::GetComponentHierarchy(Actor, Components);
+
+		if (Components.Num() > 0)
+		{
+			// Bitsets to avoid repeated testing of FDiffableComponentInfo that has already been matched to a known component
+			TBitArray<> OriginalDiffableComponentInfosToConsider(true, OriginalActorAnnotation ? OriginalActorAnnotation->DiffableComponentInfos.Num() : 0);
+			TBitArray<> DiffableComponentInfosToConsider(true, DiffableComponentInfos.Num());
+
+			const bool bIsChildActor = Actor->IsChildActor();
+			DiffableComponentPairs.Reserve(Components.Num());
+			for (UActorComponent* Component : Components)
+			{
+				if (Component && (bIsChildActor || Component->IsCreatedByConstructionScript()))
+				{
+					auto FindDiffableComponent = [Component](const TArray<FDiffableComponentInfo>& DiffableComponentInfosToSearch, TBitArray<>& DiffableComponentInfosBitset) -> const UE::Transaction::FDiffableObject*
+					{
+						const UObject* ComponentTemplate = Component->GetArchetype();
+						for (TConstSetBitIterator<> BitsetIt(DiffableComponentInfosBitset); BitsetIt; ++BitsetIt)
+						{
+							const FDiffableComponentInfo& PotentialDiffableComponentInfo = DiffableComponentInfosToSearch[BitsetIt.GetIndex()];
+							if (PotentialDiffableComponentInfo.ComponentSourceInfo.MatchesComponent(Component, ComponentTemplate))
+							{
+								DiffableComponentInfosBitset[BitsetIt.GetIndex()] = false;
+								return &PotentialDiffableComponentInfo.DiffableComponent;
+							}
+						}
+						return nullptr;
+					};
+
+					FDiffableComponentPair& DiffableComponentPair = DiffableComponentPairs.AddDefaulted_GetRef();
+					DiffableComponentPair.CurrentComponent = Component;
+					DiffableComponentPair.OldDiffableComponent = OriginalActorAnnotation ? FindDiffableComponent(OriginalActorAnnotation->DiffableComponentInfos, OriginalDiffableComponentInfosToConsider) : nullptr;
+					DiffableComponentPair.NewDiffableComponent = FindDiffableComponent(DiffableComponentInfos, DiffableComponentInfosToConsider);
+				}
+			}
+		}
+	}
+
+	if (DiffableComponentPairs.Num() > 0)
+	{
+		auto FindOrAddAdditionalObjectChange = [&OutAdditionalObjectChanges](UActorComponent* CurrentComponent, const FTransactionObjectId& OldObjectId, FTransactionObjectDeltaChange&& ComponentDeltaChange)
+		{
+			if (FTransactionObjectChange* ExistingComponentChange = OutAdditionalObjectChanges.Find(CurrentComponent))
+			{
+				ExistingComponentChange->DeltaChange.Merge(ComponentDeltaChange);
+			}
+			else
+			{
+				OutAdditionalObjectChanges.Add(CurrentComponent, FTransactionObjectChange{ OldObjectId, MoveTemp(ComponentDeltaChange) });
+			}
+		};
+
+		auto ComputeAdditionalObjectChange = [&FindOrAddAdditionalObjectChange](const UE::Transaction::FDiffableObject& OldDiffableComponent, const UE::Transaction::FDiffableObject& NewDiffableComponent, UActorComponent* CurrentComponent, const bool bForcePendingKillChange = false)
+		{
+			UE::Transaction::DiffUtil::FGenerateObjectDiffOptions DiffOptions;
+			{
+				TSet<const FProperty*> PropertiesToSkip;
+
+				// Skip properties modified during construction when calculating the diff
+				CurrentComponent->GetUCSModifiedProperties(PropertiesToSkip);
+
+				// If this is the owning Actor's root scene component, always include relative transform properties as GetUCSModifiedProperties incorrectly considers them modified (due to changing during placement)
+				if (CurrentComponent->IsA<USceneComponent>())
+				{
+					const AActor* ComponentOwner = CurrentComponent->GetOwner();
+					if (ComponentOwner && ComponentOwner->GetRootComponent() == CurrentComponent)
+					{
+						PropertiesToSkip.Remove(FindFProperty<FProperty>(USceneComponent::StaticClass(), USceneComponent::GetRelativeLocationPropertyName()));
+						PropertiesToSkip.Remove(FindFProperty<FProperty>(USceneComponent::StaticClass(), USceneComponent::GetRelativeRotationPropertyName()));
+						PropertiesToSkip.Remove(FindFProperty<FProperty>(USceneComponent::StaticClass(), USceneComponent::GetRelativeScale3DPropertyName()));
+					}
+				}
+
+				DiffOptions.ShouldSkipProperty = [CurrentComponent, PropertiesToSkip = MoveTemp(PropertiesToSkip)](FName PropertyName) -> bool
+				{
+					if (const FProperty* Property = FindFProperty<FProperty>(CurrentComponent->GetClass(), PropertyName))
+					{
+						return !Property->HasAnyPropertyFlags(CPF_Edit | CPF_Interp)
+							|| Property->IsA<FMulticastDelegateProperty>()
+							|| PropertiesToSkip.Contains(Property);
+					}
+					return false;
+				};
+			}
+
+			FTransactionObjectDeltaChange ComponentDeltaChange = UE::Transaction::DiffUtil::GenerateObjectDiff(OldDiffableComponent, NewDiffableComponent, DiffOptions);
+			ComponentDeltaChange.bHasPendingKillChange |= bForcePendingKillChange;
+
+			if (ComponentDeltaChange.HasChanged())
+			{
+				FindOrAddAdditionalObjectChange(CurrentComponent, OldDiffableComponent.ObjectInfo, MoveTemp(ComponentDeltaChange));
+			}
+		};
+
+		for (const FDiffableComponentPair& DiffableComponentPair : DiffableComponentPairs)
+		{
+			check(DiffableComponentPair.CurrentComponent);
+
+			if (DiffableComponentPair.OldDiffableComponent && DiffableComponentPair.NewDiffableComponent)
+			{
+				// Component already existed; diff against its previous state
+				ComputeAdditionalObjectChange(*DiffableComponentPair.OldDiffableComponent, *DiffableComponentPair.NewDiffableComponent, DiffableComponentPair.CurrentComponent);
+			}
+			else if (DiffableComponentPair.OldDiffableComponent && !DiffableComponentPair.NewDiffableComponent)
+			{
+				// Component is deleted; just mark it as pending kill
+				FTransactionObjectDeltaChange ComponentDeltaChange;
+				ComponentDeltaChange.bHasPendingKillChange = true;
+
+				FindOrAddAdditionalObjectChange(DiffableComponentPair.CurrentComponent, DiffableComponentPair.OldDiffableComponent->ObjectInfo, MoveTemp(ComponentDeltaChange));
+			}
+			else if (DiffableComponentPair.NewDiffableComponent)
+			{
+				check(!DiffableComponentPair.OldDiffableComponent);
+
+				// Component is newly added; diff against its archetype
+				if (const UObject* ArchetypeObject = DiffableComponentPair.CurrentComponent->GetArchetype(); ArchetypeObject && !DiffableComponentPair.CurrentComponent->HasAnyFlags(RF_ClassDefaultObject | RF_ArchetypeObject))
+				{
+					UE::Transaction::FDiffableObject ArchetypeDiffableObject = UE::Transaction::DiffUtil::GetDiffableObject(ArchetypeObject, UE::Transaction::DiffUtil::SerializeProperties);
+					ArchetypeDiffableObject.SetObject(DiffableComponentPair.CurrentComponent); // Use the real component here, as the archetype info will never be valid for an instance
+
+					ComputeAdditionalObjectChange(ArchetypeDiffableObject, *DiffableComponentPair.NewDiffableComponent, DiffableComponentPair.CurrentComponent, /*bForcePendingKillChange*/true);
+				}
+			}
+		}
+	}
+}
+
 bool FActorTransactionAnnotation::HasInstanceData() const
 {
 	return (ActorTransactionAnnotationData.bRootComponentDataCached || ActorTransactionAnnotationData.ComponentInstanceData.HasInstanceData());
diff --git a/Engine/Source/Runtime/Engine/Public/ActorTransactionAnnotation.h b/Engine/Source/Runtime/Engine/Public/ActorTransactionAnnotation.h
--- a/Engine/Source/Runtime/Engine/Public/ActorTransactionAnnotation.h
+++ b/Engine/Source/Runtime/Engine/Public/ActorTransactionAnnotation.h
@@ -7,6 +7,7 @@
 #include "ComponentInstanceDataCache.h"
 
 #if WITH_EDITOR
+#include "TransactionCommon.h"
 #include "Misc/ITransactionObjectAnnotation.h"
 #endif
 
@@ -71,6 +72,8 @@ public:
 	//~ ITransactionObjectAnnotation interface
 	virtual void AddReferencedObjects(FReferenceCollector& Collector) override;
 	virtual void Serialize(FArchive& Ar) override;
+	virtual bool SupportsAdditionalObjectChanges() const override { return true; }
+	virtual void ComputeAdditionalObjectChanges(const ITransactionObjectAnnotation* OriginalAnnotation, TMap<UObject*, FTransactionObjectChange>& OutAdditionalObjectChanges) override;
 
 	bool HasInstanceData() const;
 
@@ -79,5 +82,16 @@ public:
 private:
 	FActorTransactionAnnotation();
 	FActorTransactionAnnotation(const AActor* InActor, FComponentInstanceDataCache&& InComponentInstanceData, const bool InCacheRootComponentData = true);
+	
+	struct FDiffableComponentInfo
+	{
+		FDiffableComponentInfo() = default;
+		explicit FDiffableComponentInfo(const UActorComponent* Component);
+
+		FActorComponentInstanceSourceInfo ComponentSourceInfo;
+		UE::Transaction::FDiffableObject DiffableComponent;
+	};
+
+	TArray<FDiffableComponentInfo> DiffableComponentInfos;
 };
 #endif // WITH_EDITOR
