Add support for resizing the scene rendertargets by dynamically adjusting the max resolution based on a quantized version of the computed resolution fraction. This results in more efficient clears and resolves (since these now scale with resolution), and more efficient use of fast VRAM on platforms which use it.

Quantization can be controlled via r.DynamicRes.UpperBoundQuantization, which specifies the number of steps. The step size is given by (MaxScreenPercentage-MinScreenPercentage)/UpperBoundQuantization. Set this to 0 to disable.

Only recommended for use with the transient allocator (on supported platforms) with a large transient texture cache (e.g RHI.TransientAllocator.TextureCacheSize=512). A rule-of-thumb for cache size is 32 * r.DynamicRes.UpperBoundQuantization.

diff --git a/Engine/Source/Runtime/Engine/Private/DynamicResolution.cpp b/Engine/Source/Runtime/Engine/Private/DynamicResolution.cpp
--- a/Engine/Source/Runtime/Engine/Private/DynamicResolution.cpp
+++ b/Engine/Source/Runtime/Engine/Private/DynamicResolution.cpp
@@ -13,7 +13,6 @@
 #include "UnrealEngine.h"
 #include "RenderGraphEvent.h"
 
-
 static TAutoConsoleVariable<float> CVarDynamicResMinSP(
 	TEXT("r.DynamicRes.MinScreenPercentage"),
 	DynamicRenderScaling::FractionToPercentage(DynamicRenderScaling::FHeuristicSettings::kDefaultMinResolutionFraction),
@@ -85,6 +84,15 @@ static TAutoConsoleVariable<int32> CVarTimingMeasureModel(
 	TEXT(" 0: Same as stat unit (default);\n 1: Timestamp queries."),
 	ECVF_RenderThreadSafe | ECVF_Default);
 
+static TAutoConsoleVariable<int32> CVarUpperBoundQuantization(
+	TEXT("r.DynamicRes.UpperBoundQuantization"),
+	DynamicRenderScaling::FHeuristicSettings::kDefaultUpperBoundQuantization,
+	TEXT("Quantization step count to use for upper bound screen percentage.\n")
+	TEXT("If non-zero, rendertargets will be resized based on the dynamic resolution fraction, saving GPU time during clears and resolves.\n")
+	TEXT("Only recommended for use with the transient allocator (on supported platforms) with a large transient texture cache (e.g RHI.TransientAllocator.TextureCacheSize=512)"),
+	ECVF_RenderThreadSafe | ECVF_Default);
+
+
 
 DynamicRenderScaling::FHeuristicSettings GetPrimaryDynamicResolutionSettings()
 {
@@ -92,6 +100,7 @@ DynamicRenderScaling::FHeuristicSettings GetPrimaryDynamicResolutionSettings()
 	BudgetSetting.Model = DynamicRenderScaling::EHeuristicModel::Quadratic;
 	BudgetSetting.MinResolutionFraction      = DynamicRenderScaling::GetPercentageCVarToFraction(CVarDynamicResMinSP);
 	BudgetSetting.MaxResolutionFraction      = DynamicRenderScaling::GetPercentageCVarToFraction(CVarDynamicResMaxSP);
+	BudgetSetting.UpperBoundQuantization     = CVarUpperBoundQuantization.GetValueOnAnyThread();
 	// BudgetSetting.BudgetMs depends on the cost of other buckets.
 	BudgetSetting.ChangeThreshold            = DynamicRenderScaling::GetPercentageCVarToFraction(CVarChangeThreshold);
 	BudgetSetting.IncreaseAmortizationFactor = CVarIncreaseAmortizationFactor.GetValueOnAnyThread();
@@ -140,6 +149,7 @@ void FDynamicResolutionHeuristicProxy::ResetInternal()
 
 	NumberOfFramesSinceScreenPercentageChange = 0;
 	CurrentFrameResolutionFractions.SetAll(1.0f);
+	CurrentFrameMaxResolutionFractions.SetAll(1.0f);
 
 	// Ignore previous frame timings.
 	IgnoreFrameRemainingCount = 1;
@@ -504,18 +514,41 @@ void FDynamicResolutionHeuristicProxy::RefreshCurentFrameResolutionFraction_Rend
 	{
 		NumberOfFramesSinceScreenPercentageChange++;
 	}
-}
-
-// static
-DynamicRenderScaling::TMap<float> FDynamicResolutionHeuristicProxy::GetResolutionFractionUpperBounds()
-{
-	DynamicRenderScaling::TMap<float> MaxResolutionFractions;
+	
+	// Compute max resolution for each budget by quantizing the new resolution fraction (falls back to the MaxResolution setting if BudgetSetting.UpperBoundQuantization==0)
+	DynamicRenderScaling::TMap<float> NewMaxResolutionFractions;
 	for (TLinkedList<DynamicRenderScaling::FBudget*>::TIterator BudgetIt(DynamicRenderScaling::FBudget::GetGlobalList()); BudgetIt; BudgetIt.Next())
 	{
 		const DynamicRenderScaling::FBudget& Budget = **BudgetIt;
-		MaxResolutionFractions[Budget] = Budget.GetSettings().MaxResolutionFraction;
+		DynamicRenderScaling::FHeuristicSettings BudgetSettings = Budget.GetSettings();
+
+		if (BudgetSettings.IsEnabled() || Budget == GDynamicPrimaryResolutionFraction)
+		{
+			float NewMaxResolutionFraction = BudgetSettings.MaxResolutionFraction;
+			if (BudgetSettings.UpperBoundQuantization > 0)
+			{
+				float CurrentResolutionFraction = NewResolutionFractions[Budget];
+				float AvailableRange = BudgetSettings.MaxResolutionFraction - BudgetSettings.MinResolutionFraction;
+				float QuantizationStepSize = AvailableRange / float(BudgetSettings.UpperBoundQuantization);
+				NewMaxResolutionFraction = FMath::CeilToFloat(CurrentResolutionFraction / QuantizationStepSize) * QuantizationStepSize;
+				NewMaxResolutionFraction = FMath::Min(NewMaxResolutionFraction, BudgetSettings.MaxResolutionFraction);
+			}
+			NewMaxResolutionFractions[Budget] = NewMaxResolutionFraction;
+		}
+		else
+		{
+			NewMaxResolutionFractions[Budget] = 1.0f;
+		}
 	}
-	return MaxResolutionFractions;
+
+	CurrentFrameMaxResolutionFractions = NewMaxResolutionFractions;
+}
+
+// static
+DynamicRenderScaling::TMap<float> FDynamicResolutionHeuristicProxy::GetResolutionFractionUpperBounds() const
+{
+	check(IsInGameThread() || IsInRenderingThread());
+	return CurrentFrameMaxResolutionFractions;
 }
 
 
diff --git a/Engine/Source/Runtime/Engine/Private/GameViewportClient.cpp b/Engine/Source/Runtime/Engine/Private/GameViewportClient.cpp
--- a/Engine/Source/Runtime/Engine/Private/GameViewportClient.cpp
+++ b/Engine/Source/Runtime/Engine/Private/GameViewportClient.cpp
@@ -1407,6 +1407,7 @@ void UGameViewportClient::Draw(FViewport* InViewport, FCanvas* SceneCanvas)
 			{
 				// Keep same name as before for primary screen percentage
 				CSV_CUSTOM_STAT_GLOBAL(DynamicResolutionPercentage, DynamicResolutionStateInfos.ResolutionFractionApproximations[GDynamicPrimaryResolutionFraction] * 100.0f, ECsvCustomStatOp::Set);
+				CSV_CUSTOM_STAT_GLOBAL(DynamicResolutionPercentageMax, DynamicResolutionStateInfos.ResolutionFractionUpperBounds[GDynamicPrimaryResolutionFraction] * 100.0f, ECsvCustomStatOp::Set);
 			}
 			for (TLinkedList<DynamicRenderScaling::FBudget*>::TIterator BudgetIt(DynamicRenderScaling::FBudget::GetGlobalList()); BudgetIt; BudgetIt.Next())
 			{
diff --git a/Engine/Source/Runtime/Engine/Public/DynamicResolutionProxy.h b/Engine/Source/Runtime/Engine/Public/DynamicResolutionProxy.h
--- a/Engine/Source/Runtime/Engine/Public/DynamicResolutionProxy.h
+++ b/Engine/Source/Runtime/Engine/Public/DynamicResolutionProxy.h
@@ -50,7 +50,7 @@ public:
 
 
 	/** Returns the view fraction upper bound. */
-	static DynamicRenderScaling::TMap<float> GetResolutionFractionUpperBounds();
+	DynamicRenderScaling::TMap<float> GetResolutionFractionUpperBounds() const;
 
 	/** Creates a default dynamic resolution state using this proxy that queries GPU timing from the RHI. */
 	static TSharedPtr< class IDynamicResolutionState > CreateDefaultState();
@@ -107,6 +107,7 @@ private:
 
 	// Current frame's view fraction.
 	DynamicRenderScaling::TMap<float> CurrentFrameResolutionFractions;
+	DynamicRenderScaling::TMap<float> CurrentFrameMaxResolutionFractions;
 	DynamicRenderScaling::TMap<int32> BudgetHistorySizes;
 
 	// Frame counter to allocate unique ID for CommitPreviousFrameGPUTimings_RenderThread().
diff --git a/Engine/Source/Runtime/RenderCore/Public/DynamicRenderScaling.h b/Engine/Source/Runtime/RenderCore/Public/DynamicRenderScaling.h
--- a/Engine/Source/Runtime/RenderCore/Public/DynamicRenderScaling.h
+++ b/Engine/Source/Runtime/RenderCore/Public/DynamicRenderScaling.h
@@ -38,6 +38,7 @@ struct RENDERCORE_API FHeuristicSettings final
 	static constexpr float kDefaultTargetedHeadRoom = 0.05f;
 	static constexpr float kDefaultIncreaseAmortizationFactor = 0.9f;
 	static constexpr int32 kDefaultFractionQuantization = 0;
+	static constexpr int32 kDefaultUpperBoundQuantization = 0;
 
 
 	EHeuristicModel Model = EHeuristicModel::Unknown;
@@ -50,7 +51,8 @@ struct RENDERCORE_API FHeuristicSettings final
 	float TargetedHeadRoom           = kDefaultTargetedHeadRoom;
 	float IncreaseAmortizationFactor = kDefaultIncreaseAmortizationFactor;
 	int32 FractionQuantization       = kDefaultFractionQuantization;
-	
+	int32 UpperBoundQuantization     = kDefaultUpperBoundQuantization;
+
 
 	/** Returns whether the heuristic is enabled or not. */
 	bool IsEnabled() const;
diff --git a/Engine/Source/Runtime/Renderer/Private/TranslucentRendering.cpp b/Engine/Source/Runtime/Renderer/Private/TranslucentRendering.cpp
--- a/Engine/Source/Runtime/Renderer/Private/TranslucentRendering.cpp
+++ b/Engine/Source/Runtime/Renderer/Private/TranslucentRendering.cpp
@@ -73,6 +73,15 @@ static TAutoConsoleVariable<float> CVarTranslucencyChangeThreshold(
 	TEXT("Minimal increase percentage threshold to alow when changing resolution of translucency."),
 	ECVF_RenderThreadSafe | ECVF_Default);
 
+static TAutoConsoleVariable<int32> CVarTranslucencyUpperBoundQuantization(
+	TEXT("r.Translucency.DynamicRes.UpperBoundQuantization"),
+	DynamicRenderScaling::FHeuristicSettings::kDefaultUpperBoundQuantization,
+	TEXT("Quantization step count to use for upper bound screen percentage.\n")
+	TEXT("If non-zero, rendertargets will be resized based on the dynamic resolution fraction, saving GPU time during clears and resolves.\n")
+	TEXT("Only recommended for use with the transient allocator (on supported platforms) with a large transient texture cache (e.g RHI.TransientAllocator.TextureCacheSize=512)"),
+	ECVF_RenderThreadSafe | ECVF_Default);
+
+
 int32 GSeparateTranslucencyUpsampleMode = 1;
 static FAutoConsoleVariableRef CVarSeparateTranslucencyUpsampleMode(
 	TEXT("r.SeparateTranslucencyUpsampleMode"),
@@ -92,7 +101,6 @@ static TAutoConsoleVariable<int32> CVarParallelTranslucency(
 	TEXT("Toggles parallel translucency rendering. Parallel rendering must be enabled for this to have an effect."),
 	ECVF_RenderThreadSafe);
 
-
 DynamicRenderScaling::FHeuristicSettings GetDynamicTranslucencyResolutionSettings()
 {
 	DynamicRenderScaling::FHeuristicSettings BucketSetting;
@@ -103,6 +111,7 @@ DynamicRenderScaling::FHeuristicSettings GetDynamicTranslucencyResolutionSetting
 	BucketSetting.BudgetMs              = CVarTranslucencyTimeBudget.GetValueOnAnyThread();
 	BucketSetting.ChangeThreshold       = DynamicRenderScaling::GetPercentageCVarToFraction(CVarTranslucencyChangeThreshold);
 	BucketSetting.TargetedHeadRoom      = DynamicRenderScaling::GetPercentageCVarToFraction(CVarTranslucencyTargetedHeadRoomPercentage);
+	BucketSetting.UpperBoundQuantization = CVarTranslucencyUpperBoundQuantization.GetValueOnAnyThread();
 	return BucketSetting;
 }
 
