New material created with Strata enable, no longer contains a legacy conversion but a slab node instead.

diff --git a/Engine/Source/Editor/MaterialEditor/Private/MaterialEditor.cpp b/Engine/Source/Editor/MaterialEditor/Private/MaterialEditor.cpp
--- a/Engine/Source/Editor/MaterialEditor/Private/MaterialEditor.cpp
+++ b/Engine/Source/Editor/MaterialEditor/Private/MaterialEditor.cpp
@@ -4686,7 +4686,7 @@ void FMaterialEditor::OnCreateStrataNodeForPin(const FToolMenuContext& InMenuCon
 
 	// Set position of new node to be close to node we clicked on
 	FVector2D NewNodePos;
-	NewNodePos.X = PinNode->NodePosX + (bTargetPinIsInput ? -200 : 300);
+	NewNodePos.X = PinNode->NodePosX + (bTargetPinIsInput ? -300 : 300);
 	NewNodePos.Y = PinNode->NodePosY;
 
 
diff --git a/Engine/Source/Runtime/Engine/Private/Materials/Material.cpp b/Engine/Source/Runtime/Engine/Private/Materials/Material.cpp
--- a/Engine/Source/Runtime/Engine/Private/Materials/Material.cpp
+++ b/Engine/Source/Runtime/Engine/Private/Materials/Material.cpp
@@ -3118,6 +3118,15 @@ void UMaterial::ConvertMaterialToStrataMaterial()
 		ConvertLegacyToStrataBlendMode();
 		bInvalidateShader = true;
 	}
+	else if (!bUseMaterialAttributes && !EditorOnly->FrontMaterial.IsConnected() && GetExpressions().IsEmpty())
+	{
+		const int32 NewNodeOffsetX = -300;
+		UMaterialExpressionStrataSlabBSDF* SlabNode = NewObject<UMaterialExpressionStrataSlabBSDF>(this);
+		SlabNode->MaterialExpressionEditorX = EditorX + NewNodeOffsetX;
+		EditorOnly->FrontMaterial.Connect(0, SlabNode);
+		bRelinkCustomOutputNodes = false;
+		bInvalidateShader = true;
+	}
 	else if (!bUseMaterialAttributes && !EditorOnly->FrontMaterial.IsConnected())
 	{
 		if (MaterialDomain == MD_Surface)
