Marked DepthOfField* properties in FPostProcessSettings as "EditAnywhere".

diff --git a/Engine/Source/Runtime/Engine/Classes/Engine/Scene.h b/Engine/Source/Runtime/Engine/Classes/Engine/Scene.h
--- a/Engine/Source/Runtime/Engine/Classes/Engine/Scene.h
+++ b/Engine/Source/Runtime/Engine/Classes/Engine/Scene.h
@@ -1994,61 +1994,61 @@ struct FPostProcessSettings
 	TObjectPtr<class UTexture> ColorGradingLUT;
 
 	/** Width of the camera sensor to assume, in mm. */
-	UPROPERTY(BlueprintReadWrite, Category="Lens|Depth of Field", meta=(ForceUnits=mm, ClampMin = "0.1", UIMin="0.1", UIMax= "1000.0", editcondition = "bOverride_DepthOfFieldSensorWidth", DisplayName = "Sensor Width (mm)"))
+	UPROPERTY(EditAnywhere, BlueprintReadWrite, Category="Lens|Depth of Field", meta=(ForceUnits=mm, ClampMin = "0.1", UIMin="0.1", UIMax= "1000.0", editcondition = "bOverride_DepthOfFieldSensorWidth", DisplayName = "Sensor Width (mm)"))
 	float DepthOfFieldSensorWidth;
 
 	/** Width of the camera sensor to assume, in mm. */
-	UPROPERTY(BlueprintReadWrite, Category = "Lens|Depth of Field", meta = (ClampMin = "1.0", ClampMax = "2.0", editcondition = "bOverride_DepthOfFieldSqueezeFactor", DisplayName = "Squeeze Factor"))
+	UPROPERTY(EditAnywhere, BlueprintReadWrite, Category = "Lens|Depth of Field", meta = (ClampMin = "1.0", ClampMax = "2.0", editcondition = "bOverride_DepthOfFieldSqueezeFactor", DisplayName = "Squeeze Factor"))
 	float DepthOfFieldSqueezeFactor;
 
 	/** Distance in which the Depth of Field effect should be sharp, in unreal units (cm) */
-	UPROPERTY(interp, BlueprintReadWrite, Category="Lens|Depth of Field", meta=(ClampMin = "0.0", UIMin = "1.0", UIMax = "10000.0", editcondition = "bOverride_DepthOfFieldFocalDistance", DisplayName = "Focal Distance"))
+	UPROPERTY(interp, EditAnywhere, BlueprintReadWrite, Category="Lens|Depth of Field", meta=(ClampMin = "0.0", UIMin = "1.0", UIMax = "10000.0", editcondition = "bOverride_DepthOfFieldFocalDistance", DisplayName = "Focal Distance"))
 	float DepthOfFieldFocalDistance;
 
 	/** CircleDOF only: Depth blur km for 50% */
-	UPROPERTY(interp, BlueprintReadWrite, Category="Lens|Depth of Field", meta=(ClampMin = "0.000001", ClampMax = "100.0", editcondition = "bOverride_DepthOfFieldDepthBlurAmount", DisplayName = "Depth Blur km for 50%"))
+	UPROPERTY(interp, EditAnywhere, BlueprintReadWrite, Category="Lens|Depth of Field", meta=(ClampMin = "0.000001", ClampMax = "100.0", editcondition = "bOverride_DepthOfFieldDepthBlurAmount", DisplayName = "Depth Blur km for 50%"))
 	float DepthOfFieldDepthBlurAmount;
 
 	/** CircleDOF only: Depth blur radius in pixels at 1920x */
-	UPROPERTY(interp, BlueprintReadWrite, Category="Lens|Depth of Field", meta=(ClampMin = "0.0", UIMax = "4.0", editcondition = "bOverride_DepthOfFieldDepthBlurRadius", DisplayName = "Depth Blur Radius"))
+	UPROPERTY(interp, EditAnywhere, BlueprintReadWrite, Category="Lens|Depth of Field", meta=(ClampMin = "0.0", UIMax = "4.0", editcondition = "bOverride_DepthOfFieldDepthBlurRadius", DisplayName = "Depth Blur Radius"))
 	float DepthOfFieldDepthBlurRadius;
 
 	/** Artificial region where all content is in focus, starting after DepthOfFieldFocalDistance, in unreal units  (cm) */
-	UPROPERTY(interp, BlueprintReadWrite, Category="Lens|Mobile Depth of Field", meta=(UIMin = "0.0", UIMax = "10000.0", editcondition = "bOverride_DepthOfFieldFocalRegion", DisplayName = "Focal Region"))
+	UPROPERTY(interp, EditAnywhere, BlueprintReadWrite, Category="Lens|Mobile Depth of Field", meta=(UIMin = "0.0", UIMax = "10000.0", editcondition = "bOverride_DepthOfFieldFocalRegion", DisplayName = "Focal Region"))
 	float DepthOfFieldFocalRegion;
 
 	/** To define the width of the transition region next to the focal region on the near side (cm) */
-	UPROPERTY(interp, BlueprintReadWrite, Category="Lens|Mobile Depth of Field", meta=(UIMin = "0.0", UIMax = "10000.0", editcondition = "bOverride_DepthOfFieldNearTransitionRegion", DisplayName = "Near Transition Region"))
+	UPROPERTY(interp, EditAnywhere, BlueprintReadWrite, Category="Lens|Mobile Depth of Field", meta=(UIMin = "0.0", UIMax = "10000.0", editcondition = "bOverride_DepthOfFieldNearTransitionRegion", DisplayName = "Near Transition Region"))
 	float DepthOfFieldNearTransitionRegion;
 
 	/** To define the width of the transition region next to the focal region on the near side (cm) */
-	UPROPERTY(interp, BlueprintReadWrite, Category="Lens|Mobile Depth of Field", meta=(UIMin = "0.0", UIMax = "10000.0", editcondition = "bOverride_DepthOfFieldFarTransitionRegion", DisplayName = "Far Transition Region"))
+	UPROPERTY(interp, EditAnywhere, BlueprintReadWrite, Category="Lens|Mobile Depth of Field", meta=(UIMin = "0.0", UIMax = "10000.0", editcondition = "bOverride_DepthOfFieldFarTransitionRegion", DisplayName = "Far Transition Region"))
 	float DepthOfFieldFarTransitionRegion;
 
 	/** SM5: BokehDOF only: To amplify the depth of field effect (like aperture)  0=off 
 	    ES3_1: Used to blend DoF. 0=off
 	*/
-	UPROPERTY(interp, BlueprintReadWrite, Category="Lens|Mobile Depth of Field", meta=(ClampMin = "0.0", ClampMax = "2.0", editcondition = "bOverride_DepthOfFieldScale", DisplayName = "Scale"))
+	UPROPERTY(interp, EditAnywhere, BlueprintReadWrite, Category="Lens|Mobile Depth of Field", meta=(ClampMin = "0.0", ClampMax = "2.0", editcondition = "bOverride_DepthOfFieldScale", DisplayName = "Scale"))
 	float DepthOfFieldScale;
 
 	/** Gaussian only: Maximum size of the Depth of Field blur (in percent of the view width) (note: performance cost scales with size) */
-	UPROPERTY(interp, BlueprintReadWrite, Category="Lens|Mobile Depth of Field", meta=(UIMin = "0.0", UIMax = "32.0", editcondition = "bOverride_DepthOfFieldNearBlurSize", DisplayName = "Near Blur Size"))
+	UPROPERTY(interp, EditAnywhere, BlueprintReadWrite, Category="Lens|Mobile Depth of Field", meta=(UIMin = "0.0", UIMax = "32.0", editcondition = "bOverride_DepthOfFieldNearBlurSize", DisplayName = "Near Blur Size"))
 	float DepthOfFieldNearBlurSize;
 
 	/** Gaussian only: Maximum size of the Depth of Field blur (in percent of the view width) (note: performance cost scales with size) */
-	UPROPERTY(interp, BlueprintReadWrite, Category="Lens|Mobile Depth of Field", meta=(UIMin = "0.0", UIMax = "32.0", editcondition = "bOverride_DepthOfFieldFarBlurSize", DisplayName = "Far Blur Size"))
+	UPROPERTY(interp, EditAnywhere, BlueprintReadWrite, Category="Lens|Mobile Depth of Field", meta=(UIMin = "0.0", UIMax = "32.0", editcondition = "bOverride_DepthOfFieldFarBlurSize", DisplayName = "Far Blur Size"))
 	float DepthOfFieldFarBlurSize;
 
 	/** Occlusion tweak factor 1 (0.18 to get natural occlusion, 0.4 to solve layer color leaking issues) */
-	UPROPERTY(interp, BlueprintReadWrite, Category="Lens|Mobile Depth of Field", AdvancedDisplay, meta=(ClampMin = "0.0", ClampMax = "1.0", editcondition = "bOverride_DepthOfFieldOcclusion", DisplayName = "Occlusion"))
+	UPROPERTY(interp, EditAnywhere, BlueprintReadWrite, Category="Lens|Mobile Depth of Field", AdvancedDisplay, meta=(ClampMin = "0.0", ClampMax = "1.0", editcondition = "bOverride_DepthOfFieldOcclusion", DisplayName = "Occlusion"))
 	float DepthOfFieldOcclusion;
 
 	/** Artificial distance to allow the skybox to be in focus (e.g. 200000), <=0 to switch the feature off, only for GaussianDOF, can cost performance */
-	UPROPERTY(interp, BlueprintReadWrite, Category="Lens|Mobile Depth of Field", AdvancedDisplay, meta=(ClampMin = "0.0", ClampMax = "200000.0", editcondition = "bOverride_DepthOfFieldSkyFocusDistance", DisplayName = "Sky Distance"))
+	UPROPERTY(interp, EditAnywhere, BlueprintReadWrite, Category="Lens|Mobile Depth of Field", AdvancedDisplay, meta=(ClampMin = "0.0", ClampMax = "200000.0", editcondition = "bOverride_DepthOfFieldSkyFocusDistance", DisplayName = "Sky Distance"))
 	float DepthOfFieldSkyFocusDistance;
 
 	/** Artificial circular mask to (near) blur content outside the radius, only for GaussianDOF, diameter in percent of screen width, costs performance if the mask is used, keep Feather can Radius on default to keep it off */
-	UPROPERTY(interp, BlueprintReadWrite, Category="Lens|Mobile Depth of Field", AdvancedDisplay, meta=(UIMin = "0.0", UIMax = "100.0", editcondition = "bOverride_DepthOfFieldVignetteSize", DisplayName = "Vignette Size"))
+	UPROPERTY(interp, EditAnywhere, BlueprintReadWrite, Category="Lens|Mobile Depth of Field", AdvancedDisplay, meta=(UIMin = "0.0", UIMax = "100.0", editcondition = "bOverride_DepthOfFieldVignetteSize", DisplayName = "Vignette Size"))
 	float DepthOfFieldVignetteSize;
 
 	/** Strength of motion blur, 0:off, should be renamed to intensity */
