Fix for moving foliage instances in editor causing PerInstanceRandom corruption.

#rnx

diff --git a/Engine/Source/Runtime/Engine/Private/InstancedStaticMesh.cpp b/Engine/Source/Runtime/Engine/Private/InstancedStaticMesh.cpp
--- a/Engine/Source/Runtime/Engine/Private/InstancedStaticMesh.cpp
+++ b/Engine/Source/Runtime/Engine/Private/InstancedStaticMesh.cpp
@@ -487,7 +487,7 @@ void FStaticMeshInstanceBuffer::UpdateFromCommandBuffer_RenderThread(FInstanceUp
 			InstanceData->NullifyInstance(InstanceIndex);
 			break;
 		case FInstanceUpdateCmdBuffer::Update:
-			InstanceData->SetInstance(InstanceIndex, FMatrix44f(Cmd.XForm), 0);		// LWC_TODO: precision loss?
+			InstanceData->SetInstance(InstanceIndex, FMatrix44f(Cmd.XForm));		// LWC_TODO: precision loss?
 			break;
 		case FInstanceUpdateCmdBuffer::EditorData:
 			InstanceData->SetInstanceEditorData(InstanceIndex, Cmd.HitProxyColor, Cmd.bSelected);
diff --git a/Engine/Source/Runtime/Engine/Public/StaticMeshResources.h b/Engine/Source/Runtime/Engine/Public/StaticMeshResources.h
--- a/Engine/Source/Runtime/Engine/Public/StaticMeshResources.h
+++ b/Engine/Source/Runtime/Engine/Public/StaticMeshResources.h
@@ -1284,32 +1284,18 @@ public:
 
 	FORCEINLINE void SetInstance(int32 InstanceIndex, const FMatrix44f& Transform, const FVector2D& LightmapUVBias, const FVector2D& ShadowmapUVBias)
 	{
-		FVector4f OldOrigin;
-		GetInstanceOriginInternal(InstanceIndex, OldOrigin);
-
-		FVector4f NewOrigin(Transform.M[3][0], Transform.M[3][1], Transform.M[3][2], OldOrigin.Component(3));
-		SetInstanceOriginInternal(InstanceIndex, NewOrigin);
-
-		FVector4f InstanceTransform[3];
-		InstanceTransform[0] = FVector4f(Transform.M[0][0], Transform.M[0][1], Transform.M[0][2], 0.0f);
-		InstanceTransform[1] = FVector4f(Transform.M[1][0], Transform.M[1][1], Transform.M[1][2], 0.0f);
-		InstanceTransform[2] = FVector4f(Transform.M[2][0], Transform.M[2][1], Transform.M[2][2], 0.0f);
-
-		if (bUseHalfFloat)
-		{
-			SetInstanceTransformInternal<FFloat16>(InstanceIndex, InstanceTransform);
-		}
-		else
-		{
-			SetInstanceTransformInternal<float>(InstanceIndex, InstanceTransform);
-		}
-
-		SetInstanceLightMapDataInternal(InstanceIndex, FVector4f(LightmapUVBias.X, LightmapUVBias.Y, ShadowmapUVBias.X, ShadowmapUVBias.Y));
+		float RandomInstanceID;
+		GetInstanceRandomID(InstanceIndex, RandomInstanceID);
+		SetInstance(InstanceIndex, Transform, RandomInstanceID, LightmapUVBias, ShadowmapUVBias);
+	}
 
-		for (int32 i = 0; i < NumCustomDataFloats; ++i)
-		{
-			SetInstanceCustomDataInternal(InstanceIndex, i, 0);
-		}
+	FORCEINLINE_DEBUGGABLE void SetInstance(int32 InstanceIndex, const FMatrix44f& Transform)
+	{
+		const FVector2D& LightmapUVBias = FVector2D::ZeroVector;
+		const FVector2D& ShadowmapUVBias = FVector2D::ZeroVector;
+		float RandomInstanceID;
+		GetInstanceRandomID(InstanceIndex, RandomInstanceID);
+		SetInstance(InstanceIndex, Transform, RandomInstanceID, LightmapUVBias, ShadowmapUVBias);
 	}
 
 	FORCEINLINE void SetInstanceLightMapData(int32 InstanceIndex, const FVector2D& LightmapUVBias, const FVector2D& ShadowmapUVBias)
