Scene Capture:  Fix issue where scene captures go black in the editor when the texure size is edited when "bCaptureEveryFrame" is false.  Need to request a refresh at the end of an edit in that case (in addition to during interactive edits in any case).

#rnx

diff --git a/Engine/Source/Runtime/Engine/Private/TextureRenderTarget2D.cpp b/Engine/Source/Runtime/Engine/Private/TextureRenderTarget2D.cpp
--- a/Engine/Source/Runtime/Engine/Private/TextureRenderTarget2D.cpp
+++ b/Engine/Source/Runtime/Engine/Private/TextureRenderTarget2D.cpp
@@ -214,19 +214,23 @@ void UTextureRenderTarget2D::PostEditChangeProperty(FPropertyChangedEvent& Prope
     // SRGB may have been changed by Super, reset it since we prefer to honor explicit user choice
 	SRGB = IsSRGB();
 
-	// If this is an interactive size edit, notify any scene capture component that points to this render target that it needs to refresh.
-	// During interactive edits, time is paused, so the Tick function which normally handles capturing isn't called.
+	// Notify any scene capture components that point to this texture that they may need to refresh
 	static const FName SizeXName = GET_MEMBER_NAME_CHECKED(UTextureRenderTarget2D, SizeX);
 	static const FName SizeYName = GET_MEMBER_NAME_CHECKED(UTextureRenderTarget2D, SizeY);
 
-	if ((PropertyChangedEvent.ChangeType & EPropertyChangeType::Interactive) && (PropertyChangedEvent.GetPropertyName() == SizeXName || PropertyChangedEvent.GetPropertyName() == SizeYName))
+	if ((PropertyChangedEvent.GetPropertyName() == SizeXName || PropertyChangedEvent.GetPropertyName() == SizeYName))
 	{
 		for (TObjectIterator<USceneCaptureComponent2D> It; It; ++It)
 		{
 			USceneCaptureComponent2D* SceneCaptureComponent = *It;
 			if (SceneCaptureComponent->TextureTarget == this)
 			{
-				SceneCaptureComponent->CaptureSceneDeferred();
+				// During interactive edits, time is paused, so the Tick function which normally handles capturing isn't called, and we
+				// need a manual refresh.  We also need a refresh if the capture doesn't happen automatically every frame.
+				if ((PropertyChangedEvent.ChangeType & EPropertyChangeType::Interactive) || !SceneCaptureComponent->bCaptureEveryFrame)
+				{
+					SceneCaptureComponent->CaptureSceneDeferred();
+				}
 			}
 		}
 	}
diff --git a/Engine/Source/Runtime/Engine/Private/TextureRenderTargetCube.cpp b/Engine/Source/Runtime/Engine/Private/TextureRenderTargetCube.cpp
--- a/Engine/Source/Runtime/Engine/Private/TextureRenderTargetCube.cpp
+++ b/Engine/Source/Runtime/Engine/Private/TextureRenderTargetCube.cpp
@@ -110,18 +110,22 @@ void UTextureRenderTargetCube::PostEditChangeProperty(FPropertyChangedEvent& Pro
 
 	Super::PostEditChangeProperty(PropertyChangedEvent);
 
-	// If this is an interactive size edit, notify any scene capture component that points to this render target that it needs to refresh.
-	// During interactive edits, time is paused, so the Tick function which normally handles capturing isn't called.
+	// Notify any scene capture components that point to this texture that they may need to refresh
 	static const FName SizeXName = GET_MEMBER_NAME_CHECKED(UTextureRenderTargetCube, SizeX);
 
-	if ((PropertyChangedEvent.ChangeType & EPropertyChangeType::Interactive) && PropertyChangedEvent.GetPropertyName() == SizeXName)
+	if (PropertyChangedEvent.GetPropertyName() == SizeXName)
 	{
 		for (TObjectIterator<USceneCaptureComponentCube> It; It; ++It)
 		{
 			USceneCaptureComponentCube* SceneCaptureComponent = *It;
 			if (SceneCaptureComponent->TextureTarget == this)
 			{
-				SceneCaptureComponent->CaptureSceneDeferred();
+				// During interactive edits, time is paused, so the Tick function which normally handles capturing isn't called, and we
+				// need a manual refresh.  We also need a refresh if the capture doesn't happen automatically every frame.
+				if ((PropertyChangedEvent.ChangeType & EPropertyChangeType::Interactive) || !SceneCaptureComponent->bCaptureEveryFrame)
+				{
+					SceneCaptureComponent->CaptureSceneDeferred();
+				}
 			}
 		}
 	}
