Some FDiffableObjectDataWriter perf improvements

 - Switched SerializedTaggedData to be a TSortedMap now that it no longer has to be a UPROPERTY, as property counts for a type are typically small.
 - Added a cache to avoid redundant searches for FSerializedTaggedData when the tagged data key hasn't changed between writes (eg, when serializing arrays or structs).

#rnx

diff --git a/Engine/Source/Runtime/Engine/Private/TransactionCommon.cpp b/Engine/Source/Runtime/Engine/Private/TransactionCommon.cpp
--- a/Engine/Source/Runtime/Engine/Private/TransactionCommon.cpp
+++ b/Engine/Source/Runtime/Engine/Private/TransactionCommon.cpp
@@ -326,7 +326,7 @@ FName FDiffableObjectDataWriter::GetTaggedDataKey() const
 		const FArchiveSerializedPropertyChain* PropertyChain = GetSerializedPropertyChain();
 		if (PropertyChain && PropertyChain->GetNumProperties() > 0)
 		{
-			TaggedDataKey = CachedSerializedTaggedPropertyKey.SyncCache(PropertyChain);
+			TaggedDataKey = CachedSerializedTaggedPropertyKey.SyncCache(*PropertyChain);
 		}
 	}
 
@@ -397,7 +397,7 @@ void FDiffableObjectDataWriter::OnDataSerialized(int64 InOffset, int64 InNum)
 	const FName SerializedTaggedDataKey = GetTaggedDataKey();
 	if (!SerializedTaggedDataKey.IsNone())
 	{
-		FSerializedTaggedData& SerializedTaggedData = DiffableObject.SerializedTaggedData.FindOrAdd(SerializedTaggedDataKey);
+		FSerializedTaggedData& SerializedTaggedData = CachedSerializedTaggedDataEntry.SyncCache(DiffableObject, SerializedTaggedDataKey);
 		SerializedTaggedData.AppendSerializedData(InOffset, InNum);
 	}
 }
@@ -421,24 +421,27 @@ FArchive& FDiffableObjectDataWriter::operator<<(UObject*& Res)
 	return *this << ResPathName;
 }
 
-FName FDiffableObjectDataWriter::FCachedPropertyKey::SyncCache(const FArchiveSerializedPropertyChain* InPropertyChain)
+FName FDiffableObjectDataWriter::FCachedPropertyKey::SyncCache(const FArchiveSerializedPropertyChain& InPropertyChain)
 {
-	if (InPropertyChain)
+	const uint32 CurrentUpdateCount = InPropertyChain.GetUpdateCount();
+	if (CurrentUpdateCount != LastUpdateCount)
 	{
-		const uint32 CurrentUpdateCount = InPropertyChain->GetUpdateCount();
-		if (CurrentUpdateCount != LastUpdateCount)
-		{
-			CachedKey = InPropertyChain->GetNumProperties() > 0 ? InPropertyChain->GetPropertyFromRoot(0)->GetFName() : FName();
-			LastUpdateCount = CurrentUpdateCount;
-		}
+		CachedKey = InPropertyChain.GetNumProperties() > 0 ? InPropertyChain.GetPropertyFromRoot(0)->GetFName() : FName();
+		LastUpdateCount = CurrentUpdateCount;
 	}
-	else
+
+	return CachedKey;
+}
+
+FSerializedTaggedData& FDiffableObjectDataWriter::FCachedTaggedDataEntry::SyncCache(FDiffableObject& InDiffableObject, const FName InSerializedTaggedDataKey)
+{
+	if (CachedKey != InSerializedTaggedDataKey)
 	{
-		CachedKey = FName();
-		LastUpdateCount = 0;
+		CachedKey = InSerializedTaggedDataKey;
+		CachedEntryPtr = &InDiffableObject.SerializedTaggedData.FindOrAdd(InSerializedTaggedDataKey);
 	}
 
-	return CachedKey;
+	return *CachedEntryPtr;
 }
 
 
diff --git a/Engine/Source/Runtime/Engine/Public/TransactionCommon.h b/Engine/Source/Runtime/Engine/Public/TransactionCommon.h
--- a/Engine/Source/Runtime/Engine/Public/TransactionCommon.h
+++ b/Engine/Source/Runtime/Engine/Public/TransactionCommon.h
@@ -9,6 +9,7 @@
 #include "Serialization/ArchiveSerializedPropertyChain.h"
 #include "Templates/TypeHash.h"
 #include "Misc/TransactionObjectEvent.h"
+#include "Containers/SortedMap.h"
 
 class FReferenceCollector;
 
@@ -274,7 +275,7 @@ public:
 	FSerializedObjectData SerializedData;
 
 	/** Information about tagged data (mainly properties) that were serialized within this object */
-	TMap<FName, FSerializedTaggedData> SerializedTaggedData;
+	TSortedMap<FName, FSerializedTaggedData, FDefaultAllocator, FNameFastLess> SerializedTaggedData;
 };
 
 ENGINE_API extern const FName TaggedDataKey_UnknownData;
@@ -369,13 +370,23 @@ private:
 	struct FCachedPropertyKey
 	{
 	public:
-		FName SyncCache(const FArchiveSerializedPropertyChain* InPropertyChain);
+		FName SyncCache(const FArchiveSerializedPropertyChain& InPropertyChain);
 
 	private:
 		FName CachedKey;
 		uint32 LastUpdateCount = 0;
 	};
 
+	struct FCachedTaggedDataEntry
+	{
+	public:
+		FSerializedTaggedData& SyncCache(FDiffableObject& InDiffableObject, const FName InSerializedTaggedDataKey);
+
+	private:
+		FName CachedKey;
+		FSerializedTaggedData* CachedEntryPtr = nullptr;
+	};
+
 	FDiffableObject& DiffableObject;
 	TArrayView<const FProperty*> PropertiesToSerialize;
 
@@ -386,6 +397,7 @@ private:
 	mutable int32 TaggedDataKeyIndex_UnknownData = 0;
 	mutable int32 TaggedDataKeyIndex_ScriptData = 0;
 	mutable FCachedPropertyKey CachedSerializedTaggedPropertyKey;
+	mutable FCachedTaggedDataEntry CachedSerializedTaggedDataEntry;
 };
 
 namespace DiffUtil
