Change default value of looping flag
#rnx

diff --git a/Engine/Source/Runtime/Engine/Private/Animation/AnimSequenceBase.cpp b/Engine/Source/Runtime/Engine/Private/Animation/AnimSequenceBase.cpp
--- a/Engine/Source/Runtime/Engine/Private/Animation/AnimSequenceBase.cpp
+++ b/Engine/Source/Runtime/Engine/Private/Animation/AnimSequenceBase.cpp
@@ -32,7 +32,7 @@ CSV_DECLARE_CATEGORY_MODULE_EXTERN(ENGINE_API, Animation);
 UAnimSequenceBase::UAnimSequenceBase(const FObjectInitializer& ObjectInitializer)
 	: Super(ObjectInitializer)
 	, RateScale(1.0f)
-	, bLoop(true)
+	, bLoop(false)
 #if WITH_EDITORONLY_DATA
 	, DataModel(nullptr)
 	, bPopulatingDataModel(false)
