ActorComponent
* ActorComponents can now set their default replication condition via GetReplicationCondition.

diff --git a/Engine/Source/Runtime/Engine/Classes/Components/ActorComponent.h b/Engine/Source/Runtime/Engine/Classes/Components/ActorComponent.h
--- a/Engine/Source/Runtime/Engine/Classes/Components/ActorComponent.h
+++ b/Engine/Source/Runtime/Engine/Classes/Components/ActorComponent.h
@@ -503,6 +503,9 @@ public:
 	*/
 	bool IsUsingRegisteredSubObjectList() const { return bReplicateUsingRegisteredSubObjectList; }
 
+	/** Returns the replication condition for the component. Only works if the actor owning the component has bReplicateUsingRegisteredSubObjectList enabled. */
+	virtual ELifetimeCondition GetReplicationCondition() const { return GetIsReplicated() ? COND_None : COND_Never;  }
+
 	/** Called on the component right before replication occurs */
 	virtual void PreReplication(IRepChangedPropertyTracker & ChangedPropertyTracker);
 
diff --git a/Engine/Source/Runtime/Engine/Private/ActorReplication.cpp b/Engine/Source/Runtime/Engine/Private/ActorReplication.cpp
--- a/Engine/Source/Runtime/Engine/Private/ActorReplication.cpp
+++ b/Engine/Source/Runtime/Engine/Private/ActorReplication.cpp
@@ -504,7 +504,7 @@ bool AActor::ReplicateSubobjects(UActorChannel *Channel, FOutBunch *Bunch, FRepl
 
 	for (UActorComponent* ActorComp : ReplicatedComponents)
 	{
-		if (ActorComp && ActorComp->GetIsReplicated())
+		if (ActorComp && ActorComp->GetReplicationCondition() != COND_Never)
 		{
 			UActorChannel::SetCurrentSubObjectOwner(ActorComp);
 			WroteSomething |= ActorComp->ReplicateSubobjects(Channel, Bunch, RepFlags);		// Lets the component add subobjects before replicating its own properties.
@@ -517,7 +517,7 @@ bool AActor::ReplicateSubobjects(UActorChannel *Channel, FOutBunch *Bunch, FRepl
 
 ELifetimeCondition AActor::AllowActorComponentToReplicate(const UActorComponent* ComponentToReplicate) const
 {
-	return ComponentToReplicate->GetIsReplicated() ? COND_None : COND_Never;
+	return ComponentToReplicate->GetReplicationCondition();
 }
 
 void AActor::SetReplicatedComponentNetCondition(const UActorComponent* ReplicatedComponent, ELifetimeCondition NetCondition)
