Pull-request from github: Fixes crash in bone rendering when BoneColors or WorldTransforms arrays are different size than the bones array which can happen with Pose Watch.

#9452

diff --git a/Engine/Source/Runtime/Engine/Private/SkeletalDebugRendering.cpp b/Engine/Source/Runtime/Engine/Private/SkeletalDebugRendering.cpp
--- a/Engine/Source/Runtime/Engine/Private/SkeletalDebugRendering.cpp
+++ b/Engine/Source/Runtime/Engine/Private/SkeletalDebugRendering.cpp
@@ -335,6 +335,11 @@ void DrawBones(
 	for (int32 Index = 0; Index < RequiredBones.Num(); ++Index)
 	{
 		const int32 BoneIndex = RequiredBones[Index];
+		if ((!BoneColors.IsEmpty() && BoneIndex >= BoneColors.Num()) ||
+			BoneIndex >= WorldTransforms.Num())
+		{
+			continue;
+		}
 
 		// skips bones that should not be drawn
 		const bool bDoDraw = DrawConfig.bForceDraw || DrawConfig.BoneDrawMode == EBoneDrawMode::All || BonesToDraw[BoneIndex];
@@ -360,6 +365,10 @@ void DrawBones(
 		for (int32 ChildIndex = 0; ChildIndex < RefSkeleton.GetNum(); ++ChildIndex)
 		{
 			const int32 ParentIndex = RefSkeleton.GetParentIndex(ChildIndex);
+			if (ParentIndex >= WorldTransforms.Num())
+			{
+				continue;
+			}
 			if (ParentIndex == BoneIndex && RequiredBones.Contains(ChildIndex))
 			{
 				ChildPositions.Add(WorldTransforms[ChildIndex].GetLocation());
