Chaos - fix inertia tensor with CoM Nudge and an inertia scale (could lead to negative inertia)

diff --git a/Engine/Source/Runtime/Engine/Private/PhysicsEngine/BodyUtils.cpp b/Engine/Source/Runtime/Engine/Private/PhysicsEngine/BodyUtils.cpp
--- a/Engine/Source/Runtime/Engine/Private/PhysicsEngine/BodyUtils.cpp
+++ b/Engine/Source/Runtime/Engine/Private/PhysicsEngine/BodyUtils.cpp
@@ -76,15 +76,6 @@ namespace BodyUtils
 		MassProps.InertiaTensor *= MassRatio;
 		MassProps.CenterOfMass += MassModifierTransform.TransformVector(OwningBodyInstance->COMNudge);
 
-		// If we move the center of mass, we need to update the inertia using parallel-axis theorem. If we don't do this
-		// and the center of mass is moved significantly it can cause jitter (inertia too small for the contact positions)
-		if (CVars::bPhysicsComNudgeAdjustInertia)
-		{
-			MassProps.InertiaTensor.M[0][0] += MassProps.Mass * OwningBodyInstance->COMNudge.X * OwningBodyInstance->COMNudge.X;
-			MassProps.InertiaTensor.M[1][1] += MassProps.Mass * OwningBodyInstance->COMNudge.Y * OwningBodyInstance->COMNudge.Y;
-			MassProps.InertiaTensor.M[2][2] += MassProps.Mass * OwningBodyInstance->COMNudge.Z * OwningBodyInstance->COMNudge.Z;
-		}
-
 		// Scale the inertia tensor by the owning body instance's InertiaTensorScale
 		// NOTE: PhysX scales the inertia by the mass increase we would get from the scale change, even though we 
 		// don't actually scale the mass at all based on InertiaScale. This is non-intuituve. E.g., you may expect 
@@ -99,6 +90,17 @@ namespace BodyUtils
 			MassProps.InertiaTensor = Chaos::Utilities::ScaleInertia(MassProps.InertiaTensor, OwningBodyInstance->InertiaTensorScale, bInertaScaleIncludeMass);
 		}
 
+		// If we move the center of mass, we need to update the inertia using parallel-axis theorem. If we don't do this
+		// and the center of mass is moved significantly it can cause jitter (inertia too small for the contact positions)
+		// NOTE: This must come after ScaleInertia because ScaleInertia effectively calculates the "equivalent box" dimensions
+		// which is not always possible, e.g., if we move the CoM outside of a box you will get negative elements in the scaled inertia!
+		if (CVars::bPhysicsComNudgeAdjustInertia)
+		{
+			MassProps.InertiaTensor.M[0][0] += MassProps.Mass * OwningBodyInstance->COMNudge.X * OwningBodyInstance->COMNudge.X;
+			MassProps.InertiaTensor.M[1][1] += MassProps.Mass * OwningBodyInstance->COMNudge.Y * OwningBodyInstance->COMNudge.Y;
+			MassProps.InertiaTensor.M[2][2] += MassProps.Mass * OwningBodyInstance->COMNudge.Z * OwningBodyInstance->COMNudge.Z;
+		}
+
 		return MassProps;
 	}
 
