World Partition Server Streaming: show if streaming sources are local or remote in debug display.

#rnx

diff --git a/Engine/Source/Runtime/Engine/Private/WorldPartition/WorldPartitionStreamingPolicy.cpp b/Engine/Source/Runtime/Engine/Private/WorldPartition/WorldPartitionStreamingPolicy.cpp
--- a/Engine/Source/Runtime/Engine/Private/WorldPartition/WorldPartitionStreamingPolicy.cpp
+++ b/Engine/Source/Runtime/Engine/Private/WorldPartition/WorldPartitionStreamingPolicy.cpp
@@ -129,7 +129,7 @@ void UWorldPartitionStreamingPolicy::UpdateStreamingSources()
 			const FVector ViewLocation = WorldToLocal.TransformPosition(GCurrentLevelEditingViewportClient->GetViewLocation());
 			const FRotator ViewRotation = WorldToLocal.TransformRotation(GCurrentLevelEditingViewportClient->GetViewRotation().Quaternion()).Rotator();
 			static const FName NAME_SIE(TEXT("SIE"));
-			StreamingSources.Add(FWorldPartitionStreamingSource(NAME_SIE, ViewLocation, ViewRotation, EStreamingSourceTargetState::Activated, /*bBlockOnSlowLoading=*/false, EStreamingSourcePriority::Default));
+			StreamingSources.Add(FWorldPartitionStreamingSource(NAME_SIE, ViewLocation, ViewRotation, EStreamingSourceTargetState::Activated, /*bBlockOnSlowLoading=*/false, EStreamingSourcePriority::Default, false));
 		}
 		else
 #endif
@@ -153,7 +153,8 @@ void UWorldPartitionStreamingPolicy::UpdateStreamingSources()
 							const EStreamingSourceTargetState TargetState = PlayerController->StreamingSourceShouldActivate() ? EStreamingSourceTargetState::Activated : EStreamingSourceTargetState::Loaded;
 							const bool bBlockOnSlowLoading = PlayerController->StreamingSourceShouldBlockOnSlowStreaming();
 							const EStreamingSourcePriority StreamingSourcePriority = PlayerController->GetStreamingSourcePriority();
-							StreamingSources.Add(FWorldPartitionStreamingSource(PlayerController->GetFName(), ViewLocation, ViewRotation, TargetState, bBlockOnSlowLoading, StreamingSourcePriority));
+							const bool bRemote = !PlayerController->IsLocalController();
+							StreamingSources.Add(FWorldPartitionStreamingSource(PlayerController->GetFName(), ViewLocation, ViewRotation, TargetState, bBlockOnSlowLoading, StreamingSourcePriority, bRemote));
 						}
 					}
 				}
diff --git a/Engine/Source/Runtime/Engine/Public/WorldPartition/WorldPartitionReplay.h b/Engine/Source/Runtime/Engine/Public/WorldPartition/WorldPartitionReplay.h
--- a/Engine/Source/Runtime/Engine/Public/WorldPartition/WorldPartitionReplay.h
+++ b/Engine/Source/Runtime/Engine/Public/WorldPartition/WorldPartitionReplay.h
@@ -21,7 +21,7 @@ struct FWorldPartitionReplayStreamingSource : public FWorldPartitionStreamingSou
 	}
 
 	FWorldPartitionReplayStreamingSource(FName InName, const FVector& InLocation, const FRotator& InRotation, EStreamingSourceTargetState InTargetState, bool bInBlockOnSlowLoading, EStreamingSourcePriority InPriority, float InVelocity)
-		: FWorldPartitionStreamingSource(InName, InLocation, InRotation, InTargetState, bInBlockOnSlowLoading, InPriority, InVelocity)
+		: FWorldPartitionStreamingSource(InName, InLocation, InRotation, InTargetState, bInBlockOnSlowLoading, InPriority, false, InVelocity)
 	{
 		bReplay = true;
 	}
diff --git a/Engine/Source/Runtime/Engine/Public/WorldPartition/WorldPartitionStreamingSource.h b/Engine/Source/Runtime/Engine/Public/WorldPartition/WorldPartitionStreamingSource.h
--- a/Engine/Source/Runtime/Engine/Public/WorldPartition/WorldPartitionStreamingSource.h
+++ b/Engine/Source/Runtime/Engine/Public/WorldPartition/WorldPartitionStreamingSource.h
@@ -325,9 +325,10 @@ struct ENGINE_API FWorldPartitionStreamingSource
 		, DebugColor(ForceInit)
 		, TargetGrid(NAME_None)
 		, bReplay(false)
+		, bRemote(false)
 	{}
 
-	FWorldPartitionStreamingSource(FName InName, const FVector& InLocation, const FRotator& InRotation, EStreamingSourceTargetState InTargetState, bool bInBlockOnSlowLoading, EStreamingSourcePriority InPriority, float InVelocity = 0.f)
+	FWorldPartitionStreamingSource(FName InName, const FVector& InLocation, const FRotator& InRotation, EStreamingSourceTargetState InTargetState, bool bInBlockOnSlowLoading, EStreamingSourcePriority InPriority, bool bRemote, float InVelocity = 0.f)
 		: Name(InName)
 		, Location(InLocation)
 		, Rotation(InRotation)
@@ -338,6 +339,7 @@ struct ENGINE_API FWorldPartitionStreamingSource
 		, DebugColor(ForceInit)
 		, TargetGrid(NAME_None)
 		, bReplay(false)
+		, bRemote(bRemote)
 	{}
 
 	FColor GetDebugColor() const
@@ -386,6 +388,9 @@ struct ENGINE_API FWorldPartitionStreamingSource
 	/** If true, this streaming source is from a replay recording */
 	bool bReplay;
 
+	/** If true, this streaming source is from a remote session */
+	bool bRemote;
+
 	/** Returns a box encapsulating all shapes. */
 	FORCEINLINE FBox CalcBounds(float InGridLoadingRange, FName InGridName, const FSoftObjectPath& InGridHLODLayer, bool bCalcIn2D = false) const
 	{
@@ -410,8 +415,9 @@ struct ENGINE_API FWorldPartitionStreamingSource
 	{
 		const FVector Direction = Rotation.Euler();
 		return FString::Printf(
-			TEXT("Priority: %d | %s | %s | Pos: X=%lld,Y=%lld,Z=%lld | Rot: X=%d,Y=%d,Z=%d | Vel: %3.2f m/s (%d mph)"), 
+			TEXT("Priority: %d | %s | %s | %s | Pos: X=%lld,Y=%lld,Z=%lld | Rot: X=%d,Y=%d,Z=%d | Vel: %3.2f m/s (%d mph)"), 
 			Priority, 
+			bRemote ? TEXT("Remote") : TEXT("Local"),
 			GetStreamingSourceTargetStateName(TargetState),
 			bBlockOnSlowLoading ? TEXT("Blocking") : TEXT("NonBlocking"),
 			(int64)Location.X, (int64)Location.Y, (int64)Location.Z, 
