Anim Notify States now get informed if they reached the end or were cancelled before it

diff --git a/Engine/Source/Runtime/Engine/Classes/Animation/AnimNotifies/AnimNotify.h b/Engine/Source/Runtime/Engine/Classes/Animation/AnimNotifies/AnimNotify.h
--- a/Engine/Source/Runtime/Engine/Classes/Animation/AnimNotifies/AnimNotify.h
+++ b/Engine/Source/Runtime/Engine/Classes/Animation/AnimNotifies/AnimNotify.h
@@ -23,6 +23,7 @@ public:
 	UAnimSequenceBase* SequenceAsset;
 	FAnimNotifyEvent* NotifyEvent;
 	int32 MontageInstanceID;
+	bool bReachedEnd = false;
 
 	FBranchingPointNotifyPayload()
 		: SkelMeshComponent(nullptr)
@@ -31,11 +32,12 @@ public:
 		, MontageInstanceID(INDEX_NONE)
 	{}
 
-	FBranchingPointNotifyPayload(USkeletalMeshComponent* InSkelMeshComponent, UAnimSequenceBase* InSequenceAsset, FAnimNotifyEvent* InNotifyEvent, int32 InMontageInstanceID)
+	FBranchingPointNotifyPayload(USkeletalMeshComponent* InSkelMeshComponent, UAnimSequenceBase* InSequenceAsset, FAnimNotifyEvent* InNotifyEvent, int32 InMontageInstanceID, bool bInReachedEnd = false)
 		: SkelMeshComponent(InSkelMeshComponent)
 		, SequenceAsset(InSequenceAsset)
 		, NotifyEvent(InNotifyEvent)
 		, MontageInstanceID(InMontageInstanceID)
+		, bReachedEnd(bInReachedEnd)
 	{}
 };
 
diff --git a/Engine/Source/Runtime/Engine/Private/Animation/AnimMontage.cpp b/Engine/Source/Runtime/Engine/Private/Animation/AnimMontage.cpp
--- a/Engine/Source/Runtime/Engine/Private/Animation/AnimMontage.cpp
+++ b/Engine/Source/Runtime/Engine/Private/Animation/AnimMontage.cpp
@@ -1609,7 +1609,7 @@ void FAnimMontageInstance::Terminate()
 
 			if (NotifyEvent.NotifyStateClass)
 			{
-				FBranchingPointNotifyPayload BranchingPointNotifyPayload(AnimInstance->GetSkelMeshComponent(), Montage, &NotifyEvent, InstanceID);
+				FBranchingPointNotifyPayload BranchingPointNotifyPayload(AnimInstance->GetSkelMeshComponent(), Montage, &NotifyEvent, InstanceID, false);
 				TRACE_ANIM_NOTIFY(AnimInstance.Get(), NotifyEvent, End);
 				NotifyEvent.NotifyStateClass->BranchingPointNotifyEnd(BranchingPointNotifyPayload);
 
@@ -2604,7 +2604,7 @@ bool FAnimMontageInstance::UpdateActiveStateBranchingPoints(float CurrentTrackPo
 
 				if (!bNotifyIsActive)
 				{
-					FBranchingPointNotifyPayload BranchingPointNotifyPayload(AnimInstance->GetSkelMeshComponent(), Montage, &NotifyEvent, InstanceID);
+					FBranchingPointNotifyPayload BranchingPointNotifyPayload(AnimInstance->GetSkelMeshComponent(), Montage, &NotifyEvent, InstanceID, true);
 					TRACE_ANIM_NOTIFY(AnimInstance.Get(), NotifyEvent, End);
 					NotifyEvent.NotifyStateClass->BranchingPointNotifyEnd(BranchingPointNotifyPayload);
 
@@ -2696,7 +2696,7 @@ void FAnimMontageInstance::BranchingPointEventHandler(const FBranchingPointMarke
 				}
 				else
 				{
-					FBranchingPointNotifyPayload BranchingPointNotifyPayload(AnimInstance->GetSkelMeshComponent(), Montage, NotifyEvent, InstanceID);
+					FBranchingPointNotifyPayload BranchingPointNotifyPayload(AnimInstance->GetSkelMeshComponent(), Montage, NotifyEvent, InstanceID, true);
 					TRACE_ANIM_NOTIFY(AnimInstance.Get(), *NotifyEvent, End);
 					NotifyEvent->NotifyStateClass->BranchingPointNotifyEnd(BranchingPointNotifyPayload);
 
diff --git a/Engine/Source/Runtime/Engine/Private/Animation/AnimNotifyEndDataContext.cpp b/Engine/Source/Runtime/Engine/Private/Animation/AnimNotifyEndDataContext.cpp
new file mode 100644
--- /dev/null
+++ b/Engine/Source/Runtime/Engine/Private/Animation/AnimNotifyEndDataContext.cpp
@@ -0,0 +1,14 @@
+// Copyright Epic Games, Inc. All Rights Reserved.
+
+#include "Animation/AnimNotifyEndDataContext.h"
+
+namespace UE::Anim {
+
+IMPLEMENT_NOTIFY_CONTEXT_INTERFACE(FAnimNotifyEndDataContext)
+
+FAnimNotifyEndDataContext::FAnimNotifyEndDataContext(bool bInReachedEnd)
+	: bReachedEnd(bInReachedEnd)
+{
+}
+	
+}	// namespace UE::Anim
diff --git a/Engine/Source/Runtime/Engine/Private/Animation/AnimNotifyLibrary.cpp b/Engine/Source/Runtime/Engine/Private/Animation/AnimNotifyLibrary.cpp
new file mode 100644
--- /dev/null
+++ b/Engine/Source/Runtime/Engine/Private/Animation/AnimNotifyLibrary.cpp
@@ -0,0 +1,17 @@
+// Copyright Epic Games, Inc. All Rights Reserved.
+
+#include "Animation/AnimNotifyLibrary.h"
+#include "Animation/AnimNotifyQueue.h"
+#include "Animation/AnimNotifyEndDataContext.h"
+
+bool UAnimNotifyLibrary::NotifyStateReachedEnd(const FAnimNotifyEventReference& EventReference)
+{
+	const UE::Anim::FAnimNotifyEndDataContext* EndData = EventReference.GetContextData<UE::Anim::FAnimNotifyEndDataContext>();
+
+	if (EndData != nullptr)
+	{
+		return EndData->bReachedEnd;
+	}
+
+	return false;
+}
diff --git a/Engine/Source/Runtime/Engine/Private/Animation/AnimNotifyState.cpp b/Engine/Source/Runtime/Engine/Private/Animation/AnimNotifyState.cpp
--- a/Engine/Source/Runtime/Engine/Private/Animation/AnimNotifyState.cpp
+++ b/Engine/Source/Runtime/Engine/Private/Animation/AnimNotifyState.cpp
@@ -3,6 +3,7 @@
 #include "Animation/AnimNotifies/AnimNotifyState.h"
 #include "Animation/AnimTypes.h"
 #include "Animation/AnimNotifies/AnimNotify.h"
+#include "Animation/AnimNotifyEndDataContext.h"
 #include "UObject/ObjectSaveContext.h"
 #include "Animation/AnimSequenceBase.h"
 
@@ -71,7 +72,12 @@ void UAnimNotifyState::BranchingPointNotifyTick(FBranchingPointNotifyPayload& Br
 
 void UAnimNotifyState::BranchingPointNotifyEnd(FBranchingPointNotifyPayload& BranchingPointPayload)
 {
-	const FAnimNotifyEventReference EventReference;
+	FAnimNotifyEventReference EventReference;
+	if (BranchingPointPayload.bReachedEnd)
+	{
+		EventReference.AddContextData<UE::Anim::FAnimNotifyEndDataContext>(true);
+	}
+
 	NotifyEnd(BranchingPointPayload.SkelMeshComponent, BranchingPointPayload.SequenceAsset, EventReference);
 }
 
diff --git a/Engine/Source/Runtime/Engine/Private/Animation/AnimSequenceBase.cpp b/Engine/Source/Runtime/Engine/Private/Animation/AnimSequenceBase.cpp
--- a/Engine/Source/Runtime/Engine/Private/Animation/AnimSequenceBase.cpp
+++ b/Engine/Source/Runtime/Engine/Private/Animation/AnimSequenceBase.cpp
@@ -5,6 +5,7 @@
 #include "AnimationRuntime.h"
 #include "Animation/AnimNotifies/AnimNotify.h"
 #include "Animation/AnimNotifies/AnimNotifyState.h"
+#include "Animation/AnimNotifyEndDataContext.h"
 #include "Animation/AnimInstance.h"
 #include "Logging/TokenizedMessage.h"
 #include "Logging/MessageLog.h"
@@ -15,6 +16,7 @@
 #include "Animation/MirrorDataTable.h"
 #include "Animation/AnimData/AnimDataModel.h"
 #include "Modules/ModuleManager.h"
+#include "MathUtil.h"
 
 #if WITH_EDITOR
 #include "IAnimationDataControllerModule.h"
@@ -415,6 +417,12 @@ void UAnimSequenceBase::GetAnimNotifiesFromDeltaPositions(const float& PreviousP
 				{
 					NotifyContext.ActiveNotifies.Emplace(&AnimNotifyEvent, this, nullptr);
 				}
+
+				const bool bHasFinished = CurrentPosition <= FMathf::Max(NotifyStartTime, 0.f);
+				if (bHasFinished)
+				{
+					NotifyContext.ActiveNotifies.Top().AddContextData<UE::Anim::FAnimNotifyEndDataContext>(true);
+				}
 			}
 		}
 	}
@@ -438,6 +446,11 @@ void UAnimSequenceBase::GetAnimNotifiesFromDeltaPositions(const float& PreviousP
 					NotifyContext.ActiveNotifies.Emplace(&AnimNotifyEvent, this, nullptr);
 				}
 
+				const bool bHasFinished = CurrentPosition >= NotifyEndTime;
+				if (bHasFinished)
+				{
+					NotifyContext.ActiveNotifies.Top().AddContextData<UE::Anim::FAnimNotifyEndDataContext>(true);
+				}
 			}
 		}
 	}
diff --git a/Engine/Source/Runtime/Engine/Public/Animation/AnimNotifyEndDataContext.h b/Engine/Source/Runtime/Engine/Public/Animation/AnimNotifyEndDataContext.h
new file mode 100644
--- /dev/null
+++ b/Engine/Source/Runtime/Engine/Public/Animation/AnimNotifyEndDataContext.h
@@ -0,0 +1,21 @@
+// Copyright Epic Games, Inc. All Rights Reserved.
+
+#pragma once
+
+#include "AnimNodeMessages.h"
+
+namespace UE::Anim {
+	
+// Notify Context Data containing if the notify state reached the end or was cancelled
+// This is useful i.e When there's simultaneous instances of the same montage, to trigger start/end events of its notify states.
+class ENGINE_API FAnimNotifyEndDataContext : public UE::Anim::IAnimNotifyEventContextDataInterface
+{
+	DECLARE_NOTIFY_CONTEXT_INTERFACE(FAnimNotifyEndDataContext)
+
+public:
+	explicit FAnimNotifyEndDataContext(bool bInReachedEnd);
+
+	const bool bReachedEnd = false;
+};
+
+}	// namespace UE::Anim
diff --git a/Engine/Source/Runtime/Engine/Public/Animation/AnimNotifyLibrary.h b/Engine/Source/Runtime/Engine/Public/Animation/AnimNotifyLibrary.h
new file mode 100644
--- /dev/null
+++ b/Engine/Source/Runtime/Engine/Public/Animation/AnimNotifyLibrary.h
@@ -0,0 +1,28 @@
+// Copyright Epic Games, Inc. All Rights Reserved.
+
+#pragma once
+
+#include "CoreMinimal.h"
+#include "Kismet/BlueprintFunctionLibrary.h"
+#include "Animation/AnimNotifyQueue.h"
+#include "AnimNotifyLibrary.generated.h"
+
+
+struct FAnimNotifyEventReference;
+
+/**
+*	A library of commonly used functionality for Notifies, exposed to blueprint.
+*/
+UCLASS(meta = (ScriptName = "UAnimNotifyLibrary"))
+class ENGINE_API UAnimNotifyLibrary : public UBlueprintFunctionLibrary
+{
+	GENERATED_BODY()
+public:
+
+	/** Get whether the notify state reached the end (was not cancelled)
+	*
+	* @param EventReference		The event to inspect
+	*/
+	UFUNCTION(BlueprintPure, Category = "Utilities|Animation|Notifies" , meta = (ScriptMethod))
+	static bool NotifyStateReachedEnd(const FAnimNotifyEventReference& EventReference);
+};
diff --git a/Engine/Source/Runtime/Engine/Public/Animation/AnimNotifyQueue.h b/Engine/Source/Runtime/Engine/Public/Animation/AnimNotifyQueue.h
--- a/Engine/Source/Runtime/Engine/Public/Animation/AnimNotifyQueue.h
+++ b/Engine/Source/Runtime/Engine/Public/Animation/AnimNotifyQueue.h
@@ -29,14 +29,11 @@ struct FAnimNotifyEventReference
 	{}
 
 	FAnimNotifyEventReference(const FAnimNotifyEventReference& rhs)
-		:
-		 ContextData(rhs.ContextData)
+		: ContextData(rhs.ContextData)
 		, Notify(rhs.Notify)
 		, MirrorTable(rhs.MirrorTable)
 		, NotifySource(rhs.NotifySource)
-
-	{
-	}
+	{}
 
 	FAnimNotifyEventReference(const FAnimNotifyEvent* InNotify, const UObject* InNotifySource)
 		: Notify(InNotify)
@@ -49,7 +46,8 @@ struct FAnimNotifyEventReference
 	, MirrorTable(MirrorDataTable)
     , NotifySource(InNotifySource)
 	{}
-	
+
+
 	const FAnimNotifyEvent* GetNotify() const
 	{
 		return NotifySource ? Notify : nullptr;
@@ -97,6 +95,19 @@ struct FAnimNotifyEventReference
 	// Pulls relevant data from the tick record
 	void GatherTickRecordData(const FAnimTickRecord& InTickRecord);
 
+	// Allows adding extra context data after GatherTickRecordData has been exectuted
+	template<typename Type, typename... TArgs>
+	void AddContextData(TArgs&&... Args)
+	{
+		static_assert(TPointerIsConvertibleFromTo<Type, const UE::Anim::IAnimNotifyEventContextDataInterface>::Value, "'Type' template parameter to MakeContextData must be derived from IAnimNotifyEventContextDataInterface");
+		if (!ContextData.IsValid())
+		{
+			ContextData = MakeShared<TArray<TUniquePtr<const UE::Anim::IAnimNotifyEventContextDataInterface>>>();
+		}
+
+		ContextData->Add(MakeUnique<Type>(Forward<TArgs>(Args)...));
+	}
+
 private:
 	// Context data gleaned from the tick record
 	TSharedPtr<TArray<TUniquePtr<const UE::Anim::IAnimNotifyEventContextDataInterface>>> ContextData;
