Fixed some material failing compilation after conversion to strata.

diff --git a/Engine/Source/Runtime/Engine/Private/Materials/HLSLMaterialTranslator.cpp b/Engine/Source/Runtime/Engine/Private/Materials/HLSLMaterialTranslator.cpp
--- a/Engine/Source/Runtime/Engine/Private/Materials/HLSLMaterialTranslator.cpp
+++ b/Engine/Source/Runtime/Engine/Private/Materials/HLSLMaterialTranslator.cpp
@@ -11070,6 +11070,18 @@ int32 FHLSLMaterialTranslator::StrataCompilePreview(int32 StrataDataCodeChunk)
 	return PreviewCodeChunk;
 }
 
+bool FHLSLMaterialTranslator::StrataSkipsOpacityEvaluation()
+{
+	return !IsTranslucentBlendMode(Material->GetStrataBlendMode())
+		&& Material->GetShadingModels().CountShadingModels() == 1
+		&& !Material->GetShadingModels().HasShadingModel(MSM_SingleLayerWater)
+		&& !Material->GetShadingModels().HasShadingModel(MSM_Subsurface)
+		&& !Material->GetShadingModels().HasShadingModel(MSM_SubsurfaceProfile)
+		&& !Material->GetShadingModels().HasShadingModel(MSM_PreintegratedSkin)
+		&& !Material->GetShadingModels().HasShadingModel(MSM_TwoSidedFoliage)
+		&& !Material->GetShadingModels().HasShadingModel(MSM_PreintegratedSkin);
+}
+
 int32 FHLSLMaterialTranslator::MapARPassthroughCameraUV(int32 UV)
 {
 	if (UV == INDEX_NONE)
diff --git a/Engine/Source/Runtime/Engine/Private/Materials/HLSLMaterialTranslator.h b/Engine/Source/Runtime/Engine/Private/Materials/HLSLMaterialTranslator.h
--- a/Engine/Source/Runtime/Engine/Private/Materials/HLSLMaterialTranslator.h
+++ b/Engine/Source/Runtime/Engine/Private/Materials/HLSLMaterialTranslator.h
@@ -1089,6 +1089,7 @@ protected:
 	virtual int32 StrataMetalnessToDiffuseAlbedoF0(int32 BaseColor, int32 Specular, int32 Metallic, int32 OutputIndex) override;
 	virtual int32 StrataHazinessToSecondaryRoughness(int32 BaseRoughness, int32 Haziness, int32 OutputIndex) override;
 	virtual int32 StrataCompilePreview(int32 StrataDataCodeChunk) override;
+	virtual bool StrataSkipsOpacityEvaluation() override;
 	
 	virtual FStrataOperator& StrataCompilationRegisterOperator(int32 OperatorType, UMaterialExpression* Expression, UMaterialExpression* Parent, bool bUseParameterBlending = false) override;
 	virtual FStrataOperator& StrataCompilationGetOperator(UMaterialExpression* Expression) override;
diff --git a/Engine/Source/Runtime/Engine/Private/Materials/MaterialExpressions.cpp b/Engine/Source/Runtime/Engine/Private/Materials/MaterialExpressions.cpp
--- a/Engine/Source/Runtime/Engine/Private/Materials/MaterialExpressions.cpp
+++ b/Engine/Source/Runtime/Engine/Private/Materials/MaterialExpressions.cpp
@@ -21796,7 +21796,20 @@ int32 UMaterialExpressionStrataLegacyConversion::Compile(class FMaterialCompiler
 	FStrataOperator& StrataOperator = Compiler->StrataCompilationGetOperator(this);
 	StrataOperator.BSDFRegisteredSharedLocalBasis = NewRegisteredSharedLocalBasis;
 
-	const bool bHasDynamicShadingModels = ConvertedStrataMaterialInfo.CountShadingModels() > 1;
+	int32 ShadingModelCount = ConvertedStrataMaterialInfo.CountShadingModels();
+	int32 OpacityCodeChunk = INDEX_NONE;
+	if (!Compiler->StrataSkipsOpacityEvaluation())
+	{
+		// We evaluate opacity only for shading models and blending mode requiring it.
+		// For instance, a translucent shader reading depth for soft fading should no evaluate opacity when an instance forces an opaque mode.
+		OpacityCodeChunk = CompileWithDefaultFloat1(Compiler, Opacity, 1.0f);
+	}
+	else
+	{
+		OpacityCodeChunk = Compiler->Constant(1.0f);
+	}
+
+	const bool bHasDynamicShadingModels = ShadingModelCount > 1;
 	// We probably need to do something along these line as well :::
 	int32 OutputCodeChunk = Compiler->StrataConversionFromLegacy(
 		bHasDynamicShadingModels,
@@ -21815,7 +21828,7 @@ int32 UMaterialExpressionStrataLegacyConversion::Compile(class FMaterialCompiler
 		CompileWithDefaultFloat1(Compiler, ClearCoatRoughness, 0.1f),
 		// Misc
 		CompileWithDefaultFloat3(Compiler, EmissiveColor, 0.0f, 0.0f, 0.0f),
-		CompileWithDefaultFloat1(Compiler, Opacity, 1.0f),
+		OpacityCodeChunk,
 		CompileWithDefaultFloat3(Compiler, TransmittanceColor, 0.5f, 0.5f, 0.5f),
 		// Water
 		CompileWithDefaultFloat3(Compiler, WaterScatteringCoefficients, 0.0f, 0.0f, 0.0f),
diff --git a/Engine/Source/Runtime/Engine/Public/MaterialCompiler.h b/Engine/Source/Runtime/Engine/Public/MaterialCompiler.h
--- a/Engine/Source/Runtime/Engine/Public/MaterialCompiler.h
+++ b/Engine/Source/Runtime/Engine/Public/MaterialCompiler.h
@@ -610,6 +610,11 @@ public:
 	virtual int32 StrataMetalnessToDiffuseAlbedoF0(int32 BaseColor, int32 Specular, int32 Metallic, int32 OutputIndex) = 0;
 	virtual int32 StrataHazinessToSecondaryRoughness(int32 BaseRoughness, int32 Haziness, int32 OutputIndex) = 0;
 	virtual int32 StrataCompilePreview(int32 StrataDataCodeChunk) = 0;
+
+	/**
+	 * This is dedicated to skip evaluating any opacity input when a material instance toggles the translucent blend mode to opaque.
+	 */
+	virtual bool StrataSkipsOpacityEvaluation() = 0;
 	
 	/**
 	 * Register an operator of the tree representation the Strata material and its topology.
@@ -1270,6 +1275,11 @@ public:
 		return Compiler->StrataCompilePreview(StrataDataCodeChunk);
 	}
 
+	virtual bool StrataSkipsOpacityEvaluation() override
+	{
+		return Compiler->StrataSkipsOpacityEvaluation();
+	}
+
 	virtual FStrataOperator& StrataCompilationRegisterOperator(int32 OperatorType, UMaterialExpression* Expression, UMaterialExpression* Parent, bool bUseParameterBlending = false) override
 	{
 		return Compiler->StrataCompilationRegisterOperator(OperatorType, Expression, Parent, bUseParameterBlending);
