Enabled timermanager uobject stats in test builds.

diff --git a/Engine/Source/Runtime/Engine/Public/TimerManager.h b/Engine/Source/Runtime/Engine/Public/TimerManager.h
--- a/Engine/Source/Runtime/Engine/Public/TimerManager.h
+++ b/Engine/Source/Runtime/Engine/Public/TimerManager.h
@@ -39,23 +39,17 @@ struct FTimerUnifiedDelegate
 	{
 		if (FuncDelegate.IsBound())
 		{
-#if STATS
-			TStatId StatId = TStatId();
-			UObject* Object = FuncDelegate.GetUObject();
-			if (Object)
-			{
-				StatId = Object->GetStatID();
-			}
-			FScopeCycleCounter Context(StatId);
-#endif
+			FScopeCycleCounterUObject Context(FuncDelegate.GetUObject());
 			FuncDelegate.Execute();
 		}
 		else if (FuncDynDelegate.IsBound())
 		{
+			FScopeCycleCounterUObject Context(FuncDynDelegate.GetUObject());
 			FuncDynDelegate.ProcessDelegate<UObject>(nullptr);
 		}
 		else if ( FuncCallback )
 		{
+			QUICK_SCOPE_CYCLE_COUNTER(STAT_FTimerUnifiedDelegate_Execute);
 			FuncCallback();
 		}
 	}
