Fix for Async Queries Thread safety issue: Make sure that results cannot be read before all async tasks are finished

#tests preflight

#localization none

diff --git a/Engine/Source/Runtime/Engine/Private/Collision/WorldCollisionAsync.cpp b/Engine/Source/Runtime/Engine/Private/Collision/WorldCollisionAsync.cpp
--- a/Engine/Source/Runtime/Engine/Private/Collision/WorldCollisionAsync.cpp
+++ b/Engine/Source/Runtime/Engine/Private/Collision/WorldCollisionAsync.cpp
@@ -457,6 +457,10 @@ bool UWorld::QueryTraceData(const FTraceHandle& Handle, FTraceDatum& OutData)
 	}
 
 	AsyncTraceData& DataBuffer = AsyncTraceState.GetBufferForPreviousFrame();
+	if (!DataBuffer.bAsyncTasksCompleted)
+	{
+		return false;
+	}
 	if (auto* Data = FBufferIndexPair(Handle._Data.Index).DatumLookup(DataBuffer.TraceData))
 	{
 		OutData = *Data;
@@ -474,6 +478,10 @@ bool UWorld::QueryOverlapData(const FTraceHandle& Handle, FOverlapDatum& OutData
 	}
 
 	AsyncTraceData& DataBuffer = AsyncTraceState.GetBufferForPreviousFrame();
+	if (!DataBuffer.bAsyncTasksCompleted)
+	{
+		return false;
+	}
 	if (auto* Data = FBufferIndexPair(Handle._Data.Index).DatumLookup(DataBuffer.OverlapData))
 	{
 		OutData = *Data;
@@ -506,6 +514,7 @@ void UWorld::ResetAsyncTrace()
 
 	// Wait for thread
 	WaitForAllAsyncTraceTasks();
+	DataBufferExecuted.bAsyncTasksCompleted = true;
 
 	// do run delegates before starting next round
 	for (int32 Idx = 0; Idx != DataBufferExecuted.NumQueuedTraceData; ++Idx)
@@ -538,6 +547,7 @@ void UWorld::FinishAsyncTrace()
 	NewAsyncBuffer.bAsyncAllowed = true;
 	NewAsyncBuffer.NumQueuedTraceData = 0;
 	NewAsyncBuffer.NumQueuedOverlapData = 0;
+	NewAsyncBuffer.bAsyncTasksCompleted = false;
 
 }
 
diff --git a/Engine/Source/Runtime/Engine/Public/WorldCollision.h b/Engine/Source/Runtime/Engine/Public/WorldCollision.h
--- a/Engine/Source/Runtime/Engine/Public/WorldCollision.h
+++ b/Engine/Source/Runtime/Engine/Public/WorldCollision.h
@@ -284,6 +284,8 @@ struct AsyncTraceData : FNoncopyable
 	 */
 	bool					bAsyncAllowed; 
 
+	bool					bAsyncTasksCompleted;
+
 	/**  Thread completion event for batch **/
 	FGraphEventArray		AsyncTraceCompletionEvent;
 
@@ -291,6 +293,7 @@ struct AsyncTraceData : FNoncopyable
 		: NumQueuedTraceData(0)
 		, NumQueuedOverlapData(0)
 		, bAsyncAllowed(false)
+		, bAsyncTasksCompleted(false)
 	{}
 };
 
