Update tooltip on DefaultAmbientZoneSettings

diff --git a/Engine/Source/Runtime/Engine/Classes/GameFramework/WorldSettings.h b/Engine/Source/Runtime/Engine/Classes/GameFramework/WorldSettings.h
--- a/Engine/Source/Runtime/Engine/Classes/GameFramework/WorldSettings.h
+++ b/Engine/Source/Runtime/Engine/Classes/GameFramework/WorldSettings.h
@@ -709,7 +709,7 @@ public:
 	UPROPERTY(EditAnywhere, config, Category=Audio)
 	FReverbSettings DefaultReverbSettings;
 
-	/** Default interior settings used by audio volumes.												*/
+	/** Default interior settings applied to sounds that have "apply ambient volumes" set to true on their SoundClass. */
 	UPROPERTY(EditAnywhere, config, Category=Audio)
 	FInteriorSettings DefaultAmbientZoneSettings;
 
