Fixed: Missing uniform buffer for a custom instance data in Editor mobile preview

diff --git a/Engine/Source/Runtime/Engine/Private/InstancedStaticMesh.cpp b/Engine/Source/Runtime/Engine/Private/InstancedStaticMesh.cpp
--- a/Engine/Source/Runtime/Engine/Private/InstancedStaticMesh.cpp
+++ b/Engine/Source/Runtime/Engine/Private/InstancedStaticMesh.cpp
@@ -867,7 +867,7 @@ void FInstancedStaticMeshVertexFactory::InitRHI()
 		Elements.Add(AccessStreamComponent(Data.PositionComponent,0));
 	}
 
-	const bool bCanUseGPUScene = UseGPUScene(GMaxRHIShaderPlatform, GMaxRHIFeatureLevel);
+	const bool bCanUseGPUScene = UseGPUScene(GMaxRHIShaderPlatform, GetFeatureLevel());
 
 	// only tangent,normal are used by the stream. the binormal is derived in the shader
 	uint8 TangentBasisAttributes[2] = { 1, 2 };
