Fix obj list to report max (reserved) size properly.

diff --git a/Engine/Source/Runtime/Engine/Private/UnrealEngine.cpp b/Engine/Source/Runtime/Engine/Private/UnrealEngine.cpp
--- a/Engine/Source/Runtime/Engine/Private/UnrealEngine.cpp
+++ b/Engine/Source/Runtime/Engine/Private/UnrealEngine.cpp
@@ -8813,7 +8813,7 @@ bool UEngine::HandleObjCommand( const TCHAR* Cmd, FOutputDevice& Ar )
 				{
 					FArchiveCountMem Count(*It);
 					NumBytes = Count.GetNum();
-					MaxBytes = Count.GetNum();
+					MaxBytes = Count.GetMax();
 					It->GetResourceSizeEx(TrueResourceSize);
 				}
 
