Allow for a single WorldPartition to have multiple ActorDescContainer
Added FActorDescContainerCollection which can contains multiple UActorDescContainer
ActorDescContainer can be created to load ExternalFolder in various places for a world partition
WorldPartition no longer inherits UActorDescContainer and now inherits FActorDescContainerCollection
ActorsDesc are added to the right container on save based on their packagepath
HLODs generation only occurs on the Main WorldPartition Actors for the moment
Some WorldPartitionBuilders commandlets do not support having WorldPartitions with multiple ACtorDescContainer. In that case, the commandlet fails and outputs an error message.

#test cook AncientGame, cook city_sample, cook asteria_terrain, HLods generation on AncientGame, Level Instances workflow

diff --git a/Engine/Plugins/Editor/DataValidation/Source/DataValidation/Private/WorldPartitionChangelistValidator.cpp b/Engine/Plugins/Editor/DataValidation/Source/DataValidation/Private/WorldPartitionChangelistValidator.cpp
--- a/Engine/Plugins/Editor/DataValidation/Source/DataValidation/Private/WorldPartitionChangelistValidator.cpp
+++ b/Engine/Plugins/Editor/DataValidation/Source/DataValidation/Private/WorldPartitionChangelistValidator.cpp
@@ -176,7 +176,7 @@ EDataValidationResult UWorldPartitionChangelistValidator::ValidateActorsAndDataL
 		if (World != nullptr)
 		{
 			// World is Loaded reuse the ActorDescContainer of the World
-			ActorDescContainer = World->GetWorldPartition();
+			ActorDescContainer = World->GetWorldPartition()->GetActorDescContainer();
 		}
 		else
 		{
diff --git a/Engine/Plugins/Editor/WorldPartitionHLODUtilities/Source/Private/WorldPartitionHLODUtilities.cpp b/Engine/Plugins/Editor/WorldPartitionHLODUtilities/Source/Private/WorldPartitionHLODUtilities.cpp
--- a/Engine/Plugins/Editor/WorldPartitionHLODUtilities/Source/Private/WorldPartitionHLODUtilities.cpp
+++ b/Engine/Plugins/Editor/WorldPartitionHLODUtilities/Source/Private/WorldPartitionHLODUtilities.cpp
@@ -4,6 +4,7 @@
 
 #if WITH_EDITOR
 
+#include "WorldPartition/ActorDescContainer.h"
 #include "WorldPartition/WorldPartition.h"
 #include "WorldPartition/WorldPartitionHandle.h"
 #include "WorldPartition/DataLayer/DataLayerInstance.h"
@@ -41,7 +42,7 @@ static UWorldPartitionLevelStreamingDynamic* CreateLevelStreamingFromHLODActor(A
 	const FName LevelStreamingName = FName(*FString::Printf(TEXT("HLODLevelStreaming_%s"), *InHLODActor->GetName()));
 	TArray<FWorldPartitionRuntimeCellObjectMapping> Mappings;
 	Mappings.Reserve(InHLODActor->GetSubActors().Num());
-	Algo::Transform(InHLODActor->GetSubActors(), Mappings, [](const FHLODSubActor& SubActor) { return FWorldPartitionRuntimeCellObjectMapping(SubActor.ActorPackage, SubActor.ActorPath, SubActor.ContainerID, SubActor.ContainerTransform, SubActor.ContainerPackage); });
+	Algo::Transform(InHLODActor->GetSubActors(), Mappings, [World](const FHLODSubActor& SubActor) { return FWorldPartitionRuntimeCellObjectMapping(SubActor.ActorPackage, SubActor.ActorPath, SubActor.ContainerID, SubActor.ContainerTransform, SubActor.ContainerPackage, World->GetPackage()->GetFName()); });
 
 	UWorldPartitionLevelStreamingDynamic* LevelStreaming = UWorldPartitionLevelStreamingDynamic::LoadInEditor(World, LevelStreamingName, Mappings);
 	check(LevelStreaming);
diff --git a/Engine/Source/Editor/DataLayerEditor/Private/DataLayer/DataLayerHierarchy.cpp b/Engine/Source/Editor/DataLayerEditor/Private/DataLayer/DataLayerHierarchy.cpp
--- a/Engine/Source/Editor/DataLayerEditor/Private/DataLayer/DataLayerHierarchy.cpp
+++ b/Engine/Source/Editor/DataLayerEditor/Private/DataLayer/DataLayerHierarchy.cpp
@@ -8,6 +8,7 @@
 #include "WorldDataLayersTreeItem.h"
 #include "DataLayerMode.h"
 #include "DataLayer/DataLayerEditorSubsystem.h"
+#include "WorldPartition/WorldPartition.h"
 #include "WorldPartition/DataLayer/DataLayerAsset.h"
 #include "WorldPartition/DataLayer/DataLayerSubsystem.h"
 #include "WorldPartition/DataLayer/DataLayerInstanceWithAsset.h"
diff --git a/Engine/Source/Editor/DataLayerEditor/Private/DataLayer/DataLayersActorDescTreeItem.h b/Engine/Source/Editor/DataLayerEditor/Private/DataLayer/DataLayersActorDescTreeItem.h
--- a/Engine/Source/Editor/DataLayerEditor/Private/DataLayer/DataLayersActorDescTreeItem.h
+++ b/Engine/Source/Editor/DataLayerEditor/Private/DataLayer/DataLayersActorDescTreeItem.h
@@ -6,6 +6,7 @@
 #include "SceneOutlinerFwd.h"
 #include "ActorDescTreeItem.h"
 #include "WorldPartition/DataLayer/DataLayerInstance.h"
+#include "WorldPartition/ActorDescContainerCollection.h"
 #include "LevelInstance/LevelInstanceSubsystem.h"
 
 //////////////////////////////////////////////////////////////////////////
@@ -20,6 +21,12 @@ struct FDataLayerActorDescTreeItemData
 		, Container(InContainer)
 		, DataLayer(InDataLayer)
 	{}
+
+	FDataLayerActorDescTreeItemData(const FGuid& InActorGuid, FActorDescContainerCollection* InContainerCollection, UDataLayerInstance* InDataLayer)
+		: ActorGuid(InActorGuid)
+		, Container(InContainerCollection->GetActorDescContainer(InActorGuid))
+		, DataLayer(InDataLayer)
+	{}
 	
 	const FGuid& ActorGuid;
 	UActorDescContainer* const Container;
diff --git a/Engine/Source/Editor/SceneOutliner/Private/ActorBrowsingMode.cpp b/Engine/Source/Editor/SceneOutliner/Private/ActorBrowsingMode.cpp
--- a/Engine/Source/Editor/SceneOutliner/Private/ActorBrowsingMode.cpp
+++ b/Engine/Source/Editor/SceneOutliner/Private/ActorBrowsingMode.cpp
@@ -34,6 +34,7 @@
 #include "LevelInstance/LevelInstanceInterface.h"
 #include "LevelInstance/LevelInstanceSubsystem.h"
 #include "LevelInstance/LevelInstanceEditorInstanceActor.h"
+#include "WorldPartition/WorldPartition.h"
 #include "EditorLevelUtils.h"
 #include "EditorViewportCommands.h"
 #include "SceneOutlinerActorSCCColumn.h"
diff --git a/Engine/Source/Editor/SceneOutliner/Private/ActorDescTreeItem.cpp b/Engine/Source/Editor/SceneOutliner/Private/ActorDescTreeItem.cpp
--- a/Engine/Source/Editor/SceneOutliner/Private/ActorDescTreeItem.cpp
+++ b/Engine/Source/Editor/SceneOutliner/Private/ActorDescTreeItem.cpp
@@ -1,6 +1,7 @@
 // Copyright Epic Games, Inc. All Rights Reserved.
 
 #include "ActorDescTreeItem.h"
+#include "WorldPartition/ActorDescContainer.h"
 #include "WorldPartition/WorldPartition.h"
 #include "WorldPartition/WorldPartitionActorDesc.h"
 #include "Widgets/DeclarativeSyntaxSupport.h"
@@ -119,8 +120,11 @@ private:
 			{
 				FFormatNamedArguments Args;
 				Args.Add(TEXT("ActorLabel"), FText::FromString(TreeItem->GetDisplayString()));
-
-				if (UWorldPartition* WorldPartition = Cast<UWorldPartition>(ActorDesc->GetContainer()); WorldPartition->IsActorPinned(ActorDesc->GetGuid()))
+				
+				UActorDescContainer* ActorDescContainer = ActorDesc->GetContainer();
+				UWorld* World = ActorDescContainer != nullptr ? ActorDescContainer->GetWorld() : nullptr;
+				UWorldPartition* WorldPartition = World != nullptr ? World->GetWorldPartition() : nullptr;
+				if (WorldPartition && WorldPartition->IsActorPinned(ActorDesc->GetGuid()))
 				{
 					Args.Add(TEXT("UnloadState"), LOCTEXT("UnloadedDataLayer", "(Unloaded DataLayer)"));
 				}
@@ -234,7 +238,20 @@ FActorDescTreeItem::FActorDescTreeItem(const FGuid& InActorGuid, UActorDescConta
 	, ActorDescHandle(Container, InActorGuid)
 	, ActorGuid(InActorGuid)
 {
-	ID = FSceneOutlinerTreeItemID(InActorGuid);
+	Initialize();
+}
+
+FActorDescTreeItem::FActorDescTreeItem(const FGuid& InActorGuid, FActorDescContainerCollection* ContainerCollection)
+	: IActorBaseTreeItem(Type)
+	, ActorDescHandle(ContainerCollection, InActorGuid)
+	, ActorGuid(InActorGuid)
+{
+	Initialize();
+}
+
+void FActorDescTreeItem::Initialize()
+{
+	ID = FSceneOutlinerTreeItemID(ActorGuid);
 
 	if (const FWorldPartitionActorDesc* const ActorDesc = ActorDescHandle.Get())
 	{
@@ -304,7 +321,7 @@ bool FActorDescTreeItem::GetPinnedState() const
 {
 	if (ActorDescHandle.IsValid())
 	{
-		UWorldPartition* WorldPartition = Cast<UWorldPartition>(ActorDescHandle->GetContainer());
+		UWorldPartition* WorldPartition = ActorDescHandle->GetContainer()->GetWorld()->GetWorldPartition();
 		return WorldPartition ? WorldPartition->IsActorPinned(GetGuid()) : false;
 	}
 	return false;
diff --git a/Engine/Source/Editor/SceneOutliner/Private/ActorHierarchy.cpp b/Engine/Source/Editor/SceneOutliner/Private/ActorHierarchy.cpp
--- a/Engine/Source/Editor/SceneOutliner/Private/ActorHierarchy.cpp
+++ b/Engine/Source/Editor/SceneOutliner/Private/ActorHierarchy.cpp
@@ -20,6 +20,8 @@
 #include "LevelInstance/LevelInstanceInterface.h"
 #include "LevelInstance/LevelInstanceSubsystem.h"
 #include "Modules/ModuleManager.h"
+#include "WorldPartition/ActorDescContainer.h"
+#include "WorldPartition/WorldPartition.h"
 #include "WorldPartition/WorldPartitionHelpers.h"
 #include "WorldPartition/WorldPartitionActorDesc.h"
 #include "WorldPartition/IWorldPartitionEditorModule.h"
diff --git a/Engine/Source/Editor/SceneOutliner/Private/SceneOutlinerModule.cpp b/Engine/Source/Editor/SceneOutliner/Private/SceneOutlinerModule.cpp
--- a/Engine/Source/Editor/SceneOutliner/Private/SceneOutlinerModule.cpp
+++ b/Engine/Source/Editor/SceneOutliner/Private/SceneOutlinerModule.cpp
@@ -22,6 +22,7 @@
 #include "ActorDescTreeItem.h"
 #include "FolderTreeItem.h"
 #include "WorldTreeItem.h"
+#include "WorldPartition/ActorDescContainer.h"
 #include "WorldPartition/DataLayer/DataLayerInstance.h"
 #include "WorldPartition/DataLayer/DataLayerSubsystem.h"
 #include "LevelInstance/LevelInstanceSubsystem.h"
diff --git a/Engine/Source/Editor/SceneOutliner/Public/ActorDescTreeItem.h b/Engine/Source/Editor/SceneOutliner/Public/ActorDescTreeItem.h
--- a/Engine/Source/Editor/SceneOutliner/Public/ActorDescTreeItem.h
+++ b/Engine/Source/Editor/SceneOutliner/Public/ActorDescTreeItem.h
@@ -6,8 +6,9 @@
 #include "Widgets/SWidget.h"
 #include "ActorBaseTreeItem.h"
 #include "UObject/ObjectKey.h"
-#include "WorldPartition/WorldPartition.h"
+#include "WorldPartition/WorldPartitionHandle.h"
 
+class UActorDescContainer;
 class FWorldPartitionActorDesc;
 class UToolMenu;
 
@@ -40,6 +41,9 @@ public:
 	/** Construct this item from an actor desc */
 	FActorDescTreeItem(const FGuid& InActorGuid, UActorDescContainer* Container);
 
+	/** Construct this item from an actor desc */
+	FActorDescTreeItem(const FGuid& InActorGuid, FActorDescContainerCollection* ContainerCollection);
+
 	/* Begin ISceneOutlinerTreeItem Implementation */
 	virtual bool IsValid() const override { return ActorDescHandle.Get() != nullptr; }
 	virtual FSceneOutlinerTreeItemID GetID() const override;
@@ -66,5 +70,6 @@ protected:
 
 private:
 	void CopyActorFilePathtoClipboard() const;
+	void Initialize();
 	FGuid ActorGuid;
 };
diff --git a/Engine/Source/Editor/UnrealEd/Private/Commandlets/ContentCommandlets.cpp b/Engine/Source/Editor/UnrealEd/Private/Commandlets/ContentCommandlets.cpp
--- a/Engine/Source/Editor/UnrealEd/Private/Commandlets/ContentCommandlets.cpp
+++ b/Engine/Source/Editor/UnrealEd/Private/Commandlets/ContentCommandlets.cpp
@@ -48,6 +48,7 @@
 #include "CollectionManagerModule.h"
 #include "ICollectionManager.h"
 #include "CommandletSourceControlUtils.h"
+#include "WorldPartition/ActorDescContainer.h"
 #include "WorldPartition/WorldPartition.h"
 #include "WorldPartition/WorldPartitionHelpers.h"
 #include "LevelInstance/LevelInstanceInterface.h"
diff --git a/Engine/Source/Editor/UnrealEd/Private/WorldPartition/WorldPartitionFoliageBuilder.cpp b/Engine/Source/Editor/UnrealEd/Private/WorldPartition/WorldPartitionFoliageBuilder.cpp
--- a/Engine/Source/Editor/UnrealEd/Private/WorldPartition/WorldPartitionFoliageBuilder.cpp
+++ b/Engine/Source/Editor/UnrealEd/Private/WorldPartition/WorldPartitionFoliageBuilder.cpp
@@ -2,6 +2,7 @@
 
 #include "WorldPartition/WorldPartitionFoliageBuilder.h"
 
+#include "WorldPartition/ActorDescContainer.h"
 #include "WorldPartition/WorldPartition.h"
 #include "WorldPartition/WorldPartitionHelpers.h"
 #include "PackageSourceControlHelper.h"
@@ -44,6 +45,15 @@ bool UWorldPartitionFoliageBuilder::RunInternal(UWorld* World, const FCellInfo&
 		return false;
 	}
 
+	if (WorldPartition->GetActorDescContainerCount() > 1)
+	{ 
+		// @todo_ow: This commandlet is currently unsupported for worldpartitions having multiple containers.
+		// ForEachActorWithLoading below loops on all actors of all containers of a WorldPartition and saves new actors.
+		// The new actors are saved within the main container WorldPartition container. Newly saved actors coming from different actor desc container should be saved in the right container. 
+		UE_LOG(LogWorldPartitionFoliageBuilder, Error, TEXT("FoliageBuilder Commandlet is unsupported on WorldPartition using more than 1 ExternalActor Folder."));
+		return false;
+	}
+
 	UE_LOG(LogWorldPartitionFoliageBuilder, Display, TEXT("Changing Foliage Grid Size from %d to %d"), World->GetWorldSettings()->InstancedFoliageGridSize, NewGridSize);
 			
 	// Update Setting
diff --git a/Engine/Source/Editor/UnrealEd/Private/WorldPartition/WorldPartitionHLODsBuilder.cpp b/Engine/Source/Editor/UnrealEd/Private/WorldPartition/WorldPartitionHLODsBuilder.cpp
--- a/Engine/Source/Editor/UnrealEd/Private/WorldPartition/WorldPartitionHLODsBuilder.cpp
+++ b/Engine/Source/Editor/UnrealEd/Private/WorldPartition/WorldPartitionHLODsBuilder.cpp
@@ -22,6 +22,7 @@
 #include "DerivedDataCacheInterface.h"
 
 #include "WorldPartition/DataLayer/DataLayerInstance.h"
+#include "WorldPartition/ActorDescContainer.h"
 #include "WorldPartition/WorldPartition.h"
 #include "WorldPartition/WorldPartitionHelpers.h"
 #include "WorldPartition/HLOD/HLODActor.h"
@@ -353,7 +354,7 @@ bool UWorldPartitionHLODsBuilder::SetupHLODActors()
 	UE_LOG(LogWorldPartitionHLODsBuilder, Display, TEXT("#### World HLOD actors ####"));
 
 	int32 NumActors = 0;
-	for (FActorDescList::TIterator<AWorldPartitionHLOD> HLODIterator(WorldPartition); HLODIterator; ++HLODIterator)
+	for (FActorDescList::TIterator<AWorldPartitionHLOD> HLODIterator(WorldPartition->GetActorDescContainer()); HLODIterator; ++HLODIterator)
 	{
 		FWorldPartitionActorDesc* HLODActorDesc = *HLODIterator;
 		FString PackageName = HLODActorDesc->GetActorPackage().ToString();
@@ -536,7 +537,7 @@ bool UWorldPartitionHLODsBuilder::DeleteHLODActors()
 	UE_LOG(LogWorldPartitionHLODsBuilder, Display, TEXT("#### Deleting HLOD actors ####"));
 
 	TArray<FString> PackagesToDelete;
-	for (FActorDescList::TIterator<AWorldPartitionHLOD> HLODIterator(WorldPartition); HLODIterator; ++HLODIterator)
+	for (FActorDescList::TIterator<AWorldPartitionHLOD> HLODIterator(WorldPartition->GetActorDescContainer()); HLODIterator; ++HLODIterator)
 	{
 		FWorldPartitionActorDesc* HLODActorDesc = *HLODIterator;
 		FString PackageName = HLODActorDesc->GetActorPackage().ToString();
@@ -657,7 +658,7 @@ bool UWorldPartitionHLODsBuilder::DumpStats()
 
 	TArray<FHLODStat> HLODStats;
 
-	for (UActorDescContainer::TIterator<AWorldPartitionHLOD> HLODIterator(WorldPartition); HLODIterator; ++HLODIterator)
+	for (UActorDescContainer::TIterator<AWorldPartitionHLOD> HLODIterator(WorldPartition->GetActorDescContainer()); HLODIterator; ++HLODIterator)
 	{
 		FWorldPartitionReference HLODActorRef(WorldPartition, HLODIterator->GetGuid());
 		AWorldPartitionHLOD* HLODActor = CastChecked<AWorldPartitionHLOD>(HLODActorRef->GetActor());
@@ -802,7 +803,7 @@ TArray<TArray<FGuid>> UWorldPartitionHLODsBuilder::GetHLODWorkloads(int32 NumWor
 {
 	// Build a mapping of 1 HLOD[Level] -> N HLOD[Level - 1]
 	TMap<FGuid, TArray<FGuid>>	HLODParenting;
-	for (FActorDescList::TIterator<AWorldPartitionHLOD> HLODIterator(WorldPartition); HLODIterator; ++HLODIterator)
+	for (FActorDescList::TIterator<AWorldPartitionHLOD> HLODIterator(WorldPartition->GetActorDescContainer()); HLODIterator; ++HLODIterator)
 	{
 		TArray<FGuid>& ChildHLODs = HLODParenting.Add(HLODIterator->GetGuid());
 
diff --git a/Engine/Source/Editor/UnrealEd/Private/WorldPartition/WorldPartitionRenameDuplicateBuilder.cpp b/Engine/Source/Editor/UnrealEd/Private/WorldPartition/WorldPartitionRenameDuplicateBuilder.cpp
--- a/Engine/Source/Editor/UnrealEd/Private/WorldPartition/WorldPartitionRenameDuplicateBuilder.cpp
+++ b/Engine/Source/Editor/UnrealEd/Private/WorldPartition/WorldPartitionRenameDuplicateBuilder.cpp
@@ -17,6 +17,7 @@
 #include "UObject/ObjectRedirector.h"
 #include "ProfilingDebugging/ScopedTimers.h"
 #include "HAL/FileManager.h"
+#include "WorldPartition/ActorDescContainer.h"
 
 DEFINE_LOG_CATEGORY_STATIC(LogWorldPartitionRenameDuplicateBuilder, All, All);
 
@@ -141,9 +142,14 @@ bool UWorldPartitionRenameDuplicateBuilder::RunInternal(UWorld* World, const FCe
 		return false;
 	}
 
+	if (WorldPartition->GetActorDescContainerCount() > 1)
+	{
+		UE_LOG(LogWorldPartitionRenameDuplicateBuilder, Warning, TEXT("Actors coming from containers outside of %s will not be duplicated"), *WorldPartition->GetActorDescContainer()->GetExternalActorPath());
+	}
+
 	// Build actor clusters
 	TArray<TPair<FGuid, TArray<FGuid>>> ActorsWithRefs;
-	for (FActorDescList::TIterator<> ActorDescIterator(WorldPartition); ActorDescIterator; ++ActorDescIterator)
+	for (FActorDescContainerCollection::TIterator<> ActorDescIterator(WorldPartition); ActorDescIterator; ++ActorDescIterator)
 	{
 		ActorsWithRefs.Emplace(ActorDescIterator->GetGuid(), ActorDescIterator->GetReferences());
 	}
@@ -385,11 +391,11 @@ bool UWorldPartitionRenameDuplicateBuilder::RunInternal(UWorld* World, const FCe
 		{
 			// Validate results
 			UE_SCOPED_TIMER(TEXT("Validating actors"), LogWorldPartitionRenameDuplicateBuilder, Display);
-			for (UActorDescContainer::TConstIterator<> ActorDescIterator(WorldPartition); ActorDescIterator; ++ActorDescIterator)
+			for (FActorDescContainerCollection::TConstIterator<> ActorDescIterator(WorldPartition); ActorDescIterator; ++ActorDescIterator)
 			{
 				const FWorldPartitionActorDesc* SourceActorDesc = *ActorDescIterator;
 				FGuid* DuplicatedGuid = DuplicatedActorGuids.Find(SourceActorDesc->GetGuid());
-				const FWorldPartitionActorDesc* NewActorDesc = NewWorld->GetWorldPartition()->GetActorDesc(DuplicatedGuid ? *DuplicatedGuid : SourceActorDesc->GetGuid());
+				const FWorldPartitionActorDesc* NewActorDesc = NewWorld->GetWorldPartition()->GetActorDescContainer()->GetActorDesc(DuplicatedGuid ? *DuplicatedGuid : SourceActorDesc->GetGuid());
 				if (!NewActorDesc)
 				{
 					UE_LOG(LogWorldPartitionRenameDuplicateBuilder, Warning, TEXT("Failed to find source actor for Actor: %s"), *SourceActorDesc->GetActorPath().ToString());
diff --git a/Engine/Source/Editor/UnrealEd/Private/WorldPartition/WorldPartitionResaveActorsBuilder.cpp b/Engine/Source/Editor/UnrealEd/Private/WorldPartition/WorldPartitionResaveActorsBuilder.cpp
--- a/Engine/Source/Editor/UnrealEd/Private/WorldPartition/WorldPartitionResaveActorsBuilder.cpp
+++ b/Engine/Source/Editor/UnrealEd/Private/WorldPartition/WorldPartitionResaveActorsBuilder.cpp
@@ -176,7 +176,7 @@ bool UWorldPartitionResaveActorsBuilder::RunInternal(UWorld* World, const FCellI
 
 		// Build clusters
 		TArray<TPair<FGuid, TArray<FGuid>>> ActorsWithRefs;
-		for (FActorDescList::TIterator<> ActorDescIterator(WorldPartition); ActorDescIterator; ++ActorDescIterator)
+		for (FActorDescContainerCollection::TIterator<> ActorDescIterator(WorldPartition); ActorDescIterator; ++ActorDescIterator)
 		{
 			ActorsWithRefs.Emplace(ActorDescIterator->GetGuid(), ActorDescIterator->GetReferences());
 		}
diff --git a/Engine/Source/Runtime/Engine/Private/WorldPartition/ActorDescContainer.cpp b/Engine/Source/Runtime/Engine/Private/WorldPartition/ActorDescContainer.cpp
--- a/Engine/Source/Runtime/Engine/Private/WorldPartition/ActorDescContainer.cpp
+++ b/Engine/Source/Runtime/Engine/Private/WorldPartition/ActorDescContainer.cpp
@@ -5,6 +5,7 @@
 #if WITH_EDITOR
 #include "Editor.h"
 #include "AssetRegistry/AssetRegistryModule.h"
+#include "WorldPartition/WorldPartition.h"
 #include "WorldPartition/WorldPartitionLog.h"
 #include "WorldPartition/WorldPartitionHandle.h"
 #include "WorldPartition/WorldPartitionActorDescUtils.h"
@@ -33,17 +34,16 @@ void UActorDescContainer::Initialize(UWorld* InWorld, FName InPackageName)
 	
 	if (!ContainerPackageName.IsNone())
 	{
-		const FString LevelPathStr = ContainerPackageName.ToString();
-		const FString LevelExternalActorsPath = ULevel::GetExternalActorsPath(LevelPathStr);
+		const FString ContainerExternalActorsPath = GetExternalActorPath();
 
 		// Do a synchronous scan of the level external actors path.			
 		IAssetRegistry& AssetRegistry = FModuleManager::LoadModuleChecked<FAssetRegistryModule>(TEXT("AssetRegistry")).Get();
-		AssetRegistry.ScanPathsSynchronous({ LevelExternalActorsPath }, /*bForceRescan*/false, /*bIgnoreDenyListScanFilters*/false);
+		AssetRegistry.ScanPathsSynchronous({ ContainerExternalActorsPath }, /*bForceRescan*/true, /*bIgnoreDenyListScanFilters*/false);
 
 		FARFilter Filter;
 		Filter.bRecursivePaths = true;
 		Filter.bIncludeOnlyOnDiskAssets = true;
-		Filter.PackagePaths.Add(*LevelExternalActorsPath);
+		Filter.PackagePaths.Add(*ContainerExternalActorsPath);
 
 		AssetRegistry.GetAssets(Filter, Assets);
 	}
@@ -104,36 +104,44 @@ void UActorDescContainer::BeginDestroy()
 }
 
 #if WITH_EDITOR
-void UActorDescContainer::AddActorDescriptor(FWorldPartitionActorDesc* ActorDesc)
+FString UActorDescContainer::GetExternalActorPath() const
 {
-	FActorDescList::AddActorDescriptor(ActorDesc);
-	ActorDesc->SetContainer(this);
+	return ULevel::GetExternalActorsPath(ContainerPackageName.ToString());
 }
 
-void UActorDescContainer::RemoveActorDescriptor(FWorldPartitionActorDesc* ActorDesc)
+bool UActorDescContainer::IsActorDescHandled(const AActor* Actor) const
 {
-	ActorDesc->SetContainer(nullptr);
-	FActorDescList::RemoveActorDescriptor(ActorDesc);
+	FString ActorPackageName = Actor->GetPackage()->GetName();
+	return ActorPackageName.StartsWith(GetExternalActorPath());
 }
 
+bool UActorDescContainer::IsMainPartitionContainer() const
+{
+	return GetWorld() != nullptr && GetWorld()->GetWorldPartition()->GetActorDescContainer() == this;
+}
 
-void UActorDescContainer::OnWorldRenamed(UWorld* RenamedWorld)
+void UActorDescContainer::AddActorDescriptor(FWorldPartitionActorDesc* ActorDesc)
 {
-	if (GetWorld() == RenamedWorld)
-	{
-		OnWorldRenamed();
-	}
+	FActorDescList::AddActorDescriptor(ActorDesc);
+	ActorDesc->SetContainer(this);
 }
 
-void UActorDescContainer::OnWorldRenamed()
+void UActorDescContainer::RemoveActorDescriptor(FWorldPartitionActorDesc* ActorDesc)
 {
-	// Update container package
-	ContainerPackageName = GetWorld()->GetPackage()->GetFName();
+	ActorDesc->SetContainer(nullptr);
+	FActorDescList::RemoveActorDescriptor(ActorDesc);
 }
 
 bool UActorDescContainer::ShouldHandleActorEvent(const AActor* Actor)
 {
-	return Actor && Actor->IsMainPackageActor() && (Actor->GetLevel() != nullptr) && (Actor->GetLevel()->GetPackage()->GetFName() == ContainerPackageName);
+	if (Actor != nullptr)
+	{
+		bool bHandlesActor = GetActorDesc(Actor->GetActorGuid()) != nullptr || IsActorDescHandled(Actor);
+		return bHandlesActor && Actor->IsMainPackageActor() && (Actor->GetLevel() != nullptr);
+		
+	}
+	
+	return false;
 }
 
 void UActorDescContainer::OnObjectPreSave(UObject* Object, FObjectPreSaveContext SaveContext)
@@ -189,14 +197,17 @@ void UActorDescContainer::OnObjectsReplaced(const TMap<UObject*, UObject*>& OldT
 	}
 }
 
-void UActorDescContainer::RemoveActor(const FGuid& ActorGuid)
+bool UActorDescContainer::RemoveActor(const FGuid& ActorGuid)
 {
 	if (TUniquePtr<FWorldPartitionActorDesc>* ExistingActorDesc = GetActorDescriptor(ActorGuid))
 	{
 		OnActorDescRemoved(ExistingActorDesc->Get());
 		RemoveActorDescriptor(ExistingActorDesc->Get());
 		ExistingActorDesc->Reset();
+		return true;
 	}
+
+	return false;
 }
 
 void UActorDescContainer::LoadAllActors(TArray<FWorldPartitionReference>& OutReferences)
@@ -222,7 +233,6 @@ void UActorDescContainer::RegisterEditorDelegates()
 {
 	if (ShouldRegisterDelegates())
 	{
-		FWorldDelegates::OnPostWorldRename.AddUObject(this, &UActorDescContainer::OnWorldRenamed);
 		FCoreUObjectDelegates::OnObjectPreSave.AddUObject(this, &UActorDescContainer::OnObjectPreSave);
 		FEditorDelegates::OnPackageDeleted.AddUObject(this, &UActorDescContainer::OnPackageDeleted);
 		FCoreUObjectDelegates::OnObjectsReplaced.AddUObject(this, &UActorDescContainer::OnObjectsReplaced);
@@ -233,7 +243,6 @@ void UActorDescContainer::UnregisterEditorDelegates()
 {
 	if (ShouldRegisterDelegates())
 	{
-		FWorldDelegates::OnPostWorldRename.RemoveAll(this);
 		FCoreUObjectDelegates::OnObjectPreSave.RemoveAll(this);
 		FEditorDelegates::OnPackageDeleted.RemoveAll(this);
 		FCoreUObjectDelegates::OnObjectsReplaced.RemoveAll(this);
@@ -243,10 +252,36 @@ void UActorDescContainer::UnregisterEditorDelegates()
 void UActorDescContainer::OnActorDescAdded(FWorldPartitionActorDesc* NewActorDesc)
 {
 	OnActorDescAddedEvent.Broadcast(NewActorDesc);
+
+	if (UWorldPartition* WorldPartition = GetTypedOuter<UWorldPartition>())
+	{
+		WorldPartition->OnActorDescAdded(NewActorDesc);
+	}
 }
 
 void UActorDescContainer::OnActorDescRemoved(FWorldPartitionActorDesc* ActorDesc)
 {
 	OnActorDescRemovedEvent.Broadcast(ActorDesc);
+
+	if (UWorldPartition* WorldPartition = GetTypedOuter<UWorldPartition>())
+	{
+		WorldPartition->OnActorDescRemoved(ActorDesc);
+	}
+}
+
+void UActorDescContainer::OnActorDescUpdating(FWorldPartitionActorDesc* ActorDesc)
+{
+	if (UWorldPartition* WorldPartition = GetTypedOuter<UWorldPartition>())
+	{
+		WorldPartition->OnActorDescUpdating(ActorDesc);
+	}
+}
+
+void UActorDescContainer::OnActorDescUpdated(FWorldPartitionActorDesc* ActorDesc)
+{
+	if (UWorldPartition* WorldPartition = GetTypedOuter<UWorldPartition>())
+	{
+		WorldPartition->OnActorDescUpdated(ActorDesc);
+	}
 }
 #endif
diff --git a/Engine/Source/Runtime/Engine/Private/WorldPartition/ActorDescContainerCollection.cpp b/Engine/Source/Runtime/Engine/Private/WorldPartition/ActorDescContainerCollection.cpp
new file mode 100644
--- /dev/null
+++ b/Engine/Source/Runtime/Engine/Private/WorldPartition/ActorDescContainerCollection.cpp
@@ -0,0 +1,188 @@
+// Copyright Epic Games, Inc. All Rights Reserved.
+#include "WorldPartition/ActorDescContainerCollection.h"
+
+#include "WorldPartition/ActorDescContainer.h"
+#include "UObject/UObjectGlobals.h"
+#include "WorldPartition/WorldPartition.h"
+
+#if WITH_EDITOR
+
+void FActorDescContainerCollection::AddContainer(UActorDescContainer* Container)
+{
+	ActorDescContainerCollection.Add(Container);
+	Container->OnActorDescAddedEvent.AddRaw(this, &FActorDescContainerCollection::OnActorDescAdded);
+	Container->OnActorDescRemovedEvent.AddRaw(this, &FActorDescContainerCollection::OnActorDescRemoved);
+}
+
+bool FActorDescContainerCollection::RemoveContainer(UActorDescContainer* Container)
+{
+	Container->OnActorDescRemovedEvent.RemoveAll(this);
+	return ActorDescContainerCollection.RemoveSwap(Container) > 0;
+}
+
+bool FActorDescContainerCollection::Contains(const FName& ContainerPackageName) const
+{
+	 return Find(ContainerPackageName) != nullptr;
+}
+
+UActorDescContainer* FActorDescContainerCollection::Find(const FName& ContainerPackageName) const
+{
+	auto* ContainerPtr = ActorDescContainerCollection.FindByPredicate([&ContainerPackageName](const UActorDescContainer* ActorDescContainer) { return ActorDescContainer->GetContainerPackage() == ContainerPackageName; });
+	return ContainerPtr != nullptr ? *ContainerPtr : nullptr;
+}
+
+const FWorldPartitionActorDesc* FActorDescContainerCollection::GetActorDesc(const FGuid& Guid) const
+{
+	const FWorldPartitionActorDesc* ActorDesc = nullptr;
+	ForEachActorDescContainerBreakable([&Guid, &ActorDesc](const UActorDescContainer* ActorDescContainer)
+	{
+		ActorDesc =  ActorDescContainer->GetActorDesc(Guid);
+		return ActorDesc == nullptr;
+	});
+
+	return ActorDesc;
+}
+
+FWorldPartitionActorDesc* FActorDescContainerCollection::GetActorDesc(const FGuid& Guid)
+{
+	return const_cast<FWorldPartitionActorDesc*>(const_cast<const FActorDescContainerCollection*>(this)->GetActorDesc(Guid));
+}
+
+const FWorldPartitionActorDesc& FActorDescContainerCollection::GetActorDescChecked(const FGuid& Guid) const
+{
+	const FWorldPartitionActorDesc* ActorDesc = GetActorDesc(Guid);
+	check(ActorDesc != nullptr);
+
+	static FWorldPartitionActorDesc EmptyDescriptor;
+	return ActorDesc != nullptr ? *ActorDesc : EmptyDescriptor;
+}
+
+FWorldPartitionActorDesc& FActorDescContainerCollection::GetActorDescChecked(const FGuid& Guid)
+{
+	return const_cast<FWorldPartitionActorDesc&>(const_cast<const FActorDescContainerCollection*>(this)->GetActorDescChecked(Guid));
+}
+
+const FWorldPartitionActorDesc* FActorDescContainerCollection::GetActorDesc(const FString& ActorPath) const
+{
+	const FWorldPartitionActorDesc* ActorDesc = nullptr;
+	ForEachActorDescContainerBreakable([&ActorPath, &ActorDesc](const UActorDescContainer* ActorDescContainer)
+	{
+		ActorDesc = ActorDescContainer->GetActorDesc(ActorPath);
+		return ActorDesc == nullptr;
+	});
+
+	return ActorDesc;
+}
+
+const FWorldPartitionActorDesc* FActorDescContainerCollection::GetActorDesc(const FSoftObjectPath& InActorPath) const
+{
+	const FWorldPartitionActorDesc* ActorDesc = nullptr;
+	ForEachActorDescContainerBreakable([&InActorPath, &ActorDesc](const UActorDescContainer* ActorDescContainer)
+	{
+		ActorDesc = ActorDescContainer->GetActorDesc(InActorPath);
+		return ActorDesc == nullptr;
+	});
+
+	return ActorDesc;
+}
+
+bool FActorDescContainerCollection::RemoveActor(const FGuid& ActorGuid)
+{
+	bool bRemoved = false;
+	ForEachActorDescContainerBreakable([&ActorGuid, &bRemoved](UActorDescContainer* ActorDescContainer)
+	{
+		bRemoved = ActorDescContainer->RemoveActor(ActorGuid);
+		return !bRemoved;
+	});
+
+	return bRemoved;
+}
+
+void FActorDescContainerCollection::OnPackageDeleted(UPackage* Package)
+{
+	ForEachActorDescContainer([Package](UActorDescContainer* ActorDescContainer)
+	{
+		ActorDescContainer->OnPackageDeleted(Package);
+	});
+}
+
+bool FActorDescContainerCollection::IsActorDescHandled(const AActor* Actor) const
+{
+	bool bIsHandled = false;
+	ForEachActorDescContainerBreakable([Actor, &bIsHandled](UActorDescContainer* ActorDescContainer)
+	{
+		bIsHandled = ActorDescContainer->IsActorDescHandled(Actor);
+		return !bIsHandled;
+	});
+
+	return bIsHandled;
+}
+
+void FActorDescContainerCollection::LoadAllActors(TArray<FWorldPartitionReference>& OutReferences)
+{
+	ForEachActorDescContainer([&OutReferences](UActorDescContainer* ActorDescContainer)
+	{
+		ActorDescContainer->LoadAllActors(OutReferences);
+	});
+}
+
+const UActorDescContainer* FActorDescContainerCollection::GetActorDescContainer(const FGuid& ActorGuid) const
+{
+	UActorDescContainer* ActorDescContainer = nullptr;
+	ForEachActorDescContainerBreakable([&ActorGuid, &ActorDescContainer](UActorDescContainer* InActorDescContainer)
+	{
+		if (InActorDescContainer->GetActorDesc(ActorGuid) != nullptr)
+		{
+			ActorDescContainer = InActorDescContainer;
+		}
+		return ActorDescContainer == nullptr;
+	});
+
+	return ActorDescContainer;
+}
+
+UActorDescContainer* FActorDescContainerCollection::GetActorDescContainer(const FGuid& ActorGuid)
+{
+	return const_cast<UActorDescContainer*>(const_cast<const FActorDescContainerCollection*>(this)->GetActorDescContainer(ActorGuid));
+}
+
+void FActorDescContainerCollection::ForEachActorDescContainerBreakable(TFunctionRef<bool(UActorDescContainer*)> Func) const
+{
+	for (UActorDescContainer* ActorDescContainer : ActorDescContainerCollection)
+	{
+		if (!Func(ActorDescContainer))
+		{
+			break;
+		}
+	}
+}
+
+void FActorDescContainerCollection::ForEachActorDescContainerBreakable(TFunctionRef<bool(UActorDescContainer*)> Func)
+{
+	const_cast<const FActorDescContainerCollection*>(this)->ForEachActorDescContainerBreakable(Func);
+}
+
+void FActorDescContainerCollection::ForEachActorDescContainer(TFunctionRef<void(UActorDescContainer*)> Func) const
+{
+	for (UActorDescContainer* ActorDescContainer : ActorDescContainerCollection)
+	{
+		Func(ActorDescContainer);
+	}
+}
+
+void FActorDescContainerCollection::ForEachActorDescContainer(TFunctionRef<void(UActorDescContainer*)> Func)
+{
+	const_cast<const FActorDescContainerCollection*>(this)->ForEachActorDescContainer(Func);
+}
+
+void FActorDescContainerCollection::OnActorDescAdded(FWorldPartitionActorDesc* ActorDesc)
+{
+	OnActorDescAddedEvent.Broadcast(ActorDesc);
+}
+
+void FActorDescContainerCollection::OnActorDescRemoved(FWorldPartitionActorDesc* ActorDesc)
+{
+	OnActorDescRemovedEvent.Broadcast(ActorDesc);
+}
+
+#endif
\ No newline at end of file
diff --git a/Engine/Source/Runtime/Engine/Private/WorldPartition/RuntimeSpatialHash/RuntimeSpatialHashHLOD.cpp b/Engine/Source/Runtime/Engine/Private/WorldPartition/RuntimeSpatialHash/RuntimeSpatialHashHLOD.cpp
--- a/Engine/Source/Runtime/Engine/Private/WorldPartition/RuntimeSpatialHash/RuntimeSpatialHashHLOD.cpp
+++ b/Engine/Source/Runtime/Engine/Private/WorldPartition/RuntimeSpatialHash/RuntimeSpatialHashHLOD.cpp
@@ -12,6 +12,7 @@
 #include "WorldPartition/HLOD/IWorldPartitionHLODUtilitiesModule.h"
 #include "WorldPartition/DataLayer/DataLayerSubsystem.h"
 
+#include "WorldPartition/ActorDescContainer.h"
 #include "WorldPartition/WorldPartitionHelpers.h"
 
 #include "UObject/GCObjectScopeGuard.h"
@@ -87,7 +88,7 @@ static void DeletePackage(UWorldPartition* WorldPartition, FWorldPartitionActorD
 	if (ActorDesc->IsLoaded())
 	{
 		DeletePackage(ActorDesc->GetActor()->GetPackage(), SourceControlHelper);
-		WorldPartition->OnPackageDeleted(ActorDesc->GetActor()->GetPackage());
+		WorldPartition->GetActorDescContainer()->OnPackageDeleted(ActorDesc->GetActor()->GetPackage());
 	}
 	else
 	{
@@ -260,7 +261,7 @@ static TMap<UHLODLayer*, int32> GatherHLODLayers(UWorldPartition* WorldPartition
 	// Gather up all HLODLayers referenced by the actors, along with the HLOD level at which it was used
 	TMap<UHLODLayer*, int32> HLODLayersLevel;
 
-	for (FActorDescList::TIterator<> ActorDescIterator(WorldPartition); ActorDescIterator; ++ActorDescIterator)
+	for (FActorDescList::TIterator<> ActorDescIterator(WorldPartition->GetActorDescContainer()); ActorDescIterator; ++ActorDescIterator)
 	{
 		const FWorldPartitionActorDesc& ActorDesc = **ActorDescIterator;
 		if (!ActorDesc.GetActorNativeClass()->IsChildOf<AWorldPartitionHLOD>())
@@ -420,7 +421,7 @@ bool UWorldPartitionRuntimeSpatialHash::GenerateHLOD(ISourceControlHelper* Sourc
 
 	TSet<FWorldPartitionHandle> InvalidHLODActors;
 
-	for (FActorDescList::TIterator<AWorldPartitionHLOD> HLODIterator(WorldPartition); HLODIterator; ++HLODIterator)
+	for (FActorDescList::TIterator<AWorldPartitionHLOD> HLODIterator(WorldPartition->GetActorDescContainer()); HLODIterator; ++HLODIterator)
 	{
 		uint64 CellHash = HLODIterator->GetCellHash();
 
diff --git a/Engine/Source/Runtime/Engine/Private/WorldPartition/WorldPartition.cpp b/Engine/Source/Runtime/Engine/Private/WorldPartition/WorldPartition.cpp
--- a/Engine/Source/Runtime/Engine/Private/WorldPartition/WorldPartition.cpp
+++ b/Engine/Source/Runtime/Engine/Private/WorldPartition/WorldPartition.cpp
@@ -4,6 +4,7 @@
 	WorldPartition.cpp: UWorldPartition implementation
 =============================================================================*/
 #include "WorldPartition/WorldPartition.h"
+#include "WorldPartition/ActorDescContainer.h"
 #include "WorldPartition/WorldPartitionLog.h"
 #include "WorldPartition/WorldPartitionRuntimeCell.h"
 #include "WorldPartition/WorldPartitionStreamingPolicy.h"
@@ -272,9 +273,11 @@ void UWorldPartition::OnGCPostReachabilityAnalysis()
 
 void UWorldPartition::OnPackageDirtyStateChanged(UPackage* Package)
 {
+	auto ShouldHandleActor = [this](AActor* Actor) {return Actor && Actor->IsMainPackageActor() && (Actor->GetLevel() != nullptr) && IsActorDescHandled(Actor);};
+
 	if (IsStreamingEnabled())
 	{
-		if (AActor* Actor = AActor::FindActorInPackage(Package); ShouldHandleActorEvent(Actor))
+		if (AActor* Actor = AActor::FindActorInPackage(Package); ShouldHandleActor(Actor))
 		{
 			if (FWorldPartitionHandle ActorHandle(this, Actor->GetActorGuid()); ActorHandle.IsValid())
 			{
@@ -440,13 +443,12 @@ void UWorldPartition::Initialize(UWorld* InWorld, const FTransform& InTransform)
 
 		{
 			TRACE_CPUPROFILER_EVENT_SCOPE(UActorDescContainer::Initialize);
-			UActorDescContainer::Initialize(World, PackageName);
-			check(bContainerInitialized);
+			ActorDescContainer = RegisterActorDescContainer(PackageName);
 		}
 
 		{
 			TRACE_CPUPROFILER_EVENT_SCOPE(UActorDescContainer::Hash);
-			for (FActorDescList::TIterator<> ActorDescIterator(this); ActorDescIterator; ++ActorDescIterator)
+			for (FActorDescContainerCollection::TIterator<> ActorDescIterator(this); ActorDescIterator; ++ActorDescIterator)
 			{
 				if (bIsInstanced)
 				{
@@ -475,7 +477,7 @@ void UWorldPartition::Initialize(UWorld* InWorld, const FTransform& InTransform)
 		FlushAsyncLoading();
 
 		// Make sure to preload only AWorldDataLayers actor first
-		for (UActorDescContainer::TIterator<> ActorDescIterator(this); ActorDescIterator; ++ActorDescIterator)
+		for (FActorDescList::TIterator<> ActorDescIterator(ActorDescContainer); ActorDescIterator; ++ActorDescIterator)
 		{
 			if (ActorDescIterator->GetActorNativeClass()->IsChildOf<AWorldDataLayers>())
 			{
@@ -620,6 +622,9 @@ void UWorldPartition::Uninitialize()
 
 		DirtyActors.Empty();
 
+		UninitializeActorDescContainers();
+		ActorDescContainer = nullptr;
+
 		EditorHash = nullptr;
 #endif		
 
@@ -630,8 +635,6 @@ PRAGMA_DISABLE_DEPRECATION_WARNINGS
 		OnWorldPartitionUninitialized.Broadcast(this);
 PRAGMA_ENABLE_DEPRECATION_WARNINGS
 	}
-
-	Super::Uninitialize();
 }
 
 bool UWorldPartition::IsInitialized() const
@@ -682,6 +685,7 @@ void UWorldPartition::RegisterDelegates()
 			FCoreUObjectDelegates::PostReachabilityAnalysis.AddUObject(this, &UWorldPartition::OnGCPostReachabilityAnalysis);
 			GEditor->OnPostBugItGoCalled().AddUObject(this, &UWorldPartition::OnPostBugItGoCalled);
 			GEditor->OnEditorClose().AddUObject(this, &UWorldPartition::SavePerUserSettings);
+			FWorldDelegates::OnPostWorldRename.AddUObject(this, &UWorldPartition::OnWorldRenamed);
 
 			if (!IsRunningCommandlet())
 			{
@@ -717,6 +721,7 @@ void UWorldPartition::UnregisterDelegates()
 	{
 		if (IsMainWorldPartition())
 		{
+			FWorldDelegates::OnPostWorldRename.RemoveAll(this);
 			FEditorDelegates::PreBeginPIE.RemoveAll(this);
 			FEditorDelegates::PrePIEEnded.RemoveAll(this);
 			FEditorDelegates::CancelPIE.RemoveAll(this);
@@ -869,8 +874,6 @@ bool UWorldPartition::IsSimulating()
 #if WITH_EDITOR
 void UWorldPartition::OnActorDescAdded(FWorldPartitionActorDesc* NewActorDesc)
 {
-	Super::OnActorDescAdded(NewActorDesc);
-
 	NewActorDesc->bIsForcedNonSpatiallyLoaded = !IsStreamingEnabled();
 
 	HashActorDesc(NewActorDesc);
@@ -882,9 +885,7 @@ void UWorldPartition::OnActorDescAdded(FWorldPartitionActorDesc* NewActorDesc)
 }
 
 void UWorldPartition::OnActorDescRemoved(FWorldPartitionActorDesc* ActorDesc)
-{
-	Super::OnActorDescRemoved(ActorDesc);
-	
+{	
 	UnhashActorDesc(ActorDesc);
 		
 	if (WorldPartitionEditor)
@@ -921,7 +922,7 @@ bool UWorldPartition::GetInstancingContext(const FLinkerInstancingContext*& OutI
 
 const FTransform& UWorldPartition::GetInstanceTransform() const
 {
-	return InstanceTransform.IsSet() ? InstanceTransform.GetValue() : Super::GetInstanceTransform();
+	return InstanceTransform.IsSet() ? InstanceTransform.GetValue() : FTransform::Identity;
 }
 
 #if WITH_EDITOR
@@ -968,7 +969,7 @@ void UWorldPartition::SetCanBeUsedByLevelInstance(bool bInCanBeUsedByLevelInstan
 
 void UWorldPartition::OnEnableStreamingChanged()
 {
-	for (UActorDescContainer::TIterator<> ActorDescIterator(this); ActorDescIterator; ++ActorDescIterator)
+	for (FActorDescContainerCollection::TIterator<> ActorDescIterator(this); ActorDescIterator; ++ActorDescIterator)
 	{
 		UnhashActorDesc(*ActorDescIterator);
 		ActorDescIterator->bIsForcedNonSpatiallyLoaded = !IsStreamingEnabled();
@@ -1008,6 +1009,22 @@ void UWorldPartition::UnhashActorDesc(FWorldPartitionActorDesc* ActorDesc)
 	FWorldPartitionHandle ActorHandle(this, ActorDesc->GetGuid());
 	EditorHash->UnhashActor(ActorHandle);
 }
+
+void UWorldPartition::HashActorDescContainer(UActorDescContainer* InActorDescContainer)
+{
+	for (UActorDescContainer::TIterator<> It(InActorDescContainer); It; ++It)
+	{
+		HashActorDesc(*It);
+	}
+}
+
+void UWorldPartition::UnhashActorDescContainer(UActorDescContainer* InActorDescContainer)
+{
+	for (UActorDescContainer::TIterator<> It(InActorDescContainer); It; ++It)
+	{
+		UnhashActorDesc(*It);
+	}
+}
 #endif
 
 void UWorldPartition::Serialize(FArchive& Ar)
@@ -1039,6 +1056,15 @@ void UWorldPartition::Serialize(FArchive& Ar)
 	}
 }
 
+UWorld* UWorldPartition::GetWorld() const
+{
+	if (World)
+	{
+		return World;
+	}
+	return Super::GetWorld();
+}
+
 bool UWorldPartition::ResolveSubobject(const TCHAR* SubObjectPath, UObject*& OutObject, bool bLoadIfExists)
 {
 	if (GetWorld())
@@ -1089,14 +1115,14 @@ bool UWorldPartition::ResolveSubobject(const TCHAR* SubObjectPath, UObject*& Out
 void UWorldPartition::AddReferencedObjects(UObject* InThis, FReferenceCollector& Collector)
 {
 #if WITH_EDITOR
+	UWorldPartition* This = CastChecked<UWorldPartition>(InThis);
+
 	// We need to keep all dirty actors alive, mainly for deleted actors. Normally, these actors are only referenced
 	// by the transaction buffer, but we clear it when unloading regions, etc. and we don't want these actors to die.
 	// Also, we must avoid reporting these references when not collecting garbage, as code such as package deletion
 	// will skip packages with actors still referenced (via GatherObjectReferencersForDeletion).
 	if (IsGarbageCollecting())
 	{
-		UWorldPartition* This = CastChecked<UWorldPartition>(InThis);
-	
 		Collector.AllowEliminatingReferences(false);
 		for (auto& [ActorReference, Actor] : This->DirtyActors)
 		{
@@ -1104,6 +1130,11 @@ void UWorldPartition::AddReferencedObjects(UObject* InThis, FReferenceCollector&
 		}
 		Collector.AllowEliminatingReferences(true);
 	}
+
+	for (const UActorDescContainer* Container : This->ActorDescContainerCollection)
+	{
+		Collector.AddReferencedObject(Container);
+	}
 #endif
 
 	Super::AddReferencedObjects(InThis, Collector);
@@ -1132,7 +1163,7 @@ void UWorldPartition::Tick(float DeltaSeconds)
 		if (!IsStreamingEnabled() && SupportsStreaming() && !bStreamingWasEnabled && !bShouldEnableStreamingWarned)
 		{
 			FBox AllActorsBounds(ForceInit);
-			for (UActorDescContainer::TConstIterator<> ActorDescIterator(this); ActorDescIterator; ++ActorDescIterator)
+			for (FActorDescContainerCollection::TConstIterator<> ActorDescIterator(this); ActorDescIterator; ++ActorDescIterator)
 			{
 				if (ActorDescIterator->GetIsSpatiallyLoadedRaw())
 				{
@@ -1271,7 +1302,7 @@ void UWorldPartition::SavePerUserSettings()
 		GetMutableDefault<UWorldPartitionEditorPerProjectUserSettings>()->SetEditorLoadedRegions(GetWorld(), GetUserLoadedEditorRegions());
 
 		TArray<FName> EditorLoadedLocationVolumes;
-		for (UActorDescContainer::TConstIterator<> ActorDescIterator(this); ActorDescIterator; ++ActorDescIterator)
+		for (FActorDescContainerCollection::TConstIterator<> ActorDescIterator(this); ActorDescIterator; ++ActorDescIterator)
 		{
 			if (ALocationVolume* LocationVolume = Cast<ALocationVolume>(ActorDescIterator->GetActor()))
 			{
@@ -1296,7 +1327,7 @@ void UWorldPartition::DumpActorDescs(const FString& Path)
 	{
 		TArray<const FWorldPartitionActorDesc*> ActorDescs;
 		TMap<FName, FString> DataLayersDumpString = GetDataLayersDumpString(this);
-		for (UActorDescContainer::TConstIterator<> ActorDescIterator(this); ActorDescIterator; ++ActorDescIterator)
+		for (FActorDescContainerCollection::TConstIterator<> ActorDescIterator(this); ActorDescIterator; ++ActorDescIterator)
 		{
 			ActorDescs.Add(*ActorDescIterator);
 		}
@@ -1353,6 +1384,81 @@ void UWorldPartition::AppendAssetRegistryTags(TArray<FAssetRegistryTag>& OutTags
 	}
 }
 
+UActorDescContainer* UWorldPartition::RegisterActorDescContainer(const FName& ContainerPackage)
+{
+	if (!Contains(ContainerPackage))
+	{
+		UActorDescContainer* ContainerToRegister = NewObject<UActorDescContainer>(this);
+		ContainerToRegister->Initialize(GetWorld(), ContainerPackage);
+		AddContainer(ContainerToRegister);
+
+		FWorldPartitionReference WDLReference;
+		for (FActorDescList::TIterator<> ActorDescIterator(ContainerToRegister); ActorDescIterator; ++ActorDescIterator)
+		{
+			if (ActorDescIterator->GetActorNativeClass()->IsChildOf<AWorldDataLayers>())
+			{
+				WDLReference = FWorldPartitionReference(this, ActorDescIterator->GetGuid());
+				break;
+			}
+		}
+
+		if (IsInitialized() && EditorHash != nullptr)
+		{
+			HashActorDescContainer(ContainerToRegister);
+		}
+
+		OnActorDescContainerRegistered.Broadcast(ContainerToRegister);
+
+		return ContainerToRegister;
+	}
+	
+	return nullptr;
+}
+
+bool UWorldPartition::UnregisterActorDescContainer(UActorDescContainer* InActorDescContainer)
+{
+	if (Contains(InActorDescContainer->GetContainerPackage()))
+	{
+		TArray<FGuid> ActorGuids;
+		for (UActorDescContainer::TIterator<> It(InActorDescContainer); It; ++It)
+		{
+			FWorldPartitionHandle ActorHandle(this, It->GetGuid());
+			if (ActorHandle.IsValid())
+			{
+				ActorGuids.Add(It->GetGuid());
+				DirtyActors.Remove(ActorHandle.ToReference());
+			}
+		}
+
+		UnpinActors(ActorGuids);
+
+		OnActorDescContainerUnregistered.Broadcast(InActorDescContainer);
+
+		if (IsInitialized() && EditorHash != nullptr)
+		{
+			UnhashActorDescContainer(InActorDescContainer);
+		}
+
+		InActorDescContainer->Uninitialize();
+
+		verify(RemoveContainer(InActorDescContainer));
+
+		return true;
+	}
+
+	return false;
+}
+
+void UWorldPartition::UninitializeActorDescContainers()
+{
+	for (UActorDescContainer* Container : ActorDescContainerCollection)
+	{
+		Container->Uninitialize();
+	}
+
+	Empty();
+}
+
 void UWorldPartition::PinActors(const TArray<FGuid>& ActorGuids)
 {
 	if (PinnedActors)
@@ -1407,12 +1513,15 @@ void UWorldPartition::LoadLastLoadedRegions()
 	}
 }
 
-void UWorldPartition::OnWorldRenamed()
+void UWorldPartition::OnWorldRenamed(UWorld* RenamedWorld)
 {
-	UActorDescContainer::OnWorldRenamed();
+	if (GetWorld() == RenamedWorld)
+	{
+		ActorDescContainer->SetContainerPackage(GetWorld()->GetPackage()->GetFName());
 
-	// World was renamed so existing context is invalid.
-	InstancingContext = FLinkerInstancingContext();
+		// World was renamed so existing context is invalid.
+		InstancingContext = FLinkerInstancingContext();
+	}
 }
 
 void UWorldPartition::RemapSoftObjectPath(FSoftObjectPath& ObjectPath)
diff --git a/Engine/Source/Runtime/Engine/Private/WorldPartition/WorldPartitionActorDesc.cpp b/Engine/Source/Runtime/Engine/Private/WorldPartition/WorldPartitionActorDesc.cpp
--- a/Engine/Source/Runtime/Engine/Private/WorldPartition/WorldPartitionActorDesc.cpp
+++ b/Engine/Source/Runtime/Engine/Private/WorldPartition/WorldPartitionActorDesc.cpp
@@ -15,6 +15,7 @@
 #include "UObject/UE5MainStreamObjectVersion.h"
 #include "UObject/UE5ReleaseStreamObjectVersion.h"
 #include "UObject/FortniteNCBranchObjectVersion.h"
+#include "WorldPartition/WorldPartition.h"
 #include "WorldPartition/WorldPartitionLog.h"
 #include "WorldPartition/ActorDescContainer.h"
 #include "WorldPartition/HLOD/HLODLayer.h"
@@ -496,7 +497,10 @@ AActor* FWorldPartitionActorDesc::Load() const
 		const FLinkerInstancingContext* InstancingContext = nullptr;
 		if (Container)
 		{
-			Container->GetInstancingContext(InstancingContext);
+			if (UWorldPartition* WorldPartition = Container->GetTypedOuter<UWorldPartition>())
+			{
+				WorldPartition->GetInstancingContext(InstancingContext);
+			}
 		}
 
 		UPackage* Package = nullptr;
diff --git a/Engine/Source/Runtime/Engine/Private/WorldPartition/WorldPartitionActorLoaderInterface.cpp b/Engine/Source/Runtime/Engine/Private/WorldPartition/WorldPartitionActorLoaderInterface.cpp
--- a/Engine/Source/Runtime/Engine/Private/WorldPartition/WorldPartitionActorLoaderInterface.cpp
+++ b/Engine/Source/Runtime/Engine/Private/WorldPartition/WorldPartitionActorLoaderInterface.cpp
@@ -96,11 +96,23 @@ bool IWorldPartitionActorLoaderInterface::ILoaderAdapter::IsLoaded() const
 void IWorldPartitionActorLoaderInterface::ILoaderAdapter::RegisterDelegates()
 {
 	FDataLayersEditorBroadcast::Get().OnActorDataLayersEditorLoadingStateChanged().AddRaw(this, &IWorldPartitionActorLoaderInterface::ILoaderAdapter::OnActorDataLayersEditorLoadingStateChanged);
+	check(World->GetWorldPartition() != nullptr);
+	if (UWorldPartition* WorldPartition = World->GetWorldPartition())
+	{
+		WorldPartition->OnActorDescContainerRegistered.AddRaw(this, &IWorldPartitionActorLoaderInterface::ILoaderAdapter::ILoaderAdapter::OnActorDescContainerInitialize);
+		WorldPartition->OnActorDescContainerUnregistered.AddRaw(this, &IWorldPartitionActorLoaderInterface::ILoaderAdapter::ILoaderAdapter::OnActorDescContainerUninitialize);
+	}
 }
 
 void IWorldPartitionActorLoaderInterface::ILoaderAdapter::UnregisterDelegates()
 {
 	FDataLayersEditorBroadcast::Get().OnActorDataLayersEditorLoadingStateChanged().RemoveAll(this);
+	check(World->GetWorldPartition() != nullptr);
+	if (UWorldPartition* WorldPartition = World->GetWorldPartition())
+	{
+		WorldPartition->OnActorDescContainerRegistered.RemoveAll(this);
+		WorldPartition->OnActorDescContainerUnregistered.RemoveAll(this);
+	}
 }
 
 bool IWorldPartitionActorLoaderInterface::ILoaderAdapter::PassActorDescFilter(const FWorldPartitionHandle& Actor) const
@@ -215,6 +227,41 @@ void IWorldPartitionActorLoaderInterface::ILoaderAdapter::RefreshLoadedState()
 	}
 }
 
+void IWorldPartitionActorLoaderInterface::ILoaderAdapter::OnActorDescContainerInitialize(UActorDescContainer* Container)
+{
+	RefreshLoadedState();
+}
+
+void IWorldPartitionActorLoaderInterface::ILoaderAdapter::OnActorDescContainerUninitialize(UActorDescContainer* Container)
+{
+	// @todo_ow :
+	// If the container is removed before broadcasting its uninitialize event we could simply call RefreshLoadedState(), but it would break call symetry
+	TArray<FWorldPartitionHandle> ActorsToUnload;
+	for(UActorDescContainer::TConstIterator<> It(Container); It; ++It)
+	{
+		if( ActorReferences.Find(It->GetGuid()))
+		{
+			ActorsToUnload.Emplace(Container, It->GetGuid());
+		}
+	}
+	
+	if (ActorsToUnload.Num())
+	{
+		FWorldPartitionLoadingContext::FDeferred LoadingContext;
+	
+		FScopedSlowTask SlowTask(ActorsToUnload.Num(), LOCTEXT("ActorsUnoading", "Unloading actors..."));
+		SlowTask.MakeDialogDelayed(1.0f);
+	
+		for (FWorldPartitionHandle& ActorToUnload : ActorsToUnload)
+		{
+			RemoveReferenceToActor(ActorToUnload);
+			SlowTask.EnterProgressFrame(1);
+		}
+	
+		PostLoadedStateChanged(0, ActorsToUnload.Num());
+	}
+}
+
 bool IWorldPartitionActorLoaderInterface::ILoaderAdapter::ShouldActorBeLoaded(const FWorldPartitionHandle& Actor) const
 {
 	check(Actor.IsValid());
diff --git a/Engine/Source/Runtime/Engine/Private/WorldPartition/WorldPartitionHandle.cpp b/Engine/Source/Runtime/Engine/Private/WorldPartition/WorldPartitionHandle.cpp
--- a/Engine/Source/Runtime/Engine/Private/WorldPartition/WorldPartitionHandle.cpp
+++ b/Engine/Source/Runtime/Engine/Private/WorldPartition/WorldPartitionHandle.cpp
@@ -3,6 +3,8 @@
 #include "WorldPartition/WorldPartitionHandle.h"
 #include "WorldPartition/WorldPartitionActorDesc.h"
 #include "WorldPartition/ActorDescContainer.h"
+#include "WorldPartition/WorldPartition.h"
+
 
 #if WITH_EDITOR
 /**
@@ -23,6 +25,11 @@ UActorDescContainer* FWorldPartitionImplBase::GetActorDescContainer(TUniquePtr<F
 	return ActorDesc ? ActorDesc->Get()->GetContainer() : nullptr;
 }
 
+UActorDescContainer* FWorldPartitionImplBase::GetActorDescContainer(FActorDescContainerCollection* ContainerCollection, const FGuid& ActorGuid)
+{
+	return ContainerCollection ? ContainerCollection->GetActorDescContainer(ActorGuid) : nullptr;
+}
+
 bool FWorldPartitionImplBase::IsActorDescLoaded(FWorldPartitionActorDesc* ActorDesc)
 {
 	return ActorDesc->IsLoaded();
@@ -71,7 +78,7 @@ void FWorldPartitionLoadingContext::FImmediate::RegisterActor(FWorldPartitionAct
 		UActorDescContainer* Container = ActorDesc->GetContainer();
 		check(Container);
 
-		const FTransform& ContainerTransform = Container->GetInstanceTransform();
+		const FTransform& ContainerTransform = Container->GetTypedOuter<UWorldPartition>()->GetInstanceTransform();
 		const FTransform* ContainerTransformPtr = ContainerTransform.Equals(FTransform::Identity) ? nullptr : &ContainerTransform;
 
 		Actor->GetLevel()->AddLoadedActor(Actor, ContainerTransformPtr);
@@ -86,7 +93,7 @@ void FWorldPartitionLoadingContext::FImmediate::UnregisterActor(FWorldPartitionA
 		UActorDescContainer* Container = ActorDesc->GetContainer();
 		check(Container);
 
-		const FTransform& ContainerTransform = Container->GetInstanceTransform();
+		const FTransform& ContainerTransform = Container->GetTypedOuter<UWorldPartition>()->GetInstanceTransform();
 		const FTransform* ContainerTransformPtr = ContainerTransform.Equals(FTransform::Identity) ? nullptr : &ContainerTransform;
 
 		Actor->GetLevel()->RemoveLoadedActor(Actor, ContainerTransformPtr);
@@ -106,7 +113,7 @@ FWorldPartitionLoadingContext::FDeferred::~FDeferred()
 	{
 		for (auto& [Container, ContainerOp] : LocalContainerOps)
 		{
-			const FTransform& ContainerTransform = Container->GetInstanceTransform();
+			const FTransform& ContainerTransform = Container->GetTypedOuter<UWorldPartition>()->GetInstanceTransform();
 			const FTransform* ContainerTransformPtr = ContainerTransform.Equals(FTransform::Identity) ? nullptr : &ContainerTransform;
 
 			auto CreateActorList = [](TArray<AActor*>& ActorList, const TSet<FWorldPartitionActorDesc*>& SourceList) -> ULevel*
diff --git a/Engine/Source/Runtime/Engine/Private/WorldPartition/WorldPartitionHelpers.cpp b/Engine/Source/Runtime/Engine/Private/WorldPartition/WorldPartitionHelpers.cpp
--- a/Engine/Source/Runtime/Engine/Private/WorldPartition/WorldPartitionHelpers.cpp
+++ b/Engine/Source/Runtime/Engine/Private/WorldPartition/WorldPartitionHelpers.cpp
@@ -48,7 +48,7 @@ void FWorldPartitionHelpers::ForEachIntersectingActorDesc(UWorldPartition* World
 
 void FWorldPartitionHelpers::ForEachActorDesc(UWorldPartition* WorldPartition, TSubclassOf<AActor> ActorClass, TFunctionRef<bool(const FWorldPartitionActorDesc*)> Func)
 {
-	for (UActorDescContainer::TConstIterator<> ActorDescIterator(WorldPartition); ActorDescIterator; ++ActorDescIterator)
+	for (FActorDescContainerCollection::TConstIterator<> ActorDescIterator(WorldPartition); ActorDescIterator; ++ActorDescIterator)
 	{
 		if (IsActorDescClassCompatibleWith(*ActorDescIterator, ActorClass))
 		{
@@ -132,7 +132,7 @@ void FWorldPartitionHelpers::ForEachActorWithLoading(UWorldPartition* WorldParti
 		DoCollectGarbage();
 	};
 
-	for (UActorDescContainer::TConstIterator<> ActorDescIterator(WorldPartition); ActorDescIterator; ++ActorDescIterator)
+	for (FActorDescContainerCollection::TConstIterator<> ActorDescIterator(WorldPartition); ActorDescIterator; ++ActorDescIterator)
 	{
 		if (IsActorDescClassCompatibleWith(*ActorDescIterator, Params.ActorClass))
 		{
diff --git a/Engine/Source/Runtime/Engine/Private/WorldPartition/WorldPartitionLevelHelper.cpp b/Engine/Source/Runtime/Engine/Private/WorldPartition/WorldPartitionLevelHelper.cpp
--- a/Engine/Source/Runtime/Engine/Private/WorldPartition/WorldPartitionLevelHelper.cpp
+++ b/Engine/Source/Runtime/Engine/Private/WorldPartition/WorldPartitionLevelHelper.cpp
@@ -336,7 +336,7 @@ bool FWorldPartitionLevelHelper::LoadActors(UWorld* InOwningWorld, ULevel* InDes
 			Context = &NewContext;
 		}
 		
-		const FName ContainerPackageInstanceName = Context->RemapPackage(PackageObjectMapping.ContainerPackage);
+		const FName ContainerPackageInstanceName = Context->RemapPackage(PackageObjectMapping.WorldPackage);
 		if (PackageObjectMapping.ContainerPackage != ContainerPackageInstanceName)
 		{
 			const FString ActorPackageName = FPackageName::ObjectPathToPackageName(PackageObjectMapping.Package.ToString());
diff --git a/Engine/Source/Runtime/Engine/Private/WorldPartition/WorldPartitionRuntimeLevelStreamingCell.cpp b/Engine/Source/Runtime/Engine/Private/WorldPartition/WorldPartitionRuntimeLevelStreamingCell.cpp
--- a/Engine/Source/Runtime/Engine/Private/WorldPartition/WorldPartitionRuntimeLevelStreamingCell.cpp
+++ b/Engine/Source/Runtime/Engine/Private/WorldPartition/WorldPartitionRuntimeLevelStreamingCell.cpp
@@ -10,6 +10,7 @@
 #include "WorldPartition/WorldPartitionDebugHelper.h"
 
 #if WITH_EDITOR
+#include "WorldPartition/ActorDescContainer.h"
 #include "WorldPartition/WorldPartitionLevelStreamingPolicy.h"
 #include "WorldPartition/WorldPartitionLevelHelper.h"
 #include "Engine/LevelStreamingAlwaysLoaded.h"
@@ -100,7 +101,7 @@ void UWorldPartitionRuntimeLevelStreamingCell::SetIsAlwaysLoaded(bool bInIsAlway
 void UWorldPartitionRuntimeLevelStreamingCell::AddActorToCell(const FWorldPartitionActorDescView& ActorDescView, const FActorContainerID& InContainerID, const FTransform& InContainerTransform, const UActorDescContainer* InContainer)
 {
 	check(!ActorDescView.GetActorIsEditorOnly());
-	Packages.Emplace(ActorDescView.GetActorPackage(), ActorDescView.GetActorPath(), InContainerID, InContainerTransform, InContainer->GetContainerPackage());
+	Packages.Emplace(ActorDescView.GetActorPackage(), ActorDescView.GetActorPath(), InContainerID, InContainerTransform, InContainer->GetContainerPackage(), GetWorld()->GetPackage()->GetFName());
 }
 
 UWorldPartitionLevelStreamingDynamic* UWorldPartitionRuntimeLevelStreamingCell::CreateLevelStreaming(const FString& InPackageName) const
diff --git a/Engine/Source/Runtime/Engine/Private/WorldPartition/WorldPartitionStreamingGeneration.cpp b/Engine/Source/Runtime/Engine/Private/WorldPartition/WorldPartitionStreamingGeneration.cpp
--- a/Engine/Source/Runtime/Engine/Private/WorldPartition/WorldPartitionStreamingGeneration.cpp
+++ b/Engine/Source/Runtime/Engine/Private/WorldPartition/WorldPartitionStreamingGeneration.cpp
@@ -337,7 +337,7 @@ class FWorldPartitionStreamingGenerator
 		{
 			for (AActor* Actor : InContainer->GetWorld()->PersistentLevel->Actors)
 			{
-				if (IsValid(Actor) && Actor->IsPackageExternal() && Actor->IsMainPackageActor() && !Actor->IsEditorOnly() && !InContainer->GetActorDesc(Actor->GetActorGuid()))
+				if (IsValid(Actor) && Actor->IsPackageExternal() && Actor->IsMainPackageActor() && !Actor->IsEditorOnly() && InContainer->IsActorDescHandled(Actor) && !InContainer->GetActorDesc(Actor->GetActorGuid()))
 				{
 					FWorldPartitionActorDescView ModifiedActorDescView = GetModifiedActorDesc(Actor, InContainer);
 					RegisterActorDescView(Actor->GetActorGuid(), ModifiedActorDescView);
@@ -857,6 +857,11 @@ private:
 };
 
 bool UWorldPartition::GenerateStreaming(TArray<FString>* OutPackagesToGenerate)
+{
+	return GenerateContainerStreaming(ActorDescContainer, OutPackagesToGenerate);
+}
+
+bool UWorldPartition::GenerateContainerStreaming(const UActorDescContainer* InActorDescContainer, TArray<FString>* OutPackagesToGenerate /* = nullptr */)
 {
 	FActorDescList* ModifiedActorsDescList = nullptr;
 
@@ -880,7 +885,7 @@ bool UWorldPartition::GenerateStreaming(TArray<FString>* OutPackagesToGenerate)
 	FWorldPartitionStreamingGenerator StreamingGenerator(ModifiedActorsDescList, ErrorHandler, IsStreamingEnabled());
 
 	// Preparation Phase
-	StreamingGenerator.PreparationPhase(this);
+	StreamingGenerator.PreparationPhase(InActorDescContainer);
 
 	StreamingGenerator.DumpStateLog(HierarchicalLogAr);
 
@@ -913,7 +918,7 @@ void UWorldPartition::GenerateHLOD(ISourceControlHelper* SourceControlHelper, bo
 {
 	FStreamingGenerationLogErrorHandler LogErrorHandler;
 	FWorldPartitionStreamingGenerator StreamingGenerator(nullptr, &LogErrorHandler, IsStreamingEnabled(), { AWorldPartitionHLOD::StaticClass() });
-	StreamingGenerator.PreparationPhase(this);
+	StreamingGenerator.PreparationPhase(ActorDescContainer);
 
 	TUniquePtr<FArchive> LogFileAr = FWorldPartitionStreamingGenerator::CreateDumpStateLogArchive(TEXT("HLOD"));
 	FHierarchicalLogArchive HierarchicalLogAr(*LogFileAr);
@@ -924,7 +929,7 @@ void UWorldPartition::GenerateHLOD(ISourceControlHelper* SourceControlHelper, bo
 
 void UWorldPartition::CheckForErrors(IStreamingGenerationErrorHandler* ErrorHandler) const
 {
-	CheckForErrors(ErrorHandler, this, IsStreamingEnabled());
+	CheckForErrors(ErrorHandler, ActorDescContainer, IsStreamingEnabled());
 }
 
 void UWorldPartition::CheckForErrors(IStreamingGenerationErrorHandler* ErrorHandler, const UActorDescContainer* ActorDescContainer, bool bEnableStreaming)
diff --git a/Engine/Source/Runtime/Engine/Private/WorldPartitionBlueprintLibrary.cpp b/Engine/Source/Runtime/Engine/Private/WorldPartitionBlueprintLibrary.cpp
--- a/Engine/Source/Runtime/Engine/Private/WorldPartitionBlueprintLibrary.cpp
+++ b/Engine/Source/Runtime/Engine/Private/WorldPartitionBlueprintLibrary.cpp
@@ -66,6 +66,18 @@ void UWorldPartitionBlueprintLibrary::OnWorldPartitionUninitialized(UWorldPartit
 	}
 }
 
+bool UWorldPartitionBlueprintLibrary::GetActorDescs(const UWorldPartition* WorldPartiton, TArray<FActorDesc>& OutActorDescs)
+{
+	bool bResult = true;
+
+	WorldPartiton->ForEachActorDescContainer([&](const UActorDescContainer* ActorDescContainer)
+	{
+		bResult &= GetActorDescs(ActorDescContainer, FTransform::Identity, OutActorDescs);
+	});
+	
+	return bResult;
+}
+
 bool UWorldPartitionBlueprintLibrary::GetActorDescs(const UActorDescContainer* InContainer, const FTransform& InTransform, TArray<FActorDesc>& OutActorDescs)
 {
 	bool bResult = true;
@@ -96,52 +108,54 @@ bool UWorldPartitionBlueprintLibrary::GetActorDescs(const UActorDescContainer* I
 	return bResult;
 }
 
-bool UWorldPartitionBlueprintLibrary::GetIntersectingActorDescs(const UActorDescContainer* InContainer, const FBox& InBox, const FTransform& InTransform, TArray<FActorDesc>& OutActorDescs)
+bool UWorldPartitionBlueprintLibrary::HandleIntersectingActorDesc(const FWorldPartitionActorDesc* ActorDesc, const FBox& InBox, const FTransform& InTransform, TArray<FActorDesc>& OutActorDescs)
 {
 	bool bResult = true;
+	const UActorDescContainer* SubContainer;
+	EContainerClusterMode SubClusterMode;
+	FTransform SubTransform;
 
-	auto HandleActorDesc = [&InBox, &InTransform, &OutActorDescs](const FWorldPartitionActorDesc* ActorDesc)
+	if (ActorDesc->IsContainerInstance())
 	{
-		bool bResult = true;
-		const UActorDescContainer* SubContainer;
-		EContainerClusterMode SubClusterMode;
-		FTransform SubTransform;
-
-		if (ActorDesc->IsContainerInstance())
+		if (ActorDesc->GetContainerInstance(SubContainer, SubTransform, SubClusterMode))
 		{
-			if (ActorDesc->GetContainerInstance(SubContainer, SubTransform, SubClusterMode))
-			{
-				bResult &= GetIntersectingActorDescs(SubContainer, InBox, SubTransform * InTransform, OutActorDescs);
-			}
-			else
-			{
-				bResult = false;
-			}
+			bResult &= GetIntersectingActorDescs(SubContainer, InBox, SubTransform * InTransform, OutActorDescs);
 		}
 		else
 		{
-			OutActorDescs.Emplace(*ActorDesc, InTransform);
+			bResult = false;
 		}
+	}
+	else
+	{
+		OutActorDescs.Emplace(*ActorDesc, InTransform);
+	}
 
-		return bResult;
-	};
+	return bResult;
+}
+
+bool UWorldPartitionBlueprintLibrary::GetIntersectingActorDescs(UWorldPartition* WorldPartition, const FBox& InBox, TArray<FActorDesc>& OutActorDescs)
+{
+	bool bResult = true;
 
-	if (const UWorldPartition* WorldPartition = Cast<UWorldPartition>(InContainer))
+	FWorldPartitionHelpers::ForEachIntersectingActorDesc(WorldPartition, InBox, [&bResult, &InBox, &OutActorDescs](const FWorldPartitionActorDesc* ActorDesc)
 	{
-		FWorldPartitionHelpers::ForEachIntersectingActorDesc(const_cast<UWorldPartition*>(WorldPartition), InBox, [&bResult, &InBox, &InTransform, &OutActorDescs, &HandleActorDesc](const FWorldPartitionActorDesc* ActorDesc)
-		{
-			bResult &= HandleActorDesc(ActorDesc);
-			return true;
-		});
-	}
-	else
+		bResult &= HandleIntersectingActorDesc(ActorDesc, InBox, FTransform::Identity, OutActorDescs);
+		return true;
+	});
+
+	return bResult;
+}
+
+bool UWorldPartitionBlueprintLibrary::GetIntersectingActorDescs(const UActorDescContainer* InContainer, const FBox& InBox, const FTransform& InTransform, TArray<FActorDesc>& OutActorDescs)
+{
+	bool bResult = true;
+
+	for (FActorDescList::TConstIterator<> ActorDescIt(InContainer); ActorDescIt; ++ActorDescIt)
 	{
-		for (FActorDescList::TConstIterator<> ActorDescIt(InContainer); ActorDescIt; ++ActorDescIt)
+		if (ActorDescIt->GetBounds().Intersect(InBox))
 		{
-			if (ActorDescIt->GetBounds().Intersect(InBox))
-			{
-				bResult &= HandleActorDesc(*ActorDescIt);
-			}
+			bResult &= HandleIntersectingActorDesc(*ActorDescIt, InBox, InTransform, OutActorDescs);
 		}
 	}
 
@@ -206,7 +220,7 @@ bool UWorldPartitionBlueprintLibrary::GetActorDescs(TArray<FActorDesc>& OutActor
 #if WITH_EDITOR
 	if (UWorldPartition* WorldPartition = GetWorldPartition())
 	{
-		return GetActorDescs(WorldPartition, FTransform::Identity, OutActorDescs);
+		return GetActorDescs(WorldPartition, OutActorDescs);
 	}
 #endif
 	return false;
@@ -217,7 +231,7 @@ bool UWorldPartitionBlueprintLibrary::GetIntersectingActorDescs(const FBox& InBo
 #if WITH_EDITOR
 	if (UWorldPartition* WorldPartition = GetWorldPartition())
 	{
-		return GetIntersectingActorDescs(WorldPartition, InBox, FTransform::Identity, OutActorDescs);
+		return GetIntersectingActorDescs(WorldPartition, InBox, OutActorDescs);
 	}
 #endif
 	return false;
diff --git a/Engine/Source/Runtime/Engine/Public/WorldPartition/ActorDescContainer.h b/Engine/Source/Runtime/Engine/Public/WorldPartition/ActorDescContainer.h
--- a/Engine/Source/Runtime/Engine/Public/WorldPartition/ActorDescContainer.h
+++ b/Engine/Source/Runtime/Engine/Public/WorldPartition/ActorDescContainer.h
@@ -4,7 +4,6 @@
 #include "CoreMinimal.h"
 #include "UObject/ObjectMacros.h"
 #include "UObject/Object.h"
-#include "GameFramework/Actor.h"
 #include "WorldPartition/ActorDescList.h"
 #include "WorldPartition/WorldPartitionHandle.h"
 #include "ActorDescContainer.generated.h"
@@ -19,28 +18,32 @@ class ENGINE_API UActorDescContainer : public UObject, public FActorDescList
 
 public:
 	void Initialize(UWorld* World, FName InPackageName);
-	virtual void Uninitialize();
+	void Uninitialize();
 	
 	virtual UWorld* GetWorld() const override;
 
 #if WITH_EDITOR
-	virtual void OnObjectPreSave(UObject* Object, FObjectPreSaveContext SaveContext);
-	virtual void OnPackageDeleted(UPackage* Package);
-	virtual void OnObjectsReplaced(const TMap<UObject*, UObject*>& OldToNewObjectMap);
+	bool IsInitialized() const { return bContainerInitialized; }
+
+	void OnObjectPreSave(UObject* Object, FObjectPreSaveContext SaveContext);
+	void OnPackageDeleted(UPackage* Package);
+	void OnObjectsReplaced(const TMap<UObject*, UObject*>& OldToNewObjectMap);
 
 	FName GetContainerPackage() const { return ContainerPackageName; }
+	void SetContainerPackage(const FName& InContainerPackageName) { ContainerPackageName = InContainerPackageName;	}
+	
+	bool IsMainPartitionContainer() const;
+
+	FString GetExternalActorPath() const;
 
 	/** Removes an actor desc without the need to load a package */
-	virtual void RemoveActor(const FGuid& ActorGuid);
+	bool RemoveActor(const FGuid& ActorGuid);
 
 	void LoadAllActors(TArray<FWorldPartitionReference>& OutReferences);
 
-	/** Instancing support */
-	virtual bool GetInstancingContext(const FLinkerInstancingContext*& OutInstancingContext) const { return false; }
+	bool IsActorDescHandled(const AActor* Actor) const;
 #endif
 
-	virtual const FTransform& GetInstanceTransform() const { return FTransform::Identity; }
-
 #if WITH_EDITOR
 public:
 	DECLARE_EVENT_OneParam(UWorldPartition, FActorDescAddedEvent, FWorldPartitionActorDesc*);
@@ -67,14 +70,10 @@ protected:
 	virtual void RemoveActorDescriptor(FWorldPartitionActorDesc* ActorDesc) override;
 	//~ End FActorDescList Interface
 
-	// Events
-	void OnWorldRenamed(UWorld* RenamedWorld);
-	virtual void OnWorldRenamed();
-
-	virtual void OnActorDescAdded(FWorldPartitionActorDesc* NewActorDesc);
-	virtual void OnActorDescRemoved(FWorldPartitionActorDesc* ActorDesc);
-	virtual void OnActorDescUpdating(FWorldPartitionActorDesc* ActorDesc) {}
-	virtual void OnActorDescUpdated(FWorldPartitionActorDesc* ActorDesc) {}
+	void OnActorDescAdded(FWorldPartitionActorDesc* NewActorDesc);
+	void OnActorDescRemoved(FWorldPartitionActorDesc* ActorDesc);
+	void OnActorDescUpdating(FWorldPartitionActorDesc* ActorDesc);
+	void OnActorDescUpdated(FWorldPartitionActorDesc* ActorDesc);
 
 	bool ShouldHandleActorEvent(const AActor* Actor);
 
diff --git a/Engine/Source/Runtime/Engine/Public/WorldPartition/ActorDescContainerCollection.h b/Engine/Source/Runtime/Engine/Public/WorldPartition/ActorDescContainerCollection.h
new file mode 100644
--- /dev/null
+++ b/Engine/Source/Runtime/Engine/Public/WorldPartition/ActorDescContainerCollection.h
@@ -0,0 +1,175 @@
+// Copyright Epic Games, Inc. All Rights Reserved.
+#pragma once
+
+#include "CoreMinimal.h"
+#include "UObject/Object.h"
+#include "WorldPartition/ActorDescContainer.h"
+
+class FWorldPartitionActorDesc;
+class FReferenceCollector;
+
+class ENGINE_API FActorDescContainerCollection
+{
+#if WITH_EDITOR
+public:
+	void AddContainer(UActorDescContainer* Container);
+	bool RemoveContainer(UActorDescContainer* Container);
+	bool Contains(const FName& ContainerPackageName) const;
+	UActorDescContainer* Find(const FName& ContainerPackageName) const;
+
+	void Empty() { ActorDescContainerCollection.Empty(); }
+	bool IsEmpty() const { return ActorDescContainerCollection.IsEmpty(); }
+	uint32 GetActorDescContainerCount() const { return ActorDescContainerCollection.Num(); }
+
+	FWorldPartitionActorDesc* GetActorDesc(const FGuid& Guid);
+	const FWorldPartitionActorDesc* GetActorDesc(const FGuid& Guid) const;
+
+	FWorldPartitionActorDesc& GetActorDescChecked(const FGuid& Guid);
+	const FWorldPartitionActorDesc& GetActorDescChecked(const FGuid& Guid) const;
+
+	const FWorldPartitionActorDesc* GetActorDesc(const FString& PackageName) const;
+	const FWorldPartitionActorDesc* GetActorDesc(const FSoftObjectPath& InActorPath) const;
+
+	bool RemoveActor(const FGuid& ActorGuid);
+
+	void OnPackageDeleted(UPackage* Package);
+
+	bool IsActorDescHandled(const AActor* Actor) const;
+
+	void LoadAllActors(TArray<FWorldPartitionReference>& OutReferences);
+
+	DECLARE_EVENT_OneParam(UWorldPartition, FActorDescAddedEvent, FWorldPartitionActorDesc*);
+	FActorDescAddedEvent OnActorDescAddedEvent;
+
+	DECLARE_EVENT_OneParam(UWorldPartition, FActorDescRemovedEvent, FWorldPartitionActorDesc*);
+	FActorDescRemovedEvent OnActorDescRemovedEvent;
+
+	UActorDescContainer* GetActorDescContainer(const FGuid& ActorGuid);
+	const UActorDescContainer* GetActorDescContainer(const FGuid& ActorGuid) const;
+
+	void ForEachActorDescContainerBreakable(TFunctionRef<bool(UActorDescContainer*)> Func);
+	void ForEachActorDescContainerBreakable(TFunctionRef<bool(UActorDescContainer*)> Func) const;
+
+	void ForEachActorDescContainer(TFunctionRef<void(UActorDescContainer*)> Func);
+	void ForEachActorDescContainer(TFunctionRef<void(UActorDescContainer*)> Func) const;
+
+protected:
+	TArray<UActorDescContainer*> ActorDescContainerCollection;
+
+public:
+	void OnActorDescAdded(FWorldPartitionActorDesc* ActorDesc);
+	void OnActorDescRemoved(FWorldPartitionActorDesc* ActorDesc);
+
+public:
+	template<bool bConst, class ActorType>
+	class TBaseIterator
+	{
+		static_assert(TIsDerivedFrom<ActorType, AActor>::IsDerived, "Type is not derived from AActor.");
+
+	protected:
+		typedef UActorDescContainer ContainerType;
+		typedef TArray<ContainerType*> ContainerCollectionType;
+		typedef typename TChooseClass<bConst, ContainerCollectionType::TConstIterator, ContainerCollectionType::TIterator>::Result ContainerIteratorType;
+		typedef typename TChooseClass<bConst, typename ContainerType::TConstIterator<ActorType>, typename ContainerType::TIterator<ActorType>>::Result ActDescIteratorType;
+
+		typedef typename FWorldPartitionActorDescType<ActorType>::Type ValueType;
+		typedef typename TChooseClass<bConst, const ValueType*, ValueType*>::Result ReturnType;
+
+	public:
+		template<class T>
+		TBaseIterator(T* Collection, UClass* InActorClass)
+			: ContainerIterator(Collection->ActorDescContainerCollection)
+		{
+			if (ContainerIterator)
+			{
+				ActorsIterator = MakeUnique<ActDescIteratorType>(*ContainerIterator, InActorClass);
+			}
+		}
+
+		/**
+		 * Iterates to next suitable actor desc.
+		 */
+		void operator++()
+		{
+			++(*ActorsIterator);
+
+			if (!(*ActorsIterator))
+			{
+				AdvanceToNextRegistryContainer();
+			}
+		}
+
+		/**
+		 * Returns the current suitable actor desc pointed at by the Iterator
+		 *
+		 * @return	Current suitable actor desc
+		 */
+		FORCEINLINE ReturnType operator*() const
+		{
+			return StaticCast<ReturnType>(**ActorsIterator);
+		}
+
+		/**
+		 * Returns the current suitable actor desc pointed at by the Iterator
+		 *
+		 * @return	Current suitable actor desc
+		 */
+		FORCEINLINE ReturnType operator->() const
+		{
+			return StaticCast<ReturnType>(**ActorsIterator);
+		}
+		/**
+		 * Returns whether the iterator has reached the end and no longer points
+		 * to a suitable actor desc.
+		 *
+		 * @return true if iterator points to a suitable actor desc, false if it has reached the end
+		 */
+		FORCEINLINE explicit operator bool() const
+		{
+			return ActorsIterator ? (bool)*ActorsIterator : false;
+		}
+
+	protected:
+		FORCEINLINE void AdvanceToNextRegistryContainer()
+		{
+			if (++ContainerIterator)
+			{
+				ActorsIterator = MakeUnique<ActDescIteratorType>(*ContainerIterator, ActorsIterator->GetActorClass());
+			}
+			else
+			{
+				ActorsIterator.Reset();
+			}
+		}
+
+		ContainerIteratorType ContainerIterator;
+		TUniquePtr<ActDescIteratorType> ActorsIterator;
+	};
+
+	template <class ActorType = AActor>
+	class TIterator : public TBaseIterator<false, ActorType>
+	{
+		typedef TBaseIterator<false, ActorType> BaseType;
+
+	public:
+		template<class T>
+		TIterator(T* Collection, UClass* InActorClass = nullptr)
+			: BaseType(Collection, InActorClass ? InActorClass : ActorType::StaticClass())
+		{}
+	};
+
+	template <class ActorType = AActor>
+	class TConstIterator : public TBaseIterator<true, ActorType>
+	{
+		typedef TBaseIterator<true, ActorType> BaseType;
+
+	public:
+		template<class T>
+		TConstIterator(T* Collection, UClass* InActorClass = nullptr)
+			: BaseType(Collection, InActorClass ? InActorClass : ActorType::StaticClass())
+		{}
+	};
+
+
+#endif
+};
diff --git a/Engine/Source/Runtime/Engine/Public/WorldPartition/ActorDescList.h b/Engine/Source/Runtime/Engine/Public/WorldPartition/ActorDescList.h
--- a/Engine/Source/Runtime/Engine/Public/WorldPartition/ActorDescList.h
+++ b/Engine/Source/Runtime/Engine/Public/WorldPartition/ActorDescList.h
@@ -32,6 +32,7 @@ public:
 
 	int32 GetActorDescCount() const { return ActorsByGuid.Num(); }
 
+	bool IsEmpty() const { return GetActorDescCount() == 0; }
 	void Empty();
 
 	template<bool bConst, class ActorType>
@@ -101,6 +102,13 @@ public:
 			return (bool)ActorsIterator;
 		}
 
+		/**
+		 * Returns the actor class on which the iterator iterates on.
+		 *
+		 * @return the actor class
+		 */
+		FORCEINLINE UClass* GetActorClass() const { return ActorClass; }
+
 	protected:
 		/**
 		 * Determines whether the iterator currently points to a valid actor desc or not.
diff --git a/Engine/Source/Runtime/Engine/Public/WorldPartition/WorldPartition.h b/Engine/Source/Runtime/Engine/Public/WorldPartition/WorldPartition.h
--- a/Engine/Source/Runtime/Engine/Public/WorldPartition/WorldPartition.h
+++ b/Engine/Source/Runtime/Engine/Public/WorldPartition/WorldPartition.h
@@ -13,7 +13,7 @@
 #include "WorldPartition/WorldPartitionActorDesc.h"
 #include "WorldPartition/WorldPartitionStreamingSource.h"
 #include "WorldPartition/WorldPartitionHandle.h"
-#include "WorldPartition/ActorDescContainer.h"
+#include "WorldPartition/ActorDescContainerCollection.h"
 
 #if WITH_EDITOR
 #include "WorldPartition/WorldPartitionActorLoaderInterface.h"
@@ -25,6 +25,7 @@
 #include "WorldPartition.generated.h"
 
 class FWorldPartitionActorDesc;
+class UActorDescContainer;
 class UWorldPartitionEditorHash;
 class UWorldPartitionRuntimeCell;
 class UWorldPartitionRuntimeHash;
@@ -74,7 +75,7 @@ public:
 #endif
 
 UCLASS(AutoExpandCategories=(WorldPartition))
-class ENGINE_API UWorldPartition final : public UActorDescContainer
+class ENGINE_API UWorldPartition final : public UObject, public FActorDescContainerCollection
 {
 	GENERATED_UCLASS_BODY()
 
@@ -84,6 +85,7 @@ class ENGINE_API UWorldPartition final : public UActorDescContainer
 	friend class FWorldPartitionEditorModule;
 	friend class FWorldPartitionDetails;
 	friend class FUnrealEdMisc;
+	friend class UActorDescContainer;
 
 public:
 #if WITH_EDITOR
@@ -122,23 +124,25 @@ private:
 	void OnBeginPlay();
 	void OnEndPlay();
 
-	//~ Begin UActorDescContainer Interface
-	virtual void OnWorldRenamed() override;
+	// WorldDeletegates Events
+	void OnWorldRenamed(UWorld* RenamedWorld);
 
-	virtual void OnActorDescAdded(FWorldPartitionActorDesc* NewActorDesc) override;
-	virtual void OnActorDescRemoved(FWorldPartitionActorDesc* ActorDesc) override;
-	virtual void OnActorDescUpdating(FWorldPartitionActorDesc* ActorDesc) override;
-	virtual void OnActorDescUpdated(FWorldPartitionActorDesc* ActorDesc) override;
+	// ActorDescContainer Events
+	void OnActorDescAdded(FWorldPartitionActorDesc* NewActorDesc);
+	void OnActorDescRemoved(FWorldPartitionActorDesc* ActorDesc);
+	void OnActorDescUpdating(FWorldPartitionActorDesc* ActorDesc);
+	void OnActorDescUpdated(FWorldPartitionActorDesc* ActorDesc);
 
-	virtual bool GetInstancingContext(const FLinkerInstancingContext*& OutInstancingContext) const override;
+	bool GetInstancingContext(const FLinkerInstancingContext*& OutInstancingContext) const;
 #endif
 
 public:
-	virtual const FTransform& GetInstanceTransform() const override;
+	const FTransform& GetInstanceTransform() const;
 	//~ End UActorDescContainer Interface
 
 	//~ Begin UObject Interface
 	virtual void Serialize(FArchive& Ar) override;
+	virtual UWorld* GetWorld() const override;
 	virtual bool ResolveSubobject(const TCHAR* SubObjectPath, UObject*& OutObject, bool bLoadIfExists) override;
 	static void AddReferencedObjects(UObject* InThis, FReferenceCollector& Collector);
 	//~ End UObject Interface
@@ -148,6 +152,7 @@ public:
 
 	// Streaming generation
 	bool GenerateStreaming(TArray<FString>* OutPackagesToGenerate = nullptr);
+	bool GenerateContainerStreaming(const UActorDescContainer* ActorDescContainer, TArray<FString>* OutPackagesToGenerate = nullptr);
 	void FlushStreaming();
 
 	void RemapSoftObjectPath(FSoftObjectPath& ObjectPath);
@@ -175,6 +180,13 @@ public:
 
 	void AppendAssetRegistryTags(TArray<FAssetRegistryTag>& OutTags) const;
 
+	DECLARE_MULTICAST_DELEGATE_OneParam(FActorDescContainerRegistrationDelegate, UActorDescContainer*);
+	FActorDescContainerRegistrationDelegate OnActorDescContainerRegistered;
+	FActorDescContainerRegistrationDelegate OnActorDescContainerUnregistered;
+	UActorDescContainer* RegisterActorDescContainer(const FName& ContainerPackage);
+	bool UnregisterActorDescContainer(UActorDescContainer* Container);
+	void UninitializeActorDescContainers();
+
 	// Actors pinning
 	void PinActors(const TArray<FGuid>& ActorGuids);
 	void UnpinActors(const TArray<FGuid>& ActorGuids);
@@ -193,7 +205,7 @@ public:
 
 	void Initialize(UWorld* World, const FTransform& InTransform);
 	bool IsInitialized() const;
-	virtual void Uninitialize() override;
+	void Uninitialize();
 
 	bool CanStream() const;
 	bool IsMainWorldPartition() const;
@@ -254,9 +266,16 @@ private:
 #endif
 
 public:
+	UActorDescContainer* GetActorDescContainer() const { return ActorDescContainer; }
+
+	UPROPERTY()
+	TObjectPtr<UActorDescContainer> ActorDescContainer;
+
 	UPROPERTY()
 	TObjectPtr<UWorldPartitionRuntimeHash> RuntimeHash;
 
+	UPROPERTY(Transient)
+	TObjectPtr<UWorld> World;
 private:
 #if WITH_EDITOR
 	bool bForceGarbageCollection;
@@ -298,6 +317,8 @@ private:
 #if WITH_EDITOR
 	void HashActorDesc(FWorldPartitionActorDesc* ActorDesc);
 	void UnhashActorDesc(FWorldPartitionActorDesc* ActorDesc);
+	void HashActorDescContainer(UActorDescContainer* ActorDescContainer);
+	void UnhashActorDescContainer(UActorDescContainer* ActorDescContainer);
 
 public:
 	// Editor loader adapters management
diff --git a/Engine/Source/Runtime/Engine/Public/WorldPartition/WorldPartitionActorDesc.h b/Engine/Source/Runtime/Engine/Public/WorldPartition/WorldPartitionActorDesc.h
--- a/Engine/Source/Runtime/Engine/Public/WorldPartition/WorldPartitionActorDesc.h
+++ b/Engine/Source/Runtime/Engine/Public/WorldPartition/WorldPartitionActorDesc.h
@@ -63,6 +63,7 @@ class ENGINE_API FWorldPartitionActorDesc
 {
 	friend class AActor;
 	friend class UWorldPartition;
+	friend class FActorDescContainerCollection;
 	friend struct FWorldPartitionHandleImpl;
 	friend struct FWorldPartitionReferenceImpl;
 	friend struct FWorldPartitionActorDescUtils;
@@ -263,4 +264,4 @@ protected:
 
 	static TMap<TSubclassOf<AActor>, FActorDescDeprecator> Deprecators;
 };
-#endif
\ No newline at end of file
+#endif
diff --git a/Engine/Source/Runtime/Engine/Public/WorldPartition/WorldPartitionActorLoaderInterface.h b/Engine/Source/Runtime/Engine/Public/WorldPartition/WorldPartitionActorLoaderInterface.h
--- a/Engine/Source/Runtime/Engine/Public/WorldPartition/WorldPartitionActorLoaderInterface.h
+++ b/Engine/Source/Runtime/Engine/Public/WorldPartition/WorldPartitionActorLoaderInterface.h
@@ -41,6 +41,9 @@ public:
 		virtual TOptional<FString> GetLabel() const { return TOptional<FString>(); }
 		virtual TOptional<FColor> GetColor() const { return TOptional<FColor>(); }
 
+		void OnActorDescContainerInitialize(UActorDescContainer* Container);
+		void OnActorDescContainerUninitialize(UActorDescContainer* Container);
+
 	protected:
 		// Private interface
 		virtual void ForEachActor(TFunctionRef<void(const FWorldPartitionHandle&)> InOperation) const =0;
diff --git a/Engine/Source/Runtime/Engine/Public/WorldPartition/WorldPartitionBlueprintLibrary.h b/Engine/Source/Runtime/Engine/Public/WorldPartition/WorldPartitionBlueprintLibrary.h
--- a/Engine/Source/Runtime/Engine/Public/WorldPartition/WorldPartitionBlueprintLibrary.h
+++ b/Engine/Source/Runtime/Engine/Public/WorldPartition/WorldPartitionBlueprintLibrary.h
@@ -71,8 +71,11 @@ class UWorldPartitionBlueprintLibrary : public UBlueprintFunctionLibrary
 	static TMap<UWorldPartition*, TUniquePtr<FLoaderAdapterActorList>> LoaderAdapterActorListMap;
 	static FDelegateHandle OnWorldPartitionUninitializedHandle;
 
+	static bool GetActorDescs(const UWorldPartition* WorldPartition, TArray<FActorDesc>& OutActorDescs);
 	static bool GetActorDescs(const UActorDescContainer* InContainer, const FTransform& InTransform, TArray<FActorDesc>& OutActorDescs);
+	static bool GetIntersectingActorDescs(UWorldPartition* WorldPartition, const FBox& InBox, TArray<FActorDesc>& OutActorDescs);
 	static bool GetIntersectingActorDescs(const UActorDescContainer* InContainer, const FBox& InBox, const FTransform& InTransform, TArray<FActorDesc>& OutActorDescs);
+	static bool HandleIntersectingActorDesc(const FWorldPartitionActorDesc* ActorDesc, const FBox& InBox, const FTransform& InTransform, TArray<FActorDesc>& OutActorDescs);
 #endif
 
 public:
diff --git a/Engine/Source/Runtime/Engine/Public/WorldPartition/WorldPartitionHandle.h b/Engine/Source/Runtime/Engine/Public/WorldPartition/WorldPartitionHandle.h
--- a/Engine/Source/Runtime/Engine/Public/WorldPartition/WorldPartitionHandle.h
+++ b/Engine/Source/Runtime/Engine/Public/WorldPartition/WorldPartitionHandle.h
@@ -8,6 +8,7 @@
 
 #if WITH_EDITOR
 class UActorDescContainer;
+class FActorDescContainerCollection;
 class FWorldPartitionActorDesc;
 struct FWorldPartitionHandleImpl;
 struct FWorldPartitionReferenceImpl;
@@ -96,9 +97,9 @@ public:
 		}
 	}
 
-	FORCEINLINE TWorldPartitionHandle(UActorDescContainer* Container, const FGuid& ActorGuid)
-		: Container(Container)
-		, ActorDesc(Impl::GetActorDesc(Container, ActorGuid))
+	FORCEINLINE TWorldPartitionHandle(UActorDescContainer* InContainer, const FGuid& ActorGuid)
+		: Container(InContainer)
+		, ActorDesc(Impl::GetActorDesc(InContainer, ActorGuid))
 	{
 		if (IsValid())
 		{
@@ -106,6 +107,21 @@ public:
 		}
 	}
 
+	FORCEINLINE TWorldPartitionHandle(FActorDescContainerCollection* ContainerCollection, const FGuid& ActorGuid)
+		: Container(Impl::GetActorDescContainer(ContainerCollection, ActorGuid))
+		, ActorDesc(nullptr)
+	{
+		if (Container != nullptr)
+		{
+			ActorDesc = Impl::GetActorDesc(Container.Get(), ActorGuid);
+		}
+
+		if (IsValid())
+		{
+			IncRefCount();
+		}
+	}
+
 	FORCEINLINE TWorldPartitionHandle(const TWorldPartitionHandle& Other)
 		: Container(nullptr)
 		, ActorDesc(nullptr)
@@ -306,6 +322,7 @@ struct ENGINE_API FWorldPartitionImplBase
 {
 	static TUniquePtr<FWorldPartitionActorDesc>* GetActorDesc(UActorDescContainer* Container, const FGuid& ActorGuid);
 	static UActorDescContainer* GetActorDescContainer(TUniquePtr<FWorldPartitionActorDesc>* ActorDesc);
+	static UActorDescContainer* GetActorDescContainer(FActorDescContainerCollection* ContainerCollection, const FGuid& ActorGuid);
 	static bool IsActorDescLoaded(FWorldPartitionActorDesc* ActorDesc);
 };
 
diff --git a/Engine/Source/Runtime/Engine/Public/WorldPartition/WorldPartitionRuntimeCell.h b/Engine/Source/Runtime/Engine/Public/WorldPartition/WorldPartitionRuntimeCell.h
--- a/Engine/Source/Runtime/Engine/Public/WorldPartition/WorldPartitionRuntimeCell.h
+++ b/Engine/Source/Runtime/Engine/Public/WorldPartition/WorldPartitionRuntimeCell.h
@@ -38,13 +38,14 @@ struct FWorldPartitionRuntimeCellObjectMapping
 #endif
 	{}
 
-	FWorldPartitionRuntimeCellObjectMapping(FName InPackage, FName InPath, const FActorContainerID& InContainerID, const FTransform& InContainerTransform, FName InContainerPackage)
+	FWorldPartitionRuntimeCellObjectMapping(FName InPackage, FName InPath, const FActorContainerID& InContainerID, const FTransform& InContainerTransform, FName InContainerPackage, FName InWorldPackage)
 #if WITH_EDITORONLY_DATA
 		: Package(InPackage)
 		, Path(InPath)
 		, ContainerID(InContainerID)
 		, ContainerTransform(InContainerTransform)
 		, ContainerPackage(InContainerPackage)
+		, WorldPackage(InWorldPackage)
 		, LoadedPath(InPath)
 #endif
 	{}
@@ -80,6 +81,12 @@ struct FWorldPartitionRuntimeCellObjectMapping
 	UPROPERTY()
 	FName ContainerPackage;
 
+	/**
+	 * Package of the World 
+	 */
+	UPROPERTY()
+	FName WorldPackage;
+
 	/**
 	* Loaded actor path (when cooking or pie)
 	* 
diff --git a/Engine/Source/Runtime/Landscape/Private/LandscapeEdit.cpp b/Engine/Source/Runtime/Landscape/Private/LandscapeEdit.cpp
--- a/Engine/Source/Runtime/Landscape/Private/LandscapeEdit.cpp
+++ b/Engine/Source/Runtime/Landscape/Private/LandscapeEdit.cpp
@@ -70,6 +70,7 @@ LandscapeEdit.cpp: Landscape editing
 #include "ScopedTransaction.h"
 #include "Editor.h"
 #include "PhysicalMaterials/PhysicalMaterial.h"
+#include "WorldPartition/ActorDescContainer.h"
 #include "WorldPartition/WorldPartition.h"
 #include "WorldPartition/WorldPartitionHandle.h"
 #include "WorldPartition/WorldPartitionHelpers.h"
