UMaterialInstance::ValidateTextureOverrides() should not validate VT usage when VT support is disabled

diff --git a/Engine/Source/Runtime/Engine/Private/Materials/MaterialInstance.cpp b/Engine/Source/Runtime/Engine/Private/Materials/MaterialInstance.cpp
--- a/Engine/Source/Runtime/Engine/Private/Materials/MaterialInstance.cpp
+++ b/Engine/Source/Runtime/Engine/Private/Materials/MaterialInstance.cpp
@@ -1213,6 +1213,7 @@ void UMaterialInstance::ValidateTextureOverrides(ERHIFeatureLevel::Type InFeatur
 
 	const UMaterial* Material = GetMaterial();
 	const FMaterialResource* CurrentResource = Material->GetMaterialResource(InFeatureLevel);
+	const bool bShouldValidateVTUsage = UseVirtualTexturing(GMaxRHIFeatureLevel);
 
 	if (!CurrentResource)
 	{
@@ -1239,7 +1240,7 @@ void UMaterialInstance::ValidateTextureOverrides(ERHIFeatureLevel::Type InFeatur
 					{
 						UE_LOG(LogMaterial, Error, TEXT("MaterialInstance \"%s\" parameter '%s' assigned texture \"%s\" has invalid type, required 2D texture"), *MaterialName, *TextureInfo.GetParameterName().ToString(), *Texture->GetName());
 					}
-					else if (TextureType & MCT_TextureVirtual)
+					else if (bShouldValidateVTUsage && (TextureType & MCT_TextureVirtual))
 					{
 						UE_LOG(LogMaterial, Error, TEXT("MaterialInstance \"%s\" parameter '%s' assigned texture \"%s\" requires non-virtual texture"), *MaterialName, *TextureInfo.GetParameterName().ToString(), *Texture->GetName());
 					}
@@ -1273,7 +1274,7 @@ void UMaterialInstance::ValidateTextureOverrides(ERHIFeatureLevel::Type InFeatur
 					{
 						UE_LOG(LogMaterial, Error, TEXT("MaterialInstance \"%s\" parameter '%s' assigned texture \"%s\" has invalid type, required 2D texture"), *MaterialName, *TextureInfo.GetParameterName().ToString(), *Texture->GetName());
 					}
-					else if (!(TextureType & MCT_TextureVirtual))
+					else if (bShouldValidateVTUsage && !(TextureType & MCT_TextureVirtual))
 					{
 						UE_LOG(LogMaterial, Error, TEXT("MaterialInstance \"%s\" parameter '%s' assigned texture \"%s\" requires virtual texture"), *MaterialName, *TextureInfo.GetParameterName().ToString(), *Texture->GetName());
 					}
