Property retirement list only needs to be sized to the number of custom delta properties, not the full rep layout

diff --git a/Engine/Source/Runtime/Engine/Private/DataReplication.cpp b/Engine/Source/Runtime/Engine/Private/DataReplication.cpp
--- a/Engine/Source/Runtime/Engine/Private/DataReplication.cpp
+++ b/Engine/Source/Runtime/Engine/Private/DataReplication.cpp
@@ -668,7 +668,11 @@ void FObjectReplicator::StartReplicating(class UActorChannel * InActorChannel)
 				// Allocate retirement list.
 				// SetNum now constructs, so this is safe
 
-				SendingRepState->Retirement.SetNum(ObjectClass->ClassReps.Num());
+				check(RepLayout);
+				const FRepLayout& LocalRepLayout = *RepLayout;
+				const int32 NumLifetimeCustomDeltaProperties = FNetSerializeCB::GetNumLifetimeCustomDeltaProperties(LocalRepLayout);
+
+				SendingRepState->Retirement.SetNum(NumLifetimeCustomDeltaProperties);
 
 				if (OwningChannel->SpawnAcked)
 				{
