Attempted fix for SoundWave.h small memory incrase in size.

diff --git a/Engine/Source/Runtime/Engine/Private/SampleBufferIO.cpp b/Engine/Source/Runtime/Engine/Private/SampleBufferIO.cpp
--- a/Engine/Source/Runtime/Engine/Private/SampleBufferIO.cpp
+++ b/Engine/Source/Runtime/Engine/Private/SampleBufferIO.cpp
@@ -485,7 +485,12 @@ namespace Audio
 
 		// Emplace wav data in the RawData component of the sound wave.
 		FSharedBuffer Buffer = FSharedBuffer::Clone(SerializedWavData.GetData(), SerializedWavData.Num());
+
+#if WITH_EDITORONLY_DATA
 		CurrentSoundWave->RawData.UpdatePayload(Buffer);
+#else //WITH_EDITORONLY_DATA
+		checkNoEntry();
+#endif //WITH_EDITORONLY_DATA
 
 		USoundWave* SavedSoundWave = CurrentSoundWave;
 
