Fix ITokenizedMessageErrorHandler::OnLevelInstanceInvalidWorldAsset to output a complete asset path so goto content browser works.

#rnx

diff --git a/Engine/Source/Runtime/Engine/Private/WorldPartition/ErrorHandling/WorldPartitionStreamingGenerationTokenizedMessageErrorHandler.cpp b/Engine/Source/Runtime/Engine/Private/WorldPartition/ErrorHandling/WorldPartitionStreamingGenerationTokenizedMessageErrorHandler.cpp
--- a/Engine/Source/Runtime/Engine/Private/WorldPartition/ErrorHandling/WorldPartitionStreamingGenerationTokenizedMessageErrorHandler.cpp
+++ b/Engine/Source/Runtime/Engine/Private/WorldPartition/ErrorHandling/WorldPartitionStreamingGenerationTokenizedMessageErrorHandler.cpp
@@ -164,11 +164,14 @@ void ITokenizedMessageErrorHandler::OnLevelInstanceInvalidWorldAsset(const FWorl
 		break;
 	};
 
+	const FString WorldAssetStr = WorldAsset.ToString();
+	const FString WorldAssetPath = WorldAssetStr + TEXT(".") + FPackageName::GetShortName(WorldAssetStr);
+
 	TSharedRef<FTokenizedMessage> Message = FTokenizedMessage::Create(MessageSeverity);
 	Message->AddToken(FTextToken::Create(LOCTEXT("TokenMessage_WorldPartition_LevelInstance", "Level Instance")))
 		->AddToken(FTextToken::Create(FText::FromString(GetFullActorName(ActorDescView))))
 		->AddToken(FTextToken::Create(ReasonText))
-		->AddToken(FAssetNameToken::Create(WorldAsset.ToString()));
+		->AddToken(FAssetNameToken::Create(WorldAssetPath));
 
 	AddAdditionalNameToken(Message, FName(TEXT("WorldPartition_LevelInstanceInvalidWorldAsset_CheckForErrors")));
 
