- Fix for bad change to handle cooked content in editor

#rnx

diff --git a/Engine/Source/Runtime/Engine/Private/Particles/ParticleModules.cpp b/Engine/Source/Runtime/Engine/Private/Particles/ParticleModules.cpp
--- a/Engine/Source/Runtime/Engine/Private/Particles/ParticleModules.cpp
+++ b/Engine/Source/Runtime/Engine/Private/Particles/ParticleModules.cpp
@@ -4567,320 +4567,317 @@ void UParticleModuleTypeDataGpu::BeginDestroy()
 void UParticleModuleTypeDataGpu::Build( FParticleEmitterBuildInfo& EmitterBuildInfo )
 {
 #if WITH_EDITOR
-	if (GetOutermost()->bIsCookedForEditor)
-	{
-		return;
-	}
+	if (!GetOutermost()->bIsCookedForEditor)
+	{
+		FVector4Distribution Curve;
+		FComposableFloatDistribution ZeroDistribution;
+		FComposableFloatDistribution OneDistribution;
+		FVectorDistribution VectorDistribution;
+		FVector MinValue,MaxValue;
+		ZeroDistribution.InitializeWithConstant(0.0f);
+		OneDistribution.InitializeWithConstant(1.0f);
+
+		// Store off modules and properties required for simulation.
+		EmitterInfo.RequiredModule = EmitterBuildInfo.RequiredModule;
+		EmitterInfo.SpawnModule = EmitterBuildInfo.SpawnModule;
+		EmitterInfo.SpawnPerUnitModule = EmitterBuildInfo.SpawnPerUnitModule;
+		EmitterInfo.SpawnModules = EmitterBuildInfo.SpawnModules;
+
+		// Store the inverse of max size.
+		EmitterInfo.InvMaxSize.X = EmitterBuildInfo.MaxSize.X > UE_KINDA_SMALL_NUMBER ? (1.0f / EmitterBuildInfo.MaxSize.X) : 1.0f;
+		EmitterInfo.InvMaxSize.Y = EmitterBuildInfo.MaxSize.Y > UE_KINDA_SMALL_NUMBER ? (1.0f / EmitterBuildInfo.MaxSize.Y) : 1.0f;
+
+		// Compute the value by which to scale rotation rate.
+		const float RotationRateScale = EmitterBuildInfo.MaxRotationRate * EmitterBuildInfo.MaxLifetime;
+
+		// Store the maximum rotation rate (make sure it is never zero).
+		EmitterInfo.InvRotationRateScale = (RotationRateScale > UE_KINDA_SMALL_NUMBER || RotationRateScale < -UE_KINDA_SMALL_NUMBER) ?
+			1.0f / RotationRateScale : 1.0f;
+
+		// A particle's initial size is stored as 1 / MaxSize, so scale by MaxSize.
+		EmitterBuildInfo.SizeScale.ScaleByConstantVector( FVector( EmitterBuildInfo.MaxSize.X, EmitterBuildInfo.MaxSize.Y, 0.0f ) );
+
+		// Build and store the color curve.
+		EmitterBuildInfo.ColorScale.Resample(0.0f, 1.0f);
+		EmitterBuildInfo.AlphaScale.Resample(0.0f, 1.0f);
+		FComposableDistribution::BuildVector4(
+			Curve,
+			EmitterBuildInfo.ColorScale,
+			EmitterBuildInfo.AlphaScale );
+		FComposableDistribution::QuantizeVector4(
+			ResourceData.QuantizedColorSamples,
+			ResourceData.ColorScale,
+			ResourceData.ColorBias,
+			Curve );
+
+		// The misc curve is laid out as: R:SizeX G:SizeY B:SubImageIndex A:Unused.
+		EmitterBuildInfo.SizeScale.Resample(0.0f, 1.0f);
+		EmitterBuildInfo.SubImageIndex.Resample(0.0f, 1.0f);
+		FComposableDistribution::BuildVector4(
+			Curve, 
+			EmitterBuildInfo.SizeScale, 
+			EmitterBuildInfo.SubImageIndex, 
+			ZeroDistribution );
+		FComposableDistribution::QuantizeVector4(
+			ResourceData.QuantizedMiscSamples,
+			ResourceData.MiscScale,
+			ResourceData.MiscBias,
+			Curve );
+
+		// Resilience.
+		bool bBounceOnCollision = EmitterBuildInfo.CollisionResponse == EParticleCollisionResponse::Bounce;
+		FComposableFloatDistribution NormalizedResilience(
+			bBounceOnCollision ? EmitterBuildInfo.Resilience : ZeroDistribution
+			);
+		NormalizedResilience.Normalize(&ResourceData.ResilienceScale, &ResourceData.ResilienceBias);
+		FComposableDistribution::BuildFloat(EmitterInfo.Resilience, NormalizedResilience);
+
+		// The simulation attributes curve is: R:DragScale G:VelocityFieldScale B:Resilience A:OrbitRandom.
+		EmitterBuildInfo.VectorFieldScaleOverLife.Resample(0.0f, 1.0f);
+		EmitterBuildInfo.DragScale.Resample(0.0f, 1.0f);
+		EmitterBuildInfo.ResilienceScaleOverLife.Resample(0.0f, 1.0f);
+		FComposableDistribution::BuildVector4(
+			Curve,
+			EmitterBuildInfo.DragScale,
+			EmitterBuildInfo.VectorFieldScaleOverLife,
+			EmitterBuildInfo.ResilienceScaleOverLife,
+			OneDistribution );
+		FComposableDistribution::QuantizeVector4(
+			ResourceData.QuantizedSimulationAttrSamples,
+			ResourceData.SimulationAttrCurveScale,
+			ResourceData.SimulationAttrCurveBias,
+			Curve );
+
+		// Friction used during collision.
+		if (bBounceOnCollision)
+		{
+			ResourceData.OneMinusFriction = 1.0f - EmitterBuildInfo.Friction;
+			ResourceData.CollisionRandomSpread = EmitterBuildInfo.CollisionRandomSpread;
+			ResourceData.CollisionRandomDistribution = EmitterBuildInfo.CollisionRandomDistribution;
+		}
+		else
+		{
+			ResourceData.OneMinusFriction = 0.0f;
+			ResourceData.CollisionRandomSpread = 0.0f;
+			ResourceData.CollisionRandomDistribution = 1.0f;
+		}
 
-	FVector4Distribution Curve;
-	FComposableFloatDistribution ZeroDistribution;
-	FComposableFloatDistribution OneDistribution;
-	FVectorDistribution VectorDistribution;
-	FVector MinValue,MaxValue;
-	ZeroDistribution.InitializeWithConstant(0.0f);
-	OneDistribution.InitializeWithConstant(1.0f);
-
-	// Store off modules and properties required for simulation.
-	EmitterInfo.RequiredModule = EmitterBuildInfo.RequiredModule;
-	EmitterInfo.SpawnModule = EmitterBuildInfo.SpawnModule;
-	EmitterInfo.SpawnPerUnitModule = EmitterBuildInfo.SpawnPerUnitModule;
-	EmitterInfo.SpawnModules = EmitterBuildInfo.SpawnModules;
-
-	// Store the inverse of max size.
-	EmitterInfo.InvMaxSize.X = EmitterBuildInfo.MaxSize.X > UE_KINDA_SMALL_NUMBER ? (1.0f / EmitterBuildInfo.MaxSize.X) : 1.0f;
-	EmitterInfo.InvMaxSize.Y = EmitterBuildInfo.MaxSize.Y > UE_KINDA_SMALL_NUMBER ? (1.0f / EmitterBuildInfo.MaxSize.Y) : 1.0f;
-
-	// Compute the value by which to scale rotation rate.
-	const float RotationRateScale = EmitterBuildInfo.MaxRotationRate * EmitterBuildInfo.MaxLifetime;
-
-	// Store the maximum rotation rate (make sure it is never zero).
-	EmitterInfo.InvRotationRateScale = (RotationRateScale > UE_KINDA_SMALL_NUMBER || RotationRateScale < -UE_KINDA_SMALL_NUMBER) ?
-		1.0f / RotationRateScale : 1.0f;
-
-	// A particle's initial size is stored as 1 / MaxSize, so scale by MaxSize.
-	EmitterBuildInfo.SizeScale.ScaleByConstantVector( FVector( EmitterBuildInfo.MaxSize.X, EmitterBuildInfo.MaxSize.Y, 0.0f ) );
-
-	// Build and store the color curve.
-	EmitterBuildInfo.ColorScale.Resample(0.0f, 1.0f);
-	EmitterBuildInfo.AlphaScale.Resample(0.0f, 1.0f);
-	FComposableDistribution::BuildVector4(
-		Curve,
-		EmitterBuildInfo.ColorScale,
-		EmitterBuildInfo.AlphaScale );
-	FComposableDistribution::QuantizeVector4(
-		ResourceData.QuantizedColorSamples,
-		ResourceData.ColorScale,
-		ResourceData.ColorBias,
-		Curve );
-
-	// The misc curve is laid out as: R:SizeX G:SizeY B:SubImageIndex A:Unused.
-	EmitterBuildInfo.SizeScale.Resample(0.0f, 1.0f);
-	EmitterBuildInfo.SubImageIndex.Resample(0.0f, 1.0f);
-	FComposableDistribution::BuildVector4(
-		Curve, 
-		EmitterBuildInfo.SizeScale, 
-		EmitterBuildInfo.SubImageIndex, 
-		ZeroDistribution );
-	FComposableDistribution::QuantizeVector4(
-		ResourceData.QuantizedMiscSamples,
-		ResourceData.MiscScale,
-		ResourceData.MiscBias,
-		Curve );
-
-	// Resilience.
-	bool bBounceOnCollision = EmitterBuildInfo.CollisionResponse == EParticleCollisionResponse::Bounce;
-	FComposableFloatDistribution NormalizedResilience(
-		bBounceOnCollision ? EmitterBuildInfo.Resilience : ZeroDistribution
-		);
-	NormalizedResilience.Normalize(&ResourceData.ResilienceScale, &ResourceData.ResilienceBias);
-	FComposableDistribution::BuildFloat(EmitterInfo.Resilience, NormalizedResilience);
-
-	// The simulation attributes curve is: R:DragScale G:VelocityFieldScale B:Resilience A:OrbitRandom.
-	EmitterBuildInfo.VectorFieldScaleOverLife.Resample(0.0f, 1.0f);
-	EmitterBuildInfo.DragScale.Resample(0.0f, 1.0f);
-	EmitterBuildInfo.ResilienceScaleOverLife.Resample(0.0f, 1.0f);
-	FComposableDistribution::BuildVector4(
-		Curve,
-		EmitterBuildInfo.DragScale,
-		EmitterBuildInfo.VectorFieldScaleOverLife,
-		EmitterBuildInfo.ResilienceScaleOverLife,
-		OneDistribution );
-	FComposableDistribution::QuantizeVector4(
-		ResourceData.QuantizedSimulationAttrSamples,
-		ResourceData.SimulationAttrCurveScale,
-		ResourceData.SimulationAttrCurveBias,
-		Curve );
-
-	// Friction used during collision.
-	if (bBounceOnCollision)
-	{
-		ResourceData.OneMinusFriction = 1.0f - EmitterBuildInfo.Friction;
-		ResourceData.CollisionRandomSpread = EmitterBuildInfo.CollisionRandomSpread;
-		ResourceData.CollisionRandomDistribution = EmitterBuildInfo.CollisionRandomDistribution;
-	}
-	else
-	{
-		ResourceData.OneMinusFriction = 0.0f;
-		ResourceData.CollisionRandomSpread = 0.0f;
-		ResourceData.CollisionRandomDistribution = 1.0f;
-	}
+		// Collision time bias, used to kill particles on collision if desired.
+		if (EmitterBuildInfo.CollisionResponse == EParticleCollisionResponse::Kill)
+		{
+			// By adding 1.1 to relative time it will kill the particle.
+			ResourceData.CollisionTimeBias = 1.1f;
+		}
+		else
+		{
+			ResourceData.CollisionTimeBias = 0.0f;
+		}
 
-	// Collision time bias, used to kill particles on collision if desired.
-	if (EmitterBuildInfo.CollisionResponse == EParticleCollisionResponse::Kill)
-	{
-		// By adding 1.1 to relative time it will kill the particle.
-		ResourceData.CollisionTimeBias = 1.1f;
-	}
-	else
-	{
-		ResourceData.CollisionTimeBias = 0.0f;
-	}
+		// Parameters used to derive the collision radius from the size of the sprite.
+		// Note that the sprite size is the diameter, so bake a 1/2 in to the radius
+		// scale to convert to radius.
+		ResourceData.CollisionRadiusScale = EmitterBuildInfo.CollisionRadiusScale * 0.5f;
+		ResourceData.CollisionRadiusBias = EmitterBuildInfo.CollisionRadiusBias;
 
-	// Parameters used to derive the collision radius from the size of the sprite.
-	// Note that the sprite size is the diameter, so bake a 1/2 in to the radius
-	// scale to convert to radius.
-	ResourceData.CollisionRadiusScale = EmitterBuildInfo.CollisionRadiusScale * 0.5f;
-	ResourceData.CollisionRadiusBias = EmitterBuildInfo.CollisionRadiusBias;
+		// If appropriate, set up the sub-image size parameter.
+		EParticleSubUVInterpMethod InterpMethod = EmitterBuildInfo.RequiredModule->InterpolationMethod;
+		if ( InterpMethod == PSUVIM_Linear || InterpMethod == PSUVIM_Linear_Blend )
+		{
+			ResourceData.SubImageSize.X = EmitterBuildInfo.RequiredModule->SubImages_Horizontal;
+			ResourceData.SubImageSize.Y = EmitterBuildInfo.RequiredModule->SubImages_Vertical;
+			ResourceData.SubImageSize.Z = 1.0f / ResourceData.SubImageSize.X;
+			ResourceData.SubImageSize.W = 1.0f / ResourceData.SubImageSize.Y;
+		}
+		else
+		{
+			ResourceData.SubImageSize = FVector4( 1.0f, 1.0f, 1.0f, 1.0f );
+		}
 
-	// If appropriate, set up the sub-image size parameter.
-	EParticleSubUVInterpMethod InterpMethod = EmitterBuildInfo.RequiredModule->InterpolationMethod;
-	if ( InterpMethod == PSUVIM_Linear || InterpMethod == PSUVIM_Linear_Blend )
-	{
-		ResourceData.SubImageSize.X = EmitterBuildInfo.RequiredModule->SubImages_Horizontal;
-		ResourceData.SubImageSize.Y = EmitterBuildInfo.RequiredModule->SubImages_Vertical;
-		ResourceData.SubImageSize.Z = 1.0f / ResourceData.SubImageSize.X;
-		ResourceData.SubImageSize.W = 1.0f / ResourceData.SubImageSize.Y;
-	}
-	else
-	{
-		ResourceData.SubImageSize = FVector4( 1.0f, 1.0f, 1.0f, 1.0f );
-	}
-
-	// Store the size-by-speed parameters.
-	ResourceData.SizeBySpeed.X = FMath::Max<float>(EmitterBuildInfo.SizeScaleBySpeed.X, 0.0f);
-	ResourceData.SizeBySpeed.Y = FMath::Max<float>(EmitterBuildInfo.SizeScaleBySpeed.Y, 0.0f);
-	ResourceData.SizeBySpeed.Z = FMath::Max<float>(EmitterBuildInfo.MaxSizeScaleBySpeed.X, 0.0f);
-	ResourceData.SizeBySpeed.W = FMath::Max<float>(EmitterBuildInfo.MaxSizeScaleBySpeed.Y, 0.0f);
-
-	// Point attractor.
-	{
-		float RadiusSq = EmitterBuildInfo.PointAttractorRadius *
-			EmitterBuildInfo.PointAttractorRadius;
-		EmitterBuildInfo.PointAttractorStrength.ScaleByConstant(RadiusSq);
-		FComposableDistribution::BuildFloat(
-			EmitterInfo.PointAttractorStrength,
-			EmitterBuildInfo.PointAttractorStrength );
-		EmitterInfo.PointAttractorPosition = EmitterBuildInfo.PointAttractorPosition;
-		EmitterInfo.PointAttractorRadiusSq = RadiusSq;
-	}
-
-	// Store the constant acceleration to apply to particles.
-	ResourceData.ConstantAcceleration = EmitterBuildInfo.ConstantAcceleration;
-	EmitterInfo.ConstantAcceleration = EmitterBuildInfo.ConstantAcceleration;
-
-	// Compute the orbit offset amount.
-	FComposableDistribution::BuildVector(VectorDistribution, EmitterBuildInfo.OrbitOffset);
-	VectorDistribution.GetRange(&MinValue, &MaxValue);
-
-	// One half required due to integration in the shader.
-	MinValue *= 0.5f;
-	MaxValue *= 0.5f;
-
-	// Store the orbit offset range.
-	ResourceData.OrbitOffsetBase = MinValue;
-	ResourceData.OrbitOffsetRange = MaxValue - MinValue;
-
-	// Compute the orbit frequencies.
-	FComposableDistribution::BuildVector(VectorDistribution, EmitterBuildInfo.OrbitRotationRate);
-	VectorDistribution.GetRange(&MinValue, &MaxValue);
-
-	// # rotations to radians. Flip Z to be consistent with CPU orbit.
-	MinValue *= (2.0f * UE_PI);
-	MaxValue *= (2.0f * UE_PI);
-	MinValue.Z *= -1.0f;
-	MaxValue.Z *= -1.0f;
-
-	// Store the orbit frequency range.
-	ResourceData.OrbitFrequencyBase = MinValue;
-	ResourceData.OrbitFrequencyRange = MaxValue - MinValue;
-
-	// Compute the orbit phase.
-	FComposableDistribution::BuildVector(VectorDistribution, EmitterBuildInfo.OrbitInitialRotation);
-	VectorDistribution.GetRange(&MinValue, &MaxValue);
-
-	// # rotations to radians. Flip Z to be consistent with CPU orbit.
-	MinValue *= (2.0f * UE_PI);
-	MaxValue *= (2.0f * UE_PI);
-	MinValue.Z *= -1.0f;
-	MaxValue.Z *= -1.0f;
-
-	// Store the orbit phase range.
-	ResourceData.OrbitPhaseBase = MinValue;
-	ResourceData.OrbitPhaseRange = MaxValue - MinValue;
-
-	// Determine around which axes particles are orbiting.
-	const float OrbitX = (ResourceData.OrbitFrequencyBase.X != 0.0f || ResourceData.OrbitFrequencyRange.X != 0.0f || ResourceData.OrbitPhaseBase.X != 0.0f || ResourceData.OrbitPhaseRange.X != 0.0f) ? 1.0f : 0.0f;
-	const float OrbitY = (ResourceData.OrbitFrequencyBase.Y != 0.0f || ResourceData.OrbitFrequencyRange.Y != 0.0f || ResourceData.OrbitPhaseBase.Y != 0.0f || ResourceData.OrbitPhaseRange.Y != 0.0f) ? 1.0f : 0.0f;
-	const float OrbitZ = (ResourceData.OrbitFrequencyBase.Z != 0.0f || ResourceData.OrbitFrequencyRange.Z != 0.0f || ResourceData.OrbitPhaseBase.Z != 0.0f || ResourceData.OrbitPhaseRange.Z != 0.0f) ? 1.0f : 0.0f;
-
-	// Make some adjustments to mimic CPU orbit as much as possible.
-	if (OrbitX != 0.0f)
-	{
-		ResourceData.OrbitPhaseBase.X += (0.5f * UE_PI);
-	}
-
-	if (OrbitZ != 0.0f)
-	{
-		ResourceData.OrbitPhaseBase.Z += (0.5f * UE_PI);
-	}
-
-	// Compute an offset to position the particle at the beginning of its orbit.
-	EmitterInfo.OrbitOffsetBase.X = 2.0f * ResourceData.OrbitOffsetBase.X * (OrbitY * FMath::Cos(ResourceData.OrbitPhaseBase.Y) + OrbitZ * FMath::Sin(ResourceData.OrbitPhaseBase.Z));
-	EmitterInfo.OrbitOffsetBase.Y = 2.0f * ResourceData.OrbitOffsetBase.Y * (OrbitZ * FMath::Cos(ResourceData.OrbitPhaseBase.Z) + OrbitX * FMath::Sin(ResourceData.OrbitPhaseBase.X));
-	EmitterInfo.OrbitOffsetBase.Z = 2.0f * ResourceData.OrbitOffsetBase.Z * (OrbitX * FMath::Cos(ResourceData.OrbitPhaseBase.X) + OrbitY * FMath::Sin(ResourceData.OrbitPhaseBase.Y));
-	EmitterInfo.OrbitOffsetRange.X = -EmitterInfo.OrbitOffsetBase.X + 2.0f * (ResourceData.OrbitOffsetBase.X + ResourceData.OrbitOffsetRange.X)
-		* (OrbitY * FMath::Cos(ResourceData.OrbitPhaseBase.Y + ResourceData.OrbitPhaseRange.Y) + OrbitZ * FMath::Sin(ResourceData.OrbitPhaseBase.Z + ResourceData.OrbitPhaseRange.Z));
-	EmitterInfo.OrbitOffsetRange.Y = -EmitterInfo.OrbitOffsetBase.Y + 2.0f * (ResourceData.OrbitOffsetBase.Y + ResourceData.OrbitOffsetRange.Y)
-		* (OrbitZ * FMath::Cos(ResourceData.OrbitPhaseBase.Z + ResourceData.OrbitPhaseRange.Z) + OrbitX * FMath::Sin(ResourceData.OrbitPhaseBase.X + ResourceData.OrbitPhaseRange.X));
-	EmitterInfo.OrbitOffsetRange.Z = -EmitterInfo.OrbitOffsetBase.Z + 2.0f * (ResourceData.OrbitOffsetBase.Z + ResourceData.OrbitOffsetRange.Z)
-		* (OrbitX * FMath::Cos(ResourceData.OrbitPhaseBase.X + ResourceData.OrbitPhaseRange.X) + OrbitY * FMath::Sin(ResourceData.OrbitPhaseBase.Y + ResourceData.OrbitPhaseRange.Y));
-
-	// Local vector field.
-	EmitterInfo.LocalVectorField.Field = EmitterBuildInfo.LocalVectorField;
-	EmitterInfo.LocalVectorField.Transform = EmitterBuildInfo.LocalVectorFieldTransform;
-	EmitterInfo.LocalVectorField.MinInitialRotation = FRotator::MakeFromEuler(
-		EmitterBuildInfo.LocalVectorFieldMinInitialRotation * 360.0f );
-	EmitterInfo.LocalVectorField.MaxInitialRotation = FRotator::MakeFromEuler(
-		EmitterBuildInfo.LocalVectorFieldMaxInitialRotation * 360.0f );
-	EmitterInfo.LocalVectorField.RotationRate = FRotator::MakeFromEuler(
-		EmitterBuildInfo.LocalVectorFieldRotationRate * 360.0f );
-	EmitterInfo.LocalVectorField.Intensity = EmitterBuildInfo.LocalVectorFieldIntensity;
-	EmitterInfo.LocalVectorField.Tightness = EmitterBuildInfo.LocalVectorFieldTightness;
-	EmitterInfo.LocalVectorField.bIgnoreComponentTransform = EmitterBuildInfo.bLocalVectorFieldIgnoreComponentTransform;
-	EmitterInfo.LocalVectorField.bTileX = EmitterBuildInfo.bLocalVectorFieldTileX;
-	EmitterInfo.LocalVectorField.bTileY = EmitterBuildInfo.bLocalVectorFieldTileY;
-	EmitterInfo.LocalVectorField.bTileZ = EmitterBuildInfo.bLocalVectorFieldTileZ;
-	EmitterInfo.LocalVectorField.bUseFixDT = EmitterBuildInfo.bLocalVectorFieldUseFixDT;
-
-	// Vector field scales.
-	FComposableFloatDistribution NormalizedVectorFieldScale(EmitterBuildInfo.VectorFieldScale);
-	NormalizedVectorFieldScale.Normalize(&ResourceData.PerParticleVectorFieldScale, &ResourceData.PerParticleVectorFieldBias);
-	FComposableDistribution::BuildFloat(EmitterInfo.VectorFieldScale, NormalizedVectorFieldScale);
-
-	if (EmitterBuildInfo.RequiredModule->bUseLocalSpace)
-	{
-		ResourceData.GlobalVectorFieldScale = 0.0f;
-		ResourceData.GlobalVectorFieldTightness = -1;
-	}
-	else
-	{
-		ResourceData.GlobalVectorFieldScale = EmitterBuildInfo.GlobalVectorFieldScale;
-		ResourceData.GlobalVectorFieldTightness = EmitterBuildInfo.GlobalVectorFieldTightness;
-	}
+		// Store the size-by-speed parameters.
+		ResourceData.SizeBySpeed.X = FMath::Max<float>(EmitterBuildInfo.SizeScaleBySpeed.X, 0.0f);
+		ResourceData.SizeBySpeed.Y = FMath::Max<float>(EmitterBuildInfo.SizeScaleBySpeed.Y, 0.0f);
+		ResourceData.SizeBySpeed.Z = FMath::Max<float>(EmitterBuildInfo.MaxSizeScaleBySpeed.X, 0.0f);
+		ResourceData.SizeBySpeed.W = FMath::Max<float>(EmitterBuildInfo.MaxSizeScaleBySpeed.Y, 0.0f);
 
-	// Drag coefficient.
-	FComposableFloatDistribution NormalizedDragCoefficient(EmitterBuildInfo.DragCoefficient);
-	NormalizedDragCoefficient.Normalize(&ResourceData.DragCoefficientScale, &ResourceData.DragCoefficientBias);
-	FComposableDistribution::BuildFloat(EmitterInfo.DragCoefficient, NormalizedDragCoefficient);
+		// Point attractor.
+		{
+			float RadiusSq = EmitterBuildInfo.PointAttractorRadius *
+				EmitterBuildInfo.PointAttractorRadius;
+			EmitterBuildInfo.PointAttractorStrength.ScaleByConstant(RadiusSq);
+			FComposableDistribution::BuildFloat(
+				EmitterInfo.PointAttractorStrength,
+				EmitterBuildInfo.PointAttractorStrength );
+			EmitterInfo.PointAttractorPosition = EmitterBuildInfo.PointAttractorPosition;
+			EmitterInfo.PointAttractorRadiusSq = RadiusSq;
+		}
 
-	// Set the scale by which rotation rate must be multiplied.
-	ResourceData.RotationRateScale = RotationRateScale;
+		// Store the constant acceleration to apply to particles.
+		ResourceData.ConstantAcceleration = EmitterBuildInfo.ConstantAcceleration;
+		EmitterInfo.ConstantAcceleration = EmitterBuildInfo.ConstantAcceleration;
+
+		// Compute the orbit offset amount.
+		FComposableDistribution::BuildVector(VectorDistribution, EmitterBuildInfo.OrbitOffset);
+		VectorDistribution.GetRange(&MinValue, &MaxValue);
+
+		// One half required due to integration in the shader.
+		MinValue *= 0.5f;
+		MaxValue *= 0.5f;
+
+		// Store the orbit offset range.
+		ResourceData.OrbitOffsetBase = MinValue;
+		ResourceData.OrbitOffsetRange = MaxValue - MinValue;
+
+		// Compute the orbit frequencies.
+		FComposableDistribution::BuildVector(VectorDistribution, EmitterBuildInfo.OrbitRotationRate);
+		VectorDistribution.GetRange(&MinValue, &MaxValue);
+
+		// # rotations to radians. Flip Z to be consistent with CPU orbit.
+		MinValue *= (2.0f * UE_PI);
+		MaxValue *= (2.0f * UE_PI);
+		MinValue.Z *= -1.0f;
+		MaxValue.Z *= -1.0f;
+
+		// Store the orbit frequency range.
+		ResourceData.OrbitFrequencyBase = MinValue;
+		ResourceData.OrbitFrequencyRange = MaxValue - MinValue;
+
+		// Compute the orbit phase.
+		FComposableDistribution::BuildVector(VectorDistribution, EmitterBuildInfo.OrbitInitialRotation);
+		VectorDistribution.GetRange(&MinValue, &MaxValue);
+
+		// # rotations to radians. Flip Z to be consistent with CPU orbit.
+		MinValue *= (2.0f * UE_PI);
+		MaxValue *= (2.0f * UE_PI);
+		MinValue.Z *= -1.0f;
+		MaxValue.Z *= -1.0f;
+
+		// Store the orbit phase range.
+		ResourceData.OrbitPhaseBase = MinValue;
+		ResourceData.OrbitPhaseRange = MaxValue - MinValue;
+
+		// Determine around which axes particles are orbiting.
+		const float OrbitX = (ResourceData.OrbitFrequencyBase.X != 0.0f || ResourceData.OrbitFrequencyRange.X != 0.0f || ResourceData.OrbitPhaseBase.X != 0.0f || ResourceData.OrbitPhaseRange.X != 0.0f) ? 1.0f : 0.0f;
+		const float OrbitY = (ResourceData.OrbitFrequencyBase.Y != 0.0f || ResourceData.OrbitFrequencyRange.Y != 0.0f || ResourceData.OrbitPhaseBase.Y != 0.0f || ResourceData.OrbitPhaseRange.Y != 0.0f) ? 1.0f : 0.0f;
+		const float OrbitZ = (ResourceData.OrbitFrequencyBase.Z != 0.0f || ResourceData.OrbitFrequencyRange.Z != 0.0f || ResourceData.OrbitPhaseBase.Z != 0.0f || ResourceData.OrbitPhaseRange.Z != 0.0f) ? 1.0f : 0.0f;
+
+		// Make some adjustments to mimic CPU orbit as much as possible.
+		if (OrbitX != 0.0f)
+		{
+			ResourceData.OrbitPhaseBase.X += (0.5f * UE_PI);
+		}
 
-	// Camera motion blur.
-	ResourceData.CameraMotionBlurAmount = this->CameraMotionBlurAmount;
+		if (OrbitZ != 0.0f)
+		{
+			ResourceData.OrbitPhaseBase.Z += (0.5f * UE_PI);
+		}
 
-	// Compute the maximum lifetime of particles in this emitter. */
-	EmitterInfo.MaxLifetime = 0.0f;
-	for (int32 ModuleIndex = 0; ModuleIndex < EmitterInfo.SpawnModules.Num(); ++ModuleIndex)
-	{
-		UParticleModuleLifetimeBase* LifetimeModule = Cast<UParticleModuleLifetimeBase>(EmitterInfo.SpawnModules[ModuleIndex]);
-		if (LifetimeModule)
+		// Compute an offset to position the particle at the beginning of its orbit.
+		EmitterInfo.OrbitOffsetBase.X = 2.0f * ResourceData.OrbitOffsetBase.X * (OrbitY * FMath::Cos(ResourceData.OrbitPhaseBase.Y) + OrbitZ * FMath::Sin(ResourceData.OrbitPhaseBase.Z));
+		EmitterInfo.OrbitOffsetBase.Y = 2.0f * ResourceData.OrbitOffsetBase.Y * (OrbitZ * FMath::Cos(ResourceData.OrbitPhaseBase.Z) + OrbitX * FMath::Sin(ResourceData.OrbitPhaseBase.X));
+		EmitterInfo.OrbitOffsetBase.Z = 2.0f * ResourceData.OrbitOffsetBase.Z * (OrbitX * FMath::Cos(ResourceData.OrbitPhaseBase.X) + OrbitY * FMath::Sin(ResourceData.OrbitPhaseBase.Y));
+		EmitterInfo.OrbitOffsetRange.X = -EmitterInfo.OrbitOffsetBase.X + 2.0f * (ResourceData.OrbitOffsetBase.X + ResourceData.OrbitOffsetRange.X)
+			* (OrbitY * FMath::Cos(ResourceData.OrbitPhaseBase.Y + ResourceData.OrbitPhaseRange.Y) + OrbitZ * FMath::Sin(ResourceData.OrbitPhaseBase.Z + ResourceData.OrbitPhaseRange.Z));
+		EmitterInfo.OrbitOffsetRange.Y = -EmitterInfo.OrbitOffsetBase.Y + 2.0f * (ResourceData.OrbitOffsetBase.Y + ResourceData.OrbitOffsetRange.Y)
+			* (OrbitZ * FMath::Cos(ResourceData.OrbitPhaseBase.Z + ResourceData.OrbitPhaseRange.Z) + OrbitX * FMath::Sin(ResourceData.OrbitPhaseBase.X + ResourceData.OrbitPhaseRange.X));
+		EmitterInfo.OrbitOffsetRange.Z = -EmitterInfo.OrbitOffsetBase.Z + 2.0f * (ResourceData.OrbitOffsetBase.Z + ResourceData.OrbitOffsetRange.Z)
+			* (OrbitX * FMath::Cos(ResourceData.OrbitPhaseBase.X + ResourceData.OrbitPhaseRange.X) + OrbitY * FMath::Sin(ResourceData.OrbitPhaseBase.Y + ResourceData.OrbitPhaseRange.Y));
+
+		// Local vector field.
+		EmitterInfo.LocalVectorField.Field = EmitterBuildInfo.LocalVectorField;
+		EmitterInfo.LocalVectorField.Transform = EmitterBuildInfo.LocalVectorFieldTransform;
+		EmitterInfo.LocalVectorField.MinInitialRotation = FRotator::MakeFromEuler(
+			EmitterBuildInfo.LocalVectorFieldMinInitialRotation * 360.0f );
+		EmitterInfo.LocalVectorField.MaxInitialRotation = FRotator::MakeFromEuler(
+			EmitterBuildInfo.LocalVectorFieldMaxInitialRotation * 360.0f );
+		EmitterInfo.LocalVectorField.RotationRate = FRotator::MakeFromEuler(
+			EmitterBuildInfo.LocalVectorFieldRotationRate * 360.0f );
+		EmitterInfo.LocalVectorField.Intensity = EmitterBuildInfo.LocalVectorFieldIntensity;
+		EmitterInfo.LocalVectorField.Tightness = EmitterBuildInfo.LocalVectorFieldTightness;
+		EmitterInfo.LocalVectorField.bIgnoreComponentTransform = EmitterBuildInfo.bLocalVectorFieldIgnoreComponentTransform;
+		EmitterInfo.LocalVectorField.bTileX = EmitterBuildInfo.bLocalVectorFieldTileX;
+		EmitterInfo.LocalVectorField.bTileY = EmitterBuildInfo.bLocalVectorFieldTileY;
+		EmitterInfo.LocalVectorField.bTileZ = EmitterBuildInfo.bLocalVectorFieldTileZ;
+		EmitterInfo.LocalVectorField.bUseFixDT = EmitterBuildInfo.bLocalVectorFieldUseFixDT;
+
+		// Vector field scales.
+		FComposableFloatDistribution NormalizedVectorFieldScale(EmitterBuildInfo.VectorFieldScale);
+		NormalizedVectorFieldScale.Normalize(&ResourceData.PerParticleVectorFieldScale, &ResourceData.PerParticleVectorFieldBias);
+		FComposableDistribution::BuildFloat(EmitterInfo.VectorFieldScale, NormalizedVectorFieldScale);
+
+		if (EmitterBuildInfo.RequiredModule->bUseLocalSpace)
 		{
-			EmitterInfo.MaxLifetime += LifetimeModule->GetMaxLifetime();
+			ResourceData.GlobalVectorFieldScale = 0.0f;
+			ResourceData.GlobalVectorFieldTightness = -1;
+		}
+		else
+		{
+			ResourceData.GlobalVectorFieldScale = EmitterBuildInfo.GlobalVectorFieldScale;
+			ResourceData.GlobalVectorFieldTightness = EmitterBuildInfo.GlobalVectorFieldTightness;
 		}
-	}
 
-	// Compute the maximum number of particles allowed for this emitter.
-	EmitterInfo.MaxParticleCount = FMath::Max<int32>( 1, EmitterBuildInfo.EstimatedMaxActiveParticleCount );
+		// Drag coefficient.
+		FComposableFloatDistribution NormalizedDragCoefficient(EmitterBuildInfo.DragCoefficient);
+		NormalizedDragCoefficient.Normalize(&ResourceData.DragCoefficientScale, &ResourceData.DragCoefficientBias);
+		FComposableDistribution::BuildFloat(EmitterInfo.DragCoefficient, NormalizedDragCoefficient);
 
-	// Store screen alignment for particles.
-	EmitterInfo.ScreenAlignment = EmitterBuildInfo.RequiredModule->ScreenAlignment;
-	ResourceData.ScreenAlignment = EmitterBuildInfo.RequiredModule->ScreenAlignment;
+		// Set the scale by which rotation rate must be multiplied.
+		ResourceData.RotationRateScale = RotationRateScale;
 
-	EmitterInfo.bRemoveHMDRoll = EmitterBuildInfo.RequiredModule->bRemoveHMDRoll;
-	EmitterInfo.MinFacingCameraBlendDistance = EmitterBuildInfo.RequiredModule->MinFacingCameraBlendDistance;
-	EmitterInfo.MaxFacingCameraBlendDistance = EmitterBuildInfo.RequiredModule->MaxFacingCameraBlendDistance;
+		// Camera motion blur.
+		ResourceData.CameraMotionBlurAmount = this->CameraMotionBlurAmount;
+
+		// Compute the maximum lifetime of particles in this emitter. */
+		EmitterInfo.MaxLifetime = 0.0f;
+		for (int32 ModuleIndex = 0; ModuleIndex < EmitterInfo.SpawnModules.Num(); ++ModuleIndex)
+		{
+			UParticleModuleLifetimeBase* LifetimeModule = Cast<UParticleModuleLifetimeBase>(EmitterInfo.SpawnModules[ModuleIndex]);
+			if (LifetimeModule)
+			{
+				EmitterInfo.MaxLifetime += LifetimeModule->GetMaxLifetime();
+			}
+		}
+
+		// Compute the maximum number of particles allowed for this emitter.
+		EmitterInfo.MaxParticleCount = FMath::Max<int32>( 1, EmitterBuildInfo.EstimatedMaxActiveParticleCount );
+
+		// Store screen alignment for particles.
+		EmitterInfo.ScreenAlignment = EmitterBuildInfo.RequiredModule->ScreenAlignment;
+		ResourceData.ScreenAlignment = EmitterBuildInfo.RequiredModule->ScreenAlignment;
+
+		EmitterInfo.bRemoveHMDRoll = EmitterBuildInfo.RequiredModule->bRemoveHMDRoll;
+		EmitterInfo.MinFacingCameraBlendDistance = EmitterBuildInfo.RequiredModule->MinFacingCameraBlendDistance;
+		EmitterInfo.MaxFacingCameraBlendDistance = EmitterBuildInfo.RequiredModule->MaxFacingCameraBlendDistance;
 	
-	ResourceData.bRemoveHMDRoll = EmitterBuildInfo.RequiredModule->bRemoveHMDRoll;
-	ResourceData.MinFacingCameraBlendDistance = EmitterBuildInfo.RequiredModule->MinFacingCameraBlendDistance;
-	ResourceData.MaxFacingCameraBlendDistance = EmitterBuildInfo.RequiredModule->MaxFacingCameraBlendDistance;
+		ResourceData.bRemoveHMDRoll = EmitterBuildInfo.RequiredModule->bRemoveHMDRoll;
+		ResourceData.MinFacingCameraBlendDistance = EmitterBuildInfo.RequiredModule->MinFacingCameraBlendDistance;
+		ResourceData.MaxFacingCameraBlendDistance = EmitterBuildInfo.RequiredModule->MaxFacingCameraBlendDistance;
 
-	// Particle axis lock
-	for (int32 ModuleIndex = 0; ModuleIndex < EmitterInfo.SpawnModules.Num(); ++ModuleIndex)
-	{
-		UParticleModuleOrientationAxisLock* AxisLockModule = Cast<UParticleModuleOrientationAxisLock>(EmitterInfo.SpawnModules[ModuleIndex]);
-		if (AxisLockModule)
+		// Particle axis lock
+		for (int32 ModuleIndex = 0; ModuleIndex < EmitterInfo.SpawnModules.Num(); ++ModuleIndex)
 		{
-			EmitterInfo.LockAxisFlag = AxisLockModule->LockAxisFlags;
-			ResourceData.LockAxisFlag = AxisLockModule->LockAxisFlags;
-			break;
+			UParticleModuleOrientationAxisLock* AxisLockModule = Cast<UParticleModuleOrientationAxisLock>(EmitterInfo.SpawnModules[ModuleIndex]);
+			if (AxisLockModule)
+			{
+				EmitterInfo.LockAxisFlag = AxisLockModule->LockAxisFlags;
+				ResourceData.LockAxisFlag = AxisLockModule->LockAxisFlags;
+				break;
+			}
 		}
-	}
 
-	ResourceData.PivotOffset = EmitterBuildInfo.PivotOffset;
+		ResourceData.PivotOffset = EmitterBuildInfo.PivotOffset;
 
-	// Store color and scale when using particle parameters.
-	EmitterInfo.DynamicColor = EmitterBuildInfo.DynamicColor;
-	EmitterInfo.DynamicAlpha= EmitterBuildInfo.DynamicAlpha;
-	EmitterInfo.DynamicColorScale = EmitterBuildInfo.DynamicColorScale;
-	EmitterInfo.DynamicAlphaScale = EmitterBuildInfo.DynamicAlphaScale;
+		// Store color and scale when using particle parameters.
+		EmitterInfo.DynamicColor = EmitterBuildInfo.DynamicColor;
+		EmitterInfo.DynamicAlpha= EmitterBuildInfo.DynamicAlpha;
+		EmitterInfo.DynamicColorScale = EmitterBuildInfo.DynamicColorScale;
+		EmitterInfo.DynamicAlphaScale = EmitterBuildInfo.DynamicAlphaScale;
 
-	// Collision flag.
-	EmitterInfo.bEnableCollision = EmitterBuildInfo.bEnableCollision;
-	EmitterInfo.CollisionMode = (EParticleCollisionMode::Type)EmitterBuildInfo.CollisionMode;
+		// Collision flag.
+		EmitterInfo.bEnableCollision = EmitterBuildInfo.bEnableCollision;
+		EmitterInfo.CollisionMode = (EParticleCollisionMode::Type)EmitterBuildInfo.CollisionMode;
+	}
 #endif
 
-
 	// Create or update GPU resources.
 	if ( EmitterInfo.Resources )
 	{
