Attempted fix for LoadZerothChunk crash in an editor cooked build

diff --git a/Engine/Source/Runtime/Engine/Private/SoundWave.cpp b/Engine/Source/Runtime/Engine/Private/SoundWave.cpp
--- a/Engine/Source/Runtime/Engine/Private/SoundWave.cpp
+++ b/Engine/Source/Runtime/Engine/Private/SoundWave.cpp
@@ -466,17 +466,24 @@ bool FSoundWaveData::LoadZerothChunk()
 #endif // WITH_EDITOR
 			{
 				// The zeroth chunk is cooked out to RunningPlatformData, so retrieve it.
-				check(GetNumChunks() > 0);
-				FStreamedAudioChunk& ZerothChunk = GetChunk(0);
+				if (ensureMsgf(GetNumChunks() > 0, TEXT("SoundWave '%s' is cooked but contains no chunks."), *GetFName().ToString()))
+				{
+					FStreamedAudioChunk& ZerothChunk = GetChunk(0);
 
-				// Sanity check to ensure bulk size is set up
-				UE_CLOG(ZerothChunk.BulkData.GetBulkDataSize() != ZerothChunk.DataSize, LogAudio, Warning
-					, TEXT("SoundWave '%s' bulk data serialized out had a mismatched size with the DataSize field."
-						"\nBulk Data Reported Size: %d"
-						"\nBulk Data Actual Size: %ld")
-					, *GetFName().ToString(), ZerothChunk.DataSize, ZerothChunk.BulkData.GetBulkDataSize());
+					// Sanity check to ensure bulk size is set up
+					UE_CLOG(ZerothChunk.BulkData.GetBulkDataSize() != ZerothChunk.DataSize, LogAudio, Warning
+						, TEXT("SoundWave '%s' bulk data serialized out had a mismatched size with the DataSize field."
+							"\nBulk Data Reported Size: %d"
+							"\nBulk Data Actual Size: %ld")
+						, *GetFName().ToString(), ZerothChunk.DataSize, ZerothChunk.BulkData.GetBulkDataSize());
 
-				ZerothChunkData = ZerothChunk.BulkData.GetCopyAsBuffer(ZerothChunk.AudioDataSize, true);
+					ZerothChunkData = ZerothChunk.BulkData.GetCopyAsBuffer(ZerothChunk.AudioDataSize, true);
+				}
+				else
+				{
+					// failed as we have no chunks.
+					return false;
+				}
 			}
 		}
 	}
