Fix CIS issue.

diff --git a/Engine/Source/Runtime/Engine/Private/Materials/HLSLMaterialTranslator.cpp b/Engine/Source/Runtime/Engine/Private/Materials/HLSLMaterialTranslator.cpp
--- a/Engine/Source/Runtime/Engine/Private/Materials/HLSLMaterialTranslator.cpp
+++ b/Engine/Source/Runtime/Engine/Private/Materials/HLSLMaterialTranslator.cpp
@@ -11077,7 +11077,6 @@ bool FHLSLMaterialTranslator::StrataSkipsOpacityEvaluation()
 		&& !Material->GetShadingModels().HasShadingModel(MSM_SingleLayerWater)
 		&& !Material->GetShadingModels().HasShadingModel(MSM_Subsurface)
 		&& !Material->GetShadingModels().HasShadingModel(MSM_SubsurfaceProfile)
-		&& !Material->GetShadingModels().HasShadingModel(MSM_PreintegratedSkin)
 		&& !Material->GetShadingModels().HasShadingModel(MSM_TwoSidedFoliage)
 		&& !Material->GetShadingModels().HasShadingModel(MSM_PreintegratedSkin);
 }
