WorldPartition: Avoid creating sub containers for Validation

This removes invalid reference errors on every Level Instance contained in a WP World when validating a CL containing actors of that WP world without the world being loaded first.

diff --git a/Engine/Source/Runtime/Engine/Private/WorldPartition/WorldPartitionStreamingGeneration.cpp b/Engine/Source/Runtime/Engine/Private/WorldPartition/WorldPartitionStreamingGeneration.cpp
--- a/Engine/Source/Runtime/Engine/Private/WorldPartition/WorldPartitionStreamingGeneration.cpp
+++ b/Engine/Source/Runtime/Engine/Private/WorldPartition/WorldPartitionStreamingGeneration.cpp
@@ -294,7 +294,10 @@ class FWorldPartitionStreamingGenerator
 		{
 			if (InActorDescView.IsContainerInstance())
 			{
-				OutContainerInstances.Add(InActorDescView);
+				if (!bIsValidating)
+				{
+					OutContainerInstances.Add(InActorDescView);
+				}
 			}
 			else
 			{
@@ -729,8 +732,9 @@ class FWorldPartitionStreamingGenerator
 	}
 
 public:
-	FWorldPartitionStreamingGenerator(FActorDescList* InModifiedActorsDescList, IStreamingGenerationErrorHandler* InErrorHandler, bool bInEnableStreaming, TArray<TSubclassOf<AActor>> InFilteredClasses = {})
+	FWorldPartitionStreamingGenerator(FActorDescList* InModifiedActorsDescList, IStreamingGenerationErrorHandler* InErrorHandler, bool bInEnableStreaming, bool bInIsValidating = false, TArray<TSubclassOf<AActor>> InFilteredClasses = {})
 		: bEnableStreaming(bInEnableStreaming)
+		, bIsValidating(bInIsValidating)
 		, ModifiedActorsDescList(InModifiedActorsDescList)
 		, FilteredClasses(InFilteredClasses)
 		, ErrorHandler(InErrorHandler ? InErrorHandler : &NullErrorHandler)	
@@ -838,6 +842,7 @@ public:
 
 private:
 	bool bEnableStreaming;
+	bool bIsValidating;
 	FActorDescList* ModifiedActorsDescList;
 	TArray<TSubclassOf<AActor>> FilteredClasses;
 	IStreamingGenerationErrorHandler* ErrorHandler;
@@ -917,7 +922,8 @@ void UWorldPartition::FlushStreaming()
 void UWorldPartition::GenerateHLOD(ISourceControlHelper* SourceControlHelper, bool bCreateActorsOnly)
 {
 	FStreamingGenerationLogErrorHandler LogErrorHandler;
-	FWorldPartitionStreamingGenerator StreamingGenerator(nullptr, &LogErrorHandler, IsStreamingEnabled(), { AWorldPartitionHLOD::StaticClass() });
+	const bool bIsValidating = false;
+	FWorldPartitionStreamingGenerator StreamingGenerator(nullptr, &LogErrorHandler, IsStreamingEnabled(), bIsValidating, { AWorldPartitionHLOD::StaticClass() });
 	StreamingGenerator.PreparationPhase(ActorDescContainer);
 
 	TUniquePtr<FArchive> LogFileAr = FWorldPartitionStreamingGenerator::CreateDumpStateLogArchive(TEXT("HLOD"));
@@ -935,7 +941,8 @@ void UWorldPartition::CheckForErrors(IStreamingGenerationErrorHandler* ErrorHand
 void UWorldPartition::CheckForErrors(IStreamingGenerationErrorHandler* ErrorHandler, const UActorDescContainer* ActorDescContainer, bool bEnableStreaming)
 {
 	FActorDescList ModifiedActorDescList;
-	FWorldPartitionStreamingGenerator StreamingGenerator(ActorDescContainer->GetWorld() ? &ModifiedActorDescList : nullptr, ErrorHandler, bEnableStreaming);
+	const bool bIsValidating = true;
+	FWorldPartitionStreamingGenerator StreamingGenerator(ActorDescContainer->GetWorld() ? &ModifiedActorDescList : nullptr, ErrorHandler, bEnableStreaming, bIsValidating);
 	StreamingGenerator.PreparationPhase(ActorDescContainer);
 }
 #endif // WITH_EDITOR
