Move destructor outside of ifdef for UHT's sake.

#rnx

diff --git a/Engine/Source/Runtime/Engine/Private/Net/Iris/ReplicationSystem/ActorReplicationBridge.cpp b/Engine/Source/Runtime/Engine/Private/Net/Iris/ReplicationSystem/ActorReplicationBridge.cpp
--- a/Engine/Source/Runtime/Engine/Private/Net/Iris/ReplicationSystem/ActorReplicationBridge.cpp
+++ b/Engine/Source/Runtime/Engine/Private/Net/Iris/ReplicationSystem/ActorReplicationBridge.cpp
@@ -987,4 +987,5 @@ float UActorReplicationBridge::GetMinSupportedNetUpdateFrequency() const
 }
 #else
 UActorReplicationBridge::UActorReplicationBridge() = default;
+UActorReplicationBridge::~UActorReplicationBridge() = default;
 #endif
diff --git a/Engine/Source/Runtime/Engine/Public/Net/Iris/ReplicationSystem/ActorReplicationBridge.h b/Engine/Source/Runtime/Engine/Public/Net/Iris/ReplicationSystem/ActorReplicationBridge.h
--- a/Engine/Source/Runtime/Engine/Public/Net/Iris/ReplicationSystem/ActorReplicationBridge.h
+++ b/Engine/Source/Runtime/Engine/Public/Net/Iris/ReplicationSystem/ActorReplicationBridge.h
@@ -46,11 +46,10 @@ class UActorReplicationBridge final : public UObjectReplicationBridge
 
 public:
 	ENGINE_API UActorReplicationBridge();
+	virtual ENGINE_API ~UActorReplicationBridge() override;
 
 #if UE_WITH_IRIS
 
-	ENGINE_API virtual ~UActorReplicationBridge();
-
 	ENGINE_API static UActorReplicationBridge* Create(UNetDriver* NetDriver);
 
 	/** Sets the net driver for the bridge. */
