Log warnings and errors if wave transformations change sample rate or channel count.

diff --git a/Engine/Source/Runtime/Engine/Private/AudioDerivedData.cpp b/Engine/Source/Runtime/Engine/Private/AudioDerivedData.cpp
--- a/Engine/Source/Runtime/Engine/Private/AudioDerivedData.cpp
+++ b/Engine/Source/Runtime/Engine/Private/AudioDerivedData.cpp
@@ -1354,8 +1354,12 @@ static void CookSimpleWave(const FAudioCookInputs& Inputs, TArray<uint8>& Output
 			Transformation->ProcessAudio(TransformationInfo);
 		}
 
+		UE_CLOG(WaveSampleRate != TransformationInfo.SampleRate, LogAudioDerivedData, Warning, TEXT("Wave transformations which alter the sample rate are not supported. Cooked audio for %s may be incorrect"), *Inputs.SoundFullName);
 		WaveSampleRate = TransformationInfo.SampleRate;
+
+		UE_CLOG(NumChannels != TransformationInfo.NumChannels, LogAudioDerivedData, Error, TEXT("Wave transformations which alter number of channels are not supported. Cooked audio for %s may be incorrect"), *Inputs.SoundFullName);
 		NumChannels = TransformationInfo.NumChannels;
+
 		NumSamples = InputFloatBuffer.Num();
 	}
 	
@@ -1590,10 +1594,13 @@ static void CookSurroundWave(const FAudioCookInputs& Inputs,  TArray<uint8>& Out
 				Transformation->ProcessAudio(TransformationInfo);
 			}
 
+			UE_CLOG(WaveSampleRate != TransformationInfo.SampleRate, LogAudioDerivedData, Warning, TEXT("Wave transformations which alter the sample rate are not supported. Cooked audio for %s may be incorrect"), *Inputs.SoundFullName);
 			WaveSampleRate = TransformationInfo.SampleRate;
 			bNeedsResample = Inputs.SampleRateOverride > 0 && Inputs.SampleRateOverride != (float)WaveSampleRate;
 			
+			UE_CLOG(ChannelCount != TransformationInfo.NumChannels, LogAudioDerivedData, Error, TEXT("Wave transformations which alter number of channels are not supported. Cooked audio for %s may be incorrect"), *Inputs.SoundFullName);
 			ChannelCount = TransformationInfo.NumChannels;
+
 			NumFrames = InterleavedFloatBuffer.Num() / ChannelCount;
 		}
 
