Fix copy past error in cvar usage for PerQualityMinLOD skeletalMesh

diff --git a/Engine/Source/Runtime/Engine/Private/SkeletalMesh.cpp b/Engine/Source/Runtime/Engine/Private/SkeletalMesh.cpp
--- a/Engine/Source/Runtime/Engine/Private/SkeletalMesh.cpp
+++ b/Engine/Source/Runtime/Engine/Private/SkeletalMesh.cpp
@@ -3254,7 +3254,7 @@ void USkeletalMesh::FinishPostLoadInternal(FSkinnedAssetPostLoadContext& Context
 	// Also test default value, since PerPLatformData can have Default !=0 and no PerPlaform data overrides.
 	bool bConvertMinLODData = (PerQualityLevelData.PerQuality.Num() == 0 && PerQualityLevelData.Default == 0) && (PerPlatformData.PerPlatform.Num() != 0 || PerPlatformData.Default != 0);
 
-	if (GEngine && GEngine->UseStaticMeshMinLODPerQualityLevels && bConvertMinLODData)
+	if (GEngine && GEngine->UseSkeletalMeshMinLODPerQualityLevels && bConvertMinLODData)
 	{
 		// get the platform groups
 		const TArray<FName>& PlatformGroupNameArray = PlatformInfo::GetAllPlatformGroupNames();
@@ -5308,7 +5308,7 @@ void USkeletalMesh::SetSkinWeightProfilesData(int32 LODIndex, FSkinWeightProfile
 
 void USkeletalMesh::OnLodStrippingQualityLevelChanged(IConsoleVariable* Variable) {
 #if WITH_EDITOR || PLATFORM_DESKTOP
-	if (GEngine && GEngine->UseStaticMeshMinLODPerQualityLevels)
+	if (GEngine && GEngine->UseSkeletalMeshMinLODPerQualityLevels)
 	{
 		for (TObjectIterator<USkeletalMesh> It; It; ++It)
 		{
