Improve blueprint API for soft object paths, they can now be easily created from path strings and loaded objects
Rename Get Path Name to the more explicit Get Object Path String and change the more complicated object conversion functions to show the full name
Add autocast to the safe object conversion functions

diff --git a/Engine/Source/Runtime/Engine/Classes/Kismet/KismetSystemLibrary.h b/Engine/Source/Runtime/Engine/Classes/Kismet/KismetSystemLibrary.h
--- a/Engine/Source/Runtime/Engine/Classes/Kismet/KismetSystemLibrary.h
+++ b/Engine/Source/Runtime/Engine/Classes/Kismet/KismetSystemLibrary.h
@@ -105,11 +105,15 @@ class ENGINE_API UKismetSystemLibrary : public UBlueprintFunctionLibrary
 	UFUNCTION(BlueprintPure, Category = "Utilities")
 	static FString GetObjectName(const UObject* Object);
 
-	// Returns the full path to the specified object.
-	UFUNCTION(BlueprintPure, Category="Utilities")
+	// Returns the full path to the specified object as a string
+	UFUNCTION(BlueprintPure, Category="Utilities", meta = (DisplayName = "Get Object Path String"))
 	static FString GetPathName(const UObject* Object);
 
-	// Returns the full system path to a UObject
+	// Returns the full path to the specified object as a Soft Object Path
+	UFUNCTION(BlueprintPure, Category = "Utilities")
+	static FSoftObjectPath GetSoftObjectPath(const UObject* Object);
+
+	// Returns the full file system path to a UObject
 	// If given a non-asset UObject, it will return an empty string
 	UFUNCTION(BlueprintPure, Category = "Utilities|Paths")
 	static FString GetSystemPath(const UObject* Object);
@@ -122,7 +126,11 @@ class ENGINE_API UKismetSystemLibrary : public UBlueprintFunctionLibrary
 
 	// Returns the display name of a class
 	UFUNCTION(BlueprintPure, Category = "Utilities", meta = (DisplayName = "Get Class Display Name"))
-	static FString GetClassDisplayName(UClass* Class);
+	static FString GetClassDisplayName(const UClass* Class);
+
+	// Returns the full path to the specified class as a Soft Class Path (that can be used as a Soft Object Path)
+	UFUNCTION(BlueprintPure, Category = "Utilities")
+	static FSoftClassPath GetSoftClassPath(const UClass* Class);
 
 	// Returns the outer object of an object.
 	UFUNCTION(BlueprintPure, Category = "Utilities")
@@ -233,27 +241,34 @@ class ENGINE_API UKismetSystemLibrary : public UBlueprintFunctionLibrary
 	UFUNCTION(BlueprintPure, Category="Utilities|Platform")
 	static FString GetDeviceId();
 
-	/** Converts an object into a class */
-	UFUNCTION(BlueprintPure, meta=(DisplayName = "Cast To Class", CompactNodeTitle = "->", DeterminesOutputType = "Class"), Category="Utilities")
+	/** Casts from an object to a class, this will only work if the object is already a class */
+	UFUNCTION(BlueprintPure, meta=(DisplayName = "Cast To Class", DeterminesOutputType = "Class"), Category="Utilities")
 	static UClass* Conv_ObjectToClass(UObject* Object, TSubclassOf<UObject> Class);
 
 	/** Converts an interfance into an object */
-	UFUNCTION(BlueprintPure, meta=(DisplayName = "To Object (Interface)", CompactNodeTitle = "->"), Category="Utilities")
+	UFUNCTION(BlueprintPure, meta=(DisplayName = "To Object (Interface)", CompactNodeTitle = "->", BlueprintAutocast), Category="Utilities")
 	static UObject* Conv_InterfaceToObject(const FScriptInterface& Interface); 
 
-	/** Builds a SoftObjectPath struct. Generally you should be using Soft Object References/Ptr types instead */
-	UFUNCTION(BlueprintPure, Category = "Utilities", meta = (Keywords = "construct build", NativeMakeFunc, BlueprintThreadSafe))
+	/** Builds a Soft Object Path struct from a string that contains a full /folder/packagename.object path */
+	UFUNCTION(BlueprintPure, Category = "Utilities", meta = (Keywords = "construct build", NativeMakeFunc, BlueprintThreadSafe, BlueprintAutocast))
 	static FSoftObjectPath MakeSoftObjectPath(const FString& PathString);
 
 	/** Gets the path string out of a Soft Object Path */
-	UFUNCTION(BlueprintPure, Category = "Utilities", meta = (NativeBreakFunc, BlueprintThreadSafe))
+	UFUNCTION(BlueprintPure, Category = "Utilities", meta = (NativeBreakFunc, BlueprintThreadSafe, BlueprintAutocast))
 	static void BreakSoftObjectPath(FSoftObjectPath InSoftObjectPath, FString& PathString);
 
 	/** Converts a Soft Object Path into a base Soft Object Reference, this is not guaranteed to be resolvable */
-	UFUNCTION(BlueprintPure, meta = (DisplayName = "To Soft Object Reference", CompactNodeTitle = "->"), Category = "Utilities")
+	UFUNCTION(BlueprintPure, meta = (DisplayName = "To Soft Object Reference"), Category = "Utilities")
 	static TSoftObjectPtr<UObject> Conv_SoftObjPathToSoftObjRef(const FSoftObjectPath& SoftObjectPath);
 
-	/** Builds a SoftClassPath struct. Generally you should be using Soft Class References/Ptr types instead */
+	/** Converts a Soft Object Reference into a Soft Object Path */
+	UFUNCTION(BlueprintPure, meta = (DisplayName = "To Soft Object Path", CompactNodeTitle = "->", BlueprintAutocast), Category = "Utilities")
+	static FSoftObjectPath Conv_SoftObjRefToSoftObjPath(TSoftObjectPtr<UObject> SoftObjectReference);
+
+	/** 
+	 * Builds a Soft Class Path struct from a string that contains a full /folder/packagename.class path.
+	 * For blueprint classes, this needs to point to the actual class (often with _C) and not the blueprint editor asset
+	 */
 	UFUNCTION(BlueprintPure, Category = "Utilities", meta = (Keywords = "construct build", NativeMakeFunc, BlueprintThreadSafe))
 	static FSoftClassPath MakeSoftClassPath(const FString& PathString);
 
@@ -262,15 +277,19 @@ class ENGINE_API UKismetSystemLibrary : public UBlueprintFunctionLibrary
 	static void BreakSoftClassPath(FSoftClassPath InSoftClassPath, FString& PathString);
 
 	/** Converts a Soft Class Path into a base Soft Class Reference, this is not guaranteed to be resolvable */
-	UFUNCTION(BlueprintPure, meta = (DisplayName = "To Soft Class Reference", CompactNodeTitle = "->"), Category = "Utilities")
+	UFUNCTION(BlueprintPure, meta = (DisplayName = "To Soft Class Reference"), Category = "Utilities")
 	static TSoftClassPtr<UObject> Conv_SoftClassPathToSoftClassRef(const FSoftClassPath& SoftClassPath);
 
+	/** Converts a Soft Object Reference into a Soft Class Path (which can be used like a Soft Object Path) */
+	UFUNCTION(BlueprintPure, meta = (DisplayName = "To Soft Class Path", CompactNodeTitle = "->", BlueprintAutocast), Category = "Utilities")
+	static FSoftClassPath Conv_SoftObjRefToSoftClassPath(TSoftClassPtr<UObject> SoftClassReference);
+
 	/** Returns true if the Soft Object Reference is not null */
 	UFUNCTION(BlueprintPure, Category = "Utilities", meta = (BlueprintThreadSafe))
 	static bool IsValidSoftObjectReference(const TSoftObjectPtr<UObject>& SoftObjectReference);
 
-	/** Converts a Soft Object Reference to a string. The other direction is not provided because it cannot be validated */
-	UFUNCTION(BlueprintPure, meta = (DisplayName = "To String (SoftObjectReference)", CompactNodeTitle = "->", BlueprintThreadSafe), Category = "Utilities")
+	/** Converts a Soft Object Reference to a path string */
+	UFUNCTION(BlueprintPure, meta = (DisplayName = "To String (SoftObjectReference)", CompactNodeTitle = "->", BlueprintThreadSafe, BlueprintAutocast), Category = "Utilities")
 	static FString Conv_SoftObjectReferenceToString(const TSoftObjectPtr<UObject>& SoftObjectReference);
 
 	/** Returns true if the values are equal (A == B) */
@@ -289,8 +308,8 @@ class ENGINE_API UKismetSystemLibrary : public UBlueprintFunctionLibrary
 	UFUNCTION(BlueprintPure, Category = "Utilities", meta = (BlueprintThreadSafe))
 	static bool IsValidSoftClassReference(const TSoftClassPtr<UObject>& SoftClassReference);
 
-	/** Converts a Soft Class Reference to a string. The other direction is not provided because it cannot be validated */
-	UFUNCTION(BlueprintPure, meta = (DisplayName = "To String (SoftClassReference)", CompactNodeTitle = "->", BlueprintThreadSafe), Category = "Utilities")
+	/** Converts a Soft Class Reference to a path string */
+	UFUNCTION(BlueprintPure, meta = (DisplayName = "To String (SoftClassReference)", CompactNodeTitle = "->", BlueprintThreadSafe, BlueprintAutocast), Category = "Utilities")
 	static FString Conv_SoftClassReferenceToString(const TSoftClassPtr<UObject>& SoftClassReference);
 
 	/** Returns true if the values are equal (A == B) */
@@ -2007,7 +2026,7 @@ class ENGINE_API UKismetSystemLibrary : public UBlueprintFunctionLibrary
 	static bool IsValidPrimaryAssetId(FPrimaryAssetId PrimaryAssetId);
 
 	/** Converts a Primary Asset Id to a string. The other direction is not provided because it cannot be validated */
-	UFUNCTION(BlueprintPure, meta = (DisplayName = "To String (PrimaryAssetId)", CompactNodeTitle = "->", ScriptMethod="ToString", BlueprintThreadSafe), Category = "AssetManager")
+	UFUNCTION(BlueprintPure, meta = (DisplayName = "To String (PrimaryAssetId)", CompactNodeTitle = "->", ScriptMethod="ToString", BlueprintThreadSafe, BlueprintAutocast), Category = "AssetManager")
 	static FString Conv_PrimaryAssetIdToString(FPrimaryAssetId PrimaryAssetId);
 
 	/** Returns true if the values are equal (A == B) */
@@ -2023,7 +2042,7 @@ class ENGINE_API UKismetSystemLibrary : public UBlueprintFunctionLibrary
 	static bool IsValidPrimaryAssetType(FPrimaryAssetType PrimaryAssetType);
 
 	/** Converts a Primary Asset Type to a string. The other direction is not provided because it cannot be validated */
-	UFUNCTION(BlueprintPure, meta = (DisplayName = "To String (PrimaryAssetType)", CompactNodeTitle = "->", ScriptMethod="ToString", BlueprintThreadSafe), Category = "AssetManager")
+	UFUNCTION(BlueprintPure, meta = (DisplayName = "To String (PrimaryAssetType)", CompactNodeTitle = "->", ScriptMethod="ToString", BlueprintThreadSafe, BlueprintAutocast), Category = "AssetManager")
 	static FString Conv_PrimaryAssetTypeToString(FPrimaryAssetType PrimaryAssetType);
 
 	/** Returns true if the values are equal (A == B) */
diff --git a/Engine/Source/Runtime/Engine/Private/KismetSystemLibrary.cpp b/Engine/Source/Runtime/Engine/Private/KismetSystemLibrary.cpp
--- a/Engine/Source/Runtime/Engine/Private/KismetSystemLibrary.cpp
+++ b/Engine/Source/Runtime/Engine/Private/KismetSystemLibrary.cpp
@@ -82,6 +82,11 @@ FString UKismetSystemLibrary::GetPathName(const UObject* Object)
 	return GetPathNameSafe(Object);
 }
 
+FSoftObjectPath UKismetSystemLibrary::GetSoftObjectPath(const UObject* Object)
+{
+	return FSoftObjectPath(Object);
+}
+
 FString UKismetSystemLibrary::GetSystemPath(const UObject* Object)
 {
 	if (Object && Object->IsAsset())
@@ -111,11 +116,16 @@ FString UKismetSystemLibrary::GetDisplayName(const UObject* Object)
 	return Object ? Object->GetName() : FString();
 }
 
-FString UKismetSystemLibrary::GetClassDisplayName(UClass* Class)
+FString UKismetSystemLibrary::GetClassDisplayName(const UClass* Class)
 {
 	return Class ? Class->GetName() : FString();
 }
 
+FSoftClassPath UKismetSystemLibrary::GetSoftClassPath(const UClass* Class)
+{
+	return FSoftClassPath(Class);
+}
+
 UObject* UKismetSystemLibrary::GetOuterObject(const UObject* Object)
 {
 	return Object ? Object->GetOuter() : nullptr;
@@ -1133,6 +1143,11 @@ TSoftObjectPtr<UObject> UKismetSystemLibrary::Conv_SoftObjPathToSoftObjRef(const
 	return TSoftObjectPtr<UObject>(SoftObjectPath);
 }
 
+FSoftObjectPath UKismetSystemLibrary::Conv_SoftObjRefToSoftObjPath(TSoftObjectPtr<UObject> SoftObjectReference)
+{
+	return SoftObjectReference.ToSoftObjectPath();
+}
+
 FSoftClassPath UKismetSystemLibrary::MakeSoftClassPath(const FString& PathString)
 {
 	FSoftClassPath SoftClassPath(PathString);
@@ -1156,6 +1171,12 @@ TSoftClassPtr<UObject> UKismetSystemLibrary::Conv_SoftClassPathToSoftClassRef(co
 	return TSoftClassPtr<UObject>(SoftClassPath);
 }
 
+FSoftClassPath UKismetSystemLibrary::Conv_SoftObjRefToSoftClassPath(TSoftClassPtr<UObject> SoftClassReference)
+{
+	// TSoftClassPtr and FSoftClassPath are not directly compatible
+	return FSoftClassPath(SoftClassReference.ToString());
+}
+
 bool UKismetSystemLibrary::IsValidSoftObjectReference(const TSoftObjectPtr<UObject>& SoftObjectReference)
 {
 	return !SoftObjectReference.IsNull();
