Fix for montages behaving incorrectly when slot is placed in PostProcessAnimBp and using "bUseMainInstanceMontageEvaluationData"

diff --git a/Engine/Source/Runtime/Engine/Private/Animation/AnimInstance.cpp b/Engine/Source/Runtime/Engine/Private/Animation/AnimInstance.cpp
--- a/Engine/Source/Runtime/Engine/Private/Animation/AnimInstance.cpp
+++ b/Engine/Source/Runtime/Engine/Private/Animation/AnimInstance.cpp
@@ -3486,6 +3486,17 @@ FMarkerSyncAnimPosition UAnimInstance::GetSyncGroupPosition(FName InSyncGroupNam
 
 void UAnimInstance::UpdateMontageEvaluationData()
 {
+	if (IsUsingMainInstanceMontageEvaluationData())
+	{
+		const USkeletalMeshComponent* Comp = GetOwningComponent();
+		if (Comp && Comp->GetAnimInstance() != this)
+		{
+			// If we're using the main instance's montage eval data
+			// and we're not the main instance, then skip updating this instance's montage eval data
+			return;
+		}
+	}
+
 	FAnimInstanceProxy& Proxy = GetProxyOnGameThread<FAnimInstanceProxy>();
 
 	Proxy.GetMontageEvaluationData().Reset(MontageInstances.Num());
