Fix for montage inertial blends not working when slot node is placed in linked layer/graph with "bUseMainInstanceMontageEvaluationData"

diff --git a/Engine/Source/Runtime/Engine/Private/Animation/AnimInstanceProxy.cpp b/Engine/Source/Runtime/Engine/Private/Animation/AnimInstanceProxy.cpp
--- a/Engine/Source/Runtime/Engine/Private/Animation/AnimInstanceProxy.cpp
+++ b/Engine/Source/Runtime/Engine/Private/Animation/AnimInstanceProxy.cpp
@@ -59,7 +59,6 @@ FAnimInstanceProxy::FAnimInstanceProxy()
 	, FrameCounterForUpdate(0)
 	, FrameCounterForNodeUpdate(0)
 	, CacheBonesRecursionCounter(0)
-	, MainMontageEvaluationData(&MontageEvaluationData)
 	, bUpdatingRoot(false)
 	, bBoneCachesInvalidated(false)
 	, bShouldExtractRootMotion(false)
@@ -68,6 +67,7 @@ FAnimInstanceProxy::FAnimInstanceProxy()
 	, bIsBeingDebugged(false)
 #endif
 	, bInitializeSubsystems(false)
+	, bUseMainInstanceMontageEvaluationData(false)
 {
 }
 
@@ -86,7 +86,6 @@ FAnimInstanceProxy::FAnimInstanceProxy(UAnimInstance* Instance)
 	, FrameCounterForUpdate(0)
 	, FrameCounterForNodeUpdate(0)
 	, CacheBonesRecursionCounter(0)
-	, MainMontageEvaluationData(&MontageEvaluationData)
 	, bUpdatingRoot(false)
 	, bBoneCachesInvalidated(false)
 	, bShouldExtractRootMotion(false)
@@ -95,6 +94,7 @@ FAnimInstanceProxy::FAnimInstanceProxy(UAnimInstance* Instance)
 	, bIsBeingDebugged(false)
 #endif
 	, bInitializeSubsystems(false)
+	, bUseMainInstanceMontageEvaluationData(false)
 {
 }
 
@@ -665,7 +665,7 @@ void FAnimInstanceProxy::InitializeObjects(UAnimInstance* InAnimInstance)
 	if(UAnimInstance* MainAnimInstance = SkeletalMeshComponent->GetAnimInstance())
 	{
 		MainInstanceProxy = &MainAnimInstance->GetProxyOnAnyThread<FAnimInstanceProxy>();
-		MainMontageEvaluationData = InAnimInstance->IsUsingMainInstanceMontageEvaluationData() ? &MainInstanceProxy->MontageEvaluationData : &MontageEvaluationData;
+		bUseMainInstanceMontageEvaluationData = InAnimInstance->IsUsingMainInstanceMontageEvaluationData();
 	}
 
 	if (SkeletalMeshComponent->GetSkeletalMeshAsset() != nullptr)
@@ -2016,6 +2016,16 @@ const FMontageEvaluationState* FAnimInstanceProxy::GetActiveMontageEvaluationSta
 	return nullptr;
 }
 
+TMap<FName, UE::Anim::FSlotInertializationRequest>& FAnimInstanceProxy::GetSlotGroupInertializationRequestMap()
+{
+	if (bUseMainInstanceMontageEvaluationData && GetMainInstanceProxy())
+	{
+		return GetMainInstanceProxy()->SlotGroupInertializationRequestMap;
+	}
+
+	return SlotGroupInertializationRequestMap;
+}
+
 void FAnimInstanceProxy::GatherDebugData(FNodeDebugData& DebugData)
 {
 	GatherDebugData_WithRoot(DebugData, RootNode, NAME_AnimGraph);
@@ -2411,14 +2421,22 @@ float FAnimInstanceProxy::GetRelevantAnimTimeFraction(int32 MachineIndex, int32
 
 TArray<FMontageEvaluationState>& FAnimInstanceProxy::GetMontageEvaluationData()
 {
-	check(MainMontageEvaluationData);
-	return *MainMontageEvaluationData;
+	if (bUseMainInstanceMontageEvaluationData && GetMainInstanceProxy())
+	{
+		return GetMainInstanceProxy()->MontageEvaluationData;
+	}
+
+	return MontageEvaluationData;
 }
 
 const TArray<FMontageEvaluationState>& FAnimInstanceProxy::GetMontageEvaluationData() const
 {
-	check(MainMontageEvaluationData);
-	return *MainMontageEvaluationData;
+	if (bUseMainInstanceMontageEvaluationData && GetMainInstanceProxy())
+	{
+		return GetMainInstanceProxy()->MontageEvaluationData;
+	}
+
+	return MontageEvaluationData;
 }
 bool FAnimInstanceProxy::WasAnimNotifyStateActiveInAnyState(TSubclassOf<UAnimNotifyState> AnimNotifyStateType) const
 {
diff --git a/Engine/Source/Runtime/Engine/Public/Animation/AnimInstanceProxy.h b/Engine/Source/Runtime/Engine/Public/Animation/AnimInstanceProxy.h
--- a/Engine/Source/Runtime/Engine/Public/Animation/AnimInstanceProxy.h
+++ b/Engine/Source/Runtime/Engine/Public/Animation/AnimInstanceProxy.h
@@ -747,7 +747,7 @@ protected:
 		Note that there might be multiple Active at the same time. This will only return the first active one it finds. **/
 	const FMontageEvaluationState* GetActiveMontageEvaluationState() const;
 
-	TMap<FName, UE::Anim::FSlotInertializationRequest>& GetSlotGroupInertializationRequestMap() { return SlotGroupInertializationRequestMap; }
+	TMap<FName, UE::Anim::FSlotInertializationRequest>& GetSlotGroupInertializationRequestMap();
 
 	/** Access montage array data */
 	TArray<FMontageEvaluationState>& GetMontageEvaluationData();
@@ -1067,7 +1067,6 @@ private:
 private:
 	/** Copy of UAnimInstance::MontageInstances data used for update & evaluation */
 	TArray<FMontageEvaluationState> MontageEvaluationData;
-	TArray<FMontageEvaluationState>* MainMontageEvaluationData;
 
 	// Inertialization request for each slot.
 	TMap<FName, UE::Anim::FSlotInertializationRequest> SlotGroupInertializationRequestMap;
@@ -1127,4 +1126,6 @@ private:
 
 	// Whether subsystems should be initialized
 	uint8 bInitializeSubsystems : 1;
+
+	uint8 bUseMainInstanceMontageEvaluationData : 1;
 };
