Do not get a WorldPartitionReference on the WorldDataLayer before the instancing context is properly filled

diff --git a/Engine/Source/Runtime/Engine/Private/WorldPartition/WorldPartition.cpp b/Engine/Source/Runtime/Engine/Private/WorldPartition/WorldPartition.cpp
--- a/Engine/Source/Runtime/Engine/Private/WorldPartition/WorldPartition.cpp
+++ b/Engine/Source/Runtime/Engine/Private/WorldPartition/WorldPartition.cpp
@@ -1391,18 +1391,18 @@ UActorDescContainer* UWorldPartition::RegisterActorDescContainer(const FName& Co
 		ContainerToRegister->Initialize(GetWorld(), ContainerPackage);
 		AddContainer(ContainerToRegister);
 
-		FWorldPartitionReference WDLReference;
-		for (FActorDescList::TIterator<> ActorDescIterator(ContainerToRegister); ActorDescIterator; ++ActorDescIterator)
+		if (IsInitialized() && EditorHash != nullptr)
 		{
-			if (ActorDescIterator->GetActorNativeClass()->IsChildOf<AWorldDataLayers>())
+			FWorldPartitionReference WDLReference;
+			for (FActorDescList::TIterator<> ActorDescIterator(ContainerToRegister); ActorDescIterator; ++ActorDescIterator)
 			{
-				WDLReference = FWorldPartitionReference(this, ActorDescIterator->GetGuid());
-				break;
+				if (ActorDescIterator->GetActorNativeClass()->IsChildOf<AWorldDataLayers>())
+				{
+					WDLReference = FWorldPartitionReference(this, ActorDescIterator->GetGuid());
+					break;
+				}
 			}
-		}
 
-		if (IsInitialized() && EditorHash != nullptr)
-		{
 			HashActorDescContainer(ContainerToRegister);
 		}
 
