When an Actor Folder is deleted and doesn't have a parent folder, we allow to delete it from the disk

diff --git a/Engine/Source/Runtime/Engine/Private/Level.cpp b/Engine/Source/Runtime/Engine/Private/Level.cpp
--- a/Engine/Source/Runtime/Engine/Private/Level.cpp
+++ b/Engine/Source/Runtime/Engine/Private/Level.cpp
@@ -3662,15 +3662,24 @@ TSet<FGuid> ULevel::GetDeletedAndUnreferencedActorFolders() const
 		return FoldersToDelete;
 	}
 
+	TSet<FGuid> RootFoldersToDelete;
 	for (const auto& Pair : ActorFolders)
 	{
 		UActorFolder* ActorFolder = Pair.Value;
 		if (ActorFolder->IsMarkedAsDeleted())
 		{
-			FoldersToDelete.Add(ActorFolder->GetGuid());
+			if (!ActorFolder->GetParent())
+			{
+				RootFoldersToDelete.Add(ActorFolder->GetGuid());
+			}
+			else
+			{
+				FoldersToDelete.Add(ActorFolder->GetGuid());
+			}
 		}
 	}
 
+	// Remove all other folders that are still referenced by Actors or other Folders
 	for (const auto& Pair : ActorFolders)
 	{
 		UActorFolder* ActorFolder = Pair.Value;
@@ -3700,6 +3709,8 @@ TSet<FGuid> ULevel::GetDeletedAndUnreferencedActorFolders() const
 		}
 	}
 
+	// Allow folders that don't have parent folders to be deleted as it won't affect child actors/folders (they are already at the root)
+	FoldersToDelete.Append(RootFoldersToDelete);
 	return FoldersToDelete;
 }
 
