PR #9489: Fix badly name GetInputKeyAtDistanceAlongSpline (Safely) (Contributed by [REDACTED_PERSON])

diff --git a/Engine/Source/Runtime/Engine/Classes/Components/SplineComponent.h b/Engine/Source/Runtime/Engine/Classes/Components/SplineComponent.h
--- a/Engine/Source/Runtime/Engine/Classes/Components/SplineComponent.h
+++ b/Engine/Source/Runtime/Engine/Classes/Components/SplineComponent.h
@@ -610,9 +610,20 @@ public:
 	FVector GetDefaultUpVector(ESplineCoordinateSpace::Type CoordinateSpace) const;
 
 	/** Given a distance along the length of this spline, return the corresponding input key at that point */
-	UFUNCTION(BlueprintCallable, Category=Spline)
+	/** This method has been deprecated because it was incorrectly returning the input key at time. To maintain the same behavior,
+	 *  replace it with GetTimeAtDistanceAlongSpline. To actually get the input key, instead call GetInputKeyValueAtDistanceAlongSpline. */
+	UFUNCTION(BlueprintCallable, Category=Spline, meta = (DeprecatedFunction, DeprecationMessage = "Please use GetInputKeyValueAtDistanceAlongSpline to get input key at distance or GetTimeAtDistanceAlongSpline to get time value (normalized to duration) at distance (same logic as deprecated function)."))
 	float GetInputKeyAtDistanceAlongSpline(float Distance) const;
 
+	/** Given a distance along the length of this spline, return the corresponding input key at that point 
+      * with a fractional component between the current input key and the next as a percentage. */
+	UFUNCTION(BlueprintCallable, Category = Spline)
+	float GetInputKeyValueAtDistanceAlongSpline(float Distance) const;
+	
+	/** Given a distance along the length of this spline, return the corresponding time at that point */
+	UFUNCTION(BlueprintCallable, Category = Spline)
+	float GetTimeAtDistanceAlongSpline(float Distance) const;
+
 	/** Given a distance along the length of this spline, return the point in space where this puts you */
 	UFUNCTION(BlueprintCallable, Category=Spline)
 	FVector GetLocationAtDistanceAlongSpline(float Distance, ESplineCoordinateSpace::Type CoordinateSpace) const;
diff --git a/Engine/Source/Runtime/Engine/Private/Components/SplineComponent.cpp b/Engine/Source/Runtime/Engine/Private/Components/SplineComponent.cpp
--- a/Engine/Source/Runtime/Engine/Private/Components/SplineComponent.cpp
+++ b/Engine/Source/Runtime/Engine/Private/Components/SplineComponent.cpp
@@ -1147,6 +1147,23 @@ FVector USplineComponent::GetDefaultUpVector(ESplineCoordinateSpace::Type Coordi
 
 
 float USplineComponent::GetInputKeyAtDistanceAlongSpline(float Distance) const
+{
+	return GetTimeAtDistanceAlongSpline(Distance);
+}
+
+float USplineComponent::GetInputKeyValueAtDistanceAlongSpline(float Distance) const
+{
+	const int32 NumPoints = SplineCurves.Position.Points.Num();
+
+	if (NumPoints < 2)
+	{
+		return 0.0f;
+	}
+	
+	return SplineCurves.ReparamTable.Eval(Distance, 0.0f);
+}
+
+float USplineComponent::GetTimeAtDistanceAlongSpline(float Distance) const
 {
 	const int32 NumPoints = SplineCurves.Position.Points.Num();
 
@@ -1159,7 +1176,6 @@ float USplineComponent::GetInputKeyAtDistanceAlongSpline(float Distance) const
 	return SplineCurves.ReparamTable.Eval(Distance, 0.0f) * TimeMultiplier;
 }
 
-
 FVector USplineComponent::GetLocationAtDistanceAlongSpline(float Distance, ESplineCoordinateSpace::Type CoordinateSpace) const
 {
 	const float Param = SplineCurves.ReparamTable.Eval(Distance, 0.0f);
