Fix static analysis warning

diff --git a/Engine/Source/Runtime/Engine/Private/PrimitiveSceneProxy.cpp b/Engine/Source/Runtime/Engine/Private/PrimitiveSceneProxy.cpp
--- a/Engine/Source/Runtime/Engine/Private/PrimitiveSceneProxy.cpp
+++ b/Engine/Source/Runtime/Engine/Private/PrimitiveSceneProxy.cpp
@@ -494,7 +494,7 @@ void FPrimitiveSceneProxy::UpdateUniformBuffer()
 				.PreSkinnedLocalBounds(PreSkinnedLocalBounds)
 				//.ObjectDrawDistance(FVector2f(GetMinDrawDistance(), GetMaxDrawDistance())) // TODO
 				.ReceivesDecals(bReceivesDecals)
-				.CacheShadowAsStatic(PrimitiveSceneInfo->ShouldCacheShadowAsStatic())
+				.CacheShadowAsStatic(PrimitiveSceneInfo ? PrimitiveSceneInfo->ShouldCacheShadowAsStatic() : false)
 				.OutputVelocity(bOutputVelocity)
 				.EvaluateWorldPositionOffset(EvaluateWorldPositionOffset())
 				.LightingChannelMask(GetLightingChannelMask())
