When a local light use an IES profile and use the IES profile intensity, ensures user can't change the light unit to avoid invalid brightness conversion.

diff --git a/Engine/Source/Runtime/Engine/Classes/Components/LightComponent.h b/Engine/Source/Runtime/Engine/Classes/Components/LightComponent.h
--- a/Engine/Source/Runtime/Engine/Classes/Components/LightComponent.h
+++ b/Engine/Source/Runtime/Engine/Classes/Components/LightComponent.h
@@ -167,7 +167,7 @@ class ENGINE_API ULightComponent : public ULightComponentBase
 	TObjectPtr<class UTextureLightProfile> IESTexture;
 
 	/** true: take light brightness from IES profile, false: use the light brightness - the maximum light in one direction is used to define no masking. Use with InverseSquareFalloff. Will be disabled if a valid IES profile texture is not supplied. */
-	UPROPERTY(EditAnywhere, BlueprintReadOnly, Category=LightProfiles, meta=(DisplayName = "Use IES Intensity"))
+	UPROPERTY(EditAnywhere, BlueprintReadOnly, Category=LightProfiles, meta=(DisplayName = "Use IES Intensity", EditCondition="IESTexture!=nullptr"))
 	uint32 bUseIESBrightness : 1;
 
 	/** Global scale for IES brightness contribution. Only available when "Use IES Brightness" is selected, and a valid IES profile texture is set */
diff --git a/Engine/Source/Runtime/Engine/Classes/Components/LocalLightComponent.h b/Engine/Source/Runtime/Engine/Classes/Components/LocalLightComponent.h
--- a/Engine/Source/Runtime/Engine/Classes/Components/LocalLightComponent.h
+++ b/Engine/Source/Runtime/Engine/Classes/Components/LocalLightComponent.h
@@ -23,7 +23,7 @@ class ENGINE_API ULocalLightComponent : public ULightComponent
 	 * The peak luminous intensity is measured in candelas,
 	 * while the luminous power is measured in lumens.
 	 */
-	UPROPERTY(EditAnywhere, BlueprintReadOnly, Category=Light, AdvancedDisplay, meta=(DisplayName="Intensity Units"))
+	UPROPERTY(EditAnywhere, BlueprintReadOnly, Category=Light, AdvancedDisplay, meta=(DisplayName="Intensity Units", EditCondition="!bUseIESBrightness"))
 	ELightUnits IntensityUnits;
 
 	UPROPERTY()
