World Partition: fix unit tests when containers are registered without a world partition object (can also happen during changelist validations, if we ever load actors).

#rnx

diff --git a/Engine/Source/Runtime/Engine/Private/WorldPartition/WorldPartitionHandle.cpp b/Engine/Source/Runtime/Engine/Private/WorldPartition/WorldPartitionHandle.cpp
--- a/Engine/Source/Runtime/Engine/Private/WorldPartition/WorldPartitionHandle.cpp
+++ b/Engine/Source/Runtime/Engine/Private/WorldPartition/WorldPartitionHandle.cpp
@@ -78,8 +78,12 @@ void FWorldPartitionLoadingContext::FImmediate::RegisterActor(FWorldPartitionAct
 		UActorDescContainer* Container = ActorDesc->GetContainer();
 		check(Container);
 
-		const FTransform& ContainerTransform = Container->GetTypedOuter<UWorldPartition>()->GetInstanceTransform();
-		const FTransform* ContainerTransformPtr = ContainerTransform.Equals(FTransform::Identity) ? nullptr : &ContainerTransform;
+		const FTransform* ContainerTransformPtr = nullptr;
+		if (UWorldPartition* WorldPartition = Container->GetTypedOuter<UWorldPartition>())
+		{
+			const FTransform& ContainerTransform = WorldPartition->GetInstanceTransform();
+			ContainerTransformPtr = ContainerTransform.Equals(FTransform::Identity) ? nullptr : &ContainerTransform;
+		}
 
 		Actor->GetLevel()->AddLoadedActor(Actor, ContainerTransformPtr);
 	}
@@ -93,8 +97,12 @@ void FWorldPartitionLoadingContext::FImmediate::UnregisterActor(FWorldPartitionA
 		UActorDescContainer* Container = ActorDesc->GetContainer();
 		check(Container);
 
-		const FTransform& ContainerTransform = Container->GetTypedOuter<UWorldPartition>()->GetInstanceTransform();
-		const FTransform* ContainerTransformPtr = ContainerTransform.Equals(FTransform::Identity) ? nullptr : &ContainerTransform;
+		const FTransform* ContainerTransformPtr = nullptr;
+		if (UWorldPartition* WorldPartition = Container->GetTypedOuter<UWorldPartition>())
+		{
+			const FTransform& ContainerTransform = WorldPartition->GetInstanceTransform();
+			ContainerTransformPtr = ContainerTransform.Equals(FTransform::Identity) ? nullptr : &ContainerTransform;
+		}
 
 		Actor->GetLevel()->RemoveLoadedActor(Actor, ContainerTransformPtr);
 
@@ -113,8 +121,12 @@ FWorldPartitionLoadingContext::FDeferred::~FDeferred()
 	{
 		for (auto& [Container, ContainerOp] : LocalContainerOps)
 		{
-			const FTransform& ContainerTransform = Container->GetTypedOuter<UWorldPartition>()->GetInstanceTransform();
-			const FTransform* ContainerTransformPtr = ContainerTransform.Equals(FTransform::Identity) ? nullptr : &ContainerTransform;
+			const FTransform* ContainerTransformPtr = nullptr;
+			if (UWorldPartition* WorldPartition = Container->GetTypedOuter<UWorldPartition>())
+			{
+				const FTransform& ContainerTransform = WorldPartition->GetInstanceTransform();
+				ContainerTransformPtr = ContainerTransform.Equals(FTransform::Identity) ? nullptr : &ContainerTransform;
+			}
 
 			auto CreateActorList = [](TArray<AActor*>& ActorList, const TSet<FWorldPartitionActorDesc*>& SourceList) -> ULevel*
 			{
