KismetMathLibrary FBox overloads

diff --git a/Engine/Source/Runtime/Engine/Classes/Kismet/KismetMathLibrary.h b/Engine/Source/Runtime/Engine/Classes/Kismet/KismetMathLibrary.h
--- a/Engine/Source/Runtime/Engine/Classes/Kismet/KismetMathLibrary.h
+++ b/Engine/Source/Runtime/Engine/Classes/Kismet/KismetMathLibrary.h
@@ -1980,6 +1980,10 @@ class ENGINE_API UKismetMathLibrary : public UBlueprintFunctionLibrary
 	UFUNCTION(BlueprintPure, Category = "Math|Random", meta=(ScriptMethod = "RandomPointInBoxExtents", NotBlueprintThreadSafe))
 	static FVector RandomPointInBoundingBox(const FVector Center, const FVector HalfSize);
 
+	/** Returns a random point within the specified bounding box. */
+	UFUNCTION(BlueprintPure, Category = "Math|Random", meta=(DisplayName = "Random Point In Bounding Box (Box)", ScriptMethod = "RandomPointInBoxExtents", NotBlueprintThreadSafe))
+	static FVector RandomPointInBoundingBox_Box(const FBox Box);
+
 	/** 
 	 * Returns a random vector with length of 1, within the specified cone, with uniform random distribution.
 	 * @param ConeDir					The base "center" direction of the cone.
@@ -3999,6 +4003,10 @@ class ENGINE_API UKismetMathLibrary : public UBlueprintFunctionLibrary
 	UFUNCTION(BlueprintPure, Category="Math|Random")
 	static FVector RandomPointInBoundingBoxFromStream(const FVector Center, const FVector HalfSize, const FRandomStream& Stream);
 
+	/** Returns a random point within the specified bounding box. */
+	UFUNCTION(BlueprintPure, Category="Math|Random", meta=(DisplayName="Random Point In Bounding Box From Stream (Box)"))
+	static FVector RandomPointInBoundingBoxFromStream_Box(const FBox Box, const FRandomStream& Stream);
+
 	/** Create a random rotation */
 	UFUNCTION(BlueprintPure, Category="Math|Random")
 	static FRotator RandomRotatorFromStream(bool bRoll, const FRandomStream& Stream);
@@ -4128,6 +4136,16 @@ class ENGINE_API UKismetMathLibrary : public UBlueprintFunctionLibrary
 	UFUNCTION(BlueprintPure, Category = "Math|Geometry")
 	static bool IsPointInBox(FVector Point, FVector BoxOrigin, FVector BoxExtent);
 
+	/**
+     * Determines whether the given point is in a box. Includes points on the box.
+     *
+     * @param Point			Point to test
+     * @param Box			Box to test against
+     * @return Whether the point is in the box.
+     */
+    UFUNCTION(BlueprintPure, Category = "Math|Geometry", meta=(DisplayName = "Is Point In Box (Box)"))
+    static bool IsPointInBox_Box(FVector Point, FBox Box);
+
 	/**
 	* Determines whether a given point is in a box with a given transform. Includes points on the box.
 	*
@@ -4139,6 +4157,17 @@ class ENGINE_API UKismetMathLibrary : public UBlueprintFunctionLibrary
 	UFUNCTION(BlueprintPure, Category = "Math|Geometry")
 	static bool IsPointInBoxWithTransform(FVector Point, const FTransform& BoxWorldTransform, FVector BoxExtent);
 
+	/**
+	* Determines whether a given point is in a box with a given transform. Includes points on the box.
+	*
+	* @param Point				Point to test
+	* @param BoxWorldTransform	Component-to-World transform of the box.
+	* @param Box				Box to test agains in component space.
+	* @return Whether the point is in the box.
+	*/
+	UFUNCTION(BlueprintPure, Category = "Math|Geometry", meta=(DisplayName = "Is Point In Box With Transform (Box)"))
+	static bool IsPointInBoxWithTransform_Box(FVector Point, const FTransform& BoxWorldTransform, FBox BoxExtent);
+
 	/**
 	* Returns Slope Pitch and Roll angles in degrees based on the following information: 
 	*
diff --git a/Engine/Source/Runtime/Engine/Classes/Kismet/KismetMathLibrary.inl b/Engine/Source/Runtime/Engine/Classes/Kismet/KismetMathLibrary.inl
--- a/Engine/Source/Runtime/Engine/Classes/Kismet/KismetMathLibrary.inl
+++ b/Engine/Source/Runtime/Engine/Classes/Kismet/KismetMathLibrary.inl
@@ -1812,6 +1812,12 @@ FVector UKismetMathLibrary::RandomPointInBoundingBox(const FVector Center, const
 	return FMath::RandPointInBox(FBox(BoxMin, BoxMax));
 }
 
+KISMET_MATH_INLINE
+FVector UKismetMathLibrary::RandomPointInBoundingBox_Box(const FBox Box)
+{
+	return FMath::RandPointInBox(Box);
+}
+
 KISMET_MATH_FORCEINLINE
 FVector UKismetMathLibrary::RandomUnitVectorInConeInRadians(FVector ConeDir, float ConeHalfAngleInRadians)
 {
@@ -1828,7 +1834,7 @@ FVector UKismetMathLibrary::GetReflectionVector(FVector Direction, FVector Surfa
 KISMET_MATH_FORCEINLINE
 FVector UKismetMathLibrary::MirrorVectorByNormal(FVector InVect, FVector InNormal)
 {
-	// Same as GetReflectionVector() above - though they are bound to different names so keeping both for now
+	// Same as GetReflectionVector() above - though they are bound to cdifferent names so keeping both for now
 	return FMath::GetReflectionVector(InVect, InNormal);
 }
 
diff --git a/Engine/Source/Runtime/Engine/Private/KismetMathLibrary.cpp b/Engine/Source/Runtime/Engine/Private/KismetMathLibrary.cpp
--- a/Engine/Source/Runtime/Engine/Private/KismetMathLibrary.cpp
+++ b/Engine/Source/Runtime/Engine/Private/KismetMathLibrary.cpp
@@ -1092,6 +1092,11 @@ FVector UKismetMathLibrary::RandomPointInBoundingBoxFromStream(const FVector Cen
 	return Stream.RandPointInBox(FBox(BoxMin, BoxMax));
 }
 
+FVector UKismetMathLibrary::RandomPointInBoundingBoxFromStream_Box(const FBox Box, const FRandomStream& Stream)
+{
+	return Stream.RandPointInBox(Box);
+}
+
 FRotator UKismetMathLibrary::RandomRotatorFromStream(bool bRoll, const FRandomStream& Stream)
 {
 	FRotator RRot;
@@ -1356,6 +1361,11 @@ bool UKismetMathLibrary::IsPointInBox(FVector Point, FVector BoxOrigin, FVector
 	return Box.IsInsideOrOn(Point);
 }
 
+bool UKismetMathLibrary::IsPointInBox_Box(FVector Point, FBox Box)
+{
+	return Box.IsInsideOrOn(Point);
+}
+
 bool UKismetMathLibrary::IsPointInBoxWithTransform(FVector Point, const FTransform& BoxWorldTransform, FVector BoxExtent)
 {
 	// Put point in component space
@@ -1365,6 +1375,15 @@ bool UKismetMathLibrary::IsPointInBoxWithTransform(FVector Point, const FTransfo
 	return Box.IsInsideOrOn(PointInComponentSpace);
 }
 
+bool UKismetMathLibrary::IsPointInBoxWithTransform_Box(FVector Point, const FTransform& BoxWorldTransform, FBox Box)
+{
+	// Put point in component space
+	const FVector PointInComponentSpace = BoxWorldTransform.InverseTransformPosition(Point);
+	// Now it's just a normal point-in-box test, with a box at the origin.
+	Box = Box.MoveTo(PointInComponentSpace);
+	return Box.IsInsideOrOn(PointInComponentSpace);
+}
+
 void UKismetMathLibrary::GetSlopeDegreeAngles(const FVector& MyRightYAxis, const FVector& FloorNormal, const FVector& UpVector, float& OutSlopePitchDegreeAngle, float& OutSlopeRollDegreeAngle)
 {
 	const FVector FloorZAxis = FloorNormal;
