Wait for current render thread to complete, which includes the wait on transcode tasks completion, before allowing GC of the RVT component.
The RVT component owns StreamingTexture which has an FVirtualTextureBuiltData object which any transcode tasks will still be referencing.

diff --git a/Engine/Source/Runtime/Engine/Classes/Components/RuntimeVirtualTextureComponent.h b/Engine/Source/Runtime/Engine/Classes/Components/RuntimeVirtualTextureComponent.h
--- a/Engine/Source/Runtime/Engine/Classes/Components/RuntimeVirtualTextureComponent.h
+++ b/Engine/Source/Runtime/Engine/Classes/Components/RuntimeVirtualTextureComponent.h
@@ -4,6 +4,7 @@
 
 #include "CoreMinimal.h"
 #include "Engine/TextureDefines.h"
+#include "RenderCommandFence.h"
 #include "SceneComponent.h"
 #include "RuntimeVirtualTextureComponent.generated.h"
 
@@ -83,6 +84,9 @@ protected:
 	DECLARE_MULTICAST_DELEGATE_TwoParams(FGetHidePrimitivesDelegate, bool&, bool&);
 	FGetHidePrimitivesDelegate HidePrimitivesDelegate;
 
+	/** A fence to track render thread has finished with StreamingTexture data before destroy. */
+	FRenderCommandFence DestroyFence;
+
 public:
 	/**
 	 * This function marks an area of the runtime virtual texture as dirty.
@@ -147,6 +151,8 @@ public:
 
 protected:
 	//~ Begin UObject Interface
+	virtual void BeginDestroy() override;
+	virtual bool IsReadyForFinishDestroy() override;
 #if WITH_EDITOR
 	virtual bool CanEditChange(const FProperty* InProperty) const override;
 #endif
diff --git a/Engine/Source/Runtime/Engine/Private/Components/RuntimeVirtualTextureComponent.cpp b/Engine/Source/Runtime/Engine/Private/Components/RuntimeVirtualTextureComponent.cpp
--- a/Engine/Source/Runtime/Engine/Private/Components/RuntimeVirtualTextureComponent.cpp
+++ b/Engine/Source/Runtime/Engine/Private/Components/RuntimeVirtualTextureComponent.cpp
@@ -23,6 +23,22 @@ URuntimeVirtualTextureComponent::URuntimeVirtualTextureComponent(const FObjectIn
 	Mobility = EComponentMobility::Stationary;
 }
 
+void URuntimeVirtualTextureComponent::BeginDestroy()
+{
+	Super::BeginDestroy();
+	
+	// Queuing up a render fence means that we will have cleaned up the scene proxy/virtual texture producer before finishing the destroy.
+	// This means that any transcode tasks will have finished *before* we garbage collect our StreamingTexture.
+	// That's important because the transcode tasks reference the FVirtualTextureBuiltData from the StreamingTexture.
+	DestroyFence.BeginFence();
+}
+
+bool URuntimeVirtualTextureComponent::IsReadyForFinishDestroy()
+{
+	bool bResult = Super::IsReadyForFinishDestroy() && DestroyFence.IsFenceComplete();
+	return bResult;
+}
+
 #if WITH_EDITOR
 
 void URuntimeVirtualTextureComponent::OnRegister()
