Fix r.dynamicres.testScreenPercentage. We need to use the budget value for the upper bound if we're forcing resolution rather than the dynamically computed value. This requires exposing it to the state via a separate method.

diff --git a/Engine/Source/Runtime/Engine/Private/DynamicResolution.cpp b/Engine/Source/Runtime/Engine/Private/DynamicResolution.cpp
--- a/Engine/Source/Runtime/Engine/Private/DynamicResolution.cpp
+++ b/Engine/Source/Runtime/Engine/Private/DynamicResolution.cpp
@@ -1033,6 +1033,17 @@ public:
 		return Proxy->Heuristic.GetResolutionFractionUpperBounds();
 	}
 
+	virtual DynamicRenderScaling::TMap<float> GetResolutionFractionsUpperBoundBudgetValue() const override
+	{
+		DynamicRenderScaling::TMap<float> MaxResolutionFractions;
+		for (TLinkedList<DynamicRenderScaling::FBudget*>::TIterator BudgetIt(DynamicRenderScaling::FBudget::GetGlobalList()); BudgetIt; BudgetIt.Next())
+		{
+			const DynamicRenderScaling::FBudget& Budget = **BudgetIt;
+			MaxResolutionFractions[Budget] = Budget.GetSettings().MaxResolutionFraction;
+		}
+		return MaxResolutionFractions;
+	}
+
 	virtual void ProcessEvent(EDynamicResolutionStateEvent Event) override
 	{
 		check(IsInGameThread());
diff --git a/Engine/Source/Runtime/Engine/Private/UnrealEngine.cpp b/Engine/Source/Runtime/Engine/Private/UnrealEngine.cpp
--- a/Engine/Source/Runtime/Engine/Private/UnrealEngine.cpp
+++ b/Engine/Source/Runtime/Engine/Private/UnrealEngine.cpp
@@ -12496,7 +12496,7 @@ void UEngine::GetDynamicResolutionCurrentStateInfos(FDynamicResolutionStateInfos
 			OutInfos.Status = EDynamicResolutionStatus::DebugForceEnabled;
 			OutInfos.ResolutionFractionApproximations = DynamicResolutionState->GetResolutionFractionsApproximation();
 			OutInfos.ResolutionFractionApproximations[GDynamicPrimaryResolutionFraction] = ForceResolutionFraction;
-			OutInfos.ResolutionFractionUpperBounds = DynamicResolutionState->GetResolutionFractionsUpperBound();
+			OutInfos.ResolutionFractionUpperBounds = DynamicResolutionState->GetResolutionFractionsUpperBoundBudgetValue();
 		}
 		#endif
 		else if (DynamicResolutionState->IsEnabled())
diff --git a/Engine/Source/Runtime/Engine/Public/DynamicResolutionState.h b/Engine/Source/Runtime/Engine/Public/DynamicResolutionState.h
--- a/Engine/Source/Runtime/Engine/Public/DynamicResolutionState.h
+++ b/Engine/Source/Runtime/Engine/Public/DynamicResolutionState.h
@@ -59,6 +59,12 @@ protected:
 	/** Returns the max resolution resolution fraction. */
 	virtual DynamicRenderScaling::TMap<float> GetResolutionFractionsUpperBound() const = 0;
 
+	/** Returns the max resolution resolution fraction as specified in the budget (this can differ from the upper bound if the upper bound is dynamic)*/
+	virtual DynamicRenderScaling::TMap<float> GetResolutionFractionsUpperBoundBudgetValue() const
+	{
+		return GetResolutionFractionsUpperBound();
+	}
+
 	/** Enables/Disables dynamic resolution. This is only called by GEngine automatically. */
 	virtual void SetEnabled(bool bEnable) = 0;
 
