Adding nullptr checks which caused crashes for CIS tests

diff --git a/Engine/Source/Runtime/Engine/Private/WorldPartition/HLOD/HLODSubsystem.cpp b/Engine/Source/Runtime/Engine/Private/WorldPartition/HLOD/HLODSubsystem.cpp
--- a/Engine/Source/Runtime/Engine/Private/WorldPartition/HLOD/HLODSubsystem.cpp
+++ b/Engine/Source/Runtime/Engine/Private/WorldPartition/HLOD/HLODSubsystem.cpp
@@ -155,7 +155,11 @@ void UHLODSubsystem::OnWorldPartitionInitialized(UWorldPartition* InWorldPartiti
 	FWorldPartitionHLODRuntimeData& WorldPartitionHLODRuntimeData = WorldPartitionsHLODRuntimeData.Add(InWorldPartition, FWorldPartitionHLODRuntimeData());
 	
 	TSet<const UWorldPartitionRuntimeCell*> StreamingCells;
-	InWorldPartition->RuntimeHash->GetAllStreamingCells(StreamingCells, /*bAllDataLayers*/ true);
+	if (InWorldPartition && InWorldPartition->RuntimeHash)
+	{
+		InWorldPartition->RuntimeHash->GetAllStreamingCells(StreamingCells, /*bAllDataLayers*/ true);
+	}
+	
 
 	// Build cell to HLOD mapping
 	for (const UWorldPartitionRuntimeCell* Cell : StreamingCells)
diff --git a/Engine/Source/Runtime/Engine/Private/WorldPartition/WorldPartitionStreamingPolicy.cpp b/Engine/Source/Runtime/Engine/Private/WorldPartition/WorldPartitionStreamingPolicy.cpp
--- a/Engine/Source/Runtime/Engine/Private/WorldPartition/WorldPartitionStreamingPolicy.cpp
+++ b/Engine/Source/Runtime/Engine/Private/WorldPartition/WorldPartitionStreamingPolicy.cpp
@@ -273,7 +273,10 @@ void UWorldPartitionStreamingPolicy::UpdateStreamingState()
 		if (WorldPartition->CanStream())
 		{
 			UWorldPartitionRuntimeCell::DirtyStreamingSourceCacheEpoch();
-			WorldPartition->RuntimeHash->GetStreamingCells(StreamingSources, FrameActivateCells, FrameLoadCells);
+			if (WorldPartition->RuntimeHash)
+			{
+				WorldPartition->RuntimeHash->GetStreamingCells(StreamingSources, FrameActivateCells, FrameLoadCells);
+			}
 
 			// Activation superseeds Loading
 			LoadStreamingCells = LoadStreamingCells.Difference(ActivateStreamingCells);
@@ -302,16 +305,19 @@ void UWorldPartitionStreamingPolicy::UpdateStreamingState()
 		}
 
 		if (!bIsServerStreamingEnabled)
-		{
-			const UDataLayerSubsystem* DataLayerSubsystem = WorldPartition->GetWorld()->GetSubsystem<UDataLayerSubsystem>();
+		{			
+			if (WorldPartition->RuntimeHash)
+			{
+				const UDataLayerSubsystem* DataLayerSubsystem = WorldPartition->GetWorld()->GetSubsystem<UDataLayerSubsystem>();
 
-			// Non Data Layer Cells + Active Data Layers
-			WorldPartition->RuntimeHash->GetAllStreamingCells(ActivateStreamingCells, /*bAllDataLayers=*/ false, /*bDataLayersOnly=*/ false, DataLayerSubsystem->GetEffectiveActiveDataLayerNames());
+				// Non Data Layer Cells + Active Data Layers
+				WorldPartition->RuntimeHash->GetAllStreamingCells(ActivateStreamingCells, /*bAllDataLayers=*/ false, /*bDataLayersOnly=*/ false, DataLayerSubsystem->GetEffectiveActiveDataLayerNames());
 
-			// Loaded Data Layers Cells only
-			if (DataLayerSubsystem->GetEffectiveLoadedDataLayerNames().Num())
-			{
-				WorldPartition->RuntimeHash->GetAllStreamingCells(LoadStreamingCells, /*bAllDataLayers=*/ false, /*bDataLayersOnly=*/ true, DataLayerSubsystem->GetEffectiveLoadedDataLayerNames());
+				// Loaded Data Layers Cells only
+				if (DataLayerSubsystem->GetEffectiveLoadedDataLayerNames().Num())
+				{
+					WorldPartition->RuntimeHash->GetAllStreamingCells(LoadStreamingCells, /*bAllDataLayers=*/ false, /*bDataLayersOnly=*/ true, DataLayerSubsystem->GetEffectiveLoadedDataLayerNames());
+				}
 			}
 		}
 	}
@@ -401,7 +407,10 @@ void UWorldPartitionStreamingPolicy::UpdateStreamingState()
 				AddToWorldCells.Add(ActivatedCell);
 			}
 		}
-		WorldPartition->RuntimeHash->SortStreamingCellsByImportance(AddToWorldCells, StreamingSources, SortedAddToWorldCells);
+		if (WorldPartition->RuntimeHash)
+		{
+			WorldPartition->RuntimeHash->SortStreamingCellsByImportance(AddToWorldCells, StreamingSources, SortedAddToWorldCells);
+		}
 
 		// Update level streaming priority so that UWorld::UpdateLevelStreaming will naturally process the levels in the correct order
 		const int32 MaxPrio = SortedAddToWorldCells.Num();
@@ -440,13 +449,16 @@ void UWorldPartitionStreamingPolicy::UpdateDebugCellsStreamingPriority(const TSe
 		CellsToSort.Append(ActivateStreamingCells);
 		CellsToSort.Append(LoadStreamingCells);
 
-		TArray<const UWorldPartitionRuntimeCell*, TInlineAllocator<256>> SortedCells;
-		WorldPartition->RuntimeHash->SortStreamingCellsByImportance(CellsToSort, StreamingSources, SortedCells);
-		const int32 CellCount = SortedCells.Num();
-		int32 CellPrio = 0;
-		for (const UWorldPartitionRuntimeCell* SortedCell : SortedCells)
+		if (WorldPartition->RuntimeHash)
 		{
-			const_cast<UWorldPartitionRuntimeCell*>(SortedCell)->SetDebugStreamingPriority(float(CellPrio++) / CellCount);
+			TArray<const UWorldPartitionRuntimeCell*, TInlineAllocator<256>> SortedCells;
+			WorldPartition->RuntimeHash->SortStreamingCellsByImportance(CellsToSort, StreamingSources, SortedCells);
+			const int32 CellCount = SortedCells.Num();
+			int32 CellPrio = 0;
+			for (const UWorldPartitionRuntimeCell* SortedCell : SortedCells)
+			{
+				const_cast<UWorldPartitionRuntimeCell*>(SortedCell)->SetDebugStreamingPriority(float(CellPrio++) / CellCount);
+			}
 		}
 	}
 }
@@ -465,15 +477,18 @@ void UWorldPartitionStreamingPolicy::UpdateStreamingPerformance(const TSet<const
 		return;
 	}
 
-	EWorldPartitionStreamingPerformance NewStreamingPerformance = WorldPartition->RuntimeHash->GetStreamingPerformance(CellsToActivate);
-		
-	if (StreamingPerformance != NewStreamingPerformance)
+	if (WorldPartition->RuntimeHash)
 	{
-		UE_LOG(LogWorldPartition, Log, TEXT("Streaming performance changed: %s -> %s"),
-			*StaticEnum<EWorldPartitionStreamingPerformance>()->GetDisplayNameTextByValue((int64)StreamingPerformance).ToString(),
-			*StaticEnum<EWorldPartitionStreamingPerformance>()->GetDisplayNameTextByValue((int64)NewStreamingPerformance).ToString());
-		
-		StreamingPerformance = NewStreamingPerformance;
+		EWorldPartitionStreamingPerformance NewStreamingPerformance = WorldPartition->RuntimeHash->GetStreamingPerformance(CellsToActivate);
+
+		if (StreamingPerformance != NewStreamingPerformance)
+		{
+			UE_LOG(LogWorldPartition, Log, TEXT("Streaming performance changed: %s -> %s"),
+				*StaticEnum<EWorldPartitionStreamingPerformance>()->GetDisplayNameTextByValue((int64)StreamingPerformance).ToString(),
+				*StaticEnum<EWorldPartitionStreamingPerformance>()->GetDisplayNameTextByValue((int64)NewStreamingPerformance).ToString());
+
+			StreamingPerformance = NewStreamingPerformance;
+		}
 	}
 
 #if !UE_BUILD_SHIPPING
@@ -574,10 +589,13 @@ void UWorldPartitionStreamingPolicy::SetCellsStateToLoaded(const TSet<const UWor
 	TRACE_CPUPROFILER_EVENT_SCOPE(UWorldPartitionStreamingPolicy::SetCellsStateToLoaded);
 
 	int32 MaxCellsToLoad = GetMaxCellsToLoad();
-
 	// Sort cells based on importance
 	TArray<const UWorldPartitionRuntimeCell*, TInlineAllocator<256>> SortedCells;
-	WorldPartition->RuntimeHash->SortStreamingCellsByImportance(ToLoadCells, StreamingSources, SortedCells);
+
+	if (WorldPartition->RuntimeHash)
+	{
+		WorldPartition->RuntimeHash->SortStreamingCellsByImportance(ToLoadCells, StreamingSources, SortedCells);
+	}
 
 	// Trigger cell loading. Depending on actual state of cell limit loading.
 	for (const UWorldPartitionRuntimeCell* Cell : SortedCells)
@@ -614,7 +632,10 @@ void UWorldPartitionStreamingPolicy::SetCellsStateToActivated(const TSet<const U
 
 	// Sort cells based on importance
 	TArray<const UWorldPartitionRuntimeCell*, TInlineAllocator<256>> SortedCells;
-	WorldPartition->RuntimeHash->SortStreamingCellsByImportance(ToActivateCells, StreamingSources, SortedCells);
+	if (WorldPartition->RuntimeHash)
+	{	
+		WorldPartition->RuntimeHash->SortStreamingCellsByImportance(ToActivateCells, StreamingSources, SortedCells);
+	}
 
 	// Trigger cell activation. Depending on actual state of cell limit loading.
 	for (const UWorldPartitionRuntimeCell* Cell : SortedCells)
@@ -695,8 +716,10 @@ bool UWorldPartitionStreamingPolicy::IsStreamingCompleted(EWorldPartitionRuntime
 	for (const FWorldPartitionStreamingQuerySource& QuerySource : QuerySources)
 	{
 		TSet<const UWorldPartitionRuntimeCell*> Cells;
-		WorldPartition->RuntimeHash->GetStreamingCells(QuerySource, Cells);
-
+		if (WorldPartition->RuntimeHash)
+		{
+			WorldPartition->RuntimeHash->GetStreamingCells(QuerySource, Cells);
+		}
 		for (const UWorldPartitionRuntimeCell* Cell : Cells)
 		{
 			EWorldPartitionRuntimeCellState CellState = Cell->GetCurrentState();
@@ -736,12 +759,16 @@ bool UWorldPartitionStreamingPolicy::IsStreamingCompleted(EWorldPartitionRuntime
 
 FVector2D UWorldPartitionStreamingPolicy::GetDrawRuntimeHash2DDesiredFootprint(const FVector2D& CanvasSize)
 {
-	return WorldPartition->RuntimeHash->GetDraw2DDesiredFootprint(CanvasSize);
+	if (WorldPartition->RuntimeHash)
+	{
+		return WorldPartition->RuntimeHash->GetDraw2DDesiredFootprint(CanvasSize);
+	}
+	return FVector2D::ZeroVector;
 }
 
 void UWorldPartitionStreamingPolicy::DrawRuntimeHash2D(class UCanvas* Canvas, const FVector2D& PartitionCanvasSize, const FVector2D& Offset)
 {
-	if (StreamingSources.Num() > 0)
+	if (StreamingSources.Num() > 0 && WorldPartition->RuntimeHash)
 	{
 		WorldPartition->RuntimeHash->Draw2D(Canvas, StreamingSources, PartitionCanvasSize, Offset);
 	}
@@ -750,7 +777,9 @@ void UWorldPartitionStreamingPolicy::DrawRuntimeHash2D(class UCanvas* Canvas, co
 void UWorldPartitionStreamingPolicy::DrawRuntimeHash3D()
 {
 	UWorld* World = WorldPartition->GetWorld();
-	if (World->GetNetMode() != NM_DedicatedServer && WorldPartition->IsInitialized())
+	if (World->GetNetMode() != NM_DedicatedServer
+		&& WorldPartition->IsInitialized()
+		&& WorldPartition->RuntimeHash)
 	{
 		WorldPartition->RuntimeHash->Draw3D(StreamingSources);
 	}
