PR #9507: Allow FUniqueNetId to be implicitly converted to an FUniqueNetIdRepl. (Contributed by [REDACTED_PERSON])

diff --git a/Engine/Source/Runtime/Engine/Classes/GameFramework/OnlineReplStructs.h b/Engine/Source/Runtime/Engine/Classes/GameFramework/OnlineReplStructs.h
--- a/Engine/Source/Runtime/Engine/Classes/GameFramework/OnlineReplStructs.h
+++ b/Engine/Source/Runtime/Engine/Classes/GameFramework/OnlineReplStructs.h
@@ -53,6 +53,11 @@ struct FUniqueNetIdRepl : public FUniqueNetIdWrapper
 		: FUniqueNetIdWrapper(InUniqueNetId)
 	{
 	}
+	
+	FUniqueNetIdRepl(const FUniqueNetId& InUniqueNetId)
+		: FUniqueNetIdWrapper(InUniqueNetId)
+	{
+	}
 
 	virtual ~FUniqueNetIdRepl() {}
 
