Don't load assets in UAssetManager::GatherPublicAssetsForPackage, use the package name from the asset data instead.

#rnx

diff --git a/Engine/Source/Runtime/Engine/Private/AssetManager.cpp b/Engine/Source/Runtime/Engine/Private/AssetManager.cpp
--- a/Engine/Source/Runtime/Engine/Private/AssetManager.cpp
+++ b/Engine/Source/Runtime/Engine/Private/AssetManager.cpp
@@ -4046,23 +4046,13 @@ void UAssetManager::GatherPublicAssetsForPackage(FName PackagePath, TArray<FName
 	GetAssetRegistry().EnumerateAssets(Filter, [&PackagesToCook](const FAssetData& AssetData)
 	{
 		// this package can be externally referenced; include it in the cook
-		if (const UPackage* const OuterPackage = AssetData.GetPackage())
+		if (!AssetData.PackageName.IsNone())
 		{
-			FName PackageFName = OuterPackage->GetLoadedPath().GetPackageFName();
-			if (!PackageFName.IsNone())
-			{
-				UE_LOG(LogAssetManager, Verbose,
-					TEXT("GatherPublicAssetsForPackage: Adding public package [%s] (instigator: [%s]"),
-					*PackageFName.ToString(),
-					*AssetData.AssetName.ToString());
-				PackagesToCook.AddUnique(PackageFName);
-			}
-			else
-			{
-				UE_LOG(LogAssetManager, Error,
-					TEXT("GatherPublicAssetsForPackage: Failed to resolve package name for asset [%s]"),
-					*AssetData.AssetName.ToString());
-			}
+			UE_LOG(LogAssetManager, Verbose,
+				TEXT("GatherPublicAssetsForPackage: Adding public package [%s] (instigator: [%s]"),
+				*AssetData.PackageName.ToString(),
+				*AssetData.AssetName.ToString());
+			PackagesToCook.AddUnique(AssetData.PackageName);
 		}
 		else
 		{
