World Partition: move all code assuming that actor descriptor containers are outered to a world partition object into the actor descriptor container.

#rnx

diff --git a/Engine/Source/Runtime/Engine/Private/WorldPartition/ActorDescContainer.cpp b/Engine/Source/Runtime/Engine/Private/WorldPartition/ActorDescContainer.cpp
--- a/Engine/Source/Runtime/Engine/Private/WorldPartition/ActorDescContainer.cpp
+++ b/Engine/Source/Runtime/Engine/Private/WorldPartition/ActorDescContainer.cpp
@@ -4,6 +4,7 @@
 
 #if WITH_EDITOR
 #include "Editor.h"
+#include "UObject/LinkerInstancingContext.h"
 #include "AssetRegistry/AssetRegistryModule.h"
 #include "WorldPartition/WorldPartition.h"
 #include "WorldPartition/WorldPartitionLog.h"
@@ -284,4 +285,26 @@ void UActorDescContainer::OnActorDescUpdated(FWorldPartitionActorDesc* ActorDesc
 		WorldPartition->OnActorDescUpdated(ActorDesc);
 	}
 }
+
+const FLinkerInstancingContext* UActorDescContainer::GetInstancingContext() const
+{
+	const FLinkerInstancingContext* InstancingContext = nullptr;
+
+	if (UWorldPartition* WorldPartition = GetTypedOuter<UWorldPartition>())
+	{
+		WorldPartition->GetInstancingContext(InstancingContext);
+	}
+
+	return InstancingContext;
+}
+
+const FTransform& UActorDescContainer::GetInstanceTransform() const
+{
+	if (UWorldPartition* WorldPartition = GetTypedOuter<UWorldPartition>())
+	{
+		return WorldPartition->GetInstanceTransform();
+	}
+
+	return FTransform::Identity;
+}
 #endif
diff --git a/Engine/Source/Runtime/Engine/Private/WorldPartition/WorldPartitionActorDesc.cpp b/Engine/Source/Runtime/Engine/Private/WorldPartition/WorldPartitionActorDesc.cpp
--- a/Engine/Source/Runtime/Engine/Private/WorldPartition/WorldPartitionActorDesc.cpp
+++ b/Engine/Source/Runtime/Engine/Private/WorldPartition/WorldPartitionActorDesc.cpp
@@ -15,7 +15,6 @@
 #include "UObject/UE5MainStreamObjectVersion.h"
 #include "UObject/UE5ReleaseStreamObjectVersion.h"
 #include "UObject/FortniteNCBranchObjectVersion.h"
-#include "WorldPartition/WorldPartition.h"
 #include "WorldPartition/WorldPartitionLog.h"
 #include "WorldPartition/ActorDescContainer.h"
 #include "WorldPartition/HLOD/HLODLayer.h"
@@ -494,14 +493,7 @@ AActor* FWorldPartitionActorDesc::Load() const
 	// Then, if the actor isn't loaded, load it
 	if (ActorPtr.IsExplicitlyNull())
 	{
-		const FLinkerInstancingContext* InstancingContext = nullptr;
-		if (Container)
-		{
-			if (UWorldPartition* WorldPartition = Container->GetTypedOuter<UWorldPartition>())
-			{
-				WorldPartition->GetInstancingContext(InstancingContext);
-			}
-		}
+		const FLinkerInstancingContext* InstancingContext = Container ? Container->GetInstancingContext() : nullptr;
 
 		UPackage* Package = nullptr;
 
diff --git a/Engine/Source/Runtime/Engine/Private/WorldPartition/WorldPartitionHandle.cpp b/Engine/Source/Runtime/Engine/Private/WorldPartition/WorldPartitionHandle.cpp
--- a/Engine/Source/Runtime/Engine/Private/WorldPartition/WorldPartitionHandle.cpp
+++ b/Engine/Source/Runtime/Engine/Private/WorldPartition/WorldPartitionHandle.cpp
@@ -78,12 +78,8 @@ void FWorldPartitionLoadingContext::FImmediate::RegisterActor(FWorldPartitionAct
 		UActorDescContainer* Container = ActorDesc->GetContainer();
 		check(Container);
 
-		const FTransform* ContainerTransformPtr = nullptr;
-		if (UWorldPartition* WorldPartition = Container->GetTypedOuter<UWorldPartition>())
-		{
-			const FTransform& ContainerTransform = WorldPartition->GetInstanceTransform();
-			ContainerTransformPtr = ContainerTransform.Equals(FTransform::Identity) ? nullptr : &ContainerTransform;
-		}
+		const FTransform& ContainerTransform = Container->GetInstanceTransform();
+		const FTransform* ContainerTransformPtr = ContainerTransform.Equals(FTransform::Identity) ? nullptr : &ContainerTransform;
 
 		Actor->GetLevel()->AddLoadedActor(Actor, ContainerTransformPtr);
 	}
@@ -97,12 +93,8 @@ void FWorldPartitionLoadingContext::FImmediate::UnregisterActor(FWorldPartitionA
 		UActorDescContainer* Container = ActorDesc->GetContainer();
 		check(Container);
 
-		const FTransform* ContainerTransformPtr = nullptr;
-		if (UWorldPartition* WorldPartition = Container->GetTypedOuter<UWorldPartition>())
-		{
-			const FTransform& ContainerTransform = WorldPartition->GetInstanceTransform();
-			ContainerTransformPtr = ContainerTransform.Equals(FTransform::Identity) ? nullptr : &ContainerTransform;
-		}
+		const FTransform& ContainerTransform = Container->GetInstanceTransform();
+		const FTransform* ContainerTransformPtr = ContainerTransform.Equals(FTransform::Identity) ? nullptr : &ContainerTransform;
 
 		Actor->GetLevel()->RemoveLoadedActor(Actor, ContainerTransformPtr);
 
@@ -121,12 +113,8 @@ FWorldPartitionLoadingContext::FDeferred::~FDeferred()
 	{
 		for (auto& [Container, ContainerOp] : LocalContainerOps)
 		{
-			const FTransform* ContainerTransformPtr = nullptr;
-			if (UWorldPartition* WorldPartition = Container->GetTypedOuter<UWorldPartition>())
-			{
-				const FTransform& ContainerTransform = WorldPartition->GetInstanceTransform();
-				ContainerTransformPtr = ContainerTransform.Equals(FTransform::Identity) ? nullptr : &ContainerTransform;
-			}
+			const FTransform& ContainerTransform = Container->GetInstanceTransform();
+			const FTransform* ContainerTransformPtr = ContainerTransform.Equals(FTransform::Identity) ? nullptr : &ContainerTransform;
 
 			auto CreateActorList = [](TArray<AActor*>& ActorList, const TSet<FWorldPartitionActorDesc*>& SourceList) -> ULevel*
 			{
diff --git a/Engine/Source/Runtime/Engine/Public/WorldPartition/ActorDescContainer.h b/Engine/Source/Runtime/Engine/Public/WorldPartition/ActorDescContainer.h
--- a/Engine/Source/Runtime/Engine/Public/WorldPartition/ActorDescContainer.h
+++ b/Engine/Source/Runtime/Engine/Public/WorldPartition/ActorDescContainer.h
@@ -8,6 +8,8 @@
 #include "WorldPartition/WorldPartitionHandle.h"
 #include "ActorDescContainer.generated.h"
 
+class FLinkerInstancingContext;
+
 UCLASS()
 class ENGINE_API UActorDescContainer : public UObject, public FActorDescList
 {
@@ -54,6 +56,9 @@ public:
 
 	DECLARE_MULTICAST_DELEGATE_OneParam(FActorDescContainerInitializeDelegate, UActorDescContainer*);
 	static FActorDescContainerInitializeDelegate OnActorDescContainerInitialized;
+
+	const FLinkerInstancingContext* GetInstancingContext() const;
+	const FTransform& GetInstanceTransform() const;
 #endif
 
 	UPROPERTY(Transient)
