Add a GetActorPackageName overload that takes a Base directory string instead of a LevelPackage to generate the ActorPackageName

diff --git a/Engine/Source/Runtime/Engine/Classes/Engine/Level.h b/Engine/Source/Runtime/Engine/Classes/Engine/Level.h
--- a/Engine/Source/Runtime/Engine/Classes/Engine/Level.h
+++ b/Engine/Source/Runtime/Engine/Classes/Engine/Level.h
@@ -1158,6 +1158,13 @@ public:
 	 */
 	static ENGINE_API FString GetActorPackageName(UPackage* InLevelPackage, EActorPackagingScheme ActorPackagingScheme, const FString& InActorPath);
 
+	/**
+	 * Get the package name for this actor
+	 * @param InActorPath the fully qualified actor path, in the format: 'Outermost.Outer.Name'
+	 * @return the package name
+	 */
+	static ENGINE_API FString GetActorPackageName(const FString& InBaseDir, EActorPackagingScheme ActorPackagingScheme, const FString& InActorPath);
+
 	/**
 	 * Get the folder containing the external actors for this level path
 	 * @param InLevelPackageName The package name to get the external actors path of
diff --git a/Engine/Source/Runtime/Engine/Private/Level.cpp b/Engine/Source/Runtime/Engine/Private/Level.cpp
--- a/Engine/Source/Runtime/Engine/Private/Level.cpp
+++ b/Engine/Source/Runtime/Engine/Private/Level.cpp
@@ -3335,6 +3335,11 @@ bool ULevel::HasAnyActorsOfType(UClass *SearchType)
 
 #if WITH_EDITOR
 FString ULevel::GetActorPackageName(UPackage* InLevelPackage, EActorPackagingScheme ActorPackagingScheme, const FString& InActorPath)
+{
+	return GetActorPackageName(GetExternalActorsPath(InLevelPackage), ActorPackagingScheme, InActorPath);
+}
+
+FString ULevel::GetActorPackageName(const FString& InBaseDir, EActorPackagingScheme ActorPackagingScheme, const FString& InActorPath)
 {
 	// Convert the actor path to lowercase to make sure we get the same hash for case insensitive file systems
 	FString ActorPath = InActorPath.ToLower();
@@ -3353,12 +3358,10 @@ FString ULevel::GetActorPackageName(UPackage* InLevelPackage, EActorPackagingSch
 	FString GuidBase36 = PackageGuid.ToString(EGuidFormats::Base36Encoded);
 	check(GuidBase36.Len());
 
-	FString BaseDir = GetExternalActorsPath(InLevelPackage);
-
 	TStringBuilderWithBuffer<TCHAR, NAME_SIZE> ActorPackageName;
 
 	uint32 FilenameOffset = 0;
-	ActorPackageName.Append(BaseDir);
+	ActorPackageName.Append(InBaseDir);
 	ActorPackageName.Append(TEXT("/"));
 
 	switch (ActorPackagingScheme)
