fix reinstancing in rewind debugger so that reinstanced anim instances will show up in the timeline view

diff --git a/Engine/Source/Runtime/Engine/Private/Animation/AnimInstance.cpp b/Engine/Source/Runtime/Engine/Private/Animation/AnimInstance.cpp
--- a/Engine/Source/Runtime/Engine/Private/Animation/AnimInstance.cpp
+++ b/Engine/Source/Runtime/Engine/Private/Animation/AnimInstance.cpp
@@ -3853,6 +3853,8 @@ void UAnimInstance::HandleObjectsReinstanced(const TMap<UObject*, UObject*>& Old
 		
 		if(bThisObjectWasReinstanced)
 		{
+			TRACE_OBJECT_LIFETIME_BEGIN(this);
+
 			RecalcRequiredBones();
 			
 			// Reinit proxy
