Add stat scopes for investigating ServerUpdateLevelVisibility
Avoid streaming level iteration in ValidateLevelVisibility when a linker is found
Add map of level names to destroyed net guids for processing destruction infos when level visibility changes

diff --git a/Engine/Source/Runtime/Engine/Classes/Engine/NetDriver.h b/Engine/Source/Runtime/Engine/Classes/Engine/NetDriver.h
--- a/Engine/Source/Runtime/Engine/Classes/Engine/NetDriver.h
+++ b/Engine/Source/Runtime/Engine/Classes/Engine/NetDriver.h
@@ -1087,6 +1087,16 @@ public:
 	 */
 	TMap<FNetworkGUID, TUniquePtr<FActorDestructionInfo>>	DestroyedStartupOrDormantActors;
 
+private:
+	/** Tracks the network guids in DestroyedStartupOrDormantActors above, but keyed on the streaming level name. */
+	TMap<FName, TSet<FNetworkGUID>> DestroyedStartupOrDormantActorsByLevel;
+
+public:
+	const TSet<FNetworkGUID>& GetDestroyedStartupOrDormantActors(const FName& LevelName)
+	{
+		return DestroyedStartupOrDormantActorsByLevel.FindOrAdd(LevelName);
+	}
+
 	/** The server adds an entry into this map for every startup actor that has been renamed, and will
 	 *  always map from current name to original name
 	 */
@@ -1921,9 +1931,9 @@ private:
 	FDelegateHandle PostGarbageCollectHandle;
 	void PostGarbageCollect();
 
-	FActorDestructionInfo* CreateDestructionInfo(UNetDriver* NetDriver, AActor* ThisActor, FActorDestructionInfo *DestructionInfo);
+	FActorDestructionInfo* CreateDestructionInfo(AActor* ThisActor, FActorDestructionInfo *DestructionInfo);
 
-	void CreateReplicatedStaticActorDestructionInfo(UNetDriver* NetDriver, ULevel* Level, const FReplicatedStaticActorDestructionInfo& Info);
+	void CreateReplicatedStaticActorDestructionInfo(ULevel* Level, const FReplicatedStaticActorDestructionInfo& Info);
 
 	void FlushActorDormancyInternal(AActor *Actor);
 
diff --git a/Engine/Source/Runtime/Engine/Private/NetConnection.cpp b/Engine/Source/Runtime/Engine/Private/NetConnection.cpp
--- a/Engine/Source/Runtime/Engine/Private/NetConnection.cpp
+++ b/Engine/Source/Runtime/Engine/Private/NetConnection.cpp
@@ -125,6 +125,8 @@ namespace UE_NetConnectionPrivate
 
 	const FValidateLevelVisibilityResult ValidateLevelVisibility(UWorld* World, const FUpdateLevelVisibilityLevelInfo& LevelVisibility)
 	{
+		QUICK_SCOPE_CYCLE_COUNTER(STAT_Net_ValidateLevelVisibility);
+
 		// Verify that we were passed a valid level name
 		// If we have a linker we know it has been loaded off disk successfully
 		// If we have a file it is fine too
@@ -136,9 +138,8 @@ namespace UE_NetConnectionPrivate
 		Result.Linker = FLinkerLoad::FindExistingLinkerForPackage(Result.Package);
 
 		Result.StreamingLevel = FLevelUtils::FindStreamingLevel(World, LevelVisibility.PackageName);
-		Result.bLevelExists = Result.Linker ||
-			(!FPackageName::IsMemoryPackage(PackageNameStr) && FPackageName::DoesPackageExist(LevelVisibility.FileName.ToString())) ||
-			Result.StreamingLevel != nullptr;
+		Result.bLevelExists = Result.Linker || Result.StreamingLevel != nullptr ||
+			(!FPackageName::IsMemoryPackage(PackageNameStr) && FPackageName::DoesPackageExist(LevelVisibility.FileName.ToString()));
 
 		return Result;
 	}
@@ -1451,6 +1452,8 @@ void UNetConnection::UpdateLevelVisibility(const FUpdateLevelVisibilityLevelInfo
 
 void UNetConnection::UpdateLevelVisibilityInternal(const FUpdateLevelVisibilityLevelInfo& LevelVisibility)
 {
+	QUICK_SCOPE_CYCLE_COUNTER(STAT_NetConnection_UpdateLevelVisibilityInternal)
+
 	using namespace UE_NetConnectionPrivate;
 
 	GNumClientUpdateLevelVisibility++;
@@ -1483,18 +1486,21 @@ void UNetConnection::UpdateLevelVisibilityInternal(const FUpdateLevelVisibilityL
 
 			UpdateCachedLevelVisibility(LevelVisibility.PackageName);
 
-			QUICK_USE_CYCLE_STAT(NetUpdateLevelVisibility_UpdateDormantActors, STATGROUP_Net);
-
-			// Any destroyed actors that were destroyed prior to the streaming level being unloaded for the client will not be in the connections
-			// destroyed actors list when the level is reloaded, so seek them out and add in
-			for (const auto& DestroyedPair : Driver->DestroyedStartupOrDormantActors)
 			{
-				if (DestroyedPair.Value->StreamingLevelName == LevelVisibility.PackageName)
+				QUICK_SCOPE_CYCLE_COUNTER(STAT_NetUpdateLevelVisibility_DestructionInfos);
+
+				// Any destroyed actors that were destroyed prior to the streaming level being unloaded for the client will not be in the connections
+				// destroyed actors list when the level is reloaded, so seek them out and add in
+				const TSet<FNetworkGUID>& DestructionInfoGuids = Driver->GetDestroyedStartupOrDormantActors(LevelVisibility.PackageName);
+				
+				for (const FNetworkGUID& DestroyedGuid : DestructionInfoGuids)
 				{
-					AddDestructionInfo(DestroyedPair.Value.Get());
+					AddDestructionInfo(Driver->DestroyedStartupOrDormantActors[DestroyedGuid].Get());
 				}
 			}
 
+			QUICK_SCOPE_CYCLE_COUNTER(STAT_NetUpdateLevelVisibility_UpdateDormantActors);
+
 			// Any dormant actor that has changes flushed or made before going dormant needs to be updated on the client 
 			// when the streaming level is loaded, so mark them active for this connection
 			UWorld* LevelWorld = nullptr;
@@ -1592,6 +1598,8 @@ void UNetConnection::UpdateLevelVisibilityInternal(const FUpdateLevelVisibilityL
 		}
 #endif // UE_WITH_IRIS
 
+		QUICK_SCOPE_CYCLE_COUNTER(STAT_NetUpdateLevelVisibility_LevelUnloadChannelClose);
+
 		// Close any channels now that have actors that were apart of the level the client just unloaded
 		for (auto It = ActorChannels.CreateIterator(); It; ++It)
 		{
diff --git a/Engine/Source/Runtime/Engine/Private/NetDriver.cpp b/Engine/Source/Runtime/Engine/Private/NetDriver.cpp
--- a/Engine/Source/Runtime/Engine/Private/NetDriver.cpp
+++ b/Engine/Source/Runtime/Engine/Private/NetDriver.cpp
@@ -2475,6 +2475,15 @@ void UNetDriver::Serialize( FArchive& Ar )
 			}
 		);
 
+		GRANULAR_NETWORK_MEMORY_TRACKING_TRACK("DestroyedStartupOrDormantActorsByLevel", 
+			DestroyedStartupOrDormantActorsByLevel.CountBytes(Ar);
+
+			for (const auto& DestroyedPair : DestroyedStartupOrDormantActorsByLevel)
+			{
+				DestroyedPair.Value.CountBytes(Ar);
+			}
+		);
+
 		GRANULAR_NETWORK_MEMORY_TRACKING_TRACK("RenamedStartupActors", RenamedStartupActors.CountBytes(Ar));
 
 		PRAGMA_DISABLE_DEPRECATION_WARNINGS
@@ -3136,21 +3145,21 @@ bool UNetDriver::Exec( UWorld* InWorld, const TCHAR* Cmd, FOutputDevice& Ar )
 	}
 }
 
-FActorDestructionInfo* UNetDriver::CreateDestructionInfo( UNetDriver* NetDriver, AActor* ThisActor, FActorDestructionInfo *DestructionInfo )
+FActorDestructionInfo* UNetDriver::CreateDestructionInfo(AActor* ThisActor, FActorDestructionInfo *DestructionInfo )
 {
 	if (DestructionInfo)
 	{
 		return DestructionInfo;
 	}
 
-	FNetworkGUID NetGUID = NetDriver->GuidCache->GetOrAssignNetGUID( ThisActor );
+	FNetworkGUID NetGUID = GuidCache->GetOrAssignNetGUID( ThisActor );
 	if (NetGUID.IsDefault())
 	{
 		UE_LOG(LogNet, Error, TEXT("CreateDestructionInfo got an invalid NetGUID for %s"), *ThisActor->GetName());
 		return nullptr;
 	}
 
-	TUniquePtr<FActorDestructionInfo>& NewInfoPtr = NetDriver->DestroyedStartupOrDormantActors.FindOrAdd( NetGUID );
+	TUniquePtr<FActorDestructionInfo>& NewInfoPtr = DestroyedStartupOrDormantActors.FindOrAdd(NetGUID);
 	if (NewInfoPtr.IsValid() == false)
 	{
 		NewInfoPtr = MakeUnique<FActorDestructionInfo>();
@@ -3164,7 +3173,7 @@ FActorDestructionInfo* UNetDriver::CreateDestructionInfo( UNetDriver* NetDriver,
 	NewInfo.PathName = ThisActor->GetName();
 
 	// Look for renamed actor now so we can clear it after the destroy is queued
-	FName RenamedPath = NetDriver->RenamedStartupActors.FindRef(ThisActor->GetFName());
+	FName RenamedPath = RenamedStartupActors.FindRef(ThisActor->GetFName());
 	if (RenamedPath != NAME_None)
 	{
 		NewInfo.PathName = RenamedPath.ToString();
@@ -3179,6 +3188,9 @@ FActorDestructionInfo* UNetDriver::CreateDestructionInfo( UNetDriver* NetDriver,
 		NewInfo.StreamingLevelName = NAME_None;
 	}
 
+	TSet<FNetworkGUID>& DestroyedGuidsForLevel = DestroyedStartupOrDormantActorsByLevel.FindOrAdd(NewInfo.StreamingLevelName);
+	DestroyedGuidsForLevel.Add(NetGUID);
+
 	NewInfo.Reason = EChannelCloseReason::Destroyed;
 
 	return &NewInfo;
@@ -3208,7 +3220,7 @@ void UNetDriver::NotifyActorDestroyed( AActor* ThisActor, bool IsSeamlessTravel
 		if (bShouldCreateDestructionInfo)
 		{
 			UE_LOG(LogNet, VeryVerbose, TEXT("NotifyActorDestroyed %s - StartupActor"), *ThisActor->GetPathName() );
-			DestructionInfo = CreateDestructionInfo( this, ThisActor, DestructionInfo);
+			DestructionInfo = CreateDestructionInfo(ThisActor, DestructionInfo);
 		}
 
 		const FNetworkObjectInfo* NetworkObjectInfo = GetNetworkObjectList().Find( ThisActor ).Get();
@@ -3233,7 +3245,7 @@ void UNetDriver::NotifyActorDestroyed( AActor* ThisActor, bool IsSeamlessTravel
 				{
 					// Make a new destruction info if necessary. It is necessary if the actor is dormant or recently dormant because
 					// even though the client knew about the actor at some point, it doesn't have a channel to handle destruction.
-					DestructionInfo = CreateDestructionInfo(this, ThisActor, DestructionInfo);
+					DestructionInfo = CreateDestructionInfo(ThisActor, DestructionInfo);
 					if (DestructionInfo)
 					{
 						Connection->AddDestructionInfo(DestructionInfo);
@@ -5524,14 +5536,14 @@ void UNetDriver::UpdateCrashContext()
 	}
 }
 
-void UNetDriver::CreateReplicatedStaticActorDestructionInfo(UNetDriver* NetDriver, ULevel* Level, const FReplicatedStaticActorDestructionInfo& Info)
+void UNetDriver::CreateReplicatedStaticActorDestructionInfo(ULevel* Level, const FReplicatedStaticActorDestructionInfo& Info)
 {
-	check(NetDriver && Level);
+	check(Level);
 
 #if UE_WITH_IRIS
-	if (NetDriver->ReplicationSystem)
+	if (ReplicationSystem)
 	{
-		if (UActorReplicationBridge* Bridge = NetDriver->ReplicationSystem->GetReplicationBridgeAs<UActorReplicationBridge>())
+		if (UActorReplicationBridge* Bridge = ReplicationSystem->GetReplicationBridgeAs<UActorReplicationBridge>())
 		{			
 			// Add explicit destruction info for this object
 			UReplicationBridge::FEndReplicationParameters Params;
@@ -5545,16 +5557,16 @@ void UNetDriver::CreateReplicatedStaticActorDestructionInfo(UNetDriver* NetDrive
 	}
 #endif // UE_WITH_IRIS
 
-	FNetworkGUID NetGUID = NetDriver->GuidCache->AssignNewNetGUIDFromPath_Server( Info.PathName.ToString(), Info.ObjOuter.Get(), Info.ObjClass );
+	FNetworkGUID NetGUID = GuidCache->AssignNewNetGUIDFromPath_Server(Info.PathName.ToString(), Info.ObjOuter.Get(), Info.ObjClass);
 	if (NetGUID.IsDefault())
 	{
 		UE_LOG(LogNet, Error, TEXT("CreateReplicatedStaticActorDestructionInfo got an invalid NetGUID for %s"), *Info.PathName.ToString());
 		return;
 	}
 
-	UE_LOG(LogNet, VeryVerbose, TEXT("CreateReplicatedStaticActorDestructionInfo %s %s %s %s"), *NetDriver->GetName(), *Level->GetName(), *Info.PathName.ToString(), *NetGUID.ToString());
+	UE_LOG(LogNet, VeryVerbose, TEXT("CreateReplicatedStaticActorDestructionInfo %s %s %s %s"), *GetName(), *Level->GetName(), *Info.PathName.ToString(), *NetGUID.ToString());
 	
-	TUniquePtr<FActorDestructionInfo>& NewInfoPtr = NetDriver->DestroyedStartupOrDormantActors.FindOrAdd( NetGUID );
+	TUniquePtr<FActorDestructionInfo>& NewInfoPtr = DestroyedStartupOrDormantActors.FindOrAdd(NetGUID);
 	if (NewInfoPtr.IsValid() == false)
 	{
 		NewInfoPtr = MakeUnique<FActorDestructionInfo>();
@@ -5568,7 +5580,7 @@ void UNetDriver::CreateReplicatedStaticActorDestructionInfo(UNetDriver* NetDrive
 	NewInfo.PathName = Info.PathName.ToString();
 
 	// Look for renamed actor now so we can clear it after the destroy is queued
-	FName RenamedPath = NetDriver->RenamedStartupActors.FindRef(Info.PathName);
+	FName RenamedPath = RenamedStartupActors.FindRef(Info.PathName);
 	if (RenamedPath != NAME_None)
 	{
 		NewInfo.PathName = RenamedPath.ToString();
@@ -5583,6 +5595,9 @@ void UNetDriver::CreateReplicatedStaticActorDestructionInfo(UNetDriver* NetDrive
 		NewInfo.StreamingLevelName = NAME_None;
 	}
 
+	TSet<FNetworkGUID>& DestroyedGuidsForLevel = DestroyedStartupOrDormantActorsByLevel.FindOrAdd(NewInfo.StreamingLevelName);
+	DestroyedGuidsForLevel.Add(NetGUID);
+
 	NewInfo.Reason = EChannelCloseReason::Destroyed;
 }
 
@@ -5601,7 +5616,7 @@ void UNetDriver::InitDestroyedStartupActors()
 				const TArray<FReplicatedStaticActorDestructionInfo>& DestroyedReplicatedStaticActors = Level->GetDestroyedReplicatedStaticActors();
 				for(const FReplicatedStaticActorDestructionInfo& Info : DestroyedReplicatedStaticActors)
 				{
-					CreateReplicatedStaticActorDestructionInfo(this, Level, Info);
+					CreateReplicatedStaticActorDestructionInfo(Level, Info);
 				}
 			}
 		}
@@ -5754,6 +5769,7 @@ void UNetDriver::SetWorld(class UWorld* InWorld)
 void UNetDriver::ResetGameWorldState()
 {
 	DestroyedStartupOrDormantActors.Empty();
+	DestroyedStartupOrDormantActorsByLevel.Empty();
 	RenamedStartupActors.Empty();
 
 	if ( NetCache.IsValid() )
diff --git a/Engine/Source/Runtime/Engine/Private/PlayerController.cpp b/Engine/Source/Runtime/Engine/Private/PlayerController.cpp
--- a/Engine/Source/Runtime/Engine/Private/PlayerController.cpp
+++ b/Engine/Source/Runtime/Engine/Private/PlayerController.cpp
@@ -366,6 +366,8 @@ void APlayerController::ClientFlushLevelStreaming_Implementation()
 
 void APlayerController::ServerUpdateLevelVisibility_Implementation(const FUpdateLevelVisibilityLevelInfo& LevelVisibility)
 {
+	QUICK_SCOPE_CYCLE_COUNTER(STAT_ServerUpdateLevelVisibility);
+
 	UNetConnection* Connection = Cast<UNetConnection>(Player);
 	if (Connection != NULL)
 	{
@@ -386,6 +388,8 @@ void APlayerController::ServerUpdateLevelVisibility_Implementation(const FUpdate
 
 bool APlayerController::ServerUpdateLevelVisibility_Validate(const FUpdateLevelVisibilityLevelInfo& LevelVisibility)
 {
+	QUICK_SCOPE_CYCLE_COUNTER(STAT_ServerUpdateLevelVisibility_Validate);
+
 	RPC_VALIDATE(LevelVisibility.PackageName.IsValid());
 
 	FText Reason;
