#ue5 - Fix type truncation warnings in ComponentVisualizers module.

diff --git a/Engine/Source/Editor/ComponentVisualizers/Private/PointLightComponentVisualizer.cpp b/Engine/Source/Editor/ComponentVisualizers/Private/PointLightComponentVisualizer.cpp
--- a/Engine/Source/Editor/ComponentVisualizers/Private/PointLightComponentVisualizer.cpp
+++ b/Engine/Source/Editor/ComponentVisualizers/Private/PointLightComponentVisualizer.cpp
@@ -96,14 +96,14 @@ namespace TextureLightProfileVisualizerImpl
 		FVector LocalToLight = InvLightTransform.TransformVector( ToLight );
 
 		// -1..1
-		float DotProd = FVector::DotProduct(ToLight, LightDirection);
+		double DotProd = FVector::DotProduct(ToLight, LightDirection);
 		// -PI..PI (this distortion could be put into the texture but not without quality loss or more memory)
-		float Angle = FMath::Asin(DotProd);
+		float Angle = (float)FMath::Asin(DotProd);
 		// 0..1
-		float NormAngle = Angle / PI + 0.5f;
+		float NormAngle = Angle / UE_PI + 0.5f;
 
-		float TangentAngle = FMath::Atan2( -LocalToLight.Z, -LocalToLight.Y ); // -Y represents 0/360 horizontal angle and we're rotating counter-clockwise
-		float NormTangentAngle = TangentAngle / (PI * 2.f) + 0.5f;
+		float TangentAngle = (float)FMath::Atan2( -LocalToLight.Z, -LocalToLight.Y ); // -Y represents 0/360 horizontal angle and we're rotating counter-clockwise
+		float NormTangentAngle = (TangentAngle / (UE_PI * 2.f) + 0.5f);
 
 		return FilterLightProfile( TextureLightProfileData, NormAngle, NormTangentAngle );
 	}
diff --git a/Engine/Source/Editor/ComponentVisualizers/Private/SplineComponentVisualizer.cpp b/Engine/Source/Editor/ComponentVisualizers/Private/SplineComponentVisualizer.cpp
--- a/Engine/Source/Editor/ComponentVisualizers/Private/SplineComponentVisualizer.cpp
+++ b/Engine/Source/Editor/ComponentVisualizers/Private/SplineComponentVisualizer.cpp
@@ -465,12 +465,12 @@ void FSplineComponentVisualizer::AddReferencedObjects(FReferenceCollector& Colle
 	}
 }
 
-static float GetDashSize(const FSceneView* View, const FVector& Start, const FVector& End, float Scale)
+static double GetDashSize(const FSceneView* View, const FVector& Start, const FVector& End, float Scale)
 {
-	const float StartW = View->WorldToScreen(Start).W;
-	const float EndW = View->WorldToScreen(End).W;
+	const double StartW = View->WorldToScreen(Start).W;
+	const double EndW = View->WorldToScreen(End).W;
 
-	const float WLimit = 10.0f;
+	const double WLimit = 10.0f;
 	if (StartW > WLimit || EndW > WLimit)
 	{
 		return FMath::Max(StartW, EndW) * Scale;
@@ -622,7 +622,7 @@ void FSplineComponentVisualizer::DrawVisualization(const UActorComponent* Compon
 				// For constant interpolation - don't draw ticks - just draw dotted line.
 				if (SplineInfo.Points[KeyIdx - 1].InterpMode == CIM_Constant)
 				{
-					const float DashSize = GetDashSize(View, OldKeyPos, NewKeyPos, 0.03f);
+					const double DashSize = GetDashSize(View, OldKeyPos, NewKeyPos, 0.03f);
 					if (DashSize > 0.0f)
 					{
 						DrawDashedLine(PDI, OldKeyPos, NewKeyPos, LineColor, DashSize, SDPG_World);
@@ -648,7 +648,7 @@ void FSplineComponentVisualizer::DrawVisualization(const UActorComponent* Compon
 					// Then draw a line for each substep.
 					constexpr int32 NumSteps = 20;
 					constexpr float PartialGradientProportion = 0.75f;
-					constexpr int32 PartialNumSteps = NumSteps * PartialGradientProportion;
+					constexpr int32 PartialNumSteps = (int32)(NumSteps * PartialGradientProportion);
 					const float SegmentLineThickness = GetDefault<ULevelEditorViewportSettings>()->SplineLineThicknessAdjustment;
 
 					for (int32 StepIdx = 1; StepIdx <= NumSteps; StepIdx++)
@@ -924,7 +924,7 @@ bool FSplineComponentVisualizer::VisProxyHandleClick(FEditorViewportClient* InVi
 				float SubsegmentStartKey = static_cast<float>(SelectedSegmentIndex);
 				FVector SubsegmentStart = SplineComp->GetLocationAtSplineInputKey(SubsegmentStartKey, ESplineCoordinateSpace::World);
 
-				float ClosestDistance = TNumericLimits<float>::Max();
+				double ClosestDistance = TNumericLimits<double>::Max();
 				FVector BestLocation = SubsegmentStart;
 
 				for (int32 Step = 1; Step < NumSubdivisions; Step++)
@@ -936,7 +936,7 @@ bool FSplineComponentVisualizer::VisProxyHandleClick(FEditorViewportClient* InVi
 					FVector RayClosest;
 					FMath::SegmentDistToSegmentSafe(SubsegmentStart, SubsegmentEnd, Click.GetOrigin(), Click.GetOrigin() + Click.GetDirection() * 50000.0f, SplineClosest, RayClosest);
 
-					const float Distance = FVector::DistSquared(SplineClosest, RayClosest);
+					const double Distance = FVector::DistSquared(SplineClosest, RayClosest);
 					if (Distance < ClosestDistance)
 					{
 						ClosestDistance = Distance;
@@ -1661,7 +1661,7 @@ bool FSplineComponentVisualizer::HandleSnapTo(const bool bInAlign, const bool bI
 						EditedPoint.InterpMode = CIM_CurveUser;
 
 						// Get delta rotation between up vector and hit normal
-						FVector WorldUpVector = SplineComp->GetUpVectorAtSplineInputKey(KeyIdx, ESplineCoordinateSpace::World);
+						FVector WorldUpVector = SplineComp->GetUpVectorAtSplineInputKey((float)KeyIdx, ESplineCoordinateSpace::World);
 						FQuat DeltaRotate = FQuat::FindBetweenNormals(WorldUpVector, Hit.Normal);
 
 						// Rotate tangent according to delta rotation
@@ -1725,17 +1725,17 @@ void FSplineComponentVisualizer::OnSnapKeyToNearestSplinePoint(ESplineComponentS
 	FInterpCurvePoint<FVector>& EditedPosition = SplineComp->GetSplinePointsPosition().Points[LastKeyIndexSelected];
 	const FVector WorldPos = SplineComp->GetComponentTransform().TransformPosition(EditedPosition.OutVal); // convert local-space position to world-space
 
-	float NearestDistanceSquared = 0.0f;
+	double NearestDistanceSquared = 0.0f;
 	USplineComponent* NearestSplineComp = nullptr;
 	int32 NearestKeyIndex = INDEX_NONE;
 
-	static const float SnapTol = 5000.0f;
-	float SnapTolSquared = SnapTol * SnapTol;
+	static const double SnapTol = 5000.0f;
+	double SnapTolSquared = SnapTol * SnapTol;
 
 	auto UpdateNearestKey = [WorldPos, SnapTolSquared, &NearestDistanceSquared, &NearestSplineComp, &NearestKeyIndex](USplineComponent* InSplineComp, int InKeyIdx)
 	{
 		const FVector TestKeyWorldPos = InSplineComp->GetLocationAtSplinePoint(InKeyIdx, ESplineCoordinateSpace::World);
-		float TestDistanceSquared = FVector::DistSquared(TestKeyWorldPos, WorldPos);
+		double TestDistanceSquared = FVector::DistSquared(TestKeyWorldPos, WorldPos);
 
 		if (TestDistanceSquared < SnapTolSquared && (NearestKeyIndex == INDEX_NONE || TestDistanceSquared < NearestDistanceSquared))
 		{
@@ -1916,15 +1916,15 @@ void FSplineComponentVisualizer::SnapKeyToTransform(const ESplineComponentSnapMo
 		if (InSnapMode == ESplineComponentSnapMode::AlignPerpendicularToTangent)
 		{
 			// Rotate tangent by 90 degrees
-			const FQuat DeltaRotate(WorldUpVector, HALF_PI);
+			const FQuat DeltaRotate(WorldUpVector, UE_HALF_PI);
 			NewTangent = DeltaRotate.RotateVector(NewTangent);
 		}
 
 		const FVector Tangent = SplineComp->GetComponentTransform().GetRotation().RotateVector(EditedPosition.ArriveTangent); // convert local-space tangent vectors to world-space
 
 		// Swap the tangents if they are not pointing in the same general direction
-		float CurrentAngle = FMath::Acos(FVector::DotProduct(Tangent, NewTangent) / Tangent.Size());
-		if (CurrentAngle > HALF_PI)
+		double CurrentAngle = FMath::Acos(FVector::DotProduct(Tangent, NewTangent) / Tangent.Size());
+		if (CurrentAngle > UE_HALF_PI)
 		{
 			NewTangent = SplineComp->GetComponentTransform().GetRotation().Inverse().RotateVector(NewTangent * -1.0f) * Tangent.Size(); // convert world-space tangent vectors into local-space
 		}
@@ -2266,7 +2266,7 @@ bool FSplineComponentVisualizer::DuplicateKeyForAltDrag(const FVector& InDrag)
 
 	// When dragging from end point, maximum angle is 60 degrees from attached segment
 	// to determine whether to split existing segment or create a new point
-	static const float Angle60 = 1.0472;
+	static const double Angle60 = 1.0472;
 
 	// Insert duplicates into the list, highest index first, so that the lower indices remain the same
 	FInterpCurveVector& SplinePosition = SplineComp->GetSplinePointsPosition();
@@ -2277,7 +2277,7 @@ bool FSplineComponentVisualizer::DuplicateKeyForAltDrag(const FVector& InDrag)
 
 	// Determine direction to insert new point				
 	bool bHasPrevKey = SplineComp->IsClosedLoop() || CurrentIndex > 0;
-	float PrevAngle = 0.0f;
+	double PrevAngle = 0.0f;
 	if (bHasPrevKey)
 	{
 		// Wrap index around for closed-looped splines
@@ -2295,7 +2295,7 @@ bool FSplineComponentVisualizer::DuplicateKeyForAltDrag(const FVector& InDrag)
 	}
 
 	bool bHasNextKey = SplineComp->IsClosedLoop() || CurrentIndex + 1 < NumPoints;
-	float NextAngle = 0.0f;
+	double NextAngle = 0.0f;
 	if (bHasNextKey)
 	{
 		// Wrap index around for closed-looped splines
@@ -2452,21 +2452,21 @@ void FSplineComponentVisualizer::SplitSegment(const FVector& InWorldPos, int32 I
 	int32 SourceIndex = bCopyFromSegmentBeginIndex ? SegmentBeginIndex : SegmentEndIndex;
 
 	FInterpCurvePoint<FVector> NewPoint(
-		SegmentSplitIndex,
+		(float)SegmentSplitIndex,
 		SplinePos,
 		FVector::ZeroVector,
 		FVector::ZeroVector,
 		SplinePosition.Points[SourceIndex].InterpMode);
 
 	FInterpCurvePoint<FQuat> NewRotPoint(
-		SegmentSplitIndex,
+		(float)SegmentSplitIndex,
 		NewRot,
 		FQuat::Identity,
 		FQuat::Identity,
 		CIM_CurveAuto);
 
 	FInterpCurvePoint<FVector> NewScalePoint(
-		SegmentSplitIndex,
+		(float)SegmentSplitIndex,
 		NewScale,
 		FVector::ZeroVector,
 		FVector::ZeroVector,
@@ -2642,21 +2642,21 @@ void FSplineComponentVisualizer::AddSegment(const FVector& InWorldPos, bool bApp
 	FVector LocalPos = SplineComp->GetComponentTransform().InverseTransformPosition(InWorldPos);
 
 	FInterpCurvePoint<FVector> NewPoint(
-		NewKeyIdx,
+		(float)NewKeyIdx,
 		LocalPos,
 		FVector::ZeroVector,
 		FVector::ZeroVector,
 		SplinePosition.Points[KeyIdx].InterpMode);
 
 	FInterpCurvePoint<FQuat> NewRotPoint(
-		NewKeyIdx,
+		(float)NewKeyIdx,
 		SplineRotation.Points[KeyIdx].OutVal,
 		FQuat::Identity,
 		FQuat::Identity,
 		CIM_CurveAuto);
 
 	FInterpCurvePoint<FVector> NewScalePoint(
-		NewKeyIdx,
+		(float)NewKeyIdx,
 		SplineScale.Points[KeyIdx].OutVal,
 		FVector::ZeroVector,
 		FVector::ZeroVector,
diff --git a/Engine/Source/Editor/ComponentVisualizers/Private/SplineGeneratorPanel.cpp b/Engine/Source/Editor/ComponentVisualizers/Private/SplineGeneratorPanel.cpp
--- a/Engine/Source/Editor/ComponentVisualizers/Private/SplineGeneratorPanel.cpp
+++ b/Engine/Source/Editor/ComponentVisualizers/Private/SplineGeneratorPanel.cpp
@@ -180,7 +180,7 @@ double CalcTangentMultiplier(const float InRadius, const float InRotInc)
 	// Calculate the difference between the actual interpolated midpoint and expected interpolated midpoint
 	const FVector ActualVal = FMath::CubicInterp(P0, T0, P1, T1, A);
 	const FVector ExpectedVal = P0.RotateAngleAxis(InRotInc * A, FVector::UpVector);
-	const float Diff = (ActualVal.X - ExpectedVal.X);
+	const double Diff = (ActualVal.X - ExpectedVal.X);
 
 	// Do a partial calculation of the cubic interpolation equation
 	static constexpr double C1 = (A3 - (2 * A2) + A), C2 = (A3 - A2);
@@ -420,13 +420,13 @@ void ULineSplineGenerator::BuildCurve()
 	SelectedSplineComponent->SplineCurves = CachedSplineCurves;
 
 	const bool bPrepend = !(ShapeAddMode & (EShapeAddMode::AppendAfter | EShapeAddMode::InsertAfter));
-	const float LengthFlip = bPrepend ? -1.f : 1.f;
+	const double LengthFlip = bPrepend ? -1.f : 1.f;
 
 	// Find starting features
 	const FVector StartPoint = SelectedSplineComponent->GetLocationAtSplinePoint(GetAddIndex(-1), ESplineCoordinateSpace::Local);
 	const FQuat StartQuat = SelectedSplineComponent->GetQuaternionAtSplinePoint(GetAddIndex(-1), ESplineCoordinateSpace::Local);
 
-	float LineInc = 0.0f;
+	double LineInc = 0.0f;
 	FVector LineDir = FVector::ZeroVector;
 	
 	bEnableUpToNextPoint = false;
@@ -446,18 +446,18 @@ void ULineSplineGenerator::BuildCurve()
 		const FVector Diff = NextPoint - StartPoint;
 		Length = Diff.Size();
 		LineDir = Diff.GetSafeNormal();
-		LineInc = Diff.Size() / float(NumberOfPoints + 1);
+		LineInc = Diff.Size() / double(NumberOfPoints + 1);
 	}
 	else
 	{
-		LineInc = Length / float(NumberOfPoints);
+		LineInc = Length / double(NumberOfPoints);
 		LineDir = StartQuat.GetForwardVector() * LengthFlip;
 	}
 
 	for (int32 Index = 0; Index < NumberOfPoints; Index++)
 	{
 		const int32 AddIdx = GetAddIndex(Index);
-		const FVector NewPoint = StartPoint + (LineInc * float(Index + 1)) * LineDir;
+		const FVector NewPoint = StartPoint + (LineInc * double(Index + 1)) * LineDir;
 		SelectedSplineComponent->AddSplinePointAtIndex(NewPoint, AddIdx, ESplineCoordinateSpace::Local);
 	}
 }
diff --git a/Engine/Source/Editor/ComponentVisualizers/Private/SplineMeshComponentVisualizer.cpp b/Engine/Source/Editor/ComponentVisualizers/Private/SplineMeshComponentVisualizer.cpp
--- a/Engine/Source/Editor/ComponentVisualizers/Private/SplineMeshComponentVisualizer.cpp
+++ b/Engine/Source/Editor/ComponentVisualizers/Private/SplineMeshComponentVisualizer.cpp
@@ -75,9 +75,9 @@ void FSplineMeshComponentVisualizer::DrawVisualization(const UActorComponent* Co
 			PDI->SetHitProxy(NULL);
 			const float DefaultDashSize = 5.f;
 			const int32 MaxDashDrawCount = 65536;
-			const float LineSize = TangentWorldDirection.Size();
-			int32 DrawCount = FMath::CeilToInt(LineSize / (2.f * DefaultDashSize));
-			float DashSize = (DrawCount > MaxDashDrawCount) ? LineSize / (2.f * MaxDashDrawCount) : DefaultDashSize;
+			const double LineSize = TangentWorldDirection.Size();
+			int64 DrawCount = FMath::CeilToInt64(LineSize / (2.f * DefaultDashSize));
+			double DashSize = (DrawCount > MaxDashDrawCount) ? LineSize / (2.f * MaxDashDrawCount) : DefaultDashSize;
 			DrawDashedLine(PDI, KeyPos, KeyPos + TangentWorldDirection, Color, DashSize, SDPG_Foreground);
 			DrawDashedLine(PDI, KeyPos, KeyPos - TangentWorldDirection, Color, DashSize, SDPG_Foreground);
 			PDI->SetHitProxy(new HSplineMeshTangentHandleProxy(Component, PointIndex, false));
@@ -331,11 +331,11 @@ bool FSplineMeshComponentVisualizer::HandleInputDelta(FEditorViewportClient* Vie
 				// Break tangent into direction and length so we can change its scale (the 'tension')
 				// independently of its direction.
 				FVector Direction;
-				float Length;
+				double Length;
 				KeyTangent.ToDirectionAndLength(Direction, Length);
 
 				// Figure out which component has changed, and use it
-				float DeltaScaleValue = (DeltaScale.X != 0.0f) ? DeltaScale.X : ((DeltaScale.Y != 0.0f) ? DeltaScale.Y : DeltaScale.Z);
+				double DeltaScaleValue = (DeltaScale.X != 0.0f) ? DeltaScale.X : ((DeltaScale.Y != 0.0f) ? DeltaScale.Y : DeltaScale.Z);
 
 				// Change scale, avoiding singularity by never allowing a scale of 0, hence preserving direction.
 				Length += DeltaScaleValue * 10.0f;
diff --git a/Engine/Source/Editor/ComponentVisualizers/Private/SpringComponentVisualizer.cpp b/Engine/Source/Editor/ComponentVisualizers/Private/SpringComponentVisualizer.cpp
--- a/Engine/Source/Editor/ComponentVisualizers/Private/SpringComponentVisualizer.cpp
+++ b/Engine/Source/Editor/ComponentVisualizers/Private/SpringComponentVisualizer.cpp
@@ -23,16 +23,16 @@ void FSpringComponentVisualizer::DrawVisualization( const UActorComponent* Compo
 	{
 		//interpolate color by compression
 		const float CompressionNormalized = SpringComp->GetNormalizedCompressionScalar();
-		float R = FMath::Lerp(RestColor.R, CompressedColor.R, CompressionNormalized);
-		float G = FMath::Lerp(RestColor.G, CompressedColor.G, CompressionNormalized);
-		float B = FMath::Lerp(RestColor.B, CompressedColor.B, CompressionNormalized);
+		uint8 R = FMath::Lerp(RestColor.R, CompressedColor.R, CompressionNormalized);
+		uint8 G = FMath::Lerp(RestColor.G, CompressedColor.G, CompressionNormalized);
+		uint8 B = FMath::Lerp(RestColor.B, CompressedColor.B, CompressionNormalized);
 		const FColor CurrentColor = SpringComp->IsActive() ? FColor(R, G, B) : DisabledColor;
 		
 		//draw capsule for spring sweep
 		const FTransform& WorldTM = SpringComp->GetComponentToWorld();
 		const FVector SpringStart = WorldTM.GetLocation();
 		const FVector SpringEnd = SpringComp->GetSpringCurrentEndPoint();
-		const float HalfHeight = (SpringEnd - SpringStart).Size() * 0.5f + SpringComp->SpringRadius;
+		const double HalfHeight = (SpringEnd - SpringStart).Size() * 0.5f + SpringComp->SpringRadius;
 		DrawWireCapsule(PDI, FMath::Lerp(SpringStart, SpringEnd, 0.5f), WorldTM.GetUnitAxis(EAxis::Z), WorldTM.GetUnitAxis(EAxis::Y), WorldTM.GetUnitAxis(EAxis::X), CurrentColor, SpringComp->SpringRadius, HalfHeight, 25, SDPG_World);
 	}
 }
diff --git a/Engine/Source/Editor/ComponentVisualizers/Public/AttenuatedComponentVisualizer.h b/Engine/Source/Editor/ComponentVisualizers/Public/AttenuatedComponentVisualizer.h
--- a/Engine/Source/Editor/ComponentVisualizers/Public/AttenuatedComponentVisualizer.h
+++ b/Engine/Source/Editor/ComponentVisualizers/Public/AttenuatedComponentVisualizer.h
@@ -72,8 +72,8 @@ public:
 
 						if (ShapeDetails.Falloff > 0.f || ShapeDetails.Extents.Z > 0.f)
 						{
-							float ConeRadius = ShapeDetails.Extents.X + ShapeDetails.Falloff + ShapeDetails.ConeOffset;
-							float ConeAngle = ShapeDetails.Extents.Y + ShapeDetails.Extents.Z;
+							double ConeRadius = ShapeDetails.Extents.X + ShapeDetails.Falloff + ShapeDetails.ConeOffset;
+							double ConeAngle = ShapeDetails.Extents.Y + ShapeDetails.Extents.Z;
 							DrawWireSphereCappedCone(PDI, Origin, ConeRadius, ConeAngle, 16, 4, 10, OuterRadiusColor, SDPG_World);
 
 							ConeRadius = ShapeDetails.Extents.X + ShapeDetails.ConeOffset;
@@ -82,8 +82,8 @@ public:
 						}
 						else
 						{
-							const float ConeRadius = ShapeDetails.Extents.X + ShapeDetails.ConeOffset;
-							const float ConeAngle = ShapeDetails.Extents.Y;
+							const double ConeRadius = ShapeDetails.Extents.X + ShapeDetails.ConeOffset;
+							const double ConeAngle = ShapeDetails.Extents.Y;
 							DrawWireSphereCappedCone(PDI, Origin, ConeRadius, ConeAngle, 16, 4, 10, OuterRadiusColor, SDPG_World );
 						}
 
diff --git a/Engine/Source/Editor/ComponentVisualizers/Public/SplineGeneratorPanel.h b/Engine/Source/Editor/ComponentVisualizers/Public/SplineGeneratorPanel.h
--- a/Engine/Source/Editor/ComponentVisualizers/Public/SplineGeneratorPanel.h
+++ b/Engine/Source/Editor/ComponentVisualizers/Public/SplineGeneratorPanel.h
@@ -254,7 +254,7 @@ public:
 	int32 NumberOfPoints;
 
 	UPROPERTY(Transient, EditAnywhere, Category = ShapeParameters, Meta = (ClampMin = "0", EditCondition = "!bUpToNextPoint"))
-	float Length;
+	double Length;
 
 	UPROPERTY()
 	bool bEnableUpToNextPoint;
diff --git a/Engine/Source/Runtime/Engine/Private/PrimitiveDrawingUtils.cpp b/Engine/Source/Runtime/Engine/Private/PrimitiveDrawingUtils.cpp
--- a/Engine/Source/Runtime/Engine/Private/PrimitiveDrawingUtils.cpp
+++ b/Engine/Source/Runtime/Engine/Private/PrimitiveDrawingUtils.cpp
@@ -534,7 +534,7 @@ void DrawCone(FPrimitiveDrawInterface* PDI,const FMatrix& ConeToWorld, float Ang
 }
 
 
-void BuildCylinderVerts(const FVector& Base, const FVector& XAxis, const FVector& YAxis, const FVector& ZAxis, float Radius, float HalfHeight, uint32 Sides, TArray<FDynamicMeshVertex>& OutVerts, TArray<uint32>& OutIndices)
+void BuildCylinderVerts(const FVector& Base, const FVector& XAxis, const FVector& YAxis, const FVector& ZAxis, double Radius, double HalfHeight, uint32 Sides, TArray<FDynamicMeshVertex>& OutVerts, TArray<uint32>& OutIndices)
 {
 	const float	AngleDelta = 2.0f * UE_PI / Sides;
 	FVector	LastVertex = Base + XAxis * Radius;
@@ -637,10 +637,10 @@ void BuildCylinderVerts(const FVector& Base, const FVector& XAxis, const FVector
 
 }
 
-void GetCylinderMesh(const FVector& Start, const FVector& End, float Radius, int32 Sides, const FMaterialRenderProxy* MaterialInstance, uint8 DepthPriority, int32 ViewIndex, FMeshElementCollector& Collector, HHitProxy* HitProxy)
+void GetCylinderMesh(const FVector& Start, const FVector& End, double Radius, int32 Sides, const FMaterialRenderProxy* MaterialInstance, uint8 DepthPriority, int32 ViewIndex, FMeshElementCollector& Collector, HHitProxy* HitProxy)
 {
 	FVector Dir = End - Start;
-	float Length = Dir.Size();
+	double Length = Dir.Size();
 
 	if (Length > UE_SMALL_NUMBER)
 	{
@@ -654,12 +654,12 @@ void GetCylinderMesh(const FVector& Start, const FVector& End, float Radius, int
 }
 
 void GetCylinderMesh(const FVector& Base, const FVector& XAxis, const FVector& YAxis, const FVector& ZAxis,
-				  float Radius, float HalfHeight, int32 Sides, const FMaterialRenderProxy* MaterialRenderProxy, uint8 DepthPriority, int32 ViewIndex, FMeshElementCollector& Collector, HHitProxy* HitProxy)
+				  double Radius, double HalfHeight, int32 Sides, const FMaterialRenderProxy* MaterialRenderProxy, uint8 DepthPriority, int32 ViewIndex, FMeshElementCollector& Collector, HHitProxy* HitProxy)
 {
 	GetCylinderMesh( FMatrix::Identity, Base, XAxis, YAxis, ZAxis, Radius, HalfHeight, Sides, MaterialRenderProxy, DepthPriority, ViewIndex, Collector, HitProxy );
 }
 
-void GetCylinderMesh(const FMatrix& CylToWorld, const FVector& Base, const FVector& XAxis, const FVector& YAxis, const FVector& ZAxis, float Radius, float HalfHeight, uint32 Sides, const FMaterialRenderProxy* MaterialRenderProxy, uint8 DepthPriority, int32 ViewIndex, FMeshElementCollector& Collector, HHitProxy* HitProxy)
+void GetCylinderMesh(const FMatrix& CylToWorld, const FVector& Base, const FVector& XAxis, const FVector& YAxis, const FVector& ZAxis, double Radius, double HalfHeight, uint32 Sides, const FMaterialRenderProxy* MaterialRenderProxy, uint8 DepthPriority, int32 ViewIndex, FMeshElementCollector& Collector, HHitProxy* HitProxy)
 {
 	TArray<FDynamicMeshVertex> MeshVerts;
 	TArray<uint32> MeshIndices;
@@ -684,12 +684,12 @@ void GetConeMesh(const FMatrix& LocalToWorld, float AngleWidth, float AngleHeigh
 	MeshBuilder.GetMesh(LocalToWorld, MaterialRenderProxy, DepthPriority, false, false, true, ViewIndex, Collector, HitProxy);
 }
 
-void GetCapsuleMesh(const FVector& Origin, const FVector& XAxis, const FVector& YAxis, const FVector& ZAxis, const FLinearColor& Color, float Radius, float HalfHeight, int32 NumSides, const FMaterialRenderProxy* MaterialRenderProxy, uint8 DepthPriority, bool bDisableBackfaceCulling, int32 ViewIndex, FMeshElementCollector& Collector, HHitProxy* HitProxy)
+void GetCapsuleMesh(const FVector& Origin, const FVector& XAxis, const FVector& YAxis, const FVector& ZAxis, const FLinearColor& Color, double Radius, double HalfHeight, int32 NumSides, const FMaterialRenderProxy* MaterialRenderProxy, uint8 DepthPriority, bool bDisableBackfaceCulling, int32 ViewIndex, FMeshElementCollector& Collector, HHitProxy* HitProxy)
 {
-	const float HalfAxis = FMath::Max<float>(HalfHeight - Radius, 1.f);
+	const double HalfAxis = FMath::Max<double>(HalfHeight - Radius, 1.f);
 	const FVector BottomEnd = Origin + Radius * ZAxis;
 	const FVector TopEnd = BottomEnd + (2 * HalfAxis) * ZAxis;
-	const float CylinderHalfHeight = (TopEnd - BottomEnd).Size() * 0.5;
+	const double CylinderHalfHeight = (TopEnd - BottomEnd).Size() * 0.5;
 	const FVector CylinderLocation = BottomEnd + CylinderHalfHeight * ZAxis;
 
 	GetOrientedHalfSphereMesh(TopEnd, FRotationMatrix::MakeFromXY(XAxis, YAxis).Rotator(), FVector(Radius), NumSides, NumSides, 0, UE_PI / 2, MaterialRenderProxy, DepthPriority, bDisableBackfaceCulling, ViewIndex, Collector, false, HitProxy);
@@ -699,12 +699,12 @@ void GetCapsuleMesh(const FVector& Origin, const FVector& XAxis, const FVector&
 
 
 void DrawCylinder(FPrimitiveDrawInterface* PDI,const FVector& Base, const FVector& XAxis, const FVector& YAxis, const FVector& ZAxis,
-	float Radius, float HalfHeight, uint32 Sides, const FMaterialRenderProxy* MaterialRenderProxy, uint8 DepthPriority)
+	double Radius, double HalfHeight, uint32 Sides, const FMaterialRenderProxy* MaterialRenderProxy, uint8 DepthPriority)
 {
 	DrawCylinder( PDI, FMatrix::Identity, Base, XAxis, YAxis, ZAxis, Radius, HalfHeight, Sides, MaterialRenderProxy, DepthPriority );
 }
 
-void DrawCylinder(FPrimitiveDrawInterface* PDI, const FMatrix& CylToWorld, const FVector& Base, const FVector& XAxis, const FVector& YAxis, const FVector& ZAxis, float Radius, float HalfHeight, uint32 Sides, const FMaterialRenderProxy* MaterialRenderProxy, uint8 DepthPriority)
+void DrawCylinder(FPrimitiveDrawInterface* PDI, const FMatrix& CylToWorld, const FVector& Base, const FVector& XAxis, const FVector& YAxis, const FVector& ZAxis, double Radius, double HalfHeight, uint32 Sides, const FMaterialRenderProxy* MaterialRenderProxy, uint8 DepthPriority)
 {
 	TArray<FDynamicMeshVertex> MeshVerts;
 	TArray<uint32> MeshIndices;
@@ -718,10 +718,10 @@ void DrawCylinder(FPrimitiveDrawInterface* PDI, const FMatrix& CylToWorld, const
 	MeshBuilder.Draw(PDI, CylToWorld, MaterialRenderProxy, DepthPriority,0.f);
 }
 
-void DrawCylinder(class FPrimitiveDrawInterface* PDI, const FVector& Start, const FVector& End, float Radius, int32 Sides, const FMaterialRenderProxy* MaterialInstance, uint8 DepthPriority)
+void DrawCylinder(class FPrimitiveDrawInterface* PDI, const FVector& Start, const FVector& End, double Radius, int32 Sides, const FMaterialRenderProxy* MaterialInstance, uint8 DepthPriority)
 {
 	FVector Dir = End - Start;
-	float Length = Dir.Size();
+	double Length = Dir.Size();
 
 	if (Length > UE_SMALL_NUMBER)
 	{
@@ -735,7 +735,7 @@ void DrawCylinder(class FPrimitiveDrawInterface* PDI, const FVector& Start, cons
 }
 
 void DrawTorus(FPrimitiveDrawInterface* PDI, const FMatrix& Transform, const FVector& XAxis, const FVector& YAxis, 
-			   float OuterRadius, float InnerRadius, int32 OuterSegments, int32 InnerSegments, const FMaterialRenderProxy* MaterialRenderProxy, 
+			   double OuterRadius, double InnerRadius, int32 OuterSegments, int32 InnerSegments, const FMaterialRenderProxy* MaterialRenderProxy,
 			   uint8 DepthPriority, bool bPartial, float Angle, bool bEndCaps)
 {
 	if (OuterSegments < 3)
@@ -841,7 +841,7 @@ void DrawTorus(FPrimitiveDrawInterface* PDI, const FMatrix& Transform, const FVe
 }
 
 
-void DrawDisc(class FPrimitiveDrawInterface* PDI,const FVector& Base,const FVector& XAxis,const FVector& YAxis,FColor Color,float Radius,int32 NumSides,const FMaterialRenderProxy* MaterialRenderProxy, uint8 DepthPriority)
+void DrawDisc(class FPrimitiveDrawInterface* PDI,const FVector& Base,const FVector& XAxis,const FVector& YAxis,FColor Color,double Radius,int32 NumSides,const FMaterialRenderProxy* MaterialRenderProxy, uint8 DepthPriority)
 {
 	check (NumSides >= 3);
 
@@ -1016,7 +1016,7 @@ void DrawWireBox(FPrimitiveDrawInterface* PDI, const FMatrix& Matrix, const FBox
 }
 
 
-void DrawCircle(FPrimitiveDrawInterface* PDI, const FVector& Base, const FVector& X, const FVector& Y, const FLinearColor& Color, float Radius, int32 NumSides, uint8 DepthPriority, float Thickness, float DepthBias, bool bScreenSpace)
+void DrawCircle(FPrimitiveDrawInterface* PDI, const FVector& Base, const FVector& X, const FVector& Y, const FLinearColor& Color, double Radius, int32 NumSides, uint8 DepthPriority, float Thickness, float DepthBias, bool bScreenSpace)
 {
 	const float	AngleDelta = 2.0f * UE_PI / NumSides;
 	FVector	LastVertex = Base + X * Radius;
@@ -1029,7 +1029,7 @@ void DrawCircle(FPrimitiveDrawInterface* PDI, const FVector& Base, const FVector
 	}
 }
 
-void DrawArc(FPrimitiveDrawInterface* PDI, const FVector Base, const FVector X, const FVector Y, const float MinAngle, const float MaxAngle, const float Radius, const int32 Sections, const FLinearColor& Color, uint8 DepthPriority)
+void DrawArc(FPrimitiveDrawInterface* PDI, const FVector Base, const FVector X, const FVector Y, const float MinAngle, const float MaxAngle, const double Radius, const int32 Sections, const FLinearColor& Color, uint8 DepthPriority)
 {
 	float AngleStep = (MaxAngle - MinAngle)/((float)(Sections));
 	float CurrentAngle = MinAngle;
@@ -1120,21 +1120,21 @@ void DrawRectangle(class FPrimitiveDrawInterface* PDI, const FVector& Center, co
 	PDI->DrawLine(Positions[3], Positions[0], Color, DepthPriority, Thickness, DepthBias, bScreenSpace);
 }
 
-void DrawWireSphere(class FPrimitiveDrawInterface* PDI, const FVector& Base, const FLinearColor& Color, float Radius, int32 NumSides, uint8 DepthPriority, float Thickness, float DepthBias, bool bScreenSpace)
+void DrawWireSphere(class FPrimitiveDrawInterface* PDI, const FVector& Base, const FLinearColor& Color, double Radius, int32 NumSides, uint8 DepthPriority, float Thickness, float DepthBias, bool bScreenSpace)
 {
 	DrawCircle(PDI, Base, FVector(1,0,0), FVector(0,1,0), Color, Radius, NumSides, DepthPriority, Thickness, DepthBias, bScreenSpace);
 	DrawCircle(PDI, Base, FVector(1, 0, 0), FVector(0, 0, 1), Color, Radius, NumSides, DepthPriority, Thickness, DepthBias, bScreenSpace);
 	DrawCircle(PDI, Base, FVector(0, 1, 0), FVector(0, 0, 1), Color, Radius, NumSides, DepthPriority, Thickness, DepthBias, bScreenSpace);
 }
 
-void DrawWireSphereAutoSides(class FPrimitiveDrawInterface* PDI, const FVector& Base, const FLinearColor& Color, float Radius, uint8 DepthPriority, float Thickness, float DepthBias, bool bScreenSpace)
+void DrawWireSphereAutoSides(class FPrimitiveDrawInterface* PDI, const FVector& Base, const FLinearColor& Color, double Radius, uint8 DepthPriority, float Thickness, float DepthBias, bool bScreenSpace)
 {
 	// Guess a good number of sides
 	int32 NumSides =  FMath::Clamp<int32>(Radius/4.f, 16, 64);
 	DrawWireSphere(PDI, Base, Color, Radius, NumSides, DepthPriority, Thickness, DepthBias, bScreenSpace);
 }
 
-void DrawWireSphere(class FPrimitiveDrawInterface* PDI, const FTransform& Transform, const FLinearColor& Color, float Radius, int32 NumSides, uint8 DepthPriority, float Thickness, float DepthBias, bool bScreenSpace)
+void DrawWireSphere(class FPrimitiveDrawInterface* PDI, const FTransform& Transform, const FLinearColor& Color, double Radius, int32 NumSides, uint8 DepthPriority, float Thickness, float DepthBias, bool bScreenSpace)
 {
 	DrawCircle(PDI, Transform.GetLocation(), Transform.GetScaledAxis(EAxis::X), Transform.GetScaledAxis(EAxis::Y), Color, Radius, NumSides, DepthPriority, Thickness, DepthBias, bScreenSpace);
 	DrawCircle(PDI, Transform.GetLocation(), Transform.GetScaledAxis(EAxis::X), Transform.GetScaledAxis(EAxis::Z), Color, Radius, NumSides, DepthPriority, Thickness, DepthBias, bScreenSpace);
@@ -1142,14 +1142,14 @@ void DrawWireSphere(class FPrimitiveDrawInterface* PDI, const FTransform& Transf
 }
 
 
-void DrawWireSphereAutoSides(class FPrimitiveDrawInterface* PDI, const FTransform& Transform, const FLinearColor& Color, float Radius, uint8 DepthPriority, float Thickness, float DepthBias, bool bScreenSpace)
+void DrawWireSphereAutoSides(class FPrimitiveDrawInterface* PDI, const FTransform& Transform, const FLinearColor& Color, double Radius, uint8 DepthPriority, float Thickness, float DepthBias, bool bScreenSpace)
 {
 	// Guess a good number of sides
 	int32 NumSides =  FMath::Clamp<int32>(Radius/4.f, 16, 64);
 	DrawWireSphere(PDI, Transform, Color, Radius, NumSides, DepthPriority, Thickness, DepthBias, bScreenSpace);
 }
 
-void DrawWireCylinder(FPrimitiveDrawInterface* PDI, const FVector& Base, const FVector& X, const FVector& Y, const FVector& Z, const FLinearColor& Color, float Radius, float HalfHeight, int32 NumSides, uint8 DepthPriority, float Thickness, float DepthBias, bool bScreenSpace)
+void DrawWireCylinder(FPrimitiveDrawInterface* PDI, const FVector& Base, const FVector& X, const FVector& Y, const FVector& Z, const FLinearColor& Color, double Radius, double HalfHeight, int32 NumSides, uint8 DepthPriority, float Thickness, float DepthBias, bool bScreenSpace)
 {
 	const float	AngleDelta = 2.0f * UE_PI / NumSides;
 	FVector	LastVertex = Base + X * Radius;
@@ -1167,7 +1167,7 @@ void DrawWireCylinder(FPrimitiveDrawInterface* PDI, const FVector& Base, const F
 }
 
 
-static void DrawHalfCircle(FPrimitiveDrawInterface* PDI, const FVector& Base, const FVector& X, const FVector& Y, const FLinearColor& Color, float Radius, int32 NumSides, float Thickness, float DepthBias, bool bScreenSpace)
+static void DrawHalfCircle(FPrimitiveDrawInterface* PDI, const FVector& Base, const FVector& X, const FVector& Y, const FLinearColor& Color, double Radius, int32 NumSides, float Thickness, float DepthBias, bool bScreenSpace)
 {
 	const float	AngleDelta = (float)UE_PI / ((float)NumSides);
 	FVector	LastVertex = Base + X * Radius;
@@ -1180,7 +1180,7 @@ static void DrawHalfCircle(FPrimitiveDrawInterface* PDI, const FVector& Base, co
 	}	
 }
 
-void DrawWireCapsule(FPrimitiveDrawInterface* PDI, const FVector& Base, const FVector& X, const FVector& Y, const FVector& Z, const FLinearColor& Color, float Radius, float HalfHeight, int32 NumSides, uint8 DepthPriority, float Thickness, float DepthBias, bool bScreenSpace)
+void DrawWireCapsule(FPrimitiveDrawInterface* PDI, const FVector& Base, const FVector& X, const FVector& Y, const FVector& Z, const FLinearColor& Color, double Radius, double HalfHeight, int32 NumSides, uint8 DepthPriority, float Thickness, float DepthBias, bool bScreenSpace)
 {
 	const FVector Origin = Base;
 	const FVector XAxis = X.GetSafeNormal();
@@ -1188,10 +1188,10 @@ void DrawWireCapsule(FPrimitiveDrawInterface* PDI, const FVector& Base, const FV
 	const FVector ZAxis = Z.GetSafeNormal();
 
 	// because we are drawing a capsule we have to have room for the "domed caps"
-	const float XScale = X.Size();
-	const float YScale = Y.Size();
-	const float ZScale = Z.Size();
-	float CapsuleRadius = Radius * FMath::Max(XScale,YScale);
+	const double XScale = X.Size();
+	const double YScale = Y.Size();
+	const double ZScale = Z.Size();
+	double CapsuleRadius = Radius * FMath::Max(XScale,YScale);
 	HalfHeight *= ZScale;
 	CapsuleRadius = FMath::Clamp(CapsuleRadius, 0.f, HalfHeight);	//cap radius based on total height
 	HalfHeight -= CapsuleRadius;
@@ -1230,7 +1230,7 @@ void DrawWireCapsule(FPrimitiveDrawInterface* PDI, const FVector& Base, const FV
 	}
 }
 
-void DrawWireCone(FPrimitiveDrawInterface* PDI, TArray<FVector>& Verts, const FTransform& Transform, float ConeLength, float ConeAngle, int32 ConeSides, const FLinearColor& Color, uint8 DepthPriority, float Thickness, float DepthBias, bool bScreenSpace)
+void DrawWireCone(FPrimitiveDrawInterface* PDI, TArray<FVector>& Verts, const FTransform& Transform, double ConeLength, double ConeAngle, int32 ConeSides, const FLinearColor& Color, uint8 DepthPriority, float Thickness, float DepthBias, bool bScreenSpace)
 {
 	static const float TwoPI = 2.0f * UE_PI;
 	static const float ToRads = UE_PI / 180.0f;
@@ -1272,12 +1272,12 @@ void DrawWireCone(FPrimitiveDrawInterface* PDI, TArray<FVector>& Verts, const FT
 	PDI->DrawTranslucentLine(Verts[Verts.Num() - 1], Verts[0], Color, DepthPriority, Thickness, DepthBias, bScreenSpace);
 }
 
-void DrawWireCone(FPrimitiveDrawInterface* PDI, TArray<FVector>& Verts, const FMatrix& Transform, float ConeLength, float ConeAngle, int32 ConeSides, const FLinearColor& Color, uint8 DepthPriority, float Thickness, float DepthBias, bool bScreenSpace)
+void DrawWireCone(FPrimitiveDrawInterface* PDI, TArray<FVector>& Verts, const FMatrix& Transform, double ConeLength, double ConeAngle, int32 ConeSides, const FLinearColor& Color, uint8 DepthPriority, float Thickness, float DepthBias, bool bScreenSpace)
 {
 	DrawWireCone(PDI, Verts, FTransform(Transform), ConeLength, ConeAngle, ConeSides, Color, DepthPriority, Thickness, DepthBias, bScreenSpace);
 }
 
-void DrawWireSphereCappedCone(FPrimitiveDrawInterface* PDI, const FTransform& Transform, float ConeLength, float ConeAngle, int32 ConeSides, int32 ArcFrequency, int32 CapSegments, const FLinearColor& Color, uint8 DepthPriority)
+void DrawWireSphereCappedCone(FPrimitiveDrawInterface* PDI, const FTransform& Transform, double ConeLength, double ConeAngle, int32 ConeSides, int32 ArcFrequency, int32 CapSegments, const FLinearColor& Color, uint8 DepthPriority)
 {
 	// The cap only works if there are an even number of verts generated so add another if needed 
 	if ((ConeSides & 0x1) != 0)
@@ -1300,7 +1300,7 @@ void DrawWireSphereCappedCone(FPrimitiveDrawInterface* PDI, const FTransform& Tr
 }
 
 
-void DrawWireChoppedCone(FPrimitiveDrawInterface* PDI,const FVector& Base,const FVector& X,const FVector& Y,const FVector& Z,const FLinearColor& Color,float Radius, float TopRadius,float HalfHeight,int32 NumSides,uint8 DepthPriority)
+void DrawWireChoppedCone(FPrimitiveDrawInterface* PDI,const FVector& Base,const FVector& X,const FVector& Y,const FVector& Z,const FLinearColor& Color,double Radius, double TopRadius,double HalfHeight,int32 NumSides,uint8 DepthPriority)
 {
 	const float	AngleDelta = 2.0f * UE_PI / NumSides;
 	FVector	LastVertex = Base + X * Radius;
@@ -1415,16 +1415,16 @@ void DrawWireStar(FPrimitiveDrawInterface* PDI,const FVector& Position, float Si
 	PDI->DrawLine(Position + Size * FVector(0,0,1), Position - Size * FVector(0,0,1), Color, DepthPriority);
 }
 
-void DrawDashedLine(FPrimitiveDrawInterface* PDI, const FVector& Start, const FVector& End, const FLinearColor& Color, float DashSize, uint8 DepthPriority, float DepthBias)
+void DrawDashedLine(FPrimitiveDrawInterface* PDI, const FVector& Start, const FVector& End, const FLinearColor& Color, double DashSize, uint8 DepthPriority, float DepthBias)
 {
 	FVector LineDir = End - Start;
-	float LineLeft = (End - Start).Size();
+	double LineLeft = (End - Start).Size();
 	if (LineLeft)
 	{
 		LineDir /= LineLeft;
 	}
 
-	const int32 nLines = FMath::CeilToInt(LineLeft / (DashSize*2));
+	const int32 nLines = FMath::CeilToInt32(LineLeft / (DashSize*2));
 	PDI->AddReserveLines(DepthPriority, nLines, DepthBias != 0);
 
 	const FVector Dash = (DashSize * LineDir);
diff --git a/Engine/Source/Runtime/Engine/Public/SceneManagement.h b/Engine/Source/Runtime/Engine/Public/SceneManagement.h
--- a/Engine/Source/Runtime/Engine/Public/SceneManagement.h
+++ b/Engine/Source/Runtime/Engine/Public/SceneManagement.h
@@ -2679,13 +2679,13 @@ extern ENGINE_API void DrawSphere(class FPrimitiveDrawInterface* PDI,const FVect
 extern ENGINE_API void DrawCone(class FPrimitiveDrawInterface* PDI,const FMatrix& ConeToWorld, float Angle1, float Angle2, uint32 NumSides, bool bDrawSideLines, const FLinearColor& SideLineColor, const FMaterialRenderProxy* MaterialRenderProxy, uint8 DepthPriority);
 
 extern ENGINE_API void DrawCylinder(class FPrimitiveDrawInterface* PDI,const FVector& Base, const FVector& XAxis, const FVector& YAxis, const FVector& ZAxis,
-	float Radius, float HalfHeight, uint32 Sides, const FMaterialRenderProxy* MaterialInstance, uint8 DepthPriority);
+	double Radius, double HalfHeight, uint32 Sides, const FMaterialRenderProxy* MaterialInstance, uint8 DepthPriority);
 
 extern ENGINE_API void DrawCylinder(class FPrimitiveDrawInterface* PDI, const FMatrix& CylToWorld, const FVector& Base, const FVector& XAxis, const FVector& YAxis, const FVector& ZAxis,
-	float Radius, float HalfHeight, uint32 Sides, const FMaterialRenderProxy* MaterialInstance, uint8 DepthPriority);
+	double Radius, double HalfHeight, uint32 Sides, const FMaterialRenderProxy* MaterialInstance, uint8 DepthPriority);
 
 //Draws a cylinder along the axis from Start to End
-extern ENGINE_API void DrawCylinder(class FPrimitiveDrawInterface* PDI, const FVector& Start, const FVector& End, float Radius, int32 Sides, const FMaterialRenderProxy* MaterialInstance, uint8 DepthPriority);
+extern ENGINE_API void DrawCylinder(class FPrimitiveDrawInterface* PDI, const FVector& Start, const FVector& End, double Radius, int32 Sides, const FMaterialRenderProxy* MaterialInstance, uint8 DepthPriority);
 
 
 extern ENGINE_API void GetBoxMesh(const FMatrix& BoxToWorld,const FVector& Radii,const FMaterialRenderProxy* MaterialRenderProxy,uint8 DepthPriority,int32 ViewIndex,FMeshElementCollector& Collector, HHitProxy* HitProxy = NULL);
@@ -2698,16 +2698,16 @@ extern ENGINE_API void GetSphereMesh(const FVector& Center, const FVector& Radii
 extern ENGINE_API void GetSphereMesh(const FVector& Center,const FVector& Radii,int32 NumSides,int32 NumRings,const FMaterialRenderProxy* MaterialRenderProxy,uint8 DepthPriority,
 									bool bDisableBackfaceCulling,int32 ViewIndex,FMeshElementCollector& Collector, bool bUseSelectionOutline, HHitProxy* HitProxy);
 extern ENGINE_API void GetCylinderMesh(const FVector& Base, const FVector& XAxis, const FVector& YAxis, const FVector& ZAxis,
-									float Radius, float HalfHeight, int32 Sides, const FMaterialRenderProxy* MaterialInstance, uint8 DepthPriority, int32 ViewIndex, FMeshElementCollector& Collector, HHitProxy* HitProxy = NULL);
+									double Radius, double HalfHeight, int32 Sides, const FMaterialRenderProxy* MaterialInstance, uint8 DepthPriority, int32 ViewIndex, FMeshElementCollector& Collector, HHitProxy* HitProxy = NULL);
 extern ENGINE_API void GetCylinderMesh(const FMatrix& CylToWorld, const FVector& Base, const FVector& XAxis, const FVector& YAxis, const FVector& ZAxis,
-									float Radius, float HalfHeight, uint32 Sides, const FMaterialRenderProxy* MaterialInstance, uint8 DepthPriority, int32 ViewIndex, FMeshElementCollector& Collector, HHitProxy* HitProxy = NULL);
+									double Radius, double HalfHeight, uint32 Sides, const FMaterialRenderProxy* MaterialInstance, uint8 DepthPriority, int32 ViewIndex, FMeshElementCollector& Collector, HHitProxy* HitProxy = NULL);
 //Draws a cylinder along the axis from Start to End
-extern ENGINE_API void GetCylinderMesh(const FVector& Start, const FVector& End, float Radius, int32 Sides, const FMaterialRenderProxy* MaterialInstance, uint8 DepthPriority, int32 ViewIndex, FMeshElementCollector& Collector, HHitProxy* HitProxy = NULL);
+extern ENGINE_API void GetCylinderMesh(const FVector& Start, const FVector& End, double Radius, int32 Sides, const FMaterialRenderProxy* MaterialInstance, uint8 DepthPriority, int32 ViewIndex, FMeshElementCollector& Collector, HHitProxy* HitProxy = NULL);
 
 
 extern ENGINE_API void GetConeMesh(const FMatrix& LocalToWorld, float AngleWidth, float AngleHeight, uint32 NumSides,
 									const FMaterialRenderProxy* MaterialRenderProxy, uint8 DepthPriority, int32 ViewIndex, FMeshElementCollector& Collector, HHitProxy* HitProxy = NULL);
-extern ENGINE_API void GetCapsuleMesh(const FVector& Origin, const FVector& XAxis, const FVector& YAxis, const FVector& ZAxis, const FLinearColor& Color, float Radius, float HalfHeight, int32 NumSides,
+extern ENGINE_API void GetCapsuleMesh(const FVector& Origin, const FVector& XAxis, const FVector& YAxis, const FVector& ZAxis, const FLinearColor& Color, double Radius, double HalfHeight, int32 NumSides,
 									const FMaterialRenderProxy* MaterialRenderProxy, uint8 DepthPriority, bool bDisableBackfaceCulling, int32 ViewIndex, FMeshElementCollector& Collector, HHitProxy* HitProxy = NULL);
 
 
@@ -2731,7 +2731,7 @@ extern ENGINE_API void GetCapsuleMesh(const FVector& Origin, const FVector& XAxi
  * @param	bEndCaps				If partial, whether the ends should be capped with triangles.
  */
 extern ENGINE_API void DrawTorus(FPrimitiveDrawInterface* PDI, const FMatrix& Transform, const FVector& XAxis, const FVector& YAxis, 
-								 float OuterRadius, float InnerRadius, int32 OuterSegments, int32 InnerSegments, const FMaterialRenderProxy* MaterialRenderProxy, uint8 DepthPriority, bool bPartial, float Angle, bool bEndCaps);
+								 double OuterRadius, double InnerRadius, int32 OuterSegments, int32 InnerSegments, const FMaterialRenderProxy* MaterialRenderProxy, uint8 DepthPriority, bool bPartial, float Angle, bool bEndCaps);
 
 /**
  * Draws a circle using triangles.
@@ -2746,7 +2746,7 @@ extern ENGINE_API void DrawTorus(FPrimitiveDrawInterface* PDI, const FMatrix& Tr
  * @param	MaterialRenderProxy		Material to use for render 
  * @param	DepthPriority			Depth priority for the circle.
  */
-extern ENGINE_API void DrawDisc(class FPrimitiveDrawInterface* PDI,const FVector& Base,const FVector& XAxis,const FVector& YAxis,FColor Color,float Radius,int32 NumSides, const FMaterialRenderProxy* MaterialRenderProxy, uint8 DepthPriority);
+extern ENGINE_API void DrawDisc(class FPrimitiveDrawInterface* PDI,const FVector& Base,const FVector& XAxis,const FVector& YAxis,FColor Color,double Radius,int32 NumSides, const FMaterialRenderProxy* MaterialRenderProxy, uint8 DepthPriority);
 
 /**
  * Draws a rectangle using triangles.
@@ -2820,7 +2820,7 @@ extern ENGINE_API void DrawWireBox(class FPrimitiveDrawInterface* PDI, const FMa
  * @param	DepthPriority	Depth priority for the circle.
  * @param	Thickness		Thickness of the lines comprising the circle
  */
-extern ENGINE_API void DrawCircle(class FPrimitiveDrawInterface* PDI, const FVector& Base, const FVector& X, const FVector& Y, const FLinearColor& Color, float Radius, int32 NumSides, uint8 DepthPriority, float Thickness = 0.0f, float DepthBias = 0.0f, bool bScreenSpace = false);
+extern ENGINE_API void DrawCircle(class FPrimitiveDrawInterface* PDI, const FVector& Base, const FVector& X, const FVector& Y, const FLinearColor& Color, double Radius, int32 NumSides, uint8 DepthPriority, float Thickness = 0.0f, float DepthBias = 0.0f, bool bScreenSpace = false);
 
 
 /**
@@ -2837,7 +2837,7 @@ extern ENGINE_API void DrawCircle(class FPrimitiveDrawInterface* PDI, const FVec
  * @param	Color			Color of the circle.
  * @param	DepthPriority	Depth priority for the circle.
  */
-extern ENGINE_API void DrawArc(FPrimitiveDrawInterface* PDI, const FVector Base, const FVector X, const FVector Y, const float MinAngle, const float MaxAngle, const float Radius, const int32 Sections, const FLinearColor& Color, uint8 DepthPriority);
+extern ENGINE_API void DrawArc(FPrimitiveDrawInterface* PDI, const FVector Base, const FVector X, const FVector Y, const float MinAngle, const float MaxAngle, const double Radius, const int32 Sections, const FLinearColor& Color, uint8 DepthPriority);
 
 /**
  * Draws a rectangle using lines.
@@ -2867,8 +2867,8 @@ extern ENGINE_API void DrawRectangle(FPrimitiveDrawInterface* PDI, const FVector
  * @param	DepthPriority	Depth priority for the circle.
  * @param	Thickness		Thickness of the lines comprising the sphere
  */
-extern ENGINE_API void DrawWireSphere(class FPrimitiveDrawInterface* PDI, const FVector& Base, const FLinearColor& Color, float Radius, int32 NumSides, uint8 DepthPriority, float Thickness = 0.0f, float DepthBias = 0.0f, bool bScreenSpace = false);
-extern ENGINE_API void DrawWireSphere(class FPrimitiveDrawInterface* PDI, const FTransform& Transform, const FLinearColor& Color, float Radius, int32 NumSides, uint8 DepthPriority, float Thickness = 0.0f, float DepthBias = 0.0f, bool bScreenSpace = false);
+extern ENGINE_API void DrawWireSphere(class FPrimitiveDrawInterface* PDI, const FVector& Base, const FLinearColor& Color, double Radius, int32 NumSides, uint8 DepthPriority, float Thickness = 0.0f, float DepthBias = 0.0f, bool bScreenSpace = false);
+extern ENGINE_API void DrawWireSphere(class FPrimitiveDrawInterface* PDI, const FTransform& Transform, const FLinearColor& Color, double Radius, int32 NumSides, uint8 DepthPriority, float Thickness = 0.0f, float DepthBias = 0.0f, bool bScreenSpace = false);
 
 /**
  * Draws a sphere using circles, automatically calculating a reasonable number of sides
@@ -2880,8 +2880,8 @@ extern ENGINE_API void DrawWireSphere(class FPrimitiveDrawInterface* PDI, const
  * @param	DepthPriority	Depth priority for the circle.
  * @param	Thickness		Thickness of the lines comprising the sphere
  */
-extern ENGINE_API void DrawWireSphereAutoSides(class FPrimitiveDrawInterface* PDI, const FVector& Base, const FLinearColor& Color, float Radius, uint8 DepthPriority, float Thickness = 0.0f, float DepthBias = 0.0f, bool bScreenSpace = false);
-extern ENGINE_API void DrawWireSphereAutoSides(class FPrimitiveDrawInterface* PDI, const FTransform& Transform, const FLinearColor& Color, float Radius, uint8 DepthPriority, float Thickness = 0.0f, float DepthBias = 0.0f, bool bScreenSpace = false);
+extern ENGINE_API void DrawWireSphereAutoSides(class FPrimitiveDrawInterface* PDI, const FVector& Base, const FLinearColor& Color, double Radius, uint8 DepthPriority, float Thickness = 0.0f, float DepthBias = 0.0f, bool bScreenSpace = false);
+extern ENGINE_API void DrawWireSphereAutoSides(class FPrimitiveDrawInterface* PDI, const FTransform& Transform, const FLinearColor& Color, double Radius, uint8 DepthPriority, float Thickness = 0.0f, float DepthBias = 0.0f, bool bScreenSpace = false);
 
 /**
  * Draws a wireframe cylinder.
@@ -2898,7 +2898,7 @@ extern ENGINE_API void DrawWireSphereAutoSides(class FPrimitiveDrawInterface* PD
  * @param	DepthPriority	Depth priority for the cylinder.
  * @param	Thickness		Thickness of the lines comprising the cylinder
  */
-extern ENGINE_API void DrawWireCylinder(class FPrimitiveDrawInterface* PDI, const FVector& Base, const FVector& X, const FVector& Y, const FVector& Z, const FLinearColor& Color, float Radius, float HalfHeight, int32 NumSides, uint8 DepthPriority, float Thickness = 0.0f, float DepthBias = 0.0f, bool bScreenSpace = false);
+extern ENGINE_API void DrawWireCylinder(class FPrimitiveDrawInterface* PDI, const FVector& Base, const FVector& X, const FVector& Y, const FVector& Z, const FLinearColor& Color, double Radius, double HalfHeight, int32 NumSides, uint8 DepthPriority, float Thickness = 0.0f, float DepthBias = 0.0f, bool bScreenSpace = false);
 
 /**
  * Draws a wireframe capsule.
@@ -2915,7 +2915,7 @@ extern ENGINE_API void DrawWireCylinder(class FPrimitiveDrawInterface* PDI, cons
  * @param	DepthPriority	Depth priority for the cylinder.
  * @param	Thickness		Thickness of the lines comprising the cylinder
  */
-extern ENGINE_API void DrawWireCapsule(class FPrimitiveDrawInterface* PDI, const FVector& Base, const FVector& X, const FVector& Y, const FVector& Z, const FLinearColor& Color, float Radius, float HalfHeight, int32 NumSides, uint8 DepthPriority, float Thickness = 0.0f, float DepthBias = 0.0f, bool bScreenSpace = false);
+extern ENGINE_API void DrawWireCapsule(class FPrimitiveDrawInterface* PDI, const FVector& Base, const FVector& X, const FVector& Y, const FVector& Z, const FLinearColor& Color, double Radius, double HalfHeight, int32 NumSides, uint8 DepthPriority, float Thickness = 0.0f, float DepthBias = 0.0f, bool bScreenSpace = false);
 
 /**
  * Draws a wireframe chopped cone (cylinder with independent top and bottom radius).
@@ -2932,7 +2932,7 @@ extern ENGINE_API void DrawWireCapsule(class FPrimitiveDrawInterface* PDI, const
  * @param	NumSides		Numbers of sides that the cone has.
  * @param	DepthPriority	Depth priority for the cone.
  */
-extern ENGINE_API void DrawWireChoppedCone(class FPrimitiveDrawInterface* PDI,const FVector& Base,const FVector& X,const FVector& Y,const FVector& Z,const FLinearColor& Color,float Radius,float TopRadius,float HalfHeight,int32 NumSides,uint8 DepthPriority);
+extern ENGINE_API void DrawWireChoppedCone(class FPrimitiveDrawInterface* PDI,const FVector& Base,const FVector& X,const FVector& Y,const FVector& Z,const FLinearColor& Color,double Radius,double TopRadius,double HalfHeight,int32 NumSides,uint8 DepthPriority);
 
 /**
  * Draws a wireframe cone
@@ -2947,8 +2947,8 @@ extern ENGINE_API void DrawWireChoppedCone(class FPrimitiveDrawInterface* PDI,co
  * @param	Verts			Out param, the positions of the verts at the cone base.
  * @param	Thickness		Thickness of the lines comprising the cone
  */
-extern ENGINE_API void DrawWireCone(class FPrimitiveDrawInterface* PDI, TArray<FVector>& Verts, const FMatrix& Transform, float ConeLength, float ConeAngle, int32 ConeSides, const FLinearColor& Color, uint8 DepthPriority, float Thickness = 0.0f, float DepthBias = 0.0f, bool bScreenSpace = false);
-extern ENGINE_API void DrawWireCone(class FPrimitiveDrawInterface* PDI, TArray<FVector>& Verts, const FTransform& Transform, float ConeLength, float ConeAngle, int32 ConeSides, const FLinearColor& Color, uint8 DepthPriority, float Thickness = 0.0f, float DepthBias = 0.0f, bool bScreenSpace = false);
+extern ENGINE_API void DrawWireCone(class FPrimitiveDrawInterface* PDI, TArray<FVector>& Verts, const FMatrix& Transform, double ConeLength, double ConeAngle, int32 ConeSides, const FLinearColor& Color, uint8 DepthPriority, float Thickness = 0.0f, float DepthBias = 0.0f, bool bScreenSpace = false);
+extern ENGINE_API void DrawWireCone(class FPrimitiveDrawInterface* PDI, TArray<FVector>& Verts, const FTransform& Transform, double ConeLength, double ConeAngle, int32 ConeSides, const FLinearColor& Color, uint8 DepthPriority, float Thickness = 0.0f, float DepthBias = 0.0f, bool bScreenSpace = false);
 
 /**
  * Draws a wireframe cone with a arcs on the cap
@@ -2963,7 +2963,7 @@ extern ENGINE_API void DrawWireCone(class FPrimitiveDrawInterface* PDI, TArray<F
  * @param	Color			Color of the cone.
  * @param	DepthPriority	Depth priority for the cone.
  */
-extern ENGINE_API void DrawWireSphereCappedCone(FPrimitiveDrawInterface* PDI, const FTransform& Transform, float ConeLength, float ConeAngle, int32 ConeSides, int32 ArcFrequency, int32 CapSegments, const FLinearColor& Color, uint8 DepthPriority);
+extern ENGINE_API void DrawWireSphereCappedCone(FPrimitiveDrawInterface* PDI, const FTransform& Transform, double ConeLength, double ConeAngle, int32 ConeSides, int32 ArcFrequency, int32 CapSegments, const FLinearColor& Color, uint8 DepthPriority);
 
 /**
  * Draws an oriented box.
@@ -3028,7 +3028,7 @@ extern ENGINE_API void DrawWireStar(class FPrimitiveDrawInterface* PDI, const FV
  * @param	DashSize		Size of each of the dashes that makes up the line.
  * @param	DepthPriority	Depth priority for the line.
  */
-extern ENGINE_API void DrawDashedLine(class FPrimitiveDrawInterface* PDI, const FVector& Start, const FVector& End, const FLinearColor& Color, float DashSize, uint8 DepthPriority, float DepthBias = 0.0f);
+extern ENGINE_API void DrawDashedLine(class FPrimitiveDrawInterface* PDI, const FVector& Start, const FVector& End, const FLinearColor& Color, double DashSize, uint8 DepthPriority, float DepthBias = 0.0f);
 
 /**
  * Draws a wireframe diamond.
@@ -3084,7 +3084,7 @@ extern ENGINE_API void DrawFrustumWireframe(
 extern ENGINE_API FVector CalcConeVert(float Angle1, float Angle2, float AzimuthAngle);
 extern ENGINE_API void BuildConeVerts(float Angle1, float Angle2, float Scale, float XOffset, uint32 NumSides, TArray<FDynamicMeshVertex>& OutVerts, TArray<uint32>& OutIndices);
 
-void BuildCylinderVerts(const FVector& Base, const FVector& XAxis, const FVector& YAxis, const FVector& ZAxis, float Radius, float HalfHeight, uint32 Sides, TArray<FDynamicMeshVertex>& OutVerts, TArray<uint32>& OutIndices);
+void BuildCylinderVerts(const FVector& Base, const FVector& XAxis, const FVector& YAxis, const FVector& ZAxis, double Radius, double HalfHeight, uint32 Sides, TArray<FDynamicMeshVertex>& OutVerts, TArray<uint32>& OutIndices);
 
 
 /**
