Removing an undefined function from SkeltalMeshMerge.

diff --git a/Engine/Source/Runtime/Engine/Public/SkeletalMeshMerge.h b/Engine/Source/Runtime/Engine/Public/SkeletalMeshMerge.h
--- a/Engine/Source/Runtime/Engine/Public/SkeletalMeshMerge.h
+++ b/Engine/Source/Runtime/Engine/Public/SkeletalMeshMerge.h
@@ -363,8 +363,4 @@ private:
 	 */
 	template<typename VertexDataType>
 	void CopyVertexFromSource(VertexDataType& DestVert, const FSkeletalMeshLODRenderData& SrcLODData, int32 SourceVertIdx, const FMergeSectionInfo& MergeSectionInfo);
-
-	/** Copy skin weight info from source LOD model - templatized per SourceLODModel extra bone influence */
-	template<typename SkinWeightType, bool bHasExtraBoneInfluences, typename BoneIndexType>
-	void CopyWeightFromSource(SkinWeightType& DestWeight, const FSkeletalMeshLODRenderData& SrcLODData, int32 SourceVertIdx, const FMergeSectionInfo& MergeSectionInfo);
 };
