World Partition: fix child actor components removing their child actor from the world when it's already done when removing the component owner.

#rnx

diff --git a/Engine/Source/Runtime/Engine/Private/Components/ChildActorComponent.cpp b/Engine/Source/Runtime/Engine/Private/Components/ChildActorComponent.cpp
--- a/Engine/Source/Runtime/Engine/Private/Components/ChildActorComponent.cpp
+++ b/Engine/Source/Runtime/Engine/Private/Components/ChildActorComponent.cpp
@@ -774,20 +774,32 @@ void UChildActorComponent::DestroyChildActor()
 	// If we own an Actor, kill it now unless we don't have authority on it, for that we rely on the server
 	// If the level is being removed then don't destroy the child actor so re-adding it doesn't
 	// need to create a new actor
-	auto IsLevelBeingRemoved = [this]() -> bool
+	auto IsBeingRemovedFromLevel = [this]() -> bool
 	{
 		if (AActor* MyOwner = GetOwner())
 		{
 			if (ULevel* MyLevel = MyOwner->GetLevel())
 			{
-				return MyLevel->bIsBeingRemoved;
+				if (MyLevel->bIsBeingRemoved)
+				{
+					return true;
+				}
+
+#if WITH_EDITOR
+				// In the editor, the child actor can already be removed from the world if the component owner was removed first
+				UWorld* MyWorld = MyLevel->GetWorld();
+				if (MyWorld && !MyWorld->IsGameWorld() && !MyLevel->Actors.Contains(ChildActor))
+				{
+					return true;
+				}
+#endif
 			}
 		}
 
 		return false;
 	};
 
-	if (ChildActor && (ChildActor->HasAuthority() || !IsChildActorReplicated()) && !IsLevelBeingRemoved())
+	if (ChildActor && (ChildActor->HasAuthority() || !IsChildActorReplicated()) && !IsBeingRemovedFromLevel())
 	{
 		if (!GExitPurge)
 		{
