Add an interface to RuntimeHashes so they can dynamically inject/remove content to be streamed.
RuntimeSpatialHash can inject/remove dynamically StreamingGrids.
Added an error when mapping gridname grid to catch any gridname clash

diff --git a/Engine/Source/Runtime/Engine/Private/WorldPartition/WorldPartitionRuntimeCell.cpp b/Engine/Source/Runtime/Engine/Private/WorldPartition/WorldPartitionRuntimeCell.cpp
--- a/Engine/Source/Runtime/Engine/Private/WorldPartition/WorldPartitionRuntimeCell.cpp
+++ b/Engine/Source/Runtime/Engine/Private/WorldPartition/WorldPartitionRuntimeCell.cpp
@@ -57,6 +57,12 @@ void UWorldPartitionRuntimeCell::SetDebugInfo(int64 InCoordX, int64 InCoordY, in
 	UpdateDebugName();
 }
 
+void UWorldPartitionRuntimeCell::SetGridName(FName InGridName)
+{
+	DebugInfo.GridName = InGridName;
+	UpdateDebugName();
+}
+
 void UWorldPartitionRuntimeCell::UpdateDebugName()
 {
 	TStringBuilder<512> Builder;
diff --git a/Engine/Source/Runtime/Engine/Private/WorldPartition/WorldPartitionRuntimeSpatialHash.cpp b/Engine/Source/Runtime/Engine/Private/WorldPartition/WorldPartitionRuntimeSpatialHash.cpp
--- a/Engine/Source/Runtime/Engine/Private/WorldPartition/WorldPartitionRuntimeSpatialHash.cpp
+++ b/Engine/Source/Runtime/Engine/Private/WorldPartition/WorldPartitionRuntimeSpatialHash.cpp
@@ -33,6 +33,7 @@
 #include "RenderUtils.h"
 #include "Algo/Transform.h"
 #include "Algo/Copy.h"
+#include "Algo/Accumulate.h"
 #include "Math/Range.h"
 #include "UObject/ObjectSaveContext.h"
 #include "Components/LineBatchComponent.h"
@@ -52,6 +53,8 @@
 #include "Misc/UObjectToken.h"
 #include "Misc/MapErrors.h"
 
+#include "UObject/UObjectGlobals.h"
+
 extern UNREALED_API class UEditorEngine* GEditor;
 #endif //WITH_EDITOR
 
@@ -735,6 +738,7 @@ UWorldPartitionRuntimeSpatialHash::UWorldPartitionRuntimeSpatialHash(const FObje
 #if WITH_EDITORONLY_DATA
 	, bPreviewGrids(false)
 #endif
+	, bIsNameToGridMappingDirty(true)
 {}
 
 void UWorldPartitionRuntimeSpatialHash::PreSave(const class ITargetPlatform* TargetPlatform)
@@ -767,6 +771,35 @@ void UWorldPartitionRuntimeSpatialHash::DrawPreview() const
 	GridPreviewer.Draw(GetWorld(), Grids, bPreviewGrids);
 }
 
+URuntimeHashExternalStreamingObjectBase* UWorldPartitionRuntimeSpatialHash::StoreToExternalStreamingObject(UObject* StreamingObjectOuter, FName StreamingObjectName)
+{
+	check(!GetWorld()->IsGameWorld());
+
+	FName StreamingObjectUniqueName = MakeUniqueObjectName(StreamingObjectOuter, URuntimeSpatialHashExternalStreamingObject::StaticClass(), StreamingObjectName);
+	URuntimeSpatialHashExternalStreamingObject* StreamingObject = NewObject<URuntimeSpatialHashExternalStreamingObject>(StreamingObjectOuter, StreamingObjectUniqueName);
+	StreamingObject->StreamingGrids = MoveTemp(StreamingGrids);
+
+	// Rename grids so the name of external streaming objects do not clashes with the main streaming grids.
+	for (FSpatialHashStreamingGrid& StreamingGrid : StreamingObject->StreamingGrids)
+	{
+		FString ExternalGridName = StreamingGrid.GridName.ToString() + TEXT("(") + StreamingObjectName.ToString() + TEXT(")");
+		StreamingGrid.GridName = FName(*ExternalGridName);
+
+		for (FSpatialHashStreamingGridLevel& GridLevel : StreamingGrid.GridLevels)
+		{
+			for (FSpatialHashStreamingGridLayerCell& GridLayerCell : GridLevel.LayerCells)
+			{
+				for (UWorldPartitionRuntimeSpatialHashCell* Cell : GridLayerCell.GridCells)
+				{
+					Cell->SetGridName(StreamingGrid.GridName);
+				}
+			}
+		}
+	}
+
+	return StreamingObject;
+}
+
 void UWorldPartitionRuntimeSpatialHash::SetDefaultValues()
 {
 	check(!Grids.Num());
@@ -849,7 +882,7 @@ void UWorldPartitionRuntimeSpatialHash::DumpStateLog(FHierarchicalLogArchive& Ar
 {
 	Super::DumpStateLog(Ar);
 
-	for (FSpatialHashStreamingGrid& StreamingGrid : StreamingGrids)
+	ForEachStreamingGrid([&Ar, this](FSpatialHashStreamingGrid& StreamingGrid)
 	{
 		Ar.Printf(TEXT("----------------------------------------------------------------------------------------------------------------"));
 		Ar.Printf(TEXT("%s - Runtime Spatial Hash - Streaming Grid - %s"), *GetWorld()->GetName(), *StreamingGrid.GridName.ToString());
@@ -923,7 +956,7 @@ void UWorldPartitionRuntimeSpatialHash::DumpStateLog(FHierarchicalLogArchive& Ar
 			}
 		}
 		Ar.Printf(TEXT(""));
-	}	
+	});
 }
 
 FName UWorldPartitionRuntimeSpatialHash::GetCellName(UWorldPartition* WorldPartition, FName InGridName, const FGridCellCoord& InCellGlobalCoord, const FDataLayersID& InDataLayerID)
@@ -1183,6 +1216,7 @@ TArray<UWorldPartitionRuntimeCell*> UWorldPartitionRuntimeSpatialHash::GetAlways
 bool UWorldPartitionRuntimeSpatialHash::PopulateGeneratorPackageForCook(const TArray<ICookPackageSplitter::FGeneratedPackageForPreSave>& InGeneratedPackages, TArray<UPackage*>& OutModifiedPackages)
 {
 	check(IsRunningCookCommandlet());
+	check(ExternalStreamingObjects.IsEmpty()); // Make sure we don't have external streaming objects grids so we won't gather their cells
 
 	OutModifiedPackages.Reset();
 	for (UWorldPartitionRuntimeCell* Cell : GetAlwaysLoadedCells())
@@ -1263,7 +1297,7 @@ int32 UWorldPartitionRuntimeSpatialHash::GetAllStreamingCells(TSet<const UWorldP
 {
 	const bool bShouldConsiderClientOnlyVisible = ShouldConsiderClientOnlyVisibleCells();
 
-	for (const FSpatialHashStreamingGrid& StreamingGrid : StreamingGrids)
+	ForEachStreamingGrid([&](const FSpatialHashStreamingGrid& StreamingGrid)
 	{
 		if (!StreamingGrid.bClientOnlyVisible || bShouldConsiderClientOnlyVisible)
 		{
@@ -1285,7 +1319,7 @@ int32 UWorldPartitionRuntimeSpatialHash::GetAllStreamingCells(TSet<const UWorldP
 				}
 			}
 		}
-	}
+	});
 
 	return Cells.Num();
 }
@@ -1315,38 +1349,146 @@ bool UWorldPartitionRuntimeSpatialHash::GetStreamingCells(const TArray<FWorldPar
 	if (Sources.Num() == 0)
 	{
 		// Get always loaded cells
-		for (const FSpatialHashStreamingGrid& StreamingGrid : StreamingGrids)
+		ForEachStreamingGrid([&](const FSpatialHashStreamingGrid& StreamingGrid)
 		{
 			if (!StreamingGrid.bClientOnlyVisible || bShouldConsiderClientOnlyVisible)
 			{
 				StreamingGrid.GetAlwaysLoadedCells(DataLayerSubsystem, OutActivateCells.GetCells(), OutLoadCells.GetCells());
 			}
-		}
+		});
 	}
 	else
 	{
 		// Get cells based on streaming sources
-		for (const FSpatialHashStreamingGrid& StreamingGrid : StreamingGrids)
+		ForEachStreamingGrid([&](const FSpatialHashStreamingGrid& StreamingGrid)
 		{
 			if (!StreamingGrid.bClientOnlyVisible || bShouldConsiderClientOnlyVisible)
 			{
 				StreamingGrid.GetCells(Sources, DataLayerSubsystem, OutActivateCells, OutLoadCells, GetEffectiveEnableZCulling(bEnableZCulling));
 			}
-		}
+		});
 	}
 
 	return !!(OutActivateCells.Num() + OutLoadCells.Num());
 }
 
-const TMap<FName, const FSpatialHashStreamingGrid*>& UWorldPartitionRuntimeSpatialHash::GetNameToGridMapping() const
+bool UWorldPartitionRuntimeSpatialHash::InjectExternalStreamingObject(URuntimeHashExternalStreamingObjectBase* ExternalStreamingObject)
+{
+	check(GetWorld()->IsGameWorld());
+
+	if (!ExternalStreamingObjects.Contains(ExternalStreamingObject))
+	{
+		if (URuntimeSpatialHashExternalStreamingObject* SpatialHashStreamingObject = Cast<URuntimeSpatialHashExternalStreamingObject>(ExternalStreamingObject))
+		{
+			ExternalStreamingObjects.Add(SpatialHashStreamingObject);
+			bIsNameToGridMappingDirty = true;
+			return true;
+		}
+		else
+		{
+			UE_LOG(LogWorldPartition, Error, TEXT("Incompatible external streaming object to inject. Hash Type: %s. Streaming Object Type: %s."), 
+				*UWorldPartitionRuntimeSpatialHash::StaticClass()->GetName(), *SpatialHashStreamingObject->GetClass()->GetName());
+		}
+	}
+	return false;
+}
+
+bool UWorldPartitionRuntimeSpatialHash::RemoveExternalStreamingObject(URuntimeHashExternalStreamingObjectBase* ExternalStreamingObject)
+{
+	check(GetWorld()->IsGameWorld());
+
+	if (URuntimeSpatialHashExternalStreamingObject* SpatialHashStreamingObject = Cast<URuntimeSpatialHashExternalStreamingObject>(ExternalStreamingObject))
+	{
+		if (ExternalStreamingObjects.Remove(SpatialHashStreamingObject) > 0)
+		{
+			bIsNameToGridMappingDirty = true;
+			return true;
+		}
+	}
+	else
+	{
+		UE_LOG(LogWorldPartition, Error, TEXT("Incompatible external streaming object to remove. Hash Type: %s. Streaming Object Type: %s."),
+			*UWorldPartitionRuntimeSpatialHash::StaticClass()->GetName(), *SpatialHashStreamingObject->GetClass()->GetName());
+	}
+	
+	return false;
+}
+
+uint32 UWorldPartitionRuntimeSpatialHash::GetNumGrids() const
 {
-	if (NameToGridMapping.IsEmpty())
+	uint32 NumGrids = Algo::Accumulate(ExternalStreamingObjects, StreamingGrids.Num(),
+		[](uint32 AccumulatedNumGrids, const TWeakObjectPtr<URuntimeSpatialHashExternalStreamingObject>& StreamingObject)
+	{
+		if (StreamingObject.IsValid())
+		{
+			AccumulatedNumGrids += StreamingObject->StreamingGrids.Num();
+		}
+
+		return AccumulatedNumGrids;
+
+	});
+
+	return NumGrids;
+}
+
+void UWorldPartitionRuntimeSpatialHash::ForEachStreamingGrid(TFunctionRef<void(FSpatialHashStreamingGrid&)> Func)
+{
+	for (FSpatialHashStreamingGrid& StreamingGrid : StreamingGrids)
 	{
-		for (const FSpatialHashStreamingGrid& StreamingGrid : StreamingGrids)
+		Func(StreamingGrid);
+	}
+
+	for (TWeakObjectPtr<URuntimeSpatialHashExternalStreamingObject>& StreamingObject : ExternalStreamingObjects)
+	{
+		if (StreamingObject.IsValid())
 		{
-			NameToGridMapping.Add(StreamingGrid.GridName, &StreamingGrid);
+			for (FSpatialHashStreamingGrid& StreamingGrid : StreamingObject->StreamingGrids)
+			{
+				Func(StreamingGrid);
+			}
 		}
 	}
+}
+
+void UWorldPartitionRuntimeSpatialHash::ForEachStreamingGrid(TFunctionRef<void(const FSpatialHashStreamingGrid&)> Func) const
+{
+	for (const FSpatialHashStreamingGrid& StreamingGrid : StreamingGrids)
+	{
+		Func(StreamingGrid);
+	}
+
+	for (const TWeakObjectPtr<URuntimeSpatialHashExternalStreamingObject>& StreamingObject : ExternalStreamingObjects)
+	{
+		if (StreamingObject.IsValid())
+		{
+			for (const FSpatialHashStreamingGrid& StreamingGrid : StreamingObject->StreamingGrids)
+			{
+				Func(StreamingGrid);
+			}
+		}
+	}
+}
+
+const TMap<FName, const FSpatialHashStreamingGrid*>& UWorldPartitionRuntimeSpatialHash::GetNameToGridMapping() const
+{
+	if (bIsNameToGridMappingDirty)
+	{
+		ForEachStreamingGrid([&](const FSpatialHashStreamingGrid& StreamingGrid)
+		{
+			// Only keep the first one if names are clashing. ForEachStreamingGrids start iterating with the main grids then iterates on external Objects.
+			// Prioritize the main grid in the mapping.
+			if (!NameToGridMapping.Contains(StreamingGrid.GridName))
+			{
+				NameToGridMapping.Add(StreamingGrid.GridName, &StreamingGrid);
+			}
+			else
+			{
+				UE_LOG(LogWorldPartition, Error, TEXT("2 Streaming grids share the name: \"%s\". Some features will not work properly."), *StreamingGrid.GridName.ToString());
+			}
+		});
+
+		bIsNameToGridMappingDirty = false;
+	}
 
 	return NameToGridMapping;
 }
@@ -1501,11 +1643,19 @@ bool UWorldPartitionRuntimeSpatialHash::ContainsRuntimeHash(const FString& Name)
 
 TArray<const FSpatialHashStreamingGrid*> UWorldPartitionRuntimeSpatialHash::GetFilteredStreamingGrids() const
 {
+	URuntimeSpatialHashExternalStreamingObject* ExternalStreamingObject = NewObject<URuntimeSpatialHashExternalStreamingObject>(const_cast<UWorldPartitionRuntimeSpatialHash*>(this), TEXT("test"));
+	const_cast<UWorldPartitionRuntimeSpatialHash*>(this)->InjectExternalStreamingObject(ExternalStreamingObject);
+
 	TArray<const FSpatialHashStreamingGrid*> FilteredStreamingGrids;
-	FilteredStreamingGrids.Reserve(StreamingGrids.Num());
-	Algo::TransformIf(StreamingGrids, FilteredStreamingGrids,
-		[](const FSpatialHashStreamingGrid& StreamingGrid) { return FWorldPartitionDebugHelper::IsDebugRuntimeHashGridShown(StreamingGrid.GridName); },
-		[](const FSpatialHashStreamingGrid& StreamingGrid) { return &StreamingGrid; });
+	FilteredStreamingGrids.Reserve(GetNumGrids());
+	ForEachStreamingGrid([&FilteredStreamingGrids](const FSpatialHashStreamingGrid& StreamingGrid)
+	{
+		if (FWorldPartitionDebugHelper::IsDebugRuntimeHashGridShown(StreamingGrid.GridName))
+		{
+			FilteredStreamingGrids.Add(&StreamingGrid);
+		}
+	});
+
 	return FilteredStreamingGrids;
 }
 
diff --git a/Engine/Source/Runtime/Engine/Public/WorldPartition/WorldPartitionRuntimeCell.h b/Engine/Source/Runtime/Engine/Public/WorldPartition/WorldPartitionRuntimeCell.h
--- a/Engine/Source/Runtime/Engine/Public/WorldPartition/WorldPartitionRuntimeCell.h
+++ b/Engine/Source/Runtime/Engine/Public/WorldPartition/WorldPartitionRuntimeCell.h
@@ -204,6 +204,7 @@ class ENGINE_API UWorldPartitionRuntimeCell : public UObject, public IWorldParti
 	void SetDataLayers(const TArray<const UDataLayerInstance*>& InDataLayerInstances);
 	void SetMinMaxZ(const FVector2D& InMinMaxZ);
 	void SetDebugInfo(int64 InCoordX, int64 InCoordY, int64 InCoordZ, FName InGridName);
+	void SetGridName(FName InGridName);
 	virtual void AddActorToCell(const FWorldPartitionActorDescView& ActorDescView, const FActorContainerID& InContainerID, const FTransform& InContainerTransform, const UActorDescContainer* InContainer) PURE_VIRTUAL(UWorldPartitionRuntimeCell::AddActorToCell,);
 	virtual int32 GetActorCount() const PURE_VIRTUAL(UWorldPartitionRuntimeCell::GetActorCount, return 0;);
 
diff --git a/Engine/Source/Runtime/Engine/Public/WorldPartition/WorldPartitionRuntimeHash.h b/Engine/Source/Runtime/Engine/Public/WorldPartition/WorldPartitionRuntimeHash.h
--- a/Engine/Source/Runtime/Engine/Public/WorldPartition/WorldPartitionRuntimeHash.h
+++ b/Engine/Source/Runtime/Engine/Public/WorldPartition/WorldPartitionRuntimeHash.h
@@ -29,6 +29,15 @@ enum class EWorldPartitionStreamingPerformance : uint8
 	Critical
 };
 
+UCLASS(Abstract)
+class ENGINE_API URuntimeHashExternalStreamingObjectBase : public UObject
+{
+	GENERATED_BODY()
+
+public:
+
+};
+
 UCLASS(Abstract, Config=Engine, AutoExpandCategories=(WorldPartition), Within = WorldPartition)
 class ENGINE_API UWorldPartitionRuntimeHash : public UObject
 {
@@ -43,6 +52,8 @@ class ENGINE_API UWorldPartitionRuntimeHash : public UObject
 	virtual bool GenerateHLOD(ISourceControlHelper* SourceControlHelper, const IStreamingGenerationContext* StreamingGenerationContext, bool bCreateActorsOnly) { return false; }
 	virtual void DrawPreview() const {}
 
+	virtual URuntimeHashExternalStreamingObjectBase* StoreToExternalStreamingObject(UObject* StreamingObjectOuter, FName StreamingObjectName) { return nullptr; }
+
 	virtual void DumpStateLog(FHierarchicalLogArchive& Ar);
 
 	// PIE/Game methods
@@ -69,6 +80,9 @@ class ENGINE_API UWorldPartitionRuntimeHash : public UObject
 	virtual void SortStreamingCellsByImportance(const TSet<const UWorldPartitionRuntimeCell*>& InCells, const TArray<FWorldPartitionStreamingSource>& InSources, TArray<const UWorldPartitionRuntimeCell*, TInlineAllocator<256>>& OutSortedCells) const;
 	EWorldPartitionStreamingPerformance GetStreamingPerformance(const TSet<const UWorldPartitionRuntimeCell*>& CellToActivate) const;
 
+	virtual bool InjectExternalStreamingObject(URuntimeHashExternalStreamingObjectBase* ExternalStreamingObject) { return false; }
+	virtual bool RemoveExternalStreamingObject(URuntimeHashExternalStreamingObjectBase* ExternalStreamingObject) { return false; }
+
 	/* Returns desired footprint that Draw2D should take relative to given Canvas size (the value can exceed the given size).
 	 * UWorldPartitionSubSystem will re-adapt the size relative to all others UWorldPartitionRuntimeHash and provide the correct size to Draw2D.
 	 *
diff --git a/Engine/Source/Runtime/Engine/Public/WorldPartition/WorldPartitionRuntimeSpatialHash.h b/Engine/Source/Runtime/Engine/Public/WorldPartition/WorldPartitionRuntimeSpatialHash.h
--- a/Engine/Source/Runtime/Engine/Public/WorldPartition/WorldPartitionRuntimeSpatialHash.h
+++ b/Engine/Source/Runtime/Engine/Public/WorldPartition/WorldPartitionRuntimeSpatialHash.h
@@ -214,10 +214,20 @@ class ASpatialHashRuntimeGridInfo : public AInfo
 	GENERATED_UCLASS_BODY()
 
 public:
-	UPROPERTY(EditAnywhere, Category=Settings)
+	UPROPERTY()
 	FSpatialHashRuntimeGrid	GridSettings;
 };
 
+UCLASS()
+class URuntimeSpatialHashExternalStreamingObject : public URuntimeHashExternalStreamingObjectBase
+{
+	GENERATED_BODY()
+
+public:
+	UPROPERTY();
+	TArray<FSpatialHashStreamingGrid> StreamingGrids;
+};
+
 UCLASS()
 class ENGINE_API UWorldPartitionRuntimeSpatialHash : public UWorldPartitionRuntimeHash
 {
@@ -238,6 +248,8 @@ public:
 	virtual bool GenerateHLOD(ISourceControlHelper* SourceControlHelper, const IStreamingGenerationContext* StreamingGenerationContext, bool bCreateActorsOnly) override;
 	virtual void DrawPreview() const override;
 
+	virtual URuntimeHashExternalStreamingObjectBase* StoreToExternalStreamingObject(UObject* StreamingObjectOuter, FName StreamingObjectName) override;
+
 	FName GetCellName(FName InGridName, const FGridCellCoord& InCellGlobalCoord, const FDataLayersID& InDataLayerID) const;
 	static FName GetCellName(UWorldPartition* WorldPartition, FName InGridName, const FGridCellCoord& InCellGlobalCoord, const FDataLayersID& InDataLayerID);
 
@@ -249,10 +261,18 @@ public:
 	virtual int32 GetAllStreamingCells(TSet<const UWorldPartitionRuntimeCell*>& Cells, bool bAllDataLayers = false, bool bDataLayersOnly = false, const TSet<FName>& InDataLayers = TSet<FName>()) const override;
 	virtual bool GetStreamingCells(const FWorldPartitionStreamingQuerySource& QuerySource, TSet<const UWorldPartitionRuntimeCell*>& OutCells) const override;
 	virtual bool GetStreamingCells(const TArray<FWorldPartitionStreamingSource>& Sources, UWorldPartitionRuntimeHash::FStreamingSourceCells& OutActivateCells, UWorldPartitionRuntimeHash::FStreamingSourceCells& OutLoadCells) const override;
+
+	virtual bool InjectExternalStreamingObject(URuntimeHashExternalStreamingObjectBase* ExternalStreamingObject) override;
+	virtual bool RemoveExternalStreamingObject(URuntimeHashExternalStreamingObjectBase* ExternalStreamingObject) override;
+
+	uint32 GetNumGrids() const;
 		
 	static FString GetCellCoordString(const FGridCellCoord& InCellGlobalCoord);
 
 protected:
+	void ForEachStreamingGrid(TFunctionRef<void(FSpatialHashStreamingGrid&)> Func);
+	void ForEachStreamingGrid(TFunctionRef<void(const FSpatialHashStreamingGrid&)> Func) const;
+
 	bool ShouldConsiderClientOnlyVisibleCells() const;
 	virtual EWorldPartitionStreamingPerformance GetStreamingPerformanceForCell(const UWorldPartitionRuntimeCell* Cell) const override;
 
@@ -288,8 +308,11 @@ private:
 	 */
 	UPROPERTY(NonPIEDuplicateTransient)
 	TArray<FSpatialHashStreamingGrid> StreamingGrids;
-
 	mutable TMap<FName, const FSpatialHashStreamingGrid*> NameToGridMapping;
+	mutable bool bIsNameToGridMappingDirty;
+
+	UPROPERTY(Transient)
+	TArray<TWeakObjectPtr<URuntimeSpatialHashExternalStreamingObject>> ExternalStreamingObjects;
 	
 	virtual FVector2D GetDraw2DDesiredFootprint(const FVector2D& CanvasSize) const override;
 	virtual void Draw2D(class UCanvas* Canvas, const TArray<FWorldPartitionStreamingSource>& Sources, const FVector2D& PartitionCanvasSize, const FVector2D& Offset) const override;
