Removed debug code

diff --git a/Engine/Source/Runtime/Engine/Private/WorldPartition/WorldPartitionRuntimeSpatialHash.cpp b/Engine/Source/Runtime/Engine/Private/WorldPartition/WorldPartitionRuntimeSpatialHash.cpp
--- a/Engine/Source/Runtime/Engine/Private/WorldPartition/WorldPartitionRuntimeSpatialHash.cpp
+++ b/Engine/Source/Runtime/Engine/Private/WorldPartition/WorldPartitionRuntimeSpatialHash.cpp
@@ -1643,9 +1643,6 @@ bool UWorldPartitionRuntimeSpatialHash::ContainsRuntimeHash(const FString& Name)
 
 TArray<const FSpatialHashStreamingGrid*> UWorldPartitionRuntimeSpatialHash::GetFilteredStreamingGrids() const
 {
-	URuntimeSpatialHashExternalStreamingObject* ExternalStreamingObject = NewObject<URuntimeSpatialHashExternalStreamingObject>(const_cast<UWorldPartitionRuntimeSpatialHash*>(this), TEXT("test"));
-	const_cast<UWorldPartitionRuntimeSpatialHash*>(this)->InjectExternalStreamingObject(ExternalStreamingObject);
-
 	TArray<const FSpatialHashStreamingGrid*> FilteredStreamingGrids;
 	FilteredStreamingGrids.Reserve(GetNumGrids());
 	ForEachStreamingGrid([&FilteredStreamingGrids](const FSpatialHashStreamingGrid& StreamingGrid)
