Reverts dynamic res to old heuristic until further time is spent on multiple-dynres heuristic

diff --git a/Engine/Source/Runtime/Engine/Private/DynamicResolution.cpp b/Engine/Source/Runtime/Engine/Private/DynamicResolution.cpp
--- a/Engine/Source/Runtime/Engine/Private/DynamicResolution.cpp
+++ b/Engine/Source/Runtime/Engine/Private/DynamicResolution.cpp
@@ -217,304 +217,228 @@ void FDynamicResolutionHeuristicProxy::CommitPreviousFrameGPUTimings_RenderThrea
 
 void FDynamicResolutionHeuristicProxy::RefreshCurentFrameResolutionFraction_RenderThread()
 {
-	// Compute new frame resolution fractions only if have an history to work with.
-	if (HistorySize == 0)
-	{
-		return;
-	}
-
 	// GPU time budget per frame=.
 	const float FrameTimeBudgetMs = CVarFrameTimeBudget.GetValueOnRenderThread();
 
 	// Targeted GPU time, lower than budget to diggest noise.
 	const float TargetedGPUBusyTimeMs = FrameTimeBudgetMs * (1.0f - CVarTargetedGPUHeadRoomPercentage.GetValueOnRenderThread() / 100.0f);
 
-	const float FrameWeightExponent = CVarFrameWeightExponent.GetValueOnRenderThread();
-	const int32 MaxConsecutiveOverbudgetGPUFrameCount = FMath::Max(CVarMaxConsecutiveOverbudgetGPUFrameCount.GetValueOnRenderThread(), 2);
-	
-	// Find the first frame that have data.
-	int32 StartFrameId = 0;
-	{
-		for (; StartFrameId < HistorySize; StartFrameId++)
-		{
-			const FrameHistoryEntry& FrameEntry = GetPreviousFrameEntry(StartFrameId);
+	// Resolution fraction to downscale to when CPU bound.
+	const float CPUBoundResolutionFraction = 1.0f; // CVarCPUBoundScreenPercentage.GetValueOnRenderThread() / 100.0f;
 
-			// Ignores frames that does not have any GPU timing yet. 
-			if (FrameEntry.TotalFrameGPUBusyTimeMs < 0)
-			{
-				continue;
-			}
-			else
-			{
-				break;
-			}
-		}
-	}
-
-	// Not enough data to work with.
-	if (HistorySize - StartFrameId <= MaxConsecutiveOverbudgetGPUFrameCount)
-	{
-		return;
-	}
+	// New ResolutionFraction to use for this frame.
+	float NewFrameResolutionFraction = CurrentFrameResolutionFractions[GDynamicPrimaryResolutionFraction];
 
-	check(StartFrameId + MaxConsecutiveOverbudgetGPUFrameCount < HistorySize);
+	// Whether there is a GPU over budget panic.
+	bool bGPUOverbugetPanic = false;
 
-	// Look whether all the last MaxConsecutiveOverbudgetGPUFrameCount frames are over budget.
-	bool bGPUFrameOverbudgetPanic = true;
-	for (int32 BrowsingFrameId = StartFrameId; BrowsingFrameId < StartFrameId + MaxConsecutiveOverbudgetGPUFrameCount; BrowsingFrameId++)
+	// Compute NewFrameResolutionFraction if have an history to work with.
+	if (HistorySize > 0)
 	{
-		const FrameHistoryEntry& FrameEntry = GetPreviousFrameEntry(BrowsingFrameId);
-		bGPUFrameOverbudgetPanic &= (FrameEntry.TotalFrameGPUBusyTimeMs > FrameTimeBudgetMs);
-	}
-
-	// Whether a cruising resolution change can happen.
-	bool bCanChangeResolution = NumberOfFramesSinceScreenPercentageChange >= CVarFrameChangePeriod.GetValueOnRenderThread();
-
-	// Array to know how much a frame is overbudget
-	//TArray<float, TInlineAllocator<256>> TotalBudgetOverBudgetSavings;
-	//TotalBudgetOverBudgetSavings.SetNumZeroed(MaxFrameId);
-
-	// Function to detect whether the history size is big enough for detecting an overbudget
-	auto HasNotEnoughDataForOverBudgetDetection = [&](int32 InHistorySize)
-	{
-		return (InHistorySize - StartFrameId) <= MaxConsecutiveOverbudgetGPUFrameCount;
-	};
+		const float FrameWeightExponent = CVarFrameWeightExponent.GetValueOnRenderThread();
+		const float RenderThreadHeadRoomFromGPUMs = 1.0f; //CVarCPUTimeHeadRoom.GetValueOnRenderThread();
+		const int32 MaxConsecutiveOverbudgetGPUFrameCount = FMath::Max(CVarMaxConsecutiveOverbudgetGPUFrameCount.GetValueOnRenderThread(), 2);
 
-	// Function to know whether a bucket is overbudget.
-	auto IsBudgetOverBudget = [&](const DynamicRenderScaling::FBudget& Budget, const DynamicRenderScaling::FHeuristicSettings& BudgetSetting)
-	{
-		check(BudgetSetting.IsEnabled());
+		const float OutLierTimeThreshold = 0.0f; //FrameTimeBudgetMs * CVarOutlierThreshold.GetValueOnRenderThread();
 
-		if (HasNotEnoughDataForOverBudgetDetection(BudgetHistorySizes[Budget]))
-		{
-			// Not enough data to work with.
-			return false;
-		}
+		NewFrameResolutionFraction = 0.0f;
 
-		bool bBudgetIsOverBudget = true;
+		// Total weight of NewFrameResolutionFraction.
+		float TotalWeight = 0.0f;
 
-		for (int32 BrowsingFrameId = StartFrameId; BrowsingFrameId < StartFrameId + MaxConsecutiveOverbudgetGPUFrameCount; BrowsingFrameId++)
-		{
-			const FrameHistoryEntry& FrameEntry = GetPreviousFrameEntry(BrowsingFrameId);
+		// Frame weight.
+		float Weight = 1.0f;
 
-			float bOverBudget = FrameEntry.BudgetTimingMs[Budget] - BudgetSetting.BudgetMs;
-			bBudgetIsOverBudget &= bOverBudget > 0.0f;
-		}
+		// Number of consecutive frames that have over budget GPU.
+		int32 ConsecutiveOverbudgetGPUFramCount = 0;
 
-		return bBudgetIsOverBudget;
-	};
+		// Whether should continue browsing the frame history.
+		bool bCaryOnBrowsingFrameHistory = true;
 
-	// Compute
-	auto ComputeMaxFrameId = [&](const DynamicRenderScaling::FBudget& Budget, bool bBudgetIsOverBudget)
-	{
-		int32 MaxFrameId = BudgetHistorySizes[Budget];
-		
-		if (bGPUFrameOverbudgetPanic || bBudgetIsOverBudget)
+		// Number of frame browsed.
+		int32 FrameCount = 0;
+		int32 BrowsingFrameId = 0;
+		for (; BrowsingFrameId < HistorySize && bCaryOnBrowsingFrameHistory; BrowsingFrameId++)
 		{
-			MaxFrameId = FMath::Min(MaxFrameId, StartFrameId + MaxConsecutiveOverbudgetGPUFrameCount);
-		}
+			const FrameHistoryEntry& NextFrameEntry = GetPreviousFrameEntry(BrowsingFrameId - 1);
+			const FrameHistoryEntry& FrameEntry = GetPreviousFrameEntry(BrowsingFrameId + 0);
+			const FrameHistoryEntry& PreviousFrameEntry = GetPreviousFrameEntry(BrowsingFrameId + 1);
 
-		check(MaxFrameId <= BudgetHistorySizes[Budget]);
-		check(MaxFrameId <= HistorySize);
-		return MaxFrameId;
-	};
-
-	DynamicRenderScaling::TMap<float> NewResolutionFractions;
-	bool bCommitResolutionChange = false;
-
-	// Execute heuristic on the bucket.
-	auto ComputeNewResolutionFraction = [&](const DynamicRenderScaling::FBudget& Budget, const DynamicRenderScaling::FHeuristicSettings& BudgetSetting, bool bBudgetIsOverBudget)
-	{
-		check(BudgetSetting.IsEnabled());
-
-		float NewResolutionFraction = CurrentFrameResolutionFractions[Budget];
-		if (!HasNotEnoughDataForOverBudgetDetection(BudgetHistorySizes[Budget]))
-		{
-			float CurrentResolutionFraction = CurrentFrameResolutionFractions[Budget];
-			if (BudgetSetting.bModelScalesWithPrimaryScreenPercentage)
+			// Ignores frames that does not have any GPU timing yet. 
+			if (FrameEntry.TotalFrameGPUBusyTimeMs < 0)
 			{
-				CurrentResolutionFraction *= CurrentFrameResolutionFractions[GDynamicPrimaryResolutionFraction];
+				continue;
 			}
 
-			int32 MaxFrameId = ComputeMaxFrameId(Budget, bBudgetIsOverBudget);
-
-			// Find out a new resolution fraction based on whether this bucket is overbudget.
-			float MostRecentResolutionFractionScale = 1.0f;
-			float SuggestedResolutionFraction = 0.0f;
-			float SuggestedResolutionFractionWeight = 0.0f;
-			float Weight = 1.0f;
-			for (int32 BrowsingFrameId = StartFrameId; BrowsingFrameId < MaxFrameId; BrowsingFrameId++)
+			const bool bGPUTimingsHaveCPUBubbles = false;
+
+			// Compute the maximum thread time between game and render thread for this <i>th frame.
+			// We add a head room for render thread to detect when the GPU could be starving.
+			const float MaxThreadTimeMs = FMath::Max(
+				FrameEntry.RenderThreadTimeMs + (bGPUTimingsHaveCPUBubbles ? RenderThreadHeadRoomFromGPUMs : 0),
+				FrameEntry.GameThreadTimeMs);
+
+			// Whether bound by game thread.
+			const bool bIsGameThreadBound = FrameEntry.GameThreadTimeMs > FrameTimeBudgetMs;
+
+			// Whether bound by render thread.
+			const bool bIsRenderThreadBound = (FrameEntry.RenderThreadTimeMs + (bGPUTimingsHaveCPUBubbles ? RenderThreadHeadRoomFromGPUMs : 0)) > TargetedGPUBusyTimeMs;
+
+			// Whether the frame is CPU bound.
+			const bool bIsCPUBound = bIsGameThreadBound || bIsRenderThreadBound;
+
+			// Whether the render thread might create bubbles in GPU timings.
+			const bool bMayHaveGPUBubbles = bGPUTimingsHaveCPUBubbles && bIsRenderThreadBound;
+
+			// Whether GPU is over budget, when not CPU bound.
+			const bool bHasOverbudgetGPU = !bIsCPUBound && FrameEntry.TotalFrameGPUBusyTimeMs > FrameTimeBudgetMs;
+
+			// Whether this is possibly an outlier compared to previous and next frame.
+			// Normally should look for out lier compared to previous and next frames timings, however
+			// only interested in detected them if they are over budget, when cruising at targeted
+			// GPU busy time.
+			const bool bIsGPUOutlier =
+				bHasOverbudgetGPU &&
+				ConsecutiveOverbudgetGPUFramCount == 0 &&
+				OutLierTimeThreshold > 0.0f &&
+				FrameEntry.TotalFrameGPUBusyTimeMs > OutLierTimeThreshold &&
+				(!NextFrameEntry.HasGPUTimings() || (
+					NextFrameEntry.TotalFrameGPUBusyTimeMs < OutLierTimeThreshold&&
+					NextFrameEntry.GameThreadTimeMs < OutLierTimeThreshold&&
+					NextFrameEntry.RenderThreadTimeMs < OutLierTimeThreshold)) &&
+				(!PreviousFrameEntry.HasGPUTimings() || (
+					PreviousFrameEntry.TotalFrameGPUBusyTimeMs < OutLierTimeThreshold&&
+					PreviousFrameEntry.GameThreadTimeMs < OutLierTimeThreshold&&
+					PreviousFrameEntry.RenderThreadTimeMs < OutLierTimeThreshold));
+
+			// Look if this is multiple consecutive GPU over budget frames.
+			if (bHasOverbudgetGPU)
 			{
-				const FrameHistoryEntry& FrameEntry = GetPreviousFrameEntry(BrowsingFrameId);
-				float TimingMs = FrameEntry.BudgetTimingMs[Budget];
+				ConsecutiveOverbudgetGPUFramCount++;
+				check(ConsecutiveOverbudgetGPUFramCount <= MaxConsecutiveOverbudgetGPUFrameCount);
 
-				float ResolutionFraction = FrameEntry.ResolutionFractions[Budget];
-				if (BudgetSetting.bModelScalesWithPrimaryScreenPercentage)
+				// Max number of over budget frames where reached.
+				if (ConsecutiveOverbudgetGPUFramCount == MaxConsecutiveOverbudgetGPUFrameCount)
 				{
-					ResolutionFraction *= FrameEntry.ResolutionFractions[GDynamicPrimaryResolutionFraction];
-				}
+					// We ignore frames in history that happen before consecutive GPU overbudget frames.
+					bCaryOnBrowsingFrameHistory = false;
 
-				// Removing the savings from other bucket
-				//if (Budget != EDynamicResolutionTimingBudget::PrimaryScreenPercentage)
-				//{
-				//	TimingMs -= TotalBudgetOverBudgetSavings[BrowsingFrameId];
-				//}
-
-				// Estimate how much a resolution fraction should be changed.
-				float ResolutionFractionFactor = BudgetSetting.EstimateResolutionFactor(TimingMs);
-
-				// Don't increase the resolution if frame overbudget.
-				if (bGPUFrameOverbudgetPanic && ResolutionFractionFactor > 1.0f)
-				{
-					ResolutionFractionFactor = 1.0f;
-				}
-
-				float AdjustedResolutionFraction = ResolutionFractionFactor * ResolutionFraction;
-
-				// Track how much saving can be done in a bucket to max out PrimaryScreenPercentage
-				//if (Budget != EDynamicResolutionTimingBudget::PrimaryScreenPercentage)
-				//{
-				//	float EstimatedSavings = (
-				//		FrameEntry.BudgetTimingMs[Budget] *
-				//		(1.0f - BudgetSetting.EstimateTimeFactor(FrameEntry.ResolutionFractions[Budget], AdjustedResolutionFraction)));
-				//
-				//	TotalBudgetOverBudgetSavings[BrowsingFrameId] += EstimatedSavings;
-				//}
-
-				SuggestedResolutionFraction += AdjustedResolutionFraction * Weight;
-				SuggestedResolutionFractionWeight += Weight;
-				Weight *= FrameWeightExponent;
-
-				if (BrowsingFrameId == StartFrameId && BudgetSetting.bModelScalesWithPrimaryScreenPercentage)
-				{
-					MostRecentResolutionFractionScale = FrameEntry.ResolutionFractions[GDynamicPrimaryResolutionFraction];
+					ensureMsgf(!bIsGPUOutlier, TEXT("Does not make any sens to have ConsecutiveOverbudgetGPUFramCount >= 2 if this frame is detected as being an outlier."));
 				}
 			}
+			else
+			{
+				ConsecutiveOverbudgetGPUFramCount = 0;
+			}
 
-			NewResolutionFraction = BudgetSetting.CorrectNewResolutionFraction(
-				CurrentResolutionFraction,
-				SuggestedResolutionFraction / SuggestedResolutionFractionWeight,
-				MostRecentResolutionFractionScale);
-
-			NewResolutionFraction /= MostRecentResolutionFractionScale;
+			// Ignore if an outlier.
+			if (bIsGPUOutlier)
+			{
+				continue;
+			}
 
-			bool bTakeNewResolutionFraction = (
-				BudgetSetting.DoesResolutionChangeEnough(CurrentResolutionFraction, NewResolutionFraction, bCanChangeResolution) ||
-				bBudgetIsOverBudget);
+			float SuggestedResolutionFraction = 1.0f;
 
-			if (bTakeNewResolutionFraction)
+			// If have reliable GPU times, or guess there is not GPU bubbles -> estimate the suggested resolution fraction that could have been used.
+			if (!bGPUTimingsHaveCPUBubbles ||
+				(bGPUTimingsHaveCPUBubbles && !bIsCPUBound && !bMayHaveGPUBubbles))
 			{
-				bCommitResolutionChange = true;
+				// This assumes GPU busy time is directly proportional to ResolutionFraction^2, but in practice
+				// this is more A * ResolutionFraction^2 + B with B >= 0 non constant unknown cost such as unscaled
+				// post processing, vertex fetching & processing, occlusion queries, shadow maps rendering...
+				//
+				// This assumption means we may drop ResolutionFraction lower than needed, or be slower to increase
+				// resolution.
+				//
+				// TODO: If we have RHI guarantee of frame timing association, we could make an estimation of B
+				// At resolution change that happen every N frames, amortized over time and scaled down as the
+				// standard variation of the GPU timing over non resolution changing frames increases.
+				SuggestedResolutionFraction = (
+					FMath::Sqrt(TargetedGPUBusyTimeMs / FrameEntry.TotalFrameGPUBusyTimeMs)
+					* FrameEntry.ResolutionFractions[GDynamicPrimaryResolutionFraction]);
+			}
+			else if (FrameEntry.ResolutionFractions[GDynamicPrimaryResolutionFraction] > CPUBoundResolutionFraction)
+			{
+				// When CPU bound, we no longer want to downscale, or drop resolution if CPU and GPU are sharing same memory.
+				SuggestedResolutionFraction = CPUBoundResolutionFraction;
 			}
 			else
 			{
-				NewResolutionFraction = CurrentFrameResolutionFractions[Budget];
+				// Unable to know reliably the GPU time so don't try to make any prediction to avoid issues such
+				// as resolution keep shrinking because GPU time is slightly more than render thread time, or
+				// increasing because GPU time slightly lower than frame time.
+				SuggestedResolutionFraction = FrameEntry.ResolutionFractions[GDynamicPrimaryResolutionFraction];
 			}
-		}
-
-		return NewResolutionFraction;
-	};
-
-	// Estimate the new cost of the bucket.
-	auto EstimateBudgetCost = [&](const DynamicRenderScaling::FBudget& Budget, bool bBudgetIsOverBudget)
-	{
-		const DynamicRenderScaling::FHeuristicSettings& BudgetSetting = Budget.GetSettings();
-		check(BudgetSetting.IsEnabled());
-
-		int32 MaxFrameId = ComputeMaxFrameId(Budget, bBudgetIsOverBudget);
 
-		float CurrentResolutionFraction = CurrentFrameResolutionFractions[Budget];
-		float NewResolutionFraction = NewResolutionFractions[Budget];
+			NewFrameResolutionFraction += SuggestedResolutionFraction * Weight;
+			TotalWeight += Weight;
+			FrameCount++;
 
-		ensure(MaxFrameId != StartFrameId);
-
-		float BudgetPredictedCost = 0.0;
-		float BudgetPredictedCostWeight = 0.0;
-		float Weight = 1.0f;
-		for (int32 BrowsingFrameId = StartFrameId; BrowsingFrameId < MaxFrameId; BrowsingFrameId++)
-		{
-			const FrameHistoryEntry& FrameEntry = GetPreviousFrameEntry(BrowsingFrameId);
-			float EstimatedCost = (
-				FrameEntry.BudgetTimingMs[Budget] *
-				BudgetSetting.EstimateTimeFactor(CurrentResolutionFraction, NewResolutionFraction));
-
-			BudgetPredictedCost += EstimatedCost * Weight;
-			BudgetPredictedCostWeight += Weight;
 			Weight *= FrameWeightExponent;
 		}
 
-		return BudgetPredictedCost / BudgetPredictedCostWeight;
-	};
+		NewFrameResolutionFraction /= TotalWeight;
 
-	// Iterate through bucket and adjust them.
-	float TotalStaticBudgetPredictedCost = 0.0f;
-	float TotalPrimaryScreenPercentageDynamicBudgetPredictedCost = 0.0f;
-	for (TLinkedList<DynamicRenderScaling::FBudget*>::TIterator BudgetIt(DynamicRenderScaling::FBudget::GetGlobalList()); BudgetIt; BudgetIt.Next())
-	{
-		const DynamicRenderScaling::FBudget& Budget = **BudgetIt;
-		const DynamicRenderScaling::FHeuristicSettings& HeuristicSettings = Budget.GetSettings();
+		// If immediate previous frames where over budget, react immediately.
+		bGPUOverbugetPanic = FrameCount > 0 && ConsecutiveOverbudgetGPUFramCount == FrameCount;
 
-		// Primary screen percentage is done last.
-		if (Budget == GDynamicPrimaryResolutionFraction)
+		// If over budget, reset history size to 0 so that this frame really behave as a first frame after an history reset.
+		if (bGPUOverbugetPanic)
 		{
-			continue;
+			HistorySize = 0;
 		}
-
-		// Skip if the bucket is not enabled.
-		if (HeuristicSettings.IsEnabled())
+		// If not immediately over budget, refine the new resolution fraction.
+		else
 		{
-			bool bBudgetIsOverBudget = IsBudgetOverBudget(Budget, HeuristicSettings);
+			// If scaling the resolution, look if this is above a threshold compared to current res.
+			if (CVarChangeThreshold.GetValueOnRenderThread())
+			{
+				float IncreaseResolutionFractionThreshold = CurrentFrameResolutionFractions[GDynamicPrimaryResolutionFraction] * (1.0f + CVarChangeThreshold.GetValueOnRenderThread() / 100.0f);
 
-			NewResolutionFractions[Budget] = ComputeNewResolutionFraction(Budget, HeuristicSettings, bBudgetIsOverBudget);
+				float DecreaseResolutionFractionThreshold = CurrentFrameResolutionFractions[GDynamicPrimaryResolutionFraction] * (1.0f - CVarChangeThreshold.GetValueOnRenderThread() / 100.0f);
 
-			float BudgetCost = EstimateBudgetCost(Budget, bBudgetIsOverBudget);
-			if (HeuristicSettings.bModelScalesWithPrimaryScreenPercentage)
-			{
-				TotalPrimaryScreenPercentageDynamicBudgetPredictedCost += BudgetCost;
-			}
-			else
-			{
-				TotalStaticBudgetPredictedCost += BudgetCost;
+				if (NewFrameResolutionFraction > CurrentFrameResolutionFractions[GDynamicPrimaryResolutionFraction] && NewFrameResolutionFraction < IncreaseResolutionFractionThreshold)
+				{
+					NewFrameResolutionFraction = CurrentFrameResolutionFractions[GDynamicPrimaryResolutionFraction];
+				}
 			}
 
-			// Reset the size of the history for the bucket when it is over budget.
-			if (bBudgetIsOverBudget)
+			// If scaling the resolution up, amortize to avoid oscillations.
+			if (NewFrameResolutionFraction > CurrentFrameResolutionFractions[GDynamicPrimaryResolutionFraction])
 			{
-				BudgetHistorySizes[Budget] = 0;
+				NewFrameResolutionFraction = FMath::Lerp(
+					CurrentFrameResolutionFractions[GDynamicPrimaryResolutionFraction],
+					NewFrameResolutionFraction,
+					CVarIncreaseAmortizationFactor.GetValueOnRenderThread());
 			}
 		}
-		else
-		{
-			NewResolutionFractions[Budget] = 1.0f;
-			continue;
-		}
+
+		// Clamp resolution fraction.
+		NewFrameResolutionFraction = FMath::Clamp(
+			NewFrameResolutionFraction,
+			CVarDynamicResMinSP.GetValueOnRenderThread() / 100.0f,
+			CurrentFrameMaxResolutionFractions[GDynamicPrimaryResolutionFraction]);
 	}
 
-	// Takes care of primary screen percentage to finally fit the frame within budget
+	// Update the current frame's resolution fraction.
 	{
-		const DynamicRenderScaling::FBudget& Budget = GDynamicPrimaryResolutionFraction;
-
-		// Set the budget of the budget.
-		DynamicRenderScaling::FHeuristicSettings HeuristicSettings = Budget.GetSettings();
-		HeuristicSettings.BudgetMs = FMath::Max(TargetedGPUBusyTimeMs - TotalStaticBudgetPredictedCost, 0.5f);
-		check(HeuristicSettings.IsEnabled());
-
-		bool bBudgetIsOverBudget = IsBudgetOverBudget(Budget, HeuristicSettings);
+		// CVarChangeThreshold avoids very small changes.
+		bool bWouldBeWorthChangingRes = CurrentFrameResolutionFractions[GDynamicPrimaryResolutionFraction] != NewFrameResolutionFraction;
 
-		NewResolutionFractions[Budget] = ComputeNewResolutionFraction(Budget, HeuristicSettings, bBudgetIsOverBudget);
+		// We do not change resolution too often to avoid interferences with temporal sub pixel in TAA upsample.
+		if ((bWouldBeWorthChangingRes &&
+			NumberOfFramesSinceScreenPercentageChange >= CVarFrameChangePeriod.GetValueOnRenderThread()) ||
+			bGPUOverbugetPanic)
+		{
+			NumberOfFramesSinceScreenPercentageChange = 0;
+			CurrentFrameResolutionFractions[GDynamicPrimaryResolutionFraction] = NewFrameResolutionFraction;
+		}
+		else
+		{
+			NumberOfFramesSinceScreenPercentageChange++;
+		}
 	}
 
-	// Commit the new resolution fractions.
-	if (bCommitResolutionChange)
-	{
-		NumberOfFramesSinceScreenPercentageChange = 0;
-		CurrentFrameResolutionFractions = NewResolutionFractions;
-	}
-	else
-	{
-		NumberOfFramesSinceScreenPercentageChange++;
-	}
-	
 	// Compute max resolution for each budget by quantizing the new resolution fraction (falls back to the MaxResolution setting if BudgetSetting.UpperBoundQuantization==0)
 	DynamicRenderScaling::TMap<float> NewMaxResolutionFractions;
 	for (TLinkedList<DynamicRenderScaling::FBudget*>::TIterator BudgetIt(DynamicRenderScaling::FBudget::GetGlobalList()); BudgetIt; BudgetIt.Next())
@@ -527,7 +451,7 @@ void FDynamicResolutionHeuristicProxy::RefreshCurentFrameResolutionFraction_Rend
 			float NewMaxResolutionFraction = BudgetSettings.MaxResolutionFraction;
 			if (BudgetSettings.UpperBoundQuantization > 0)
 			{
-				float CurrentResolutionFraction = NewResolutionFractions[Budget];
+				float CurrentResolutionFraction = CurrentFrameResolutionFractions[Budget];
 				float AvailableRange = BudgetSettings.MaxResolutionFraction - BudgetSettings.MinResolutionFraction;
 				float QuantizationStepSize = AvailableRange / float(BudgetSettings.UpperBoundQuantization);
 				NewMaxResolutionFraction = FMath::CeilToFloat(CurrentResolutionFraction / QuantizationStepSize) * QuantizationStepSize;
