Use pass by ref instead of pass by value for TArray to elliminate a copy in UNetDriver::ServerReplicateActors_PrioritizeActors.

#tests PF compile + editor tests

diff --git a/Engine/Source/Runtime/Engine/Classes/Engine/NetDriver.h b/Engine/Source/Runtime/Engine/Classes/Engine/NetDriver.h
--- a/Engine/Source/Runtime/Engine/Classes/Engine/NetDriver.h
+++ b/Engine/Source/Runtime/Engine/Classes/Engine/NetDriver.h
@@ -1792,7 +1792,7 @@ protected:
 	*/
 	int32 ServerReplicateActors_PrepConnections( const float DeltaSeconds );
 	void ServerReplicateActors_BuildConsiderList( TArray<FNetworkObjectInfo*>& OutConsiderList, const float ServerTickTime );
-	int32 ServerReplicateActors_PrioritizeActors( UNetConnection* Connection, const TArray<FNetViewer>& ConnectionViewers, const TArray<FNetworkObjectInfo*> ConsiderList, const bool bCPUSaturated, FActorPriority*& OutPriorityList, FActorPriority**& OutPriorityActors );
+	int32 ServerReplicateActors_PrioritizeActors( UNetConnection* Connection, const TArray<FNetViewer>& ConnectionViewers, const TArray<FNetworkObjectInfo*>& ConsiderList, const bool bCPUSaturated, FActorPriority*& OutPriorityList, FActorPriority**& OutPriorityActors );
 	int32 ServerReplicateActors_ProcessPrioritizedActors( UNetConnection* Connection, const TArray<FNetViewer>& ConnectionViewers, FActorPriority** PriorityActors, const int32 FinalSortedCount, int32& OutUpdated );
 #endif
 
diff --git a/Engine/Source/Runtime/Engine/Private/NetDriver.cpp b/Engine/Source/Runtime/Engine/Private/NetDriver.cpp
--- a/Engine/Source/Runtime/Engine/Private/NetDriver.cpp
+++ b/Engine/Source/Runtime/Engine/Private/NetDriver.cpp
@@ -4365,7 +4365,7 @@ static FORCEINLINE_DEBUGGABLE bool ShouldActorGoDormant( AActor* Actor, const TA
 	return true;
 }
 
-int32 UNetDriver::ServerReplicateActors_PrioritizeActors( UNetConnection* Connection, const TArray<FNetViewer>& ConnectionViewers, const TArray<FNetworkObjectInfo*> ConsiderList, const bool bCPUSaturated, FActorPriority*& OutPriorityList, FActorPriority**& OutPriorityActors )
+int32 UNetDriver::ServerReplicateActors_PrioritizeActors( UNetConnection* Connection, const TArray<FNetViewer>& ConnectionViewers, const TArray<FNetworkObjectInfo*>& ConsiderList, const bool bCPUSaturated, FActorPriority*& OutPriorityList, FActorPriority**& OutPriorityActors )
 {
 	SCOPE_CYCLE_COUNTER( STAT_NetPrioritizeActorsTime );
 
