Added Session ID's and Client ID's (Connection ID's) to packet protocol.

diff --git a/Engine/Source/Runtime/Engine/Classes/Engine/Engine.h b/Engine/Source/Runtime/Engine/Classes/Engine/Engine.h
--- a/Engine/Source/Runtime/Engine/Classes/Engine/Engine.h
+++ b/Engine/Source/Runtime/Engine/Classes/Engine/Engine.h
@@ -2976,6 +2976,20 @@ private:
 	FScreenSaverInhibitor*  ScreenSaverInhibitorRunnable;
 
 
+	/** Increments every time a non-seamless travel happens on a server, to generate net session id's. Written to config to preserve id upon crash. */
+	UPROPERTY(Config)
+	uint32 GlobalNetTravelCount = 0;
+
+public:
+	void IncrementGlobalNetTravelCount()
+	{
+		GlobalNetTravelCount++;
+	}
+
+	uint32 GetGlobalNetTravelCount() const
+	{
+		return GlobalNetTravelCount;
+	}
 public:
 
 	/** A list of named UNetDriver definitions */
diff --git a/Engine/Source/Runtime/Engine/Classes/Engine/NetDriver.h b/Engine/Source/Runtime/Engine/Classes/Engine/NetDriver.h
--- a/Engine/Source/Runtime/Engine/Classes/Engine/NetDriver.h
+++ b/Engine/Source/Runtime/Engine/Classes/Engine/NetDriver.h
@@ -903,11 +903,15 @@ public:
 	/** Set the NetDriver's NetDriverDefintion. */
 	void SetNetDriverDefinition(FName NewNetDriverDefinition);
 
+	UE_DEPRECATED(5.1, "Use GetNetDriverDefinition instead.")
+	FName GetNetDriverDefintion() const { return NetDriverDefinition; }
+
+	/** Get the NetDriver's NetDriverDefintion. */
+	FName GetNetDriverDefinition() const { return NetDriverDefinition; }
+
 	/** Callback after the engine created the NetDriver and set our name for the first time */
 	void PostCreation(bool bInitializeWithIris);
 
-	/** Get the NetDriver's NetDriverDefintion. */
-	FName GetNetDriverDefintion() const { return NetDriverDefinition; }
 
 	void InitPacketSimulationSettings();
 
@@ -1906,6 +1910,11 @@ public:
 		TotalOutOfOrderPacketsDuplicate += Count;
 	}
 
+	uint32 GetCachedGlobalNetTravelCount() const
+	{
+		return CachedGlobalNetTravelCount;
+	}
+
 	bool DidHitchLastFrame() const;
 
 	static bool IsDormInitialStartupActor(AActor* Actor);
@@ -2001,4 +2010,8 @@ private:
 
 	/** Stat tracking for the total number of out of order packets that were duplicates */
 	int32 TotalOutOfOrderPacketsDuplicate = 0;
+
+
+	/** Cached value for UEngine.GlobalNetTravelCount, at the time of NetDriver initialization */
+	uint32 CachedGlobalNetTravelCount = 0;
 };
diff --git a/Engine/Source/Runtime/Engine/Private/GameModeBase.cpp b/Engine/Source/Runtime/Engine/Private/GameModeBase.cpp
--- a/Engine/Source/Runtime/Engine/Private/GameModeBase.cpp
+++ b/Engine/Source/Runtime/Engine/Private/GameModeBase.cpp
@@ -491,10 +491,16 @@ void AGameModeBase::ProcessServerTravel(const FString& URL, bool bAbsolute)
 		World->SeamlessTravel(World->NextURL, bAbsolute);
 		World->NextURL = TEXT("");
 	}
-	// Switch immediately if not networking.
-	else if (NetMode != NM_DedicatedServer && NetMode != NM_ListenServer)
+	else
 	{
-		World->NextSwitchCountdown = 0.0f;
+		// Switch immediately if not networking.
+		if (NetMode != NM_DedicatedServer && NetMode != NM_ListenServer)
+		{
+			World->NextSwitchCountdown = 0.0f;
+		}
+
+		GEngine->IncrementGlobalNetTravelCount();
+		GEngine->SaveConfig();
 	}
 #endif // WITH_SERVER_CODE
 }
diff --git a/Engine/Source/Runtime/Engine/Private/NetConnection.cpp b/Engine/Source/Runtime/Engine/Private/NetConnection.cpp
--- a/Engine/Source/Runtime/Engine/Private/NetConnection.cpp
+++ b/Engine/Source/Runtime/Engine/Private/NetConnection.cpp
@@ -432,7 +432,7 @@ void UNetConnection::InitBase(UNetDriver* InDriver,class FSocket* InSocket, cons
 	InitHandler();
 
 
-	FaultRecovery.InitDefaults((Driver != nullptr ? Driver->NetDriverName.ToString() : TEXT("")), this);
+	FaultRecovery.InitDefaults((Driver != nullptr ? Driver->GetNetDriverDefinition().ToString() : TEXT("")), this);
 
 	if (CVarLogUnhandledFaults.GetValueOnAnyThread() != 0)
 	{
@@ -488,7 +488,7 @@ void UNetConnection::InitBase(UNetDriver* InDriver,class FSocket* InSocket, cons
 
 	if (bIsServer && !bIsReplay)
 	{
-		RPCDoS.Init(Driver->NetDriverName, AnalyticsAggregator,
+		RPCDoS.Init(Driver->GetNetDriverDefinition(), AnalyticsAggregator,
 			[WorldPtr = TWeakObjectPtr<UWorld>(InDriver->GetWorld())]() -> UWorld*
 			{
 				return WorldPtr.IsValid() ? WorldPtr.Get() : nullptr;
@@ -602,7 +602,7 @@ void UNetConnection::InitHandler()
 											MoveTemp(NotifyAddHandler));
 
 			Handler->NotifyAnalyticsProvider(Driver->AnalyticsProvider, Driver->AnalyticsAggregator);
-			Handler->Initialize(Mode, MaxPacket * 8, false, nullptr, nullptr, Driver->NetDriverName);
+			Handler->Initialize(Mode, MaxPacket * 8, false, nullptr, nullptr, Driver->GetNetDriverDefinition());
 
 
 			// Add handling for the stateless connect handshake, for connectionless packets, as the outermost layer
diff --git a/Engine/Source/Runtime/Engine/Private/NetDriver.cpp b/Engine/Source/Runtime/Engine/Private/NetDriver.cpp
--- a/Engine/Source/Runtime/Engine/Private/NetDriver.cpp
+++ b/Engine/Source/Runtime/Engine/Private/NetDriver.cpp
@@ -1400,6 +1400,8 @@ bool UNetDriver::InitBase(bool bInitAsClient, FNetworkNotify* InNotify, const FU
 		InitDestroyedStartupActors();
 	}
 
+	CachedGlobalNetTravelCount = GEngine->GetGlobalNetTravelCount();
+
 	return bSuccess;
 }
 
diff --git a/Engine/Source/Runtime/Engine/Private/PacketHandlers/StatelessConnectHandlerComponent.cpp b/Engine/Source/Runtime/Engine/Private/PacketHandlers/StatelessConnectHandlerComponent.cpp
--- a/Engine/Source/Runtime/Engine/Private/PacketHandlers/StatelessConnectHandlerComponent.cpp
+++ b/Engine/Source/Runtime/Engine/Private/PacketHandlers/StatelessConnectHandlerComponent.cpp
@@ -8,6 +8,8 @@
 #include "Engine/NetConnection.h"
 #include "Net/Core/Misc/PacketAudit.h"
 #include "Misc/NetworkVersion.h"
+#include "Misc/StringBuilder.h"
+#include "Misc/ConfigCacheIni.h"
 
 
 DEFINE_LOG_CATEGORY(LogHandshake);
@@ -38,19 +40,31 @@ DEFINE_LOG_CATEGORY(LogHandshake);
  * and the server responding with a unique 'Cookie' value, which the client has to respond with.
  *
  * Client - Initial Connect:
- * [?:MagicHeader][HandshakeBit][RestartHandshakeBit][8:MinVersion][8:CurVersion][8:HandshakePacketType][8:SentPacketCount][32:NetworkVersion][SecretIdBit][28:PacketSizeFiller][AlignPad][?:RandomData]
+ *
+ * [?:MagicHeader][2:SessionID][3:ClientID][HandshakeBit][RestartHandshakeBit]
+ * [8:MinVersion][8:CurVersion][8:HandshakePacketType][8:SentPacketCount][32:NetworkVersion]
+ * [SecretIdBit] [28:PacketSizeFiller][AlignPad][?:RandomData]
  *													--->
  *															Server - Stateless Handshake Challenge:
- *															[?:MagicHeader][HandshakeBit][RestartHandshakeBit][8:MinVersion][8:CurVersion][8:HandshakePacketType][8:SentPacketCount][32:NetworkVersion][SecretIdBit][8:Timestamp][20:Cookie][AlignPad][?:RandomData]
+ *
+ *															[?:MagicHeader][2:SessionID][3:ClientID][HandshakeBit][RestartHandshakeBit]
+ *															[8:MinVersion][8:CurVersion][8:HandshakePacketType][8:SentPacketCount][32:NetworkVersion]
+ *															[SecretIdBit][8:Timestamp][20:Cookie][AlignPad][?:RandomData]
  *													<---
  * Client - Stateless Challenge Response:
- * [?:MagicHeader][HandshakeBit][RestartHandshakeBit][8:MinVersion][8:CurVersion][8:HandshakePacketType][8:SentPacketCount][32:NetworkVersion][SecretIdBit][8:Timestamp][20:Cookie][AlignPad][?:RandomData]
+ *
+ * [?:MagicHeader][2:SessionID][3:ClientID][HandshakeBit][RestartHandshakeBit]
+ * [8:MinVersion][8:CurVersion][8:HandshakePacketType][8:SentPacketCount][32:NetworkVersion]
+ * [SecretIdBit][8:Timestamp][20:Cookie][AlignPad][?:RandomData]
  *													--->
  *															Server:
  *															Ignore, or create UNetConnection.
  *
- *															Server - Stateless Handshake Ack
- *															[?:MagicHeader][HandshakeBit][RestartHandshakeBit][8:MinVersion][8:CurVersion][8:HandshakePacketType][8:SentPacketCount][32:NetworkVersion][SecretIdBit][8:Timestamp][20:Cookie][AlignPad][?:RandomData]
+ *															Server - Stateless Handshake Ack:
+ *
+ *															[?:MagicHeader][2:SessionID][3:ClientID][HandshakeBit][RestartHandshakeBit]
+ *															[8:MinVersion][8:CurVersion][8:HandshakePacketType][8:SentPacketCount][32:NetworkVersion]
+ *															[SecretIdBit][8:Timestamp][20:Cookie][AlignPad][?:RandomData]
  *													<---
  * Client:
  * Handshake Complete.
@@ -62,14 +76,19 @@ DEFINE_LOG_CATEGORY(LogHandshake);
  * so the protocol has been crafted so the server sends only a minimal (1 byte) response, to minimize DRDoS reflection amplification.
  *
  *															Server - Restart Handshake Request:
- *															[?:MagicHeader][HandshakeBit][RestartHandshakeBit][8:MinVersion][8:CurVersion][8:HandshakePacketType][8:SentPacketCount][32:NetworkVersion][AlignPad][?:RandomData]
+ *
+ *															[?:MagicHeader][2:SessionID][3:ClientID][HandshakeBit][RestartHandshakeBit]
+ *															[8:HandshakePacketType][8:SentPacketCount][32:NetworkVersion][AlignPad][?:RandomData]
  *													<--
  * Client -  Initial Connect (as above)
  *													-->
  *															Server -  Stateless Handshake Challenge (as above)
  *													<--
- * Client - Stateless Challenge Response + Original Cookie
- * [?:MagicHeader][HandshakeBit][RestartHandshakeBit][8:MinVersion][8:CurVersion][8:HandshakePacketType][8:SentPacketCount][SecretIdBit][8:Timestamp][20:Cookie][20:OriginalCookie][32:NetworkVersion][AlignPad][?:RandomData]
+ * Client - Stateless Challenge Response + Original Cookie:
+ *
+ * [?:MagicHeader][2:SessionID][3:ClientID][HandshakeBit][RestartHandshakeBit]
+ * [8:MinVersion][8:CurVersion][8:HandshakePacketType][8:SentPacketCount][32:NetworkVersion]
+ * [SecretIdBit][8:Timestamp][20:Cookie][20:OriginalCookie][AlignPad][?:RandomData]
  *													-->
  *															Server:
  *															Ignore, or restore UNetConnection.
@@ -82,6 +101,8 @@ DEFINE_LOG_CATEGORY(LogHandshake);
  *
  *
  *	- MagicHeader:			An optional static/predefined header, between 0-32 bits in size. Serves no purpose for the handshake code.
+ *	- SessionID:			Session id incremented serverside every non-seamless server travel, to prevent non-ephemeral old/new-session crosstalk.
+ *	- ClientID:				Connection id incremented clientside every connection per-NetDriver, to prevent non-ephemeral old/new-connection crosstalk.
  *	- HandshakeBit:			Bit signifying whether a packet is a handshake packet. Applied to all game packets.
  *	- SecretIdBit:			For handshake packets, specifies which HandshakeSecret array was used to generate Cookie.
  *	- RestartHandshakeBit:  Sent by the server when it detects normal game traffic from an unknown IP/port combination.
@@ -105,8 +126,9 @@ DEFINE_LOG_CATEGORY(LogHandshake);
  *
  * Game Protocol Changes:
  *
- * Every game (as opposed to handshake) packet starts with an extra bit, represented by [HandshakeBit], and game packets set this to 0.
- * This is the only change to game packets. When HandshakeBit is set to 1, the separate protocol above is used for handshake packets.
+ * Every packet (game and handshake) starts with the MagicHeader bits (if set), then 2 SessionID bits and 3 ClientID bits,
+ * and finally HandshakeBit (which is 0 for game packets, going through normal PacketHandler/NetConnection protocol processing,
+ * and 1 for handshake packets, going through the separate protocol documented above).
  *
  *
  *
@@ -144,6 +166,33 @@ DEFINE_LOG_CATEGORY(LogHandshake);
  *
  * The client carries on with the handshake as normal, but when completing the handshake, the client also sends the cookie it previously connected with.
  * The server looks up the NetConnection associated with that cookie, and then updates the address for the connection.
+ *
+ *
+ *
+ * SessionID/ClientID and non-ephemeral sockets
+ *
+ * The packet protocol has a reliance on IP packet ephemeral ports (the random client-specified source port) to differentiate client connections,
+ * but not all socket subsystems provide something which serves this role - some socket subsystems only provide a static address without a port,
+ * where the address remains indistinguishable between old/new client connections, e.g. when performing a non-seamless travel between levels.
+ *
+ * The SessionID value solves this for the case of non-seamless travel, by specifying an incrementing server-authoritative ID for the game session
+ * (which changes upon non-seamless travel).
+ *
+ * The ClientID solves this for the case of clients reconnecting to a server they are currently connected to,
+ * or which their old connection is pending timeout from (e.g. after a crash or other fault requiring a reconnect),
+ * by specifying an incrementing client-authoritative ID for the connection (per-NetDriver - so e.g. Game and Beacon drivers increment separately).
+ *
+ *
+ * This is not a complete solution, however - multiple clients from the same address will not work, presently.
+ * ClientID has enough bits to implement this in the future, while keeping net compatibility, but is non-trivial and not guaranteed to be added.
+ *
+ * If this is to be added though, it will require adjustments to the serverside NetDriver receive code, to set the ClientID as the address port,
+ * and will require adjustments to the clientside setting of ClientID, to perform inter-process communication when picking the ClientID,
+ * so that the value is unique for every NetConnection across every game process, connecting to the same server address + port.
+ *
+ *
+ * Also, if the server and client are using the same *Engine.ini file, the last process to close will clobber the
+ * GlobalNetTravelCount/CachedClientID increments from the other process, which are used for SessionID/ClientID.
  */
 
 
@@ -201,7 +250,7 @@ namespace UE::Net
 #ifdef HANDSHAKE_MIN_VERSION_OVERRIDE
 	static int32 MinSupportedHandshakeVersion = HANDSHAKE_MIN_VERSION_OVERRIDE;
 #else
-	static int32 MinSupportedHandshakeVersion = static_cast<uint8>(EHandshakeVersion::Original);
+	static int32 MinSupportedHandshakeVersion = static_cast<uint8>(EHandshakeVersion::SessionClientId);
 #endif
 
 	/** The current compile-time handshake version */
@@ -234,6 +283,22 @@ namespace UE::Net
 		TEXT("Whether or not the stateless handshake should enforce the Network CL version, instead of the higher level netcode. ")
 		TEXT("(does NOT result in a graceful rejection, packet is dropped without response - higher level netcode would have sent upgrade message)"));
 
+	static int32 GVerifyNetSessionID = 1;
+
+	FAutoConsoleVariableRef CVarNetVerifyNetSessionID(
+		TEXT("net.VerifyNetSessionID"),
+		GVerifyNetSessionID,
+		TEXT("Whether or not verification of the packet SessionID value is performed."));
+
+	static int32 GVerifyNetClientID = 1;
+
+	FAutoConsoleVariableRef CVarNetVerifyNetClientID(
+		TEXT("net.VerifyNetClientID"),
+		GVerifyNetSessionID,
+		TEXT("Whether or not verification of the packet ClientID value is performed."));
+
+
+
 
 	/** The base amount of random data to add to handshake packets */
 	static constexpr int32 BaseRandomDataLengthBytes		= 16;
@@ -355,7 +420,8 @@ void StatelessConnectHandlerComponent::SendInitialPacket(EHandshakeVersion Hands
 
 		if (ServerConn != nullptr)
 		{
-			FBitWriter InitialPacket(GetAdjustedSizeBits(HANDSHAKE_PACKET_SIZE_BITS) + (BaseRandomDataLengthBytes * 8) + 1 /* Termination bit */);
+			const int32 AdjustedSize = GetAdjustedSizeBits(HANDSHAKE_PACKET_SIZE_BITS, HandshakeVersion);
+			FBitWriter InitialPacket(AdjustedSize + (BaseRandomDataLengthBytes * 8) + 1 /* Termination bit */);
 			uint8 bHandshakePacket = 1;
 
 			if (MagicHeader.Num() > 0)
@@ -363,6 +429,12 @@ void StatelessConnectHandlerComponent::SendInitialPacket(EHandshakeVersion Hands
 				InitialPacket.SerializeBits(MagicHeader.GetData(), MagicHeader.Num());
 			}
 
+			if (HandshakeVersion >= EHandshakeVersion::SessionClientId)
+			{
+				InitialPacket.SerializeBits(&CachedGlobalNetTravelCount, SessionIDSizeBits);
+				InitialPacket.SerializeBits(&CachedClientID, ClientIDSizeBits);
+			}
+
 			InitialPacket.WriteBit(bHandshakePacket);
 
 
@@ -436,13 +508,14 @@ void StatelessConnectHandlerComponent::SendInitialPacket(EHandshakeVersion Hands
 }
 
 void StatelessConnectHandlerComponent::SendConnectChallenge(TSharedPtr<const FInternetAddr> ClientAddress, EHandshakeVersion HandshakeVersion,
-															uint8 ClientSentHandshakePacketCount)
+															uint8 ClientSentHandshakePacketCount, uint32 InClientID)
 {
 	using namespace UE::Net;
 
 	if (Driver != nullptr)
 	{
-		FBitWriter ChallengePacket(GetAdjustedSizeBits(HANDSHAKE_PACKET_SIZE_BITS) + (BaseRandomDataLengthBytes * 8) + 1 /* Termination bit */);
+		const int32 AdjustedSize = GetAdjustedSizeBits(HANDSHAKE_PACKET_SIZE_BITS, HandshakeVersion);
+		FBitWriter ChallengePacket(AdjustedSize + (BaseRandomDataLengthBytes * 8) + 1 /* Termination bit */);
 		uint8 bHandshakePacket = 1;
 		uint8 bRestartHandshake = 0; // Ignored clientside
 		double Timestamp = Driver->GetElapsedTime();
@@ -455,6 +528,12 @@ void StatelessConnectHandlerComponent::SendConnectChallenge(TSharedPtr<const FIn
 			ChallengePacket.SerializeBits(MagicHeader.GetData(), MagicHeader.Num());
 		}
 
+		if (HandshakeVersion >= EHandshakeVersion::SessionClientId)
+		{
+			ChallengePacket.SerializeBits(&CachedGlobalNetTravelCount, SessionIDSizeBits);
+			ChallengePacket.SerializeBits(&InClientID, ClientIDSizeBits);
+		}
+
 		ChallengePacket.WriteBit(bHandshakePacket);
 		ChallengePacket.WriteBit(bRestartHandshake);
 
@@ -542,10 +621,9 @@ void StatelessConnectHandlerComponent::SendChallengeResponse(EHandshakeVersion H
 
 	if (ServerConn != nullptr)
 	{
-		int32 RestartHandshakeResponseSize = RESTART_RESPONSE_SIZE_BITS;
-
-		const int32 BaseSize = GetAdjustedSizeBits(bRestartedHandshake ? RestartHandshakeResponseSize : HANDSHAKE_PACKET_SIZE_BITS);
-		FBitWriter ResponsePacket(BaseSize + (BaseRandomDataLengthBytes * 8) + 1 /* Termination bit */);
+		const int32 AdjustedSize = GetAdjustedSizeBits((bRestartedHandshake ? RESTART_RESPONSE_SIZE_BITS : HANDSHAKE_PACKET_SIZE_BITS),
+														HandshakeVersion);
+		FBitWriter ResponsePacket(AdjustedSize + (BaseRandomDataLengthBytes * 8) + 1 /* Termination bit */);
 		uint8 bHandshakePacket = 1;
 		uint8 bRestartHandshake = (bRestartedHandshake ? 1 : 0);
 
@@ -554,6 +632,12 @@ void StatelessConnectHandlerComponent::SendChallengeResponse(EHandshakeVersion H
 			ResponsePacket.SerializeBits(MagicHeader.GetData(), MagicHeader.Num());
 		}
 
+		if (HandshakeVersion >= EHandshakeVersion::SessionClientId)
+		{
+			ResponsePacket.SerializeBits(&CachedGlobalNetTravelCount, SessionIDSizeBits);
+			ResponsePacket.SerializeBits(&CachedClientID, ClientIDSizeBits);
+		}
+
 		ResponsePacket.WriteBit(bHandshakePacket);
 		ResponsePacket.WriteBit(bRestartHandshake);
 
@@ -637,13 +721,14 @@ void StatelessConnectHandlerComponent::SendChallengeResponse(EHandshakeVersion H
 }
 
 void StatelessConnectHandlerComponent::SendChallengeAck(TSharedPtr<const FInternetAddr> ClientAddress, EHandshakeVersion HandshakeVersion,
-														uint8 ClientSentHandshakePacketCount, uint8 InCookie[COOKIE_BYTE_SIZE])
+														uint8 ClientSentHandshakePacketCount, uint32 InClientID, uint8 InCookie[COOKIE_BYTE_SIZE])
 {
 	using namespace UE::Net;
 
 	if (Driver != nullptr)
 	{
-		FBitWriter AckPacket(GetAdjustedSizeBits(HANDSHAKE_PACKET_SIZE_BITS) + (BaseRandomDataLengthBytes * 8) + 1 /* Termination bit */);
+		const int32 AdjustedSize = GetAdjustedSizeBits(HANDSHAKE_PACKET_SIZE_BITS, HandshakeVersion);
+		FBitWriter AckPacket(AdjustedSize + (BaseRandomDataLengthBytes * 8) + 1 /* Termination bit */);
 		uint8 bHandshakePacket = 1;
 		uint8 bRestartHandshake = 0; // Ignored clientside
 		double Timestamp  = -1.0;
@@ -653,6 +738,12 @@ void StatelessConnectHandlerComponent::SendChallengeAck(TSharedPtr<const FIntern
 			AckPacket.SerializeBits(MagicHeader.GetData(), MagicHeader.Num());
 		}
 
+		if (HandshakeVersion >= EHandshakeVersion::SessionClientId)
+		{
+			AckPacket.SerializeBits(&CachedGlobalNetTravelCount, SessionIDSizeBits);
+			AckPacket.SerializeBits(&InClientID, ClientIDSizeBits);
+		}
+
 		AckPacket.WriteBit(bHandshakePacket);
 		AckPacket.WriteBit(bRestartHandshake);
 
@@ -730,13 +821,14 @@ void StatelessConnectHandlerComponent::SendChallengeAck(TSharedPtr<const FIntern
 }
 
 void StatelessConnectHandlerComponent::SendRestartHandshakeRequest(const TSharedPtr<const FInternetAddr> ClientAddress,
-																	EHandshakeVersion HandshakeVersion)
+																	EHandshakeVersion HandshakeVersion, uint32 InClientID)
 {
 	using namespace UE::Net;
 
 	if (Driver != nullptr)
 	{
-		FBitWriter RestartPacket(GetAdjustedSizeBits(RESTART_HANDSHAKE_PACKET_SIZE_BITS) + (BaseRandomDataLengthBytes * 8) + 1 /* Termination bit */);
+		const int32 AdjustedSize = GetAdjustedSizeBits(RESTART_HANDSHAKE_PACKET_SIZE_BITS, HandshakeVersion);
+		FBitWriter RestartPacket(AdjustedSize + (BaseRandomDataLengthBytes * 8) + 1 /* Termination bit */);
 		uint8 bHandshakePacket = 1;
 		uint8 bRestartHandshake = 1;
 
@@ -745,6 +837,12 @@ void StatelessConnectHandlerComponent::SendRestartHandshakeRequest(const TShared
 			RestartPacket.SerializeBits(MagicHeader.GetData(), MagicHeader.Num());
 		}
 
+		if (HandshakeVersion >= EHandshakeVersion::SessionClientId)
+		{
+			RestartPacket.SerializeBits(&CachedGlobalNetTravelCount, SessionIDSizeBits);
+			RestartPacket.SerializeBits(&InClientID, ClientIDSizeBits);
+		}
+
 		RestartPacket.WriteBit(bHandshakePacket);
 		RestartPacket.WriteBit(bRestartHandshake);
 
@@ -822,7 +920,7 @@ void StatelessConnectHandlerComponent::CapHandshakePacket(FBitWriter& HandshakeP
 {
 	using namespace UE::Net;
 
-	uint32 NumBits = HandshakePacket.GetNumBits() - GetAdjustedSizeBits(0);
+	uint32 NumBits = HandshakePacket.GetNumBits() - GetAdjustedSizeBits(0, HandshakeVersion);
 
 #if !UE_BUILD_SHIPPING
 	if (HandshakeVersion == EHandshakeVersion::Original)
@@ -848,11 +946,14 @@ void StatelessConnectHandlerComponent::CapHandshakePacket(FBitWriter& HandshakeP
 		int32 RandomDataLengthBytes = BaseRandomDataLengthBytes - FMath::RandRange(0, RandomDataLengthVarianceBytes);
 
 		// In versions that must stay compatible with the original protocol, make sure there isn't a size collision with the original restart response
-		if (NumBits + (RandomDataLengthBytes * 8) == OriginalRestartResponseSizeBits)
+		if (HandshakeVersion < EHandshakeVersion::SessionClientId)
 		{
-			RandomDataLengthBytes = FMath::Max(0, RandomDataLengthBytes - 1);
+			if (NumBits + (RandomDataLengthBytes * 8) == OriginalRestartResponseSizeBits)
+			{
+				RandomDataLengthBytes = FMath::Max(0, RandomDataLengthBytes - 1);
+			}
 		}
-		
+
 		for (int32 RandIdx=0; RandIdx<RandomDataLengthBytes; RandIdx++)
 		{
 			uint8 RandVal = FMath::Rand() % 255;
@@ -886,15 +987,84 @@ void StatelessConnectHandlerComponent::SetDriver(UNetDriver* InDriver)
 				InitFromConnectionless(StatelessComponent);
 			}
 		}
+
+		CachedGlobalNetTravelCount = Driver->GetCachedGlobalNetTravelCount() & ((1 << SessionIDSizeBits) - 1);
+	}
+	else //if (Handler->Mode == Handler::Mode::Client)
+	{
+		// Use basic GConfig until NetDriver per-NetDriverDefinitionName per-object-config refactor
+		TStringBuilder<256> ConfigSection;
+		const FString NetDriverDef = Driver->GetNetDriverDefinition().ToString();
+
+		ConfigSection.Append(ToCStr(NetDriverDef));
+		ConfigSection.Append(TEXT(" StatelessConnectHandlerComponent"));
+
+		const TCHAR* ConfigKey = TEXT("CachedClientID");
+		int32 ConfigCachedClientID = 0;
+
+		GConfig->GetInt(ConfigSection.ToString(), ConfigKey, ConfigCachedClientID, GEngineIni);
+
+		CachedClientID = FMath::Max(ConfigCachedClientID, 0) + 1;
+
+		if (CachedClientID > static_cast<uint32>(std::numeric_limits<int32>::max()))
+		{
+			CachedClientID = 0;
+		}
+
+		GConfig->SetInt(ConfigSection.ToString(), ConfigKey, CachedClientID, GEngineIni);
+		GConfig->Flush(false);
+
+		CachedClientID = CachedClientID & ((1 << ClientIDSizeBits) - 1);
+
+		UE_LOG(LogHandshake, Log, TEXT("Stateless Handshake: NetDriverDefinition '%s' CachedClientID: %i"), ToCStr(NetDriverDef),
+				CachedClientID);
 	}
 }
 
 void StatelessConnectHandlerComponent::Initialize()
 {
+	using namespace UE::Net;
+
 	// On the server, initializes immediately. Clientside doesn't initialize until handshake completes.
 	if (Handler->Mode == Handler::Mode::Server)
 	{
 		Initialized();
+
+#if !UE_BUILD_SHIPPING
+		const EHandshakeVersion CurHandshakeVersion = static_cast<EHandshakeVersion>(CurrentHandshakeVersion);
+		const EHandshakeVersion MinHandshakeVersion = static_cast<EHandshakeVersion>(MinSupportedHandshakeVersion);
+
+		if (CurHandshakeVersion > EHandshakeVersion::Latest)
+		{
+			UE_LOG(LogHandshake, Error, TEXT("net.CurrentHandshakeVersion value '%i' is invalid. Maximum value: %i"), CurrentHandshakeVersion,
+					static_cast<EHandshakeVersion>(EHandshakeVersion::Latest));
+		}
+		else
+		{
+			EHandshakeVersion CompatibleRangeStart = EHandshakeVersion::Original;
+			EHandshakeVersion CompatibleRangeEnd = EHandshakeVersion::Latest;
+
+			for (EHandshakeVersion CurBreakVersion : HandshakeCompatibilityBreaks)
+			{
+				if (CurBreakVersion <= CurHandshakeVersion)
+				{
+					CompatibleRangeStart = CurBreakVersion;
+				}
+				else if (CurBreakVersion > CurHandshakeVersion)
+				{
+					CompatibleRangeEnd = static_cast<EHandshakeVersion>(static_cast<uint8>(CurBreakVersion) - 1);
+					break;
+				}
+			}
+
+			if (MinHandshakeVersion < CompatibleRangeStart || MinHandshakeVersion > CompatibleRangeEnd)
+			{
+				UE_LOG(LogHandshake, Error, TEXT("net.MinHandshakeVersion value '%i' is invalid relative to net.CurrentHandshakeVersion value '%i'. ")
+						TEXT("Minimum value: %i, Maximum value: %i"), MinSupportedHandshakeVersion, CurrentHandshakeVersion,
+						static_cast<uint8>(CompatibleRangeStart), static_cast<uint8>(CompatibleRangeEnd));
+			}
+		}
+#endif
 	}
 }
 
@@ -903,6 +1073,7 @@ void StatelessConnectHandlerComponent::InitFromConnectionless(StatelessConnectHa
 	// Store the cookie/address used for the handshake, to enable server ack-retries
 	LastChallengeSuccessAddress = InConnectionlessHandler->LastChallengeSuccessAddress;
 	LastRemoteHandshakeVersion = InConnectionlessHandler->LastRemoteHandshakeVersion;
+	CachedClientID = InConnectionlessHandler->CachedClientID;
 
 	FMemory::Memcpy(AuthorisedCookie, InConnectionlessHandler->AuthorisedCookie, UE_ARRAY_COUNT(AuthorisedCookie));
 
@@ -920,6 +1091,20 @@ void StatelessConnectHandlerComponent::Incoming(FBitReader& Packet)
 		Packet.SerializeBits(&ReadMagic, MagicHeader.Num());
 	}
 
+	bool bHasValidSessionID = true;
+	bool bHasValidClientID = true;
+	uint8 SessionID = 0;
+	uint8 ClientID = 0;
+
+	if (LastRemoteHandshakeVersion >= EHandshakeVersion::SessionClientId)
+	{
+		Packet.SerializeBits(&SessionID, SessionIDSizeBits);
+		Packet.SerializeBits(&ClientID, ClientIDSizeBits);
+
+		bHasValidSessionID = GVerifyNetSessionID == 0 || (SessionID == CachedGlobalNetTravelCount && !Packet.IsError());
+		bHasValidClientID = GVerifyNetClientID == 0 || (ClientID == CachedClientID && !Packet.IsError());
+	}
+
 	bool bHandshakePacket = !!Packet.ReadBit() && !Packet.IsError();
 
 	if (bHandshakePacket)
@@ -930,7 +1115,10 @@ void StatelessConnectHandlerComponent::Incoming(FBitReader& Packet)
 
 		if (bHandshakePacket)
 		{
-			if (Handler->Mode == Handler::Mode::Client)
+			const bool bIsChallengePacket = HandshakeData.HandshakePacketType == EHandshakePacketType::Challenge && HandshakeData.Timestamp > 0.0;
+			const bool bIsInitialChallengePacket = bIsChallengePacket && State != Handler::Component::State::Initialized;
+
+			if (Handler->Mode == Handler::Mode::Client && bHasValidClientID && (bHasValidSessionID || bIsInitialChallengePacket))
 			{
 				if (State == Handler::Component::State::UnInitialized || State == Handler::Component::State::InitializedOnLocal)
 				{
@@ -940,9 +1128,15 @@ void StatelessConnectHandlerComponent::Incoming(FBitReader& Packet)
 						UE_LOG(LogHandshake, Log, TEXT("Ignoring restart handshake request, while already restarted."));
 #endif
 					}
-					// Receiving challenge, verify the timestamp is > 0.0f
-					else if (HandshakeData.HandshakePacketType == EHandshakePacketType::Challenge && HandshakeData.Timestamp > 0.0)
+					// Receiving challenge
+					else if (bIsChallengePacket)
 					{
+#if !UE_BUILD_SHIPPING
+						UE_LOG(LogHandshake, Log, TEXT("Cached server SessionID: %i"), SessionID);
+#endif
+
+						CachedGlobalNetTravelCount = SessionID;
+
 						LastChallengeTimestamp = (Driver != nullptr ? Driver->GetElapsedTime() : 0.0);
 
 						SendChallengeResponse(HandshakeData.RemoteCurVersion, HandshakeData.SecretId, HandshakeData.Timestamp, HandshakeData.Cookie);
@@ -1082,7 +1276,7 @@ void StatelessConnectHandlerComponent::Incoming(FBitReader& Packet)
 					// Ignore, could be a dupe/out-of-order challenge packet
 				}
 			}
-			else if (Handler->Mode == Handler::Mode::Server)
+			else if (Handler->Mode == Handler::Mode::Server && bHasValidSessionID && bHasValidClientID)
 			{
 				if (LastChallengeSuccessAddress.IsValid())
 				{
@@ -1093,9 +1287,20 @@ void StatelessConnectHandlerComponent::Incoming(FBitReader& Packet)
 							*LastChallengeSuccessAddress->ToString(true), *FString::FromBlob(AuthorisedCookie, COOKIE_BYTE_SIZE));
 #endif
 
-					SendChallengeAck(LastChallengeSuccessAddress, LastRemoteHandshakeVersion, 0, AuthorisedCookie);
+					SendChallengeAck(LastChallengeSuccessAddress, LastRemoteHandshakeVersion, 0, CachedClientID, AuthorisedCookie);
 				}
 			}
+			else if (!bHasValidSessionID || !bHasValidClientID)
+			{
+#if !UE_BUILD_SHIPPING
+				UE_CLOG(TrackSessionConnIDLogs(), LogHandshake, Log,
+						TEXT("Incoming: Rejecting handshake packet with invalid session id (%i vs %i) or connection id (%i vs %i)."),
+						SessionID, CachedGlobalNetTravelCount, ClientID, CachedClientID);
+#endif
+
+				// Ignore, don't trigger disconnect
+				Packet.SetAtEnd();
+			}
 		}
 		else
 		{
@@ -1109,21 +1314,35 @@ void StatelessConnectHandlerComponent::Incoming(FBitReader& Packet)
 #if !UE_BUILD_SHIPPING
 	else if (Packet.IsError())
 	{
-		UE_LOG(LogHandshake, Log, TEXT("Incoming: Error reading handshake bit from packet."));
+		UE_LOG(LogHandshake, Log, TEXT("Incoming: Error reading session id, connection id and handshake bit from packet."));
 	}
 #endif
-	// Servers should wipe LastChallengeSuccessAddress shortly after the first non-handshake packet is received by the client,
-	// in order to disable challenge ack resending
-	else if (LastInitTimestamp != 0.0 && LastChallengeSuccessAddress.IsValid() && Handler->Mode == Handler::Mode::Server)
+	else if (!bHasValidSessionID || !bHasValidClientID)
 	{
-		// Restart handshakes require extra time before disabling challenge ack resends, as NetConnection packets will already be in flight
-		const double RestartHandshakeAckResendWindow = 10.0;
-		double CurTime = Driver != nullptr ? Driver->GetElapsedTime() : 0.0;
+#if !UE_BUILD_SHIPPING
+		UE_CLOG(TrackSessionConnIDLogs(), LogHandshake, Log,
+				TEXT("Incoming: Rejecting game packet with invalid session id (%i vs %i) or connection id (%i vs %i)."),
+				SessionID, CachedGlobalNetTravelCount, ClientID, CachedClientID);
+#endif
 
-		if (CurTime - LastInitTimestamp >= RestartHandshakeAckResendWindow)
+		// Ignore, don't trigger disconnect
+		Packet.SetAtEnd();
+	}
+	else
+	{
+		// Servers should wipe LastChallengeSuccessAddress shortly after the first non-handshake packet is received by the client,
+		// in order to disable challenge ack resending
+		if (LastInitTimestamp != 0.0 && LastChallengeSuccessAddress.IsValid() && Handler->Mode == Handler::Mode::Server)
 		{
-			LastChallengeSuccessAddress.Reset();
-			LastInitTimestamp = 0.0;
+			// Restart handshakes require extra time before disabling challenge ack resends, as NetConnection packets will already be in flight
+			const double RestartHandshakeAckResendWindow = 10.0;
+			double CurTime = Driver != nullptr ? Driver->GetElapsedTime() : 0.0;
+
+			if (CurTime - LastInitTimestamp >= RestartHandshakeAckResendWindow)
+			{
+				LastChallengeSuccessAddress.Reset();
+				LastInitTimestamp = 0.0;
+			}
 		}
 	}
 }
@@ -1131,7 +1350,8 @@ void StatelessConnectHandlerComponent::Incoming(FBitReader& Packet)
 void StatelessConnectHandlerComponent::Outgoing(FBitWriter& Packet, FOutPacketTraits& Traits)
 {
 	// All UNetConnection packets must specify a zero bHandshakePacket value
-	FBitWriter NewPacket(GetAdjustedSizeBits(Packet.GetNumBits())+1, true);
+	const int32 AdjustedSize = GetAdjustedSizeBits(Packet.GetNumBits(), LastRemoteHandshakeVersion);
+	FBitWriter NewPacket(AdjustedSize + 1, true);
 	uint8 bHandshakePacket = 0;
 
 	if (MagicHeader.Num() > 0)
@@ -1139,6 +1359,12 @@ void StatelessConnectHandlerComponent::Outgoing(FBitWriter& Packet, FOutPacketTr
 		NewPacket.SerializeBits(MagicHeader.GetData(), MagicHeader.Num());
 	}
 
+	if (LastRemoteHandshakeVersion >= EHandshakeVersion::SessionClientId)
+	{
+		NewPacket.SerializeBits(&CachedGlobalNetTravelCount, SessionIDSizeBits);
+		NewPacket.SerializeBits(&CachedClientID, ClientIDSizeBits);
+	}
+
 	NewPacket.WriteBit(bHandshakePacket);
 	NewPacket.SerializeBits(Packet.GetData(), Packet.GetNumBits());
 
@@ -1159,6 +1385,20 @@ void StatelessConnectHandlerComponent::IncomingConnectionless(FIncomingPacketRef
 		Packet.SerializeBits(&ReadMagic, MagicHeader.Num());
 	}
 
+	bool bHasValidSessionID = true;
+	uint8 SessionID = 0;
+	uint8 ClientID = 0;
+
+	if (CurrentHandshakeVersion >= static_cast<uint8>(EHandshakeVersion::SessionClientId))
+	{
+		Packet.SerializeBits(&SessionID, SessionIDSizeBits);
+		Packet.SerializeBits(&ClientID, ClientIDSizeBits);
+
+		bHasValidSessionID = GVerifyNetSessionID == 0 || (SessionID == CachedGlobalNetTravelCount && !Packet.IsError());
+
+		// No ClientID validation until connected
+	}
+
 	bool bHandshakePacket = !!Packet.ReadBit() && !Packet.IsError();
 
 	LastChallengeSuccessAddress = nullptr;
@@ -1171,60 +1411,16 @@ void StatelessConnectHandlerComponent::IncomingConnectionless(FIncomingPacketRef
 
 		if (bHandshakePacket)
 		{
-			bool bValidHandshakeVersion = false;
-			EHandshakeVersion TargetVersion = static_cast<EHandshakeVersion>(CurrentHandshakeVersion);
-			const uint8 RemoteMinVersionUint8 = static_cast<uint8>(HandshakeData.RemoteMinVersion);
-
-			if (RemoteMinVersionUint8 >= MinSupportedHandshakeVersion && HandshakeData.RemoteMinVersion <= TargetVersion
-				&& HandshakeData.RemoteMinVersion <= HandshakeData.RemoteCurVersion)
-			{
-				if (HandshakeData.RemoteCurVersion <= TargetVersion)
-				{
-					TargetVersion = HandshakeData.RemoteCurVersion;
-				}
-
-				bValidHandshakeVersion = true;
-			}
-
-			const bool bCheckNetVersion = !UE_BUILD_SHIPPING || GHandshakeEnforceNetworkCLVersion;
-			bool bValidNetVersion = true;
-
-			if (bCheckNetVersion && HandshakeData.RemoteCurVersion >= EHandshakeVersion::NetCLVersion)
-			{
-				const uint32 LocalNetworkVersion = FNetworkVersion::GetLocalNetworkVersion();
-
-				if (!FNetworkVersion::IsNetworkCompatible(LocalNetworkVersion, HandshakeData.RemoteNetworkVersion))
-				{
-					bValidNetVersion = false;
-				}
-			}
-
-#if !UE_BUILD_SHIPPING
-			if (!bValidHandshakeVersion || !bValidNetVersion)
-			{
-				FDDoSDetection* DDoS = Handler->GetDDoS();
-				const uint32 LocalNetworkVersion = FNetworkVersion::GetLocalNetworkVersion();
-
-				UE_CLOG((DDoS == nullptr || !DDoS->CheckLogRestrictions()), LogHandshake, Log,
-						TEXT("IncomingConnectionless: Incompatible version. bValidHandshakeVersion: %i, bValidNetVersion: %i, ")
-						TEXT("GHandshakeEnforceNetworkCLVersion: %i, RemoteMinVersion: %i, RemoteCurVersion: %i, MinSupportedHandshakeVersion: %i, ")
-						TEXT("CurrentHandshakeVersion: %i, RemoteNetworkVersion: %i, LocalNetworkVersion: %i"),
-						(int32)bValidHandshakeVersion, (int32)bValidNetVersion, GHandshakeEnforceNetworkCLVersion, RemoteMinVersionUint8,
-						static_cast<uint8>(HandshakeData.RemoteCurVersion), MinSupportedHandshakeVersion, CurrentHandshakeVersion,
-						HandshakeData.RemoteNetworkVersion, LocalNetworkVersion);
-			}
-#endif
-
-			const bool bPassedNetVersionConditions = bValidNetVersion || !GHandshakeEnforceNetworkCLVersion;
-
-			if (Handler->Mode == Handler::Mode::Server && bValidHandshakeVersion && bPassedNetVersionConditions)
-			{
-				const bool bInitialConnect = HandshakeData.HandshakePacketType == EHandshakePacketType::InitialPacket &&
+			EHandshakeVersion TargetVersion = EHandshakeVersion::Latest;
+			bool bValidVersion = CheckVersion(HandshakeData, TargetVersion);
+			const bool bInitialConnect = HandshakeData.HandshakePacketType == EHandshakePacketType::InitialPacket &&
 												HandshakeData.Timestamp == 0.0;
 
+			if (Handler->Mode == Handler::Mode::Server && bValidVersion && (bHasValidSessionID || bInitialConnect))
+			{
 				if (bInitialConnect)
 				{
-					SendConnectChallenge(Address, TargetVersion, HandshakeData.RemoteSentHandshakePacketCount);
+					SendConnectChallenge(Address, TargetVersion, HandshakeData.RemoteSentHandshakePacketCount, ClientID);
 				}
 				// Challenge response
 				else if (Driver != nullptr)
@@ -1265,6 +1461,7 @@ void StatelessConnectHandlerComponent::IncomingConnectionless(FIncomingPacketRef
 							bRestartedHandshake = HandshakeData.bRestartHandshake;
 							LastChallengeSuccessAddress = Address->Clone();
 							LastRemoteHandshakeVersion = TargetVersion;
+							CachedClientID = ClientID;
 
 							if (TargetVersion < MinClientHandshakeVersion && static_cast<uint8>(TargetVersion) >= MinSupportedHandshakeVersion)
 							{
@@ -1273,11 +1470,20 @@ void StatelessConnectHandlerComponent::IncomingConnectionless(FIncomingPacketRef
 
 
 							// Now ack the challenge response - the cookie is stored in AuthorisedCookie, to enable retries
-							SendChallengeAck(Address, TargetVersion, HandshakeData.RemoteSentHandshakePacketCount, AuthorisedCookie);
+							SendChallengeAck(Address, TargetVersion, HandshakeData.RemoteSentHandshakePacketCount, ClientID, AuthorisedCookie);
 						}
 					}
 				}
 			}
+			else if (Handler->Mode == Handler::Mode::Server && bValidVersion && !bHasValidSessionID)
+			{
+#if !UE_BUILD_SHIPPING
+				UE_LOG(LogHandshake, Log, TEXT("IncomingConnectionless: Rejecting packet with invalid session id: %i vs %i."),
+						SessionID, CachedGlobalNetTravelCount);
+#endif
+
+				Packet.SetError();
+			}
 		}
 		else
 		{
@@ -1299,19 +1505,22 @@ void StatelessConnectHandlerComponent::IncomingConnectionless(FIncomingPacketRef
 #if !UE_BUILD_SHIPPING
 	else if (Packet.IsError())
 	{
-		UE_LOG(LogHandshake, Log, TEXT("IncomingConnectionless: Error reading handshake bit from packet."));
+		UE_LOG(LogHandshake, Log, TEXT("IncomingConnectionless: Error reading session id, connection id and handshake bit from packet."));
 	}
 #endif
-	// Late packets from recently disconnected clients may incorrectly trigger this code path, so detect and exclude those packets
-	else if (!Packet.IsError() && !PacketRef.Traits.bFromRecentlyDisconnected)
+	else if (bHasValidSessionID)
 	{
-		// The packet was fine but not a handshake packet - an existing client might suddenly be communicating on a different address.
-		// If we get them to resend their cookie, we can update the connection's info with their new address.
-		SendRestartHandshakeRequest(Address, static_cast<EHandshakeVersion>(MinSupportedHandshakeVersion));
+		// Late packets from recently disconnected clients may incorrectly trigger this code path, so detect and exclude those packets
+		if (!Packet.IsError() && !PacketRef.Traits.bFromRecentlyDisconnected)
+		{
+			// The packet was fine but not a handshake packet - an existing client might suddenly be communicating on a different address.
+			// If we get them to resend their cookie, we can update the connection's info with their new address.
+			SendRestartHandshakeRequest(Address, static_cast<EHandshakeVersion>(MinSupportedHandshakeVersion), ClientID);
+		}
 	}
 }
 
-bool StatelessConnectHandlerComponent::ParseHandshakePacket(FBitReader& Packet, FParsedHandshakeData& OutResult)
+bool StatelessConnectHandlerComponent::ParseHandshakePacket(FBitReader& Packet, FParsedHandshakeData& OutResult) const
 {
 	using namespace UE::Net;
 
@@ -1350,12 +1559,16 @@ bool StatelessConnectHandlerComponent::ParseHandshakePacket(FBitReader& Packet,
 
 	bool bValidPacket = false;
 	const int32 BitsLeft = Packet.GetBitsLeft();
-	const bool bOriginalVersion = BitsLeft == (OriginalHandshakePacketSizeBits - 1) || BitsLeft == (OriginalRestartHandshakePacketSizeBits - 1) ||
-									BitsLeft == (OriginalRestartResponseSizeBits - 1);
 
-	if (bOriginalVersion)
+	if (MinSupportedHandshakeVersion == static_cast<uint8>(EHandshakeVersion::Original))
 	{
-		return ParseHandshakePacketOriginal(Packet, OutResult);
+		const bool bOriginalVersion = BitsLeft == (OriginalHandshakePacketSizeBits - 1) || BitsLeft == (OriginalRestartHandshakePacketSizeBits - 1) ||
+										BitsLeft == (OriginalRestartResponseSizeBits - 1);
+
+		if (bOriginalVersion)
+		{
+			return ParseHandshakePacketOriginal(Packet, OutResult);
+		}
 	}
 
 
@@ -1430,7 +1643,7 @@ bool StatelessConnectHandlerComponent::ParseHandshakePacket(FBitReader& Packet,
 	return bValidPacket;
 }
 
-bool StatelessConnectHandlerComponent::ParseHandshakePacketOriginal(FBitReader& Packet, FParsedHandshakeData& OutResult)
+bool StatelessConnectHandlerComponent::ParseHandshakePacketOriginal(FBitReader& Packet, FParsedHandshakeData& OutResult) const
 {
 	using namespace UE::Net;
 
@@ -1489,7 +1702,77 @@ bool StatelessConnectHandlerComponent::ParseHandshakePacketOriginal(FBitReader&
 	return bValidPacket;
 }
 
-void StatelessConnectHandlerComponent::GenerateCookie(TSharedPtr<const FInternetAddr> ClientAddress, uint8 SecretId, double Timestamp, uint8 (&OutCookie)[20])
+bool StatelessConnectHandlerComponent::CheckVersion(const FParsedHandshakeData& HandshakeData, EHandshakeVersion& OutTargetVersion) const
+{
+	using namespace UE::Net;
+
+	bool bValidHandshakeVersion = false;
+	const uint8 RemoteMinVersionUint8 = static_cast<uint8>(HandshakeData.RemoteMinVersion);
+	EHandshakeVersion LocalMaxVersion = EHandshakeVersion::Latest;
+	bool bHasMaxVersion = false;
+
+	OutTargetVersion = static_cast<EHandshakeVersion>(CurrentHandshakeVersion);
+
+	for (EHandshakeVersion CurBreakVersion : HandshakeCompatibilityBreaks)
+	{
+		if (CurBreakVersion > OutTargetVersion)
+		{
+			// A maximum version is only enforced if we're aware of a net compatibility breakage in a higher version
+			LocalMaxVersion = static_cast<EHandshakeVersion>(static_cast<uint8>(CurBreakVersion) - 1);
+			bHasMaxVersion = true;
+
+			break;
+		}
+	}
+
+	if (RemoteMinVersionUint8 >= MinSupportedHandshakeVersion && HandshakeData.RemoteMinVersion <= OutTargetVersion
+		&& HandshakeData.RemoteMinVersion <= HandshakeData.RemoteCurVersion &&
+		(!bHasMaxVersion || HandshakeData.RemoteCurVersion <= LocalMaxVersion))
+	{
+		if (HandshakeData.RemoteCurVersion <= OutTargetVersion)
+		{
+			OutTargetVersion = HandshakeData.RemoteCurVersion;
+		}
+
+		bValidHandshakeVersion = true;
+	}
+
+	const bool bCheckNetVersion = !UE_BUILD_SHIPPING || GHandshakeEnforceNetworkCLVersion;
+	bool bValidNetVersion = true;
+
+	if (bCheckNetVersion && HandshakeData.RemoteCurVersion >= EHandshakeVersion::NetCLVersion)
+	{
+		const uint32 LocalNetworkVersion = FNetworkVersion::GetLocalNetworkVersion();
+
+		if (!FNetworkVersion::IsNetworkCompatible(LocalNetworkVersion, HandshakeData.RemoteNetworkVersion))
+		{
+			bValidNetVersion = false;
+		}
+	}
+
+#if !UE_BUILD_SHIPPING
+	if (!bValidHandshakeVersion || !bValidNetVersion)
+	{
+		FDDoSDetection* DDoS = Handler->GetDDoS();
+		const uint32 LocalNetworkVersion = FNetworkVersion::GetLocalNetworkVersion();
+
+		UE_CLOG((DDoS == nullptr || !DDoS->CheckLogRestrictions()), LogHandshake, Log,
+				TEXT("CheckVersion: Incompatible version. bValidHandshakeVersion: %i, bValidNetVersion: %i, ")
+				TEXT("GHandshakeEnforceNetworkCLVersion: %i, RemoteMinVersion: %i, RemoteCurVersion: %i, MinSupportedHandshakeVersion: %i, ")
+				TEXT("CurrentHandshakeVersion: %i, RemoteNetworkVersion: %i, LocalNetworkVersion: %i"),
+				(int32)bValidHandshakeVersion, (int32)bValidNetVersion, GHandshakeEnforceNetworkCLVersion, RemoteMinVersionUint8,
+				static_cast<uint8>(HandshakeData.RemoteCurVersion), MinSupportedHandshakeVersion, CurrentHandshakeVersion,
+				HandshakeData.RemoteNetworkVersion, LocalNetworkVersion);
+	}
+#endif
+
+	const bool bPassedNetVersionConditions = bValidNetVersion || !GHandshakeEnforceNetworkCLVersion;
+
+	return bValidHandshakeVersion && bPassedNetVersionConditions;
+}
+
+void StatelessConnectHandlerComponent::GenerateCookie(TSharedPtr<const FInternetAddr> ClientAddress, uint8 SecretId, double Timestamp,
+														uint8 (&OutCookie)[20]) const
 {
 	TArray<uint8> CookieData;
 	FMemoryWriter CookieArc(CookieData);
@@ -1536,7 +1819,8 @@ void StatelessConnectHandlerComponent::UpdateSecret()
 
 int32 StatelessConnectHandlerComponent::GetReservedPacketBits() const
 {
-	int32 ReturnVal = MagicHeader.Num() + 1 /* bHandshakePacket */;
+	// Count all base bit additions which affect NetConnection packets, regardless of handshake protocol version - as this is called upon construction
+	int32 ReturnVal = MagicHeader.Num() + SessionIDSizeBits + ClientIDSizeBits + 1 /* bHandshakePacket */;
 
 #if !UE_BUILD_SHIPPING
 	SET_DWORD_STAT(STAT_PacketReservedHandshake, ReturnVal);
@@ -1545,6 +1829,18 @@ int32 StatelessConnectHandlerComponent::GetReservedPacketBits() const
 	return ReturnVal;
 }
 
+int32 StatelessConnectHandlerComponent::GetAdjustedSizeBits(int32 InSizeBits, EHandshakeVersion HandshakeVersion) const
+{
+	int32 ReturnVal = MagicHeader.Num() + InSizeBits;
+
+	if (HandshakeVersion >= EHandshakeVersion::SessionClientId)
+	{
+		ReturnVal += SessionIDSizeBits + ClientIDSizeBits;
+	}
+
+	return ReturnVal;
+}
+
 void StatelessConnectHandlerComponent::Tick(float DeltaTime)
 {
 	using namespace UE::Net;
@@ -1615,3 +1911,24 @@ void StatelessConnectHandlerComponent::Tick(float DeltaTime)
 	}
 }
 
+#if !UE_BUILD_SHIPPING
+bool StatelessConnectHandlerComponent::TrackSessionConnIDLogs()
+{
+	const uint32 NumSessionConnIDLogsPerPeriod = 10;
+	const double SessionConnIDLogPeriod = 30.0;
+	const double ElapsedTime = Driver->GetElapsedTime();
+
+	if (ElapsedTime - LastSessionConnIDLogPeriodStart >= SessionConnIDLogPeriod)
+	{
+		LastSessionConnIDLogPeriodStart = ElapsedTime;
+		SessionConnIDLogCounter = 0;
+	}
+	else
+	{
+		SessionConnIDLogCounter++;
+	}
+
+	return SessionConnIDLogCounter < NumSessionConnIDLogsPerPeriod;
+}
+#endif
+
diff --git a/Engine/Source/Runtime/Engine/Public/PacketHandlers/StatelessConnectHandlerComponent.h b/Engine/Source/Runtime/Engine/Public/PacketHandlers/StatelessConnectHandlerComponent.h
--- a/Engine/Source/Runtime/Engine/Public/PacketHandlers/StatelessConnectHandlerComponent.h
+++ b/Engine/Source/Runtime/Engine/Public/PacketHandlers/StatelessConnectHandlerComponent.h
@@ -37,9 +37,17 @@ namespace UE::Net
 		Randomized		= 1,	// The version of the handshake protocol with randomization, versioning and debug/diagnostic tweaks
 		NetCLVersion	= 2,	// Added Network CL version, for optional extra-early client rejection
 
-		Latest			= NetCLVersion
+		/** Net compatibility break */
+
+		SessionClientId	= 3,	// Added SessionId (index of server game session) and ClientId (clientside index of NetConnection)
+
+		Latest			= SessionClientId
 	};
 
+	/** List of EHandshakeVersion's which break net compatibility (sorted from earliest to latest) */
+	static constexpr EHandshakeVersion HandshakeCompatibilityBreaks[] = { EHandshakeVersion::SessionClientId };
+
+
 	/**
 	 * The different handshake packet types
 	 */
@@ -98,8 +106,10 @@ private:
 	 * @param ClientAddress						The address of the client to send the challenge to.
 	 * @param HandshakeVersion					The handshake format version to use when sending
 	 * @param ClientSentHandshakePacketCount	The number of handshake packets the client has sent (for debugging/packet-analysis)
+	 * @param InClientID						The client-specified connection id
 	 */
-	void SendConnectChallenge(TSharedPtr<const FInternetAddr> ClientAddress, EHandshakeVersion HandshakeVersion, uint8 ClientSentHandshakePacketCount);
+	void SendConnectChallenge(TSharedPtr<const FInternetAddr> ClientAddress, EHandshakeVersion HandshakeVersion,
+								uint8 ClientSentHandshakePacketCount, uint32 InClientID);
 
 	/**
 	 * Constructs and sends the handshake challenge response packet, from the client to the server
@@ -114,21 +124,24 @@ private:
 	/**
 	 * Constructs and sends the server ack to a successful challenge response, from the server to the client.
 	 *
-	 * @param ClientAddress		The address of the client to send the ack to.
-	 * @param HandshakeVersion	The handshake format version to use when sending
+	 * @param ClientAddress						The address of the client to send the ack to.
+	 * @param HandshakeVersion					The handshake format version to use when sending
 	 * @param ClientSentHandshakePacketCount	The number of handshake packets the client has sent (for debugging/packet-analysis)
-	 * @param InCookie			The cookie value to send
+	 * @param InClientID						The client-specified connection id
+	 * @param InCookie							The cookie value to send
 	 */
 	void SendChallengeAck(TSharedPtr<const FInternetAddr> ClientAddress, EHandshakeVersion HandshakeVersion, uint8 ClientSentHandshakePacketCount,
-							uint8 InCookie[COOKIE_BYTE_SIZE]);
+							uint32 InClientID, uint8 InCookie[COOKIE_BYTE_SIZE]);
 
 	/**
 	 * Constructs and sends a request to resend the cookie, from the server to the client.
 	 *
 	 * @param ClientAddress		The address of the client to send the request to.
 	 * @param HandshakeVersion	The handshake format version to use when sending
+	 * @param InClientID		The client-specified connection id
 	 */
-	void SendRestartHandshakeRequest(const TSharedPtr<const FInternetAddr> ClientAddress, EHandshakeVersion HandshakeVersion);
+	void SendRestartHandshakeRequest(const TSharedPtr<const FInternetAddr> ClientAddress, EHandshakeVersion HandshakeVersion,
+										uint32 InClientID);
 
 
 	/**
@@ -290,9 +303,18 @@ private:
 	 * @param OutResult		The parsed handshake data, if successful
 	 * @return				Whether or not the handshake packet was parsed successfully
 	 */
-	bool ParseHandshakePacket(FBitReader& Packet, FParsedHandshakeData& OutResult);
+	bool ParseHandshakePacket(FBitReader& Packet, FParsedHandshakeData& OutResult) const;
+
+	bool ParseHandshakePacketOriginal(FBitReader& Packet, FParsedHandshakeData& OutResult) const;
 
-	bool ParseHandshakePacketOriginal(FBitReader& Packet, FParsedHandshakeData& OutResult);
+	/**
+	 * Checks the handshake protocol version and Network CL version for incoming handshake packets, serverside
+	 *
+	 * @param HandshakeData		The packet handshake data
+	 * @param OutTargetVersion	The target handshake protocol version to use, for the new connection
+	 * @return					Whether or not the checks successfully validated that the version is correct, and that communication can continue
+	 */
+	bool CheckVersion(const FParsedHandshakeData& HandshakeData, EHandshakeVersion& OutTargetVersion) const;
 
 	/**
 	 * Takes the client address plus server timestamp, and outputs a deterministic cookie value
@@ -302,7 +324,7 @@ private:
 	 * @param TimeStamp			The serverside timestamp
 	 * @param OutCookie			Outputs the generated cookie value.
 	 */
-	void GenerateCookie(TSharedPtr<const FInternetAddr> ClientAddress, uint8 SecretId, double Timestamp, uint8 (&OutCookie)[COOKIE_BYTE_SIZE]);
+	void GenerateCookie(TSharedPtr<const FInternetAddr> ClientAddress, uint8 SecretId, double Timestamp, uint8 (&OutCookie)[COOKIE_BYTE_SIZE]) const;
 
 	/**
 	 * Generates a new HandshakeSecret value
@@ -312,10 +334,16 @@ private:
 	/**
 	 * Calculate the adjusted predefined packet size, based on MagicHeader
 	 */
-	FORCEINLINE int32 GetAdjustedSizeBits(int32 InSizeBits)
-	{
-		return MagicHeader.Num() + InSizeBits;
-	}
+	int32 GetAdjustedSizeBits(int32 InSizeBits, EHandshakeVersion HandshakeVersion) const;
+
+#if !UE_BUILD_SHIPPING
+	/**
+	 * Tracks logs for invalid Session ID's or Connection ID's, and limits the number of these logs
+	 *
+	 * @return		Whether or not logging is currently allowed
+	 */
+	bool TrackSessionConnIDLogs();
+#endif
 
 
 private:
@@ -349,6 +377,14 @@ private:
 	/** The minimum supported client version, for all connections - for setting minimum bounds for the restart handshake (not aware of version) */
 	EHandshakeVersion MinClientHandshakeVersion;
 
+#if !UE_BUILD_SHIPPING
+	/** The last time an 'invalid session/connection id' log period began */
+	double LastSessionConnIDLogPeriodStart = 0.0;
+
+	/** The number of 'invalid session/connection id' logs in the current log period (for logspam limiting) */
+	uint32 SessionConnIDLogCounter = 0;
+#endif
+
 
 	/** Clientside variables */
 
@@ -388,5 +424,17 @@ private:
 
 	/** The last remote handshake version sent by the client/server */
 	EHandshakeVersion LastRemoteHandshakeVersion;
+
+	/** Cached version of the serverside UEngine.GlobalNetTravelCount value, for writing session id's */
+	uint32 CachedGlobalNetTravelCount = 0;
+
+	/** Cached version of the clientside per-NetDriverDefinition 'ClientID' value, for writing client connection id's */
+	uint32 CachedClientID = 0;
+
+	/** The size of the session id in packets (WARNING: Adjusting this is a net compatibility break) */
+	static constexpr uint32 SessionIDSizeBits = 2;
+
+	/** The size of the connection id in packets (WARNING: Adjusting this is a net compatibility break) */
+	static constexpr uint32 ClientIDSizeBits = 3;
 };
 
