Fix crash occuring when AWorldDataLayers is removed from world but the tree item is still present

#test editor

diff --git a/Engine/Source/Runtime/Engine/Private/WorldPartition/DataLayer/WorldDataLayers.cpp b/Engine/Source/Runtime/Engine/Private/WorldPartition/DataLayer/WorldDataLayers.cpp
--- a/Engine/Source/Runtime/Engine/Private/WorldPartition/DataLayer/WorldDataLayers.cpp
+++ b/Engine/Source/Runtime/Engine/Private/WorldPartition/DataLayer/WorldDataLayers.cpp
@@ -853,9 +853,9 @@ bool AWorldDataLayers::IsRuntimeRelevant() const
 
 bool AWorldDataLayers::IsSubWorldDataLayers() const
 {
-	check(GetWorld());
+	UWorld* ActorWorld = GetWorld();
 	UWorld* OuterWorld = GetTypedOuter<UWorld>();
-	return OuterWorld != nullptr && OuterWorld->GetFName() != GetWorld()->GetFName();
+	return ActorWorld != nullptr && OuterWorld != nullptr && OuterWorld->GetFName() != ActorWorld->GetFName();
 }
 
 bool AWorldDataLayers::IsTheMainWorldDataLayers() const
