Material Baking - Force disable restrictive mode when compiling properties for the export material proxy

diff --git a/Engine/Source/Developer/MaterialBaking/Private/ExportMaterialProxy.h b/Engine/Source/Developer/MaterialBaking/Private/ExportMaterialProxy.h
--- a/Engine/Source/Developer/MaterialBaking/Private/ExportMaterialProxy.h
+++ b/Engine/Source/Developer/MaterialBaking/Private/ExportMaterialProxy.h
@@ -19,6 +19,7 @@
 #include "SceneTypes.h"
 #include "Materials/Material.h"
 #include "Materials/MaterialExpressionCustomOutput.h"
+#include "MaterialRestrictiveMode.h"
 
 struct FExportMaterialCompiler : public FProxyMaterialCompiler
 {
@@ -354,6 +355,8 @@ public:
 	/** helper for CompilePropertyAndSetMaterialProperty() */
 	int32 CompilePropertyAndSetMaterialPropertyWithoutCast(EMaterialProperty Property, FMaterialCompiler* Compiler) const
 	{
+		FScopedHLSLMaterialTranslatorRestrictiveModeChange HLSLMaterialTranslatorRestrictiveModeChangeScope(false);
+
 		if (Property == MP_EmissiveColor)
 		{
 			const EBlendMode BlendMode = MaterialInterface->GetBlendMode();
diff --git a/Engine/Source/Runtime/Engine/Private/Materials/HLSLMaterialTranslator.cpp b/Engine/Source/Runtime/Engine/Private/Materials/HLSLMaterialTranslator.cpp
--- a/Engine/Source/Runtime/Engine/Private/Materials/HLSLMaterialTranslator.cpp
+++ b/Engine/Source/Runtime/Engine/Private/Materials/HLSLMaterialTranslator.cpp
@@ -2571,6 +2571,11 @@ FString FHLSLMaterialTranslator::GetParameterCodeRaw(int32 Index, const TCHAR* D
 }
 
 // Public interface for restrictive mode.
+bool GetHLSLMaterialTranslatorRestrictiveMode()
+{
+	return FHLSLMaterialTranslator::InRestrictiveMode();
+}
+
 void SetHLSLMaterialTranslatorRestrictiveMode(bool bRestrictiveMode)
 {
 	FHLSLMaterialTranslator::SetRestrictiveMode(bRestrictiveMode);
diff --git a/Engine/Source/Runtime/Engine/Public/MaterialRestrictiveMode.h b/Engine/Source/Runtime/Engine/Public/MaterialRestrictiveMode.h
--- a/Engine/Source/Runtime/Engine/Public/MaterialRestrictiveMode.h
+++ b/Engine/Source/Runtime/Engine/Public/MaterialRestrictiveMode.h
@@ -5,7 +5,28 @@
 
 #if WITH_EDITORONLY_DATA
 
+/// Get whether the HLSL material translator is in restrive mode.
+ENGINE_API bool GetHLSLMaterialTranslatorRestrictiveMode();
+
 /// Puts the HLSL material translator into restrictive mode.
 ENGINE_API void SetHLSLMaterialTranslatorRestrictiveMode(bool bRestrictiveMode);
 
+/// Scope helper to toggle restrictive mode for a code block
+struct FScopedHLSLMaterialTranslatorRestrictiveModeChange
+{
+	FORCEINLINE FScopedHLSLMaterialTranslatorRestrictiveModeChange(bool bRestrictiveMode)
+	{
+		bRestoreRestrictiveMode = GetHLSLMaterialTranslatorRestrictiveMode();
+		SetHLSLMaterialTranslatorRestrictiveMode(bRestrictiveMode);
+	}
+	
+	FORCEINLINE FScopedHLSLMaterialTranslatorRestrictiveModeChange()
+	{
+		SetHLSLMaterialTranslatorRestrictiveMode(bRestoreRestrictiveMode);
+	}
+
+private:
+	bool bRestoreRestrictiveMode;
+};
+
 #endif // WITH_EDITORONLY_DATA
