WorldPartition - Minimap: Fixed stretched tile capture when tile wasn't square

diff --git a/Engine/Source/Runtime/Engine/Private/WorldPartition/WorldPartitionMiniMapHelper.cpp b/Engine/Source/Runtime/Engine/Private/WorldPartition/WorldPartitionMiniMapHelper.cpp
--- a/Engine/Source/Runtime/Engine/Private/WorldPartition/WorldPartitionMiniMapHelper.cpp
+++ b/Engine/Source/Runtime/Engine/Private/WorldPartition/WorldPartitionMiniMapHelper.cpp
@@ -63,16 +63,9 @@ void FWorldPartitionMiniMapHelper::CaptureBoundsMiniMapToTexture(UWorld* InWorld
 	// Before capturing the scene, make sure all assets are finished compiling
 	FAssetCompilingManager::Get().FinishAllCompilation();
 
-	//Calculate Viewport size
-	FBox2D WorldBounds2D(FVector2D(InBounds.Min), FVector2D(InBounds.Max));
-	FVector2D ViewSize = WorldBounds2D.Max - WorldBounds2D.Min;
-	float AspectRatio = FMath::Abs(ViewSize.X) / FMath::Abs(ViewSize.Y);
-	uint32 ViewportWidth = InMiniMapSizeX * AspectRatio;
-	uint32 ViewportHeight = InMiniMapSizeY;
-
 	//Calculate Projection matrix based on world bounds.
 	FMatrix ProjectionMatrix;
-	FWorldPartitionMiniMapHelper::CalTopViewOfWorld(ProjectionMatrix, InBounds, ViewportWidth, ViewportHeight);
+	FWorldPartitionMiniMapHelper::CalTopViewOfWorld(ProjectionMatrix, InBounds, InMiniMapSizeX, InMiniMapSizeY);
 
 	//Using SceneCapture Actor capture the scene to buffer
 	UTextureRenderTarget2D* RenderTargetTexture = NewObject<UTextureRenderTarget2D>();
@@ -90,7 +83,7 @@ void FWorldPartitionMiniMapHelper::CaptureBoundsMiniMapToTexture(UWorld* InWorld
 	CaptureComponent->TextureTarget = RenderTargetTexture;
 	CaptureComponent->ProjectionType = ECameraProjectionMode::Orthographic;
 	CaptureComponent->CaptureSource = InCaptureSource;
-	CaptureComponent->OrthoWidth = ViewportWidth;
+	CaptureComponent->OrthoWidth = InMiniMapSizeX;
 	CaptureComponent->bUseCustomProjectionMatrix = true;
 	CaptureComponent->CustomProjectionMatrix = ProjectionMatrix;
 	CaptureComponent->bCaptureEveryFrame = false;
