Reserve bytes array for FUniqueNetIdRepl::ReplicationBytes to minimize allocations and wasted space

#tests PF compile + engine tests

diff --git a/Engine/Source/Runtime/Engine/Private/OnlineReplStructs.cpp b/Engine/Source/Runtime/Engine/Private/OnlineReplStructs.cpp
--- a/Engine/Source/Runtime/Engine/Private/OnlineReplStructs.cpp
+++ b/Engine/Source/Runtime/Engine/Private/OnlineReplStructs.cpp
@@ -125,7 +125,7 @@ void FUniqueNetIdRepl::MakeReplicationData()
 	{
 		EUniqueIdEncodingFlags EncodingFlags = (EUniqueIdEncodingFlags::IsEncoded | EUniqueIdEncodingFlags::IsEmpty);
 
-		ReplicationBytes.Empty();
+		ReplicationBytes.Empty(sizeof(EncodingFlags));
 		FMemoryWriter Writer(ReplicationBytes);
 		Writer << EncodingFlags;
 	}
@@ -178,7 +178,7 @@ void FUniqueNetIdRepl::MakeReplicationDataV1()
 		if (EnumHasAllFlags(EncodingFlags, EUniqueIdEncodingFlags::IsEncoded))
 		{
 			uint8 EncodedSize = static_cast<uint8>(EncodedSize32);
-			const int32 TotalBytes = sizeof(EncodingFlags) + sizeof(EncodedSize) + EncodedSize;
+			const int32 TotalBytes = sizeof(EncodingFlags) + sizeof(EncodedSize) + EncodedSize; // no optimization for TypeHash_Other
 			ReplicationBytes.Empty(TotalBytes);
 
 			FMemoryWriter Writer(ReplicationBytes);
@@ -199,7 +199,8 @@ void FUniqueNetIdRepl::MakeReplicationDataV1()
 		}
 		else
 		{
-			ReplicationBytes.Empty(Length);
+			const int32 TotalBytes = sizeof(EncodingFlags) + sizeof(Length) + Length * sizeof(TCHAR); // no optimization for TypeHash_Other
+			ReplicationBytes.Empty(TotalBytes);
 
 			FMemoryWriter Writer(ReplicationBytes);
 			Writer << EncodingFlags;
@@ -216,7 +217,7 @@ void FUniqueNetIdRepl::MakeReplicationDataV1()
 	{
 		EUniqueIdEncodingFlags EncodingFlags = (EUniqueIdEncodingFlags::IsEncoded | EUniqueIdEncodingFlags::IsEmpty);
 
-		ReplicationBytes.Empty();
+		ReplicationBytes.Empty(sizeof(EncodingFlags));
 		FMemoryWriter Writer(ReplicationBytes);
 		Writer << EncodingFlags;
 	}
@@ -232,7 +233,8 @@ void FUniqueNetIdRepl::MakeReplicationDataV2()
 	TArray<uint8> ReplicationData = UE::Online::FOnlineIdRegistryRegistry::Get().ToReplicationData(AccountId);
 	check(!ReplicationData.IsEmpty());
 
-	ReplicationBytes.Empty();
+	const int32 TotalBytes = sizeof(EncodingFlags) + sizeof(OnlineServicesType) + sizeof(ReplicationData.Num()) + ReplicationData.Num();
+	ReplicationBytes.Empty(TotalBytes);
 	FMemoryWriter Writer(ReplicationBytes);
 	Writer << EncodingFlags;
 	Writer << OnlineServicesType;
