Remove references to WORLD_MAX from Audio Code

diff --git a/Engine/Source/Runtime/Engine/Private/ActiveSound.cpp b/Engine/Source/Runtime/Engine/Private/ActiveSound.cpp
--- a/Engine/Source/Runtime/Engine/Private/ActiveSound.cpp
+++ b/Engine/Source/Runtime/Engine/Private/ActiveSound.cpp
@@ -119,7 +119,7 @@ FActiveSound::FActiveSound()
 	, VolumeConcurrency(0.0f)
 	, OcclusionCheckInterval(0.0f)
 	, LastOcclusionCheckTime(TNumericLimits<float>::Lowest())
-	, MaxDistance(WORLD_MAX)
+	, MaxDistance(FAudioDevice::GetMaxWorldDistance())
 	, LastLocation(FVector::ZeroVector)
 	, AudioVolumeID(0)
 	, LastUpdateTime(0.0f)
diff --git a/Engine/Source/Runtime/Engine/Private/Attenuation.cpp b/Engine/Source/Runtime/Engine/Private/Attenuation.cpp
--- a/Engine/Source/Runtime/Engine/Private/Attenuation.cpp
+++ b/Engine/Source/Runtime/Engine/Private/Attenuation.cpp
@@ -2,6 +2,7 @@
 
 #include "Engine/Attenuation.h"
 
+#include "AudioDevice.h"
 #include "DSP/Dsp.h"
 #include "EngineDefines.h"
 
@@ -52,12 +53,12 @@ float FBaseAttenuationSettings::GetMaxDimension() const
 		check(false);
 	}
 
-	return FMath::Clamp(MaxDimension, 0.0f, static_cast<float>(WORLD_MAX));
+	return FMath::Clamp(MaxDimension, 0.0f, static_cast<float>(FAudioDevice::GetMaxWorldDistance()));
 }
 
 float FBaseAttenuationSettings::GetMaxFalloffDistance() const
 {
-	static const float WorldMax = static_cast<float>(WORLD_MAX);
+	static const float WorldMax = static_cast<float>(FAudioDevice::GetMaxWorldDistance());
 	if (FalloffDistance > WorldMax)
 	{
 		return WorldMax;
diff --git a/Engine/Source/Runtime/Engine/Private/AudioDevice.cpp b/Engine/Source/Runtime/Engine/Private/AudioDevice.cpp
--- a/Engine/Source/Runtime/Engine/Private/AudioDevice.cpp
+++ b/Engine/Source/Runtime/Engine/Private/AudioDevice.cpp
@@ -180,6 +180,13 @@ FAutoConsoleVariableRef CVarFlushAudioRenderThreadOnGC(
 	TEXT("When set to 1, every time the GC runs, we flush all pending audio render thread commands.\n"),
 	ECVF_Default);
 
+static float MaxWorldDistanceCVar = UE_OLD_WORLD_MAX;
+FAutoConsoleVariableRef CVarSetAudioMaxDistance(
+	TEXT("au.MaxWorldDistance"),
+	MaxWorldDistanceCVar,
+	TEXT("Maximum world distance used in audio-related calculations (eg. attenuation).\n"),
+	ECVF_Default);
+
 static FAutoConsoleCommandWithWorld GListAvailableSpatialPluginsCommand(
 	TEXT("au.spatialization.ListAvailableSpatialPlugins"),
 	TEXT("This will output a list of currently availible/active spatialization plugins"),
@@ -5791,6 +5798,11 @@ bool FAudioDevice::GetDistanceSquaredToNearestListener(const FVector& Location,
 	return true;
 }
 
+float FAudioDevice::GetMaxWorldDistance()
+{
+	return MaxWorldDistanceCVar;
+}
+
 bool FAudioDevice::GetListenerPosition(int32 ListenerIndex, FVector& OutPosition, bool bAllowOverride) const
 {
 	OutPosition = FVector::ZeroVector;
diff --git a/Engine/Source/Runtime/Engine/Private/SoundAttenuation.cpp b/Engine/Source/Runtime/Engine/Private/SoundAttenuation.cpp
--- a/Engine/Source/Runtime/Engine/Private/SoundAttenuation.cpp
+++ b/Engine/Source/Runtime/Engine/Private/SoundAttenuation.cpp
@@ -19,7 +19,7 @@ void FSoundAttenuationSettings::PostSerialize(const FArchive& Ar)
 	if (Ar.UEVer() < VER_UE4_ATTENUATION_SHAPES)
 	{
 		FalloffDistance = RadiusMax_DEPRECATED - RadiusMin_DEPRECATED;
-
+		const float MaxDistance = FAudioDevice::GetMaxWorldDistance();
 		switch(DistanceType_DEPRECATED)
 		{
 		case SOUNDDISTANCE_Normal:
@@ -29,17 +29,17 @@ void FSoundAttenuationSettings::PostSerialize(const FArchive& Ar)
 
 		case SOUNDDISTANCE_InfiniteXYPlane:
 			AttenuationShape = EAttenuationShape::Box;
-			AttenuationShapeExtents = FVector(WORLD_MAX, WORLD_MAX, RadiusMin_DEPRECATED);
+			AttenuationShapeExtents = FVector(MaxDistance, MaxDistance, RadiusMin_DEPRECATED);
 			break;
 
 		case SOUNDDISTANCE_InfiniteXZPlane:
 			AttenuationShape = EAttenuationShape::Box;
-			AttenuationShapeExtents = FVector(WORLD_MAX, RadiusMin_DEPRECATED, WORLD_MAX);
+			AttenuationShapeExtents = FVector(MaxDistance, RadiusMin_DEPRECATED, MaxDistance);
 			break;
 
 		case SOUNDDISTANCE_InfiniteYZPlane:
 			AttenuationShape = EAttenuationShape::Box;
-			AttenuationShapeExtents = FVector(RadiusMin_DEPRECATED, WORLD_MAX, WORLD_MAX);
+			AttenuationShapeExtents = FVector(RadiusMin_DEPRECATED, MaxDistance, MaxDistance);
 			break;
 		}
 	}
diff --git a/Engine/Source/Runtime/Engine/Private/SoundBase.cpp b/Engine/Source/Runtime/Engine/Private/SoundBase.cpp
--- a/Engine/Source/Runtime/Engine/Private/SoundBase.cpp
+++ b/Engine/Source/Runtime/Engine/Private/SoundBase.cpp
@@ -63,7 +63,7 @@ float USoundBase::GetMaxDistance() const
 		}
 	}
 
-	return WORLD_MAX;
+	return FAudioDevice::GetMaxWorldDistance();
 }
 
 float USoundBase::GetDuration() const
diff --git a/Engine/Source/Runtime/Engine/Private/SoundCue.cpp b/Engine/Source/Runtime/Engine/Private/SoundCue.cpp
--- a/Engine/Source/Runtime/Engine/Private/SoundCue.cpp
+++ b/Engine/Source/Runtime/Engine/Private/SoundCue.cpp
@@ -21,6 +21,7 @@
 #include "Sound/SoundNodeWavePlayer.h"
 #include "GameFramework/GameUserSettings.h"
 #include "AudioCompressionSettingsUtils.h"
+#include "AudioDevice.h"
 #include "AudioThread.h"
 #include "DSP/Dsp.h"
 #if WITH_EDITOR
@@ -345,7 +346,7 @@ float USoundCue::FindMaxDistanceInternal() const
 	{
 		if (!Settings->bAttenuate)
 		{
-			return WORLD_MAX;
+			return FAudioDevice::GetMaxWorldDistance();
 		}
 
 		OutMaxDistance = FMath::Max(OutMaxDistance, Settings->GetMaxDimension());
diff --git a/Engine/Source/Runtime/Engine/Private/SoundNodeAttenuation.cpp b/Engine/Source/Runtime/Engine/Private/SoundNodeAttenuation.cpp
--- a/Engine/Source/Runtime/Engine/Private/SoundNodeAttenuation.cpp
+++ b/Engine/Source/Runtime/Engine/Private/SoundNodeAttenuation.cpp
@@ -17,7 +17,7 @@ USoundNodeAttenuation::USoundNodeAttenuation(const FObjectInitializer& ObjectIni
 
 float USoundNodeAttenuation::GetMaxDistance() const 
 { 
-	float MaxDistance = WORLD_MAX;
+	float MaxDistance = FAudioDevice::GetMaxWorldDistance();
 	const FSoundAttenuationSettings* Settings = GetAttenuationSettingsToApply();
 	if (Settings)
 	{
diff --git a/Engine/Source/Runtime/Engine/Public/AudioDevice.h b/Engine/Source/Runtime/Engine/Public/AudioDevice.h
--- a/Engine/Source/Runtime/Engine/Public/AudioDevice.h
+++ b/Engine/Source/Runtime/Engine/Public/AudioDevice.h
@@ -1038,6 +1038,11 @@ public:
 	*/
 	bool GetDistanceSquaredToNearestListener(const FVector& Location, float& OutSqDistance) const;
 		
+	/**
+	* Returns the global maximum distance used in the audio engine.
+	*/
+	static float GetMaxWorldDistance();
+
 	/**
 	* Returns a position from the appropriate listener representation, depending on calling thread.
 	*
