Fix Audio Attenuation Visualization Bug

diff --git a/Engine/Source/Runtime/Engine/Classes/Sound/SoundBase.h b/Engine/Source/Runtime/Engine/Classes/Sound/SoundBase.h
--- a/Engine/Source/Runtime/Engine/Classes/Sound/SoundBase.h
+++ b/Engine/Source/Runtime/Engine/Classes/Sound/SoundBase.h
@@ -74,7 +74,7 @@ public:
 	UPROPERTY(EditAnywhere, BlueprintReadOnly, Category = Sound, meta = (DisplayName = "Class"), AssetRegistrySearchable)
 	TObjectPtr<USoundClass> SoundClassObject;
 
-	/** When "au.debug.Sounds -debug" has been specified, draw this sound's attenuation shape when the sound is audible. For debugging purpose only. */
+	/** When "au.3dVisualize.Attenuation" has been specified, draw this sound's attenuation shape when the sound is audible. For debugging purpose only. */
 	UPROPERTY(EditAnywhere, Category = Developer)
 	uint8 bDebug : 1;
 
diff --git a/Engine/Source/Runtime/Engine/Private/AudioDebug.cpp b/Engine/Source/Runtime/Engine/Private/AudioDebug.cpp
--- a/Engine/Source/Runtime/Engine/Private/AudioDebug.cpp
+++ b/Engine/Source/Runtime/Engine/Private/AudioDebug.cpp
@@ -723,6 +723,29 @@ static FAutoConsoleCommandWithWorldAndArgs GAudioDebugSounds
 	}), ECVF_Cheat
 );
 
+static bool bAttenuationVisualizeEnabledCVar = false;
+FAutoConsoleCommandWithWorldAndArgs CVarAudioVisualizeAttenuation(
+	TEXT("au.3dVisualize.Attenuation"),
+	TEXT("Whether or not attenuation spheres are visible when 3d visualize is enabled. \n")
+	TEXT("0: Not Enabled, 1: Enabled"),
+	FConsoleCommandWithWorldAndArgsDelegate::CreateStatic([](const TArray<FString>& Args, UWorld* InWorld)
+	{
+		if (Args.Num() < 0)
+		{
+			return;
+		}
+
+		bAttenuationVisualizeEnabledCVar = Args[0].ToBool();
+
+		// Internally, the visualization code requires au.Debug.Sounds to be enabled,
+		// so we force it on whenever this one is enabled.
+		if (bAttenuationVisualizeEnabledCVar)
+		{
+			Audio::DebugSoundObject(Args, InWorld, Audio::DebugStatNames::Sounds, Audio::bDebugSoundsForAllViewsEnabled);
+		}
+	}),
+	ECVF_Default);
+
 static FAutoConsoleCommandWithWorldAndArgs GAudioDebugSoundWaves
 (
 	TEXT("au.Debug.SoundWaves"),
@@ -1832,7 +1855,9 @@ namespace Audio
 		Canvas->DrawShadowedString(X, Y, TEXT("Active Sounds:"), GetStatsFont(), HeaderColor);
 		Y += FontHeight;
 
-		const FString InfoText = FString::Printf(TEXT(" Sorting By: %s, Visualize Attenuation: %s"), *SortingName, bDebug ? TEXT("Enabled") : TEXT("Disabled"));
+		FAudioDeviceManager* DeviceManager = GEngine->GetAudioDeviceManager();
+		const FString InfoText = FString::Printf(TEXT(" Sorting By: %s, 3D Visualization: %s"), *SortingName, 
+			DeviceManager->IsVisualizeDebug3dEnabled() || bAttenuationVisualizeEnabledCVar ? TEXT("Enabled") : TEXT("Disabled"));
 		Canvas->DrawShadowedString(X, Y, *InfoText, GetStatsFont(), FColor(128, 255, 128));
 		Y += FontHeight;
 
@@ -1951,12 +1976,7 @@ namespace Audio
 
 		Y += FontHeight;
 
-		if (!bDebug)
-		{
-			return Y;
-		}
-
-		// Draw sound cue's sphere only in debug.
+		// Draw attenuation shape if enabled.
 		for (const FAudioStats::FStatSoundInfo& StatSoundInfo : AudioStats.StatSoundInfos)
 		{
 			const FTransform& SoundTransform = StatSoundInfo.Transform;
@@ -2329,7 +2349,7 @@ namespace Audio
 						StatSoundInfo.Transform = ActiveSound->Transform;
 						StatSoundInfo.AudioComponentID = ActiveSound->GetAudioComponentID();
 
-						if (bDebug && ActiveSound->GetSound()->bDebug)
+						if (bAttenuationVisualizeEnabledCVar && ActiveSound->GetSound()->bDebug)
 						{
 							ActiveSound->CollectAttenuationShapesForVisualization(StatSoundInfo.ShapeDetailsMap);
 						}
