Streaming generation support for ContentBundles.
Cells can now be flagged as part of a content bundle
Cells from content bundles cannot be always loaded
Cell Name now also contain the contentbundle UID
Add a OnPreStreamingGeneration callback external systems can use to perform actions before the WorldPartition Streaming Generation occurs.

#test asteria_terrain, qagame, pie, cooked game, worldpartition changelist validation

diff --git a/Engine/Source/Runtime/Engine/Private/WorldPartition/RuntimeSpatialHash/RuntimeSpatialHashGridHelper.h b/Engine/Source/Runtime/Engine/Private/WorldPartition/RuntimeSpatialHash/RuntimeSpatialHashGridHelper.h
--- a/Engine/Source/Runtime/Engine/Private/WorldPartition/RuntimeSpatialHash/RuntimeSpatialHashGridHelper.h
+++ b/Engine/Source/Runtime/Engine/Private/WorldPartition/RuntimeSpatialHash/RuntimeSpatialHashGridHelper.h
@@ -261,10 +261,11 @@ struct FSquare2DGridHelper
 #if WITH_EDITOR
 		struct FGridCellDataChunk
 		{
-			FGridCellDataChunk(const TArray<const UDataLayerInstance*>& InDataLayers)
+			FGridCellDataChunk(const TArray<const UDataLayerInstance*>& InDataLayers, const FGuid& InContentBundleID)
 			{
 				Algo::TransformIf(InDataLayers, DataLayers, [](const UDataLayerInstance* DataLayer) { return DataLayer->IsRuntime(); }, [](const UDataLayerInstance* DataLayer) { return DataLayer; });
 				DataLayersID = FDataLayersID(DataLayers);
+				ContentBundleID = InContentBundleID;
 			}
 
 			void AddActorSetInstance(const IStreamingGenerationContext::FActorSetInstance* ActorSetInstance) { ActorSetInstances.Add(ActorSetInstance); }
@@ -272,6 +273,7 @@ struct FSquare2DGridHelper
 			bool HasDataLayers() const { return !DataLayers.IsEmpty(); }
 			const TArray<const UDataLayerInstance*>& GetDataLayers() const { return DataLayers; }
 			const FDataLayersID& GetDataLayersID() const { return DataLayersID; }
+			FGuid GetContentBundleID() const { return ContentBundleID; }
 			bool operator==(const FGridCellDataChunk& InGridCellDataChunk) const { return DataLayersID == InGridCellDataChunk.DataLayersID;}
 			friend uint32 GetTypeHash(const FGridCellDataChunk& InGridCellDataChunk) { return GetTypeHash(InGridCellDataChunk.DataLayersID);}
 
@@ -279,6 +281,7 @@ struct FSquare2DGridHelper
 			TArray<const IStreamingGenerationContext::FActorSetInstance*> ActorSetInstances;
 			TArray<const UDataLayerInstance*> DataLayers;
 			FDataLayersID DataLayersID;
+			FGuid ContentBundleID;
 		};
 
 		struct FGridCell
@@ -290,7 +293,7 @@ struct FSquare2DGridHelper
 			void AddActorSetInstance(const IStreamingGenerationContext::FActorSetInstance* ActorSetInstance)
 			{
 				const FDataLayersID DataLayersID = FDataLayersID(ActorSetInstance->DataLayers);
-				FGridCellDataChunk& ActorDataChunk = DataChunks.FindOrAddByHash(DataLayersID.GetHash(), FGridCellDataChunk(ActorSetInstance->DataLayers));
+				FGridCellDataChunk& ActorDataChunk = DataChunks.FindOrAddByHash(DataLayersID.GetHash(), FGridCellDataChunk(ActorSetInstance->DataLayers, ActorSetInstance->ContentBundleID));
 				ActorDataChunk.AddActorSetInstance(ActorSetInstance);
 			}
 
diff --git a/Engine/Source/Runtime/Engine/Private/WorldPartition/RuntimeSpatialHash/RuntimeSpatialHashHLOD.cpp b/Engine/Source/Runtime/Engine/Private/WorldPartition/RuntimeSpatialHash/RuntimeSpatialHashHLOD.cpp
--- a/Engine/Source/Runtime/Engine/Private/WorldPartition/RuntimeSpatialHash/RuntimeSpatialHashHLOD.cpp
+++ b/Engine/Source/Runtime/Engine/Private/WorldPartition/RuntimeSpatialHash/RuntimeSpatialHashHLOD.cpp
@@ -172,7 +172,9 @@ static TArray<FGuid> GenerateHLODsForGrid(UWorldPartition* WorldPartition, const
 			if (bShouldGenerateHLODs)
 			{
 				SlowTask.EnterProgressFrame(1);
-				FName CellName = UWorldPartitionRuntimeSpatialHash::GetCellName(WorldPartition, RuntimeGrid.GridName, CellGlobalCoord, GridCellDataChunk.GetDataLayersID());
+				// todo_ow : Generating hlods is not yet supported for content bundles. Pass an invalid UID to specified the cell is not coming from ContentBundles.
+				FGuid ContentBundleInvalidUID;
+				FName CellName = UWorldPartitionRuntimeSpatialHash::GetCellName(WorldPartition, RuntimeGrid.GridName, CellGlobalCoord, GridCellDataChunk.GetDataLayersID(), ContentBundleInvalidUID);
 
 				UE_LOG(LogWorldPartitionRuntimeSpatialHashHLOD, Verbose, TEXT("Creating HLOD for cell %s at %s..."), *CellName.ToString(), *CellCoord.ToString());
 
diff --git a/Engine/Source/Runtime/Engine/Private/WorldPartition/WorldPartitionActorDesc.cpp b/Engine/Source/Runtime/Engine/Private/WorldPartition/WorldPartitionActorDesc.cpp
--- a/Engine/Source/Runtime/Engine/Private/WorldPartition/WorldPartitionActorDesc.cpp
+++ b/Engine/Source/Runtime/Engine/Private/WorldPartition/WorldPartitionActorDesc.cpp
@@ -454,6 +454,26 @@ FName FWorldPartitionActorDesc::GetDisplayClassName() const
 	return BaseClass.IsNone() ? GetCleanClassName(NativeClass) : GetCleanClassName(BaseClass);
 }
 
+FGuid FWorldPartitionActorDesc::GetContentBundleUID() const
+{
+	constexpr TCHAR ContentBundleFolder[] = TEXT("/ContentBundle/");
+	FString ActorPackageString = ActorPackage.ToString();
+	int FoundIdx = ActorPackageString.Find(ContentBundleFolder, ESearchCase::IgnoreCase);
+	if (FoundIdx != INDEX_NONE)
+	{
+		uint32 StartUIDIdx = FoundIdx + FCString::Strlen(ContentBundleFolder);
+		uint32 EndUIDIdx = ActorPackageString.Find(TEXT("/"), ESearchCase::IgnoreCase, ESearchDir::FromStart, StartUIDIdx);
+		ActorPackageString.MidInline(StartUIDIdx, EndUIDIdx - StartUIDIdx);
+		
+		FGuid ContentBundleUID;
+		FGuid::Parse(ActorPackageString, ContentBundleUID);
+
+		return ContentBundleUID;
+	}
+
+	return FGuid();
+}
+
 void FWorldPartitionActorDesc::CheckForErrors(IStreamingGenerationErrorHandler* ErrorHandler) const
 {
 	if (IsResaveNeeded())
diff --git a/Engine/Source/Runtime/Engine/Private/WorldPartition/WorldPartitionActorDescView.cpp b/Engine/Source/Runtime/Engine/Private/WorldPartition/WorldPartitionActorDescView.cpp
--- a/Engine/Source/Runtime/Engine/Private/WorldPartition/WorldPartitionActorDescView.cpp
+++ b/Engine/Source/Runtime/Engine/Private/WorldPartition/WorldPartitionActorDescView.cpp
@@ -140,6 +140,11 @@ const FGuid& FWorldPartitionActorDescView::GetFolderGuid() const
 	return ActorDesc->GetFolderGuid();
 }
 
+FGuid FWorldPartitionActorDescView::GetContentBundleUID() const
+{
+	return ActorDesc->GetContentBundleUID();
+}
+
 bool FWorldPartitionActorDescView::IsContainerInstance() const
 {
 	return ActorDesc->IsContainerInstance();
diff --git a/Engine/Source/Runtime/Engine/Private/WorldPartition/WorldPartitionRuntimeSpatialHash.cpp b/Engine/Source/Runtime/Engine/Private/WorldPartition/WorldPartitionRuntimeSpatialHash.cpp
--- a/Engine/Source/Runtime/Engine/Private/WorldPartition/WorldPartitionRuntimeSpatialHash.cpp
+++ b/Engine/Source/Runtime/Engine/Private/WorldPartition/WorldPartitionRuntimeSpatialHash.cpp
@@ -363,7 +363,7 @@ void FSpatialHashStreamingGrid::GetAlwaysLoadedCells(const UDataLayerSubsystem*
 			{
 				if (!Cell->HasDataLayers() || (DataLayerSubsystem && DataLayerSubsystem->IsAnyDataLayerInEffectiveRuntimeState(Cell->GetDataLayers(), EDataLayerRuntimeState::Activated)))
 				{
-					check(Cell->IsAlwaysLoaded() || Cell->HasDataLayers());
+					check(Cell->IsAlwaysLoaded() || Cell->HasDataLayers() || Cell->GetContentBundleID().IsValid());
 					OutActivateCells.Add(Cell);
 				}
 				else if(DataLayerSubsystem && DataLayerSubsystem->IsAnyDataLayerInEffectiveRuntimeState(Cell->GetDataLayers(), EDataLayerRuntimeState::Loaded))
@@ -959,12 +959,19 @@ void UWorldPartitionRuntimeSpatialHash::DumpStateLog(FHierarchicalLogArchive& Ar
 	});
 }
 
-FName UWorldPartitionRuntimeSpatialHash::GetCellName(UWorldPartition* WorldPartition, FName InGridName, const FGridCellCoord& InCellGlobalCoord, const FDataLayersID& InDataLayerID)
+FName UWorldPartitionRuntimeSpatialHash::GetCellName(UWorldPartition* WorldPartition, FName InGridName, const FGridCellCoord& InCellGlobalCoord, const FDataLayersID& InDataLayerID, const FGuid& InContentBundleID)
 {
 	const FString PackageName = FPackageName::GetShortName(WorldPartition->GetPackage());
 	const FString PackageNameNoPIEPrefix = UWorld::RemovePIEPrefix(PackageName);
 
-	return FName(*FString::Printf(TEXT("%s_%s_%s_DL%X"), *PackageNameNoPIEPrefix, *InGridName.ToString(), *GetCellCoordString(InCellGlobalCoord), InDataLayerID.GetHash()));
+	FString CellName = FString::Printf(TEXT("%s_%s_%s_DL%X"), *PackageNameNoPIEPrefix, *InGridName.ToString(), *GetCellCoordString(InCellGlobalCoord), InDataLayerID.GetHash());
+
+	if (InContentBundleID.IsValid())
+	{
+		CellName += TEXT("_CB") + InContentBundleID.ToString(EGuidFormats::Short);
+	}
+
+	return FName(*CellName);
 }
 
 bool UWorldPartitionRuntimeSpatialHash::GetPreviewGrids() const
@@ -978,10 +985,10 @@ void UWorldPartitionRuntimeSpatialHash::SetPreviewGrids(bool bInPreviewGrids)
 	bPreviewGrids = bInPreviewGrids;
 }
 
-FName UWorldPartitionRuntimeSpatialHash::GetCellName(FName InGridName, const FGridCellCoord& InCellGlobalCoord, const FDataLayersID& InDataLayerID) const
+FName UWorldPartitionRuntimeSpatialHash::GetCellName(FName InGridName, const FGridCellCoord& InCellGlobalCoord, const FDataLayersID& InDataLayerID, const FGuid& InContentBundleID) const
 {
 	UWorldPartition* WorldPartition = GetOuterUWorldPartition();
-	return UWorldPartitionRuntimeSpatialHash::GetCellName(WorldPartition, InGridName, InCellGlobalCoord, InDataLayerID);
+	return UWorldPartitionRuntimeSpatialHash::GetCellName(WorldPartition, InGridName, InCellGlobalCoord, InDataLayerID, InContentBundleID);
 }
 
 bool UWorldPartitionRuntimeSpatialHash::CreateStreamingGrid(const FSpatialHashRuntimeGrid& RuntimeGrid, const FSquare2DGridHelper& PartionedActors, UWorldPartitionStreamingPolicy* StreamingPolicy, TArray<FString>* OutPackagesToGenerate)
@@ -1024,7 +1031,7 @@ bool UWorldPartitionRuntimeSpatialHash::CreateStreamingGrid(const FSpatialHashRu
 			for (const FSquare2DGridHelper::FGridLevel::FGridCellDataChunk& GridCellDataChunk : TempCell.GetDataChunks())
 			{
 				// Cell cannot be treated as always loaded if it has data layers
-				const bool bIsCellAlwaysLoaded = (&TempCell == &PartionedActors.GetAlwaysLoadedCell()) && !GridCellDataChunk.HasDataLayers();
+				const bool bIsCellAlwaysLoaded = (&TempCell == &PartionedActors.GetAlwaysLoadedCell()) && !GridCellDataChunk.HasDataLayers() && !GridCellDataChunk.GetContentBundleID().IsValid();
 				
 				TArray<IStreamingGenerationContext::FActorInstance> FilteredActors;
 				
@@ -1078,7 +1085,7 @@ bool UWorldPartitionRuntimeSpatialHash::CreateStreamingGrid(const FSpatialHashRu
 				
 				FGridCellCoord CellGlobalCoords;
 				verify(PartionedActors.GetCellGlobalCoords(FGridCellCoord(CellCoordX, CellCoordY, Level), CellGlobalCoords));
-				FName CellName = GetCellName(CurrentStreamingGrid.GridName, CellGlobalCoords, GridCellDataChunk.GetDataLayersID());
+				FName CellName = GetCellName(CurrentStreamingGrid.GridName, CellGlobalCoords, GridCellDataChunk.GetDataLayersID(), GridCellDataChunk.GetContentBundleID());
 
 				UWorldPartitionRuntimeSpatialHashCell* StreamingCell = NewObject<UWorldPartitionRuntimeSpatialHashCell>(WorldPartition, StreamingPolicy->GetRuntimeCellClass(), CellName);
 
@@ -1097,6 +1104,7 @@ bool UWorldPartitionRuntimeSpatialHash::CreateStreamingGrid(const FSpatialHashRu
 				GridLevel.LayerCells[LayerCellIndex].GridCells.Add(StreamingCell);
 				StreamingCell->SetIsAlwaysLoaded(bIsCellAlwaysLoaded);
 				StreamingCell->SetDataLayers(GridCellDataChunk.GetDataLayers());
+				StreamingCell->SetContentBundleUID(GridCellDataChunk.GetContentBundleID());
 				StreamingCell->Level = Level;
 				StreamingCell->SetPriority(RuntimeGrid.Priority);
 				FBox2D Bounds;
diff --git a/Engine/Source/Runtime/Engine/Private/WorldPartition/WorldPartitionStreamingGeneration.cpp b/Engine/Source/Runtime/Engine/Private/WorldPartition/WorldPartitionStreamingGeneration.cpp
--- a/Engine/Source/Runtime/Engine/Private/WorldPartition/WorldPartitionStreamingGeneration.cpp
+++ b/Engine/Source/Runtime/Engine/Private/WorldPartition/WorldPartitionStreamingGeneration.cpp
@@ -117,6 +117,7 @@ class FWorldPartitionStreamingGenerator
 					ActorSetInstance.ActorSet = &ActorSet;
 					ActorSetInstance.ContainerID = ContainerInstanceDescriptor.ID;
 					ActorSetInstance.Transform = ContainerInstanceDescriptor.Transform;
+					ActorSetInstance.ContentBundleID = ReferenceActorDescView.GetContentBundleUID();
 
 					if (ContainerInstanceDescriptor.ID.IsMainContainer())
 					{
@@ -863,6 +864,8 @@ private:
 
 bool UWorldPartition::GenerateStreaming(TArray<FString>* OutPackagesToGenerate)
 {
+	OnPreGenerateStreaming.Broadcast();
+
 	return GenerateContainerStreaming(ActorDescContainer, OutPackagesToGenerate);
 }
 
diff --git a/Engine/Source/Runtime/Engine/Public/WorldPartition/WorldPartition.h b/Engine/Source/Runtime/Engine/Public/WorldPartition/WorldPartition.h
--- a/Engine/Source/Runtime/Engine/Public/WorldPartition/WorldPartition.h
+++ b/Engine/Source/Runtime/Engine/Public/WorldPartition/WorldPartition.h
@@ -155,6 +155,9 @@ public:
 	bool GenerateContainerStreaming(const UActorDescContainer* ActorDescContainer, TArray<FString>* OutPackagesToGenerate = nullptr);
 	void FlushStreaming();
 
+	DECLARE_MULTICAST_DELEGATE(FWorldPartitionGenerateStreamingDelegate);
+	FWorldPartitionGenerateStreamingDelegate OnPreGenerateStreaming;
+
 	void RemapSoftObjectPath(FSoftObjectPath& ObjectPath);
 
 	// Cooking
diff --git a/Engine/Source/Runtime/Engine/Public/WorldPartition/WorldPartitionActorDesc.h b/Engine/Source/Runtime/Engine/Public/WorldPartition/WorldPartitionActorDesc.h
--- a/Engine/Source/Runtime/Engine/Public/WorldPartition/WorldPartitionActorDesc.h
+++ b/Engine/Source/Runtime/Engine/Public/WorldPartition/WorldPartitionActorDesc.h
@@ -114,6 +114,8 @@ public:
 	virtual FName GetLevelPackage() const { return NAME_None; }
 	virtual bool GetContainerInstance(const UActorDescContainer*& OutLevelContainer, FTransform& OutLevelTransform, EContainerClusterMode& OutClusterMode) const { return false; }
 
+	FGuid GetContentBundleUID() const;
+
 	virtual const FGuid& GetSceneOutlinerParent() const { return GetParentActor(); }
 	virtual bool IsResaveNeeded() const { return false; }
 	virtual bool IsRuntimeRelevant(const FActorContainerID& InContainerID) const { return true; }
diff --git a/Engine/Source/Runtime/Engine/Public/WorldPartition/WorldPartitionActorDescView.h b/Engine/Source/Runtime/Engine/Public/WorldPartition/WorldPartitionActorDescView.h
--- a/Engine/Source/Runtime/Engine/Public/WorldPartition/WorldPartitionActorDescView.h
+++ b/Engine/Source/Runtime/Engine/Public/WorldPartition/WorldPartitionActorDescView.h
@@ -48,6 +48,8 @@ public:
 	FName GetActorName() const;
 	const FGuid& GetFolderGuid() const;
 
+	FGuid GetContentBundleUID() const;
+
 	bool IsContainerInstance() const;
 	bool GetContainerInstance(const UActorDescContainer*& OutLevelContainer, FTransform& OutLevelTransform, EContainerClusterMode& OutClusterMode) const;
 
diff --git a/Engine/Source/Runtime/Engine/Public/WorldPartition/WorldPartitionRuntimeCell.h b/Engine/Source/Runtime/Engine/Public/WorldPartition/WorldPartitionRuntimeCell.h
--- a/Engine/Source/Runtime/Engine/Public/WorldPartition/WorldPartitionRuntimeCell.h
+++ b/Engine/Source/Runtime/Engine/Public/WorldPartition/WorldPartitionRuntimeCell.h
@@ -174,6 +174,7 @@ class ENGINE_API UWorldPartitionRuntimeCell : public UObject, public IWorldParti
 	virtual bool IsDebugShown() const;
 	virtual int32 SortCompare(const UWorldPartitionRuntimeCell* Other) const;
 	virtual FName GetGridName() const { return DebugInfo.GridName; }
+	virtual FGuid const& GetContentBundleID() const { return ContentBundleID; }
 	/** Caches information on streaming source that will be used later on to sort cell. Returns true if cache was reset, else returns false. */
 	virtual bool CacheStreamingSourceInfo(const UWorldPartitionRuntimeCell::FStreamingSourceInfo& Info) const;
 
@@ -202,6 +203,7 @@ class ENGINE_API UWorldPartitionRuntimeCell : public UObject, public IWorldParti
 	bool GetClientOnlyVisible() const { return bClientOnlyVisible; }
 
 	void SetDataLayers(const TArray<const UDataLayerInstance*>& InDataLayerInstances);
+	void SetContentBundleUID(const FGuid& InContentBundleID) { ContentBundleID = InContentBundleID; }
 	void SetMinMaxZ(const FVector2D& InMinMaxZ);
 	void SetDebugInfo(int64 InCoordX, int64 InCoordY, int64 InCoordZ, FName InGridName);
 	void SetGridName(FName InGridName);
@@ -259,6 +261,9 @@ private:
 	UPROPERTY()
 	bool bBlockOnSlowLoading;
 
+	UPROPERTY()
+	FGuid ContentBundleID;
+
 protected:
 	// Source Priority
 	mutable uint8 CachedMinSourcePriority;
diff --git a/Engine/Source/Runtime/Engine/Public/WorldPartition/WorldPartitionRuntimeSpatialHash.h b/Engine/Source/Runtime/Engine/Public/WorldPartition/WorldPartitionRuntimeSpatialHash.h
--- a/Engine/Source/Runtime/Engine/Public/WorldPartition/WorldPartitionRuntimeSpatialHash.h
+++ b/Engine/Source/Runtime/Engine/Public/WorldPartition/WorldPartitionRuntimeSpatialHash.h
@@ -250,8 +250,8 @@ public:
 
 	virtual URuntimeHashExternalStreamingObjectBase* StoreToExternalStreamingObject(UObject* StreamingObjectOuter, FName StreamingObjectName) override;
 
-	FName GetCellName(FName InGridName, const FGridCellCoord& InCellGlobalCoord, const FDataLayersID& InDataLayerID) const;
-	static FName GetCellName(UWorldPartition* WorldPartition, FName InGridName, const FGridCellCoord& InCellGlobalCoord, const FDataLayersID& InDataLayerID);
+	FName GetCellName(FName InGridName, const FGridCellCoord& InCellGlobalCoord, const FDataLayersID& InDataLayerID, const FGuid& InContentBundleID) const;
+	static FName GetCellName(UWorldPartition* WorldPartition, FName InGridName, const FGridCellCoord& InCellGlobalCoord, const FDataLayersID& InDataLayerID, const FGuid& InContentBundleID);
 
 	bool GetPreviewGrids() const;
 	void SetPreviewGrids(bool bInPreviewGrids);
diff --git a/Engine/Source/Runtime/Engine/Public/WorldPartition/WorldPartitionStreamingGenerationContext.h b/Engine/Source/Runtime/Engine/Public/WorldPartition/WorldPartitionStreamingGenerationContext.h
--- a/Engine/Source/Runtime/Engine/Public/WorldPartition/WorldPartitionStreamingGenerationContext.h
+++ b/Engine/Source/Runtime/Engine/Public/WorldPartition/WorldPartitionStreamingGenerationContext.h
@@ -29,6 +29,7 @@ public:
 		FName RuntimeGrid;
 		bool bIsSpatiallyLoaded;
 		TArray<const UDataLayerInstance*> DataLayers;
+		FGuid ContentBundleID;
 		const FActorSetContainer* ContainerInstance;
 		FActorContainerID ContainerID;
 		FTransform Transform;
