Add blame cvar for Blueprint PrintString spam

diff --git a/Engine/Source/Runtime/Engine/Private/KismetSystemLibrary.cpp b/Engine/Source/Runtime/Engine/Private/KismetSystemLibrary.cpp
--- a/Engine/Source/Runtime/Engine/Private/KismetSystemLibrary.cpp
+++ b/Engine/Source/Runtime/Engine/Private/KismetSystemLibrary.cpp
@@ -45,23 +45,31 @@
 
 #define LOCTEXT_NAMESPACE "UKismetSystemLibrary"
 
-const FName PropertyGetFailedWarning = FName("PropertyGetFailedWarning");
-const FName PropertySetFailedWarning = FName("PropertySetFailedWarning");
+namespace UE::Blueprint::Private
+{
+	const FName PropertyGetFailedWarning = FName("PropertyGetFailedWarning");
+	const FName PropertySetFailedWarning = FName("PropertySetFailedWarning");
+
+	bool bBlamePrintString = false;
+	FAutoConsoleVariableRef CVarBlamePrintString(TEXT("bp.BlamePrintString"), 
+		bBlamePrintString,
+		TEXT("When true, prints the Blueprint Asset and Function that generated calls to Print String. Useful for tracking down screen message spam."));
+}
 
 UKismetSystemLibrary::UKismetSystemLibrary(const FObjectInitializer& ObjectInitializer)
 	: Super(ObjectInitializer)
 {
 	FBlueprintSupport::RegisterBlueprintWarning(
 		FBlueprintWarningDeclaration(
-			PropertyGetFailedWarning,
-			LOCTEXT("PropertyGetFailedWarning", "Property Get Failed")
+			UE::Blueprint::Private::PropertyGetFailedWarning,
+			LOCTEXT("UE::Blueprint::Private::PropertyGetFailedWarning", "Property Get Failed")
 		)
 	);
 
 	FBlueprintSupport::RegisterBlueprintWarning(
 		FBlueprintWarningDeclaration(
-			PropertySetFailedWarning,
-			LOCTEXT("PropertySetFailedWarning", "Property Set Failed")
+			UE::Blueprint::Private::PropertySetFailedWarning,
+			LOCTEXT("UE::Blueprint::Private::PropertySetFailedWarning", "Property Set Failed")
 		)
 	);
 }
@@ -318,7 +326,16 @@ void UKismetSystemLibrary::PrintString(const UObject* WorldContextObject, const
 			}
 		}
 	}
-	
+
+	if (UE::Blueprint::Private::bBlamePrintString && !FBlueprintContextTracker::Get().GetCurrentScriptStack().IsEmpty())
+	{
+		const TArrayView<const FFrame* const> ScriptStack = FBlueprintContextTracker::Get().GetCurrentScriptStack();
+		Prefix = FString::Printf(TEXT("Blueprint Object: %s\nBlueprint Function: %s\n%s"), 
+			*ScriptStack.Last()->Node->GetPackage()->GetPathName(),
+			*ScriptStack.Last()->Node->GetName(),
+			*Prefix);
+	}
+
 	const FString FinalDisplayString = Prefix + InString;
 	FString FinalLogString = FinalDisplayString;
 
@@ -2896,17 +2913,17 @@ bool UKismetSystemLibrary::Generic_GetEditorProperty(const UObject* Object, cons
 	{
 		if (EnumHasAnyFlags(AccessResult, EPropertyAccessResultFlags::AccessProtected))
 		{
-			FFrame::KismetExecutionMessage(*FString::Printf(TEXT("Property '%s' on '%s' (%s) is protected and cannot be read"), *ObjectProp->GetName(), *Object->GetPathName(), *Object->GetClass()->GetName()), ELogVerbosity::Warning, PropertyGetFailedWarning);
+			FFrame::KismetExecutionMessage(*FString::Printf(TEXT("Property '%s' on '%s' (%s) is protected and cannot be read"), *ObjectProp->GetName(), *Object->GetPathName(), *Object->GetClass()->GetName()), ELogVerbosity::Warning, UE::Blueprint::Private::PropertyGetFailedWarning);
 			return false;
 		}
 
-		FFrame::KismetExecutionMessage(*FString::Printf(TEXT("Property '%s' on '%s' (%s) cannot be read"), *ObjectProp->GetName(), *Object->GetPathName(), *Object->GetClass()->GetName()), ELogVerbosity::Warning, PropertyGetFailedWarning);
+		FFrame::KismetExecutionMessage(*FString::Printf(TEXT("Property '%s' on '%s' (%s) cannot be read"), *ObjectProp->GetName(), *Object->GetPathName(), *Object->GetClass()->GetName()), ELogVerbosity::Warning, UE::Blueprint::Private::PropertyGetFailedWarning);
 		return false;
 	}
 
 	if (EnumHasAnyFlags(AccessResult, EPropertyAccessResultFlags::ConversionFailed))
 	{
-		FFrame::KismetExecutionMessage(*FString::Printf(TEXT("Property '%s' (%s) on '%s' (%s) tried to get to a property value of the incorrect type (%s)"), *ObjectProp->GetName(), *ObjectProp->GetClass()->GetName(), *Object->GetPathName(), *Object->GetClass()->GetName(), *ValueProp->GetClass()->GetName()), ELogVerbosity::Warning, PropertyGetFailedWarning);
+		FFrame::KismetExecutionMessage(*FString::Printf(TEXT("Property '%s' (%s) on '%s' (%s) tried to get to a property value of the incorrect type (%s)"), *ObjectProp->GetName(), *ObjectProp->GetClass()->GetName(), *Object->GetPathName(), *Object->GetClass()->GetName(), *ValueProp->GetClass()->GetName()), ELogVerbosity::Warning, UE::Blueprint::Private::PropertyGetFailedWarning);
 		return false;
 	}
 
@@ -2955,7 +2972,7 @@ DEFINE_FUNCTION(UKismetSystemLibrary::execGetEditorProperty)
 		}
 		else
 		{
-			FFrame::KismetExecutionMessage(*FString::Printf(TEXT("Property '%s' on '%s' (%s) was missing"), *PropertyName.ToString(), *Object->GetPathName(), *Object->GetClass()->GetName()), ELogVerbosity::Warning, PropertyGetFailedWarning);
+			FFrame::KismetExecutionMessage(*FString::Printf(TEXT("Property '%s' on '%s' (%s) was missing"), *PropertyName.ToString(), *Object->GetPathName(), *Object->GetClass()->GetName()), ELogVerbosity::Warning, UE::Blueprint::Private::PropertyGetFailedWarning);
 		}
 	}
 
@@ -2977,35 +2994,35 @@ bool UKismetSystemLibrary::Generic_SetEditorProperty(UObject* Object, const FPro
 	{
 		if (EnumHasAnyFlags(AccessResult, EPropertyAccessResultFlags::AccessProtected))
 		{
-			FFrame::KismetExecutionMessage(*FString::Printf(TEXT("Property '%s' on '%s' (%s) is protected and cannot be set"), *ObjectProp->GetName(), *Object->GetPathName(), *Object->GetClass()->GetName()), ELogVerbosity::Warning, PropertySetFailedWarning);
+			FFrame::KismetExecutionMessage(*FString::Printf(TEXT("Property '%s' on '%s' (%s) is protected and cannot be set"), *ObjectProp->GetName(), *Object->GetPathName(), *Object->GetClass()->GetName()), ELogVerbosity::Warning, UE::Blueprint::Private::PropertySetFailedWarning);
 			return false;
 		}
 
 		if (EnumHasAnyFlags(AccessResult, EPropertyAccessResultFlags::CannotEditTemplate))
 		{
-			FFrame::KismetExecutionMessage(*FString::Printf(TEXT("Property '%s' on '%s' (%s) cannot be edited on templates"), *ObjectProp->GetName(), *Object->GetPathName(), *Object->GetClass()->GetName()), ELogVerbosity::Warning, PropertySetFailedWarning);
+			FFrame::KismetExecutionMessage(*FString::Printf(TEXT("Property '%s' on '%s' (%s) cannot be edited on templates"), *ObjectProp->GetName(), *Object->GetPathName(), *Object->GetClass()->GetName()), ELogVerbosity::Warning, UE::Blueprint::Private::PropertySetFailedWarning);
 			return false;
 		}
 
 		if (EnumHasAnyFlags(AccessResult, EPropertyAccessResultFlags::CannotEditInstance))
 		{
-			FFrame::KismetExecutionMessage(*FString::Printf(TEXT("Property '%s' on '%s' (%s) cannot be edited on instances"), *ObjectProp->GetName(), *Object->GetPathName(), *Object->GetClass()->GetName()), ELogVerbosity::Warning, PropertySetFailedWarning);
+			FFrame::KismetExecutionMessage(*FString::Printf(TEXT("Property '%s' on '%s' (%s) cannot be edited on instances"), *ObjectProp->GetName(), *Object->GetPathName(), *Object->GetClass()->GetName()), ELogVerbosity::Warning, UE::Blueprint::Private::PropertySetFailedWarning);
 			return false;
 		}
 
 		if (EnumHasAnyFlags(AccessResult, EPropertyAccessResultFlags::ReadOnly))
 		{
-			FFrame::KismetExecutionMessage(*FString::Printf(TEXT("Property '%s' on '%s' (%s) is read-only and cannot be set"), *ObjectProp->GetName(), *Object->GetPathName(), *Object->GetClass()->GetName()), ELogVerbosity::Warning, PropertySetFailedWarning);
+			FFrame::KismetExecutionMessage(*FString::Printf(TEXT("Property '%s' on '%s' (%s) is read-only and cannot be set"), *ObjectProp->GetName(), *Object->GetPathName(), *Object->GetClass()->GetName()), ELogVerbosity::Warning, UE::Blueprint::Private::PropertySetFailedWarning);
 			return false;
 		}
 
-		FFrame::KismetExecutionMessage(*FString::Printf(TEXT("Property '%s' on '%s' (%s) cannot be set"), *ObjectProp->GetName(), *Object->GetPathName(), *Object->GetClass()->GetName()), ELogVerbosity::Warning, PropertySetFailedWarning);
+		FFrame::KismetExecutionMessage(*FString::Printf(TEXT("Property '%s' on '%s' (%s) cannot be set"), *ObjectProp->GetName(), *Object->GetPathName(), *Object->GetClass()->GetName()), ELogVerbosity::Warning, UE::Blueprint::Private::PropertySetFailedWarning);
 		return false;
 	}
 
 	if (EnumHasAnyFlags(AccessResult, EPropertyAccessResultFlags::ConversionFailed))
 	{
-		FFrame::KismetExecutionMessage(*FString::Printf(TEXT("Property '%s' (%s) on '%s' (%s) tried to set from a property value of the incorrect type (%s)"), *ObjectProp->GetName(), *ObjectProp->GetClass()->GetName(), *Object->GetPathName(), *Object->GetClass()->GetName(), *ValueProp->GetClass()->GetName()), ELogVerbosity::Warning, PropertySetFailedWarning);
+		FFrame::KismetExecutionMessage(*FString::Printf(TEXT("Property '%s' (%s) on '%s' (%s) tried to set from a property value of the incorrect type (%s)"), *ObjectProp->GetName(), *ObjectProp->GetClass()->GetName(), *Object->GetPathName(), *Object->GetClass()->GetName(), *ValueProp->GetClass()->GetName()), ELogVerbosity::Warning, UE::Blueprint::Private::PropertySetFailedWarning);
 		return false;
 	}
 
@@ -3056,7 +3073,7 @@ DEFINE_FUNCTION(UKismetSystemLibrary::execSetEditorProperty)
 		}
 		else
 		{
-			FFrame::KismetExecutionMessage(*FString::Printf(TEXT("Property '%s' on '%s' (%s) was missing"), *PropertyName.ToString(), *Object->GetPathName(), *Object->GetClass()->GetName()), ELogVerbosity::Warning, PropertySetFailedWarning);
+			FFrame::KismetExecutionMessage(*FString::Printf(TEXT("Property '%s' on '%s' (%s) was missing"), *PropertyName.ToString(), *Object->GetPathName(), *Object->GetClass()->GetName()), ELogVerbosity::Warning, UE::Blueprint::Private::PropertySetFailedWarning);
 		}
 	}
 
