World Partition Editor Optimizations:
- Implement a simple way to cull out octree levels when searching for intersecting actors by providing a minimum bounding box.
- Stop looking for intersecting actors when the current octree bounds is smaller than the current effective grid size.
- Changed the effective screen grid size to switch faster to higher cell sizes.

Editor tick time goes from ~750ms to <1ms in a test map with 660K actors.

#rnx

diff --git a/Engine/Source/Editor/WorldPartitionEditor/Private/WorldPartition/SWorldPartitionEditorGrid2D.cpp b/Engine/Source/Editor/WorldPartitionEditor/Private/WorldPartition/SWorldPartitionEditorGrid2D.cpp
--- a/Engine/Source/Editor/WorldPartitionEditor/Private/WorldPartition/SWorldPartitionEditorGrid2D.cpp
+++ b/Engine/Source/Editor/WorldPartitionEditor/Private/WorldPartition/SWorldPartitionEditorGrid2D.cpp
@@ -724,6 +724,9 @@ void SWorldPartitionEditorGrid2D::Tick(const FGeometry& AllottedGeometry, const
 		}
 	}
 
+	UWorldPartitionEditorHash::FForEachIntersectingActorParams ForEachIntersectingActorParams;
+	ForEachIntersectingActorParams.MinimumBox = FBox(FVector::ZeroVector, FVector(GetSelectionSnap()));
+
 	WorldPartition->EditorHash->ForEachIntersectingActor(ViewRectWorld, [&](FWorldPartitionActorDesc* ActorDesc)
 	{
 		if (bShowActors)
@@ -744,7 +747,7 @@ void SWorldPartitionEditorGrid2D::Tick(const FGeometry& AllottedGeometry, const
 				}
 			}
 		}
-	});
+	}, ForEachIntersectingActorParams);
 
 	// Also include transient actor loader adapters that might have been spawned by blutilities, etc. Since these actors can't be saved because they are transient,
 	// they will never get an actor descriptor so they will never appear in the world partition editor. Also include unsaved, newly created actors for convenience.
diff --git a/Engine/Source/Editor/WorldPartitionEditor/Private/WorldPartition/SWorldPartitionEditorGridSpatialHash.cpp b/Engine/Source/Editor/WorldPartitionEditor/Private/WorldPartition/SWorldPartitionEditorGridSpatialHash.cpp
--- a/Engine/Source/Editor/WorldPartitionEditor/Private/WorldPartition/SWorldPartitionEditorGridSpatialHash.cpp
+++ b/Engine/Source/Editor/WorldPartitionEditor/Private/WorldPartition/SWorldPartitionEditorGridSpatialHash.cpp
@@ -80,7 +80,7 @@ int32 SWorldPartitionEditorGridSpatialHash::GetSelectionSnap() const
 	for (const UWorldPartitionEditorSpatialHash::FCellNodeHashLevel& HashLevel : EditorSpatialHash->HashLevels)
 	{
 		const FVector2D CellScreenSize = WorldToScreen.TransformVector(FVector2D(EffectiveCellSize, EffectiveCellSize));
-		if (CellScreenSize.X > 32)
+		if (CellScreenSize.X > 64)
 		{
 			break;
 		}
diff --git a/Engine/Source/Runtime/Engine/Private/WorldPartition/LoaderAdapter/LoaderAdapterSpatial.cpp b/Engine/Source/Runtime/Engine/Private/WorldPartition/LoaderAdapter/LoaderAdapterSpatial.cpp
--- a/Engine/Source/Runtime/Engine/Private/WorldPartition/LoaderAdapter/LoaderAdapterSpatial.cpp
+++ b/Engine/Source/Runtime/Engine/Private/WorldPartition/LoaderAdapter/LoaderAdapterSpatial.cpp
@@ -16,6 +16,10 @@ void ILoaderAdapterSpatial::ForEachActor(TFunctionRef<void(const FWorldPartition
 {
 	if (UWorldPartition* WorldPartition = GetWorld()->GetWorldPartition())
 	{
+		UWorldPartitionEditorHash::FForEachIntersectingActorParams ForEachIntersectingActorParams;
+		ForEachIntersectingActorParams.bIncludeSpatiallyLoadedActors = bIncludeSpatiallyLoadedActors;
+		ForEachIntersectingActorParams.bIncludeNonSpatiallyLoadedActors = bIncludeNonSpatiallyLoadedActors;
+
 		WorldPartition->EditorHash->ForEachIntersectingActor(*GetBoundingBox(), [this, WorldPartition, &InOperation](FWorldPartitionActorDesc* ActorDesc)
 		{
 			if (Intersect(ActorDesc->GetBounds()))
@@ -23,7 +27,7 @@ void ILoaderAdapterSpatial::ForEachActor(TFunctionRef<void(const FWorldPartition
 				FWorldPartitionHandle ActorHandle(WorldPartition, ActorDesc->GetGuid());
 				InOperation(ActorHandle);
 			}
-		}, bIncludeSpatiallyLoadedActors, bIncludeNonSpatiallyLoadedActors);
+		}, ForEachIntersectingActorParams);
 	}
 }
 #endif
diff --git a/Engine/Source/Runtime/Engine/Private/WorldPartition/WorldPartitionEditorHash.cpp b/Engine/Source/Runtime/Engine/Private/WorldPartition/WorldPartitionEditorHash.cpp
--- a/Engine/Source/Runtime/Engine/Private/WorldPartition/WorldPartitionEditorHash.cpp
+++ b/Engine/Source/Runtime/Engine/Private/WorldPartition/WorldPartitionEditorHash.cpp
@@ -5,3 +5,13 @@
 UWorldPartitionEditorHash::UWorldPartitionEditorHash(const FObjectInitializer& ObjectInitializer)
 	: Super(ObjectInitializer)
 {}
+
+#if WITH_EDITOR
+int32 UWorldPartitionEditorHash::ForEachIntersectingActor(const FBox& Box, TFunctionRef<void(FWorldPartitionActorDesc*)> InOperation, bool bIncludeSpatiallyLoadedActors, bool bIncludeNonSpatiallyLoadedActors)
+{
+	UWorldPartitionEditorHash::FForEachIntersectingActorParams ForEachIntersectingActorParams;
+	ForEachIntersectingActorParams.bIncludeSpatiallyLoadedActors = bIncludeSpatiallyLoadedActors;
+	ForEachIntersectingActorParams.bIncludeNonSpatiallyLoadedActors = bIncludeNonSpatiallyLoadedActors;
+	return ForEachIntersectingActor(Box, InOperation, ForEachIntersectingActorParams);
+}
+#endif
\ No newline at end of file
diff --git a/Engine/Source/Runtime/Engine/Private/WorldPartition/WorldPartitionEditorSpatialHash.cpp b/Engine/Source/Runtime/Engine/Private/WorldPartition/WorldPartitionEditorSpatialHash.cpp
--- a/Engine/Source/Runtime/Engine/Private/WorldPartition/WorldPartitionEditorSpatialHash.cpp
+++ b/Engine/Source/Runtime/Engine/Private/WorldPartition/WorldPartitionEditorSpatialHash.cpp
@@ -268,11 +268,13 @@ void UWorldPartitionEditorSpatialHash::UnhashActor(FWorldPartitionHandle& InActo
 	}
 }
 
-int32 UWorldPartitionEditorSpatialHash::ForEachIntersectingActor(const FBox& Box, TFunctionRef<void(FWorldPartitionActorDesc*)> InOperation, bool bIncludeSpatiallyLoadedActors, bool bIncludeNonSpatiallyLoadedActors)
+int32 UWorldPartitionEditorSpatialHash::ForEachIntersectingActor(const FBox& Box, TFunctionRef<void(FWorldPartitionActorDesc*)> InOperation, const FForEachIntersectingActorParams& Params)
 {
 	TSet<FWorldPartitionActorDesc*> IntersectedActorDescs;
 
-	if (bIncludeSpatiallyLoadedActors)
+	const int MininumumLevel = Params.MinimumBox.IsSet() ? GetLevelForBox(*Params.MinimumBox) : 0;
+
+	if (Params.bIncludeSpatiallyLoadedActors)
 	{
 		ForEachIntersectingCell(Box, [&](FCell* Cell)
 		{
@@ -292,10 +294,10 @@ int32 UWorldPartitionEditorSpatialHash::ForEachIntersectingActor(const FBox& Box
 					}
 				}
 			}
-		});
+		}, MininumumLevel);
 	}
 
-	if (bIncludeNonSpatiallyLoadedActors)
+	if (Params.bIncludeNonSpatiallyLoadedActors)
 	{
 		for(FWorldPartitionHandle& ActorDesc: AlwaysLoadedCell->Actors)
 		{
@@ -318,7 +320,7 @@ int32 UWorldPartitionEditorSpatialHash::ForEachIntersectingActor(const FBox& Box
 	return IntersectedActorDescs.Num();
 }
 
-int32 UWorldPartitionEditorSpatialHash::ForEachIntersectingCellInner(const FBox& Box, const FCellCoord& CellCoord, TFunctionRef<void(FCell*)> InOperation)
+int32 UWorldPartitionEditorSpatialHash::ForEachIntersectingCellInner(const FBox& Box, const FCellCoord& CellCoord, TFunctionRef<void(FCell*)> InOperation, int32 MinimumLevel)
 {
 	int32 NumIntersecting = 0;
 
@@ -330,22 +332,25 @@ int32 UWorldPartitionEditorSpatialHash::ForEachIntersectingCellInner(const FBox&
 			NumIntersecting++;
 		}
 
-		CellNodeElement->Key.ForEachChild([&](uint32 ChildIndex)
+		if (MinimumLevel <= CellCoord.Level)
 		{
-			const FCellCoord ChildCellCoord = CellCoord.GetChildCellCoord(ChildIndex);
-			const FBox CellBounds = GetCellBounds(ChildCellCoord);
-
-			if (Box.Intersect(CellBounds))
+			CellNodeElement->Key.ForEachChild([&](uint32 ChildIndex)
 			{
-				NumIntersecting += ForEachIntersectingCellInner(Box, ChildCellCoord, InOperation);
-			}
-		});
+				const FCellCoord ChildCellCoord = CellCoord.GetChildCellCoord(ChildIndex);
+				const FBox CellBounds = GetCellBounds(ChildCellCoord);
+
+				if (Box.Intersect(CellBounds))
+				{
+					NumIntersecting += ForEachIntersectingCellInner(Box, ChildCellCoord, InOperation, MinimumLevel);
+				}
+			});
+		}
 	}
 
 	return NumIntersecting;
 }
 
-int32 UWorldPartitionEditorSpatialHash::ForEachIntersectingCell(const FBox& Box, TFunctionRef<void(FCell*)> InOperation)
+int32 UWorldPartitionEditorSpatialHash::ForEachIntersectingCell(const FBox& Box, TFunctionRef<void(FCell*)> InOperation, int32 MinimumLevel)
 {
 	int32 NumIntersecting = 0;
 
@@ -357,7 +362,7 @@ int32 UWorldPartitionEditorSpatialHash::ForEachIntersectingCell(const FBox& Box,
 		{
 			ForEachIntersectingCells(SearchBox, HashLevels.Num() - 1, [&](const FCellCoord& CellCoord)
 			{
-				NumIntersecting += ForEachIntersectingCellInner(Box, CellCoord, InOperation);
+				NumIntersecting += ForEachIntersectingCellInner(Box, CellCoord, InOperation, MinimumLevel);
 			});
 		}
 	}
diff --git a/Engine/Source/Runtime/Engine/Public/WorldPartition/WorldPartitionEditorHash.h b/Engine/Source/Runtime/Engine/Public/WorldPartition/WorldPartitionEditorHash.h
--- a/Engine/Source/Runtime/Engine/Public/WorldPartition/WorldPartitionEditorHash.h
+++ b/Engine/Source/Runtime/Engine/Public/WorldPartition/WorldPartitionEditorHash.h
@@ -25,6 +25,22 @@ public:
 	virtual void HashActor(FWorldPartitionHandle& InActorHandle) PURE_VIRTUAL(UWorldPartitionEditorHash::HashActor, ;);
 	virtual void UnhashActor(FWorldPartitionHandle& InActorHandle) PURE_VIRTUAL(UWorldPartitionEditorHash::UnhashActor, ;);
 
-	virtual int32 ForEachIntersectingActor(const FBox& Box, TFunctionRef<void(FWorldPartitionActorDesc*)> InOperation, bool bIncludeSpatiallyLoadedActors = true, bool bIncludeNonSpatiallyLoadedActors = true) PURE_VIRTUAL(UWorldPartitionEditorHash::ForEachIntersectingActor, return 0;);
+	UE_DEPRECATED(5.1, "Use version that takes FForEachIntersectingActorParams instead.")	
+	int32 ForEachIntersectingActor(const FBox& Box, TFunctionRef<void(FWorldPartitionActorDesc*)> InOperation, bool bIncludeSpatiallyLoadedActors, bool bIncludeNonSpatiallyLoadedActors);
+
+	/* Struct of optional parameters passed to ForEachIntersectingActor. */
+	struct ENGINE_API FForEachIntersectingActorParams
+	{
+		/** Should we include spatially loaded actors in the query? */
+		bool bIncludeSpatiallyLoadedActors = true;
+
+		/** Should we include non-spatially loaded actors in the query? */
+		bool bIncludeNonSpatiallyLoadedActors = true;
+
+		/** Optional minimum box to stop searching for actors */
+		TOptional<FBox> MinimumBox;
+	};
+
+	virtual int32 ForEachIntersectingActor(const FBox& Box, TFunctionRef<void(FWorldPartitionActorDesc*)> InOperation, const FForEachIntersectingActorParams& Params = FForEachIntersectingActorParams()) PURE_VIRTUAL(UWorldPartitionEditorHash::ForEachIntersectingActor, return 0;);
 #endif
 };
diff --git a/Engine/Source/Runtime/Engine/Public/WorldPartition/WorldPartitionEditorSpatialHash.h b/Engine/Source/Runtime/Engine/Public/WorldPartition/WorldPartitionEditorSpatialHash.h
--- a/Engine/Source/Runtime/Engine/Public/WorldPartition/WorldPartitionEditorSpatialHash.h
+++ b/Engine/Source/Runtime/Engine/Public/WorldPartition/WorldPartitionEditorSpatialHash.h
@@ -197,14 +197,14 @@ public:
 	virtual void HashActor(FWorldPartitionHandle& InActorHandle) override;
 	virtual void UnhashActor(FWorldPartitionHandle& InActorHandle) override;
 
-	virtual int32 ForEachIntersectingActor(const FBox& Box, TFunctionRef<void(FWorldPartitionActorDesc*)> InOperation, bool bIncludeSpatiallyLoadedActors = true, bool bIncludeNonSpatiallyLoadedActors = true) override;
+	virtual int32 ForEachIntersectingActor(const FBox& Box, TFunctionRef<void(FWorldPartitionActorDesc*)> InOperation, const FForEachIntersectingActorParams& Params = FForEachIntersectingActorParams()) override;
 	// UWorldPartitionEditorHash interface end
 #endif
 
 #if WITH_EDITORONLY_DATA
 private:
-	int32 ForEachIntersectingCell(const FBox& Box, TFunctionRef<void(FCell*)> InOperation);
-	int32 ForEachIntersectingCellInner(const FBox& Box, const FCellCoord& CellCoord, TFunctionRef<void(FCell*)> InOperation);
+	int32 ForEachIntersectingCell(const FBox& Box, TFunctionRef<void(FCell*)> InOperation, int32 MinimumLevel = 0);
+	int32 ForEachIntersectingCellInner(const FBox& Box, const FCellCoord& CellCoord, TFunctionRef<void(FCell*)> InOperation, int32 MinimumLevel = 0);
 
 	UPROPERTY(Config)
 	int32 CellSize;
