Fix Linux build.

#rnx

diff --git a/Engine/Source/Runtime/Engine/Private/WorldPartition/WorldPartitionEditorHash.cpp b/Engine/Source/Runtime/Engine/Private/WorldPartition/WorldPartitionEditorHash.cpp
--- a/Engine/Source/Runtime/Engine/Private/WorldPartition/WorldPartitionEditorHash.cpp
+++ b/Engine/Source/Runtime/Engine/Private/WorldPartition/WorldPartitionEditorHash.cpp
@@ -7,6 +7,11 @@ UWorldPartitionEditorHash::UWorldPartitionEditorHash(const FObjectInitializer& O
 {}
 
 #if WITH_EDITOR
+UWorldPartitionEditorHash::FForEachIntersectingActorParams::FForEachIntersectingActorParams()
+	: bIncludeSpatiallyLoadedActors(true)
+	, bIncludeNonSpatiallyLoadedActors(true)
+{}
+
 int32 UWorldPartitionEditorHash::ForEachIntersectingActor(const FBox& Box, TFunctionRef<void(FWorldPartitionActorDesc*)> InOperation, bool bIncludeSpatiallyLoadedActors, bool bIncludeNonSpatiallyLoadedActors)
 {
 	UWorldPartitionEditorHash::FForEachIntersectingActorParams ForEachIntersectingActorParams;
diff --git a/Engine/Source/Runtime/Engine/Public/WorldPartition/WorldPartitionEditorHash.h b/Engine/Source/Runtime/Engine/Public/WorldPartition/WorldPartitionEditorHash.h
--- a/Engine/Source/Runtime/Engine/Public/WorldPartition/WorldPartitionEditorHash.h
+++ b/Engine/Source/Runtime/Engine/Public/WorldPartition/WorldPartitionEditorHash.h
@@ -31,11 +31,13 @@ public:
 	/* Struct of optional parameters passed to ForEachIntersectingActor. */
 	struct ENGINE_API FForEachIntersectingActorParams
 	{
+		FForEachIntersectingActorParams();
+
 		/** Should we include spatially loaded actors in the query? */
-		bool bIncludeSpatiallyLoadedActors = true;
+		bool bIncludeSpatiallyLoadedActors;
 
 		/** Should we include non-spatially loaded actors in the query? */
-		bool bIncludeNonSpatiallyLoadedActors = true;
+		bool bIncludeNonSpatiallyLoadedActors;
 
 		/** Optional minimum box to stop searching for actors */
 		TOptional<FBox> MinimumBox;
