Fix nullptr dereferencing when outputing error message

diff --git a/Engine/Source/Runtime/Engine/Private/WorldPartition/WorldPartitionRuntimeSpatialHash.cpp b/Engine/Source/Runtime/Engine/Private/WorldPartition/WorldPartitionRuntimeSpatialHash.cpp
--- a/Engine/Source/Runtime/Engine/Private/WorldPartition/WorldPartitionRuntimeSpatialHash.cpp
+++ b/Engine/Source/Runtime/Engine/Private/WorldPartition/WorldPartitionRuntimeSpatialHash.cpp
@@ -1395,7 +1395,7 @@ bool UWorldPartitionRuntimeSpatialHash::InjectExternalStreamingObject(URuntimeHa
 		else
 		{
 			UE_LOG(LogWorldPartition, Error, TEXT("Incompatible external streaming object to inject. Hash Type: %s. Streaming Object Type: %s."), 
-				*UWorldPartitionRuntimeSpatialHash::StaticClass()->GetName(), *SpatialHashStreamingObject->GetClass()->GetName());
+				*UWorldPartitionRuntimeSpatialHash::StaticClass()->GetName(), *ExternalStreamingObject->GetClass()->GetName());
 		}
 	}
 	return false;
