Fix MMV + FSceneView ensure

We were assuming that the instanced view UB is only needed by stereo view families, but that isn't true because currently, the base pass shaders are compiled with MMV or ISR enabled.

diff --git a/Engine/Source/Runtime/Engine/Private/SceneView.cpp b/Engine/Source/Runtime/Engine/Private/SceneView.cpp
--- a/Engine/Source/Runtime/Engine/Private/SceneView.cpp
+++ b/Engine/Source/Runtime/Engine/Private/SceneView.cpp
@@ -878,6 +878,8 @@ PRAGMA_ENABLE_DEPRECATION_WARNINGS
 	}
 
 	{
+		// The shader variants that are compiled have ISR or MMV _enabled_ in the shader, even if the current ViewFamily doesn't
+		// require multiple views functionality.
 		const UE::StereoRenderUtils::FStereoShaderAspects Aspects(ShaderPlatform);
 		bShouldBindInstancedViewUB = Aspects.IsInstancedStereoEnabled() || Aspects.IsMobileMultiViewEnabled();
 
@@ -885,7 +887,6 @@ PRAGMA_ENABLE_DEPRECATION_WARNINGS
 		ensure(bIsInstancedStereoEnabled == Aspects.IsInstancedStereoEnabled());
 		ensure(bIsMultiViewportEnabled == Aspects.IsInstancedMultiViewportEnabled());
 		ensure(bIsMobileMultiViewEnabled == (Family && Family->bRequireMultiView && Aspects.IsMobileMultiViewEnabled()));
-		ensure(bShouldBindInstancedViewUB == bIsInstancedStereoEnabled || bIsMobileMultiViewEnabled);
 	}
 
 	SetupAntiAliasingMethod();
