Fix for anim graph notifies not triggering in servers.
When  "a.EnableQueuedAnimEventsOnServer" is enabled, and servers are doing a full graph update (common in a Listen Server setup), notifies are queued and triggered at the end of the frame.

diff --git a/Engine/Source/Runtime/Engine/Private/Components/CharacterMovementComponent.cpp b/Engine/Source/Runtime/Engine/Private/Components/CharacterMovementComponent.cpp
--- a/Engine/Source/Runtime/Engine/Private/Components/CharacterMovementComponent.cpp
+++ b/Engine/Source/Runtime/Engine/Private/Components/CharacterMovementComponent.cpp
@@ -139,6 +139,14 @@ namespace CharacterMovementCVars
 		TEXT("0: Disable, 1: Enable"),
 		ECVF_Default);
 
+	static int32 EnableQueuedAnimEventsOnServer = 0;
+	FAutoConsoleVariableRef CVarEnableQueuedAnimEventsOnServer(
+		TEXT("a.EnableQueuedAnimEventsOnServer"),
+		EnableQueuedAnimEventsOnServer,
+		TEXT("Whether to enable queued anim events on servers. In most cases, when the server is doing a full anim graph update, queued notifies aren't triggered by the server, but this will enable them. Enabling this is recommended in projects using Listen Servers. \n")
+		TEXT("0: Disable, 1: Enable"),
+		ECVF_Default);
+
 	// Logging when character is stuck. Off by default in shipping.
 #if UE_BUILD_SHIPPING
 	static float StuckWarningPeriod = -1.f;
@@ -9737,8 +9745,12 @@ void UCharacterMovementComponent::MoveAutonomous
 		}
 		// TODO: SaveBaseLocation() in case tick moves us?
 
-		// Trigger Events right away, as we could be receiving multiple ServerMoves per frame.
-		CharacterOwner->GetMesh()->ConditionallyDispatchQueuedAnimEvents();
+		if (!CharacterMovementCVars::EnableQueuedAnimEventsOnServer || CharacterOwner->GetMesh()->ShouldOnlyTickMontages(DeltaTime))
+		{
+			// If we're not doing a full anim graph update on the server, 
+			// trigger events right away, as we could be receiving multiple ServerMoves per frame.
+			CharacterOwner->GetMesh()->ConditionallyDispatchQueuedAnimEvents();
+		}
 	}
 
 	if (CharacterOwner && UpdatedComponent)
