Iris - Partial bunches are not supported at this moment. Iris uses different mechanisms to support bulk data.

#rnx

diff --git a/Engine/Source/Runtime/Engine/Private/Net/Experimental/Iris/DataStreamChannel.cpp b/Engine/Source/Runtime/Engine/Private/Net/Experimental/Iris/DataStreamChannel.cpp
--- a/Engine/Source/Runtime/Engine/Private/Net/Experimental/Iris/DataStreamChannel.cpp
+++ b/Engine/Source/Runtime/Engine/Private/Net/Experimental/Iris/DataStreamChannel.cpp
@@ -51,7 +51,6 @@ void UDataStreamChannel::Init(UNetConnection* InConnection, int32 InChIndex, ECh
 		DataStreamManager->Init(InitParams);
 	}
 
-	// $IRIS TODO: Not pretty. Need to figure out something elegant. This will have to do for now
 	if (UReplicationSystem* ReplicationSystem = InConnection->Driver->GetReplicationSystem())
 	{
 		bIsReadyToHandshake = 1U;
@@ -83,12 +82,20 @@ void UDataStreamChannel::ReceivedBunch(FInBunch& Bunch)
 
 	IRIS_PROFILER_SCOPE(UDataStreamChannel_ReceivedBunch);
 
-	// We are sending dummy bunches until we are open...
+	// We are sending dummy bunches until we are open.
 	if (!Bunch.GetNumBits())
 	{
 		return;
 	}
 
+	// We do not support partial bunches.
+	if (Bunch.bPartial)
+	{
+		Bunch.SetAtEnd();
+		Bunch.SetError();
+		return;
+	}
+
 	FNetBitStreamReader BitReader;
 	BitReader.InitBits(Bunch.GetData(), Bunch.GetNumBits());
 	BitReader.Seek(Bunch.GetPosBits());
