Fix PVS false-positive warning V568: It's odd that the argument of sizeof() operator is the 'ReplicationData.Num()' expression. - I really do want to get size of a return type.

#rnx

diff --git a/Engine/Source/Runtime/Engine/Private/OnlineReplStructs.cpp b/Engine/Source/Runtime/Engine/Private/OnlineReplStructs.cpp
--- a/Engine/Source/Runtime/Engine/Private/OnlineReplStructs.cpp
+++ b/Engine/Source/Runtime/Engine/Private/OnlineReplStructs.cpp
@@ -233,7 +233,7 @@ void FUniqueNetIdRepl::MakeReplicationDataV2()
 	TArray<uint8> ReplicationData = UE::Online::FOnlineIdRegistryRegistry::Get().ToReplicationData(AccountId);
 	check(!ReplicationData.IsEmpty());
 
-	const int32 TotalBytes = sizeof(EncodingFlags) + sizeof(OnlineServicesType) + sizeof(ReplicationData.Num()) + ReplicationData.Num();
+	const int32 TotalBytes = sizeof(EncodingFlags) + sizeof(OnlineServicesType) + sizeof(TArray<uint8>::SizeType) + ReplicationData.Num();
 	ReplicationBytes.Empty(TotalBytes);
 	FMemoryWriter Writer(ReplicationBytes);
 	Writer << EncodingFlags;
