Fix typo catched by static analysis
#rnx

diff --git a/Engine/Source/Runtime/Engine/Public/MaterialRestrictiveMode.h b/Engine/Source/Runtime/Engine/Public/MaterialRestrictiveMode.h
--- a/Engine/Source/Runtime/Engine/Public/MaterialRestrictiveMode.h
+++ b/Engine/Source/Runtime/Engine/Public/MaterialRestrictiveMode.h
@@ -20,7 +20,7 @@ struct FScopedHLSLMaterialTranslatorRestrictiveModeChange
 		SetHLSLMaterialTranslatorRestrictiveMode(bRestrictiveMode);
 	}
 	
-	FORCEINLINE FScopedHLSLMaterialTranslatorRestrictiveModeChange()
+	FORCEINLINE ~FScopedHLSLMaterialTranslatorRestrictiveModeChange()
 	{
 		SetHLSLMaterialTranslatorRestrictiveMode(bRestoreRestrictiveMode);
 	}
