Networking deprecation cleanup

diff --git a/Engine/Plugins/Runtime/ReplicationGraph/Source/Public/ReplicationGraphTypes.h b/Engine/Plugins/Runtime/ReplicationGraph/Source/Public/ReplicationGraphTypes.h
--- a/Engine/Plugins/Runtime/ReplicationGraph/Source/Public/ReplicationGraphTypes.h
+++ b/Engine/Plugins/Runtime/ReplicationGraph/Source/Public/ReplicationGraphTypes.h
@@ -1296,13 +1296,6 @@ typedef TArray<FNetViewer, FReplicationGraphConnectionsAllocator> FNetViewerArra
 // Parameter structure for what we actually pass down during the Gather phase.
 struct FConnectionGatherActorListParameters
 {
-	UE_DEPRECATED(4.26, "Please use the constructor that takes a viewer array.")
-	FConnectionGatherActorListParameters(FNetViewer& InViewer, UNetReplicationGraphConnection& InConnectionManager, TSet<FName>& InClientVisibleLevelNamesRef, uint32 InReplicationFrameNum, FGatheredReplicationActorLists& InOutGatheredReplicationLists)
-		: ConnectionManager(InConnectionManager), ReplicationFrameNum(InReplicationFrameNum), OutGatheredReplicationLists(InOutGatheredReplicationLists), ClientVisibleLevelNamesRef(InClientVisibleLevelNamesRef)
-	{
-		Viewers.Emplace(InViewer);
-	}
-
 	FConnectionGatherActorListParameters(FNetViewerArray& InViewers, UNetReplicationGraphConnection& InConnectionManager, TSet<FName>& InClientVisibleLevelNamesRef, uint32 InReplicationFrameNum, FGatheredReplicationActorLists& InOutGatheredReplicationLists)
 		: Viewers(InViewers), ConnectionManager(InConnectionManager), ReplicationFrameNum(InReplicationFrameNum), OutGatheredReplicationLists(InOutGatheredReplicationLists), ClientVisibleLevelNamesRef(InClientVisibleLevelNamesRef)
 	{
diff --git a/Engine/Source/Editor/UnrealEd/Classes/Editor/EditorEngine.h b/Engine/Source/Editor/UnrealEd/Classes/Editor/EditorEngine.h
--- a/Engine/Source/Editor/UnrealEd/Classes/Editor/EditorEngine.h
+++ b/Engine/Source/Editor/UnrealEd/Classes/Editor/EditorEngine.h
@@ -803,8 +803,6 @@ private:
 	virtual TSharedPtr<SViewport> GetGameViewportWidget() const override;
 	virtual void TriggerStreamingDataRebuild() override;
 
-	// deprecated in 4.26
-	virtual bool NetworkRemapPath(UNetDriver* Driver, FString& Str, bool bReading = true) override;
 	virtual bool NetworkRemapPath(UNetConnection* Connection, FString& Str, bool bReading = true) override;
 	virtual bool NetworkRemapPath(UPendingNetGame* PendingNetGame, FString& Str, bool bReading = true) override;
 
diff --git a/Engine/Source/Editor/UnrealEd/Private/EditorEngine.cpp b/Engine/Source/Editor/UnrealEd/Private/EditorEngine.cpp
--- a/Engine/Source/Editor/UnrealEd/Private/EditorEngine.cpp
+++ b/Engine/Source/Editor/UnrealEd/Private/EditorEngine.cpp
@@ -7144,19 +7144,6 @@ bool UEditorEngine::NetworkRemapPath(UNetConnection* Connection, FString& Str, b
 	return NetworkRemapPath_local(Context, Str, bReading, Connection->IsReplay());
 }
 
-bool UEditorEngine::NetworkRemapPath(UNetDriver* Driver, FString& Str, bool bReading)
-{
-	if (Driver == nullptr)
-	{
-		return false;
-	}
-
-	// Pretty sure there's no case where you can't have a world by this point.
-	const bool bIsReplay = Driver->GetWorld() ? (Driver->GetWorld()->GetDemoNetDriver() != nullptr) : false;
-	FWorldContext& Context = GetWorldContextFromWorldChecked(Driver->GetWorld());
-	return NetworkRemapPath_local(Context, Str, bReading, bIsReplay);
-}
-
 bool UEditorEngine::NetworkRemapPath( UPendingNetGame *PendingNetGame, FString& Str, bool bReading)
 {
 	FWorldContext& Context = GetWorldContextFromPendingNetGameChecked(PendingNetGame);
diff --git a/Engine/Source/Runtime/Engine/Classes/Engine/ActorChannel.h b/Engine/Source/Runtime/Engine/Classes/Engine/ActorChannel.h
--- a/Engine/Source/Runtime/Engine/Classes/Engine/ActorChannel.h
+++ b/Engine/Source/Runtime/Engine/Classes/Engine/ActorChannel.h
@@ -182,9 +182,6 @@ public:
 
 	virtual void NotifyActorChannelOpen(AActor* InActor, FInBunch& InBunch);
 
-	UE_DEPRECATED(4.26, "Use UNetworkDriver::SendDestructionInfo instead.")
-	int64 SetChannelActorForDestroy( struct FActorDestructionInfo *DestructInfo );
-
 	/** Append any export bunches */
 	virtual void AppendExportBunches( TArray< FOutBunch* >& OutExportBunches ) override;
 
diff --git a/Engine/Source/Runtime/Engine/Classes/Engine/Engine.h b/Engine/Source/Runtime/Engine/Classes/Engine/Engine.h
--- a/Engine/Source/Runtime/Engine/Classes/Engine/Engine.h
+++ b/Engine/Source/Runtime/Engine/Classes/Engine/Engine.h
@@ -3112,8 +3112,6 @@ public:
 	void DestroyNamedNetDriver(UWorld *InWorld, FName NetDriverName);
 	void DestroyNamedNetDriver(UPendingNetGame *PendingNetGame, FName NetDriverName);
 
-	UE_DEPRECATED(4.26, "Please use NetworkRemapPath that takes a connection instead.")
-	virtual bool NetworkRemapPath(UNetDriver* Driver, FString &Str, bool bReading=true) { return false; }
 	virtual bool NetworkRemapPath(UNetConnection* Connection, FString& Str, bool bReading=true) { return false; }
 	virtual bool NetworkRemapPath(UPendingNetGame *PendingNetGame, FString &Str, bool bReading=true) { return false; }
 
diff --git a/Engine/Source/Runtime/Engine/Classes/Engine/GameEngine.h b/Engine/Source/Runtime/Engine/Classes/Engine/GameEngine.h
--- a/Engine/Source/Runtime/Engine/Classes/Engine/GameEngine.h
+++ b/Engine/Source/Runtime/Engine/Classes/Engine/GameEngine.h
@@ -17,9 +17,6 @@ class FSceneViewport;
 class UGameViewportClient;
 class UNetDriver;
 
-PRAGMA_DISABLE_DEPRECATION_WARNINGS
-// disabled for: virtual bool NetworkRemapPath(UNetDriver* Driver, FString& Str, bool bReading = true) override;
-
 /**
  * Engine that manages core systems that enable a game.
  */
@@ -140,7 +137,6 @@ public:
 	virtual float GetMaxTickRate( float DeltaTime, bool bAllowFrameRateSmoothing = true ) const override;
 	virtual void ProcessToggleFreezeCommand( UWorld* InWorld ) override;
 	virtual void ProcessToggleFreezeStreamingCommand( UWorld* InWorld ) override;
-	virtual bool NetworkRemapPath(UNetDriver* Driver, FString& Str, bool bReading = true) override;
 	virtual bool NetworkRemapPath(UNetConnection* Connection, FString& Str, bool bReading = true) override;
 	virtual bool ShouldDoAsyncEndOfFrameTasks() const override;
 
@@ -195,5 +191,3 @@ private:
 
 	TPimplPtr<FEngineConsoleCommandExecutor> CmdExec;
 };
-
-PRAGMA_ENABLE_DEPRECATION_WARNINGS
diff --git a/Engine/Source/Runtime/Engine/Classes/Engine/NetConnection.h b/Engine/Source/Runtime/Engine/Classes/Engine/NetConnection.h
--- a/Engine/Source/Runtime/Engine/Classes/Engine/NetConnection.h
+++ b/Engine/Source/Runtime/Engine/Classes/Engine/NetConnection.h
@@ -292,10 +292,11 @@ public:
 	UPROPERTY()
 	int32	MaxPacket;						// Maximum packet size.
 
-	UE_DEPRECATED(4.25, "Please use IsInternalAck/SetInternalAck instead")
+private:
 	UPROPERTY()
 	uint32 InternalAck:1;					// Internally ack all packets, for 100% reliable connections.
 
+public:
 	bool IsInternalAck() const { return bInternalAck; }
 	void SetInternalAck(bool bValue) 
 	{ 
@@ -427,16 +428,7 @@ public:
 	int32			TickCount;				// Count of ticks.
 	uint32			LastProcessedFrame;   // The last frame where we gathered and processed actors for this connection
 
-	/** The last time an ack was received */
-	UE_DEPRECATED(4.25, "Please use GetLastRecvAckTime() instead.")
-	float			LastRecvAckTime;
-
 	double GetLastRecvAckTime() const { return LastRecvAckTimestamp; }
-
-	/** Time when connection request was first initiated */
-	UE_DEPRECATED(4.25, "Please use GetConnectTime() instead.")
-	float			ConnectTime;
-
 	double GetConnectTime() const { return ConnectTimestamp; }
 
 private:
@@ -567,8 +559,6 @@ public:
 	double			OutLagTime[256];				// For lag measuring.
 	int32			OutLagPacketId[256];			// For lag measuring.
 	uint8			OutBytesPerSecondHistory[256];	// For saturation measuring.
-	UE_DEPRECATED(4.25, "RemoteSaturation is not calculated anymore and is now deprecated")
-	float			RemoteSaturation;
 	int32			InPacketId;						// Full incoming packet index.
 	int32			OutPacketId;					// Most recently sent packet.
 	int32 			OutAckPacketId;					// Most recently acked outgoing packet.
@@ -933,10 +923,6 @@ public:
 	 */
 	ENGINE_API virtual void FlushNet(bool bIgnoreSimulation = false);
 
-	/** Poll the connection. If it is timed out, close it. */
-	UE_DEPRECATED(4.26, "Now takes DeltaSeconds")
-	ENGINE_API virtual void Tick() { Tick(0.0f); }
-
 	/** Poll the connection. If it is timed out, close it. */
 	ENGINE_API virtual void Tick(float DeltaSeconds);
 
@@ -1066,34 +1052,16 @@ public:
 	 */
 	ENGINE_API virtual void NotifyAnalyticsProvider();
 
-	/**
-	 * Sets the encryption key and enables encryption.
-	 */
-	UE_DEPRECATED(4.24, "Use SetEncryptionData instead.")
-	ENGINE_API void EnableEncryptionWithKey(TArrayView<const uint8> Key);
-	
 	/**
 	 * Sets the encryption data and enables encryption.
 	 */
 	ENGINE_API void EnableEncryption(const FEncryptionData& EncryptionData);
 
-	/**
-	 * Sets the encryption key, enables encryption, and sends the encryption ack to the client.
-	 */
-	UE_DEPRECATED(4.24, "Use SetEncryptionData instead.")
-	ENGINE_API void EnableEncryptionWithKeyServer(TArrayView<const uint8> Key);
-
 	/**
 	 * Sets the encryption data, enables encryption, and sends the encryption ack to the client.
 	 */
 	ENGINE_API void EnableEncryptionServer(const FEncryptionData& EncryptionData);
 
-	/**
-	 * Sets the key for the underlying encryption packet handler component, but doesn't modify encryption enabled state.
-	 */
-	UE_DEPRECATED(4.24, "Use SetEncryptionData instead.")
-	ENGINE_API void SetEncryptionKey(TArrayView<const uint8> Key);
-
 	/**
 	 * Sets the data for the underlying encryption packet handler component, but doesn't modify encryption enabled state.
 	 */
@@ -1272,13 +1240,6 @@ public:
 	/** Returns the online platform name for the player on this connection. Only valid for client connections on servers. */
 	ENGINE_API FName GetPlayerOnlinePlatformName() const { return PlayerOnlinePlatformName; }
 	
-	/**
-	 * Sets whether or not we should ignore bunches that would attempt to open channels that are already open.
-	 * Should only be used with InternalAck.
-	 */
-	UE_DEPRECATED(4.26, "Please call SetAllowExistingChannelIndex instead.")
-	void SetIgnoreAlreadyOpenedChannels(bool bInIgnoreAlreadyOpenedChannels) { SetAllowExistingChannelIndex(bInIgnoreAlreadyOpenedChannels); }
-
 	/** Sets whether we handle opening channels with an index that already exists, used by replays to fast forward the packet stream */
 	void SetAllowExistingChannelIndex(bool bAllow);
 
diff --git a/Engine/Source/Runtime/Engine/Classes/Engine/NetworkDelegates.h b/Engine/Source/Runtime/Engine/Classes/Engine/NetworkDelegates.h
--- a/Engine/Source/Runtime/Engine/Classes/Engine/NetworkDelegates.h
+++ b/Engine/Source/Runtime/Engine/Classes/Engine/NetworkDelegates.h
@@ -132,9 +132,6 @@ struct FEncryptionKeyResponse
 	EEncryptionResponse Response;
 	/** Error message related to the response */
 	FString ErrorMsg;
-	/** Encryption key */
-	UE_DEPRECATED(4.24, "Use EncryptionData struct instead.")
-	TArray<uint8> EncryptionKey;
 	/** Encryption data */
 	FEncryptionData EncryptionData;
 
diff --git a/Engine/Source/Runtime/Engine/Classes/Engine/NetworkSettings.h b/Engine/Source/Runtime/Engine/Classes/Engine/NetworkSettings.h
--- a/Engine/Source/Runtime/Engine/Classes/Engine/NetworkSettings.h
+++ b/Engine/Source/Runtime/Engine/Classes/Engine/NetworkSettings.h
@@ -45,22 +45,6 @@ class ENGINE_API UNetworkSettings : public UDeveloperSettings
 		ToolTip="If true, origin rebasing is enabled in multiplayer games, meaning that servers and clients can have different local world origins."))
 	uint32 bEnableMultiplayerWorldOriginRebasing : 1;
 
-	//! Maximum allowable size for replicated dynamic arrays (in number of elements)
-	UE_DEPRECATED(4.26, "This property is no longer used.")
-	UPROPERTY(config, EditAnywhere, Category = replication, meta = (
-		DisplayName = "Max Array Size",
-		ToolTip = "Maximum allowable size for replicated dynamic arrays (in number of elements). Must be between 1 and 65535.",
-		ClampMin = "1", ClampMax = "65535", UIMin = "1", UIMax = "65535"))
-	int32 MaxRepArraySize = DefaultMaxRepArraySize;
-
-	//! Maximum allowable size for replicated dynamic arrays (in bytes)
-	UE_DEPRECATED(4.26, "This property is no longer used.")
-	UPROPERTY(config, EditAnywhere, Category = replication, meta = (
-		DisplayName = "Max Array Memory",
-		ToolTip = "Maximum allowable size for replicated dynamic arrays (in bytes).  Must be between 1 and 65535.",
-		ClampMin = "1", ClampMax = "65535", UIMin = "1", UIMax = "65535"))
-	int32 MaxRepArrayMemory = DefaultMaxRepArrayMemory;
-
 	/** This lists the common network emulation profiles that will be selectable in PIE settings */
 	UPROPERTY(config)
 	TArray<FNetworkEmulationProfileDescription> NetworkEmulationProfiles;
diff --git a/Engine/Source/Runtime/Engine/Classes/GameFramework/Actor.h b/Engine/Source/Runtime/Engine/Classes/GameFramework/Actor.h
--- a/Engine/Source/Runtime/Engine/Classes/GameFramework/Actor.h
+++ b/Engine/Source/Runtime/Engine/Classes/GameFramework/Actor.h
@@ -3259,35 +3259,6 @@ public:
 	/** Forces this actor to be net relevant if it is not already by default	 */
 	virtual void ForceNetRelevant();
 
-	/** Updates NetUpdateTime to the new value for future net relevancy checks */
-	UE_DEPRECATED(4.26, "No longer used.")
-	void SetNetUpdateTime(float NewUpdateTime);
-
-	/**
-	 *	Find the FNetworkObjectInfo struct associated with this actor (for the main NetDriver).
-	 *	If no info is found, one will be created and added to the main NetDriver.
-	 *
-	 *	@return The FNetworkObejctInfo associated with this Actor, or nullptr if one couldn't be created.
-	 */
-	UE_DEPRECATED(4.26, "Please use FindOrAddNetworkObjectInfo on the net driver instead.")
-	FNetworkObjectInfo* FindOrAddNetworkObjectInfo();
-
-	/**
-	 *	Find the FNetworkObjectInfo struct associated with this actor (for the main NetDriver).
-	 *
-	 *	@return The FNetworkObejctInfo associated with this Actor, or nullptr if none was found.
-	 */
-	UE_DEPRECATED(4.26, "Please use FindNetworkObjectInfo on the net driver instead.")
-	FNetworkObjectInfo* FindNetworkObjectInfo();
-	
-	UE_DEPRECATED(4.26, "Please use FindNetworkObjectInfo on the net driver instead..")
-	const FNetworkObjectInfo* FindNetworkObjectInfo() const
-	{
-		PRAGMA_DISABLE_DEPRECATION_WARNINGS
-		return const_cast<AActor*>(this)->FindNetworkObjectInfo();
-		PRAGMA_ENABLE_DEPRECATION_WARNINGS
-	}
-
 	/** Force actor to be updated to clients/demo net drivers */
 	UFUNCTION( BlueprintCallable, Category="Networking")
 	virtual void ForceNetUpdate();
diff --git a/Engine/Source/Runtime/Engine/Private/Actor.cpp b/Engine/Source/Runtime/Engine/Private/Actor.cpp
--- a/Engine/Source/Runtime/Engine/Private/Actor.cpp
+++ b/Engine/Source/Runtime/Engine/Private/Actor.cpp
@@ -5753,43 +5753,6 @@ float AActor::GetGameTimeSinceCreation() const
 	}
 }
 
-PRAGMA_DISABLE_DEPRECATION_WARNINGS
-void AActor::SetNetUpdateTime( float NewUpdateTime )
-{
-	if ( FNetworkObjectInfo* NetActor = FindNetworkObjectInfo() )
-	{
-		// Only allow the next update to be sooner than the current one
-		NetActor->NextUpdateTime = FMath::Min( NetActor->NextUpdateTime, (double)NewUpdateTime );
-	}			
-}
-PRAGMA_ENABLE_DEPRECATION_WARNINGS
-
-FNetworkObjectInfo* AActor::FindOrAddNetworkObjectInfo()
-{
-	if ( UWorld* World = GetWorld() )
-	{
-		if ( UNetDriver* NetDriver = World->GetNetDriver() )
-		{
-			return NetDriver->FindOrAddNetworkObjectInfo( this );
-		}
-	}
-
-	return nullptr;
-}
-
-FNetworkObjectInfo* AActor::FindNetworkObjectInfo()
-{
-	if ( UWorld* World = GetWorld() )
-	{
-		if ( UNetDriver* NetDriver = World->GetNetDriver() )
-		{
-			return NetDriver->FindNetworkObjectInfo(this);
-		}
-	}
-
-	return nullptr;
-}
-
 void AActor::PostRename(UObject* OldOuter, const FName OldName)
 {
 	Super::PostRename(OldOuter, OldName);
diff --git a/Engine/Source/Runtime/Engine/Private/DataChannel.cpp b/Engine/Source/Runtime/Engine/Private/DataChannel.cpp
--- a/Engine/Source/Runtime/Engine/Private/DataChannel.cpp
+++ b/Engine/Source/Runtime/Engine/Private/DataChannel.cpp
@@ -2594,44 +2594,6 @@ void UActorChannel::NotifyActorChannelOpen(AActor* InActor, FInBunch& InBunch)
 	}
 }
 
-int64 UActorChannel::SetChannelActorForDestroy( FActorDestructionInfo *DestructInfo )
-{
-	int64 NumBits = 0;
-	check(Connection);
-	check(Connection->Channels[ChIndex]==this);
-	check(DestructInfo);
-
-	if
-	(	!Closing
-	&&	(Connection->GetConnectionState()==USOCK_Open || Connection->GetConnectionState()==USOCK_Pending) )
-	{
-		// Outer must be valid to call PackageMap->WriteObject. In the case of streaming out levels, this can go null out of from underneath us. In that case, just skip the destruct info.
-		// We assume that if server unloads a level that clients will to and this will implicitly destroy all actors in it, so not worried about leaking actors client side here.
-		if (UObject* ObjOuter = DestructInfo->ObjOuter.Get())
-		{
-
-		// Send a close notify, and wait for ack.
-		FOutBunch CloseBunch( this, 1 );
-		check(!CloseBunch.IsError());
-		check(CloseBunch.bClose);
-		CloseBunch.bReliable = 1;
-		CloseBunch.CloseReason = DestructInfo->Reason;
-
-		// Serialize DestructInfo
-		NET_CHECKSUM(CloseBunch); // This is to mirror the Checksum in UPackageMapClient::SerializeNewActor
-			Connection->PackageMap->WriteObject( CloseBunch, ObjOuter, DestructInfo->NetGUID, DestructInfo->PathName );
-
-		UE_LOG(LogNetTraffic, Log, TEXT("SetChannelActorForDestroy: Channel %d. NetGUID <%s> Path: %s. Bits: %d"), ChIndex, *DestructInfo->NetGUID.ToString(), *DestructInfo->PathName, CloseBunch.GetNumBits() );
-		UE_LOG(LogNetDormancy, Verbose, TEXT("SetChannelActorForDestroy: Channel %d. NetGUID <%s> Path: %s. Bits: %d"), ChIndex, *DestructInfo->NetGUID.ToString(), *DestructInfo->PathName, CloseBunch.GetNumBits() );
-
-		SendBunch( &CloseBunch, 0 );
-		NumBits = CloseBunch.GetNumBits();
-	}
-	}
-
-	return NumBits;
-}
-
 void UActorChannel::Tick()
 {
 	Super::Tick();
diff --git a/Engine/Source/Runtime/Engine/Private/DemoNetDriver.cpp b/Engine/Source/Runtime/Engine/Private/DemoNetDriver.cpp
--- a/Engine/Source/Runtime/Engine/Private/DemoNetDriver.cpp
+++ b/Engine/Source/Runtime/Engine/Private/DemoNetDriver.cpp
@@ -2198,10 +2198,7 @@ bool UDemoNetDriver::ReplicatePrioritizedActor(const FActorPriority& ActorPriori
 
 				FScopedRepContext LevelContext(Params.Connection, DestructionInfo->Level.Get());
 
-				// Send a close bunch on the new channel
-				PRAGMA_DISABLE_DEPRECATION_WARNINGS
-				Channel->SetChannelActorForDestroy(DestructionInfo);
-				PRAGMA_ENABLE_DEPRECATION_WARNINGS
+				SendDestructionInfo(Params.Connection, DestructionInfo);
 
 				// Remove from connection's to-be-destroyed list (close bunch is reliable, so it will make it there)
 				Params.Connection->RemoveDestructionInfo(DestructionInfo);
diff --git a/Engine/Source/Runtime/Engine/Private/GameEngine.cpp b/Engine/Source/Runtime/Engine/Private/GameEngine.cpp
--- a/Engine/Source/Runtime/Engine/Private/GameEngine.cpp
+++ b/Engine/Source/Runtime/Engine/Private/GameEngine.cpp
@@ -1262,7 +1262,6 @@ void UGameEngine::FinishDestroy()
 	CmdExec.Reset();
 }
 
-//@todo: unify this and the driver version
 bool UGameEngine::NetworkRemapPath(UNetConnection* Connection, FString& Str, bool bReading /*= true*/)
 {
 	if (Connection == nullptr)
@@ -1364,104 +1363,6 @@ bool UGameEngine::NetworkRemapPath(UNetConnection* Connection, FString& Str, boo
 	return false;
 }
 
-bool UGameEngine::NetworkRemapPath(UNetDriver* Driver, FString& Str, bool bReading /*= true*/)
-{
-	if (Driver == nullptr)
-	{
-		return false;
-	}
-
-	UWorld* const World = Driver->GetWorld();
-
-	if (World == nullptr)
-	{
-		return false;
-	}
-
-	// If the driver is using a duplicate level ID, find the level collection using the driver
-	// and see if any of its levels match the prefixed name. If so, remap Str to that level's
-	// prefixed name.
-	if (Driver->GetDuplicateLevelID() != INDEX_NONE && bReading)
-	{
-		const FName PrefixedName = *UWorld::ConvertToPIEPackageName(Str, Driver->GetDuplicateLevelID());
-
-		for (const FLevelCollection& Collection : World->GetLevelCollections())
-		{
-			if (Collection.GetNetDriver() == Driver || Collection.GetDemoNetDriver() == Driver)
-			{
-				for (const ULevel* Level : Collection.GetLevels())
-				{
-					const UPackage* const CachedOutermost = Level ? Level->GetOutermost() : nullptr;
-					if (CachedOutermost && CachedOutermost->GetFName() == PrefixedName)
-					{
-						Str = PrefixedName.ToString();
-						return true;
-					}
-				}
-			}
-		}
-	}
-
-	if (!bReading)
-	{
-		return false;
-	}
-
-	// Try to find the level script objects and remap them for when demos are being replayed.
-	if (World->GetDemoNetDriver() == Driver && World->RemapCompiledScriptActor(Str))
-	{
-		return true;
-	}
-
-	// If the game has created multiple worlds, some of them may have prefixed package names,
-	// so we need to remap the world package and streaming levels for replay playback to work correctly.
-	FWorldContext& Context = GetWorldContextFromWorldChecked(World);
-	if (Context.PIEInstance == INDEX_NONE)
-	{
-		if (WorldList.Num() > 1)
-		{
-			// If this is not a PIE instance but sender is PIE, we need to strip the PIE prefix
-			const FString Stripped = UWorld::RemovePIEPrefix(Str);
-			if (!Stripped.Equals(Str, ESearchCase::CaseSensitive))
-			{
-				Str = Stripped;
-				return true;
-			}
-		}
-		return false;
-	}
-
-	// If the prefixed path matches the world package name or the name of a streaming level,
-	// return the prefixed name.
-	FString PackageNameOnly = Str;
-	FPackageName::TryConvertFilenameToLongPackageName(PackageNameOnly, PackageNameOnly);
-
-	const FString PrefixedFullName = UWorld::ConvertToPIEPackageName(Str, Context.PIEInstance);
-	const FString PrefixedPackageName = UWorld::ConvertToPIEPackageName(PackageNameOnly, Context.PIEInstance);
-	const FString WorldPackageName = World->GetOutermost()->GetName();
-
-	if (WorldPackageName == PrefixedPackageName)
-	{
-		Str = PrefixedFullName;
-		return true;
-	}
-
-	for (ULevelStreaming* StreamingLevel : World->GetStreamingLevels())
-	{
-		if (StreamingLevel != nullptr)
-		{
-			const FString StreamingLevelName = StreamingLevel->GetWorldAsset().GetLongPackageName();
-			if (StreamingLevelName == PrefixedPackageName)
-			{
-				Str = PrefixedFullName;
-				return true;
-			}
-		}
-	}
-
-	return false;
-}
-
 bool UGameEngine::ShouldDoAsyncEndOfFrameTasks() const
 {
 	return FApp::ShouldUseThreadingForPerformance() && ENamedThreads::GetRenderThread() != ENamedThreads::GameThread && !!GDoAsyncEndOfFrameTasks;
diff --git a/Engine/Source/Runtime/Engine/Private/NetConnection.cpp b/Engine/Source/Runtime/Engine/Private/NetConnection.cpp
--- a/Engine/Source/Runtime/Engine/Private/NetConnection.cpp
+++ b/Engine/Source/Runtime/Engine/Private/NetConnection.cpp
@@ -392,10 +392,6 @@ void UNetConnection::InitBase(UNetDriver* InDriver,class FSocket* InSocket, cons
 	LastTime				= 0.0;
 	LastSendTime			= DriverElapsedTime;
 	LastTickTime			= DriverElapsedTime;
-	PRAGMA_DISABLE_DEPRECATION_WARNINGS
-	LastRecvAckTime			= DriverElapsedTime;
-	ConnectTime				= DriverElapsedTime;
-	PRAGMA_ENABLE_DEPRECATION_WARNINGS
 	LastRecvAckTimestamp	= DriverElapsedTime;
 	ConnectTimestamp		= DriverElapsedTime;
 
@@ -680,14 +676,6 @@ void UNetConnection::NotifyConnectionUpdated()
 	}
 }
 
-void UNetConnection::EnableEncryptionWithKey(TArrayView<const uint8> Key)
-{
-	FEncryptionData EncryptionData;
-	EncryptionData.Key.Append(Key.GetData(), Key.Num());
-
-	EnableEncryption(EncryptionData);
-}
-
 void UNetConnection::EnableEncryption(const FEncryptionData& EncryptionData)
 {
 	if (Handler.IsValid())
@@ -707,14 +695,6 @@ void UNetConnection::EnableEncryption(const FEncryptionData& EncryptionData)
 	}
 }
 
-void UNetConnection::EnableEncryptionWithKeyServer(TArrayView<const uint8> Key)
-{
-	FEncryptionData EncryptionData;
-	EncryptionData.Key.Append(Key.GetData(), Key.Num());
-
-	EnableEncryptionServer(EncryptionData);
-}
-
 void UNetConnection::EnableEncryptionServer(const FEncryptionData& EncryptionData)
 {
 	if (GetConnectionState() != USOCK_Invalid && GetConnectionState() != USOCK_Closed && Driver)
@@ -741,14 +721,6 @@ void UNetConnection::SendClientEncryptionAck()
 	}
 }
 
-void UNetConnection::SetEncryptionKey(TArrayView<const uint8> Key)
-{
-	FEncryptionData EncryptionData;
-	EncryptionData.Key.Append(Key.GetData(), Key.Num());
-
-	SetEncryptionData(EncryptionData);
-}
-
 void UNetConnection::SetEncryptionData(const FEncryptionData& EncryptionData)
 {
 	if (Handler.IsValid())
@@ -2226,9 +2198,6 @@ void UNetConnection::ReceivedAck(int32 AckPacketId, FChannelsToClose& OutChannel
 	OutAckPacketId = AckPacketId;
 
 	// Process the bunch.
-	PRAGMA_DISABLE_DEPRECATION_WARNINGS
-	LastRecvAckTime = Driver->GetElapsedTime();
-	PRAGMA_ENABLE_DEPRECATION_WARNINGS
 	LastRecvAckTimestamp = Driver->GetElapsedTime();
 
 	PacketAnalytics.TrackAck(AckPacketId);
@@ -5142,21 +5111,7 @@ void UNetConnection::SendChallengeControlMessage(const FEncryptionKeyResponse& R
 	{
 		if (Response.Response == EEncryptionResponse::Success)
 		{
-			// handle deprecated path where only the key is set
-			PRAGMA_DISABLE_DEPRECATION_WARNINGS
-			if ((Response.EncryptionKey.Num() > 0) && (Response.EncryptionData.Key.Num() == 0))
-			{
-				FEncryptionData ResponseData = Response.EncryptionData;
-				ResponseData.Key = Response.EncryptionKey;
-
-				EnableEncryptionServer(ResponseData);
-			}
-			PRAGMA_ENABLE_DEPRECATION_WARNINGS
-			else
-			{
-				EnableEncryptionServer(Response.EncryptionData);
-			}
-
+			EnableEncryptionServer(Response.EncryptionData);
 			SendChallengeControlMessage();
 		}
 		else
diff --git a/Engine/Source/Runtime/Engine/Private/NetDriver.cpp b/Engine/Source/Runtime/Engine/Private/NetDriver.cpp
--- a/Engine/Source/Runtime/Engine/Private/NetDriver.cpp
+++ b/Engine/Source/Runtime/Engine/Private/NetDriver.cpp
@@ -4745,9 +4745,31 @@ int64 UNetDriver::SendDestructionInfo(UNetConnection* Connection, FActorDestruct
 		if (Channel)
 		{
 			// Send a close bunch on the new channel
-			PRAGMA_DISABLE_DEPRECATION_WARNINGS
-			NumBits = Channel->SetChannelActorForDestroy(DestructionInfo);
-			PRAGMA_ENABLE_DEPRECATION_WARNINGS
+			if(!Channel->Closing && (Connection->GetConnectionState() == USOCK_Open || Connection->GetConnectionState() == USOCK_Pending))
+			{
+				// Outer must be valid to call PackageMap->WriteObject. In the case of streaming out levels, this can go null out of from underneath us. In that case, just skip the destruct info.
+				// We assume that if server unloads a level that clients will to and this will implicitly destroy all actors in it, so not worried about leaking actors client side here.
+				if (UObject* ObjOuter = DestructionInfo->ObjOuter.Get())
+				{
+					// Send a close notify, and wait for ack.
+					FOutBunch CloseBunch(Channel, true);
+					check(!CloseBunch.IsError());
+					check(CloseBunch.bClose);
+					CloseBunch.bReliable = 1;
+					CloseBunch.CloseReason = DestructionInfo->Reason;
+
+					// Serialize DestructInfo
+					NET_CHECKSUM(CloseBunch); // This is to mirror the Checksum in UPackageMapClient::SerializeNewActor
+					Connection->PackageMap->WriteObject(CloseBunch, ObjOuter, DestructionInfo->NetGUID, DestructionInfo->PathName);
+
+					UE_LOG(LogNetTraffic, Log, TEXT("SendDestructionInfo: Channel %d. NetGUID <%s> Path: %s. Bits: %d"), Channel->ChIndex, *DestructionInfo->NetGUID.ToString(), *DestructionInfo->PathName, CloseBunch.GetNumBits());
+					UE_LOG(LogNetDormancy, Verbose, TEXT("SendDestructionInfo: Channel %d. NetGUID <%s> Path: %s. Bits: %d"), Channel->ChIndex, *DestructionInfo->NetGUID.ToString(), *DestructionInfo->PathName, CloseBunch.GetNumBits());
+
+					Channel->SendBunch(&CloseBunch, false);
+
+					NumBits = CloseBunch.GetNumBits();
+				}
+			}
 		}
 	}
 
diff --git a/Engine/Source/Runtime/Engine/Private/PendingNetGame.cpp b/Engine/Source/Runtime/Engine/Private/PendingNetGame.cpp
--- a/Engine/Source/Runtime/Engine/Private/PendingNetGame.cpp
+++ b/Engine/Source/Runtime/Engine/Private/PendingNetGame.cpp
@@ -391,20 +391,7 @@ void UPendingNetGame::FinalizeEncryptedConnection(const FEncryptionKeyResponse&
 		{
 			if (Response.Response == EEncryptionResponse::Success)
 			{
-				// handle deprecated path where only the key is set
-				PRAGMA_DISABLE_DEPRECATION_WARNINGS
-				if ((Response.EncryptionKey.Num() > 0) && (Response.EncryptionData.Key.Num() == 0))
-				{
-					FEncryptionData ResponseData = Response.EncryptionData;
-					ResponseData.Key = Response.EncryptionKey;
-
-					Connection->EnableEncryption(ResponseData);
-				}
-				PRAGMA_ENABLE_DEPRECATION_WARNINGS
-				else
-				{
-					Connection->EnableEncryption(Response.EncryptionData);
-				}
+				Connection->EnableEncryption(Response.EncryptionData);
 			}
 			else
 			{
diff --git a/Engine/Source/Runtime/PacketHandlers/PacketHandler/Public/PacketHandler.h b/Engine/Source/Runtime/PacketHandlers/PacketHandler/Public/PacketHandler.h
--- a/Engine/Source/Runtime/PacketHandlers/PacketHandler/Public/PacketHandler.h
+++ b/Engine/Source/Runtime/PacketHandlers/PacketHandler/Public/PacketHandler.h
@@ -805,11 +805,6 @@ public:
 	{
 	}
 
-	UE_DEPRECATED(4.26, "Remove this function (no longer pure virtual, same with OutgoingConnectionless), or use the version which takes FReceivedPacketView.")
-	virtual void IncomingConnectionless(const TSharedPtr<const FInternetAddr>& Address, FBitReader& Packet)
-	{
-	}
-
 	/**
 	 * Handles incoming packets not associated with a UNetConnection
 	 *
