Fixed bugs in FNetGUIDCache (PIE/-game only)
- FNetGUIDCache::CanClientLoadObject didn't properly handle external packages
- FNetGUIDCache::GetObjectFromNetGUID wasn't considering packages with bIsDynamicPIEPackagePending == true
- FNetGUIDCache::ShouldIgnoreWhenMissing fixed inverse logic when testing IsDynamicPIEPackagePending()

diff --git a/Engine/Source/Runtime/Engine/Private/PackageMapClient.cpp b/Engine/Source/Runtime/Engine/Private/PackageMapClient.cpp
--- a/Engine/Source/Runtime/Engine/Private/PackageMapClient.cpp
+++ b/Engine/Source/Runtime/Engine/Private/PackageMapClient.cpp
@@ -760,9 +760,22 @@ bool FNetGUIDCache::CanClientLoadObject( const UObject* Object, const FNetworkGU
 
 	// PackageMapClient can't load maps, we must wait for the client to load the map when ready
 	// These guids are special guids, where the guid and all child guids resolve once the map has been loaded
-	if ( Object != nullptr && Object->GetOutermost()->ContainsMap() )
+	if (Object)
 	{
-		return false;
+		if (Object->GetPackage()->ContainsMap())
+		{
+			return false;
+		}
+
+#if WITH_EDITOR
+		// For objects using external package, we need to do the test on the package of their outer most object
+		// (this is currently only possible in Editor)
+		UObject* OutermostObject = Object->GetOutermostObject();
+		if (OutermostObject && OutermostObject->GetPackage()->ContainsMap())
+		{
+			return false;
+		}
+#endif
 	}
 
 	// If the object is null, we can't check whether the outermost contains a map anymore, so
@@ -3309,13 +3322,20 @@ UObject* FNetGUIDCache::GetObjectFromNetGUID( const FNetworkGUID& NetGUID, const
 	// At this point, we either have an outer, or we are a package
 	check( !CacheObjectPtr->bIsPending );
 
-	if (!ensure(ObjOuter == nullptr || ObjOuter->GetOutermost()->IsFullyLoaded() || ObjOuter->GetOutermost()->HasAnyPackageFlags(TreatAsLoadedFlags)))
+	if (!ensure(ObjOuter == nullptr || ObjOuter->GetPackage()->IsFullyLoaded() || ObjOuter->GetPackage()->HasAnyPackageFlags(TreatAsLoadedFlags)))
 	{
 		UE_LOG( LogNetPackageMap, Error, TEXT( "GetObjectFromNetGUID: Outer is null or package is not fully loaded.  FullNetGUIDPath: %s Outer: %s" ), *FullNetGUIDPath( NetGUID ), *GetFullNameSafe(ObjOuter) );
 	}
 
 	// See if this object is in memory
 	Object = FindObjectFast<UObject>(ObjOuter, CacheObjectPtr->PathName);
+#if WITH_EDITOR
+	// Object must be null if the package is a dynamic PIE package with pending external objects still loading, as it would normally while object is async loading
+	if (Object && Object->GetPackage()->IsDynamicPIEPackagePending())
+	{
+		Object = NULL;
+	}
+#endif
 
 	// Assume this is a package if the outer is invalid and this is a static guid
 	const bool bIsPackage = NetGUID.IsStatic() && !CacheObjectPtr->OuterGUID.IsValid();
@@ -3528,13 +3548,13 @@ bool FNetGUIDCache::ShouldIgnoreWhenMissing( const FNetworkGUID& NetGUID ) const
 		{
 #if WITH_EDITOR
 			// Ignore if the package is a dynamic PIE package with pending external objects still loading
-			if ( !OutermostCacheObject->Object->GetOutermost()->IsDynamicPIEPackagePending() )
+			if ( OutermostCacheObject->Object->GetPackage()->IsDynamicPIEPackagePending() )
 			{
 				return true;
 			}
 #endif
 			// Sometimes, other systems async load packages, which we don't track, but still must be aware of
-			if ( !OutermostCacheObject->Object->GetOutermost()->IsFullyLoaded() )
+			if ( !OutermostCacheObject->Object->GetPackage()->IsFullyLoaded() )
 			{
 				return true;
 			}
