World Partition: Move all world partition related properties into the "World Partition Setup" category

diff --git a/Engine/Source/Runtime/Engine/Public/WorldPartition/WorldPartition.h b/Engine/Source/Runtime/Engine/Public/WorldPartition/WorldPartition.h
--- a/Engine/Source/Runtime/Engine/Public/WorldPartition/WorldPartition.h
+++ b/Engine/Source/Runtime/Engine/Public/WorldPartition/WorldPartition.h
@@ -260,7 +260,7 @@ private:
 	bool bShouldCheckEnableStreamingWarning;
 
 	/** Whether Level Instance can reference this partition. */
-	UPROPERTY(EditAnywhere, Category = "WorldPartition", AdvancedDisplay, meta = (EditConditionHides, HideEditConditionToggle, EditCondition = "!bEnableStreaming"))
+	UPROPERTY(EditAnywhere, Category = "WorldPartitionSetup", AdvancedDisplay, meta = (EditConditionHides, HideEditConditionToggle, EditCondition = "!bEnableStreaming"))
 	bool bCanBeUsedByLevelInstance;
 #endif
 
@@ -285,7 +285,7 @@ private:
 
 #if WITH_EDITORONLY_DATA
 	// Default HLOD layer
-	UPROPERTY(EditAnywhere, Category=WorldPartition, meta = (DisplayName = "Default HLOD Layer", EditCondition="bEnableStreaming", EditConditionHides, HideEditConditionToggle))
+	UPROPERTY(EditAnywhere, Category=WorldPartitionSetup, meta = (DisplayName = "Default HLOD Layer", EditCondition="bEnableStreaming", EditConditionHides, HideEditConditionToggle))
 	TObjectPtr<class UHLODLayer> DefaultHLODLayer;
 
 	TArray<FWorldPartitionReference> LoadedSubobjects;
