Fixed bug when enabling world partition server streaming : NetConnection was closed because it detected a missing level package.
- Changed UE_NetConnectionPrivate::ValidateLevelVisibility to rely on FLevelUtils::IsValidStreamingLevel instead of FLevelUtils::FindStreamingLevel to determine if a streaming level package name is valid or not.
- Added FLevelUtils::IsValidStreamingLevel which tests for a valid StreamingLevel and also handles world partition generated memory/PIE packages.

diff --git a/Engine/Source/Runtime/Engine/Private/LevelUtils.cpp b/Engine/Source/Runtime/Engine/Private/LevelUtils.cpp
--- a/Engine/Source/Runtime/Engine/Private/LevelUtils.cpp
+++ b/Engine/Source/Runtime/Engine/Private/LevelUtils.cpp
@@ -15,6 +15,7 @@
 
 #if WITH_EDITOR
 #include "ScopedTransaction.h"
+#include "WorldPartition/WorldPartition.h"
 #endif
 
 #define LOCTEXT_NAMESPACE "LevelUtils"
@@ -77,6 +78,23 @@ ULevelStreaming* FLevelUtils::FindStreamingLevel(UWorld* InWorld, const TCHAR* I
 	return FindStreamingLevel(InWorld, FName(InPackageName));
 }
 
+bool FLevelUtils::IsValidStreamingLevel(UWorld* InWorld, const TCHAR* InPackageName)
+{
+	if (FindStreamingLevel(InWorld, InPackageName))
+	{
+		return true;
+	}
+
+#if WITH_EDITOR
+	if (UWorldPartition* WorldPartition = InWorld ? InWorld->GetWorldPartition() : nullptr)
+	{
+		return WorldPartition->IsValidPackageName(InPackageName);
+	}
+#endif
+
+	return false;
+}
+
 /////////////////////////////////////////////////////////////////////////////////////////
 //
 //	Level locking/unlocking.
diff --git a/Engine/Source/Runtime/Engine/Private/NetConnection.cpp b/Engine/Source/Runtime/Engine/Private/NetConnection.cpp
--- a/Engine/Source/Runtime/Engine/Private/NetConnection.cpp
+++ b/Engine/Source/Runtime/Engine/Private/NetConnection.cpp
@@ -138,8 +138,17 @@ namespace UE_NetConnectionPrivate
 		Result.Linker = FLinkerLoad::FindExistingLinkerForPackage(Result.Package);
 
 		Result.StreamingLevel = FLevelUtils::FindStreamingLevel(World, LevelVisibility.PackageName);
-		Result.bLevelExists = Result.Linker || Result.StreamingLevel != nullptr ||
-			(!FPackageName::IsMemoryPackage(PackageNameStr) && FPackageName::DoesPackageExist(LevelVisibility.FileName.ToString()));
+		Result.bLevelExists = (Result.Linker != nullptr) || (Result.StreamingLevel != nullptr);
+
+		if (!Result.bLevelExists)
+		{
+			Result.bLevelExists = FLevelUtils::IsValidStreamingLevel(World, *PackageNameStr);
+			
+			if (!Result.bLevelExists)
+			{
+				Result.bLevelExists = (!FPackageName::IsMemoryPackage(PackageNameStr) && FPackageName::DoesPackageExist(LevelVisibility.FileName.ToString()));
+			}
+		}
 
 		return Result;
 	}
diff --git a/Engine/Source/Runtime/Engine/Private/WorldPartition/WorldPartition.cpp b/Engine/Source/Runtime/Engine/Private/WorldPartition/WorldPartition.cpp
--- a/Engine/Source/Runtime/Engine/Private/WorldPartition/WorldPartition.cpp
+++ b/Engine/Source/Runtime/Engine/Private/WorldPartition/WorldPartition.cpp
@@ -299,6 +299,20 @@ void UWorldPartition::OnPackageDirtyStateChanged(UPackage* Package)
 	}
 }
 
+// Returns whether the memory package is part of the known/valid package names
+// used by World Partition for PIE/-game streaming.
+bool UWorldPartition::IsValidPackageName(const FString& InPackageName)
+{
+	if (FPackageName::IsMemoryPackage(InPackageName))
+	{
+		// Remove PIE prefix
+		FString PackageName = UWorld::RemovePIEPrefix(InPackageName);
+		// Test if package is a valid world partition PIE package
+		return GeneratedStreamingPackageNames.Contains(PackageName);
+	}
+	return false;
+}
+
 void UWorldPartition::OnPreBeginPIE(bool bStartSimulate)
 {
 	check(!bIsPIE);
@@ -315,7 +329,18 @@ void UWorldPartition::OnPrePIEEnded(bool bWasSimulatingInEditor)
 
 void UWorldPartition::OnBeginPlay()
 {
-	GenerateStreaming();
+	TArray<FString> OutGeneratedStreamingPackageNames;
+	GenerateStreaming((bIsPIE || IsRunningGame()) ? &OutGeneratedStreamingPackageNames : nullptr);
+
+	// Prepare GeneratedStreamingPackages
+	check(GeneratedStreamingPackageNames.IsEmpty());
+	for (const FString& PackageName : OutGeneratedStreamingPackageNames)
+	{
+		// Set as memory package to avoid wasting time in UWorldPartition::IsValidPackageName (GenerateStreaming for PIE runs on the editor world)
+		FString Package = FPaths::RemoveDuplicateSlashes(FPackageName::IsMemoryPackage(PackageName) ? PackageName : TEXT("/Memory/") + PackageName);
+		GeneratedStreamingPackageNames.Add(Package);
+	}
+
 	RuntimeHash->OnBeginPlay();
 }
 
@@ -1035,6 +1060,9 @@ void UWorldPartition::Serialize(FArchive& Ar)
 	if (Ar.GetPortFlags() & PPF_DuplicateForPIE)
 	{
 		Ar << StreamingPolicy;
+#if WITH_EDITORONLY_DATA
+		Ar << GeneratedStreamingPackageNames;
+#endif
 
 #if WITH_EDITOR
 		if (Ar.IsLoading())
diff --git a/Engine/Source/Runtime/Engine/Private/WorldPartition/WorldPartitionRuntimeSpatialHash.cpp b/Engine/Source/Runtime/Engine/Private/WorldPartition/WorldPartitionRuntimeSpatialHash.cpp
--- a/Engine/Source/Runtime/Engine/Private/WorldPartition/WorldPartitionRuntimeSpatialHash.cpp
+++ b/Engine/Source/Runtime/Engine/Private/WorldPartition/WorldPartitionRuntimeSpatialHash.cpp
@@ -995,6 +995,12 @@ bool UWorldPartitionRuntimeSpatialHash::CreateStreamingGrid(const FSpatialHashRu
 {
 	TRACE_CPUPROFILER_EVENT_SCOPE(CreateStreamingGrid);
 
+	if (IsRunningCookCommandlet() && !OutPackagesToGenerate)
+	{
+		UE_LOG(LogWorldPartition, Error, TEXT("Error creating runtime streaming cells for cook, OutPackagesToGenerate is null."));
+		return false;
+	}
+
 	UWorldPartition* WorldPartition = GetOuterUWorldPartition();
 	UWorld* World = WorldPartition->GetWorld();
 	const bool bIsMainWorldPartition = (World == WorldPartition->GetTypedOuter<UWorld>());
@@ -1161,25 +1167,18 @@ bool UWorldPartitionRuntimeSpatialHash::CreateStreamingGrid(const FSpatialHashRu
 				}
 				StreamingCell->SetMinMaxZ(CellMinMaxZ);
 
-				if (IsRunningCookCommandlet())
+				// Always loaded cell actors are transfered to World's Persistent Level (see UWorldPartitionRuntimeSpatialHash::PopulateGeneratorPackageForCook)
+				if (StreamingCell->GetActorCount() && !StreamingCell->IsAlwaysLoaded())
 				{
-					UE_LOG(LogWorldPartition, Log, TEXT("Creating runtime streaming cells %s."), *StreamingCell->GetName());
-
-					if (StreamingCell->GetActorCount())
+					if (OutPackagesToGenerate)
 					{
-						// Always loaded cell actors are transfered to World's Persistent Level (see UWorldPartitionRuntimeSpatialHash::PopulateGeneratorPackageForCook)
-						if (!StreamingCell->IsAlwaysLoaded())
-						{
-							if (!OutPackagesToGenerate)
-							{
-								UE_LOG(LogWorldPartition, Error, TEXT("Error creating runtime streaming cells for cook, OutPackagesToGenerate is null."));
-								return false;
-							}
-
-							const FString PackageRelativePath = StreamingCell->GetPackageNameToCreate();
-							check(!PackageRelativePath.IsEmpty());
-							OutPackagesToGenerate->Add(PackageRelativePath);
+						const FString PackageRelativePath = StreamingCell->GetPackageNameToCreate();
+						check(!PackageRelativePath.IsEmpty());
+						OutPackagesToGenerate->Add(PackageRelativePath);
 
+						if (IsRunningCookCommandlet())
+						{
+							UE_LOG(LogWorldPartition, Log, TEXT("Creating runtime streaming cells %s."), *StreamingCell->GetName());
 							// Map relative package to StreamingCell for PopulateGeneratedPackageForCook/PopulateGeneratorPackageForCook
 							PackagesToGenerateForCook.Add(PackageRelativePath, StreamingCell);
 						}
diff --git a/Engine/Source/Runtime/Engine/Private/WorldPartition/WorldPartitionStreamingGeneration.cpp b/Engine/Source/Runtime/Engine/Private/WorldPartition/WorldPartitionStreamingGeneration.cpp
--- a/Engine/Source/Runtime/Engine/Private/WorldPartition/WorldPartitionStreamingGeneration.cpp
+++ b/Engine/Source/Runtime/Engine/Private/WorldPartition/WorldPartitionStreamingGeneration.cpp
@@ -920,6 +920,7 @@ void UWorldPartition::FlushStreaming()
 {
 	RuntimeHash->FlushStreaming();
 	StreamingPolicy = nullptr;
+	GeneratedStreamingPackageNames.Empty();
 }
 
 void UWorldPartition::GenerateHLOD(ISourceControlHelper* SourceControlHelper, bool bCreateActorsOnly)
diff --git a/Engine/Source/Runtime/Engine/Public/LevelUtils.h b/Engine/Source/Runtime/Engine/Public/LevelUtils.h
--- a/Engine/Source/Runtime/Engine/Public/LevelUtils.h
+++ b/Engine/Source/Runtime/Engine/Public/LevelUtils.h
@@ -43,6 +43,17 @@ public:
 	 */
 	static ULevelStreaming* FindStreamingLevel(UWorld* InWorld, const TCHAR* PackageName);
 
+	/**
+	 * Returns whether the given package is referenced by one of the world streaming levels or not.
+	 *
+	 * @param		InWorld			World to look in for the streaming level
+	 * @param		InPackageName	Name of the package containing the ULevel to query
+	 * @return						True if the given package is referenced by one of
+	 *								the world streaming levels, else False.
+	 */
+	static bool IsValidStreamingLevel(UWorld* InWorld, const TCHAR* InPackageName);
+
+
 	///////////////////////////////////////////////////////////////////////////
 	// Locking/unlocking levels for edit.
 
diff --git a/Engine/Source/Runtime/Engine/Public/WorldPartition/WorldPartition.h b/Engine/Source/Runtime/Engine/Public/WorldPartition/WorldPartition.h
--- a/Engine/Source/Runtime/Engine/Public/WorldPartition/WorldPartition.h
+++ b/Engine/Source/Runtime/Engine/Public/WorldPartition/WorldPartition.h
@@ -159,6 +159,7 @@ public:
 	FWorldPartitionGenerateStreamingDelegate OnPreGenerateStreaming;
 
 	void RemapSoftObjectPath(FSoftObjectPath& ObjectPath);
+	bool IsValidPackageName(const FString& InPackageName);
 
 	// Cooking
 	bool PopulateGeneratorPackageForCook(const TArray<ICookPackageSplitter::FGeneratedPackageForPreSave>& InGeneratedPackages, TArray<UPackage*>& OutModifiedPackages);
@@ -291,6 +292,8 @@ private:
 	TArray<FWorldPartitionReference> LoadedSubobjects;
 
 	TMap<FWorldPartitionReference, AActor*> DirtyActors;
+
+	TSet<FString> GeneratedStreamingPackageNames;
 #endif
 
 	EWorldPartitionInitState InitState;
