World Partition - HLOD: Prevent HLODs from being runtime relevant if streaming is disabled

diff --git a/Engine/Source/Runtime/Engine/Public/WorldPartition/DataLayer/WorldDataLayersActorDesc.h b/Engine/Source/Runtime/Engine/Public/WorldPartition/DataLayer/WorldDataLayersActorDesc.h
--- a/Engine/Source/Runtime/Engine/Public/WorldPartition/DataLayer/WorldDataLayersActorDesc.h
+++ b/Engine/Source/Runtime/Engine/Public/WorldPartition/DataLayer/WorldDataLayersActorDesc.h
@@ -58,11 +58,13 @@ public:
 	const FDataLayerInstanceDesc* GetDataLayerInstanceFromAssetPath(FName InDataLayerAssetPath) const;
 
 protected:
+	//~ Begin FWorldPartitionActorDesc Interface.
 	virtual void Init(const AActor* InActor) override;
 	virtual bool Equals(const FWorldPartitionActorDesc* Other) const override;
 	virtual void Serialize(FArchive& Ar) override;
 	virtual bool IsResaveNeeded() const override { return !IsValid(); }
-	virtual bool IsRuntimeRelevant(const FActorContainerID& InContainerID) const;
+	virtual bool IsRuntimeRelevant(const FActorContainerID& InContainerID) const override;
+	//~ End FWorldPartitionActorDesc Interface.
 
 private:
 	TArray<FDataLayerInstanceDesc> DataLayerInstances;
diff --git a/Engine/Source/Runtime/Engine/Public/WorldPartition/HLOD/HLODActorDesc.h b/Engine/Source/Runtime/Engine/Public/WorldPartition/HLOD/HLODActorDesc.h
--- a/Engine/Source/Runtime/Engine/Public/WorldPartition/HLOD/HLODActorDesc.h
+++ b/Engine/Source/Runtime/Engine/Public/WorldPartition/HLOD/HLODActorDesc.h
@@ -26,12 +26,14 @@ public:
 
 	static uint64 ComputeCellHash(const FString HLODLayerName, uint64 GridIndexX, uint64 GridIndexY, uint64 GridIndexZ, FDataLayersID DataLayersID);
 
-	virtual bool ShouldBeLoadedByEditor() const { return false; }
-
 protected:
+	//~ Begin FWorldPartitionActorDesc Interface.
 	virtual void Init(const AActor* InActor) override;
 	virtual bool Equals(const FWorldPartitionActorDesc* Other) const override;
 	virtual void Serialize(FArchive& Ar) override;
+	virtual bool ShouldBeLoadedByEditor() const override { return false; }
+	virtual bool IsRuntimeRelevant(const FActorContainerID& InContainerID) const override { return !bIsForcedNonSpatiallyLoaded; }
+	//~ End FWorldPartitionActorDesc Interface.
 
 	TArray<FHLODSubActorDesc> HLODSubActors;
 
