Implement GatherProducePageDataTasks() for FVirtualTextureLevelRedirector.
Without this RVTs with streaming mips don't do anything.
That means that we don't wait on outstanding tasks, which can lead to various crashes when we destroy before tasks are complete.

diff --git a/Engine/Source/Runtime/Engine/Private/VT/VirtualTextureLevelRedirector.cpp b/Engine/Source/Runtime/Engine/Private/VT/VirtualTextureLevelRedirector.cpp
--- a/Engine/Source/Runtime/Engine/Private/VT/VirtualTextureLevelRedirector.cpp
+++ b/Engine/Source/Runtime/Engine/Private/VT/VirtualTextureLevelRedirector.cpp
@@ -41,3 +41,19 @@ IVirtualTextureFinalizer* FVirtualTextureLevelRedirector::ProducePageData(
 	int32 vLevelOffset = vLevel < TransitionLevel ? 0 : TransitionLevel;
 	return VirtualTextures[VirtualTextureIndex]->ProducePageData(RHICmdList, FeatureLevel, Flags, ProducerHandle, LayerMask, vLevel - vLevelOffset, vAddress, RequestHandle, TargetLayers);
 }
+
+void FVirtualTextureLevelRedirector::GatherProducePageDataTasks(
+	FVirtualTextureProducerHandle const& ProducerHandle, 
+	FGraphEventArray& InOutTasks) const
+{
+	VirtualTextures[0]->GatherProducePageDataTasks(ProducerHandle, InOutTasks);
+	VirtualTextures[1]->GatherProducePageDataTasks(ProducerHandle, InOutTasks);
+}
+
+void FVirtualTextureLevelRedirector::GatherProducePageDataTasks(
+	uint64 RequestHandle,
+	FGraphEventArray& InOutTasks) const
+{
+	VirtualTextures[0]->GatherProducePageDataTasks(RequestHandle, InOutTasks);
+	VirtualTextures[1]->GatherProducePageDataTasks(RequestHandle, InOutTasks);
+}
diff --git a/Engine/Source/Runtime/Engine/Private/VT/VirtualTextureLevelRedirector.h b/Engine/Source/Runtime/Engine/Private/VT/VirtualTextureLevelRedirector.h
--- a/Engine/Source/Runtime/Engine/Private/VT/VirtualTextureLevelRedirector.h
+++ b/Engine/Source/Runtime/Engine/Private/VT/VirtualTextureLevelRedirector.h
@@ -35,6 +35,16 @@ public:
 		uint64 RequestHandle,
 		const FVTProduceTargetLayer* TargetLayers
 	) override;
+
+	virtual void GatherProducePageDataTasks(
+		FVirtualTextureProducerHandle const& ProducerHandle, 
+		FGraphEventArray& InOutTasks
+	) const override;
+
+	virtual void GatherProducePageDataTasks(
+		uint64 RequestHandle, 
+		FGraphEventArray& InOutTasks
+	) const override;
 	//~ End IVirtualTexture Interface.
 
 private:
