Don't modify QueuedBits if IsInternalAck()

diff --git a/Engine/Source/Runtime/Engine/Private/NetConnection.cpp b/Engine/Source/Runtime/Engine/Private/NetConnection.cpp
--- a/Engine/Source/Runtime/Engine/Private/NetConnection.cpp
+++ b/Engine/Source/Runtime/Engine/Private/NetConnection.cpp
@@ -2056,16 +2056,19 @@ void UNetConnection::FlushNet(bool bIgnoreSimulation)
 		
 		++OutPacketId; 
 
-		int32 NewQueuedBits = 0;
-		const bool bWouldOverflow = UE_NetConnectionPrivate::Add_DetectOverflow_Clamp(QueuedBits, PacketBytes * 8, NewQueuedBits);
-		
-		if (bWouldOverflow && !bLoggedFlushNetQueuedBitsOverflow)
+		if (!IsInternalAck())
 		{
-			UE_LOG(LogNet, Log, TEXT("UNetConnection::FlushNet: QueuedBits overflow detected! QueuedBits: %d, PacketBytes: %d. %s"), QueuedBits, PacketBytes, *Describe());
-			bLoggedFlushNetQueuedBitsOverflow = true;
-		}
+			int32 NewQueuedBits = 0;
+			const bool bWouldOverflow = UE_NetConnectionPrivate::Add_DetectOverflow_Clamp(QueuedBits, PacketBytes * 8, NewQueuedBits);
 
-		QueuedBits = NewQueuedBits;
+			if (bWouldOverflow && !bLoggedFlushNetQueuedBitsOverflow)
+			{
+				UE_LOG(LogNet, Log, TEXT("UNetConnection::FlushNet: QueuedBits overflow detected! QueuedBits: %d, PacketBytes: %d. %s"), QueuedBits, PacketBytes, *Describe());
+				bLoggedFlushNetQueuedBitsOverflow = true;
+			}
+
+			QueuedBits = NewQueuedBits;
+		}
 
 		OutBytes += PacketBytes;
 		OutTotalBytes += PacketBytes;
@@ -2173,6 +2176,11 @@ bool UNetConnection::ShouldDropOutgoingPacketForLossSimulation(int64 NumBits) co
 
 int32 UNetConnection::IsNetReady(bool Saturate)
 {
+	if (IsInternalAck())
+	{
+		return 1;
+	}
+
 	// Return whether we can send more data without saturation the connection.
 	if (Saturate)
 	{
@@ -2182,7 +2190,7 @@ int32 UNetConnection::IsNetReady(bool Saturate)
 #if !(UE_BUILD_SHIPPING || UE_BUILD_TEST)
 	if (CVarDisableBandwithThrottling.GetValueOnAnyThread() > 0)
 	{
-		return true;
+		return 1;
 	}
 #endif
 
@@ -4396,36 +4404,39 @@ void UNetConnection::Tick(float DeltaSeconds)
 
 	SaturationAnalytics.TrackFrame(!IsNetReady(false));
 
-	// Clamp DeltaTime for bandwidth limiting so that if there is a hitch, we don't try to send
-	// a large burst on the next frame, which can cause another hitch if a lot of additional replication occurs.
-	float BandwidthDeltaTime = DeltaTime;
-	if (DesiredTickRate != 0.0f)
+	if (!IsInternalAck())
 	{
-		BandwidthDeltaTime = FMath::Clamp(BandwidthDeltaTime, 0.0f, 1.0f / DesiredTickRate);
-	}
+		// Clamp DeltaTime for bandwidth limiting so that if there is a hitch, we don't try to send
+		// a large burst on the next frame, which can cause another hitch if a lot of additional replication occurs.
+		float BandwidthDeltaTime = DeltaTime;
+		if (DesiredTickRate != 0.0f)
+		{
+			BandwidthDeltaTime = FMath::Clamp(BandwidthDeltaTime, 0.0f, 1.0f / DesiredTickRate);
+		}
+
+		float DeltaBits = CurrentNetSpeed * BandwidthDeltaTime * 8.f;
 
-	float DeltaBits = CurrentNetSpeed * BandwidthDeltaTime * 8.f;
+		int32 NewQueuedBits = 0;
+		const int64 DeltaQueuedBits = -FMath::TruncToInt(DeltaBits);
+		const bool bWouldOverflow = UE_NetConnectionPrivate::Add_DetectOverflow_Clamp(QueuedBits, DeltaQueuedBits, NewQueuedBits);
 
-	int32 NewQueuedBits = 0;
-	const int64 DeltaQueuedBits = -FMath::TruncToInt(DeltaBits);
-	const bool bWouldOverflow = UE_NetConnectionPrivate::Add_DetectOverflow_Clamp(QueuedBits, DeltaQueuedBits, NewQueuedBits);
+		ensureMsgf(!bWouldOverflow, TEXT("UNetConnection::Tick: QueuedBits overflow detected! QueuedBits: %d, change: %lld, BandwidthDeltaTime: %.4f, DesiredTickRate: %.2f"),
+			QueuedBits, DeltaQueuedBits, BandwidthDeltaTime, DesiredTickRate);
 
-	ensureMsgf(!bWouldOverflow, TEXT("UNetConnection::Tick: QueuedBits overflow detected! QueuedBits: %d, change: %lld, BandwidthDeltaTime: %.4f, DesiredTickRate: %.2f"),
-		QueuedBits, DeltaQueuedBits, BandwidthDeltaTime, DesiredTickRate);
+		QueuedBits = NewQueuedBits;
 
-	QueuedBits = NewQueuedBits;
+		const int64 AllowedLag = 2 * DeltaQueuedBits;
 
-	const int64 AllowedLag = 2 * DeltaQueuedBits;
-	
-	if (QueuedBits < AllowedLag)
-	{
-		int32 NewAllowedLag = 0;
-		const bool bAllowedLagOverflow = UE_NetConnectionPrivate::Add_DetectOverflow_Clamp(0, AllowedLag, NewAllowedLag);
+		if (QueuedBits < AllowedLag)
+		{
+			int32 NewAllowedLag = 0;
+			const bool bAllowedLagOverflow = UE_NetConnectionPrivate::Add_DetectOverflow_Clamp(0, AllowedLag, NewAllowedLag);
 
-		ensureMsgf(!bAllowedLagOverflow, TEXT("UNetConnection::Tick: AllowedLag overflow detected! AllowedLag: %lld, BandwidthDeltaTime: %.4f, DesiredTickRate: %.2f"),
-			AllowedLag, BandwidthDeltaTime, DesiredTickRate);
+			ensureMsgf(!bAllowedLagOverflow, TEXT("UNetConnection::Tick: AllowedLag overflow detected! AllowedLag: %lld, BandwidthDeltaTime: %.4f, DesiredTickRate: %.2f"),
+				AllowedLag, BandwidthDeltaTime, DesiredTickRate);
 
-		QueuedBits = NewAllowedLag;
+			QueuedBits = NewAllowedLag;
+		}
 	}
 
 	bFlushedNetThisFrame = false;
