Add net.DeltaInitialFastArrayElements to allow fast array struct deltas to be generated on initial sends, disabled by default

diff --git a/Engine/Source/Runtime/Engine/Private/RepLayout.cpp b/Engine/Source/Runtime/Engine/Private/RepLayout.cpp
--- a/Engine/Source/Runtime/Engine/Private/RepLayout.cpp
+++ b/Engine/Source/Runtime/Engine/Private/RepLayout.cpp
@@ -95,6 +95,12 @@ constexpr bool GbPushModelValidateProperties = false;
 
 #endif
 
+namespace UE::Net::Private
+{
+	static bool bDeltaInitialFastArrayElements = false;
+	static FAutoConsoleVariableRef CVarDeltaInitialFastArrayElements(TEXT("net.DeltaInitialFastArrayElements"), bDeltaInitialFastArrayElements, TEXT("If true, send delta struct changelists for initial fast array elements."));
+}
+
 extern int32 GNumSharedSerializationHit;
 extern int32 GNumSharedSerializationMiss;
 
@@ -7732,12 +7738,14 @@ ERepLayoutResult FRepLayout::DeltaSerializeFastArrayProperty(FFastArrayDeltaSeri
 							NewChangelist.Empty(1);
 							const int32 ArrayElementOffset = ElementSize * IDIndexPair.Idx;
 
+							const bool bForceFail = !UE::Net::Private::bDeltaInitialFastArrayElements && (bIsInitial || ShadowArrayItemIsNew[IDIndexPair.Idx]);
+
 							// Go ahead and do a property compare here, regardless of what we'll actually use below.
 							// This is to prevent issues where someone with an initial / outdated connection doesn't properly
 							// update the changelists in our history, but does update the shadow state inadvertently.
 							FComparePropertiesSharedParams SharedParams{
 								/*bIsInitial=*/ bIsInitial,
-								/*bForceFail=*/ bIsInitial || ShadowArrayItemIsNew[IDIndexPair.Idx],
+								/*bForceFail=*/ bForceFail,
 								Flags,
 								Parents,
 								Cmds,
@@ -7791,12 +7799,14 @@ ERepLayoutResult FRepLayout::DeltaSerializeFastArrayProperty(FFastArrayDeltaSeri
 				const uint32 LastHistory = LastAckedHistory;
 				const uint32 LastChangelistDelta = LastAckedChangelistDelta;
 
-				if (LastHistory != 0 && LastChangelistDelta > 0 && LastChangelistDelta < (FDeltaArrayHistoryState::MAX_CHANGE_HISTORY - 1))
+				const bool bAllowInitialHistory = UE::Net::Private::bDeltaInitialFastArrayElements || (LastHistory != 0);
+
+				if (bAllowInitialHistory && LastChangelistDelta > 0 && LastChangelistDelta < (FDeltaArrayHistoryState::MAX_CHANGE_HISTORY - 1))
 				{
 					const FConstRepObjectDataBuffer ConstObjectData(ObjectData);
 					Changelists.SetNum(ChangedElements.Num());
 
-					// Note, we iterate from LastSentHistory + 1, because we don't want to send something if
+					// Note, we iterate from LastAckedHistory + 1, because we don't want to send something if
 					// we think it's already been sent/received.
 					// Similarly, we do <= NewChangelistHistory because we need to send the newest history.
 					for (uint32 ChangelistHistory = LastHistory + 1; ChangelistHistory <= NewChangelistHistory; ++ChangelistHistory)
