Replace a cvar GetValueOnAnyThread with a direct variable reference

diff --git a/Engine/Source/Runtime/Engine/Private/DataChannel.cpp b/Engine/Source/Runtime/Engine/Private/DataChannel.cpp
--- a/Engine/Source/Runtime/Engine/Private/DataChannel.cpp
+++ b/Engine/Source/Runtime/Engine/Private/DataChannel.cpp
@@ -107,12 +107,6 @@ static FAutoConsoleVariableRef CVarDormancyHysteresis(
 	TEXT("When > 0, represents the time we'll wait before letting a channel become fully dormant (in seconds). This can prevent churn when objects are going in and out of dormant more frequently than normal.")
 );
 
-static TAutoConsoleVariable<bool> CVarEnableNetInitialSubObjects(
-	TEXT("net.EnableNetInitialSubObjects"),
-	true,
-	TEXT("Enables new SubObjects to set bNetInitial to true to make sure all replicated properties are replicated.")
-);
-
 namespace UE::Net
 {
 	extern int32 FilterGuidRemapping;
@@ -142,6 +136,13 @@ namespace UE::Net
 		TEXT("If enabled log all the errors detected by the CompareWithLegacy cheat. Otherwise only the first ensure triggered gets logged."));
 #endif
 
+	static bool bEnableNetInitialSubObjects = true;
+	static FAutoConsoleVariableRef CVarEnableNetInitialSubObjects(
+		TEXT("net.EnableNetInitialSubObjects"),
+		bEnableNetInitialSubObjects,
+		TEXT("Enables new SubObjects to set bNetInitial to true to make sure all replicated properties are replicated.")
+	);
+
 }; // namespace UE::Net
 
 template<typename T>
@@ -4888,7 +4889,7 @@ bool UActorChannel::WriteSubObjectInBunch(UObject* Obj, FOutBunch& Bunch, FRepli
 		// to spawn this on the client).
 		Bunch.bReliable = true;
 		NewSubobject = true;
-		if (CVarEnableNetInitialSubObjects.GetValueOnAnyThread())
+		if (UE::Net::bEnableNetInitialSubObjects)
 		{
 			ObjRepFlags.bNetInitial = true;
 		}	
