(courtesy of [REDACTED_PERSON])
- Fix for some DeviceProfiles not reading in their parent TextureLODGroups
- Remove the call to UDeviceProfileManager::Get() that was being called too early, before object subsystem is ready, so it was unable to load all the other platform DeviceProfiles

#tests preflight 631258e02b7fe03eb6d9d725

diff --git a/Engine/Source/Runtime/Core/Public/Misc/DelayedAutoRegister.h b/Engine/Source/Runtime/Core/Public/Misc/DelayedAutoRegister.h
--- a/Engine/Source/Runtime/Core/Public/Misc/DelayedAutoRegister.h
+++ b/Engine/Source/Runtime/Core/Public/Misc/DelayedAutoRegister.h
@@ -16,6 +16,7 @@ enum class EDelayedRegisterRunPhase : uint8
 	ShaderTypesReady,
 	PreObjectSystemReady,
 	ObjectSystemReady,
+	DeviceProfileManagerReady,
 	EndOfEngineInit,
 
 	NumPhases,
diff --git a/Engine/Source/Runtime/Engine/Classes/DeviceProfiles/DeviceProfile.h b/Engine/Source/Runtime/Engine/Classes/DeviceProfiles/DeviceProfile.h
--- a/Engine/Source/Runtime/Engine/Classes/DeviceProfiles/DeviceProfile.h
+++ b/Engine/Source/Runtime/Engine/Classes/DeviceProfiles/DeviceProfile.h
@@ -121,7 +121,6 @@ public:
 	void ValidateProfile();
 
 	//~ Begin UObject Interface
-	virtual void PostInitProperties() override;
 	virtual void BeginDestroy() override;
 	//~ End UObject Interface
 
diff --git a/Engine/Source/Runtime/Engine/Private/DeviceProfiles/DeviceProfile.cpp b/Engine/Source/Runtime/Engine/Private/DeviceProfiles/DeviceProfile.cpp
--- a/Engine/Source/Runtime/Engine/Private/DeviceProfiles/DeviceProfile.cpp
+++ b/Engine/Source/Runtime/Engine/Private/DeviceProfiles/DeviceProfile.cpp
@@ -97,12 +97,6 @@ UDeviceProfile* UDeviceProfile::GetParentProfile(bool bIncludeDefaultObject) con
 	return ParentProfile;
 }
 
-void UDeviceProfile::PostInitProperties()
-{
-	Super::PostInitProperties();
-	ValidateTextureLODGroups();
-}
-
 void UDeviceProfile::BeginDestroy()
 {
 	UE_LOG(LogDeviceProfile, Verbose, TEXT("Device profile begin destroy: [%p] %s"), this, *GetName());
diff --git a/Engine/Source/Runtime/Engine/Private/DeviceProfiles/DeviceProfileManager.cpp b/Engine/Source/Runtime/Engine/Private/DeviceProfiles/DeviceProfileManager.cpp
--- a/Engine/Source/Runtime/Engine/Private/DeviceProfiles/DeviceProfileManager.cpp
+++ b/Engine/Source/Runtime/Engine/Private/DeviceProfiles/DeviceProfileManager.cpp
@@ -51,6 +51,10 @@ TArray<FSelectedFragmentProperties> UDeviceProfileManager::PlatformFragmentsSele
 
 UDeviceProfileManager* UDeviceProfileManager::DeviceProfileManagerSingleton = nullptr;
 
+// when objects are ready, we can create the singleton properly, this makes it as early as possible, just in case other Object creation
+// didnt't create it along the way
+FDelayedAutoRegisterHelper GDPManagerSingletonHelper(EDelayedRegisterRunPhase::ObjectSystemReady, [] { UDeviceProfileManager::Get(); });
+
 UDeviceProfileManager& UDeviceProfileManager::Get(bool bFromPostCDOContruct)
 {
 	if (DeviceProfileManagerSingleton == nullptr)
@@ -93,6 +97,9 @@ UDeviceProfileManager& UDeviceProfileManager::Get(bool bFromPostCDOContruct)
 #if ALLOW_OTHER_PLATFORM_CONFIG
 		FCoreDelegates::GatherDeviceProfileCVars.BindLambda([](const FString& DeviceProfileName) { return UDeviceProfileManager::Get().GatherDeviceProfileCVars(DeviceProfileName, EDeviceProfileMode::DPM_CacheValues); });
 #endif
+
+		// let any other code that needs the DPManager to run now
+		FDelayedAutoRegisterHelper::RunAndClearDelayedAutoRegisterDelegates(EDelayedRegisterRunPhase::DeviceProfileManagerReady);
 	}
 	return *DeviceProfileManagerSingleton;
 }
@@ -663,9 +670,11 @@ UDeviceProfile* UDeviceProfileManager::CreateProfile(const FString& ProfileName,
 			// if the config needs to come from a platform, set it now, then reload the config
 			DeviceProfile->ConfigPlatform = ConfigPlatform;
 			DeviceProfile->LoadConfig();
-			DeviceProfile->ValidateProfile();
 		}
 
+		// make sure the DP has all the LODGroups it needs
+		DeviceProfile->ValidateProfile();
+
 		// if the config didn't specify a DeviceType, use the passed in one
 		if (DeviceProfile->DeviceType.IsEmpty())
 		{
@@ -1359,6 +1368,16 @@ public:
 				{
 					Ar.Logf(TEXT("%s = %s"), *Pair.Key, *Pair.Value);
 				}
+
+				// log out the LODGroups fully
+				FArrayProperty* LODGroupsProperty = FindFProperty<FArrayProperty>(UDeviceProfile::StaticClass(), GET_MEMBER_NAME_CHECKED(UTextureLODSettings, TextureLODGroups));
+				FScriptArrayHelper_InContainer ArrayHelper(LODGroupsProperty, DeviceProfile);
+				for (int32 Index = 0; Index < ArrayHelper.Num(); Index++)
+				{
+					FString	Buffer;
+					LODGroupsProperty->Inner->ExportTextItem_Direct(Buffer, ArrayHelper.GetRawPtr(Index), ArrayHelper.GetRawPtr(Index), DeviceProfile, 0);
+					Ar.Logf(TEXT("LODGroup[%d]=%s"), Index, *Buffer);
+				}
 			}
 		}
 		else if (FParse::Command(&Cmd, TEXT("dpdumppreview")))
diff --git a/Engine/Source/Runtime/Launch/Private/LaunchEngineLoop.cpp b/Engine/Source/Runtime/Launch/Private/LaunchEngineLoop.cpp
--- a/Engine/Source/Runtime/Launch/Private/LaunchEngineLoop.cpp
+++ b/Engine/Source/Runtime/Launch/Private/LaunchEngineLoop.cpp
@@ -2576,13 +2576,7 @@ PRAGMA_ENABLE_DEPRECATION_WARNINGS
 			// This may include scalability group settings which will override
 			// the defaults set above which can then be replaced below when
 			// the game user settings are loaded and applied.
-
-#if ALLOW_OTHER_PLATFORM_CONFIG
-			// Need to initialize the DeviceProfileManager Singleton to bind FCoreDelegates::GatherDeviceProfileCVars. This delegate is needed at cook time.
-			UDeviceProfileManager::Get().InitializeCVarsForActiveDeviceProfile();
-#else
 			UDeviceProfileManager::InitializeCVarsForActiveDeviceProfile();
-#endif
 		}
 
 		{
