World Partition: fix crash when saving deleted actors.

#rnx

diff --git a/Engine/Source/Runtime/Engine/Private/WorldPartition/WorldPartition.cpp b/Engine/Source/Runtime/Engine/Private/WorldPartition/WorldPartition.cpp
--- a/Engine/Source/Runtime/Engine/Private/WorldPartition/WorldPartition.cpp
+++ b/Engine/Source/Runtime/Engine/Private/WorldPartition/WorldPartition.cpp
@@ -284,16 +284,6 @@ void UWorldPartition::OnPackageDirtyStateChanged(UPackage* Package)
 				{
 					DirtyActors.Add(ActorHandle.ToReference(), Actor);
 				}
-				else
-				{
-					// If we hold the last reference to that actor (or no reference are held at all), pin it to avoid unloading
-					if (PinnedActors && (ActorHandle->GetHardRefCount() <= 1))
-					{
-						PinnedActors->AddActors({ ActorHandle });
-					}
-
-					DirtyActors.Remove(ActorHandle.ToReference());
-				}
 			}
 		}
 	}
@@ -1175,6 +1165,22 @@ void UWorldPartition::Tick(float DeltaSeconds)
 		EditorHash->Tick(DeltaSeconds);
 	}
 
+	if (PinnedActors)
+	{
+		for (TMap<FWorldPartitionReference, AActor*>::TIterator DirtyActorIt(DirtyActors); DirtyActorIt; ++DirtyActorIt)
+		{
+			if (!DirtyActorIt.Key().IsValid() || !DirtyActorIt.Value()->GetPackage()->IsDirty())
+			{
+				if (DirtyActorIt.Key().IsValid() && DirtyActorIt.Key()->GetHardRefCount() <= 1)
+				{
+					PinnedActors->AddActors({ DirtyActorIt.Key().ToHandle() });
+				}
+
+				DirtyActorIt.RemoveCurrent();
+			}
+		}
+	}
+
 	if (bForceGarbageCollection)
 	{
 		GEngine->ForceGarbageCollection(bForceGarbageCollectionPurge);
