ContainerPackages which have a different root than the WorldPackage (/Game/ vs /Plugin/) now properly add their remapping in the InstancingContext for PIE
Fix ContentBundle not being streamed

diff --git a/Engine/Source/Runtime/Engine/Private/WorldPartition/WorldPartitionLevelStreamingDynamic.cpp b/Engine/Source/Runtime/Engine/Private/WorldPartition/WorldPartitionLevelStreamingDynamic.cpp
--- a/Engine/Source/Runtime/Engine/Private/WorldPartition/WorldPartitionLevelStreamingDynamic.cpp
+++ b/Engine/Source/Runtime/Engine/Private/WorldPartition/WorldPartitionLevelStreamingDynamic.cpp
@@ -10,6 +10,7 @@
 #include "Misc/PackageName.h"
 #include "Engine/Engine.h"
 #include "Engine/World.h"
+#include "Misc/PathViews.h"
 #include "Model.h"
 #include "ContentStreaming.h"
 #include "GameFramework/WorldSettings.h"
@@ -282,11 +283,49 @@ bool UWorldPartitionLevelStreamingDynamic::IssueLoadRequests()
 	check(!bLoadRequestInProgress);
 	bLoadSucceeded = false;
 	bLoadRequestInProgress = true;
-	// Don't do SoftObjectPath remapping for PersistentLevel actors because references can end up in different cells
-	const bool bSoftObjectRemappingEnabled = false;
-	FLinkerInstancingContext InstancingContext(bSoftObjectRemappingEnabled);
-	UPackage* RuntimePackage = RuntimeLevel->GetPackage();
-	InstancingContext.AddPackageMapping(OriginalLevelPackageName, RuntimePackage->GetFName());
+
+	auto BuildInstancingContext = [this](FLinkerInstancingContext& OutLinkInstancingContext)
+	{
+		// Don't do SoftObjectPath remapping for PersistentLevel actors because references can end up in different cells
+		OutLinkInstancingContext.SetSoftObjectPathRemappingEnabled(false);
+
+		UPackage* RuntimePackage = RuntimeLevel->GetPackage();
+		OutLinkInstancingContext.AddPackageMapping(OriginalLevelPackageName, RuntimePackage->GetFName());
+
+		for (const FWorldPartitionRuntimeCellObjectMapping& CellObjectMapping : ChildPackages)
+		{
+			if (CellObjectMapping.ContainerPackage != CellObjectMapping.WorldPackage)
+			{
+				bool bIsContainerPackageAlreadyRemapped = OutLinkInstancingContext.RemapPackage(CellObjectMapping.ContainerPackage) != CellObjectMapping.ContainerPackage;
+				if (!bIsContainerPackageAlreadyRemapped)
+				{
+					FString ContainerPackage = CellObjectMapping.ContainerPackage.ToString();
+					FString WorldPackage = CellObjectMapping.WorldPackage.ToString();
+
+					bool bWithoutSlashes = false;
+					FName ContainerMountPoint = FPackageName::GetPackageMountPoint(ContainerPackage, bWithoutSlashes);
+					FName WorldPackageMountPoint = FPackageName::GetPackageMountPoint(WorldPackage, bWithoutSlashes);
+					if (ContainerMountPoint != WorldPackageMountPoint)
+					{
+#if DO_CHECK
+						//  Validate that while the container mounting points are different, they point to the same world package.
+						FStringView RelativeWorldPackage;
+						ensure(FPathViews::TryMakeChildPathRelativeTo(MakeStringView(WorldPackage), MakeStringView(WorldPackageMountPoint.ToString()), RelativeWorldPackage));
+
+						uint32 WorldPackageInContainerPackageStartIdx = ContainerPackage.Len() - RelativeWorldPackage.Len();
+						FStringView WorldPackageInContainerPackage = MakeStringView(ContainerPackage).SubStr(WorldPackageInContainerPackageStartIdx, RelativeWorldPackage.Len());
+						check(WorldPackageInContainerPackage.Equals(RelativeWorldPackage));
+#endif // DO_CHECK
+
+						OutLinkInstancingContext.AddPackageMapping(CellObjectMapping.ContainerPackage, RuntimePackage->GetFName());
+					}
+				}
+			}
+		}
+	};
+	
+	FLinkerInstancingContext InstancingContext;
+	BuildInstancingContext(InstancingContext);
 
 	ChildPackagesToLoad.Reset(ChildPackages.Num());
 
