Change Custom Attributes project setting name
#fix add display name metadata
#misc make tooltip generation use ::GetDisplayNameText rather than generate it in place (and not respect DisplayName metadata)

diff --git a/Engine/Source/Runtime/Engine/Classes/Animation/AnimationSettings.h b/Engine/Source/Runtime/Engine/Classes/Animation/AnimationSettings.h
--- a/Engine/Source/Runtime/Engine/Classes/Animation/AnimationSettings.h
+++ b/Engine/Source/Runtime/Engine/Classes/Animation/AnimationSettings.h
@@ -61,34 +61,34 @@ class ENGINE_API UAnimationSettings : public UDeveloperSettings
 	UPROPERTY(config, EditAnywhere, Category = Performance)
 	bool bTickAnimationOnSkeletalMeshInit;
 
-	/** Names that identify bone custom attributes representing the individual components of a timecode and a subframe along with a take name.
+	/** Names that identify bone animation attributes representing the individual components of a timecode and a subframe along with a take name.
 	    These will be included in the list of bone custom attribute names to import. */
-	UPROPERTY(config, EditAnywhere, Category = CustomAttributes)
+	UPROPERTY(config, EditAnywhere, Category = AnimationAttributes, meta=(DisplayName="Bone Timecode Animation Attribute name settings"))
 	FTimecodeCustomAttributeNameSettings BoneTimecodeCustomAttributeNameSettings;
 
-	/** List of custom attribute to import directly on their corresponding bone. The meaning field allows to contextualize the attribute name and customize tooling for it. */
-	UPROPERTY(config, EditAnywhere, Category = CustomAttributes)
+	/** List of animation attribute names to import directly on their corresponding bone names. The meaning field allows to contextualize the attribute name and customize tooling for it. */
+	UPROPERTY(config, EditAnywhere, Category = AnimationAttributes, meta=(DisplayName="Bone Animation Attributes names"))
 	TArray<FCustomAttributeSetting> BoneCustomAttributesNames;
 
-	/** Gets the complete list of bone custom attribute names to consider for import.
-	    This includes the designated timecode custom attributes as well as other bone custom attributes identified in the settings. */
-	UFUNCTION(BlueprintPure, Category = CustomAttributes)
+	/** Gets the complete list of bone animation attribute names to consider for import.
+	    This includes the designated timecode animation attributes as well as other bone animation attributes identified in the settings. */
+	UFUNCTION(BlueprintPure, Category = AnimationAttributes)
 	TArray<FString> GetBoneCustomAttributeNamesToImport() const;
 
-	/** List of bone names for which all custom attributes are directly imported on the bone. */
-	UPROPERTY(config, EditAnywhere, Category = CustomAttributes)
+	/** List of bone names for which all animation attributes are directly imported on the bone. */
+	UPROPERTY(config, EditAnywhere, Category = AnimationAttributes, meta=(DisplayName="Bone names with Animation Attributes"))
 	TArray<FString> BoneNamesWithCustomAttributes;
 
-	/** Custom Attribute specific blend types (by name) */
-	UPROPERTY(config, EditAnywhere, Category = CustomAttributes)
+	/** Animation Attribute specific blend types (by name) */
+	UPROPERTY(config, EditAnywhere, Category = AnimationAttributes)
 	TMap<FName, ECustomAttributeBlendType> AttributeBlendModes;
 
-	/** Default Custom Attribute blend type */
-	UPROPERTY(config, EditAnywhere, Category = CustomAttributes)
+	/** Default Animation Attribute blend type */
+	UPROPERTY(config, EditAnywhere, Category = AnimationAttributes)
 	ECustomAttributeBlendType DefaultAttributeBlendMode;
 
 	/** Names to match against when importing FBX node transform curves as attributes (can use ? and * wildcards) */
-	UPROPERTY(config, EditAnywhere, Category = CustomAttributes)
+	UPROPERTY(config, EditAnywhere, Category = AnimationAttributes)
 	TArray<FString> TransformAttributeNames;
 
 	/** Find and Replace Expressions used for mirroring  */
