Expose add_socket method to SkeletalMesh class in Python API

diff --git a/Engine/Source/Runtime/Engine/Classes/Engine/SkeletalMesh.h b/Engine/Source/Runtime/Engine/Classes/Engine/SkeletalMesh.h
--- a/Engine/Source/Runtime/Engine/Classes/Engine/SkeletalMesh.h
+++ b/Engine/Source/Runtime/Engine/Classes/Engine/SkeletalMesh.h
@@ -2439,6 +2439,11 @@ public:
 	/** Recalculate Retarget Base Pose BoneTransform */
 	void ReallocateRetargetBasePose();
 
+	/**
+	 *	Add a skeletal socket object to this SkeletalMesh, and optionally promotes it to USkeleton socket.
+	 */
+	UFUNCTION(BlueprintCallable, Category = "Animation")
+	void AddSocket(USkeletalMeshSocket* InSocket, bool bAddToSkeleton=false);
 #endif // WITH_EDITOR
 
 	/** 
diff --git a/Engine/Source/Runtime/Engine/Private/SkeletalMesh.cpp b/Engine/Source/Runtime/Engine/Private/SkeletalMesh.cpp
--- a/Engine/Source/Runtime/Engine/Private/SkeletalMesh.cpp
+++ b/Engine/Source/Runtime/Engine/Private/SkeletalMesh.cpp
@@ -95,6 +95,7 @@
 #if WITH_EDITOR
 #include "ClothingAssetFactoryInterface.h"
 #include "ClothingSystemEditorInterfaceModule.h"
+#include "ScopedTransaction.h"
 #endif
 #include "SkeletalDebugRendering.h"
 #include "Misc/RuntimeErrors.h"
@@ -3600,6 +3601,85 @@ USkeletalMeshSocket* USkeletalMesh::FindSocket(FName InSocketName) const
 	return FindSocketAndIndex(InSocketName, DummyIdx);
 }
 
+#if WITH_EDITOR
+void USkeletalMesh::AddSocket(USkeletalMeshSocket* InSocket, bool bAddToSkeleton /*= false*/)
+{
+	if (InSocket)
+	{
+		if (InSocket->GetOuter() == this)
+		{
+			const FString SocketNameString = InSocket->SocketName.ToString();
+			const FString TrimmedNameString = SocketNameString.TrimStartAndEnd();
+			const FName TrimmedName = FName(*TrimmedNameString);
+			if (!Sockets.ContainsByPredicate([TrimmedName](const TObjectPtr<USkeletalMeshSocket>& Socket)
+			{
+				return Socket->SocketName == TrimmedName;
+			}))
+			{
+				const FReferenceSkeleton& ReferenceSkeleton = GetRefSkeleton();
+				if (ReferenceSkeleton.FindBoneIndex(InSocket->BoneName) != INDEX_NONE)
+				{				
+					if (!TrimmedNameString.IsEmpty())
+					{
+						const FScopedTransaction Transaction(LOCTEXT("AddSocket", "Add Socket"));
+						Modify();
+
+						// Update socket name in place
+						InSocket->SocketName = TrimmedName;					
+						Sockets.Add(InSocket);
+		    
+						if (bAddToSkeleton)
+						{
+							USkeleton* CurrentSkeleton = GetSkeleton();				
+							if (!CurrentSkeleton->Sockets.ContainsByPredicate([Name=InSocket->SocketName](const TObjectPtr<USkeletalMeshSocket>& Socket)
+							{
+								return Socket->SocketName == Name;
+							}))
+							{
+								const FReferenceSkeleton& SkeletonReferenceSkeleton = CurrentSkeleton->GetReferenceSkeleton();
+								if (SkeletonReferenceSkeleton.FindBoneIndex(InSocket->BoneName) != INDEX_NONE)
+								{
+									CurrentSkeleton->Modify();
+			    
+									USkeletalMeshSocket* NewSocket = DuplicateObject<USkeletalMeshSocket>(InSocket, CurrentSkeleton);
+									check(NewSocket);
+									CurrentSkeleton->Sockets.Add(NewSocket);
+								}
+								else
+								{
+									// Bone does not exists on Skeleton
+									UE_LOG(LogSkeletalMesh, Warning, TEXT("Failed to add socket to Skeleton (%s) as the provided bone name %s does not exist."), *CurrentSkeleton->GetFullName(), *InSocket->BoneName.ToString());
+								}
+							}
+						}
+					}
+					else
+					{
+						// Invalid socket name
+						UE_LOG(LogSkeletalMesh, Warning, TEXT("Failed to add socket as the provided socket name %s is invalid."), *TrimmedNameString);
+					}
+				}
+				else
+				{
+					// Bone does not exists
+					UE_LOG(LogSkeletalMesh, Warning, TEXT("Failed to add socket as the provided bone name %s does not exist."), *InSocket->BoneName.ToString());
+				}
+			}
+			else
+			{
+				// Socket already exists
+				UE_LOG(LogSkeletalMesh, Warning, TEXT("Failed to add socket as a socket with name %s already exist."), *InSocket->BoneName.ToString());
+			}
+		}
+		else
+		{
+			// Socket outer is not correct
+			UE_LOG(LogSkeletalMesh, Warning, TEXT("Failed to add socket as the socket its outer should be %s but is %s."), *GetFullName(), *InSocket->GetOuter()->GetFullName());
+		}
+	}
+}
+#endif
+
 #if !WITH_EDITOR
 
 USkeletalMesh::FSocketInfo::FSocketInfo(const USkeletalMesh* InSkeletalMesh, USkeletalMeshSocket* InSocket, int32 InSocketIndex)
