AnimSequence need not to mark package dirty
#fix remove forced MarkPackageDirty and PostEditChange from Skeleton resetting path
#misc do not mark AnimSequence as dirty for Bracket-Open/Close notifies

diff --git a/Engine/Source/Runtime/Engine/Private/Animation/AnimSequence.cpp b/Engine/Source/Runtime/Engine/Private/Animation/AnimSequence.cpp
--- a/Engine/Source/Runtime/Engine/Private/Animation/AnimSequence.cpp
+++ b/Engine/Source/Runtime/Engine/Private/Animation/AnimSequence.cpp
@@ -5375,7 +5375,8 @@ void UAnimSequence::OnModelModified(const EAnimDataModelNotifyType& NotifyType,
 		}
 	};
 
-	bool bShouldMarkPackageDirty = !FUObjectThreadContext::Get().IsRoutingPostLoad;
+	bool bShouldMarkPackageDirty = !FUObjectThreadContext::Get().IsRoutingPostLoad &&
+		NotifyType != EAnimDataModelNotifyType::BracketOpened && NotifyType != EAnimDataModelNotifyType::BracketClosed;
 
 	PRAGMA_DISABLE_DEPRECATION_WARNINGS
 	switch (NotifyType)
diff --git a/Engine/Source/Runtime/Engine/Private/AnimationAsset.cpp b/Engine/Source/Runtime/Engine/Private/AnimationAsset.cpp
--- a/Engine/Source/Runtime/Engine/Private/AnimationAsset.cpp
+++ b/Engine/Source/Runtime/Engine/Private/AnimationAsset.cpp
@@ -476,11 +476,6 @@ bool UAnimationAsset::ReplaceSkeleton(USkeleton* NewSkeleton, bool bConvertSpace
 			}
 		}
 
-		PostEditChange();
-		if (!FUObjectThreadContext::Get().IsRoutingPostLoad)
-		{
-		    MarkPackageDirty();
-		}
 		return true;
 	}
 
