Added a FindComponentByInterface

diff --git a/Engine/Source/Runtime/Engine/Classes/GameFramework/Actor.h b/Engine/Source/Runtime/Engine/Classes/GameFramework/Actor.h
--- a/Engine/Source/Runtime/Engine/Classes/GameFramework/Actor.h
+++ b/Engine/Source/Runtime/Engine/Classes/GameFramework/Actor.h
@@ -3344,6 +3344,9 @@ public:
 	UFUNCTION(BlueprintCallable, Category = "Actor", meta = (ComponentClass = "/Script/Engine.ActorComponent"), meta = (DeterminesOutputType = "ComponentClass"))
 	TArray<UActorComponent*> GetComponentsByTag(TSubclassOf<UActorComponent> ComponentClass, FName Tag) const;
 
+	/** Searches components array and returns first encountered component that implements the given interface. */
+	virtual UActorComponent* FindComponentByInterface(const TSubclassOf<UInterface> Interface) const;
+	
 	/** Gets all the components that implements the given interface. */
 	UFUNCTION(BlueprintCallable, Category = "Actor")
 	TArray<UActorComponent*> GetComponentsByInterface(TSubclassOf<UInterface> Interface) const;
@@ -3357,6 +3360,15 @@ public:
 		return (T*)FindComponentByClass(T::StaticClass());
 	}
 
+	/** Templatized version of FindComponentByInterface that handles casting for you */
+	template<class T>
+	T* FindComponentByInterface() const
+	{
+		static_assert(TPointerIsConvertibleFromTo<T, const UInterface>::Value, "'T' template parameter to FindComponentByInterface must be derived from UInterface");
+
+		return (T*)FindComponentByInterface(T::StaticClass());
+	}
+
 private:
 	/**
 	 * Internal helper function to call a compile-time lambda on all components of a given type
diff --git a/Engine/Source/Runtime/Engine/Private/Actor.cpp b/Engine/Source/Runtime/Engine/Private/Actor.cpp
--- a/Engine/Source/Runtime/Engine/Private/Actor.cpp
+++ b/Engine/Source/Runtime/Engine/Private/Actor.cpp
@@ -3411,6 +3411,25 @@ TArray<UActorComponent*> AActor::GetComponentsByTag(TSubclassOf<UActorComponent>
 	return MoveTemp(ComponentsByTag);
 }
 
+UActorComponent* AActor::FindComponentByInterface(const TSubclassOf<UInterface> Interface) const
+{
+	UActorComponent* FoundComponent = nullptr;
+
+	if (Interface)
+	{
+		for (UActorComponent* Component : GetComponents())
+		{
+			if (Component && Component->GetClass()->ImplementsInterface(Interface))
+			{
+				FoundComponent = Component;
+				break;
+			}
+		}
+	}
+
+	return FoundComponent;
+}
+
 TArray<UActorComponent*> AActor::GetComponentsByInterface(TSubclassOf<UInterface> Interface) const
 {
 	TArray<UActorComponent*> Components;
