Clamp using the budget upper bound instead of the dynamic one when computing the resolution. This avoids a feedback effect which causes resolution to constantly decrease if r.DynamicRes.UpperBoundQuantization is being used.

diff --git a/Engine/Source/Runtime/Engine/Private/DynamicResolution.cpp b/Engine/Source/Runtime/Engine/Private/DynamicResolution.cpp
--- a/Engine/Source/Runtime/Engine/Private/DynamicResolution.cpp
+++ b/Engine/Source/Runtime/Engine/Private/DynamicResolution.cpp
@@ -417,7 +417,7 @@ void FDynamicResolutionHeuristicProxy::RefreshCurentFrameResolutionFraction_Rend
 		NewFrameResolutionFraction = FMath::Clamp(
 			NewFrameResolutionFraction,
 			CVarDynamicResMinSP.GetValueOnRenderThread() / 100.0f,
-			CurrentFrameMaxResolutionFractions[GDynamicPrimaryResolutionFraction]);
+			GDynamicPrimaryResolutionFraction.GetSettings().MaxResolutionFraction);
 	}
 
 	// Update the current frame's resolution fraction.
