Fix invalid parameter pipping for Strata eye node.

diff --git a/Engine/Shaders/Private/Strata/Strata.ush b/Engine/Shaders/Private/Strata/Strata.ush
--- a/Engine/Shaders/Private/Strata/Strata.ush
+++ b/Engine/Shaders/Private/Strata/Strata.ush
@@ -163,16 +163,16 @@ struct FStrataBSDF
 //            BASIS(0)      32 |  Normal: Oct 16/16 | Tangent: 11/11/10
 //            BASIS(1)      32 |  Normal: Oct 16/16 | Tangent: 11/11/10
 //                             |
-//            STATE(0)      16 |  BSDF Type | Normal Id | Aniso | TopLayer | SSS | SSSProfile | GreyWeighV | Haze | F90 | Simple vol. | MFP Plugged | Fuzz | Thin
-//                             |  3           2           1       1          1     1            1            1      1     1             1             1      1 
+//            STATE(0)      16 |  BSDF Type | Normal Id | Aniso | TopLayer | SSSType | GreyWeighV | Haze | F90 | Simple vol. | MFP Plugged | Fuzz | Thin
+//                             |  3           2           1       1          2         1            1      1     1             1             1      1 
 //            BSDF(0)       ?? |  Adaptive
-//            STATE(1)      16 |  BSDF Type | Normal Id | Aniso | TopLayer | SSS | SSSProfile | GreyWeighV | Haze | F90 | Simple vol. | MFP Plugged | Fuzz | Thin
-//                             |  3           2           1       1          1     1            1            1      1     1             1             1      1 
+//            STATE(1)      16 |  BSDF Type | Normal Id | Aniso | TopLayer | SSSType | GreyWeighV | Haze | F90 | Simple vol. | MFP Plugged | Fuzz | Thin
+//                             |  3           2           1       1          2         1            1      1     1             1             1      1 
 //            BSDF(1)       ?? |  Adaptive
 //                             |
 // * LAYOUT 1 -------------------------------------------------------------------------------------------------------------------------------------------------------------------------
-//  (Single)  HEADER        27 |  State | AO | Aniso | TopLayer | SSS | SSSProfile | GreyWeighV | Haze | F90 | Simple vol. | MFP Plugged | Fuzz | Thin
-//                             |  8       8    1       1          1     1            1            1      1     1             1             1      1            
+//  (Single)  HEADER        27 |  State | AO | Aniso | TopLayer | SSSType | GreyWeighV | Haze | F90 | Simple vol. | MFP Plugged | Fuzz | Thin
+//                             |  8       8    1       1          2         1            1      1     1             1             1      1            
 //            BASIS(0)       0 | (Use top layer normal)
 //            BSDF(0)       ?? |  Adaptive
 // * LAYOUT 0 -------------------------------------------------------------------------------------------------------------------------------------------------------------------------          
diff --git a/Engine/Source/Runtime/Engine/Private/Materials/HLSLMaterialTranslator.cpp b/Engine/Source/Runtime/Engine/Private/Materials/HLSLMaterialTranslator.cpp
--- a/Engine/Source/Runtime/Engine/Private/Materials/HLSLMaterialTranslator.cpp
+++ b/Engine/Source/Runtime/Engine/Private/Materials/HLSLMaterialTranslator.cpp
@@ -10690,7 +10690,7 @@ int32 FHLSLMaterialTranslator::StrataHairBSDF(int32 BaseColor, int32 Scatter, in
 	);
 }
 
-int32 FHLSLMaterialTranslator::StrataEyeBSDF(int32 DiffuseAlbedo, int32 Roughness, int32 IrisMask, int32 IrisDistance, int32 EmissiveColor, int32 CorneaNormal, int32 IrisNormal, int32 IrisPlaneNormal, int32 SSSProfileId, const FString& SharedLocalBasisIndexMacro, FStrataOperator* PromoteToOperator)
+int32 FHLSLMaterialTranslator::StrataEyeBSDF(int32 DiffuseAlbedo, int32 Roughness, int32 IrisMask, int32 IrisDistance, int32 IrisNormal, int32 IrisPlaneNormal, int32 SSSProfileId, int32 EmissiveColor, int32 CorneaNormal, const FString& SharedLocalBasisIndexMacro, FStrataOperator* PromoteToOperator)
 {
 	return AddCodeChunk(
 		MCT_Strata, TEXT("PromoteParameterBlendedBSDFToOperator(GetStrataEyeBSDF(%s, %s, %s, %s, %s, %s, %s, %s, %s), Parameters.StrataTree, %u, %u, %u, %u) /* Cornea:%s Iris:%s */"),
diff --git a/Engine/Source/Runtime/Engine/Public/MaterialCompiler.h b/Engine/Source/Runtime/Engine/Public/MaterialCompiler.h
--- a/Engine/Source/Runtime/Engine/Public/MaterialCompiler.h
+++ b/Engine/Source/Runtime/Engine/Public/MaterialCompiler.h
@@ -592,7 +592,7 @@ public:
 	virtual int32 StrataVolumetricFogCloudBSDF(int32 Albedo, int32 Extinction, int32 EmissiveColor, int32 AmbientOcclusion) = 0;
 	virtual int32 StrataUnlitBSDF(int32 EmissiveColor, int32 TransmittanceColor) = 0;
 	virtual int32 StrataHairBSDF(int32 BaseColor, int32 Scatter, int32 Specular, int32 Roughness, int32 Backlit, int32 EmissiveColor, int32 Tangent, const FString& SharedLocalBasisIndexMacro, FStrataOperator* PromoteToOperator) = 0;
-	virtual int32 StrataEyeBSDF(int32 DiffuseAlbedo, int32 Roughness, int32 IrisMask, int32 IrisDistance, int32 EmissiveColor, int32 CorneaNormal, int32 IrisNormal, int32 IrisPlaneNormal, int32 SSSProfileId, const FString& SharedLocalBasisIndexMacro, FStrataOperator* PromoteToOperator) = 0;
+	virtual int32 StrataEyeBSDF(int32 DiffuseAlbedo, int32 Roughness, int32 IrisMask, int32 IrisDistance, int32 IrisNormal, int32 IrisPlaneNormal, int32 SSSProfileId, int32 EmissiveColor, int32 CorneaNormal, const FString& SharedLocalBasisIndexMacro, FStrataOperator* PromoteToOperator) = 0;
 	virtual int32 StrataSingleLayerWaterBSDF(
 		int32 BaseColor, int32 Metallic, int32 Specular, int32 Roughness, int32 EmissiveColor, int32 TopMaterialOpacity, 
 		int32 WaterAlbedo, int32 WaterExtinction, int32 WaterPhaseG, int32 ColorScaleBehindWater, int32 Normal, const FString& SharedLocalBasisIndexMacro, FStrataOperator* PromoteToOperator) = 0;
@@ -1191,9 +1191,9 @@ public:
 		return Compiler->StrataHairBSDF(BaseColor, Scatter, Specular, Roughness, Backlit, EmissiveColor, Tangent, SharedLocalBasisIndexMacro, PromoteToOperator);
 	}
 
-	virtual int32 StrataEyeBSDF(int32 DiffuseAlbedo, int32 Roughness, int32 IrisMask, int32 IrisDistance, int32 EmissiveColor, int32 CorneaNormal, int32 IrisNormal, int32 IrisPlaneNormal, int32 SSSProfileId, const FString& SharedLocalBasisIndexMacro, FStrataOperator* PromoteToOperator) override
+	virtual int32 StrataEyeBSDF(int32 DiffuseAlbedo, int32 Roughness, int32 IrisMask, int32 IrisDistance, int32 IrisNormal, int32 IrisPlaneNormal, int32 SSSProfileId, int32 EmissiveColor, int32 CorneaNormal, const FString& SharedLocalBasisIndexMacro, FStrataOperator* PromoteToOperator) override
 	{
-		return Compiler->StrataEyeBSDF(DiffuseAlbedo, Roughness, IrisMask, IrisDistance, EmissiveColor, CorneaNormal, IrisNormal, IrisPlaneNormal, SSSProfileId, SharedLocalBasisIndexMacro, PromoteToOperator);
+		return Compiler->StrataEyeBSDF(DiffuseAlbedo, Roughness, IrisMask, IrisDistance, IrisNormal, IrisPlaneNormal, SSSProfileId, EmissiveColor, CorneaNormal, SharedLocalBasisIndexMacro, PromoteToOperator);
 	}
 
 	virtual int32 StrataSingleLayerWaterBSDF(
