Add ShowUI parameter to the Trace.Screenshot command.

diff --git a/Engine/Source/Runtime/Engine/Private/ProfilingDebugging/TraceScreenshot.cpp b/Engine/Source/Runtime/Engine/Private/ProfilingDebugging/TraceScreenshot.cpp
--- a/Engine/Source/Runtime/Engine/Private/ProfilingDebugging/TraceScreenshot.cpp
+++ b/Engine/Source/Runtime/Engine/Private/ProfilingDebugging/TraceScreenshot.cpp
@@ -9,6 +9,7 @@
 #include "HAL/PlatformTLS.h"
 #include "HAL/TlsAutoCleanup.h"
 #include "Misc/DateTime.h"
+#include "Misc/CString.h"
 #include "Misc/Paths.h"
 #include "ProfilingDebugging/MiscTrace.h"
 #include "Tasks/Task.h"
@@ -17,25 +18,30 @@
 static void TraceScreenshotCommandCallback(const TArray<FString>& Args)
 {
 	FString Name;
+	bool bShowUI = false;
+
 	if (Args.Num() > 0)
 	{
 		Name = Args[0];
 	}
 
-	FTraceScreenshot::RequestScreenshot(Name);
+	if (Args.Num() > 1)
+	{
+		bShowUI = FCString::ToBool(*Args[1]);
+	}
+
+	FTraceScreenshot::RequestScreenshot(Name, bShowUI);
 }
 
 static FAutoConsoleCommand TraceScreenshotCmd(
 	TEXT("Trace.Screenshot"),
-	TEXT("[Name] Takes a screenshot and saves it in the trace."
-		" [Name] is the name of the screenshot."
-	),
+	TEXT("[Name] [ShowUI] Takes a screenshot and saves it in the trace. Ex: Trace.Screenshot ScreenshotName true"),
 	FConsoleCommandWithArgsDelegate::CreateStatic(TraceScreenshotCommandCallback)
 );
 
 bool FTraceScreenshot::bSuppressWritingToFile = false;
 
-void FTraceScreenshot::RequestScreenshot(FString Name)
+void FTraceScreenshot::RequestScreenshot(FString Name, bool bShowUI)
 {
 	if (!SHOULD_TRACE_SCREENSHOT())
 	{
@@ -49,7 +55,6 @@ void FTraceScreenshot::RequestScreenshot(FString Name)
 	{
 		Name = FDateTime::Now().ToString(TEXT("Screenshot_%Y%m%d_%H%M%S"));
 	}
-	constexpr bool bShowUI = false;
 	constexpr bool bAddUniqueSuffix = false;
 	FScreenshotRequest::RequestScreenshot(Name, bShowUI, bAddUniqueSuffix);
 }
diff --git a/Engine/Source/Runtime/Engine/Public/ProfilingDebugging/TraceScreenshot.h b/Engine/Source/Runtime/Engine/Public/ProfilingDebugging/TraceScreenshot.h
--- a/Engine/Source/Runtime/Engine/Public/ProfilingDebugging/TraceScreenshot.h
+++ b/Engine/Source/Runtime/Engine/Public/ProfilingDebugging/TraceScreenshot.h
@@ -10,7 +10,7 @@ class ULevel;
 class ENGINE_API FTraceScreenshot
 {
 public:
-	static void RequestScreenshot(FString Name);
+	static void RequestScreenshot(FString Name, bool bShowUI);
 
 	/* 
 	* Add the provided screenshot to the trace.
