PR #8791: Fix the OnScalabilitySettingsChanged event in non-monolithic builds.

#rnx

diff --git a/Engine/Source/Runtime/Engine/Private/Scalability.cpp b/Engine/Source/Runtime/Engine/Private/Scalability.cpp
--- a/Engine/Source/Runtime/Engine/Private/Scalability.cpp
+++ b/Engine/Source/Runtime/Engine/Private/Scalability.cpp
@@ -11,6 +11,8 @@
 #include "Interfaces/IAnalyticsProvider.h"
 #include "Interfaces/IProjectManager.h"
 
+Scalability::FOnScalabilitySettingsChanged Scalability::OnScalabilitySettingsChanged;
+
 static TAutoConsoleVariable<float> CVarResolutionQuality(
 	TEXT("sg.ResolutionQuality"),
 	100.0f,
diff --git a/Engine/Source/Runtime/Engine/Public/Scalability.h b/Engine/Source/Runtime/Engine/Public/Scalability.h
--- a/Engine/Source/Runtime/Engine/Public/Scalability.h
+++ b/Engine/Source/Runtime/Engine/Public/Scalability.h
@@ -150,7 +150,7 @@ namespace Scalability
 	};
 
 	DECLARE_MULTICAST_DELEGATE_OneParam(FOnScalabilitySettingsChanged, const Scalability::FQualityLevels&);
-	static FOnScalabilitySettingsChanged OnScalabilitySettingsChanged;
+	ENGINE_API extern FOnScalabilitySettingsChanged OnScalabilitySettingsChanged;
 
 	/** This is the only suggested way to set the current state - don't set CVars directly **/
 	ENGINE_API void SetQualityLevels(const FQualityLevels& QualityLevels, bool bForce = false);
