Added OnPawnBeginPlay delegate to Pawn.

diff --git a/Engine/Source/Runtime/Engine/Classes/GameFramework/Pawn.h b/Engine/Source/Runtime/Engine/Classes/GameFramework/Pawn.h
--- a/Engine/Source/Runtime/Engine/Classes/GameFramework/Pawn.h
+++ b/Engine/Source/Runtime/Engine/Classes/GameFramework/Pawn.h
@@ -29,6 +29,7 @@ ENGINE_API DECLARE_LOG_CATEGORY_EXTERN(LogDamage, Warning, All);
 
 DECLARE_DYNAMIC_MULTICAST_SPARSE_DELEGATE_OneParam(FPawnRestartedSignature, APawn, ReceiveRestartedDelegate, APawn*, Pawn);
 DECLARE_DYNAMIC_MULTICAST_SPARSE_DELEGATE_ThreeParams(FPawnControllerChangedSignature, APawn, ReceiveControllerChangedDelegate, APawn*, Pawn, AController*, OldController, AController*, NewController);
+DECLARE_MULTICAST_DELEGATE_OneParam(FOnPawnBeginPlay, APawn*);
 
 /** 
  * Pawn is the base class of all actors that can be possessed by players or AI.
@@ -149,6 +150,9 @@ public:
 
 	bool IsLocalPlayerControllerViewingAPawn() const;
 
+	/** a delegate that broadcasts a notification whenever BeginPlay() is called */
+	static FOnPawnBeginPlay OnPawnBeginPlay;
+
 private:
 	/** If Pawn is possessed by a player, points to its Player State.  Needed for network play as controllers are not replicated to clients. */
 	UPROPERTY(replicatedUsing=OnRep_PlayerState, BlueprintReadOnly, Category=Pawn, meta=(AllowPrivateAccess="true"))
@@ -261,6 +265,7 @@ public:
 	virtual void DisplayDebug(UCanvas* Canvas, const FDebugDisplayInfo& DebugDisplay, float& YL, float& YPos) override;
 	virtual void GetActorEyesViewPoint( FVector& Location, FRotator& Rotation ) const override;
 	virtual void OutsideWorldBounds() override;
+	virtual void BeginPlay() override;
 	virtual void Destroyed() override;
 	virtual void EndPlay(const EEndPlayReason::Type EndPlayReason) override;
 	virtual void PreInitializeComponents() override;
diff --git a/Engine/Source/Runtime/Engine/Private/Pawn.cpp b/Engine/Source/Runtime/Engine/Private/Pawn.cpp
--- a/Engine/Source/Runtime/Engine/Private/Pawn.cpp
+++ b/Engine/Source/Runtime/Engine/Private/Pawn.cpp
@@ -36,6 +36,8 @@
 DEFINE_LOG_CATEGORY(LogDamage);
 DEFINE_LOG_CATEGORY_STATIC(LogPawn, Warning, All);
 
+FOnPawnBeginPlay APawn::OnPawnBeginPlay;
+
 APawn::APawn(const FObjectInitializer& ObjectInitializer)
 	: Super(ObjectInitializer)
 {
@@ -171,6 +173,13 @@ void APawn::PostRegisterAllComponents()
 	UpdateNavAgent();
 }
 
+void APawn::BeginPlay()
+{
+	Super::BeginPlay();
+
+	OnPawnBeginPlay.Broadcast(this);
+}
+
 UPawnMovementComponent* APawn::GetMovementComponent() const
 {
 	return FindComponentByClass<UPawnMovementComponent>();
