[UE][FIX] Constraint default transform doesn't work for Cross Constraints
- Current code (written by me - CL-19283727) assumes that a constraint's child bone is a direct decendent of its parent bone when calculating the default parent transform, which isn't the case for cross constraints. This fix finds a common ancestor bone for both constraint bones and uses that to calculate the transform for one bone relative to the othe.

#rmx

diff --git a/Engine/Source/Runtime/Engine/Private/PhysicsEngine/ConstraintInstance.cpp b/Engine/Source/Runtime/Engine/Private/PhysicsEngine/ConstraintInstance.cpp
--- a/Engine/Source/Runtime/Engine/Private/PhysicsEngine/ConstraintInstance.cpp
+++ b/Engine/Source/Runtime/Engine/Private/PhysicsEngine/ConstraintInstance.cpp
@@ -71,24 +71,50 @@ FVector RevolutionsToRads(const FVector Revolutions)
 
 #if WITH_EDITORONLY_DATA
 
-/** Returns the child bone's transform relative to the parent bone. */
-FTransform CalculateChildTransformRelativeToParent(const FName ChildBoneName, const FName ParentBoneName, const UPhysicsAsset* const PhysicsAsset)
+/** Returns the 'To' bone's transform relative to the 'From' bone. */
+FTransform CalculateRelativeBoneTransform(const FName ToBoneName, const FName FromBoneName, FReferenceSkeleton& ReferenceSkeleton)
 {
-	FReferenceSkeleton& ReferenceSkeleton = PhysicsAsset->PreviewSkeletalMesh->GetRefSkeleton();
-	const TArray<FTransform>& LocalPose = PhysicsAsset->PreviewSkeletalMesh->GetRefSkeleton().GetRefBonePose();
-	const int32 ChildBoneIndex = ReferenceSkeleton.FindBoneIndex(ChildBoneName);
-	const int32 ParentBoneIndex = ReferenceSkeleton.FindBoneIndex(ParentBoneName);
+	FTransform RelativeBoneTransform = FTransform::Identity;
 
-	// Transform of child from parent is just child ref-pose entry.
-	FTransform RelTM = FTransform::Identity;
+	const TArray<FTransform>& LocalPose = ReferenceSkeleton.GetRefBonePose();
+	int32 ToBoneAncestorIndex = ReferenceSkeleton.FindBoneIndex(ToBoneName);
+	int32 FromBoneAncestorIndex = ReferenceSkeleton.FindBoneIndex(FromBoneName);
 
-	// Traverse the skeleton hierarchy from child to parent bone, accumulating transforms.
-	for (int32 ParentIndex = ChildBoneIndex; LocalPose.IsValidIndex(ParentIndex) && (ParentIndex != ParentBoneIndex); ParentIndex = ReferenceSkeleton.GetParentIndex(ParentIndex))
+	check(LocalPose.IsValidIndex(ToBoneAncestorIndex));
+	check(LocalPose.IsValidIndex(FromBoneAncestorIndex));
+
+	FTransform ToCommonBasisTransform = FTransform::Identity;
+	FTransform FromCommonBasisTransform = FTransform::Identity;
+
+	// Traverse the skeleton from child to parent bone, accumulating transforms until we have a transform for both bones relative to a common ancestor bone in the hierarchy.
+	while (LocalPose.IsValidIndex(ToBoneAncestorIndex) && LocalPose.IsValidIndex(FromBoneAncestorIndex))
+	{
+		if (ToBoneAncestorIndex > FromBoneAncestorIndex)
+		{
+			ToCommonBasisTransform = ToCommonBasisTransform * LocalPose[ToBoneAncestorIndex];
+			ToBoneAncestorIndex = ReferenceSkeleton.GetParentIndex(ToBoneAncestorIndex);
+		}
+		else if (FromBoneAncestorIndex > ToBoneAncestorIndex)
+		{
+			FromCommonBasisTransform = FromCommonBasisTransform * LocalPose[FromBoneAncestorIndex];
+			FromBoneAncestorIndex = ReferenceSkeleton.GetParentIndex(FromBoneAncestorIndex);
+		}
+		else // FromBoneAncestorIndex == ToBoneAncestorIndex 
+		{
+			// Found a bone that exists in the hierarchy of both the 'To' and 'From' bones.
+			break;
+		}
+	}
+
+	check(ToBoneAncestorIndex == FromBoneAncestorIndex); // A pair of bones should always have at least one common bone in their hierarchies, even if it's the root bone.
+
+	// Calculate the transform of the 'To' bone relative to the 'From' bone.
+	if (ToBoneAncestorIndex == FromBoneAncestorIndex)
 	{
-		RelTM = RelTM * LocalPose[ParentIndex];
+		RelativeBoneTransform = ToCommonBasisTransform.GetRelativeTransform(FromCommonBasisTransform);
 	}
 
-	return RelTM;
+	return RelativeBoneTransform;
 }
 
 #endif // WITH_EDITORONLY_DATA
@@ -251,7 +277,15 @@ void FConstraintInstance::UpdateDriveTarget()
 
 FTransform FConstraintInstance::CalculateDefaultParentTransform(const UPhysicsAsset* const PhysicsAsset) const
 {
-	return CalculateChildTransformRelativeToParent(ConstraintBone1, ConstraintBone2, PhysicsAsset);
+	if (PhysicsAsset)
+	{
+		if (USkeletalMesh* const PreviewSkelMesh = PhysicsAsset->GetPreviewMesh())
+		{
+			return CalculateRelativeBoneTransform(ConstraintBone1, ConstraintBone2, PreviewSkelMesh->GetRefSkeleton());
+		}
+	}
+
+	return FTransform::Identity;
 }
 
 FTransform FConstraintInstance::CalculateDefaultChildTransform() const
