World Partition - HLOD: Prevent creation of HLOD subsystem if world isn't a game world or isn't partitioned
* Additional checks to validate that no HLODs are registered if streaming is disabled

diff --git a/Engine/Source/Runtime/Engine/Private/WorldPartition/HLOD/HLODSubsystem.cpp b/Engine/Source/Runtime/Engine/Private/WorldPartition/HLOD/HLODSubsystem.cpp
--- a/Engine/Source/Runtime/Engine/Private/WorldPartition/HLOD/HLODSubsystem.cpp
+++ b/Engine/Source/Runtime/Engine/Private/WorldPartition/HLOD/HLODSubsystem.cpp
@@ -122,6 +122,24 @@ bool UHLODSubsystem::IsHLODEnabled()
 	return UHLODSubsystem::WorldPartitionHLODEnabled;
 }
 
+bool UHLODSubsystem::ShouldCreateSubsystem(UObject* Outer) const
+{
+	if (!Super::ShouldCreateSubsystem(Outer))
+	{
+		return false;
+	}
+
+	UWorld* World = Cast<UWorld>(Outer);
+	check(World);
+
+	return World->IsPartitionedWorld();
+}
+
+bool UHLODSubsystem::DoesSupportWorldType(const EWorldType::Type WorldType) const
+{
+	return WorldType == EWorldType::Game || WorldType == EWorldType::PIE;
+}
+
 void UHLODSubsystem::Initialize(FSubsystemCollectionBase& Collection)
 {
 	// Ensure the WorldPartitionSubsystem gets created before the HLODSubsystem
@@ -130,14 +148,12 @@ void UHLODSubsystem::Initialize(FSubsystemCollectionBase& Collection)
 	Super::Initialize(Collection);
 
 	UWorld* World = GetWorld();
+	check(World->IsGameWorld());
 
-	if (World->IsGameWorld())
-	{
-		GetWorld()->OnWorldPartitionInitialized().AddUObject(this, &UHLODSubsystem::OnWorldPartitionInitialized);
-		GetWorld()->OnWorldPartitionUninitialized().AddUObject(this, &UHLODSubsystem::OnWorldPartitionUninitialized);
+	GetWorld()->OnWorldPartitionInitialized().AddUObject(this, &UHLODSubsystem::OnWorldPartitionInitialized);
+	GetWorld()->OnWorldPartitionUninitialized().AddUObject(this, &UHLODSubsystem::OnWorldPartitionUninitialized);
 
-		SceneViewExtension = FSceneViewExtensions::NewExtension<FHLODResourcesResidencySceneViewExtension>(World);
-	}
+	SceneViewExtension = FSceneViewExtensions::NewExtension<FHLODResourcesResidencySceneViewExtension>(World);
 }
 
 void UHLODSubsystem::Deinitialize()
@@ -150,28 +166,33 @@ void UHLODSubsystem::Deinitialize()
 
 void UHLODSubsystem::OnWorldPartitionInitialized(UWorldPartition* InWorldPartition)
 {
-	check(!WorldPartitionsHLODRuntimeData.Contains(InWorldPartition));
-
-	FWorldPartitionHLODRuntimeData& WorldPartitionHLODRuntimeData = WorldPartitionsHLODRuntimeData.Add(InWorldPartition, FWorldPartitionHLODRuntimeData());
-	
-	TSet<const UWorldPartitionRuntimeCell*> StreamingCells;
-	if (InWorldPartition && InWorldPartition->RuntimeHash)
+	if (InWorldPartition && InWorldPartition->IsStreamingEnabled())
 	{
-		InWorldPartition->RuntimeHash->GetAllStreamingCells(StreamingCells, /*bAllDataLayers*/ true);
-	}
-	
+		check(!WorldPartitionsHLODRuntimeData.Contains(InWorldPartition));
 
-	// Build cell to HLOD mapping
-	for (const UWorldPartitionRuntimeCell* Cell : StreamingCells)
-	{
-		WorldPartitionHLODRuntimeData.CellsData.Emplace(Cell->GetFName());
+		FWorldPartitionHLODRuntimeData& WorldPartitionHLODRuntimeData = WorldPartitionsHLODRuntimeData.Add(InWorldPartition, FWorldPartitionHLODRuntimeData());
+
+		TSet<const UWorldPartitionRuntimeCell*> StreamingCells;
+		if (InWorldPartition->RuntimeHash)
+		{
+			InWorldPartition->RuntimeHash->GetAllStreamingCells(StreamingCells, /*bAllDataLayers*/ true);
+		}
+
+		// Build cell to HLOD mapping
+		for (const UWorldPartitionRuntimeCell* Cell : StreamingCells)
+		{
+			WorldPartitionHLODRuntimeData.CellsData.Emplace(Cell->GetFName());
+		}
 	}
 }
 
 void UHLODSubsystem::OnWorldPartitionUninitialized(UWorldPartition* InWorldPartition)
 {
-	check(WorldPartitionsHLODRuntimeData.Contains(InWorldPartition));
-	WorldPartitionsHLODRuntimeData.Remove(InWorldPartition);
+	if (InWorldPartition && InWorldPartition->IsStreamingEnabled())
+	{
+		check(WorldPartitionsHLODRuntimeData.Contains(InWorldPartition));
+		WorldPartitionsHLODRuntimeData.Remove(InWorldPartition);
+	}
 }
 
 void UHLODSubsystem::RegisterHLODActor(AWorldPartitionHLOD* InWorldPartitionHLOD)
@@ -179,7 +200,7 @@ void UHLODSubsystem::RegisterHLODActor(AWorldPartitionHLOD* InWorldPartitionHLOD
 	TRACE_CPUPROFILER_EVENT_SCOPE(UHLODSubsystem::RegisterHLODActor);
 
 	const UWorldPartition* WorldPartition = FHLODSubsystem::GetWorldPartition(InWorldPartitionHLOD);
-	check(WorldPartition);
+	check(WorldPartition && WorldPartition->IsStreamingEnabled());
 
 	FWorldPartitionHLODRuntimeData& WorldPartitionHLODRuntimeData = WorldPartitionsHLODRuntimeData.FindChecked(WorldPartition);
 
@@ -207,7 +228,7 @@ void UHLODSubsystem::UnregisterHLODActor(AWorldPartitionHLOD* InWorldPartitionHL
 	TRACE_CPUPROFILER_EVENT_SCOPE(UHLODSubsystem::UnregisterHLODActor);
 
 	const UWorldPartition* WorldPartition = FHLODSubsystem::GetWorldPartition(InWorldPartitionHLOD);
-	check(WorldPartition);
+	check(WorldPartition && WorldPartition->IsStreamingEnabled());
 
 	FWorldPartitionHLODRuntimeData& WorldPartitionHLODRuntimeData = WorldPartitionsHLODRuntimeData.FindChecked(WorldPartition);
 
@@ -228,8 +249,9 @@ void UHLODSubsystem::UnregisterHLODActor(AWorldPartitionHLOD* InWorldPartitionHL
 void UHLODSubsystem::OnCellShown(const UWorldPartitionRuntimeCell* InCell)
 {
 	const UWorldPartition* WorldPartition = InCell->GetTypedOuter<UWorldPartition>();
-	FWorldPartitionHLODRuntimeData& WorldPartitionHLODRuntimeData = WorldPartitionsHLODRuntimeData.FindChecked(WorldPartition);
+	check(WorldPartition && WorldPartition->IsStreamingEnabled());
 
+	FWorldPartitionHLODRuntimeData& WorldPartitionHLODRuntimeData = WorldPartitionsHLODRuntimeData.FindChecked(WorldPartition);
 	FCellData& CellData = WorldPartitionHLODRuntimeData.CellsData.FindChecked(InCell->GetFName());
 	CellData.bIsCellVisible = true;
 
@@ -249,8 +271,9 @@ void UHLODSubsystem::OnCellShown(const UWorldPartitionRuntimeCell* InCell)
 void UHLODSubsystem::OnCellHidden(const UWorldPartitionRuntimeCell* InCell)
 {
 	const UWorldPartition* WorldPartition = InCell->GetTypedOuter<UWorldPartition>();
-	FWorldPartitionHLODRuntimeData& WorldPartitionHLODRuntimeData = WorldPartitionsHLODRuntimeData.FindChecked(WorldPartition);
+	check(WorldPartition && WorldPartition->IsStreamingEnabled());
 
+	FWorldPartitionHLODRuntimeData& WorldPartitionHLODRuntimeData = WorldPartitionsHLODRuntimeData.FindChecked(WorldPartition);
 	FCellData& CellData = WorldPartitionHLODRuntimeData.CellsData.FindChecked(InCell->GetFName());
 	CellData.bIsCellVisible = false;
 
diff --git a/Engine/Source/Runtime/Engine/Public/WorldPartition/HLOD/HLODSubsystem.h b/Engine/Source/Runtime/Engine/Public/WorldPartition/HLOD/HLODSubsystem.h
--- a/Engine/Source/Runtime/Engine/Public/WorldPartition/HLOD/HLODSubsystem.h
+++ b/Engine/Source/Runtime/Engine/Public/WorldPartition/HLOD/HLODSubsystem.h
@@ -40,10 +40,15 @@ public:
 	virtual ~UHLODSubsystem();
 
 	//~ Begin USubsystem Interface.
+	virtual bool ShouldCreateSubsystem(UObject* Outer) const override;
 	virtual void Initialize(FSubsystemCollectionBase& Collection) override;
 	virtual void Deinitialize() override;
 	//~ End USubsystem Interface.
 
+	//~ Begin UWorldSubsystem Interface.
+	virtual bool DoesSupportWorldType(const EWorldType::Type WorldType) const override;
+	//~ End UWorldSubsystem Interface.
+
 	void RegisterHLODActor(AWorldPartitionHLOD* InWorldPartitionHLOD);
 	void UnregisterHLODActor(AWorldPartitionHLOD* InWorldPartitionHLOD);
 
