Integrating around robomerge block:
Removing redundant rep context for replay destruction infos caused by deprecation cleanup

diff --git a/Engine/Source/Runtime/Engine/Private/DemoNetDriver.cpp b/Engine/Source/Runtime/Engine/Private/DemoNetDriver.cpp
--- a/Engine/Source/Runtime/Engine/Private/DemoNetDriver.cpp
+++ b/Engine/Source/Runtime/Engine/Private/DemoNetDriver.cpp
@@ -2196,8 +2196,6 @@ bool UDemoNetDriver::ReplicatePrioritizedActor(const FActorPriority& ActorPriori
 			{
 				UE_LOG(LogDemo, Verbose, TEXT("TickDemoRecord creating destroy channel for NetGUID <%s,%s> Priority: %d"), *DestructionInfo->NetGUID.ToString(), *DestructionInfo->PathName, ActorPriority.Priority);
 
-				FScopedRepContext LevelContext(Params.Connection, DestructionInfo->Level.Get());
-
 				SendDestructionInfo(Params.Connection, DestructionInfo);
 
 				// Remove from connection's to-be-destroyed list (close bunch is reliable, so it will make it there)
