Add arm64 to ddc keys to avoid ddc consistency issues due to current floating point compiler switches.

diff --git a/Engine/Source/Developer/AudioFormatBink/Private/AudioFormatBink.cpp b/Engine/Source/Developer/AudioFormatBink/Private/AudioFormatBink.cpp
--- a/Engine/Source/Developer/AudioFormatBink/Private/AudioFormatBink.cpp
+++ b/Engine/Source/Developer/AudioFormatBink/Private/AudioFormatBink.cpp
@@ -70,10 +70,6 @@ public:
 		TRACE_CPUPROFILER_EVENT_SCOPE(FAudioFormatBink::Cook);
 		check(InFormat == NAME_BINKA);
 
-#if !PLATFORM_CPU_X86_FAMILY
-		UE_LOG(LogAudioFormatBink, Warning, TEXT("Bink Audio encoder has not been tested for consistency on non x86 platforms - cooks on this platform might not match cooks on x86/x64 for the same DDC key!"));
-#endif
-
 		uint8 CompressionLevel = AudioFormatBinkPrivate::GetCompressionLevelFromQualityIndex(InQualityInfo.Quality);
 		
 		void* CompressedData = 0;
diff --git a/Engine/Source/Runtime/Engine/Private/AudioDerivedData.cpp b/Engine/Source/Runtime/Engine/Private/AudioDerivedData.cpp
--- a/Engine/Source/Runtime/Engine/Private/AudioDerivedData.cpp
+++ b/Engine/Source/Runtime/Engine/Private/AudioDerivedData.cpp
@@ -177,6 +177,13 @@ static void GetStreamedAudioDerivedDataKeySuffix(
 		Version,
 		*SoundWave.CompressedDataGuid.ToString()
 		);
+
+#if PLATFORM_CPU_ARM_FAMILY
+	// Separate out arm keys as x64 and arm64 clang do not generate the same data for a given
+	// input. Add the arm specifically so that a) we avoid rebuilding the current DDC and
+	// b) we can remove it once we get arm64 to be consistent.
+	OutKeySuffix.Append(TEXT("_arm64"));
+#endif
 }
 
 /**
diff --git a/Engine/Source/Runtime/Engine/Private/TextureDerivedData.cpp b/Engine/Source/Runtime/Engine/Private/TextureDerivedData.cpp
--- a/Engine/Source/Runtime/Engine/Private/TextureDerivedData.cpp
+++ b/Engine/Source/Runtime/Engine/Private/TextureDerivedData.cpp
@@ -378,6 +378,13 @@ void GetTextureDerivedDataKeySuffix(const UTexture& Texture, const FTextureBuild
 		OutKeySuffix.Append(FString::Printf(TEXT("VT%s_"), TEXTURE_VT_DERIVEDDATA_VER));
 	}
 
+#if PLATFORM_CPU_ARM_FAMILY
+	// Separate out arm keys as x64 and arm64 clang do not generate the same data for a given
+	// input. Add the arm specifically so that a) we avoid rebuilding the current DDC and
+	// b) we can remove it once we get arm64 to be consistent.
+	OutKeySuffix.Append(TEXT("_arm64"));
+#endif
+
 	// Serialize the compressor settings into a temporary array. The archive
 	// is flagged as persistent so that machines of different endianness produce
 	// identical binary results.
