Deletes some dead code in dynamic res heuristic

diff --git a/Engine/Source/Runtime/Engine/Private/DynamicResolution.cpp b/Engine/Source/Runtime/Engine/Private/DynamicResolution.cpp
--- a/Engine/Source/Runtime/Engine/Private/DynamicResolution.cpp
+++ b/Engine/Source/Runtime/Engine/Private/DynamicResolution.cpp
@@ -236,7 +236,6 @@ void FDynamicResolutionHeuristicProxy::RefreshCurentFrameResolutionFraction_Rend
 	if (HistorySize > 0)
 	{
 		const float FrameWeightExponent = CVarFrameWeightExponent.GetValueOnRenderThread();
-		const float RenderThreadHeadRoomFromGPUMs = 1.0f; //CVarCPUTimeHeadRoom.GetValueOnRenderThread();
 		const int32 MaxConsecutiveOverbudgetGPUFrameCount = FMath::Max(CVarMaxConsecutiveOverbudgetGPUFrameCount.GetValueOnRenderThread(), 2);
 
 		const float OutLierTimeThreshold = 0.0f; //FrameTimeBudgetMs * CVarOutlierThreshold.GetValueOnRenderThread();
@@ -270,26 +269,21 @@ void FDynamicResolutionHeuristicProxy::RefreshCurentFrameResolutionFraction_Rend
 				continue;
 			}
 
-			const bool bGPUTimingsHaveCPUBubbles = false;
-
 			// Compute the maximum thread time between game and render thread for this <i>th frame.
 			// We add a head room for render thread to detect when the GPU could be starving.
 			const float MaxThreadTimeMs = FMath::Max(
-				FrameEntry.RenderThreadTimeMs + (bGPUTimingsHaveCPUBubbles ? RenderThreadHeadRoomFromGPUMs : 0),
+				FrameEntry.RenderThreadTimeMs,
 				FrameEntry.GameThreadTimeMs);
 
 			// Whether bound by game thread.
 			const bool bIsGameThreadBound = FrameEntry.GameThreadTimeMs > FrameTimeBudgetMs;
 
 			// Whether bound by render thread.
-			const bool bIsRenderThreadBound = (FrameEntry.RenderThreadTimeMs + (bGPUTimingsHaveCPUBubbles ? RenderThreadHeadRoomFromGPUMs : 0)) > TargetedGPUBusyTimeMs;
+			const bool bIsRenderThreadBound = FrameEntry.RenderThreadTimeMs > TargetedGPUBusyTimeMs;
 
 			// Whether the frame is CPU bound.
 			const bool bIsCPUBound = bIsGameThreadBound || bIsRenderThreadBound;
 
-			// Whether the render thread might create bubbles in GPU timings.
-			const bool bMayHaveGPUBubbles = bGPUTimingsHaveCPUBubbles && bIsRenderThreadBound;
-
 			// Whether GPU is over budget, when not CPU bound.
 			const bool bHasOverbudgetGPU = !bIsCPUBound && FrameEntry.TotalFrameGPUBusyTimeMs > FrameTimeBudgetMs;
 
@@ -340,8 +334,6 @@ void FDynamicResolutionHeuristicProxy::RefreshCurentFrameResolutionFraction_Rend
 			float SuggestedResolutionFraction = 1.0f;
 
 			// If have reliable GPU times, or guess there is not GPU bubbles -> estimate the suggested resolution fraction that could have been used.
-			if (!bGPUTimingsHaveCPUBubbles ||
-				(bGPUTimingsHaveCPUBubbles && !bIsCPUBound && !bMayHaveGPUBubbles))
 			{
 				// This assumes GPU busy time is directly proportional to ResolutionFraction^2, but in practice
 				// this is more A * ResolutionFraction^2 + B with B >= 0 non constant unknown cost such as unscaled
@@ -357,18 +349,6 @@ void FDynamicResolutionHeuristicProxy::RefreshCurentFrameResolutionFraction_Rend
 					FMath::Sqrt(TargetedGPUBusyTimeMs / FrameEntry.TotalFrameGPUBusyTimeMs)
 					* FrameEntry.ResolutionFractions[GDynamicPrimaryResolutionFraction]);
 			}
-			else if (FrameEntry.ResolutionFractions[GDynamicPrimaryResolutionFraction] > CPUBoundResolutionFraction)
-			{
-				// When CPU bound, we no longer want to downscale, or drop resolution if CPU and GPU are sharing same memory.
-				SuggestedResolutionFraction = CPUBoundResolutionFraction;
-			}
-			else
-			{
-				// Unable to know reliably the GPU time so don't try to make any prediction to avoid issues such
-				// as resolution keep shrinking because GPU time is slightly more than render thread time, or
-				// increasing because GPU time slightly lower than frame time.
-				SuggestedResolutionFraction = FrameEntry.ResolutionFractions[GDynamicPrimaryResolutionFraction];
-			}
 
 			NewFrameResolutionFraction += SuggestedResolutionFraction * Weight;
 			TotalWeight += Weight;
