Clean up DestroyedStartupOrDormantActorsByLevel properly when levels are removed

diff --git a/Engine/Source/Runtime/Engine/Classes/Engine/NetDriver.h b/Engine/Source/Runtime/Engine/Classes/Engine/NetDriver.h
--- a/Engine/Source/Runtime/Engine/Classes/Engine/NetDriver.h
+++ b/Engine/Source/Runtime/Engine/Classes/Engine/NetDriver.h
@@ -1953,6 +1953,8 @@ private:
 
 	void UpdateCrashContext();
 
+	void RemoveDestroyedGuidsByLevel(const ULevel* Level, const TArray<FNetworkGUID>& RemovedGUIDs);
+
 	/** Handle to FNetDelegates::OnSyncLoadDetected delegate */
 	FDelegateHandle ReportSyncLoadDelegateHandle;
 
diff --git a/Engine/Source/Runtime/Engine/Private/NetDriver.cpp b/Engine/Source/Runtime/Engine/Private/NetDriver.cpp
--- a/Engine/Source/Runtime/Engine/Private/NetDriver.cpp
+++ b/Engine/Source/Runtime/Engine/Private/NetDriver.cpp
@@ -3360,6 +3360,8 @@ void UNetDriver::NotifyStreamingLevelUnload(ULevel* Level)
 				It.RemoveCurrent();
 			}
 		}
+
+		RemoveDestroyedGuidsByLevel(Level, RemovedGUIDs);
 	}	
 
 	if (Level->LevelScriptActor)
@@ -6273,6 +6275,21 @@ TSharedPtr< FReplicationChangelistMgr > UNetDriver::GetReplicationChangeListMgr(
 	return ReplicationChangeListMgrPtr->ReplicationChangelistMgr;
 }
 
+void UNetDriver::RemoveDestroyedGuidsByLevel(const ULevel* Level, const TArray<FNetworkGUID>& RemovedGUIDs)
+{
+	check(Level);
+
+	const FName StreamingLevelName = !Level->IsPersistentLevel() ? Level->GetOutermost()->GetFName() : NAME_None;
+
+	if (TSet<FNetworkGUID>* DestroyedGuidsForLevel = DestroyedStartupOrDormantActorsByLevel.Find(StreamingLevelName))
+	{
+		for (const FNetworkGUID& NetGUID : RemovedGUIDs)
+		{
+			DestroyedGuidsForLevel->Remove(NetGUID);
+		}
+	}
+}
+
 // This method will be called when Streaming Levels become Visible.
 void UNetDriver::OnLevelAddedToWorld(ULevel* Level, UWorld* InWorld)
 {
@@ -6334,6 +6351,8 @@ void UNetDriver::OnLevelRemovedFromWorld(class ULevel* Level, class UWorld* InWo
 				RemovedGUIDs.Add(It->Key);
 				It.RemoveCurrent();
 			}
+
+			RemoveDestroyedGuidsByLevel(Level, RemovedGUIDs);
 		}
 	}
 	}
