Updated FWaitForInteractiveFramerate automation command to use a rolling window of averages over the last few seconds vs relying on GAverageFPS (which is just an approximation and weighted to the last few frames)

diff --git a/Engine/Source/Runtime/Engine/Private/Tests/AutomationCommon.cpp b/Engine/Source/Runtime/Engine/Private/Tests/AutomationCommon.cpp
--- a/Engine/Source/Runtime/Engine/Private/Tests/AutomationCommon.cpp
+++ b/Engine/Source/Runtime/Engine/Private/Tests/AutomationCommon.cpp
@@ -24,6 +24,7 @@
 #include "Tests/AutomationTestSettings.h"
 #include "GameMapsSettings.h"
 #include "IRenderCaptureProvider.h"
+#include <Runtime/Core/Public/Algo/Accumulate.h>
 
 #if WITH_AUTOMATION_TESTS
 
@@ -486,12 +487,48 @@ bool FWaitForSpecifiedMapToLoadCommand::Update()
 	return false;
 }
 
+/** Adds an easy way to average some value over discrete updates.  */
+template <typename T>
+struct TWaitForFrameRateRollingAverage
+{
+	static_assert(TIsArithmetic<T>::Value, "Unsupported type for computing a rolling average...");
+
+	TArray<T> Buffer;
+	int32 BufferOffset = 0;
+
+	void SetNum(int32 NewSize)
+	{
+		// at least 1
+		Buffer.SetNum(FMath::Max(NewSize, 1));
+		BufferOffset = BufferOffset % Buffer.Num();
+	}
+
+	void Reset()
+	{
+		BufferOffset = 0;
+	}
+
+	void Add(const T Value)
+	{
+		Buffer[BufferOffset] = Value;
+		BufferOffset = (BufferOffset + 1) % Buffer.Num();
+	}
+
+	const T Average() const
+	{
+		const int32 Num = Buffer.Num();
+		return Num > 0 ? Algo::Accumulate(Buffer, 0) / Num : 0;
+	}
+};
+
 FWaitForInteractiveFrameRate::FWaitForInteractiveFrameRate(float InDesiredFrameRate, float InDuration, float InMaxWaitTime)
 	: DesiredFrameRate(InDesiredFrameRate)
 	, Duration(InDuration)
 	, MaxWaitTime(InMaxWaitTime)
 	, StartTimeOfWait(0)
 	, StartTimeOfAcceptableFrameRate(0)
+	, LastTickTime(0)
+	, BufferIndex(0)
 {
 	UAutomationTestSettings const* AutomationTestSettings = GetDefault<UAutomationTestSettings>();
 
@@ -511,7 +548,24 @@ FWaitForInteractiveFrameRate::FWaitForInteractiveFrameRate(float InDesiredFrameR
 	}
 }
 
+void FWaitForInteractiveFrameRate::AddTickRateSample(const double Value)
+{
+	if (RollingTickRateBuffer.Num() == kSampleCount)
+	{
+		RollingTickRateBuffer[BufferIndex] = Value;
+		BufferIndex = (BufferIndex + 1) % RollingTickRateBuffer.Num();
+	}
+	else
+	{
+		RollingTickRateBuffer.Add(Value);
+	}
+}
 
+double FWaitForInteractiveFrameRate::CurrentAverageTickRate() const
+{
+	const int32 Num = RollingTickRateBuffer.Num();	
+	return Num > 0 ? (Algo::Accumulate(RollingTickRateBuffer, 0.0) / Num) : 0.0;
+}
 
 bool FWaitForInteractiveFrameRate::Update()
 {
@@ -529,16 +583,52 @@ bool FWaitForInteractiveFrameRate::Update()
 		UE_LOG(LogEngineAutomationLatentCommand, Log, TEXT("FWaitForInteractiveFrameRate: Starting wait for framerate of >= %.f FPS"), DesiredFrameRate);
 		StartTimeOfWait = TimeNow;
 		LastReportTime = TimeNow;
+		LastTickTime = TimeNow;
 	}
 
-	// Check if we're at the desired framerate, and if so how long we have been there for
-	extern ENGINE_API float GAverageFPS;
-	bool AtDesiredFramerate = GAverageFPS >= DesiredFrameRate;
+	const double TickDelta = TimeNow - LastTickTime;
 
+	// We tick at 60hz to try and hold a few seconds of data vs lots of high frequency ticks from a small window
+	if (TickDelta < (1 / kTickRate))
+	{
+		return false;
+	}
+
+	LastTickTime = TimeNow;
+
+	// Add fps sample but cap at 60 Hz. We don't want really high framerate periods to cause crappy ones to be ignored
+	AddTickRateSample(TickDelta);
+
+	const double AvgTickRate = CurrentAverageTickRate();
+
+	const double AverageFrameRate = AvgTickRate > 0 ? (1.f / AvgTickRate) : 0;
+	bool AtDesiredFramerate = AverageFrameRate >= DesiredFrameRate;
+
+	const double TimeAtFrameRate = AtDesiredFramerate ? (TimeNow - StartTimeOfAcceptableFrameRate) : 0;
 	const double ElapsedWaitTime = TimeNow - StartTimeOfWait;
 
+	// uncomment this for per-second reporting of current and average FPS
+	/*
+	const double CurrentFrameRate = 1.f / TickDelta;
+
+	static int LastSecond = 0;
+
+	if (ElapsedWaitTime > LastSecond)
+	{
+		UE_LOG(LogEngineAutomationLatentCommand, Display, TEXT("FWaitForInteractiveFrameRate: After %.f secs Current: %.02f, Average: %.02f, Desired: %.02f"), ElapsedWaitTime, CurrentFrameRate, AverageFrameRate, DesiredFrameRate);
+		LastSecond = FMath::CeilToInt(ElapsedWaitTime);
+	}
+	*/
+
 	if (!AtDesiredFramerate)
 	{
+		if (StartTimeOfAcceptableFrameRate > 0)
+		{
+			const double FailedTimeAtFrameRate = (TimeNow - StartTimeOfAcceptableFrameRate);
+
+			UE_LOG(LogEngineAutomationLatentCommand, Display, TEXT("FWaitForInteractiveFrameRate: Dropped below minimal framerate of %.f after %.02f secs. Now %.02f FPS "), 
+				DesiredFrameRate, FailedTimeAtFrameRate, AverageFrameRate);
+		}
 		StartTimeOfAcceptableFrameRate = 0;		
 	}
 	else
@@ -549,11 +639,9 @@ bool FWaitForInteractiveFrameRate::Update()
 		}
 		else
 		{
-			const double TimeAtFrameRate = TimeNow - StartTimeOfAcceptableFrameRate;
-
 			if (TimeAtFrameRate >= Duration)
 			{
-				UE_LOG(LogEngineAutomationLatentCommand, Log, TEXT("FWaitForInteractiveFrameRate: Hit %.02f FPS for %.f seconds after %.f seconds of waiting"), GAverageFPS, TimeAtFrameRate, ElapsedWaitTime);
+				UE_LOG(LogEngineAutomationLatentCommand, Log, TEXT("FWaitForInteractiveFrameRate: Hit %.02f FPS for %.f seconds after %.f seconds of waiting"), AverageFrameRate, TimeAtFrameRate, ElapsedWaitTime);
 				return true;
 			}
 		}
@@ -561,7 +649,7 @@ bool FWaitForInteractiveFrameRate::Update()
 
 	if (ElapsedWaitTime >= MaxWaitTime)
 	{
-		FString Msg = FString::Printf(TEXT("FWaitForInteractiveFrameRate: Game did not reach %.02f FPS within %.02f seconds. Current FPS=%.f. Giving up."), DesiredFrameRate, MaxWaitTime, GAverageFPS);
+		FString Msg = FString::Printf(TEXT("FWaitForInteractiveFrameRate: Game did not reach %.02f FPS within %.02f seconds. Current FPS=%.f. Giving up."), DesiredFrameRate, MaxWaitTime, AverageFrameRate);
 		UE_LOG(LogEngineAutomationLatentCommand, Display, TEXT("%s"), *Msg);
 
 		if (FAutomationTestBase* CurrentTest = FAutomationTestFramework::Get().GetCurrentTest())
@@ -574,7 +662,17 @@ bool FWaitForInteractiveFrameRate::Update()
 
 	if (TimeNow - LastReportTime >= kReportInterval)
 	{
-		UE_LOG(LogEngineAutomationLatentCommand, Display, TEXT("FWaitForInteractiveFrameRate: Waited %.f seconds. Will timeout in %.f Current FPS=%.f "), ElapsedWaitTime, MaxWaitTime- ElapsedWaitTime, GAverageFPS);
+		if (AtDesiredFramerate)
+		{
+			UE_LOG(LogEngineAutomationLatentCommand, Display, TEXT("FWaitForInteractiveFrameRate: Waited %.f seconds. Will timeout in %.f. Current FPS=%.f for %.02f secs"), 
+				ElapsedWaitTime, MaxWaitTime- ElapsedWaitTime, AverageFrameRate, TimeAtFrameRate);
+		}
+		else
+		{
+			UE_LOG(LogEngineAutomationLatentCommand, Display, TEXT("FWaitForInteractiveFrameRate: Waited %.f seconds. Will timeout in %.f. Current FPS=%.f "), 
+				ElapsedWaitTime, MaxWaitTime - ElapsedWaitTime, AverageFrameRate);
+
+		}
 		LastReportTime = TimeNow;
 	}
 
diff --git a/Engine/Source/Runtime/Engine/Public/Tests/AutomationCommon.h b/Engine/Source/Runtime/Engine/Public/Tests/AutomationCommon.h
--- a/Engine/Source/Runtime/Engine/Public/Tests/AutomationCommon.h
+++ b/Engine/Source/Runtime/Engine/Public/Tests/AutomationCommon.h
@@ -203,6 +203,12 @@ public:
 
 private:
 
+	// Add a sample to the rolling buffer where we hold tick rate
+	void AddTickRateSample(const double Value);
+
+	// return the average tick rate
+	double CurrentAverageTickRate() const;
+
 	// Time we began executing
 	double StartTimeOfWait;
 
@@ -211,6 +217,21 @@ private:
 
 	// time we last logged we are waiting
 	double LastReportTime;
+
+	// time of last tick
+	double LastTickTime;
+
+	// buffer of recent tick rate
+	TArray<double> RollingTickRateBuffer;
+
+	// index into the buffer
+	int32 BufferIndex;
+
+	// We tick at 60Hz
+	const double kTickRate = 60.0;
+
+	// How many samples we hold
+	const int kSampleCount = kTickRate * 5;
 };
 
 /**
