UE_NetConnectionPrivate::ValidateLevelVisibility in PIE now resets returned Package to null if package is a dynamic PIE package with pending external objects still loading.

diff --git a/Engine/Source/Runtime/Engine/Private/NetConnection.cpp b/Engine/Source/Runtime/Engine/Private/NetConnection.cpp
--- a/Engine/Source/Runtime/Engine/Private/NetConnection.cpp
+++ b/Engine/Source/Runtime/Engine/Private/NetConnection.cpp
@@ -135,6 +135,13 @@ namespace UE_NetConnectionPrivate
 		FValidateLevelVisibilityResult Result;
 		FString PackageNameStr = LevelVisibility.PackageName.ToString();
 		Result.Package = FindPackage(nullptr, *PackageNameStr);
+#if WITH_EDITOR
+		if (Result.Package && Result.Package->IsDynamicPIEPackagePending())
+		{
+			// Package is a dynamic PIE package with pending external objects still loading, mimic same behavior as if it was AsyncLoad
+			Result.Package = nullptr;
+		}
+#endif
 		Result.Linker = FLinkerLoad::FindExistingLinkerForPackage(Result.Package);
 
 		Result.StreamingLevel = FLevelUtils::FindStreamingLevel(World, LevelVisibility.PackageName);
