Convert bIsIntermediateNode to a property.

During graph duplication, we weren't preserving this flag when copying nodes, which we rely on for any post-graph expansion analysis.

diff --git a/Engine/Source/Runtime/Engine/Classes/EdGraph/EdGraphNode.h b/Engine/Source/Runtime/Engine/Classes/EdGraph/EdGraphNode.h
--- a/Engine/Source/Runtime/Engine/Classes/EdGraph/EdGraphNode.h
+++ b/Engine/Source/Runtime/Engine/Classes/EdGraph/EdGraphNode.h
@@ -323,6 +323,7 @@ public:
 
 private:
 	/** Whether the node was created as part of an expansion step */
+	UPROPERTY()
 	uint8 bIsIntermediateNode : 1;
 
 #if WITH_EDITORONLY_DATA
