separate arm64 ddc2 texture builds

diff --git a/Engine/Source/Runtime/Engine/Private/TextureDerivedDataBuildUtils.cpp b/Engine/Source/Runtime/Engine/Private/TextureDerivedDataBuildUtils.cpp
--- a/Engine/Source/Runtime/Engine/Private/TextureDerivedDataBuildUtils.cpp
+++ b/Engine/Source/Runtime/Engine/Private/TextureDerivedDataBuildUtils.cpp
@@ -85,6 +85,13 @@ static void WriteBuildSettings(FCbWriter& Writer, const FTextureBuildSettings& B
 
 	Writer.BeginObject();
 
+#if PLATFORM_CPU_ARM_FAMILY
+	// Separate out arm keys as x64 and arm64 clang do not generate the same data for a given
+	// input. Add the arm specifically so that a) we avoid rebuilding the current DDC and
+	// b) we can remove it once we get arm64 to be consistent.
+	WriteCbField<bool>(Writer, "bBuildIsArm64", true);
+#endif
+
 	if (BuildSettings.FormatConfigOverride)
 	{
 		Writer.AddObject("FormatConfigOverride", BuildSettings.FormatConfigOverride);
