FNames containing asset paths are deprecated. FSoftObjectPath or FTopLevelAssetPath should be used instead.
Fixups for World Partition + World Composition.

diff --git a/Engine/Plugins/Editor/DataValidation/Source/DataValidation/Private/WorldPartitionChangelistValidator.cpp b/Engine/Plugins/Editor/DataValidation/Source/DataValidation/Private/WorldPartitionChangelistValidator.cpp
--- a/Engine/Plugins/Editor/DataValidation/Source/DataValidation/Private/WorldPartitionChangelistValidator.cpp
+++ b/Engine/Plugins/Editor/DataValidation/Source/DataValidation/Private/WorldPartitionChangelistValidator.cpp
@@ -48,7 +48,7 @@ EDataValidationResult UWorldPartitionChangelistValidator::ValidateLoadedAsset_Im
 
 FString GetPrettyPackageName(const FWorldPartitionActorDescView& Desc)
 {
-	FString AssetPath = Desc.GetActorPath().ToString();
+	FString AssetPath = Desc.GetActorSoftPath().ToString();
 	
 	int32 LastDot = -1;
 	if (AssetPath.FindLastChar('.', LastDot))
@@ -75,7 +75,7 @@ EDataValidationResult UWorldPartitionChangelistValidator::ValidateActorsAndDataL
 	FSourceControlChangelistStatePtr ChangelistState = SourceControlProvider.GetState(Changelist->Changelist->AsShared(), EStateCacheUsage::Use);
 
 	// Figure out which world(s) these actors are in and split the files per world, and return the actor's native class
-	TMap<FName, TSet<FAssetData>> MapToActorsFiles;
+	TMap<FTopLevelAssetPath, TSet<FAssetData>> MapToActorsFiles;
 	auto TryAssociateActorToMap = [&MapToActorsFiles](const FAssetData& AssetData) -> UClass*
 	{
 		// Check that the asset is an actor
@@ -83,8 +83,8 @@ EDataValidationResult UWorldPartitionChangelistValidator::ValidateActorsAndDataL
 		{
 			// WorldPartition actors are all in OFPA mode so they're external
 			// Extract the MapName from the ObjectPath (<PathToPackage>.<MapName>:<Level>.<ActorName>)
-			FSoftObjectPath ActorPath = FSoftObjectPath(AssetData.ObjectPath);
-			FName MapAssetName = ActorPath.GetAssetPathName();
+			FSoftObjectPath ActorPath = AssetData.GetSoftObjectPath();
+			FTopLevelAssetPath MapAssetName = ActorPath.GetAssetPath();
 
 			TSet<FAssetData>* ActorFiles = MapToActorsFiles.Find(MapAssetName);
 
@@ -154,7 +154,7 @@ EDataValidationResult UWorldPartitionChangelistValidator::ValidateActorsAndDataL
 					{
 						if (ULevel::GetIsLevelPartitionedFromPackage(*PackageName))
 						{
-							MapToActorsFiles.FindOrAdd(AssetData.ObjectPath);
+							MapToActorsFiles.FindOrAdd(AssetData.GetSoftObjectPath().GetAssetPath());
 						}
 					}
 				}
@@ -163,13 +163,13 @@ EDataValidationResult UWorldPartitionChangelistValidator::ValidateActorsAndDataL
 	}
 
 	// For Each world 
-	for (TTuple<FName, TSet<FAssetData>>& It : MapToActorsFiles)
+	for (TTuple<FTopLevelAssetPath, TSet<FAssetData>>& It : MapToActorsFiles)
 	{
-		const FName& MapName = It.Get<0>();
+		const FTopLevelAssetPath& MapPath = It.Get<0>();
 		const TSet<FAssetData>& ActorsData = It.Get<1>();
 
 		// Find/Load the ActorDescContainer
-		UWorld* World = FindObject<UWorld>(nullptr, *MapName.ToString(), true);
+		UWorld* World = FindObject<UWorld>(nullptr, *MapPath.ToString(), true);
 
 		UActorDescContainer* ActorDescContainer = nullptr;
 		
@@ -181,14 +181,12 @@ EDataValidationResult UWorldPartitionChangelistValidator::ValidateActorsAndDataL
 		else
 		{
 			// Find in memory failed, load the ActorDescContainer
-			FSoftObjectPath MapPath = FSoftObjectPath(MapName);
-
 			ActorDescContainer = NewObject<UActorDescContainer>();
-			ActorDescContainer->Initialize(nullptr, MapPath.GetLongPackageFName());
+			ActorDescContainer->Initialize(nullptr, MapPath.GetPackageName());
 		}
 
 		// Build a set of Relevant Actor Guids to scope error messages to what's contained in the CL 
-		RelevantMap = MapName;
+		RelevantMap = MapPath;
 		RelevantActorGuids.Reset();
 
 		for (const FAssetData& ActorData : ActorsData)
@@ -203,8 +201,7 @@ EDataValidationResult UWorldPartitionChangelistValidator::ValidateActorsAndDataL
 		}
 
 		// Invoke static WorldPartition Validation from the ActorDescContainer
-		const FSoftObjectPath MapPath = FSoftObjectPath(MapName);
-		const bool bIsStreamingDisabled = ULevel::GetIsStreamingDisabledFromPackage(MapPath.GetLongPackageFName());
+		const bool bIsStreamingDisabled = ULevel::GetIsStreamingDisabledFromPackage(MapPath.GetPackageName());
 		UWorldPartition::CheckForErrors(this, ActorDescContainer, !bIsStreamingDisabled);
 	}
 
@@ -224,10 +221,10 @@ bool UWorldPartitionChangelistValidator::Filter(const FWorldPartitionActorDescVi
 		return true;
 	}
 
-	if (!RelevantMap.IsNone())
+	if (!RelevantMap.IsNull())
 	{
-		FSoftObjectPath ActorPath = ActorDescView.GetActorPath().ToString();
-		return ActorPath.GetAssetPathName() == RelevantMap;
+		FSoftObjectPath ActorPath = ActorDescView.GetActorSoftPath();
+		return ActorPath.GetAssetPath() == RelevantMap;
 	}
 
 	return false;
diff --git a/Engine/Plugins/Editor/DataValidation/Source/DataValidation/Private/WorldPartitionChangelistValidator.h b/Engine/Plugins/Editor/DataValidation/Source/DataValidation/Private/WorldPartitionChangelistValidator.h
--- a/Engine/Plugins/Editor/DataValidation/Source/DataValidation/Private/WorldPartitionChangelistValidator.h
+++ b/Engine/Plugins/Editor/DataValidation/Source/DataValidation/Private/WorldPartitionChangelistValidator.h
@@ -18,7 +18,7 @@ class DATAVALIDATION_API UWorldPartitionChangelistValidator : public UEditorVali
 protected:		
 	
 	TArray<FText>*							Errors = nullptr;
-	FName									RelevantMap;
+	FTopLevelAssetPath						RelevantMap;
 	TSet<FGuid>								RelevantActorGuids;
 	TSet<FString>							RelevantDataLayerAssets;
 	bool									SubmittingWorldDataLayers = false;
diff --git a/Engine/Plugins/Editor/WorldPartitionHLODUtilities/Source/Private/WorldPartitionHLODUtilities.cpp b/Engine/Plugins/Editor/WorldPartitionHLODUtilities/Source/Private/WorldPartitionHLODUtilities.cpp
--- a/Engine/Plugins/Editor/WorldPartitionHLODUtilities/Source/Private/WorldPartitionHLODUtilities.cpp
+++ b/Engine/Plugins/Editor/WorldPartitionHLODUtilities/Source/Private/WorldPartitionHLODUtilities.cpp
@@ -119,7 +119,10 @@ TArray<AWorldPartitionHLOD*> FWorldPartitionHLODUtilities::CreateHLODActors(FHLO
 			{
 				FSubActorsInfo& SubActorsInfo = SubActorsInfos.FindOrAdd(HLODLayer);
 
+				// Leaving this as deprecated for now until we fix up the serialization format for SubActorsInfo and do proper upgrade/deprecation 
+PRAGMA_DISABLE_DEPRECATION_WARNINGS
 				SubActorsInfo.SubActors.Emplace(ActorDescView.GetGuid(), ActorDescView.GetActorPackage(), ActorDescView.GetActorPath(), ActorInstance.GetContainerID(), ActorInstance.GetActorDescContainer()->GetContainerPackage(), ActorInstance.GetTransform());
+PRAGMA_ENABLE_DEPRECATION_WARNINGS
 				if (ActorDescView.GetIsSpatiallyLoaded())
 				{
 					SubActorsInfo.bIsSpatiallyLoaded = true;
diff --git a/Engine/Source/Editor/SceneOutliner/Private/ActorBrowsingMode.cpp b/Engine/Source/Editor/SceneOutliner/Private/ActorBrowsingMode.cpp
--- a/Engine/Source/Editor/SceneOutliner/Private/ActorBrowsingMode.cpp
+++ b/Engine/Source/Editor/SceneOutliner/Private/ActorBrowsingMode.cpp
@@ -2224,14 +2224,13 @@ bool FActorBrowsingMode::ConvertItemToAssetData(SceneOutliner::FilterBarType InI
 
 		return true;
 	}
-	// TODO: Maybe can be improved to not use TryConvertShortClassNameToPathName to optimize?
 	else if (const FActorDescTreeItem* ActorDescItem = InItem.CastTo<FActorDescTreeItem>())
 	{
 		if (const FWorldPartitionActorDesc* ActorDesc = ActorDescItem->ActorDescHandle.Get())
 		{
-			// For Unloaded Actors, grab the native class and try to convert it to a path name
+			// For Unloaded Actors, grab the native class 
 			OutAssetData = FAssetData();
-			FTopLevelAssetPath ClassPath = FAssetData::TryConvertShortClassNameToPathName(ActorDesc->GetNativeClass());
+			FTopLevelAssetPath ClassPath = ActorDesc->GetNativeClass();
 			if(ClassPath.IsValid())
 			{
 				OutAssetData.AssetClassPath = ClassPath;
diff --git a/Engine/Source/Editor/UnrealEd/Private/Commandlets/WorldPartitionConvertCommandlet.cpp b/Engine/Source/Editor/UnrealEd/Private/Commandlets/WorldPartitionConvertCommandlet.cpp
--- a/Engine/Source/Editor/UnrealEd/Private/Commandlets/WorldPartitionConvertCommandlet.cpp
+++ b/Engine/Source/Editor/UnrealEd/Private/Commandlets/WorldPartitionConvertCommandlet.cpp
@@ -463,7 +463,7 @@ void UWorldPartitionConvertCommandlet::FixupSoftObjectPaths(UPackage* OuterPacka
 				int32 DotPos = Value.GetSubPathString().Find(TEXT("."), ESearchCase::IgnoreCase, ESearchDir::FromStart);
 				if (DotPos != INDEX_NONE)
 				{
-					RemappedValue = RemapSoftObjectPaths.Find(Value.GetAssetPathName().ToString());
+					RemappedValue = RemapSoftObjectPaths.Find(Value.GetWithoutSubPath().ToString());
 					if (RemappedValue)
 					{
 						FString NewPath = *RemappedValue + ':' + Value.GetSubPathString();
diff --git a/Engine/Source/Editor/UnrealEd/Private/WorldPartition/WorldPartitionRenameDuplicateBuilder.cpp b/Engine/Source/Editor/UnrealEd/Private/WorldPartition/WorldPartitionRenameDuplicateBuilder.cpp
--- a/Engine/Source/Editor/UnrealEd/Private/WorldPartition/WorldPartitionRenameDuplicateBuilder.cpp
+++ b/Engine/Source/Editor/UnrealEd/Private/WorldPartition/WorldPartitionRenameDuplicateBuilder.cpp
@@ -398,13 +398,13 @@ bool UWorldPartitionRenameDuplicateBuilder::RunInternal(UWorld* World, const FCe
 				const FWorldPartitionActorDesc* NewActorDesc = NewWorld->GetWorldPartition()->GetActorDescContainer()->GetActorDesc(DuplicatedGuid ? *DuplicatedGuid : SourceActorDesc->GetGuid());
 				if (!NewActorDesc)
 				{
-					UE_LOG(LogWorldPartitionRenameDuplicateBuilder, Warning, TEXT("Failed to find source actor for Actor: %s"), *SourceActorDesc->GetActorPath().ToString());
+					UE_LOG(LogWorldPartitionRenameDuplicateBuilder, Warning, TEXT("Failed to find source actor for Actor: %s"), *SourceActorDesc->GetActorSoftPath().ToString());
 				}
 				else
 				{
 					if (NewActorDesc->GetReferences().Num() != SourceActorDesc->GetReferences().Num())
 					{
-						UE_LOG(LogWorldPartitionRenameDuplicateBuilder, Warning, TEXT("Actor: %s and Source Actor: %s have mismatching reference count"), *NewActorDesc->GetActorPath().ToString(), *SourceActorDesc->GetActorPath().ToString());
+						UE_LOG(LogWorldPartitionRenameDuplicateBuilder, Warning, TEXT("Actor: %s and Source Actor: %s have mismatching reference count"), *NewActorDesc->GetActorSoftPath().ToString(), *SourceActorDesc->GetActorSoftPath().ToString());
 					}
 					else
 					{
@@ -413,7 +413,7 @@ bool UWorldPartitionRenameDuplicateBuilder::RunInternal(UWorld* World, const FCe
 							FGuid* DuplicateReferenceGuid = DuplicatedActorGuids.Find(ReferenceGuid);
 							if (!NewActorDesc->GetReferences().Contains(DuplicateReferenceGuid ? *DuplicateReferenceGuid : ReferenceGuid))
 							{
-								UE_LOG(LogWorldPartitionRenameDuplicateBuilder, Warning, TEXT("Actor: %s and Source Actor: %s have mismatching reference"), *NewActorDesc->GetActorPath().ToString(), *SourceActorDesc->GetActorPath().ToString());
+								UE_LOG(LogWorldPartitionRenameDuplicateBuilder, Warning, TEXT("Actor: %s and Source Actor: %s have mismatching reference"), *NewActorDesc->GetActorSoftPath().ToString(), *SourceActorDesc->GetActorSoftPath().ToString());
 							}
 						}
 					}
diff --git a/Engine/Source/Runtime/Core/Public/UObject/UE5MainStreamObjectVersion.h b/Engine/Source/Runtime/Core/Public/UObject/UE5MainStreamObjectVersion.h
--- a/Engine/Source/Runtime/Core/Public/UObject/UE5MainStreamObjectVersion.h
+++ b/Engine/Source/Runtime/Core/Public/UObject/UE5MainStreamObjectVersion.h
@@ -277,6 +277,9 @@ struct CORE_API FUE5MainStreamObjectVersion
 		// Add flag to override MeshDeformer on a SkinnedMeshComponent.
 		AddSetMeshDeformerFlag,
 
+		// Replace FNames for class/actor paths with FSoftObjectPath
+		WorldPartitionActorDescActorAndClassPaths,
+
 		// -----<new versions can be added above this line>-------------------------------------------------
 		VersionPlusOne,
 		LatestVersion = VersionPlusOne - 1
diff --git a/Engine/Source/Runtime/Engine/Private/WorldComposition.cpp b/Engine/Source/Runtime/Engine/Private/WorldComposition.cpp
--- a/Engine/Source/Runtime/Engine/Private/WorldComposition.cpp
+++ b/Engine/Source/Runtime/Engine/Private/WorldComposition.cpp
@@ -182,7 +182,7 @@ struct FWorldTilesGatherer
 		for (const FString& FullPath : MapFilesToConsider)
 		{
 			FString TilePackageName = FPackageName::FilenameToLongPackageName(FullPath);
-			const FName TilePackagePath = *(TilePackageName + TEXT(".") + FPackageName::GetLongPackageAssetName(TilePackageName));
+			const FSoftObjectPath TilePackagePath(FName(*TilePackageName), FName(*FPackageName::GetLongPackageAssetName(TilePackageName)), FString{});
 
 			const FAssetData MapAssetData = AssetRegistry.GetAssetByObjectPath(TilePackagePath);
 
diff --git a/Engine/Source/Runtime/Engine/Private/WorldPartition/ErrorHandling/WorldPartitionStreamingGenerationErrorHandler.cpp b/Engine/Source/Runtime/Engine/Private/WorldPartition/ErrorHandling/WorldPartitionStreamingGenerationErrorHandler.cpp
--- a/Engine/Source/Runtime/Engine/Private/WorldPartition/ErrorHandling/WorldPartitionStreamingGenerationErrorHandler.cpp
+++ b/Engine/Source/Runtime/Engine/Private/WorldPartition/ErrorHandling/WorldPartitionStreamingGenerationErrorHandler.cpp
@@ -7,7 +7,7 @@
 // Helpers
 FString IStreamingGenerationErrorHandler::GetFullActorName(const FWorldPartitionActorDescView& ActorDescView)
 {
-	const FSoftObjectPath ActorPath(ActorDescView.GetActorPath());
+	const FSoftObjectPath ActorPath = ActorDescView.GetActorSoftPath();
 	return ActorDescView.GetActorLabelOrName().ToString() + TEXT(" (") + ActorPath.GetLongPackageName() + TEXT(")");
 }
 #endif
diff --git a/Engine/Source/Runtime/Engine/Private/WorldPartition/ErrorHandling/WorldPartitionStreamingGenerationTokenizedMessageErrorHandler.cpp b/Engine/Source/Runtime/Engine/Private/WorldPartition/ErrorHandling/WorldPartitionStreamingGenerationTokenizedMessageErrorHandler.cpp
--- a/Engine/Source/Runtime/Engine/Private/WorldPartition/ErrorHandling/WorldPartitionStreamingGenerationTokenizedMessageErrorHandler.cpp
+++ b/Engine/Source/Runtime/Engine/Private/WorldPartition/ErrorHandling/WorldPartitionStreamingGenerationTokenizedMessageErrorHandler.cpp
@@ -14,7 +14,7 @@ void ITokenizedMessageErrorHandler::OnInvalidReference(const FWorldPartitionActo
 {
 	TSharedRef<FTokenizedMessage> Message = FTokenizedMessage::Create(EMessageSeverity::Warning);
 	Message->AddToken(FTextToken::Create(LOCTEXT("TokenMessage_WorldPartition_Actor", "Actor")))
-		->AddToken(FActorToken::Create(ActorDescView.GetActorPath().ToString(), ActorDescView.GetGuid(), FText::FromString(GetFullActorName(ActorDescView))))
+		->AddToken(FActorToken::Create(ActorDescView.GetActorSoftPath().ToString(), ActorDescView.GetGuid(), FText::FromString(GetFullActorName(ActorDescView))))
 		->AddToken(FTextToken::Create(LOCTEXT("TokenMessage_WorldPartition_HaveMissingRefsTo", "have missing references to")))
 		->AddToken(FTextToken::Create(FText::FromString(ReferenceGuid.ToString())));
 
@@ -30,10 +30,10 @@ void ITokenizedMessageErrorHandler::OnInvalidReferenceGridPlacement(const FWorld
 
 	TSharedRef<FTokenizedMessage> Message = FTokenizedMessage::Create(EMessageSeverity::Error);
 	Message->AddToken(FTextToken::Create(ActorDescView.GetIsSpatiallyLoaded() ? SpatiallyLoadedActor : NonSpatiallyLoadedActor))
-		->AddToken(FActorToken::Create(ActorDescView.GetActorPath().ToString(), ActorDescView.GetGuid(), FText::FromString(GetFullActorName(ActorDescView))))
+		->AddToken(FActorToken::Create(ActorDescView.GetActorSoftPath().ToString(), ActorDescView.GetGuid(), FText::FromString(GetFullActorName(ActorDescView))))
 		->AddToken(FTextToken::Create(LOCTEXT("TokenMessage_WorldPartition_References", "references")))
 		->AddToken(FTextToken::Create(ReferenceActorDescView.GetIsSpatiallyLoaded() ? SpatiallyLoadedActor : NonSpatiallyLoadedActor))
-		->AddToken(FActorToken::Create(ReferenceActorDescView.GetActorPath().ToString(), ReferenceActorDescView.GetGuid(), FText::FromName(ReferenceActorDescView.GetActorLabelOrName())));
+		->AddToken(FActorToken::Create(ReferenceActorDescView.GetActorSoftPath().ToString(), ReferenceActorDescView.GetGuid(), FText::FromName(ReferenceActorDescView.GetActorLabelOrName())));
 
 	AddAdditionalNameToken(Message, FName(TEXT("WorldPartition_StreamedActorReferenceAlwaysLoadedActor_CheckForErrors")));
 
@@ -44,9 +44,9 @@ void ITokenizedMessageErrorHandler::OnInvalidReferenceDataLayers(const FWorldPar
 {
 	TSharedRef<FTokenizedMessage> Message = FTokenizedMessage::Create(EMessageSeverity::Error);
 	Message->AddToken(FTextToken::Create(LOCTEXT("TokenMessage_WorldPartition_Actor", "Actor")))
-		->AddToken(FActorToken::Create(ActorDescView.GetActorPath().ToString(), ActorDescView.GetGuid(),  FText::FromString(GetFullActorName(ActorDescView))))
+		->AddToken(FActorToken::Create(ActorDescView.GetActorSoftPath().ToString(), ActorDescView.GetGuid(),  FText::FromString(GetFullActorName(ActorDescView))))
 		->AddToken(FTextToken::Create(LOCTEXT("TokenMessage_WorldPartition_ReferenceActorInOtherDataLayers", "references an actor in a different set of runtime data layers")))
-		->AddToken(FActorToken::Create(ReferenceActorDescView.GetActorPath().ToString(), ReferenceActorDescView.GetGuid(), FText::FromName(ReferenceActorDescView.GetActorLabelOrName())));
+		->AddToken(FActorToken::Create(ReferenceActorDescView.GetActorSoftPath().ToString(), ReferenceActorDescView.GetGuid(), FText::FromName(ReferenceActorDescView.GetActorLabelOrName())));
 
 	AddAdditionalNameToken(Message, FName(TEXT("WorldPartition_ActorReferenceActorInAnotherDataLayer_CheckForErrors")));
 
@@ -57,9 +57,9 @@ void ITokenizedMessageErrorHandler::OnInvalidReferenceRuntimeGrid(const FWorldPa
 {
 	TSharedRef<FTokenizedMessage> Message = FTokenizedMessage::Create(EMessageSeverity::Error);
 	Message->AddToken(FTextToken::Create(LOCTEXT("TokenMessage_WorldPartition_Actor", "Actor")))
-		->AddToken(FActorToken::Create(ActorDescView.GetActorPath().ToString(), ActorDescView.GetGuid(), FText::FromString(GetFullActorName(ActorDescView))))
+		->AddToken(FActorToken::Create(ActorDescView.GetActorSoftPath().ToString(), ActorDescView.GetGuid(), FText::FromString(GetFullActorName(ActorDescView))))
 		->AddToken(FTextToken::Create(LOCTEXT("TokenMessage_WorldPartition_ReferenceActorInOtherRuntimeGrid", "references an actor in a different runtime grid")))
-		->AddToken(FActorToken::Create(ReferenceActorDescView.GetActorPath().ToString(), ReferenceActorDescView.GetGuid(), FText::FromName(ReferenceActorDescView.GetActorLabelOrName())));
+		->AddToken(FActorToken::Create(ReferenceActorDescView.GetActorSoftPath().ToString(), ReferenceActorDescView.GetGuid(), FText::FromName(ReferenceActorDescView.GetActorLabelOrName())));
 
 	AddAdditionalNameToken(Message, FName(TEXT("WorldPartition_ActorReferenceActorInAnotherRuntimeGrid_CheckForErrors")));
 
@@ -70,7 +70,7 @@ void ITokenizedMessageErrorHandler::OnInvalidReferenceLevelScriptStreamed(const
 {
 	TSharedRef<FTokenizedMessage> Message = FTokenizedMessage::Create(EMessageSeverity::Error);
 	Message->AddToken(FTextToken::Create(LOCTEXT("TokenMessage_WorldPartition_LevelScriptBlueprintStreamedActorReference", "Level Script Blueprint references streamed actor")))
-		->AddToken(FActorToken::Create(ActorDescView.GetActorPath().ToString(), ActorDescView.GetGuid(), FText::FromString(GetFullActorName(ActorDescView))));
+		->AddToken(FActorToken::Create(ActorDescView.GetActorSoftPath().ToString(), ActorDescView.GetGuid(), FText::FromString(GetFullActorName(ActorDescView))));
 
 	AddAdditionalNameToken(Message, FName(TEXT("WorldPartition_LevelScriptBlueprintRefefenceStreamed_CheckForErrors")));
 
@@ -82,7 +82,7 @@ void ITokenizedMessageErrorHandler::OnInvalidReferenceLevelScriptDataLayers(cons
 {
 	TSharedRef<FTokenizedMessage> Message = FTokenizedMessage::Create(EMessageSeverity::Error);
 	Message->AddToken(FTextToken::Create(LOCTEXT("TokenMessage_WorldPartition_LevelScriptBlueprintActorReference", "Level Script Blueprint references actor")))
-		->AddToken(FActorToken::Create(ActorDescView.GetActorPath().ToString(), ActorDescView.GetGuid(), FText::FromString(GetFullActorName(ActorDescView))))
+		->AddToken(FActorToken::Create(ActorDescView.GetActorSoftPath().ToString(), ActorDescView.GetGuid(), FText::FromString(GetFullActorName(ActorDescView))))
 		->AddToken(FTextToken::Create(LOCTEXT("TokenMessage_WorldPartition_LevelScriptBlueprintDataLayerReference", "with a non empty set of data layers")));
 
 	AddAdditionalNameToken(Message, FName(TEXT("WorldPartition_LevelScriptBlueprintRefefenceDataLayer_CheckForErrors")));
@@ -135,7 +135,7 @@ void ITokenizedMessageErrorHandler::OnActorNeedsResave(const FWorldPartitionActo
 {
 	TSharedRef<FTokenizedMessage> Message = FTokenizedMessage::Create(EMessageSeverity::Warning);
 	Message->AddToken(FTextToken::Create(LOCTEXT("TokenMessage_WorldPartition_ActorNeedsResave", "Actor needs resave")))
-		->AddToken(FActorToken::Create(ActorDescView.GetActorPath().ToString(), ActorDescView.GetGuid(), FText::FromString(GetFullActorName(ActorDescView))));
+		->AddToken(FActorToken::Create(ActorDescView.GetActorSoftPath().ToString(), ActorDescView.GetGuid(), FText::FromString(GetFullActorName(ActorDescView))));
 
 	AddAdditionalNameToken(Message, FName(TEXT("WorldPartition_ActorNeedsResave_CheckForErrors")));
 
@@ -144,7 +144,7 @@ void ITokenizedMessageErrorHandler::OnActorNeedsResave(const FWorldPartitionActo
 
 void ITokenizedMessageErrorHandler::OnLevelInstanceInvalidWorldAsset(const FWorldPartitionActorDescView& ActorDescView, FName WorldAsset, ELevelInstanceInvalidReason Reason)
 {
-	FSoftObjectPath ActorPath(ActorDescView.GetActorPath());
+	FSoftObjectPath ActorPath(ActorDescView.GetActorSoftPath());
 	EMessageSeverity::Type MessageSeverity = EMessageSeverity::Info;
 	FText ReasonText;
 
diff --git a/Engine/Source/Runtime/Engine/Private/WorldPartition/HLOD/HLODLayer.cpp b/Engine/Source/Runtime/Engine/Private/WorldPartition/HLOD/HLODLayer.cpp
--- a/Engine/Source/Runtime/Engine/Private/WorldPartition/HLOD/HLODLayer.cpp
+++ b/Engine/Source/Runtime/Engine/Private/WorldPartition/HLOD/HLODLayer.cpp
@@ -61,7 +61,7 @@ UHLODLayer* UHLODLayer::GetHLODLayer(const FWorldPartitionActorDescView& InActor
 	check(InWorldPartition);
 
 	const FName HLODLayerName = InActorDesc.GetHLODLayer();
-	if (UHLODLayer* HLODLayer = HLODLayerName.IsNone() ? nullptr : Cast<UHLODLayer>(FSoftObjectPath(HLODLayerName).TryLoad()))
+	if (UHLODLayer* HLODLayer = HLODLayerName.IsNone() ? nullptr : Cast<UHLODLayer>(FSoftObjectPath(HLODLayerName.ToString()).TryLoad()))
 	{
 		return HLODLayer;
 	}
diff --git a/Engine/Source/Runtime/Engine/Private/WorldPartition/WorldPartitionActorDesc.cpp b/Engine/Source/Runtime/Engine/Private/WorldPartition/WorldPartitionActorDesc.cpp
--- a/Engine/Source/Runtime/Engine/Private/WorldPartition/WorldPartitionActorDesc.cpp
+++ b/Engine/Source/Runtime/Engine/Private/WorldPartition/WorldPartitionActorDesc.cpp
@@ -8,6 +8,7 @@
 #include "UObject/LinkerInstancingContext.h"
 #include "UObject/UObjectHash.h"
 #include "UObject/MetaData.h"
+#include "AssetRegistry/AssetData.h"
 #include "Algo/Transform.h"
 #include "Engine/World.h"
 #include "Engine/Level.h"
@@ -317,10 +318,34 @@ void FWorldPartitionActorDesc::Serialize(FArchive& Ar)
 
 	if (Ar.CustomVer(FFortniteNCBranchObjectVersion::GUID) >= FFortniteNCBranchObjectVersion::WorldPartitionActorDescNativeBaseClassSerialization)
 	{
-		Ar << BaseClass;
+		if (Ar.CustomVer(FUE5MainStreamObjectVersion::GUID) < FUE5MainStreamObjectVersion::WorldPartitionActorDescActorAndClassPaths)
+		{
+			FName BaseClassPathName;
+			Ar << BaseClassPathName;
+PRAGMA_DISABLE_DEPRECATION_WARNINGS
+			BaseClass = FAssetData::TryConvertShortClassNameToPathName(BaseClassPathName);
+PRAGMA_ENABLE_DEPRECATION_WARNINGS
+		}
+		else
+		{
+			Ar << BaseClass;
+		}
 	}
 
-	Ar << NativeClass << Guid;
+	if (Ar.CustomVer(FUE5MainStreamObjectVersion::GUID) < FUE5MainStreamObjectVersion::WorldPartitionActorDescActorAndClassPaths)
+	{
+		FName NativeClassPathName;
+		Ar << NativeClassPathName;
+PRAGMA_DISABLE_DEPRECATION_WARNINGS
+		NativeClass = FAssetData::TryConvertShortClassNameToPathName(NativeClassPathName);
+PRAGMA_ENABLE_DEPRECATION_WARNINGS
+	}
+	else
+	{
+		Ar << NativeClass;
+	}
+
+	Ar << Guid;
 
 	if(Ar.CustomVer(FUE5ReleaseStreamObjectVersion::GUID) < FUE5ReleaseStreamObjectVersion::LargeWorldCoordinates)
 	{
@@ -370,7 +395,19 @@ void FWorldPartitionActorDesc::Serialize(FArchive& Ar)
 
 	if (Ar.CustomVer(FUE5MainStreamObjectVersion::GUID) < FUE5MainStreamObjectVersion::WorldPartitionActorDescSerializeArchivePersistent)
 	{
-		Ar << ActorPackage << ActorPath;
+		Ar << ActorPackage;
+		if (Ar.CustomVer(FUE5MainStreamObjectVersion::GUID) < FUE5MainStreamObjectVersion::WorldPartitionActorDescActorAndClassPaths)
+		{ 
+			FName ActorPathName;
+			Ar << ActorPathName;
+PRAGMA_DISABLE_DEPRECATION_WARNINGS
+			ActorPath = FSoftObjectPath(ActorPathName);
+PRAGMA_ENABLE_DEPRECATION_WARNINGS
+		}
+		else
+		{
+			Ar << ActorPath;
+		}
 	}
 
 	if (Ar.CustomVer(FUE5MainStreamObjectVersion::GUID) >= FUE5MainStreamObjectVersion::WorldPartitionActorDescSerializeDataLayers)
@@ -438,20 +475,21 @@ FName FWorldPartitionActorDesc::GetActorLabelOrName() const
 
 FName FWorldPartitionActorDesc::GetDisplayClassName() const
 {
-	auto GetCleanClassName = [](FName ClassName) -> FName
+	auto GetCleanClassName = [](FTopLevelAssetPath ClassPath) -> FName
 	{
+		// Should this just return ClassPath.GetAssetName() with _C removed if necessary?
 		int32 Index;
-		FString ClassNameStr(ClassName.ToString());
+		FString ClassNameStr(ClassPath.ToString());
 		if (ClassNameStr.FindLastChar(TCHAR('.'), Index))
 		{
 			FString CleanClassName = ClassNameStr.Mid(Index + 1);
 			CleanClassName.RemoveFromEnd(TEXT("_C"));
 			return *CleanClassName;
 		}
-		return ClassName;
+		return *ClassPath.ToString(); 
 	};
 
-	return BaseClass.IsNone() ? GetCleanClassName(NativeClass) : GetCleanClassName(BaseClass);
+	return BaseClass.IsNull() ? GetCleanClassName(NativeClass) : GetCleanClassName(BaseClass);
 }
 
 FGuid FWorldPartitionActorDesc::GetContentBundleUID() const
@@ -520,7 +558,7 @@ AActor* FWorldPartitionActorDesc::Load() const
 		if (InstancingContext)
 		{
 			FName RemappedPackageName = InstancingContext->RemapPackage(ActorPackage);
-			check(RemappedPackageName != ActorPath);
+			check(RemappedPackageName != ActorPath.GetLongPackageFName());
 
 			Package = CreatePackage(*RemappedPackageName.ToString());
 		}
diff --git a/Engine/Source/Runtime/Engine/Private/WorldPartition/WorldPartitionActorDescUtils.cpp b/Engine/Source/Runtime/Engine/Private/WorldPartition/WorldPartitionActorDescUtils.cpp
--- a/Engine/Source/Runtime/Engine/Private/WorldPartition/WorldPartitionActorDescUtils.cpp
+++ b/Engine/Source/Runtime/Engine/Private/WorldPartition/WorldPartitionActorDescUtils.cpp
@@ -60,7 +60,7 @@ TUniquePtr<FWorldPartitionActorDesc> FWorldPartitionActorDescUtils::GetActorDesc
 		FWorldPartitionActorDescInitData ActorDescInitData;
 		ActorDescInitData.NativeClass = ActorClass;
 		ActorDescInitData.PackageName = InAssetData.PackageName;
-		ActorDescInitData.ActorPath = InAssetData.ObjectPath;
+		ActorDescInitData.ActorPath = InAssetData.GetSoftObjectPath();
 		FBase64::Decode(ActorMetaDataStr, ActorDescInitData.SerializedData);
 
 		TUniquePtr<FWorldPartitionActorDesc> NewActorDesc(AActor::StaticCreateClassActorDesc(ActorDescInitData.NativeClass));
diff --git a/Engine/Source/Runtime/Engine/Private/WorldPartition/WorldPartitionActorDescView.cpp b/Engine/Source/Runtime/Engine/Private/WorldPartition/WorldPartitionActorDescView.cpp
--- a/Engine/Source/Runtime/Engine/Private/WorldPartition/WorldPartitionActorDescView.cpp
+++ b/Engine/Source/Runtime/Engine/Private/WorldPartition/WorldPartitionActorDescView.cpp
@@ -23,12 +23,12 @@ const FGuid& FWorldPartitionActorDescView::GetGuid() const
 	return ActorDesc->GetGuid();
 }
 
-FName FWorldPartitionActorDescView::GetBaseClass() const
+FTopLevelAssetPath FWorldPartitionActorDescView::GetBaseClass() const
 {
 	return ActorDesc->GetBaseClass();
 }
 
-FName FWorldPartitionActorDescView::GetNativeClass() const
+FTopLevelAssetPath FWorldPartitionActorDescView::GetNativeClass() const
 {
 	return ActorDesc->GetNativeClass();
 }
@@ -99,9 +99,9 @@ FName FWorldPartitionActorDescView::GetActorPackage() const
 	return ActorDesc->GetActorPackage();
 }
 
-FName FWorldPartitionActorDescView::GetActorPath() const
+FSoftObjectPath FWorldPartitionActorDescView::GetActorSoftPath() const
 {
-	return ActorDesc->GetActorPath();
+	return ActorDesc->GetActorSoftPath();
 }
 
 FName FWorldPartitionActorDescView::GetActorLabel() const
diff --git a/Engine/Source/Runtime/Engine/Private/WorldPartition/WorldPartitionConverter.cpp b/Engine/Source/Runtime/Engine/Private/WorldPartition/WorldPartitionConverter.cpp
--- a/Engine/Source/Runtime/Engine/Private/WorldPartition/WorldPartitionConverter.cpp
+++ b/Engine/Source/Runtime/Engine/Private/WorldPartition/WorldPartitionConverter.cpp
@@ -307,7 +307,7 @@ void FWorldPartitionConverter::FixupSoftObjectPaths(UPackage* OuterPackage)
 				int32 DotPos = Value.GetSubPathString().Find(TEXT("."), ESearchCase::IgnoreCase, ESearchDir::FromStart);
 				if (DotPos != INDEX_NONE)
 				{
-					RemappedValue = RemapSoftObjectPaths.Find(Value.GetAssetPathName().ToString());
+					RemappedValue = RemapSoftObjectPaths.Find(Value.GetWithoutSubPath().ToString());
 					if (RemappedValue)
 					{
 						FString NewPath = *RemappedValue + ':' + Value.GetSubPathString();
diff --git a/Engine/Source/Runtime/Engine/Private/WorldPartition/WorldPartitionHelpers.cpp b/Engine/Source/Runtime/Engine/Private/WorldPartition/WorldPartitionHelpers.cpp
--- a/Engine/Source/Runtime/Engine/Private/WorldPartition/WorldPartitionHelpers.cpp
+++ b/Engine/Source/Runtime/Engine/Private/WorldPartition/WorldPartitionHelpers.cpp
@@ -21,13 +21,13 @@ bool FWorldPartitionHelpers::IsActorDescClassCompatibleWith(const FWorldPartitio
 
 	if (!Class->IsNative())
 	{
-		if (FName ActorBaseClassName = ActorDesc->GetBaseClass(); !ActorBaseClassName.IsNone())
+		if (FTopLevelAssetPath ActorBasePath = ActorDesc->GetBaseClass(); !ActorBasePath.IsNull())
 		{
-			ActorBaseClass = LoadClass<AActor>(nullptr, *ActorBaseClassName.ToString(), nullptr, LOAD_None, nullptr);
+			ActorBaseClass = LoadClass<AActor>(nullptr, *ActorBasePath.ToString(), nullptr, LOAD_None, nullptr);
 
 			if (!ActorBaseClass)
 			{
-				UE_LOG(LogWorldPartitionHelpers, Warning, TEXT("Failed to find actor base class: %s."), *ActorBaseClassName.ToString());
+				UE_LOG(LogWorldPartitionHelpers, Warning, TEXT("Failed to find actor base class: %s."), *ActorBasePath.ToString());
 				ActorBaseClass = ActorNativeClass;
 			}
 		}
diff --git a/Engine/Source/Runtime/Engine/Private/WorldPartition/WorldPartitionLevelHelper.cpp b/Engine/Source/Runtime/Engine/Private/WorldPartition/WorldPartitionLevelHelper.cpp
--- a/Engine/Source/Runtime/Engine/Private/WorldPartition/WorldPartitionLevelHelper.cpp
+++ b/Engine/Source/Runtime/Engine/Private/WorldPartition/WorldPartitionLevelHelper.cpp
@@ -226,10 +226,10 @@ bool FWorldPartitionLevelHelper::RemapActorPath(const FActorContainerID& InConta
 		const FString NewSubPathString = FWorldPartitionLevelHelper::AddActorContainerIDToSubPathString(InContainerID, SoftObjectPath.GetSubPathString());
 
 		FNameBuilder NewAssetPathBuilder;
-		SoftObjectPath.GetAssetPathName().ToString(NewAssetPathBuilder);
+		SoftObjectPath.GetAssetPath().AppendString(NewAssetPathBuilder);
 		NewAssetPathBuilder.ReplaceAt(0, LongPackageName.Len(), ContainerPackageString);
 
-		OutActorPath = FSoftObjectPath(*NewAssetPathBuilder, NewSubPathString).ToString();
+		OutActorPath = FSoftObjectPath(FTopLevelAssetPath(NewAssetPathBuilder.ToView()), NewSubPathString).ToString();
 		return true;
 	}
 
diff --git a/Engine/Source/Runtime/Engine/Private/WorldPartition/WorldPartitionLevelStreamingPolicy.cpp b/Engine/Source/Runtime/Engine/Private/WorldPartition/WorldPartitionLevelStreamingPolicy.cpp
--- a/Engine/Source/Runtime/Engine/Private/WorldPartition/WorldPartitionLevelStreamingPolicy.cpp
+++ b/Engine/Source/Runtime/Engine/Private/WorldPartition/WorldPartitionLevelStreamingPolicy.cpp
@@ -143,8 +143,7 @@ void UWorldPartitionLevelStreamingPolicy::RemapSoftObjectPath(FSoftObjectPath& O
 			}
 
 			// Use the WorldPartition world name here instead of using the world name from the path to support converting level instance paths to main world paths.
-			ObjectPath.SetAssetPathName(FName(*FString::Printf(TEXT("%s%s.%s"), *PrefixPath, *PackagePath, *OuterWorld->GetName())));
-			ObjectPath.SetSubPathString(SrcObjectPath.GetSubPathString());
+			ObjectPath = FSoftObjectPath(FTopLevelAssetPath(FString::Printf(TEXT("%s%s.%s"), *PrefixPath, *PackagePath, *OuterWorld->GetName())), SrcObjectPath.GetSubPathString());
 			// Put back PIE prefix
 			if (OuterWorld->IsPlayInEditor() && (PIEInstanceID != INDEX_NONE))
 			{
diff --git a/Engine/Source/Runtime/Engine/Private/WorldPartition/WorldPartitionRuntimeLevelStreamingCell.cpp b/Engine/Source/Runtime/Engine/Private/WorldPartition/WorldPartitionRuntimeLevelStreamingCell.cpp
--- a/Engine/Source/Runtime/Engine/Private/WorldPartition/WorldPartitionRuntimeLevelStreamingCell.cpp
+++ b/Engine/Source/Runtime/Engine/Private/WorldPartition/WorldPartitionRuntimeLevelStreamingCell.cpp
@@ -101,7 +101,10 @@ void UWorldPartitionRuntimeLevelStreamingCell::SetIsAlwaysLoaded(bool bInIsAlway
 void UWorldPartitionRuntimeLevelStreamingCell::AddActorToCell(const FWorldPartitionActorDescView& ActorDescView, const FActorContainerID& InContainerID, const FTransform& InContainerTransform, const UActorDescContainer* InContainer)
 {
 	check(!ActorDescView.GetActorIsEditorOnly());
+	// Leaving this using deprecated functions until the serialization format of Packages is updated.
+PRAGMA_DISABLE_DEPRECATION_WARNINGS
 	Packages.Emplace(ActorDescView.GetActorPackage(), ActorDescView.GetActorPath(), InContainerID, InContainerTransform, InContainer->GetContainerPackage(), GetWorld()->GetPackage()->GetFName());
+PRAGMA_ENABLE_DEPRECATION_WARNINGS
 }
 
 UWorldPartitionLevelStreamingDynamic* UWorldPartitionRuntimeLevelStreamingCell::CreateLevelStreaming(const FString& InPackageName) const
diff --git a/Engine/Source/Runtime/Engine/Private/WorldPartition/WorldPartitionRuntimeSpatialHash.cpp b/Engine/Source/Runtime/Engine/Private/WorldPartition/WorldPartitionRuntimeSpatialHash.cpp
--- a/Engine/Source/Runtime/Engine/Private/WorldPartition/WorldPartitionRuntimeSpatialHash.cpp
+++ b/Engine/Source/Runtime/Engine/Private/WorldPartition/WorldPartitionRuntimeSpatialHash.cpp
@@ -1063,7 +1063,7 @@ bool UWorldPartitionRuntimeSpatialHash::CreateStreamingGrid(const FSpatialHashRu
 							IStreamingGenerationContext::FActorInstance ActorInstance(ActorGuid, ActorSetInstance);
 							const FWorldPartitionActorDescView& ActorDescView = ActorInstance.GetActorDescView();
 
-							if (AActor* AlwaysLoadedActor = FindObject<AActor>(nullptr, *ActorDescView.GetActorPath().ToString()))
+							if (AActor* AlwaysLoadedActor = FindObject<AActor>(nullptr, *ActorDescView.GetActorSoftPath().ToString()))
 							{
 								AlwaysLoadedActorsForPIE.Emplace(Reference, AlwaysLoadedActor);
 
@@ -1128,7 +1128,7 @@ bool UWorldPartitionRuntimeSpatialHash::CreateStreamingGrid(const FSpatialHashRu
 				for (const IStreamingGenerationContext::FActorInstance& ActorInstance : FilteredActors)
 				{
 					const FWorldPartitionActorDescView& ActorDescView = ActorInstance.GetActorDescView();
-					if (AActor* Actor = FindObject<AActor>(nullptr, *ActorDescView.GetActorPath().ToString()))
+					if (AActor* Actor = FindObject<AActor>(nullptr, *ActorDescView.GetActorSoftPath().ToString()))
 					{
 						if (ModifiedActorDescListForPIE.GetActorDesc(ActorDescView.GetGuid()) != nullptr)
 						{
@@ -1150,7 +1150,7 @@ bool UWorldPartitionRuntimeSpatialHash::CreateStreamingGrid(const FSpatialHashRu
 
 					if (ActorInstance.GetContainerID().IsMainContainer() && StreamingCell->UnsavedActorsContainer)
 					{
-						if (AActor* Actor = FindObject<AActor>(nullptr, *ActorDescView.GetActorPath().ToString()))
+						if (AActor* Actor = FindObject<AActor>(nullptr, *ActorDescView.GetActorSoftPath().ToString()))
 						{
 							StreamingCell->UnsavedActorsContainer->Actors.Add(Actor->GetFName(), Actor);
 
diff --git a/Engine/Source/Runtime/Engine/Private/WorldPartitionBlueprintLibrary.cpp b/Engine/Source/Runtime/Engine/Private/WorldPartitionBlueprintLibrary.cpp
--- a/Engine/Source/Runtime/Engine/Private/WorldPartitionBlueprintLibrary.cpp
+++ b/Engine/Source/Runtime/Engine/Private/WorldPartitionBlueprintLibrary.cpp
@@ -21,13 +21,13 @@ FActorDesc::FActorDesc(const FWorldPartitionActorDesc& InActorDesc, const FTrans
 {
 	Guid = InActorDesc.GetGuid();
 
-	if (!InActorDesc.GetBaseClass().IsNone())
+	if (!InActorDesc.GetBaseClass().IsNull())
 	{
-		Class = InActorDesc.GetBaseClass();
+		Class = FSoftObjectPath(InActorDesc.GetBaseClass(), {});
 	}
 	else
 	{
-		Class = InActorDesc.GetActorNativeClass();
+		Class = FSoftObjectPath(InActorDesc.GetActorNativeClass());
 	}
 
 	Name = InActorDesc.GetActorName();
diff --git a/Engine/Source/Runtime/Engine/Public/WorldPartition/WorldPartitionActorDesc.h b/Engine/Source/Runtime/Engine/Public/WorldPartition/WorldPartitionActorDesc.h
--- a/Engine/Source/Runtime/Engine/Public/WorldPartition/WorldPartitionActorDesc.h
+++ b/Engine/Source/Runtime/Engine/Public/WorldPartition/WorldPartitionActorDesc.h
@@ -6,6 +6,8 @@
 #include "PropertyPairsMap.h"
 #include "UObject/NameTypes.h"
 #include "UObject/ObjectMacros.h"
+#include "UObject/SoftObjectPath.h"
+#include "UObject/TopLevelAssetPath.h"
 #include "UObject/WeakObjectPtrTemplates.h"
 #include "UObject/WeakObjectPtr.h"
 #include "Templates/SubclassOf.h"
@@ -16,7 +18,7 @@ struct FWorldPartitionActorDescInitData
 {
 	UClass* NativeClass;
 	FName PackageName;
-	FName ActorPath;
+	FSoftObjectPath ActorPath;
 	TArray<uint8> SerializedData;
 };
 
@@ -73,11 +75,15 @@ public:
 
 	inline const FGuid& GetGuid() const { return Guid; }
 	
+PRAGMA_DISABLE_DEPRECATION_WARNINGS
 	UE_DEPRECATED(5.1, "GetClass is deprecated, GetBaseClass or GetNativeClass should be used instead.")
-	inline FName GetClass() const { return GetNativeClass(); }
+	inline FName GetClass() const { return GetNativeClass().ToFName(); }
+	UE_DEPRECATED(5.1, "GetActorPath is deprecated, GetActorSoftPath should be used instead.")
+	inline FName GetActorPath() const { return ActorPath.ToFName(); }
+PRAGMA_ENABLE_DEPRECATION_WARNINGS
 
-	inline FName GetBaseClass() const { return BaseClass; }
-	inline FName GetNativeClass() const { return NativeClass; }
+	inline FTopLevelAssetPath GetBaseClass() const { return BaseClass; }
+	inline FTopLevelAssetPath GetNativeClass() const { return NativeClass; }
 	inline UClass* GetActorNativeClass() const { return ActorNativeClass; }
 	inline FVector GetOrigin() const { return GetBounds().GetCenter(); }
 	inline FName GetRuntimeGrid() const { return RuntimeGrid; }
@@ -95,7 +101,7 @@ public:
 	inline const TArray<FName>& GetTags() const { return Tags; }
 	inline void SetDataLayerInstanceNames(const TArray<FName>& InDataLayerInstanceNames) { DataLayerInstanceNames = InDataLayerInstanceNames; }
 	inline FName GetActorPackage() const { return ActorPackage; }
-	inline FName GetActorPath() const { return ActorPath; }
+	inline FSoftObjectPath GetActorSoftPath() const { return ActorPath; }
 	inline FName GetActorLabel() const { return ActorLabel; }
 	inline FName GetFolderPath() const { return FolderPath; }
 	inline const FGuid& GetFolderGuid() const { return FolderGuid; }
@@ -233,10 +239,10 @@ protected:
 
 	// Persistent
 	FGuid							Guid;
-	FName							BaseClass;
-	FName							NativeClass;
+	FTopLevelAssetPath				BaseClass;
+	FTopLevelAssetPath				NativeClass;
 	FName							ActorPackage;
-	FName							ActorPath;
+	FSoftObjectPath					ActorPath;
 	FName							ActorLabel;
 	FVector							BoundsLocation;
 	FVector							BoundsExtent;
diff --git a/Engine/Source/Runtime/Engine/Public/WorldPartition/WorldPartitionActorDescView.h b/Engine/Source/Runtime/Engine/Public/WorldPartition/WorldPartitionActorDescView.h
--- a/Engine/Source/Runtime/Engine/Public/WorldPartition/WorldPartitionActorDescView.h
+++ b/Engine/Source/Runtime/Engine/Public/WorldPartition/WorldPartitionActorDescView.h
@@ -2,6 +2,8 @@
 #pragma once
 
 #include "CoreMinimal.h"
+#include "UObject/SoftObjectPath.h"
+#include "UObject/TopLevelAssetPath.h"
 #include "Misc/Optional.h"
 
 #if WITH_EDITOR
@@ -25,8 +27,8 @@ public:
 	FWorldPartitionActorDescView(const FWorldPartitionActorDesc* InActorDesc);
 
 	const FGuid& GetGuid() const;
-	FName GetBaseClass() const;
-	FName GetNativeClass() const;
+	FTopLevelAssetPath GetBaseClass() const;
+	FTopLevelAssetPath GetNativeClass() const;
 	UClass* GetActorNativeClass() const;
 	FVector GetOrigin() const;
 	FName GetRuntimeGrid() const;
@@ -39,7 +41,12 @@ public:
 	const TArray<FName>& GetRuntimeDataLayers() const;
 	const TArray<FName>& GetTags() const;
 	FName GetActorPackage() const;
-	FName GetActorPath() const;
+	
+PRAGMA_DISABLE_DEPRECATION_WARNINGS
+	UE_DEPRECATED(5.1, "GetActorPath is deprecated, GetActorSoftPath should be used instead.")
+	inline FName GetActorPath() const { return GetActorSoftPath().ToFName(); }
+PRAGMA_ENABLE_DEPRECATION_WARNINGS
+	FSoftObjectPath GetActorSoftPath() const;
 	FName GetActorLabel() const;
 	FBox GetBounds() const;
 	const TArray<FGuid>& GetReferences() const;
