FStreamableManager::CreateCombinedHandle: Fixed an ensure triggering for null handles even when SkipNull is passed in.
FStreamableManager::RequestAsyncLoad: Fixed a log level inconsistency.

diff --git a/Engine/Source/Runtime/Engine/Private/StreamableManager.cpp b/Engine/Source/Runtime/Engine/Private/StreamableManager.cpp
--- a/Engine/Source/Runtime/Engine/Private/StreamableManager.cpp
+++ b/Engine/Source/Runtime/Engine/Private/StreamableManager.cpp
@@ -1331,7 +1331,7 @@ TSharedPtr<FStreamableHandle> FStreamableManager::RequestAsyncLoad(TArray<FSoftO
 		}
 
 		// Some valid, some null
-		UE_LOG(LogStreamableManager, Warning, TEXT("RequestAsyncLoad(%s) called with both valid and null assets, null assets removed from %s!"), *DebugName, *RequestedSet);
+		UE_LOG(LogStreamableManager, Display, TEXT("RequestAsyncLoad(%s) called with both valid and null assets, null assets removed from %s!"), *DebugName, *RequestedSet);
 	}
 
 	if (TargetSet.Num() != NewRequest->RequestedAssets.Num())
@@ -1710,7 +1710,7 @@ TSharedPtr<FStreamableHandle> FStreamableManager::CreateCombinedHandle(const TCo
 	bool bFormattingFirstPass = true;
 	for (TSharedPtr<FStreamableHandle> ChildHandle : ChildHandles)
 	{
-		if (!ensure(ChildHandle.IsValid()))
+		if (!ChildHandle.IsValid())
 		{
 			if (EnumHasAllFlags(Options, EStreamableManagerCombinedHandleOptions::SkipNulls))
 			{
@@ -1718,6 +1718,7 @@ TSharedPtr<FStreamableHandle> FStreamableManager::CreateCombinedHandle(const TCo
 			}
 			else
 			{
+				ensureMsgf(ChildHandle.IsValid(), TEXT("A child handle is not valid and SkipNulls flag was not used, returning nullptr"));
 				return nullptr;
 			}
 		}
