Dump the total time spent compiling each shader job in seconds.

- If we compiled all shaders serially on a single thread this is roughly how long it would take.

diff --git a/Engine/Source/Runtime/Engine/Private/ShaderCompiler/ShaderCompiler.cpp b/Engine/Source/Runtime/Engine/Private/ShaderCompiler/ShaderCompiler.cpp
--- a/Engine/Source/Runtime/Engine/Private/ShaderCompiler/ShaderCompiler.cpp
+++ b/Engine/Source/Runtime/Engine/Private/ShaderCompiler/ShaderCompiler.cpp
@@ -3093,6 +3093,9 @@ void FShaderCompilerStats::WriteStatSummary()
 		{
 			TotalTimeForAllShaders += Iter.Value().TotalCompileTime;
 		}
+
+		UE_LOG(LogShaderCompilers, Display, TEXT("Total thread time: %.2f s"), TotalTimeForAllShaders);
+
 		if (TotalTimeAtLeastOneJobWasInFlight > 0.0)
 		{
 			double EffectiveParallelization = TotalTimeForAllShaders / TotalTimeAtLeastOneJobWasInFlight;
