Fix UActorDescContainer::IsActorDescHandled to incorrectly return true for paths sharing the same name prefix by appending a trailing slash.

#rnx

diff --git a/Engine/Source/Runtime/Engine/Private/WorldPartition/ActorDescContainer.cpp b/Engine/Source/Runtime/Engine/Private/WorldPartition/ActorDescContainer.cpp
--- a/Engine/Source/Runtime/Engine/Private/WorldPartition/ActorDescContainer.cpp
+++ b/Engine/Source/Runtime/Engine/Private/WorldPartition/ActorDescContainer.cpp
@@ -112,8 +112,9 @@ FString UActorDescContainer::GetExternalActorPath() const
 
 bool UActorDescContainer::IsActorDescHandled(const AActor* Actor) const
 {
-	FString ActorPackageName = Actor->GetPackage()->GetName();
-	return ActorPackageName.StartsWith(GetExternalActorPath());
+	const FString ActorPackageName = Actor->GetPackage()->GetName();
+	const FString ExternalActorPath = GetExternalActorPath() / TEXT("");
+	return ActorPackageName.StartsWith(ExternalActorPath);
 }
 
 bool UActorDescContainer::IsMainPartitionContainer() const
