Scene Capture:  Remove arbitrary Scene Capture 10x memory estimate provided in large render target warning in Texture Editor.  We now have a computed Scene Capture memory estimate displayed in the Texture Editor which is far more accurate, and the number provided in the warning is meaningless and confusing (got a UDN asking about the number).

The actual number in my tests was only 1.17x, but the number can go up as you add rendering features, and the overhead multiplier also varies with the bit depth of the render target.  For example, turning on TAA, which defaults to off, does create a 10x overhead, and Lumen on top of that can bump that to 20x or more.  And, if you had an 8-bit output, the output render target gets smaller, but the persistent data stays the same size (the scene renderer itself doesn't switch to monochrome for any of its persistent data -- it just blits to the lower bit depth at the end), so the overhead relative to the target's memory goes up (1.68x for the simple case with no expensive features).

#rnx

diff --git a/Engine/Source/Runtime/Engine/Private/TextureRenderTarget2D.cpp b/Engine/Source/Runtime/Engine/Private/TextureRenderTarget2D.cpp
--- a/Engine/Source/Runtime/Engine/Private/TextureRenderTarget2D.cpp
+++ b/Engine/Source/Runtime/Engine/Private/TextureRenderTarget2D.cpp
@@ -182,7 +182,7 @@ void UTextureRenderTarget2D::PostEditChangeProperty(FPropertyChangedEvent& Prope
 		const float MemoryMb = SizeX * SizeY * GPixelFormats[Format].BlockBytes / 1024.0f / 1024.0f;
 		FNumberFormattingOptions FloatFormat;
 		FloatFormat.SetMaximumFractionalDigits(1);
-		FText Message = FText::Format( NSLOCTEXT("TextureRenderTarget2D", "LargeTextureRenderTarget2DWarning", "A TextureRenderTarget2D of size {0}x{1} will use {2}Mb ({3}Mb if used with a Scene Capture), which may result in extremely poor performance or an Out Of Video Memory crash.\nAre you sure?"), FText::AsNumber(SizeX), FText::AsNumber(SizeY), FText::AsNumber(MemoryMb, &FloatFormat), FText::AsNumber(10.0f * MemoryMb, &FloatFormat));
+		FText Message = FText::Format( NSLOCTEXT("TextureRenderTarget2D", "LargeTextureRenderTarget2DWarning", "A TextureRenderTarget2D of size {0}x{1} will use {2}Mb, which may result in extremely poor performance or an Out Of Video Memory crash.\nAre you sure?"), FText::AsNumber(SizeX), FText::AsNumber(SizeY), FText::AsNumber(MemoryMb, &FloatFormat));
 		const EAppReturnType::Type Choice = FMessageDialog::Open(EAppMsgType::YesNo, Message);
 	
 		if (Choice == EAppReturnType::No)
