AnimationAttribute: Added support for user defined struct as animation attributes

diff --git a/Engine/Source/Editor/AnimationBlueprintLibrary/Private/K2_AnimationAttributeNodes.cpp b/Engine/Source/Editor/AnimationBlueprintLibrary/Private/K2_AnimationAttributeNodes.cpp
--- a/Engine/Source/Editor/AnimationBlueprintLibrary/Private/K2_AnimationAttributeNodes.cpp
+++ b/Engine/Source/Editor/AnimationBlueprintLibrary/Private/K2_AnimationAttributeNodes.cpp
@@ -18,6 +18,7 @@
 #include "UObject/Object.h"
 #include "UObject/WeakObjectPtr.h"
 #include "UObject/WeakObjectPtrTemplates.h"
+#include "Engine/UserDefinedStruct.h"
 
 #define LOCTEXT_NAMESPACE "K2Node_AnimationAttributeNodes"
 
@@ -75,6 +76,11 @@ bool UK2Node_BaseAttributeActionNode::IsConnectionDisallowed(const UEdGraphPin*
 				{
 					OutReason = FText::Format(LOCTEXT("UnregisterdStructureTypeFormat", "Must be a registered Animation Attribute, which {0} is not."), FText::FromName(AttributeType->GetFName())).ToString();
 				}
+				else if (AttributeType->IsA<UUserDefinedStruct>())
+				{
+					bAllowed = false;
+					OutReason = FText::Format(LOCTEXT("UnsupportedStructureTypeFormat", "User Defined Structs like {0} are not supported by this function"), FText::FromName(AttributeType->GetFName())).ToString();
+				}
 			}
 		}
 		else
diff --git a/Engine/Source/Editor/AnimationSettings/AnimationSettings.Build.cs b/Engine/Source/Editor/AnimationSettings/AnimationSettings.Build.cs
new file mode 100644
--- /dev/null
+++ b/Engine/Source/Editor/AnimationSettings/AnimationSettings.Build.cs
@@ -0,0 +1,19 @@
+// Copyright Epic Games, Inc. All Rights Reserved.
+
+using UnrealBuildTool;
+
+public class AnimationSettings: ModuleRules
+{
+	public AnimationSettings(ReadOnlyTargetRules Target) : base(Target)
+	{
+		PrivateDependencyModuleNames.AddRange(
+			new string [] {
+				"Core",
+				"CoreUObject",
+				"AssetRegistry",
+				"Engine", 
+				"SharedSettingsWidgets"
+			}
+		);
+	}
+}
diff --git a/Engine/Source/Editor/AnimationSettings/Private/AnimationSettingsModule.cpp b/Engine/Source/Editor/AnimationSettings/Private/AnimationSettingsModule.cpp
new file mode 100644
--- /dev/null
+++ b/Engine/Source/Editor/AnimationSettings/Private/AnimationSettingsModule.cpp
@@ -0,0 +1,95 @@
+// Copyright Epic Games, Inc. All Rights Reserved.
+
+#include "AnimationSettingsModule.h"
+
+#include "SSettingsEditorCheckoutNotice.h"
+#include "Animation/AnimationSettings.h"
+#include "AssetRegistry/AssetRegistryModule.h"
+#include "Engine/UserDefinedStruct.h"
+
+
+IMPLEMENT_MODULE(FAnimationSettingsModule, AnimationSettings);
+
+DEFINE_LOG_CATEGORY(LogAnimationSettings);
+
+#define LOCTEXT_NAMESPACE "AnimationSettings"
+
+void FAnimationSettingsModule::OnAssetRenamed(const FAssetData& InAssetData, const FString& InOldName)
+{
+	if (const UUserDefinedStruct* UserDefinedStruct = Cast<UUserDefinedStruct>(InAssetData.GetAsset()))
+	{
+		if (UAnimationSettings* Settings = UAnimationSettings::Get())
+		{
+			TArray<int32> Indices;
+			for (int32 Index = 0; Index < Settings->UserDefinedStructAttributes.Num(); Index++)
+			{
+				const TSoftObjectPtr<UUserDefinedStruct>& StructPtr = Settings->UserDefinedStructAttributes[Index];
+				UUserDefinedStruct* Struct = StructPtr.Get();
+				if (Struct == UserDefinedStruct)
+				{
+					Indices.Add(Index);
+				}
+			}
+
+			if (Indices.Num() > 0)
+			{
+				for (int32 Index : Indices)
+				{
+					Settings->UserDefinedStructAttributes[Index] = UserDefinedStruct;
+				}
+
+
+				FString RelativePath;
+				bool bIsSourceControlled = false;
+				bool bIsNewFile = false;
+
+				check(Settings->GetClass()->HasAnyClassFlags(CLASS_DefaultConfig));
+				
+				// Attempt to checkout the file automatically
+				RelativePath = Settings->GetDefaultConfigFilename();
+				bIsSourceControlled = true;
+				FString FullPath = FPaths::ConvertRelativePathToFull(RelativePath);
+
+				if (!FPlatformFileManager::Get().GetPlatformFile().FileExists(*FullPath))
+				{
+					bIsNewFile = true;
+				}
+
+				if (!bIsSourceControlled || !SettingsHelpers::CheckOutOrAddFile(FullPath))
+				{
+					SettingsHelpers::MakeWritable(FullPath);
+				}
+
+				Settings->TryUpdateDefaultConfigFile();
+
+				if (bIsNewFile && bIsSourceControlled)
+				{
+					SettingsHelpers::CheckOutOrAddFile(FullPath);
+				}
+			}	
+		}
+	}	
+}
+
+void FAnimationSettingsModule::StartupModule()
+{
+	const FAssetRegistryModule& AssetRegistryModule = FModuleManager::LoadModuleChecked<FAssetRegistryModule>(TEXT("AssetRegistry"));
+    
+   OnAssetRenamedHandle = AssetRegistryModule.Get().OnAssetRenamed().AddStatic(&FAnimationSettingsModule::OnAssetRenamed);	
+}
+
+void FAnimationSettingsModule::ShutdownModule()
+{
+	if (const FAssetRegistryModule* AssetRegistryModule = FModuleManager::GetModulePtr<FAssetRegistryModule>(TEXT("AssetRegistry")))
+	{
+		if (IAssetRegistry* AssetRegistry = AssetRegistryModule->TryGet())
+		{
+			AssetRegistry->OnAssetRenamed().Remove(OnAssetRenamedHandle);
+		}
+	}
+}
+
+
+
+#undef LOCTEXT_NAMESPACE
+
diff --git a/Engine/Source/Editor/AnimationSettings/Public/AnimationSettingsModule.h b/Engine/Source/Editor/AnimationSettings/Public/AnimationSettingsModule.h
new file mode 100644
--- /dev/null
+++ b/Engine/Source/Editor/AnimationSettings/Public/AnimationSettingsModule.h
@@ -0,0 +1,23 @@
+// Copyright Epic Games, Inc. All Rights Reserved.
+#pragma once
+
+#include "CoreMinimal.h"
+#include "Modules/ModuleInterface.h"
+#include "AssetRegistry/AssetData.h"
+
+DECLARE_LOG_CATEGORY_EXTERN(LogAnimationSettings, Log, All);
+
+class FAnimationSettingsModule : public IModuleInterface
+{
+public:
+
+	// IModuleInterface interface
+	virtual void StartupModule() override;
+	virtual void ShutdownModule() override;
+
+private:
+	static void OnAssetRenamed(const FAssetData& InAssetData, const FString& InOldName);
+
+	FDelegateHandle OnAssetRenamedHandle;
+};
+
diff --git a/Engine/Source/Editor/UnrealEd/Private/EditorEngine.cpp b/Engine/Source/Editor/UnrealEd/Private/EditorEngine.cpp
--- a/Engine/Source/Editor/UnrealEd/Private/EditorEngine.cpp
+++ b/Engine/Source/Editor/UnrealEd/Private/EditorEngine.cpp
@@ -1075,6 +1075,7 @@ void UEditorEngine::Init(IEngineLoop* InEngineLoop)
 			TEXT("CSVtoSVG"),
 			TEXT("GeometryFramework"),
 			TEXT("VirtualizationEditor"),
+			TEXT("AnimationSettings")
 		};
 
 		FScopedSlowTask ModuleSlowTask(UE_ARRAY_COUNT(ModuleNames));
diff --git a/Engine/Source/Editor/UnrealEd/UnrealEd.Build.cs b/Engine/Source/Editor/UnrealEd/UnrealEd.Build.cs
--- a/Engine/Source/Editor/UnrealEd/UnrealEd.Build.cs
+++ b/Engine/Source/Editor/UnrealEd/UnrealEd.Build.cs
@@ -265,7 +265,8 @@ public class UnrealEd : ModuleRules
 				"VirtualTexturingEditor",
 				"WorldPartitionEditor",
 				"CSVtoSVG",
-				"SourceControlWindowExtender"
+				"SourceControlWindowExtender",
+				"AnimationSettings"
 			}
 		);
 
diff --git a/Engine/Source/Runtime/Engine/Classes/Animation/AnimationSettings.h b/Engine/Source/Runtime/Engine/Classes/Animation/AnimationSettings.h
--- a/Engine/Source/Runtime/Engine/Classes/Animation/AnimationSettings.h
+++ b/Engine/Source/Runtime/Engine/Classes/Animation/AnimationSettings.h
@@ -13,6 +13,7 @@
 #include "Engine/DeveloperSettings.h"
 #include "CustomAttributes.h"
 #include "MirrorDataTable.h"
+#include "Engine/UserDefinedStruct.h"
 #include "AnimationSettings.generated.h"
 
 /**
@@ -91,6 +92,10 @@ class ENGINE_API UAnimationSettings : public UDeveloperSettings
 	UPROPERTY(config, EditAnywhere, Category = AnimationAttributes)
 	TArray<FString> TransformAttributeNames;
 
+	/** Register user defined structs as animation attributes*/
+	UPROPERTY(config, EditAnywhere, DisplayName="User Defined Struct Animation Attributes (Runtime only, Non-blendable)", Category = CustomAttributes, meta=(AllowedClasses="/Script/Engine.UserDefinedStruct"))
+	TArray<TSoftObjectPtr<UUserDefinedStruct>> UserDefinedStructAttributes;
+
 	/** Find and Replace Expressions used for mirroring  */
 	UPROPERTY(config, EditAnywhere, Category = Mirroring)
 	TArray<FMirrorFindReplaceExpression> MirrorFindReplaceExpressions;
diff --git a/Engine/Source/Runtime/Engine/Classes/Animation/AttributeBlendData.h b/Engine/Source/Runtime/Engine/Classes/Animation/AttributeBlendData.h
--- a/Engine/Source/Runtime/Engine/Classes/Animation/AttributeBlendData.h
+++ b/Engine/Source/Runtime/Engine/Classes/Animation/AttributeBlendData.h
@@ -278,6 +278,119 @@ namespace UE
 				int32 CurrentIndex;
 			};
 
+			struct TAttributeSetRawIterator
+			{
+				friend struct FAttributeBlendData;
+			protected:
+				TAttributeSetRawIterator(const FAttributeBlendData& InData, const FAttributeSet& InCollection) : Data(InData), Collection(InCollection), CurrentIndex(-1) {}
+			public:
+				/** Return the value for the currently indexed entry in the attribute set */
+				const uint8* GetValuePtr() const
+				{
+					check(CurrentIndex < Collection.DataPtrs.Num() && Collection.DataPtrs.IsValidIndex(CurrentIndex));
+					return Collection.DataPtrs[CurrentIndex];
+				}
+
+				/** Returns (container level) weight value for the current attribute's container in the attribute set */
+				const float GetWeight() const
+				{
+					check(Data.HasContainerWeights());
+					check(CurrentIndex < Collection.DataPtrs.Num() && Collection.WeightIndices.IsValidIndex(CurrentIndex));
+					return Data.GetContainerWeight(Collection.WeightIndices[CurrentIndex]);
+				}
+
+				/** Returns (bone level) weight value for the current attribute its bone and container */
+				const float GetBoneWeight() const
+				{
+					check(Data.HasBoneWeights());
+					check(CurrentIndex < Collection.DataPtrs.Num() && Collection.WeightIndices.IsValidIndex(CurrentIndex));
+					return Data.GetBoneWeight(Collection.WeightIndices[CurrentIndex], Collection.Identifier->GetIndex());
+				}
+
+				/** Returns highest (container level) weighted value for the attribute set */
+				const uint8* GetHighestWeightedValue() const
+				{
+					check(Data.HasContainerWeights());
+					check(Collection.HighestWeightedIndex < Collection.DataPtrs.Num() && Collection.DataPtrs.IsValidIndex(Collection.HighestWeightedIndex));
+					return Collection.DataPtrs[Collection.HighestWeightedIndex];
+				}
+
+				/** Returns highest (bone level) weighted value for the attribute set */
+				const uint8* GetHighestBoneWeightedValue() const
+				{
+					check(Data.HasBoneWeights());
+
+					int32 HighestIndex = INDEX_NONE;
+					float Weight = -1.f;
+					for (const int32 Index : Collection.WeightIndices)
+					{
+						const float BoneWeight = Data.GetBoneWeight(Index, Collection.Identifier->GetIndex());
+						if (BoneWeight > Weight)
+						{
+							Weight = BoneWeight;
+							HighestIndex = Index;
+						}
+					}
+					ensure(HighestIndex != INDEX_NONE);
+					return Collection.DataPtrs[Collection.WeightIndices.IndexOfByKey(HighestIndex)];
+				}
+
+				/** Returns highest (bone level) weighted value, and its weight for the attribute set */
+				void GetHighestBoneWeighted(const uint8* OutAttributePtr, float& OutWeight) const
+				{
+					check(Data.HasBoneWeights());
+
+					int32 HighestIndex = INDEX_NONE;
+					float Weight = -1.f;
+					for (const int32 Index : Collection.WeightIndices)
+					{
+						const float BoneWeight = Data.GetBoneWeight(Index, Collection.Identifier->GetIndex());
+						if (BoneWeight > Weight)
+						{
+							Weight = BoneWeight;
+							HighestIndex = Index;
+						}
+					}
+					ensure(HighestIndex != INDEX_NONE);
+					OutAttributePtr = Collection.DataPtrs[Collection.WeightIndices.IndexOfByKey(HighestIndex)];
+					OutWeight = Weight;
+				}
+
+				/** Returns the identifier for the current attribute set */
+				const FAttributeId& GetIdentifier() const
+				{
+					return *Collection.Identifier;
+				}
+
+				EAdditiveAnimationType GetAdditiveType() const
+				{
+					return Data.AdditiveType;
+				}
+
+				bool IsFilteredBlend() const
+				{
+					return Data.bPerBoneFilter;
+				}
+
+				/** Cycle through to next entry in the attribute set, returns false if the end was reached */
+				bool Next()
+				{
+					++CurrentIndex;
+					return CurrentIndex < Collection.DataPtrs.Num();
+				}
+
+				int32 GetIndex() const
+				{
+					return CurrentIndex;
+				}
+			protected:
+				/** Outer object that creates this */
+				const FAttributeBlendData& Data;
+				/** Attribute collection for current index */
+				const FAttributeSet& Collection;
+				int32 CurrentIndex;
+			};
+
 			template<typename AttributeType>
 			struct TSingleIterator
 			{
@@ -351,6 +464,78 @@ namespace UE
 				int32 CurrentIndex;
 			};
 
+			struct TSingleRawIterator
+			{
+				friend struct FAttributeBlendData;
+			protected:
+				TSingleRawIterator(const FAttributeBlendData& InData, const TArray<FUniqueAttribute>& InAttributes) : Data(InData), AttributesView(InAttributes), CurrentIndex(-1) {}
+			public:
+				/** Cycle through to next unique attribute, returns false if the end was reached */
+				bool Next()
+				{
+					++CurrentIndex;
+					return CurrentIndex < AttributesView.Num();
+				}
+
+				/** Return the value for the currently indexed unique attribute */
+				const uint8* GetValuePtr() const
+				{
+					check(CurrentIndex < AttributesView.Num() && AttributesView.IsValidIndex(CurrentIndex));
+					return AttributesView[CurrentIndex].DataPtr;
+				}
+
+				/** Returns (container level) weight value for the unique attribute its container */
+				const float GetWeight() const
+				{
+					check(Data.HasContainerWeights());
+					check(CurrentIndex < AttributesView.Num() && AttributesView.IsValidIndex(CurrentIndex));
+					
+					return Data.GetContainerWeight(AttributesView[CurrentIndex].WeightIndex);
+				}
+				
+				/** Returns (bone level) weight value for the unique attribute its bone and container */
+				const float GetBoneWeight() const
+				{
+					check(Data.HasBoneWeights());
+					check(CurrentIndex < AttributesView.Num() && AttributesView.IsValidIndex(CurrentIndex));
+					
+					return Data.GetBoneWeight(AttributesView[CurrentIndex].WeightIndex, AttributesView[CurrentIndex].Identifier->GetIndex());
+				}
+
+				/** Returns whether or not the unique attribute its (bone level) weight is the highest across the containers */
+				bool IsHighestBoneWeighted() const
+				{
+					check(Data.HasBoneWeights());
+					check(CurrentIndex < AttributesView.Num() && AttributesView.IsValidIndex(CurrentIndex));
+					check(Data.HighestBoneWeightedIndices.IsValidIndex(AttributesView[CurrentIndex].Identifier->GetIndex()));
+
+					return Data.HighestBoneWeightedIndices[AttributesView[CurrentIndex].Identifier->GetIndex()] == AttributesView[CurrentIndex].WeightIndex;
+				}
+
+				/** Returns the identifier for the current attribute set */
+				const FAttributeId& GetIdentifier() const
+				{
+					check(CurrentIndex < AttributesView.Num() && AttributesView.IsValidIndex(CurrentIndex));					
+					return *AttributesView[CurrentIndex].Identifier;
+				}
+
+				EAdditiveAnimationType GetAdditiveType() const
+				{
+					return Data.AdditiveType;
+				}
+				
+				bool IsFilteredBlend() const
+				{
+					return Data.bPerBoneFilter;
+				}
+
+			protected:
+				/** Outer object that creates this */
+				const FAttributeBlendData& Data;
+				TArrayView<const FUniqueAttribute> AttributesView;
+				int32 CurrentIndex;
+			};
+
 		public:
 			template<typename AttributeType>
 			void ForEachAttributeSet(TFunctionRef<void(TAttributeSetIterator<AttributeType>&)> ForEachFunction) const
@@ -368,6 +553,21 @@ namespace UE
 				TSingleIterator<AttributeType> It(*this, UniqueAttributes);
 				ForEachFunction(It);
 			}
+			
+			void ForEachAttributeSet(TFunctionRef<void(TAttributeSetRawIterator&)> ForEachFunction) const
+			{
+				for (const FAttributeSet& Collection : AttributeSets)
+				{
+					TAttributeSetRawIterator It(*this, Collection);
+					ForEachFunction(It);
+				}
+			}
+
+			void ForEachUniqueAttribute(TFunctionRef<void(TSingleRawIterator&)> ForEachFunction) const
+			{
+				TSingleRawIterator It(*this, UniqueAttributes);
+				ForEachFunction(It);
+			}
 		};
 	}
 }
\ No newline at end of file
diff --git a/Engine/Source/Runtime/Engine/Classes/Animation/AttributeBlendOperator.h b/Engine/Source/Runtime/Engine/Classes/Animation/AttributeBlendOperator.h
--- a/Engine/Source/Runtime/Engine/Classes/Animation/AttributeBlendOperator.h
+++ b/Engine/Source/Runtime/Engine/Classes/Animation/AttributeBlendOperator.h
@@ -423,5 +423,134 @@ namespace UE
 				AttributeType::StaticStruct()->CopyScriptStruct(InOutData, CopySource);
 			}
 		};
+
+		class FNonBlendableAttributeBlendOperator : public IAttributeBlendOperator
+		{
+		public:
+			FNonBlendableAttributeBlendOperator(const UScriptStruct* InScriptStruct) : ScriptStructPtr(InScriptStruct) {}; 
+			
+			/** Begin IAttributeBlendOperator overrides */
+			virtual void Accumulate(const FAttributeBlendData& BlendData, FStackAttributeContainer* OutAttributes) const final
+			{
+				BlendData.ForEachUniqueAttribute([&, OutAttributes](typename FAttributeBlendData::template TSingleRawIterator& It) -> void
+				{
+					while (It.Next())
+					{						
+						const uint8* Attribute= It.GetValuePtr();
+						const FAttributeId& Identifier = It.GetIdentifier();
+						const int32 ExistingIndex = OutAttributes->IndexOfByKey(ScriptStructPtr.Get(), Identifier);
+						uint8* OutAttributePtr = OutAttributes->FindOrAdd(ScriptStructPtr.Get(), Identifier);
+
+						// 'Accumulate' value if the attribute did not yet exist, or based upon being the highest weighted attribute
+						const float AttributeWeight = It.GetWeight();
+						if (ExistingIndex == INDEX_NONE || AttributeWeight > 0.5f)
+						{
+							ScriptStructPtr.Get()->CopyScriptStruct(OutAttributePtr, Attribute);
+						}
+					}
+				});
+			};
+			virtual void Interpolate(const void* FromData, const void* ToData, float Alpha, void* InOutData) const final
+			{
+				// Determine stepped value to copy
+				const void* CopySource = Alpha > 0.5f ? ToData : FromData;
+				// Using CopyScriptStruct as assignment operator might not be implemented for attribute type
+				ScriptStructPtr.Get()->CopyScriptStruct(InOutData, CopySource);	
+			};
+			virtual void Override(const FAttributeBlendData& BlendData, FStackAttributeContainer* OutAttributes) const final
+			{
+				BlendData.ForEachUniqueAttribute([&, OutAttributes](typename FAttributeBlendData::template TSingleRawIterator& It) -> void
+				{
+					while (It.Next())
+					{
+						const uint8* Attribute = It.GetValuePtr();
+						const FAttributeId& Identifier = It.GetIdentifier();
+
+						// Find or add as the attribute might already exist, so Add would fail
+						uint8* OutAttribute = OutAttributes->FindOrAdd(ScriptStructPtr.Get(),Identifier);
+						ScriptStructPtr.Get()->CopyScriptStruct(OutAttribute, Attribute);
+					}
+				});
+			};
+			virtual void Blend(const FAttributeBlendData& BlendData, FStackAttributeContainer* OutAttributes) const final
+			{
+				BlendData.ForEachAttributeSet([&, OutAttributes](typename FAttributeBlendData::template TAttributeSetRawIterator& It) -> void
+				{
+					const FAttributeId& Identifier = It.GetIdentifier();
+					const ECustomAttributeBlendType BlendType = Attributes::GetAttributeBlendType(Identifier);
+					
+					// Not iterating over each attribute, but just picking highest weighted value
+					uint8* OutAttribute = OutAttributes->Add(ScriptStructPtr.Get(), Identifier);
+
+					ScriptStructPtr.Get()->CopyScriptStruct(OutAttribute, It.GetHighestWeightedValue());
+				});
+
+				BlendData.ForEachUniqueAttribute([&, OutAttributes](typename FAttributeBlendData::template TSingleRawIterator& It) -> void
+				{
+					while (It.Next())
+					{
+						const FAttributeId& Identifier = It.GetIdentifier();
+						
+						uint8* OutAttribute = OutAttributes->Add(ScriptStructPtr.Get(), Identifier);
+						ScriptStructPtr.Get()->CopyScriptStruct(OutAttribute, It.GetValuePtr());	
+					}
+				});
+			};
+			virtual void BlendPerBone(const FAttributeBlendData& BlendData, FStackAttributeContainer* OutAttributes) const final
+			{
+				BlendData.ForEachAttributeSet([&, OutAttributes](typename FAttributeBlendData::template TAttributeSetRawIterator& It) -> void
+				{
+					const FAttributeId& Identifier = It.GetIdentifier();
+					uint8* OutAttribute = OutAttributes->FindOrAdd(ScriptStructPtr.Get(), Identifier);
+					ScriptStructPtr.Get()->CopyScriptStruct(OutAttribute, It.GetHighestBoneWeightedValue());
+				});
+
+				BlendData.ForEachUniqueAttribute([&, OutAttributes](typename FAttributeBlendData::template TSingleRawIterator& It) -> void
+				{
+					while (It.Next())
+					{
+						const FAttributeId& Identifier = It.GetIdentifier();
+						const int32 ExistingIndex = OutAttributes->IndexOfByKey(ScriptStructPtr.Get() ,Identifier);
+						const float BoneWeight = It.GetBoneWeight();
+
+						if (FAnimWeight::IsRelevant(BoneWeight) || ExistingIndex == INDEX_NONE)
+						{
+							uint8* OutAttributePtr = OutAttributes->FindOrAdd(ScriptStructPtr.Get(), Identifier);
+							const uint8* Attribute = It.GetValuePtr();
+
+							// 'Blend' value if the attribute did not yet exist, or based upon being the highest weighted attribute
+							const bool bHighestWeight = It.IsHighestBoneWeighted();
+							if (bHighestWeight || ExistingIndex == INDEX_NONE)
+							{
+								ScriptStructPtr.Get()->CopyScriptStruct(OutAttributePtr, Attribute);
+							}
+						}
+					}
+				});
+			};
+			virtual void ConvertToAdditive(const FAttributeBlendData& BlendData, FStackAttributeContainer* OutAdditiveAttributes) const final
+			{
+				BlendData.ForEachUniqueAttribute([&, OutAdditiveAttributes](typename FAttributeBlendData::template TSingleRawIterator& It) -> void
+				{
+					while (It.Next())
+					{
+						const uint8* Attribute = It.GetValuePtr();
+						const FAttributeId& Identifier = It.GetIdentifier();
+						const int32 ExistingIndex = OutAdditiveAttributes->IndexOfByKey(ScriptStructPtr.Get(), Identifier);
+						uint8* OutAttributePtr = OutAdditiveAttributes->FindOrAdd(ScriptStructPtr.Get(),Identifier);
+
+						// 'Accumulate' value if the attribute did not yet exist, or based upon being the highest weighted attribute
+						const float AttributeWeight = It.GetWeight();
+						if (ExistingIndex == INDEX_NONE)
+						{							
+							ScriptStructPtr.Get()->CopyScriptStruct(OutAttributePtr, &Attribute);
+						}
+					}
+				});
+			};
+
+		protected:
+			TWeakObjectPtr<const UScriptStruct> ScriptStructPtr;
+		};
 	}
 }
\ No newline at end of file
diff --git a/Engine/Source/Runtime/Engine/Classes/Animation/AttributesContainer.h b/Engine/Source/Runtime/Engine/Classes/Animation/AttributesContainer.h
--- a/Engine/Source/Runtime/Engine/Classes/Animation/AttributesContainer.h
+++ b/Engine/Source/Runtime/Engine/Classes/Animation/AttributesContainer.h
@@ -285,7 +285,7 @@ namespace UE
 			*
 			* @return	The ptr to the added entry, to be used for populating the data, nullptr if adding it failed
 			*/
-			uint8* Add(UScriptStruct* InScriptStruct, const FAttributeId& InAttributeId)
+			uint8* Add(const UScriptStruct* InScriptStruct, const FAttributeId& InAttributeId)
 			{
 				const int32 TypeIndex = FindOrAddTypeIndex(InScriptStruct);
 
@@ -337,7 +337,7 @@ namespace UE
 			*
 			* @return	The ptr to the added/existing entry, to be used for populating the data, nullptr if adding it failed
 			*/
-			uint8* FindOrAdd(UScriptStruct* InScriptStruct, const FAttributeId& InAttributeId)
+			uint8* FindOrAdd(const UScriptStruct* InScriptStruct, const FAttributeId& InAttributeId)
 			{
 				const int32 TypeIndex = FindOrAddTypeIndex(InScriptStruct);
 				TArray<FAttributeId>& AttributeIds = AttributeIdentifiers[TypeIndex];
@@ -552,6 +552,33 @@ namespace UE
 				return INDEX_NONE;
 			}
 
+			/**
+			* Tries to find and return the indedx of a attribute type/value entry of the specified AttributeType.
+			*
+			* @param	InAttributeId		Key to be used for seraching the entry
+			*
+			* @return	Index to the entry, INDEX_NONE if not found
+			*/
+			int32 IndexOfByKey(const UScriptStruct* InScriptStruct, const FAttributeId& InAttributeId) const
+			{
+				const int32 TypeIndex = FindTypeIndex(InScriptStruct);
+				if (TypeIndex != INDEX_NONE)
+				{
+					const TArray<int32>& UniqueBoneIndices = UniqueTypedBoneIndices[TypeIndex];
+
+					// Early out if for this bone index no attributes are currently contained
+					if (UniqueBoneIndices.Contains(InAttributeId.GetIndex()))
+					{
+						const TArray<FAttributeId>& AttributeIds = AttributeIdentifiers[TypeIndex];
+						const int32 AttributeIndex = AttributeIds.IndexOfByKey(InAttributeId);
+						
+						return AttributeIndex;
+					}
+				}
+
+				return INDEX_NONE;
+			}
+			
 			/**
 			* Removes, if existing, an attribute type/value entry of the specified AttributeType.
 			*
@@ -784,10 +811,12 @@ namespace UE
 			TArray<FAttributeId>& GetKeysInternal(int32 TypeIndex) { return AttributeIdentifiers[TypeIndex]; }
 			
 			/** Find or add a new root-level entry for the provided attribute data type, returning the index into the arrays representing the type */
-			int32 FindOrAddTypeIndex(UScriptStruct* InScriptStruct)
+			int32 FindOrAddTypeIndex(const UScriptStruct* InScriptStruct)
 			{
 				const int32 OldNum = UniqueTypes.Num();
-				const int32 TypeIndex = UniqueTypes.AddUnique(InScriptStruct);
+				// Technically we could have avoided using const cast here if we made UniqueTypes use const UScriptStruct*
+				// but a bit of refactoring is required
+				const int32 TypeIndex = UniqueTypes.AddUnique(const_cast<UScriptStruct*>(InScriptStruct));
 
 				if (UniqueTypes.Num() > OldNum)
 				{
diff --git a/Engine/Source/Runtime/Engine/Private/Animation/AnimationSettings.cpp b/Engine/Source/Runtime/Engine/Private/Animation/AnimationSettings.cpp
--- a/Engine/Source/Runtime/Engine/Private/Animation/AnimationSettings.cpp
+++ b/Engine/Source/Runtime/Engine/Private/Animation/AnimationSettings.cpp
@@ -2,6 +2,8 @@
 
 #include "Animation/AnimationSettings.h"
 #include "Animation/AnimCompress_BitwiseCompressOnly.h"
+#include "Animation/AttributeTypes.h"
+
 
 UAnimationSettings::UAnimationSettings(const FObjectInitializer& ObjectInitializer)
 	: Super(ObjectInitializer)
@@ -64,15 +66,32 @@ void UAnimationSettings::PostEditChangeProperty(struct FPropertyChangedEvent& Pr
 {
 	Super::PostEditChangeProperty(PropertyChangedEvent);
 
-// 	const FName PropertyName = PropertyChangedEvent.Property ? PropertyChangedEvent.Property->GetFName() : NAME_None;
-// 	if (PropertyName == GET_MEMBER_NAME_CHECKED(UAnimationSettings, FrictionCombineMode))
-// 	{
-// 		UPhysicalMaterial::RebuildPhysicalMaterials();
-// 	}
-// 	else if (PropertyName == GET_MEMBER_NAME_CHECKED(UAnimationSettings, LockedAxis))
-// 	{
-// 		UMovementComponent::PhysicsLockedAxisSettingChanged();
-// 	}
+ 	const FName PropertyName = PropertyChangedEvent.Property ? PropertyChangedEvent.Property->GetFName() : NAME_None;
+	
+	if (PropertyName == GET_MEMBER_NAME_CHECKED(UAnimationSettings, UserDefinedStructAttributes))
+	{
+		TArray<UScriptStruct*> NewListOfStructs;
+		for (const TSoftObjectPtr<UUserDefinedStruct>& UserDefinedStruct : UAnimationSettings::Get()->UserDefinedStructAttributes)
+		{
+			if (UserDefinedStruct.IsValid())
+			{
+				NewListOfStructs.Add(UserDefinedStruct.Get());
+				UE::Anim::AttributeTypes::RegisterNonBlendableType(UserDefinedStruct.Get());
+			}
+		}
+
+		TArray<TWeakObjectPtr<const UScriptStruct>> RegisteredTypes = UE::Anim::AttributeTypes::GetRegisteredTypes();
+		for (const TWeakObjectPtr<const UScriptStruct>& RegisteredStruct : RegisteredTypes)
+		{
+			if (RegisteredStruct.IsValid() && RegisteredStruct->IsA<UUserDefinedStruct>())
+			{
+				if (!NewListOfStructs.Contains(RegisteredStruct.Get()))
+				{
+					UE::Anim::AttributeTypes::UnregisterType(RegisteredStruct.Get());
+				}
+			}
+		}
+	}
 }
 
 
diff --git a/Engine/Source/Runtime/Engine/Private/Animation/AttributeTypes.cpp b/Engine/Source/Runtime/Engine/Private/Animation/AttributeTypes.cpp
--- a/Engine/Source/Runtime/Engine/Private/Animation/AttributeTypes.cpp
+++ b/Engine/Source/Runtime/Engine/Private/Animation/AttributeTypes.cpp
@@ -3,6 +3,8 @@
 #include "Animation/AttributeTypes.h"
 #include "Animation/BuiltInAttributeTypes.h"
 #include "Misc/DelayedAutoRegister.h"
+#include "Animation/AnimationSettings.h"
+#include "Engine/UserDefinedStruct.h"
 
 namespace UE
 {
@@ -12,6 +14,7 @@ namespace UE
 		TArray<TUniquePtr<IAttributeBlendOperator>> AttributeTypes::Operators;
 		TArray<TWeakObjectPtr<const UScriptStruct>> AttributeTypes::InterpolatableTypes;
 		std::atomic<bool> AttributeTypes::bInitialized = false;		
+		AttributeTypes::FOnAttributeTypesChanged AttributeTypes::OnAttributeTypesChangedDelegate;
 		
 		void AttributeTypes::LazyInitialize()
 		{
@@ -30,12 +33,52 @@ namespace UE
 			RegisterType<FTransformAnimationAttribute>();
 			RegisterType<FVectorAnimationAttribute>();
 			RegisterType<FQuaternionAnimationAttribute>();
+
+			for (const TSoftObjectPtr<UUserDefinedStruct>& UserDefinedStruct : UAnimationSettings::Get()->UserDefinedStructAttributes)
+			{
+				UE::Anim::AttributeTypes::RegisterNonBlendableType(UserDefinedStruct.LoadSynchronous());
+			}
 		}
 		
 		static FDelayedAutoRegisterHelper DelayedAttributeTypesInitializationHelper(EDelayedRegisterRunPhase::ObjectSystemReady, []()
 		{
 			UE::Anim::AttributeTypes::LazyInitialize();
 		});
+
+		bool AttributeTypes::RegisterNonBlendableType(const UScriptStruct* InScriptStruct)
+		{
+			if (!InScriptStruct)
+			{
+				return false;
+			}
+				
+			if (IsTypeRegistered(InScriptStruct))
+			{
+				return false;
+			}
+				
+			AttributeTypes::RegisteredTypes.Add(InScriptStruct);
+			AttributeTypes::Operators.Add(MakeUnique<FNonBlendableAttributeBlendOperator>(InScriptStruct));
+
+			OnAttributeTypesChangedDelegate.Broadcast(InScriptStruct, true);
+				
+			return true;
+		}
+
+		void AttributeTypes::UnregisterType(const UScriptStruct* InScriptStruct)
+		{
+			if (UObjectInitialized())
+			{
+				const int32 Index = AttributeTypes::RegisteredTypes.IndexOfByKey(InScriptStruct);
+				if (Index != INDEX_NONE)
+				{
+					AttributeTypes::RegisteredTypes.RemoveAtSwap(Index);
+					AttributeTypes::Operators.RemoveAtSwap(Index);
+				}
+			}
+				
+			OnAttributeTypesChangedDelegate.Broadcast(InScriptStruct, false);
+		}
 	}
 }
 
diff --git a/Engine/Source/Runtime/Engine/Public/Animation/AttributeTypes.h b/Engine/Source/Runtime/Engine/Public/Animation/AttributeTypes.h
--- a/Engine/Source/Runtime/Engine/Public/Animation/AttributeTypes.h
+++ b/Engine/Source/Runtime/Engine/Public/Animation/AttributeTypes.h
@@ -11,6 +11,7 @@
 #include "UObject/WeakObjectPtr.h"
 
 #include "CoreTypes.h"
+#include "Animation/AnimationSettings.h"
 #include "Templates/UnrealTemplate.h"
 
 enum EAdditiveAnimationType;
@@ -42,15 +43,28 @@ namespace UE
 
 		struct ENGINE_API AttributeTypes
 		{
+			DECLARE_MULTICAST_DELEGATE_TwoParams(FOnAttributeTypesChanged, const UScriptStruct*, bool /* bIsAdded */ );
+			
 		protected:
 			static TArray<TWeakObjectPtr<const UScriptStruct>> RegisteredTypes;
 			static TArray<TUniquePtr<IAttributeBlendOperator>> Operators;
 			static TArray<TWeakObjectPtr<const UScriptStruct>> InterpolatableTypes;
 			static std::atomic<bool> bInitialized;
+			static FOnAttributeTypesChanged OnAttributeTypesChangedDelegate;
+
 			static void Initialize();
+
+			/** Register user defined structs as non-blendable animation attribute */
+			static bool RegisterNonBlendableType(const UScriptStruct* InScriptStruct);
+
+			/** Unregisters a specific attribute type and deletes its associated blend operator */
+			static void UnregisterType(const UScriptStruct* InScriptStruct);
+			
 		public:			
 			static void LazyInitialize();
 
+			static FOnAttributeTypesChanged& GetOnAttributeTypesChanged() { return OnAttributeTypesChangedDelegate; };
+			
 			/** Used for registering an attribute type for which TAttributeTypeTraits::WithCustomBlendOperator is set to true, use RegisterType() otherwise */
 			template<typename AttributeType, typename OperatorType, typename... OperatorArgs>
 			static void RegisterTypeWithOperator(OperatorArgs&&... args)
@@ -140,6 +154,8 @@ namespace UE
 				LazyInitialize();
 				return AttributeTypes::RegisteredTypes;
 			}
+
+			friend class UAnimationSettings;
 		};
 	}
 }
