When a skeleton is deleted, it now removes itself from the compatible skeletons array of skeletons that reference it.
* As users are able to add empty compatible skeletons in the property window and due to CompatibleSkeletons being a TSoftObjectPtr array, programmers still need to check the pointer is valid.

diff --git a/Engine/Source/Runtime/Engine/Classes/Animation/Skeleton.h b/Engine/Source/Runtime/Engine/Classes/Animation/Skeleton.h
--- a/Engine/Source/Runtime/Engine/Classes/Animation/Skeleton.h
+++ b/Engine/Source/Runtime/Engine/Classes/Animation/Skeleton.h
@@ -485,6 +485,7 @@ protected:
 	/**
 	 * The list of compatible skeletons.
 	 * This is an array of TSoftObjectPtr in order to prevent all skeletons to be loaded, as we only want to load things on demand.
+	 * As this is EditAnywhere and an array of TSoftObjectPtr, checking validity of pointers is needed.
 	 **/
 	UPROPERTY(EditAnywhere, BlueprintReadOnly, Category = CompatibleSkeletons)
 	TArray<TSoftObjectPtr<USkeleton>> CompatibleSkeletons;
diff --git a/Engine/Source/Runtime/Engine/Private/Animation/Skeleton.cpp b/Engine/Source/Runtime/Engine/Private/Animation/Skeleton.cpp
--- a/Engine/Source/Runtime/Engine/Private/Animation/Skeleton.cpp
+++ b/Engine/Source/Runtime/Engine/Private/Animation/Skeleton.cpp
@@ -483,18 +483,6 @@ USkeleton::USkeleton(const FObjectInitializer& ObjectInitializer)
 
 void USkeleton::BeginDestroy()
 {
-	Super::BeginDestroy();
-	if (HasAnyFlags(RF_ClassDefaultObject))
-	{
-		FCoreUObjectDelegates::OnPackageReloaded.RemoveAll(this);
-	}
-
-	// Clear our own skeleton mappings.
-	ClearSkeletonRemappings();
-	OnSkeletonDestructEvent.Broadcast(this); // Inform others about this skeleton being destructed, so we can remove mappings inside those skeletons.
-	OnSkeletonDestructEvent.Clear();
-	OnSmartNamesChangedEvent.Clear();
-
 	//  Remove the loaded skeleton from the list.
 	{
 		FScopeLock Lock(&LoadedSkeletonsMutex);
@@ -506,9 +494,22 @@ void USkeleton::BeginDestroy()
 			if (Skeleton->IsCompatible(this))
 			{
 				Skeleton->OnSkeletonDestructEvent.RemoveAll(this);
+				Skeleton->RemoveCompatibleSkeleton(this);
 			}
 		}
 	}
+
+	Super::BeginDestroy();
+	if (HasAnyFlags(RF_ClassDefaultObject))
+	{
+		FCoreUObjectDelegates::OnPackageReloaded.RemoveAll(this);
+	}
+
+	// Clear our own skeleton mappings.
+	ClearSkeletonRemappings();
+	OnSkeletonDestructEvent.Broadcast(this); // Inform others about this skeleton being destructed, so we can remove mappings inside those skeletons.
+	OnSkeletonDestructEvent.Clear();
+	OnSmartNamesChangedEvent.Clear();
 }
 
 void USkeleton::PostInitProperties()
@@ -596,6 +597,12 @@ void USkeleton::PostLoad()
 	{
 		OnSmartNamesChangedEvent.AddUObject(this, &USkeleton::HandleSmartNamesChangedEvent);
 	}
+	
+	// Cleanup CompatibleSkeletons for convenience, this does not ensure CompatibleSkeletons is free of invalid TSoftObjectPtrs at access-time.
+	CompatibleSkeletons = CompatibleSkeletons.FilterByPredicate([](const TSoftObjectPtr<USkeleton>& Skeleton)
+	{
+		return Skeleton.IsValid();
+	});
 }
 
 // Regenerate all required skeleton remappings whenever curves change.
