- World Partition PIE Runtime Levels now computes a LoadTime based on the load time of all its actor packages
- World Partition debug draw runtime cell now also displays load time per cell

diff --git a/Engine/Source/Runtime/Engine/Private/WorldPartition/WorldPartitionLevelStreamingDynamic.cpp b/Engine/Source/Runtime/Engine/Private/WorldPartition/WorldPartitionLevelStreamingDynamic.cpp
--- a/Engine/Source/Runtime/Engine/Private/WorldPartition/WorldPartitionLevelStreamingDynamic.cpp
+++ b/Engine/Source/Runtime/Engine/Private/WorldPartition/WorldPartitionLevelStreamingDynamic.cpp
@@ -460,6 +460,17 @@ void UWorldPartitionLevelStreamingDynamic::FinalizeRuntimeLevel()
 	// Mark this package as fully loaded with regards to external objects
 	RuntimeLevel->GetPackage()->SetDynamicPIEPackagePending(false);
 
+	// Update runtime level package load time : take into account actor packages load time
+	float LoadTime = RuntimeLevel->GetPackage()->GetLoadTime();
+	for (AActor* Actor : RuntimeLevel->Actors)
+	{
+		if (UPackage* ActorPackage = Actor ? Actor->GetExternalPackage() : nullptr)
+		{
+			LoadTime += ActorPackage->GetLoadTime();
+		}
+	}
+	RuntimeLevel->GetPackage()->SetLoadTime(LoadTime);
+
 	PackageReferencer.RemoveReferences();
 }
 
diff --git a/Engine/Source/Runtime/Engine/Private/WorldPartition/WorldPartitionLevelStreamingPolicy.cpp b/Engine/Source/Runtime/Engine/Private/WorldPartition/WorldPartitionLevelStreamingPolicy.cpp
--- a/Engine/Source/Runtime/Engine/Private/WorldPartition/WorldPartitionLevelStreamingPolicy.cpp
+++ b/Engine/Source/Runtime/Engine/Private/WorldPartition/WorldPartitionLevelStreamingPolicy.cpp
@@ -232,7 +232,18 @@ void UWorldPartitionLevelStreamingPolicy::DrawRuntimeCellsDetails(UCanvas* Canva
 			const FColor Color = ULevelStreaming::GetLevelStreamingStatusColor(StreamingStatus);
 			for (const UWorldPartitionRuntimeCell* Cell : Cells)
 			{
-				DrawCellDetails(Cell->GetDebugName(), GEngine->GetTinyFont(), Color);
+				if ((StreamingStatus == EStreamingStatus::LEVEL_Loaded) ||
+					(StreamingStatus == EStreamingStatus::LEVEL_MakingVisible) ||
+					(StreamingStatus == EStreamingStatus::LEVEL_Visible) ||
+					(StreamingStatus == EStreamingStatus::LEVEL_MakingInvisible))
+				{
+					float LoadTime = Cell->GetLevel() ? Cell->GetLevel()->GetPackage()->GetLoadTime() : 0.f;
+					DrawCellDetails(FString::Printf(TEXT("%s (%s)"), *Cell->GetDebugName(), *FPlatformTime::PrettyTime(LoadTime)), GEngine->GetTinyFont(), Color);
+				}
+				else
+				{
+					DrawCellDetails(Cell->GetDebugName(), GEngine->GetTinyFont(), Color);
+				}
 			}
 		}
 	}
