FNames containing asset paths are deprecated. FSoftObjectPath or FTopLevelAssetPath should be used instead.
Fixups for editor asset management.

diff --git a/Engine/Plugins/Editor/AssetManagerEditor/Source/AssetManagerEditor/Private/SAssetAuditBrowser.cpp b/Engine/Plugins/Editor/AssetManagerEditor/Source/AssetManagerEditor/Private/SAssetAuditBrowser.cpp
--- a/Engine/Plugins/Editor/AssetManagerEditor/Source/AssetManagerEditor/Private/SAssetAuditBrowser.cpp
+++ b/Engine/Plugins/Editor/AssetManagerEditor/Source/AssetManagerEditor/Private/SAssetAuditBrowser.cpp
@@ -128,22 +128,22 @@ void SAssetAuditBrowser::GetSelectedPackages(const TArray<FAssetData>& Assets, T
 void SAssetAuditBrowser::EditSelectedAssets() const
 {
 	TArray<FAssetData> Assets = GetCurrentSelectionDelegate.Execute();
-	TArray<FName> AssetNames;
+	TArray<FSoftObjectPath> AssetPaths;
 
 	for (FAssetData& AssetData : Assets)
 	{
-		AssetNames.Add(AssetData.ObjectPath);
+		AssetPaths.Add(AssetData.GetSoftObjectPath());
 	}
 
-	GEditor->GetEditorSubsystem<UAssetEditorSubsystem>()->OpenEditorsForAssets(AssetNames);
+	GEditor->GetEditorSubsystem<UAssetEditorSubsystem>()->OpenEditorsForAssets(AssetPaths);
 }
 
 void SAssetAuditBrowser::OnRequestOpenAsset(const FAssetData& AssetData) const
 {
-	TArray<FName> AssetNames;
-	AssetNames.Add(AssetData.ObjectPath);
+	TArray<FSoftObjectPath> AssetPaths;
+	AssetPaths.Add(AssetData.GetSoftObjectPath());
 
-	GEditor->GetEditorSubsystem<UAssetEditorSubsystem>()->OpenEditorsForAssets(AssetNames);
+	GEditor->GetEditorSubsystem<UAssetEditorSubsystem>()->OpenEditorsForAssets(AssetPaths);
 }
 
 void SAssetAuditBrowser::SaveSelectedAssets() const
diff --git a/Engine/Plugins/Editor/AssetManagerEditor/Source/AssetManagerEditor/Private/SSizeMap.cpp b/Engine/Plugins/Editor/AssetManagerEditor/Source/AssetManagerEditor/Private/SSizeMap.cpp
--- a/Engine/Plugins/Editor/AssetManagerEditor/Source/AssetManagerEditor/Private/SSizeMap.cpp
+++ b/Engine/Plugins/Editor/AssetManagerEditor/Source/AssetManagerEditor/Private/SSizeMap.cpp
@@ -510,7 +510,7 @@ void SSizeMap::GatherDependenciesRecursively(TSharedPtr<FAssetThumbnailPool>& In
 				NodeSizeMapData.AssetData.AssetClassPath = FTopLevelAssetPath(TEXT("/None"), *LOCTEXT("MissingAsset", "MISSING!").ToString());
 
 				const FString AssetPathString = AssetPackageNameString + TEXT(".") + FPackageName::GetLongPackageAssetName(AssetPackageNameString);
-				FAssetData FoundData = CurrentRegistrySource->GetAssetByObjectPath(FName(*AssetPathString));
+				FAssetData FoundData = CurrentRegistrySource->GetAssetByObjectPath(FSoftObjectPath(AssetPathString));
 
 				if (FoundData.IsValid())
 				{
@@ -964,10 +964,10 @@ void SSizeMap::EditSelectedAssets() const
 	const FNodeSizeMapData* NodeSizeMapData = GetCurrentSizeMapData(true);
 	if (NodeSizeMapData)
 	{
-		TArray<FName> AssetNames;
-		AssetNames.Add(NodeSizeMapData->AssetData.ObjectPath);
+		TArray<FSoftObjectPath> AssetPaths;
+		AssetPaths.Add(NodeSizeMapData->AssetData.GetSoftObjectPath());
 
-		GEditor->GetEditorSubsystem<UAssetEditorSubsystem>()->OpenEditorsForAssets(AssetNames);
+		GEditor->GetEditorSubsystem<UAssetEditorSubsystem>()->OpenEditorsForAssets(AssetPaths);
 	}
 }
 	
diff --git a/Engine/Plugins/Editor/AssetReferenceRestrictions/Source/AssetReferenceRestrictions/Private/DomainAssetReferenceFilter.cpp b/Engine/Plugins/Editor/AssetReferenceRestrictions/Source/AssetReferenceRestrictions/Private/DomainAssetReferenceFilter.cpp
--- a/Engine/Plugins/Editor/AssetReferenceRestrictions/Source/AssetReferenceRestrictions/Private/DomainAssetReferenceFilter.cpp
+++ b/Engine/Plugins/Editor/AssetReferenceRestrictions/Source/AssetReferenceRestrictions/Private/DomainAssetReferenceFilter.cpp
@@ -121,7 +121,7 @@ bool FDomainAssetReferenceFilter::PassesFilterImpl(const FAssetData& AssetData,
 	// The only exception anticipated is if somehow an editor with an asset picker is pointed
 	// towards something living in the transient package, but that's not expected!
 	// If it gets hit, transient can be added to the special system mount 'Temp' domain
-	if (ensureMsgf(AssetDomain.IsValid(), TEXT("Asset %s didn't match any domain"), *AssetData.ObjectPath.ToString()))
+	if (ensureMsgf(AssetDomain.IsValid(), TEXT("Asset %s didn't match any domain"), *AssetData.GetObjectPathString()))
 	{
 		// Check the referencing domains
 		for (const TSharedPtr<FDomainData>& ReferencingDomain : ReferencingDomains)
diff --git a/Engine/Plugins/Editor/AssetSearch/Source/Private/AssetSearchDatabase.cpp b/Engine/Plugins/Editor/AssetSearch/Source/Private/AssetSearchDatabase.cpp
--- a/Engine/Plugins/Editor/AssetSearch/Source/Private/AssetSearchDatabase.cpp
+++ b/Engine/Plugins/Editor/AssetSearch/Source/Private/AssetSearchDatabase.cpp
@@ -119,7 +119,7 @@ public:
 	public: bool IsAssetUpToDate(const FAssetData& InAssetData, const FString& IndexedJsonHash)
 	{
 		FString OutIndexedJsonHash;
-		if (Statement_IsAssetUpToDate.BindAndExecuteSingle(InAssetData.ObjectPath.ToString(), OutIndexedJsonHash))
+		if (Statement_IsAssetUpToDate.BindAndExecuteSingle(InAssetData.GetObjectPathString(), OutIndexedJsonHash))
 		{
 			return OutIndexedJsonHash.Equals(IndexedJsonHash, ESearchCase::CaseSensitive);
 		}
@@ -153,7 +153,7 @@ public:
 	public: int64 GetAssetIdForAsset(const FAssetData& InAssetData)
 	{
 		int64 OutAssetId = INDEX_NONE;
-		if (Statement_GetAssetIdForAssetPath.BindAndExecuteSingle(InAssetData.ObjectPath.ToString(), OutAssetId))
+		if (Statement_GetAssetIdForAssetPath.BindAndExecuteSingle(InAssetData.GetObjectPathString(), OutAssetId))
 		{
 			return OutAssetId;
 		}
@@ -177,7 +177,7 @@ public:
 	private: FAddAssetPropertiesFromJson Statement_AddAssetProperty;
 	public: bool AddSearchRecord(const FAssetData& InAssetData, const FString& IndexedJson, const FString& IndexedJsonHash)
 	{
-		FString AssetObjectPath = InAssetData.ObjectPath.ToString();
+		FString AssetObjectPath = InAssetData.GetObjectPathString();
 		if (Statement_AddAssetToAssetTable.BindAndExecute(InAssetData.AssetName.ToString(), InAssetData.AssetClassPath.ToString(), AssetObjectPath, IndexedJsonHash))
 		{
 			int64 AssetId = Database.GetLastInsertRowId();
@@ -304,7 +304,7 @@ public:
 	private: FDeleteEntriesForAsset Statement_DeleteEntriesForAsset;
 	public: bool DeleteEntriesForAsset(const FAssetData& InAssetData)
 	{
-		return DeleteEntriesForAsset(InAssetData.ObjectPath.ToString());
+		return DeleteEntriesForAsset(InAssetData.GetObjectPathString());
 	}
 	public: bool DeleteEntriesForAsset(const FString& InAssetObjectPath)
 	{
@@ -807,7 +807,7 @@ void FAssetSearchDatabase::RemoveAssetsNotInThisSet(const TArray<FAssetData>& In
 			continue;
 		}
 
-		InAssetPaths.Add(InAsset.ObjectPath.ToString());
+		InAssetPaths.Add(InAsset.GetObjectPathString());
 	}
 
 	TArray<FString> MissingAssets;
diff --git a/Engine/Plugins/Editor/AssetSearch/Source/Private/AssetSearchManager.cpp b/Engine/Plugins/Editor/AssetSearch/Source/Private/AssetSearchManager.cpp
--- a/Engine/Plugins/Editor/AssetSearch/Source/Private/AssetSearchManager.cpp
+++ b/Engine/Plugins/Editor/AssetSearch/Source/Private/AssetSearchManager.cpp
@@ -610,7 +610,7 @@ bool FAssetSearchManager::AsyncGetDerivedDataKey(const FAssetData& InAssetData,
 			// The universal key for content is:
 			// AssetSearch_V{SerializerVersion}_{IndexersNamesAndVersions}_{ObjectPathHash}_{FileOnDiskHash}
 
-			const FString ObjectPathString = InAssetData.ObjectPath.ToString();
+			const FString ObjectPathString = InAssetData.GetObjectPathString();
 
 			FSHAHash ObjectPathHash;
 			FSHA1::HashBuffer(*ObjectPathString, ObjectPathString.Len() * sizeof(FString::ElementType), ObjectPathHash.Hash);
@@ -713,7 +713,7 @@ void FAssetSearchManager::StoreIndexForAsset(UObject* InAsset)
 
 					FTCHARToUTF8 IndexedJsonUTF8(*IndexedJson);
 					TArrayView<const uint8> IndexedJsonUTF8View((const uint8*)IndexedJsonUTF8.Get(), IndexedJsonUTF8.Length() * sizeof(UTF8CHAR));
-					GetDerivedDataCacheRef().Put(*InDDCKey, IndexedJsonUTF8View, InAssetData.ObjectPath.ToString(), false);
+					GetDerivedDataCacheRef().Put(*InDDCKey, IndexedJsonUTF8View, InAssetData.GetObjectPathString(), false);
 
 					AddOrUpdateAsset(InAssetData, IndexedJson, InDDCKey);
 				});
@@ -787,7 +787,7 @@ bool FAssetSearchManager::Tick_GameThread(float DeltaTime)
 		DownloadQueueCount--;
 		ActiveDownloads++;
 
-		DDCRequest.DDCHandle = GetDerivedDataCacheRef().GetAsynchronous(*DDCRequest.DDCKey, DDCRequest.AssetData.ObjectPath.ToString());
+		DDCRequest.DDCHandle = GetDerivedDataCacheRef().GetAsynchronous(*DDCRequest.DDCKey, DDCRequest.AssetData.GetObjectPathString());
 		ProcessDDCQueue.Enqueue(DDCRequest);
 	}
 
diff --git a/Engine/Plugins/Editor/AssetSearch/Source/Private/Providers/AssetRegistrySearchProvider.cpp b/Engine/Plugins/Editor/AssetSearch/Source/Private/Providers/AssetRegistrySearchProvider.cpp
--- a/Engine/Plugins/Editor/AssetSearch/Source/Private/Providers/AssetRegistrySearchProvider.cpp
+++ b/Engine/Plugins/Editor/AssetSearch/Source/Private/Providers/AssetRegistrySearchProvider.cpp
@@ -225,11 +225,11 @@ public:
 			bool bIsMatch = false;
 			if (InTextComparisonMode == ETextFilterTextComparisonMode::Partial)
 			{
-				bIsMatch = TextFilterUtils::TestBasicStringExpression(AssetPtr->ObjectPath, InValue, InTextComparisonMode);
+				bIsMatch = TextFilterUtils::TestBasicStringExpression(AssetPtr->GetObjectPathString(), InValue, InTextComparisonMode);
 			}
 			else
 			{
-				bIsMatch = TextFilterUtils::TestBasicStringExpression(AssetPtr->ObjectPath, InValue, InTextComparisonMode)
+				bIsMatch = TextFilterUtils::TestBasicStringExpression(AssetPtr->GetObjectPathString(), InValue, InTextComparisonMode)
 					|| TextFilterUtils::TestBasicStringExpression(AssetPtr->PackageName, InValue, InTextComparisonMode)
 					|| TextFilterUtils::TestBasicStringExpression(AssetPtr->PackagePath, InValue, InTextComparisonMode);
 			}
@@ -363,7 +363,7 @@ void FAssetRegistrySearchProvider::Search(FSearchQueryPtr SearchQuery)
 		for (const FAssetData& Asset : Assets)
 		{
 			FSearchRecord Record;
-			Record.AssetPath = Asset.ObjectPath.ToString();
+			Record.AssetPath = Asset.GetObjectPathString();
 			Record.AssetName = Asset.AssetName.ToString();
 			Record.AssetClass = Asset.AssetClassPath;
 
diff --git a/Engine/Source/Editor/UnrealEd/Private/AutoReimport/AssetSourceFilenameCache.cpp b/Engine/Source/Editor/UnrealEd/Private/AutoReimport/AssetSourceFilenameCache.cpp
--- a/Engine/Source/Editor/UnrealEd/Private/AutoReimport/AssetSourceFilenameCache.cpp
+++ b/Engine/Source/Editor/UnrealEd/Private/AutoReimport/AssetSourceFilenameCache.cpp
@@ -100,7 +100,7 @@ void FAssetSourceFilenameCache::HandleOnAssetAdded(const FAssetData& AssetData)
 	{
 		for (const auto& SourceFile : ImportData->SourceFiles)
 		{
-			SourceFileToObjectPathCache.FindOrAdd(FPaths::GetCleanFilename(SourceFile.RelativeFilename)).Add(AssetData.ObjectPath);
+			SourceFileToObjectPathCache.FindOrAdd(FPaths::GetCleanFilename(SourceFile.RelativeFilename)).Add(AssetData.GetSoftObjectPath());
 		}
 	}
 }
@@ -115,7 +115,7 @@ void FAssetSourceFilenameCache::HandleOnAssetRemoved(const FAssetData& AssetData
 			FString CleanFilename = FPaths::GetCleanFilename(SourceFile.RelativeFilename);
 			if (auto* Objects = SourceFileToObjectPathCache.Find(CleanFilename))
 			{
-				Objects->Remove(AssetData.ObjectPath);
+				Objects->Remove(AssetData.GetSoftObjectPath());
 				if (Objects->Num() == 0)
 				{
 					SourceFileToObjectPathCache.Remove(CleanFilename);
@@ -130,22 +130,20 @@ void FAssetSourceFilenameCache::HandleOnAssetRenamed(const FAssetData& AssetData
 	TOptional<FAssetImportInfo> ImportData = ExtractAssetImportInfo(AssetData);
 	if (ImportData.IsSet())
 	{
-		FName OldPathName = *OldPath;
-
 		for (auto& SourceFile : ImportData->SourceFiles)
 		{
 			FString CleanFilename = FPaths::GetCleanFilename(SourceFile.RelativeFilename);
 
 			if (auto* Objects = SourceFileToObjectPathCache.Find(CleanFilename))
 			{
-				Objects->Remove(OldPathName);
+				Objects->Remove(FSoftObjectPath(OldPath));
 				if (Objects->Num() == 0)
 				{
 					SourceFileToObjectPathCache.Remove(CleanFilename);
 				}
 			}
 
-			SourceFileToObjectPathCache.FindOrAdd(CleanFilename).Add(AssetData.ObjectPath);
+			SourceFileToObjectPathCache.FindOrAdd(CleanFilename).Add(AssetData.GetSoftObjectPath());
 		}
 	}
 
@@ -154,7 +152,7 @@ void FAssetSourceFilenameCache::HandleOnAssetRenamed(const FAssetData& AssetData
 
 void FAssetSourceFilenameCache::HandleOnAssetUpdated(const FAssetImportInfo& OldData, const UAssetImportData* ImportData)
 {
-	FName ObjectPath = FName(*ImportData->GetOuter()->GetPathName());
+	FSoftObjectPath ObjectPath = FSoftObjectPath(ImportData->GetOuter());
 
 	for (const auto& SourceFile : OldData.SourceFiles)
 	{
@@ -181,7 +179,7 @@ TArray<FAssetData> FAssetSourceFilenameCache::GetAssetsPertainingToFile(const IA
 	
 	if (const auto* ObjectPaths = SourceFileToObjectPathCache.Find(FPaths::GetCleanFilename(AbsoluteFilename)))
 	{
-		for (const FName& Path : *ObjectPaths)
+		for (const FSoftObjectPath& Path : *ObjectPaths)
 		{
 			FAssetData Asset = Registry.GetAssetByObjectPath(Path);
 			TOptional<FAssetImportInfo> ImportInfo = ExtractAssetImportInfo(Asset);
diff --git a/Engine/Source/Editor/UnrealEd/Private/Dialogs/SDeleteAssetsDialog.cpp b/Engine/Source/Editor/UnrealEd/Private/Dialogs/SDeleteAssetsDialog.cpp
--- a/Engine/Source/Editor/UnrealEd/Private/Dialogs/SDeleteAssetsDialog.cpp
+++ b/Engine/Source/Editor/UnrealEd/Private/Dialogs/SDeleteAssetsDialog.cpp
@@ -819,7 +819,7 @@ void SDeleteAssetsDialog::OnAssetsActivated(const TArray<FAssetData>& ActivatedA
 		for(const FAssetData& ActivatedAsset : ActivatedAssets)
 		{
 			FString MapFilePath;
-			if ( FEditorFileUtils::IsMapPackageAsset(ActivatedAsset.ObjectPath.ToString(), MapFilePath) )
+			if (FEditorFileUtils::IsMapPackageAsset(ActivatedAsset.GetObjectPathString(), MapFilePath))
 			{
 				if ( ActivatedAsset.IsAssetLoaded() )
 				{
diff --git a/Engine/Source/Editor/UnrealEd/Private/Subsystems/AssetEditorSubsystem.cpp b/Engine/Source/Editor/UnrealEd/Private/Subsystems/AssetEditorSubsystem.cpp
--- a/Engine/Source/Editor/UnrealEd/Private/Subsystems/AssetEditorSubsystem.cpp
+++ b/Engine/Source/Editor/UnrealEd/Private/Subsystems/AssetEditorSubsystem.cpp
@@ -595,6 +595,11 @@ void UAssetEditorSubsystem::HandleRequestOpenAssetMessage(const FAssetEditorRequ
 	OpenEditorForAsset(Message.AssetName);
 }
 
+void UAssetEditorSubsystem::OpenEditorForAsset(const FSoftObjectPath& AssetPath)
+{
+	OpenEditorForAsset(AssetPath.ToString());
+}
+
 void UAssetEditorSubsystem::OpenEditorForAsset(const FString& AssetPathName)
 {
 	// An asset needs loading
@@ -986,6 +991,14 @@ void UAssetEditorSubsystem::HandlePackageReloaded(const EPackageReloadPhase InPa
 	}
 }
 
+void UAssetEditorSubsystem::OpenEditorsForAssets(const TArray<FSoftObjectPath>& AssetsToOpen)
+{
+	for (const FSoftObjectPath& AssetName : AssetsToOpen)
+	{
+		OpenEditorForAsset(AssetName);
+	}
+}
+
 void UAssetEditorSubsystem::OpenEditorsForAssets(const TArray<FString>& AssetsToOpen)
 {
 	for (const FString& AssetName : AssetsToOpen)
diff --git a/Engine/Source/Editor/UnrealEd/Private/Subsystems/EditorAssetSubsystem.cpp b/Engine/Source/Editor/UnrealEd/Private/Subsystems/EditorAssetSubsystem.cpp
--- a/Engine/Source/Editor/UnrealEd/Private/Subsystems/EditorAssetSubsystem.cpp
+++ b/Engine/Source/Editor/UnrealEd/Private/Subsystems/EditorAssetSubsystem.cpp
@@ -59,7 +59,7 @@ namespace UE::EditorAssetUtils
 		}
 
 		FAssetRegistryModule& AssetRegistryModule = FModuleManager::LoadModuleChecked<FAssetRegistryModule>("AssetRegistry");
-		FAssetData AssetData = AssetRegistryModule.Get().GetAssetByObjectPath(*ObjectPath);
+		FAssetData AssetData = AssetRegistryModule.Get().GetAssetByObjectPath(FSoftObjectPath(ObjectPath));
 		if (!AssetData.IsValid())
 		{
 			return MakeError(FString::Printf(TEXT("The AssetData '%s' could not be found in the Asset Registry."), *ObjectPath));
@@ -78,11 +78,11 @@ namespace UE::EditorAssetUtils
 		UObject* FoundObject = AssetData.GetAsset();
 		if (!IsValid(FoundObject))
 		{
-			return MakeError(FString::Printf(TEXT("The asset '%s' exists but was not able to be loaded."), *AssetData.ObjectPath.ToString()));
+			return MakeError(FString::Printf(TEXT("The asset '%s' exists but was not able to be loaded."), *AssetData.GetObjectPathString()));
 		}
 		else if (!FoundObject->IsAsset())
 		{
-			return MakeError(FString::Printf(TEXT("'%s' is not a valid asset."), *AssetData.ObjectPath.ToString()));
+			return MakeError(FString::Printf(TEXT("'%s' is not a valid asset."), *AssetData.GetObjectPathString()));
 		}
 		return MakeValue(FoundObject);
 	}
@@ -116,7 +116,10 @@ namespace UE::EditorAssetUtils
 
 		UPackage* Package = Object->GetPackage();
 		FAssetRegistryModule& AssetRegistryModule = FModuleManager::LoadModuleChecked<FAssetRegistryModule>("AssetRegistry");
-		FAssetData AssetData = AssetRegistryModule.Get().GetAssetByObjectPath(Package->GetFName());
+PRAGMA_DISABLE_DEPRECATION_WARNINGS
+		// This should likely get GetAssetsByPackagename but the impact of changing it should be checked
+		FAssetData AssetData = AssetRegistryModule.Get().GetAssetByObjectPath(Package->GetName());
+PRAGMA_ENABLE_DEPRECATION_WARNINGS
 		if (!AssetData.IsValid())
 		{
 			return MakeError(FString::Printf(TEXT("The AssetData '%s' could not be found in the Asset Registry."), *Package->GetName()));
@@ -450,7 +453,7 @@ bool UEditorAssetSubsystem::DoesAssetExist(const FString& AssetPath)
 	}
 
 	FAssetRegistryModule& AssetRegistryModule = FModuleManager::LoadModuleChecked<FAssetRegistryModule>("AssetRegistry");
-	FAssetData AssetData = AssetRegistryModule.Get().GetAssetByObjectPath(*ObjectPath);
+	FAssetData AssetData = AssetRegistryModule.Get().GetAssetByObjectPath(FSoftObjectPath(ObjectPath));
 	if (!AssetData.IsValid())
 	{
 		return false;
@@ -480,7 +483,7 @@ bool UEditorAssetSubsystem::DoAssetsExist(const TArray<FString>& AssetPaths)
 			return false;
 		}
 
-		FAssetData AssetData = AssetRegistryModule.Get().GetAssetByObjectPath(*ObjectPath);
+		FAssetData AssetData = AssetRegistryModule.Get().GetAssetByObjectPath(FSoftObjectPath(ObjectPath));
 		if (!AssetData.IsValid())
 		{
 			return false;
@@ -895,7 +898,7 @@ bool UEditorAssetSubsystem::DuplicateDirectory(const FString& SourceDirectoryPat
 		UObject* DuplicatedAsset = Module.Get().DuplicateAsset(AssetPaths.SourceDirectoryAssetDatas[Index].AssetName.ToString(), AssetPaths.DestinationDirectoryAssetPaths[Index], AssetPaths.SourceDirectoryAssetDatas[Index].GetAsset());
 		if (DuplicatedAsset == nullptr)
 		{
-			UE_LOG(LogEditorAssetSubsystem, Warning, TEXT("DuplicateDirectory failed: Could not duplicate object '%s'"), *AssetPaths.SourceDirectoryAssetDatas[Index].ObjectPath.ToString());
+			UE_LOG(LogEditorAssetSubsystem, Warning, TEXT("DuplicateDirectory failed: Could not duplicate object '%s'"), *AssetPaths.SourceDirectoryAssetDatas[Index].GetObjectPathString());
 			bHaveAFailedAsset = true;
 		}
 	}
@@ -1018,7 +1021,7 @@ bool UEditorAssetSubsystem::RenameDirectory(const FString& SourceDirectoryPath,
 
 	for (int32 Index = 0; Index < AssetPathsResult.GetValue().SourceDirectoryAssetDatas.Num(); ++Index)
 	{
-		AssetsToRename.Add(FAssetRenameData(AssetPathsResult.GetValue().SourceDirectoryAssetDatas[Index].ObjectPath.ToString(), AssetPathsResult.GetValue().DestinationDirectoryAssetPaths[Index]));
+		AssetsToRename.Add(FAssetRenameData(AssetPathsResult.GetValue().SourceDirectoryAssetDatas[Index].GetObjectPathString(), AssetPathsResult.GetValue().DestinationDirectoryAssetPaths[Index]));
 	}
 
 	// Rename the assets
@@ -1386,7 +1389,7 @@ TArray<FString> UEditorAssetSubsystem::ListAssets(const FString& DirectoryPath,
 		AssetDatas.Reserve(AssetDatas.Num());
 		for (const FAssetData& AssetData : AssetDatas)
 		{
-			AssetPaths.Add(AssetData.ObjectPath.ToString());
+			AssetPaths.Add(AssetData.GetObjectPathString());
 		}
 	}
 
diff --git a/Engine/Source/Editor/UnrealEd/Public/AutoReimport/AssetSourceFilenameCache.h b/Engine/Source/Editor/UnrealEd/Public/AutoReimport/AssetSourceFilenameCache.h
--- a/Engine/Source/Editor/UnrealEd/Public/AutoReimport/AssetSourceFilenameCache.h
+++ b/Engine/Source/Editor/UnrealEd/Public/AutoReimport/AssetSourceFilenameCache.h
@@ -53,5 +53,5 @@ private:
 	FAssetRenamedEvent AssetRenamedEvent;
 
 	/** Map of clean filenames (no leading path information) to object paths that were imported with that file */
-	TMap<FString, TSet<FName>> SourceFileToObjectPathCache;
+	TMap<FString, TSet<FSoftObjectPath>> SourceFileToObjectPathCache;
 };
diff --git a/Engine/Source/Editor/UnrealEd/Public/Subsystems/AssetEditorSubsystem.h b/Engine/Source/Editor/UnrealEd/Public/Subsystems/AssetEditorSubsystem.h
--- a/Engine/Source/Editor/UnrealEd/Public/Subsystems/AssetEditorSubsystem.h
+++ b/Engine/Source/Editor/UnrealEd/Public/Subsystems/AssetEditorSubsystem.h
@@ -73,13 +73,20 @@ public:
 	bool IsAssetEditable(const UObject* Asset);
 
 	/** Opens an asset by path */
-	void OpenEditorForAsset(const FString& AssetPathName);
+	void OpenEditorForAsset(const FString& AssetPath);
+	void OpenEditorForAsset(const FSoftObjectPath& AssetPath);
 
 	/**
 	 * Tries to open an editor for the specified asset.  Returns true if the asset is open in an editor.
 	 * If the file is already open in an editor, it will not create another editor window but instead bring it to front
 	 */
 	bool OpenEditorForAsset(UObject* Asset, const EToolkitMode::Type ToolkitMode = EToolkitMode::Standalone, TSharedPtr<IToolkitHost> OpenedFromLevelEditor = TSharedPtr<IToolkitHost>(), const bool bShowProgressWindow = true);
+	
+	template<typename ObjectType>
+	bool OpenEditorForAsset(TObjectPtr<ObjectType> Asset, const EToolkitMode::Type ToolkitMode = EToolkitMode::Standalone, TSharedPtr<IToolkitHost> OpenedFromLevelEditor = TSharedPtr<IToolkitHost>(), const bool bShowProgressWindow = true)
+	{
+		return OpenEditorForAsset(ToRawPtr(Asset), ToolkitMode, OpenedFromLevelEditor, bShowProgressWindow);
+	}
 
 	/**
 	 * Tries to open an editor for all of the specified assets.
@@ -93,6 +100,7 @@ public:
 	/** Opens editors for the supplied assets (via OpenEditorForAsset) */
 	void OpenEditorsForAssets(const TArray<FString>& AssetsToOpen);
 	void OpenEditorsForAssets(const TArray<FName>& AssetsToOpen);
+	void OpenEditorsForAssets(const TArray<FSoftObjectPath>& AssetsToOpen);
 
 	/** Returns the primary editor if one is already open for the specified asset.
 	 * If there is one open and bFocusIfOpen is true, that editor will be brought to the foreground and focused if possible.
diff --git a/Engine/Source/Runtime/Engine/Classes/Engine/AssetManager.h b/Engine/Source/Runtime/Engine/Classes/Engine/AssetManager.h
--- a/Engine/Source/Runtime/Engine/Classes/Engine/AssetManager.h
+++ b/Engine/Source/Runtime/Engine/Classes/Engine/AssetManager.h
@@ -193,6 +193,8 @@ public:
 
 	/** Sees if the passed in object path is a registered primary asset, if so return it. Returns invalid Identifier if not found */
 	virtual FPrimaryAssetId GetPrimaryAssetIdForPath(const FSoftObjectPath& ObjectPath) const;
+
+	UE_DEPRECATED(5.1, "FName asset paths are deprecated, use FSoftObjectPath instead.")
 	virtual FPrimaryAssetId GetPrimaryAssetIdForPath(FName ObjectPath) const;
 
 	/** Sees if the package has a primary asset, useful if only the package name is available */
@@ -407,6 +409,7 @@ public:
 	virtual FPrimaryAssetId DeterminePrimaryAssetIdForObject(const UObject* Object) const;
 
 	/** Reads AssetManagerSettings for specifically redirected asset paths. This is useful if you need to convert older saved data */
+	UE_DEPRECATED(5.1, "Asset path FNames are deprecated, use FSoftObjectPath instead.")
 	virtual FName GetRedirectedAssetPath(FName OldPath) const;
 	virtual FSoftObjectPath GetRedirectedAssetPath(const FSoftObjectPath& OldPath) const;
 
@@ -696,7 +699,7 @@ protected:
 #endif // WITH_EDITOR
 
 	/** Map from object path to Primary Asset Id */
-	TMap<FName, FPrimaryAssetId> AssetPathMap;
+	TMap<FSoftObjectPath, FPrimaryAssetId> AssetPathMap;
 
 	/** Overridden asset management data for specific types */
 	TMap<FPrimaryAssetId, FPrimaryAssetRulesExplicitOverride> AssetRuleOverrides;
@@ -806,7 +809,7 @@ protected:
 	/** Redirector maps loaded out of AssetMigrations.ini */
 	TMap<FName, FName> PrimaryAssetTypeRedirects;
 	TMap<FString, FString> PrimaryAssetIdRedirects;
-	TMap<FName, FName> AssetPathRedirects;
+	TMap<FSoftObjectPath, FSoftObjectPath> AssetPathRedirects;
 
 	/** Delegate called when initial span finishes */
 	static FSimpleMulticastDelegate OnCompletedInitialScanDelegate;
diff --git a/Engine/Source/Runtime/Engine/Private/AssetManager.cpp b/Engine/Source/Runtime/Engine/Private/AssetManager.cpp
--- a/Engine/Source/Runtime/Engine/Private/AssetManager.cpp
+++ b/Engine/Source/Runtime/Engine/Private/AssetManager.cpp
@@ -76,7 +76,7 @@ struct FPrimaryAssetLoadState
 struct FPrimaryAssetData
 {
 	/** Path used to look up cached asset data in the asset registry. This will be missing the _C for blueprint classes */
-	FName AssetDataPath;
+	FSoftObjectPath AssetDataPath;
 
 	/** Path to this asset on disk */
 	FSoftObjectPtr AssetPtr;
@@ -457,7 +457,7 @@ FPrimaryAssetId UAssetManager::DeterminePrimaryAssetIdForObject(const UObject* O
 	}
 
 	FString AssetPath = AssetObject->GetPathName();
-	FPrimaryAssetId RegisteredId = GetPrimaryAssetIdForPath(FName(*AssetPath));
+	FPrimaryAssetId RegisteredId = GetPrimaryAssetIdForPath(FSoftObjectPath(AssetObject));
 
 	if (RegisteredId.IsValid())
 	{
@@ -1035,7 +1035,7 @@ bool UAssetManager::TryUpdateCachedAssetData(const FPrimaryAssetId& PrimaryAsset
 		if (!TryToSoftObjectPath<EAssetDataCanBeSubobject::No>(NewAssetData, NewAssetPath))
 		{
 			UE_LOG(LogAssetManager, Warning, TEXT("Tried to add primary asset %s, but it is not a TopLevelAsset and so cannot be a primary asset"),
-				*NewAssetData.ObjectPath.ToString())
+				*NewAssetData.GetObjectPathString())
 			return false;
 		}
 
@@ -1076,7 +1076,7 @@ bool UAssetManager::TryUpdateCachedAssetData(const FPrimaryAssetId& PrimaryAsset
 		FPrimaryAssetData& NameData = TypeData.AssetMap.FindOrAdd(PrimaryAssetId.PrimaryAssetName);
 
 		// Update data and path, don't touch state or references
-		NameData.AssetDataPath = NewAssetData.ObjectPath; // This will not have _C
+		NameData.AssetDataPath = NewAssetData.GetSoftObjectPath(); // This will not have _C
 		NameData.AssetPtr = FSoftObjectPtr(NewAssetPath); // This will have _C
 
 		// If the types don't match, update the registry
@@ -1091,7 +1091,7 @@ bool UAssetManager::TryUpdateCachedAssetData(const FPrimaryAssetId& PrimaryAsset
 		if (IsBulkScanning())
 		{
 			// Do a partial update, add to the path->asset map
-			AssetPathMap.Add(NewAssetPath.GetAssetPathName(), PrimaryAssetId);
+			AssetPathMap.Add(NewAssetPath, PrimaryAssetId);
 		}
 
 		if (NewAssetData.TaggedAssetBundles)
@@ -1117,9 +1117,10 @@ bool UAssetManager::TryUpdateCachedAssetData(const FPrimaryAssetId& PrimaryAsset
 			
 				for (const FAssetBundleEntry& Entry : NewAssetData.TaggedAssetBundles->Bundles)
 				{
-					for (const FSoftObjectPath& Path : Entry.BundleAssets)
+					for (const FTopLevelAssetPath& Path : Entry.AssetPaths)
 					{
-						Path.PostLoadPath(nullptr);
+						FSoftObjectPath FullPath(Path);
+						FullPath.PostLoadPath(nullptr);
 					}
 				}
 			}
@@ -1186,7 +1187,7 @@ bool UAssetManager::AddDynamicAsset(const FPrimaryAssetId& PrimaryAssetId, const
 	if (IsBulkScanning() && AssetPath.IsValid())
 	{
 		// Do a partial update, add to the path->asset map
-		AssetPathMap.Add(AssetPath.GetAssetPathName(), PrimaryAssetId);
+		AssetPathMap.Add(AssetPath, PrimaryAssetId);
 	}
 
 
@@ -1204,13 +1205,14 @@ bool UAssetManager::AddDynamicAsset(const FPrimaryAssetId& PrimaryAssetId, const
 
 void UAssetManager::RecursivelyExpandBundleData(FAssetBundleData& BundleData) const
 {
-	TArray<FSoftObjectPath> ReferencesToExpand;
+	TArray<FTopLevelAssetPath> ReferencesToExpand;
 	TSet<FName> FoundBundleNames;
 
 	for (const FAssetBundleEntry& Entry : BundleData.Bundles)
 	{
 		FoundBundleNames.Add(Entry.BundleName);
-		for (const FSoftObjectPath& Reference : Entry.BundleAssets)
+
+		for (const FTopLevelAssetPath& Reference : Entry.AssetPaths)
 		{
 			ReferencesToExpand.AddUnique(Reference);
 		}
@@ -1219,7 +1221,7 @@ void UAssetManager::RecursivelyExpandBundleData(FAssetBundleData& BundleData) co
 	// Expandable references can increase recursively
 	for (int32 i = 0; i < ReferencesToExpand.Num(); i++)
 	{
-		FPrimaryAssetId FoundId = GetPrimaryAssetIdForPath(ReferencesToExpand[i]);
+		FPrimaryAssetId FoundId = GetPrimaryAssetIdForPath(FSoftObjectPath(ReferencesToExpand[i]));
 		TArray<FAssetBundleEntry> FoundEntries;
 
 		if (FoundId.IsValid() && GetAssetBundleEntries(FoundId, FoundEntries))
@@ -1229,9 +1231,9 @@ void UAssetManager::RecursivelyExpandBundleData(FAssetBundleData& BundleData) co
 				// Make sure the bundle name matches
 				if (FoundBundleNames.Contains(FoundEntry.BundleName))
 				{
-					BundleData.AddBundleAssets(FoundEntry.BundleName, FoundEntry.BundleAssets);
+					BundleData.AddBundleAssets(FoundEntry.BundleName, FoundEntry.AssetPaths);
 
-					for (const FSoftObjectPath& FoundReference : FoundEntry.BundleAssets)
+					for (const FTopLevelAssetPath& FoundReference : FoundEntry.AssetPaths)
 					{
 						// Keep recursing
 						ReferencesToExpand.AddUnique(FoundReference);
@@ -1427,7 +1429,7 @@ bool UAssetManager::GetPrimaryAssetPathList(FPrimaryAssetType PrimaryAssetType,
 FPrimaryAssetId UAssetManager::GetPrimaryAssetIdForObject(UObject* Object) const
 {
 	// Use path instead of calling on Object, we only want it if it's registered
-	return GetPrimaryAssetIdForPath(FName(*Object->GetPathName()));
+	return GetPrimaryAssetIdForPath(FSoftObjectPath(Object));
 }
 
 FPrimaryAssetId UAssetManager::GetPrimaryAssetIdForData(const FAssetData& AssetData) const
@@ -1436,20 +1438,15 @@ FPrimaryAssetId UAssetManager::GetPrimaryAssetIdForData(const FAssetData& AssetD
 }
 
 FPrimaryAssetId UAssetManager::GetPrimaryAssetIdForPath(const FSoftObjectPath& ObjectPath) const
-{
-	return GetPrimaryAssetIdForPath(ObjectPath.GetAssetPathName());
-}
-
-FPrimaryAssetId UAssetManager::GetPrimaryAssetIdForPath(FName ObjectPath) const
 {
 	const FPrimaryAssetId* FoundIdentifier = AssetPathMap.Find(ObjectPath);
 
 	// Check redirector list
 	if (!FoundIdentifier)
 	{
-		FName RedirectedPath = GetRedirectedAssetPath(ObjectPath);
+		FSoftObjectPath RedirectedPath = GetRedirectedAssetPath(ObjectPath);
 
-		if (RedirectedPath != NAME_None)
+		if (!RedirectedPath.IsNull())
 		{
 			FoundIdentifier = AssetPathMap.Find(RedirectedPath);
 		}
@@ -1463,16 +1460,23 @@ FPrimaryAssetId UAssetManager::GetPrimaryAssetIdForPath(FName ObjectPath) const
 	return FPrimaryAssetId();
 }
 
+FPrimaryAssetId UAssetManager::GetPrimaryAssetIdForPath(FName ObjectPath) const
+{
+PRAGMA_DISABLE_DEPRECATION_WARNINGS
+	return GetPrimaryAssetIdForPath(FSoftObjectPath(ObjectPath));
+PRAGMA_ENABLE_DEPRECATION_WARNINGS
+}
+
 FPrimaryAssetId UAssetManager::GetPrimaryAssetIdForPackage(FName PackagePath) const
 {
 	FString PackageString = PackagePath.ToString();
 	FString AssetName = FPackageName::GetShortName(PackageString);
 
 	FPrimaryAssetId FoundId;
-	FName PossibleAssetPath = FName(*FString::Printf(TEXT("%s.%s"), *PackageString, *AssetName), FNAME_Find);
+	FSoftObjectPath PossibleAssetPath(FString::Printf(TEXT("%s.%s"), *PackageString, *AssetName));
 
 	// Try without _C first
-	if (PossibleAssetPath != NAME_None)
+	if (PossibleAssetPath.IsValid())
 	{
 		FoundId = GetPrimaryAssetIdForPath(PossibleAssetPath);
 
@@ -1483,9 +1487,9 @@ FPrimaryAssetId UAssetManager::GetPrimaryAssetIdForPackage(FName PackagePath) co
 	}
 
 	// Then try _C
-	PossibleAssetPath = FName(*FString::Printf(TEXT("%s.%s_C"), *PackageString, *AssetName), FNAME_Find);
+	PossibleAssetPath = FSoftObjectPath(FString::Printf(TEXT("%s.%s_C"), *PackageString, *AssetName));
 
-	if (PossibleAssetPath != NAME_None)
+	if (PossibleAssetPath.IsValid())
 	{
 		FoundId = GetPrimaryAssetIdForPath(PossibleAssetPath);
 	}
@@ -1645,7 +1649,10 @@ TSharedPtr<FStreamableHandle> UAssetManager::ChangeBundleStateForPrimaryAssets(c
 
 				if (Entry.IsValid())
 				{
-					PathsToLoad.Append(Entry.BundleAssets);
+					for (const FTopLevelAssetPath & Path : Entry.AssetPaths)
+					{
+						PathsToLoad.Emplace(FSoftObjectPath(Path));
+					}
 				}
 				else
 				{
@@ -1804,7 +1811,10 @@ bool UAssetManager::GetPrimaryAssetLoadSet(TSet<FSoftObjectPath>& OutAssetLoadSe
 
 		for (const FAssetBundleEntry& Entry : TempBundleData.Bundles)
 		{
-			OutAssetLoadSet.Append(Entry.BundleAssets);
+			for (const FTopLevelAssetPath& Path : Entry.AssetPaths)
+			{
+				OutAssetLoadSet.Emplace(FSoftObjectPath(Path));
+			}
 		}
 	}
 	else
@@ -2265,7 +2275,10 @@ void UAssetManager::AcquireResourcesForPrimaryAssetList(const TArray<FPrimaryAss
 			{
 				if (Entry.IsValid())
 				{
-					PathsToLoad.Append(Entry.BundleAssets);
+					for (const FTopLevelAssetPath& Path : Entry.AssetPaths)
+					{
+						PathsToLoad.Emplace(FSoftObjectPath(Path));
+					}
 				}
 			}
 		}
@@ -2527,7 +2540,7 @@ void UAssetManager::RebuildObjectReferenceList()
 			// Dynamic types can have null asset refs
 			if (!AssetRef.IsNull())
 			{
-				AssetPathMap.Add(AssetRef.GetAssetPathName(), FPrimaryAssetId(TypePair.Key, NamePair.Key));
+				AssetPathMap.Add(AssetRef, FPrimaryAssetId(TypePair.Key, NamePair.Key));
 			}
 		}
 	}
@@ -2555,18 +2568,18 @@ void UAssetManager::LoadRedirectorMaps()
 
 	for (const FAssetManagerRedirect& Redirect : Settings.AssetPathRedirects)
 	{
-		AssetPathRedirects.Add(FName(*Redirect.Old), FName(*Redirect.New));
+		AssetPathRedirects.Add(FSoftObjectPath(*Redirect.Old), FSoftObjectPath(*Redirect.New));
 	}
 
 	// Collapse all redirects to resolve recursive relationships.
-	for (const TPair<FName, FName>& Pair : AssetPathRedirects)
+	for (const TPair<FSoftObjectPath, FSoftObjectPath>& Pair : AssetPathRedirects)
 	{
-		const FName OldPath = Pair.Key;
-		FName NewPath = Pair.Value;
-		TSet<FName> CollapsedPaths;
+		const FSoftObjectPath OldPath = Pair.Key;
+		FSoftObjectPath NewPath = Pair.Value;
+		TSet<FSoftObjectPath> CollapsedPaths;
 		CollapsedPaths.Add(OldPath);
 		CollapsedPaths.Add(NewPath);
-		while (FName* NewPathValue = AssetPathRedirects.Find(NewPath)) // Does the NewPath exist as a key?
+		while (FSoftObjectPath* NewPathValue = AssetPathRedirects.Find(NewPath)) // Does the NewPath exist as a key?
 		{
 			NewPath = *NewPathValue;
 			if (CollapsedPaths.Contains(NewPath))
@@ -2629,32 +2642,28 @@ void UAssetManager::GetPreviousPrimaryAssetIds(const FPrimaryAssetId& NewId, TAr
 
 FName UAssetManager::GetRedirectedAssetPath(FName OldPath) const
 {
-	const FName* FoundPath = AssetPathRedirects.Find(OldPath);
-
-	if (FoundPath)
-	{
-		return *FoundPath;
-	}
-
-	return NAME_None;
+PRAGMA_DISABLE_DEPRECATION_WARNINGS
+	FSoftObjectPath Redirect = GetRedirectedAssetPath(FSoftObjectPath(OldPath));
+	return Redirect.ToFName();
+PRAGMA_ENABLE_DEPRECATION_WARNINGS
 }
 
 FSoftObjectPath UAssetManager::GetRedirectedAssetPath(const FSoftObjectPath& ObjectPath) const
 {
-	FName PossibleAssetPath = ObjectPath.GetAssetPathName();
+	FSoftObjectPath PossibleAssetPath = ObjectPath;
 
-	if (PossibleAssetPath == NAME_None)
+	if (PossibleAssetPath.IsNull())
 	{
 		return FSoftObjectPath();
 	}
 
-	FName RedirectedName = GetRedirectedAssetPath(PossibleAssetPath);
+	FSoftObjectPath RedirectedName = GetRedirectedAssetPath(PossibleAssetPath);
 
-	if (RedirectedName == NAME_None)
+	if (RedirectedName.IsNull())
 	{
 		return FSoftObjectPath();
 	}
-	return FSoftObjectPath(RedirectedName, ObjectPath.GetSubPathString());
+	return FSoftObjectPath(RedirectedName.GetAssetPath(), ObjectPath.GetSubPathString());
 }
 
 void UAssetManager::ExtractSoftObjectPaths(const UStruct* Struct, const void* StructValue, TArray<FSoftObjectPath>& FoundAssetReferences, const TArray<FName>& PropertiesToSkip) const
@@ -2750,7 +2759,7 @@ bool UAssetManager::GetAssetDataForPath(const FSoftObjectPath& ObjectPath, FAsse
 	{
 		FString DestinationObjectPath = Result.GetValue();
 		ConstructorHelpers::StripObjectClass(DestinationObjectPath);
-		AssetData = AssetRegistry.GetAssetByObjectPath(*DestinationObjectPath);
+		AssetData = AssetRegistry.GetAssetByObjectPath(FSoftObjectPath(DestinationObjectPath));
 		Result = AssetData.TagsAndValues.FindTag("DestinationObject");
 	}
 
@@ -2783,23 +2792,23 @@ bool TryToSoftObjectPath(const FAssetData& AssetData, FSoftObjectPath& OutObject
 	else if (EndsWithBlueprint(AssetData.AssetClassPath))
 	{
 		TStringBuilder<256> AssetPath;
-		AssetPath << AssetData.ObjectPath << TEXT("_C");
+		AssetPath << AssetData.ToSoftObjectPath() << TEXT("_C");
 		OutObjectPath = FSoftObjectPath(FStringView(AssetPath));
 		return true;
 	}
 	else if (ScanForSubobject == EAssetDataCanBeSubobject::Yes)
 	{
-		OutObjectPath = FSoftObjectPath(AssetData.ObjectPath);
+		OutObjectPath = FSoftObjectPath(AssetData.GetSoftObjectPath());
 		return true;
 	}
 	else
 	{
-		if (ContainsSubobjectDelimiter(AssetData.ObjectPath))
+		if (!AssetData.IsTopLevelAsset())
 		{
 			OutObjectPath = FSoftObjectPath();
 			return false;
 		}
-		OutObjectPath = FSoftObjectPath(AssetData.ObjectPath, /* no subobject */ FString());
+		OutObjectPath = AssetData.GetSoftObjectPath();
 		return true;
 	}
 }
@@ -2822,7 +2831,7 @@ void UAssetManager::GetAssetDataForPathInternal(IAssetRegistry& AssetRegistry, c
 	{
 		// We need to strip the class suffix because the asset registry has it listed by blueprint name
 		
-		OutAssetData = AssetRegistry.GetAssetByObjectPath(FName(*AssetPath.LeftChop(2)), bIncludeOnlyOnDiskAssets);
+		OutAssetData = AssetRegistry.GetAssetByObjectPath(FSoftObjectPath(AssetPath.LeftChop(2)), bIncludeOnlyOnDiskAssets);
 
 		if (OutAssetData.IsValid())
 		{
@@ -2830,7 +2839,7 @@ void UAssetManager::GetAssetDataForPathInternal(IAssetRegistry& AssetRegistry, c
 		}
 	}
 
-	OutAssetData = AssetRegistry.GetAssetByObjectPath(FName(*AssetPath), bIncludeOnlyOnDiskAssets);
+	OutAssetData = AssetRegistry.GetAssetByObjectPath(FSoftObjectPath(AssetPath), bIncludeOnlyOnDiskAssets);
 }
 
 bool UAssetManager::WriteCustomReport(FString FileName, TArray<FString>& FileLines) const
@@ -3009,8 +3018,8 @@ void UAssetManager::DumpBundlesForAsset(const TArray<FString>& Args)
 	UE_LOG(LogAssetManager, Display, TEXT("Dumping bundles for primary asset %s..."), *PrimaryAssetIdString);
 	for (const FAssetBundleEntry& Entry : (**FoundMap).Bundles)
 	{
-		UE_LOG(LogAssetManager, Display, TEXT("  Bundle: %s (%d assets)"), *Entry.BundleName.ToString(), Entry.BundleAssets.Num());
-		for (const FSoftObjectPath& Path : Entry.BundleAssets)
+		UE_LOG(LogAssetManager, Display, TEXT("  Bundle: %s (%d assets)"), *Entry.BundleName.ToString(), Entry.AssetPaths.Num());
+		for (const FTopLevelAssetPath& Path : Entry.AssetPaths)
 		{
 			UE_LOG(LogAssetManager, Display, TEXT("    %s"), *Path.ToString());
 		}
@@ -3718,9 +3727,9 @@ void UAssetManager::UpdateManagementDatabase(bool bForceRefresh)
 			{
 				for (const FAssetBundleEntry& Entry : (**BundleMap).Bundles)
 				{
-					for (const FSoftObjectPath& BundleAssetRef : Entry.BundleAssets)
+					for (const FTopLevelAssetPath& BundleAssetRef : Entry.AssetPaths)
 					{
-						FName PackageName = FName(*BundleAssetRef.GetLongPackageName());
+						FName PackageName = BundleAssetRef.GetPackageName();
 
 						if (PackageName.IsNone())
 						{
@@ -3959,9 +3968,9 @@ void UAssetManager::ModifyCook(TArray<FName>& PackagesToCook, TArray<FName>& Pac
 			{
 				for (const FAssetBundleEntry& FoundEntry : FoundEntries)
 				{
-					for (const FSoftObjectPath& FoundReference : FoundEntry.BundleAssets)
+					for (const FTopLevelAssetPath& FoundReference : FoundEntry.AssetPaths)
 					{
-						FName PackageName = FName(*FoundReference.GetLongPackageName());
+						FName PackageName = FoundReference.GetPackageName();
 						AssetPackages.AddUnique(PackageName);
 					}
 				}
@@ -4382,7 +4391,7 @@ void UAssetManager::OnObjectPreSave(UObject* Object, FObjectPreSaveContext SaveC
 		return;
 	}
 
-	FPrimaryAssetId FoundPrimaryAssetId = GetPrimaryAssetIdForPath(*Object->GetPathName());
+	FPrimaryAssetId FoundPrimaryAssetId = GetPrimaryAssetIdForPath(FSoftObjectPath(Object));
 	if (FoundPrimaryAssetId.IsValid())
 	{
 		TSharedRef<FPrimaryAssetTypeData>* FoundType = AssetTypeMap.Find(FoundPrimaryAssetId.PrimaryAssetType);
@@ -4423,12 +4432,13 @@ void UAssetManager::OnAssetRemoved(const FAssetData& Data)
 {
 	// This could be much more efficient if UAssetManager broadcast one large event instead of all these tiny updates, see UAssetRegistryImpl::Broadcast
 
-	FPrimaryAssetId PrimaryAssetId = GetPrimaryAssetIdForPath(Data.ObjectPath);
+	FPrimaryAssetId PrimaryAssetId = GetPrimaryAssetIdForPath(Data.GetSoftObjectPath());
 
 	// This may be a blueprint, try with _C
 	if (!PrimaryAssetId.IsValid())
 	{
-		PrimaryAssetId = GetPrimaryAssetIdForPath(FName(FNameBuilder(Data.ObjectPath).Append(TEXT("_C")).ToView(), FNAME_Find));
+		FSoftObjectPath Path(WriteToString<FName::StringBufferSize>(Data.GetSoftObjectPath(), TEXT("_C")));
+		PrimaryAssetId = GetPrimaryAssetIdForPath(Path);
 	}
 
 	CachedAssetBundles.Remove(PrimaryAssetId);
@@ -4615,7 +4625,7 @@ void UAssetManager::InitializeAssetBundlesFromMetadata_Recursive(const UStruct*
 
 				for (const FName& BundleName : BundleSet)
 				{
-					AssetBundle.AddBundleAsset(BundleName, FoundRef);
+					AssetBundle.AddBundleAsset(BundleName, FoundRef.GetAssetPath());
 				}
 			}
 			else
diff --git a/Engine/Source/Runtime/Engine/Private/PrimaryAssetLabel.cpp b/Engine/Source/Runtime/Engine/Private/PrimaryAssetLabel.cpp
--- a/Engine/Source/Runtime/Engine/Private/PrimaryAssetLabel.cpp
+++ b/Engine/Source/Runtime/Engine/Private/PrimaryAssetLabel.cpp
@@ -46,7 +46,7 @@ void UPrimaryAssetLabel::UpdateAssetBundleData()
 		TArray<FAssetData> DirectoryAssets;
 		AssetRegistry.GetAssetsByPath(PackagePath, DirectoryAssets, true);
 
-		TArray<FSoftObjectPath> NewPaths;
+		TArray<FTopLevelAssetPath> NewPaths;
 
 		for (const FAssetData& AssetData : DirectoryAssets)
 		{
@@ -54,7 +54,7 @@ void UPrimaryAssetLabel::UpdateAssetBundleData()
 
 			if (!AssetRef.IsNull())
 			{
-				NewPaths.Add(AssetRef);
+				NewPaths.Add(AssetRef.GetAssetPath());
 			}
 		}
 
@@ -64,19 +64,19 @@ void UPrimaryAssetLabel::UpdateAssetBundleData()
 
 	if (AssetCollection.CollectionName != NAME_None)
 	{
-		TArray<FSoftObjectPath> NewPaths;
-		TArray<FName> CollectionAssets;
+		TArray<FTopLevelAssetPath> NewPaths;
+		TArray<FName> CollectionAssetPathNames;
 		ICollectionManager& CollectionManager = FCollectionManagerModule::GetModule().Get();
-		CollectionManager.GetAssetsInCollection(AssetCollection.CollectionName, ECollectionShareType::CST_All, CollectionAssets);
+		CollectionManager.GetAssetsInCollection(AssetCollection.CollectionName, ECollectionShareType::CST_All, CollectionAssetPathNames);
+		TArray<FSoftObjectPath> CollectionAssets = UE::SoftObjectPath::Private::ConvertObjectPathNames(CollectionAssetPathNames);
 		for (int32 Index = 0; Index < CollectionAssets.Num(); ++Index)
 		{
 			FAssetData FoundAsset = Manager.GetAssetRegistry().GetAssetByObjectPath(CollectionAssets[Index]);
-
 			FSoftObjectPath AssetRef = Manager.GetAssetPathForData(FoundAsset);
 
 			if (!AssetRef.IsNull())
 			{
-				NewPaths.Add(AssetRef);
+				NewPaths.Add(AssetRef.GetAssetPath());
 			}
 		}
 
