Fix UAssetManager::GetRedirectedAssetPath to not infinitely recurse.

diff --git a/Engine/Source/Runtime/Engine/Private/AssetManager.cpp b/Engine/Source/Runtime/Engine/Private/AssetManager.cpp
--- a/Engine/Source/Runtime/Engine/Private/AssetManager.cpp
+++ b/Engine/Source/Runtime/Engine/Private/AssetManager.cpp
@@ -2643,27 +2643,23 @@ void UAssetManager::GetPreviousPrimaryAssetIds(const FPrimaryAssetId& NewId, TAr
 FName UAssetManager::GetRedirectedAssetPath(FName OldPath) const
 {
 PRAGMA_DISABLE_DEPRECATION_WARNINGS
-	FSoftObjectPath Redirect = GetRedirectedAssetPath(FSoftObjectPath(OldPath));
-	return Redirect.ToFName();
+	const FSoftObjectPath* Redirected = AssetPathRedirects.Find(FSoftObjectPath(OldPath));
+	if (Redirected)
+	{
+		return Redirected->ToFName();
+	}
+	return NAME_None;
 PRAGMA_ENABLE_DEPRECATION_WARNINGS
 }
 
 FSoftObjectPath UAssetManager::GetRedirectedAssetPath(const FSoftObjectPath& ObjectPath) const
 {
-	FSoftObjectPath PossibleAssetPath = ObjectPath;
-
-	if (PossibleAssetPath.IsNull())
-	{
-		return FSoftObjectPath();
-	}
-
-	FSoftObjectPath RedirectedName = GetRedirectedAssetPath(PossibleAssetPath);
-
-	if (RedirectedName.IsNull())
+	const FSoftObjectPath* RedirectedName = AssetPathRedirects.Find(ObjectPath.GetWithoutSubPath());
+	if (!RedirectedName)
 	{
 		return FSoftObjectPath();
 	}
-	return FSoftObjectPath(RedirectedName.GetAssetPath(), ObjectPath.GetSubPathString());
+	return FSoftObjectPath(RedirectedName->GetAssetPath(), ObjectPath.GetSubPathString());
 }
 
 void UAssetManager::ExtractSoftObjectPaths(const UStruct* Struct, const void* StructValue, TArray<FSoftObjectPath>& FoundAssetReferences, const TArray<FName>& PropertiesToSkip) const
