World Partition: don't show the enable streaming dialog for level instances.

#rnx

diff --git a/Engine/Source/Runtime/Engine/Private/WorldPartition/WorldPartition.cpp b/Engine/Source/Runtime/Engine/Private/WorldPartition/WorldPartition.cpp
--- a/Engine/Source/Runtime/Engine/Private/WorldPartition/WorldPartition.cpp
+++ b/Engine/Source/Runtime/Engine/Private/WorldPartition/WorldPartition.cpp
@@ -1012,7 +1012,7 @@ void UWorldPartition::HashActorDesc(FWorldPartitionActorDesc* ActorDesc)
 	FWorldPartitionHandle ActorHandle(this, ActorDesc->GetGuid());
 	EditorHash->HashActor(ActorHandle);
 
-	bShouldCheckEnableStreamingWarning = true;
+	bShouldCheckEnableStreamingWarning = IsMainWorldPartition();
 }
 
 void UWorldPartition::UnhashActorDesc(FWorldPartitionActorDesc* ActorDesc)
