Disable HLOD Warmup for dedicated server. This was preventing server from unloading cells when server streaming is enabled.

diff --git a/Engine/Source/Runtime/Engine/Private/WorldPartition/HLOD/HLODSubsystem.cpp b/Engine/Source/Runtime/Engine/Private/WorldPartition/HLOD/HLODSubsystem.cpp
--- a/Engine/Source/Runtime/Engine/Private/WorldPartition/HLOD/HLODSubsystem.cpp
+++ b/Engine/Source/Runtime/Engine/Private/WorldPartition/HLOD/HLODSubsystem.cpp
@@ -407,7 +407,7 @@ void UHLODSubsystem::MakeRenderResourcesResident(const FCellData& CellData, cons
 bool UHLODSubsystem::RequestUnloading(const UWorldPartitionRuntimeCell* InCell)
 {
 	// Test if warmup is disabled globally.
-	bool bPerformWarmpup = CVarHLODWarmupEnabled.GetValueOnGameThread() != 0 && CVarHLODWarmupNumFrames.GetValueOnGameThread() > 0;
+	bool bPerformWarmpup = (CVarHLODWarmupEnabled.GetValueOnGameThread() != 0) && (CVarHLODWarmupNumFrames.GetValueOnGameThread() > 0) && (GetWorld()->GetNetMode() != NM_DedicatedServer);
 	if (!bPerformWarmpup)
 	{
 		return true;
