Make sure UWorldPartition::bEnableStreaming default value is correctly assigned in non editor builds

diff --git a/Engine/Source/Runtime/Engine/Private/WorldPartition/WorldPartition.cpp b/Engine/Source/Runtime/Engine/Private/WorldPartition/WorldPartition.cpp
--- a/Engine/Source/Runtime/Engine/Private/WorldPartition/WorldPartition.cpp
+++ b/Engine/Source/Runtime/Engine/Private/WorldPartition/WorldPartition.cpp
@@ -216,7 +216,6 @@ UWorldPartition::UWorldPartition(const FObjectInitializer& ObjectInitializer)
 	, bShouldEnableStreamingWarned(false)
 	, bShouldCheckEnableStreamingWarning(false)
 	, bCanBeUsedByLevelInstance(false)
-	, bEnableStreaming(true)
 	, bForceGarbageCollection(false)
 	, bForceGarbageCollectionPurge(false)
 	, bIsPIE(false)
@@ -226,6 +225,8 @@ UWorldPartition::UWorldPartition(const FObjectInitializer& ObjectInitializer)
 	, StreamingPolicy(nullptr)
 	, Replay(nullptr)
 {
+	bEnableStreaming = true;
+
 #if WITH_EDITOR
 	WorldPartitionStreamingPolicyClass = UWorldPartitionLevelStreamingPolicy::StaticClass();
 #endif
