Call FCoreDelegates::OnActorLabelChanged when a default label is created via AActor::GetActorLabel

#rnx

diff --git a/Engine/Source/Runtime/Engine/Classes/GameFramework/Actor.h b/Engine/Source/Runtime/Engine/Classes/GameFramework/Actor.h
--- a/Engine/Source/Runtime/Engine/Classes/GameFramework/Actor.h
+++ b/Engine/Source/Runtime/Engine/Classes/GameFramework/Actor.h
@@ -955,7 +955,7 @@ private:
 	 * and call AActor::SetActorLabel() or FActorLabelUtilities::SetActorLabelUnique() to change the label.  Never set the label directly.
 	 */
 	UPROPERTY()
-	mutable FString ActorLabel;
+	FString ActorLabel;
 #endif
 
 #if !WITH_EDITOR && ACTOR_HAS_LABELS
diff --git a/Engine/Source/Runtime/Engine/Private/ActorEditor.cpp b/Engine/Source/Runtime/Engine/Private/ActorEditor.cpp
--- a/Engine/Source/Runtime/Engine/Private/ActorEditor.cpp
+++ b/Engine/Source/Runtime/Engine/Private/ActorEditor.cpp
@@ -1109,7 +1109,9 @@ const FString& AActor::GetActorLabel(bool bCreateIfNone) const
 		// Remember, there could already be an actor with the same label in the level.  But that's OK, because
 		// actor labels aren't supposed to be unique.  We just try to make them unique initially to help
 		// disambiguate when opening up a new level and there are hundreds of actors of the same type.
-		ActorLabel = MoveTemp(DefaultActorLabel);
+		AActor* MutableThis = const_cast<AActor*>(this);
+		MutableThis->ActorLabel = MoveTemp(DefaultActorLabel);
+		FCoreDelegates::OnActorLabelChanged.Broadcast(MutableThis);
 	}
 
 	return ActorLabel;
