Compile fixes for non unity builds

diff --git a/Engine/Source/Runtime/Engine/Private/SceneTexturesConfig.cpp b/Engine/Source/Runtime/Engine/Private/SceneTexturesConfig.cpp
--- a/Engine/Source/Runtime/Engine/Private/SceneTexturesConfig.cpp
+++ b/Engine/Source/Runtime/Engine/Private/SceneTexturesConfig.cpp
@@ -3,6 +3,7 @@
 #include "SceneTexturesConfig.h"
 #include "RenderGraphResources.h"
 #include "ShaderCompiler.h"
+#include "SceneInterface.h"
 
 FSceneTexturesConfig FSceneTexturesConfig::GlobalInstance;
 
diff --git a/Engine/Source/Runtime/Engine/Public/PrimitiveSceneProxy.h b/Engine/Source/Runtime/Engine/Public/PrimitiveSceneProxy.h
--- a/Engine/Source/Runtime/Engine/Public/PrimitiveSceneProxy.h
+++ b/Engine/Source/Runtime/Engine/Public/PrimitiveSceneProxy.h
@@ -1481,4 +1481,4 @@ ENGINE_API extern bool SupportsNaniteRendering(const FVertexFactory* RESTRICT Ve
 
 ENGINE_API extern bool SupportsNaniteRendering(const FVertexFactory* RESTRICT VertexFactory, const FPrimitiveSceneProxy* RESTRICT PrimitiveSceneProxy, const class FMaterialRenderProxy* MaterialRenderProxy, ERHIFeatureLevel::Type FeatureLevel);
 
-ENGINE_API extern bool SupportsNaniteRendering(const FVertexFactoryType* RESTRICT VertexFactoryType, const class FMaterial& Material, ERHIFeatureLevel::Type FeatureLevel);
+ENGINE_API extern bool SupportsNaniteRendering(const class FVertexFactoryType* RESTRICT VertexFactoryType, const class FMaterial& Material, ERHIFeatureLevel::Type FeatureLevel);
