Missing include files for non unity builds

diff --git a/Engine/Source/Runtime/Engine/Private/PSOPrecache.cpp b/Engine/Source/Runtime/Engine/Private/PSOPrecache.cpp
--- a/Engine/Source/Runtime/Engine/Private/PSOPrecache.cpp
+++ b/Engine/Source/Runtime/Engine/Private/PSOPrecache.cpp
@@ -5,6 +5,8 @@
 =============================================================================*/
 
 #include "PSOPrecache.h"
+#include "SceneInterface.h"
+#include "VertexFactory.h"
 
 PSOCollectorCreateFunction FPSOCollectorCreateManager::JumpTable[(int32)EShadingPath::Num][FPSOCollectorCreateManager::MaxPSOCollectorCount] = {};
 
diff --git a/Engine/Source/Runtime/Engine/Private/SceneTexturesConfig.cpp b/Engine/Source/Runtime/Engine/Private/SceneTexturesConfig.cpp
--- a/Engine/Source/Runtime/Engine/Private/SceneTexturesConfig.cpp
+++ b/Engine/Source/Runtime/Engine/Private/SceneTexturesConfig.cpp
@@ -4,6 +4,8 @@
 #include "RenderGraphResources.h"
 #include "ShaderCompiler.h"
 #include "SceneInterface.h"
+#include "Engine/Engine.h"
+#include "StereoRendering.h"
 
 FSceneTexturesConfig FSceneTexturesConfig::GlobalInstance;
 
