Fix for sync groups not working correctly for zero-duration/inertial blends
* Contributed by [REDACTED_PERSON]

diff --git a/Engine/Source/Runtime/Engine/Classes/Animation/AnimationAsset.h b/Engine/Source/Runtime/Engine/Classes/Animation/AnimationAsset.h
--- a/Engine/Source/Runtime/Engine/Classes/Animation/AnimationAsset.h
+++ b/Engine/Source/Runtime/Engine/Classes/Animation/AnimationAsset.h
@@ -303,18 +303,18 @@ struct FMarkerSyncAnimPosition
 	GENERATED_USTRUCT_BODY()
 
 	/** The marker we have passed*/
-	UPROPERTY()
+	UPROPERTY(EditAnywhere, BlueprintReadWrite, Category = Sync)
 	FName PreviousMarkerName;
 
 	/** The marker we are heading towards */
-	UPROPERTY()
+	UPROPERTY(EditAnywhere, BlueprintReadWrite, Category = Sync)
 	FName NextMarkerName;
 
 	/** Value between 0 and 1 representing where we are:
 	0   we are at PreviousMarker
 	1   we are at NextMarker
 	0.5 we are half way between the two */
-	UPROPERTY()
+	UPROPERTY(EditAnywhere, BlueprintReadWrite, Category = Sync)
 	float PositionBetweenMarkers;
 
 	/** Is this a valid Marker Sync Position */
@@ -605,11 +605,19 @@ public:
 
 	FMarkerTickContext MarkerTickContext;
 
+	// Float in 0 - 1 range representing how far through an animation we were before ticking
+	float PreviousAnimLengthRatio;
+
+	// Float in 0 - 1 range representing how far through an animation we are
+	float AnimLengthRatio;
+
 public:
 	FAnimGroupInstance()
 		: GroupLeaderIndex(INDEX_NONE)
 		, bCanUseMarkerSync(false)
 		, MontageLeaderWeight(0.f)
+		, PreviousAnimLengthRatio(0.f)
+		, AnimLengthRatio(0.f)
 	{
 	}
 
@@ -620,6 +628,8 @@ public:
 		bCanUseMarkerSync = false;
 		MontageLeaderWeight = 0.f;
 		MarkerTickContext = FMarkerTickContext();
+		PreviousAnimLengthRatio = 0.f;
+		AnimLengthRatio = 0.f;
 	}
 
 	// Checks the last tick record in the ActivePlayers array to see if it's a better leader than the current candidate.
@@ -792,6 +802,7 @@ public:
 		, bIsMarkerPositionValid(ValidMarkerNames.Num() > 0)
 		, bIsLeader(true)
 		, bOnlyOneAnimationInGroup(bInOnlyOneAnimationInGroup)
+		, bIsJoiningSyncGroup(false)
 	{
 	}
 
@@ -804,6 +815,7 @@ public:
 		, bIsMarkerPositionValid(false)
 		, bIsLeader(true)
 		, bOnlyOneAnimationInGroup(bInOnlyOneAnimationInGroup)
+		, bIsJoiningSyncGroup(false)
 	{
 	}
 
@@ -844,17 +856,15 @@ public:
 		AnimLengthRatio = NormalizedTime;
 	}
 
-	// Returns the previous synchronization point (normalized time; only legal to call if ticking a follower)
+	// Returns the previous synchronization point (normalized time)
 	float GetPreviousAnimationPositionRatio() const
 	{
-		checkSlow(!bIsLeader);
 		return PreviousAnimLengthRatio;
 	}
 
-	// Returns the synchronization point (normalized time; only legal to call if ticking a follower)
+	// Returns the synchronization point (normalized time)
 	float GetAnimationPositionRatio() const
 	{
-		checkSlow(!bIsLeader);
 		return AnimLengthRatio;
 	}
 
@@ -883,6 +893,17 @@ public:
 		return bOnlyOneAnimationInGroup;
 	}
 
+	void SetJoiningSyncGroup(bool bInIsJoiningSyncGroup)
+	{
+		bIsJoiningSyncGroup = bInIsJoiningSyncGroup;
+	}
+
+	// Are we (re)joining the sync group this tick (eg: when initializing or resuming from zero weight)?
+	bool IsJoiningSyncGroup() const
+	{
+		return bIsJoiningSyncGroup;
+	}
+
 	//Root Motion accumulated from this tick context
 	FRootMotionMovementParams RootMotionMovementParams;
 
@@ -907,6 +928,9 @@ private:
 	bool bIsLeader;
 
 	bool bOnlyOneAnimationInGroup;
+
+	// True if the asset player being ticked was not previously part of the sync group (eg: it was inactive and has now reactivated)
+	bool bIsJoiningSyncGroup;
 };
 
 USTRUCT()
diff --git a/Engine/Source/Runtime/Engine/Private/Animation/AnimSequenceBase.cpp b/Engine/Source/Runtime/Engine/Private/Animation/AnimSequenceBase.cpp
--- a/Engine/Source/Runtime/Engine/Private/Animation/AnimSequenceBase.cpp
+++ b/Engine/Source/Runtime/Engine/Private/Animation/AnimSequenceBase.cpp
@@ -469,6 +469,12 @@ void UAnimSequenceBase::RemapTracksToNewSkeleton(USkeleton* NewSkeleton, bool bC
 void UAnimSequenceBase::TickAssetPlayer(FAnimTickRecord& Instance, struct FAnimNotifyQueue& NotifyQueue, FAnimAssetTickContext& Context) const
 {
 	float CurrentTime = *(Instance.TimeAccumulator);
+	if (Context.IsJoiningSyncGroup())
+	{
+		// Synchronize the asset player time to the other sync group members when (re)joining the group
+		CurrentTime = Context.GetAnimationPositionRatio() * GetPlayLength();
+	}
+
 	float PreviousTime = CurrentTime;
 	float DeltaTime = 0.f;
 
diff --git a/Engine/Source/Runtime/Engine/Private/Animation/AnimSync.cpp b/Engine/Source/Runtime/Engine/Private/Animation/AnimSync.cpp
--- a/Engine/Source/Runtime/Engine/Private/Animation/AnimSync.cpp
+++ b/Engine/Source/Runtime/Engine/Private/Animation/AnimSync.cpp
@@ -95,6 +95,10 @@ void FAnimSync::TickAssetPlayerInstances(FAnimInstanceProxy& InProxy, float InDe
 			FAnimAssetTickContext TickContext(InDeltaSeconds, RootMotionMode, bOnlyOneAnimationInGroup, SyncGroup.ValidMarkers);
 			if (PreviousGroup)
 			{
+				// Initialize the anim position ratio from the previous frame's final anim position ratio in case we're not using marker based sync
+				TickContext.SetPreviousAnimationPositionRatio(PreviousGroup->AnimLengthRatio);
+				TickContext.SetAnimationPositionRatio(PreviousGroup->AnimLengthRatio);
+
 				const FMarkerSyncAnimPosition& EndPosition = PreviousGroup->MarkerTickContext.GetMarkerSyncEndPosition();
 				if ( EndPosition.IsValid() &&
 						(EndPosition.PreviousMarkerName == NAME_None || SyncGroup.ValidMarkers.Contains(EndPosition.PreviousMarkerName)) &&
@@ -121,6 +125,22 @@ void FAnimSync::TickAssetPlayerInstances(FAnimInstanceProxy& InProxy, float InDe
 				// if it has leader score
 				SCOPE_CYCLE_COUNTER(STAT_TickAssetPlayerInstance);
 				FScopeCycleCounterUObject Scope(GroupLeader.SourceAsset);
+
+				// determine if we are (re)joining the sync group (ie: that we were not part of the sync group on the previous update)
+				bool bIsJoiningSyncGroup = true;
+				if (PreviousGroup)
+				{
+					for (const FAnimTickRecord& PreviousPlayer : PreviousGroup->ActivePlayers)
+					{
+						if (PreviousPlayer.SourceAsset == GroupLeader.SourceAsset && PreviousPlayer.TimeAccumulator == GroupLeader.TimeAccumulator)
+						{
+							bIsJoiningSyncGroup = false;
+							break;
+						}
+					}
+				}
+				TickContext.SetJoiningSyncGroup(bIsJoiningSyncGroup);
+
 				TickContext.MarkerTickContext.MarkersPassedThisTick.Reset();
 				TickContext.RootMotionMovementParams.Clear();
 				GroupLeader.SourceAsset->TickAssetPlayer(GroupLeader, InProxy.NotifyQueue, TickContext);
@@ -133,6 +153,8 @@ void FAnimSync::TickAssetPlayerInstances(FAnimInstanceProxy& InProxy, float InDe
 				// if we're not using marker based sync, we don't care, get out
 				if (TickContext.CanUseMarkerPosition() == false)
 				{
+					SyncGroup.PreviousAnimLengthRatio = TickContext.GetPreviousAnimationPositionRatio();
+					SyncGroup.AnimLengthRatio = TickContext.GetAnimationPositionRatio();
 					SyncGroup.GroupLeaderIndex = GroupLeaderIndex;
 					break;
 				}
@@ -140,6 +162,8 @@ void FAnimSync::TickAssetPlayerInstances(FAnimInstanceProxy& InProxy, float InDe
 				else if (TickContext.MarkerTickContext.IsMarkerSyncEndValid())
 				{
 					// if this leader contains correct position, break
+					SyncGroup.PreviousAnimLengthRatio = TickContext.GetPreviousAnimationPositionRatio();
+					SyncGroup.AnimLengthRatio = TickContext.GetAnimationPositionRatio();
 					SyncGroup.MarkerTickContext = TickContext.MarkerTickContext;
 					SyncGroup.GroupLeaderIndex = GroupLeaderIndex;
 					UE_LOG(LogAnimMarkerSync, Log, TEXT("Previous Sync Group Marker Tick Context :\n%s"), *SyncGroup.MarkerTickContext.ToString());
@@ -148,6 +172,8 @@ void FAnimSync::TickAssetPlayerInstances(FAnimInstanceProxy& InProxy, float InDe
 				}
 				else
 				{
+					SyncGroup.PreviousAnimLengthRatio = TickContext.GetPreviousAnimationPositionRatio();
+					SyncGroup.AnimLengthRatio = TickContext.GetAnimationPositionRatio();
 					SyncGroup.GroupLeaderIndex = GroupLeaderIndex;
 					UE_LOG(LogAnimMarkerSync, Log, TEXT("Invalid position from Leader %d. Trying next leader"), GroupLeaderIndex);
 				}
@@ -208,7 +234,8 @@ void FAnimSync::TickAssetPlayerInstances(FAnimInstanceProxy& InProxy, float InDe
 				// if we don't have a good leader, no reason to convert to follower
 				// tick as leader
 				TickContext.ConvertToFollower();
-	
+				TickContext.SetJoiningSyncGroup(false);
+
 				for (int32 TickIndex = GroupLeaderIndex + 1; TickIndex < SyncGroup.ActivePlayers.Num(); ++TickIndex)
 				{
 					FAnimTickRecord& AssetPlayer = SyncGroup.ActivePlayers[TickIndex];
@@ -333,6 +360,10 @@ FMarkerSyncAnimPosition FAnimSync::GetSyncGroupPosition(FName InSyncGroupName) c
 		{
 			return SyncGroupInstancePtr->MarkerTickContext.GetMarkerSyncEndPosition();
 		}
+		else
+		{
+			return FMarkerSyncAnimPosition(NAME_None, NAME_None, SyncGroupInstancePtr->AnimLengthRatio);
+		}
 	}
 
 	return FMarkerSyncAnimPosition();
diff --git a/Engine/Source/Runtime/Engine/Private/Animation/BlendSpace.cpp b/Engine/Source/Runtime/Engine/Private/Animation/BlendSpace.cpp
--- a/Engine/Source/Runtime/Engine/Private/Animation/BlendSpace.cpp
+++ b/Engine/Source/Runtime/Engine/Private/Animation/BlendSpace.cpp
@@ -418,6 +418,12 @@ void UBlendSpace::TickAssetPlayer(FAnimTickRecord& Instance, struct FAnimNotifyQ
 			}
 
 			float& NormalizedCurrentTime = *(Instance.TimeAccumulator);
+			if (Context.IsJoiningSyncGroup())
+			{
+				// Synchronize the asset player time to the other sync group members when (re)joining the group
+				NormalizedCurrentTime = Context.GetAnimationPositionRatio();
+			}
+
 			float NormalizedPreviousTime = NormalizedCurrentTime;
 
 			// @note for sync group vs non sync group
