- Allow changing wp.Runtime.EnableServerStreaming while running PIE
- Added wp.Runtime.DebugServerStreaming available to toggle between server/client when using wp debug draw commands in PIE
- Hide client only visible grids from debug display when hash doesn't consider them
- World Partition 2D debug draw now shows world's net mode and if server streaming is enabled
- Allow block on slow streaming if server streaming is enabled
- Don't filter AddToWorld levels on listen server (same as dedicated server)
- Listen server with server streaming enable now considers client-only visible cells

diff --git a/Engine/Source/Runtime/Engine/Private/WorldPartition/WorldPartition.cpp b/Engine/Source/Runtime/Engine/Private/WorldPartition/WorldPartition.cpp
--- a/Engine/Source/Runtime/Engine/Private/WorldPartition/WorldPartition.cpp
+++ b/Engine/Source/Runtime/Engine/Private/WorldPartition/WorldPartition.cpp
@@ -64,8 +64,21 @@ static FAutoConsoleVariableRef CVarLoadingRangeBugItGo(
 	TEXT("Loading range for BugItGo command."),
 	ECVF_Default
 );
+
+static bool GDebugDedicatedServerStreaming = false;
+static FAutoConsoleVariableRef CVarDebugDedicatedServerStreaming(
+	TEXT("wp.Runtime.DebugDedicatedServerStreaming"),
+	GDebugDedicatedServerStreaming,
+	TEXT("Turn on/off to debug of server streaming."));
+
 #endif //WITH_EDITOR
 
+static int32 GEnableServerStreaming = 0;
+static FAutoConsoleVariableRef CVarEnableServerStreaming(
+	TEXT("wp.Runtime.EnableServerStreaming"),
+	GEnableServerStreaming,
+	TEXT("Set to 1 to enable server streaming, set to 2 to only enable it in PIE."));
+
 #define LOCTEXT_NAMESPACE "WorldPartition"
 
 #if WITH_EDITOR
@@ -884,6 +897,28 @@ const TArray<FWorldPartitionStreamingSource>& UWorldPartition::GetStreamingSourc
 	return EmptyStreamingSources;
 }
 
+bool UWorldPartition::IsServer() const
+{
+	if (UWorld* OwningWorld = GetWorld())
+	{
+		const ENetMode NetMode = OwningWorld->GetNetMode();
+		return NetMode == NM_DedicatedServer || NetMode == NM_ListenServer;
+	}
+	return false;
+}
+
+bool UWorldPartition::IsServerStreamingEnabled() const
+{
+	switch (GEnableServerStreaming)
+	{
+	case 1:	return true;
+#if WITH_EDITOR
+	case 2: return bIsPIE;
+#endif
+	}
+	return false;
+}
+
 bool UWorldPartition::IsSimulating()
 {
 #if WITH_EDITOR
@@ -1252,26 +1287,41 @@ bool UWorldPartition::IsStreamingCompleted(EWorldPartitionRuntimeCellState Query
 	return false;
 }
 
-bool UWorldPartition::CanDrawRuntimeHash() const
+bool UWorldPartition::CanDebugDraw() const
 {
-	return GetWorld()->IsGameWorld() || (UWorldPartition::IsSimulating() && StreamingPolicy);
+#if WITH_EDITOR
+	const ENetMode NetMode = GetWorld()->GetNetMode();
+	const bool bIsDedicatedServer = (NetMode == NM_DedicatedServer);
+	const bool bIsClient = (NetMode == NM_Client);
+	if ((bIsDedicatedServer && !GDebugDedicatedServerStreaming) ||
+		(bIsClient && GDebugDedicatedServerStreaming))
+	{
+		return false;
+	}
+
+	if (!bIsDedicatedServer && UWorldPartition::IsSimulating() && StreamingPolicy)
+	{
+		return true;
+	}
+#endif
+	return GetWorld()->IsGameWorld();
 }
 
 FVector2D UWorldPartition::GetDrawRuntimeHash2DDesiredFootprint(const FVector2D& CanvasSize)
 {
-	check(CanDrawRuntimeHash());
+	check(CanDebugDraw());
 	return StreamingPolicy->GetDrawRuntimeHash2DDesiredFootprint(CanvasSize);
 }
 
-void UWorldPartition::DrawRuntimeHash2D(class UCanvas* Canvas, const FVector2D& PartitionCanvasSize, const FVector2D& Offset)
+bool UWorldPartition::DrawRuntimeHash2D(class UCanvas* Canvas, const FVector2D& PartitionCanvasSize, const FVector2D& Offset)
 {
-	check(CanDrawRuntimeHash());
-	StreamingPolicy->DrawRuntimeHash2D(Canvas, PartitionCanvasSize, Offset);
+	check(CanDebugDraw());
+	return StreamingPolicy->DrawRuntimeHash2D(Canvas, PartitionCanvasSize, Offset);
 }
 
 void UWorldPartition::DrawRuntimeHash3D()
 {
-	check(CanDrawRuntimeHash());
+	check(CanDebugDraw());
 	StreamingPolicy->DrawRuntimeHash3D();
 }
 
diff --git a/Engine/Source/Runtime/Engine/Private/WorldPartition/WorldPartitionRuntimeSpatialHash.cpp b/Engine/Source/Runtime/Engine/Private/WorldPartition/WorldPartitionRuntimeSpatialHash.cpp
--- a/Engine/Source/Runtime/Engine/Private/WorldPartition/WorldPartitionRuntimeSpatialHash.cpp
+++ b/Engine/Source/Runtime/Engine/Private/WorldPartition/WorldPartitionRuntimeSpatialHash.cpp
@@ -1290,8 +1290,10 @@ bool UWorldPartitionRuntimeSpatialHash::ShouldConsiderClientOnlyVisibleCells() c
 	const UWorld* World = GetWorld();
 	if (World->IsGameWorld())
 	{
-		ENetMode NetMode = World->GetNetMode();
-		if (NetMode == NM_DedicatedServer || NetMode == NM_ListenServer)
+		// Dedicated server & Listen server without server streaming won't consider client-only visible cells
+		const ENetMode NetMode = World->GetNetMode();
+		if ((NetMode == NM_DedicatedServer) ||
+			((NetMode == NM_ListenServer) && !GetOuterUWorldPartition()->IsServerStreamingEnabled()))
 		{
 			return false;
 		}
@@ -1543,14 +1545,14 @@ FVector2D UWorldPartitionRuntimeSpatialHash::GetDraw2DDesiredFootprint(const FVe
 	return FVector2D(CanvasSize.X * GetFilteredStreamingGrids().Num(), CanvasSize.Y);
 }
 
-void UWorldPartitionRuntimeSpatialHash::Draw2D(UCanvas* Canvas, const TArray<FWorldPartitionStreamingSource>& Sources, const FVector2D& PartitionCanvasSize, const FVector2D& Offset) const
+bool UWorldPartitionRuntimeSpatialHash::Draw2D(UCanvas* Canvas, const TArray<FWorldPartitionStreamingSource>& Sources, const FVector2D& PartitionCanvasSize, const FVector2D& Offset) const
 {
 	TRACE_CPUPROFILER_EVENT_SCOPE(UWorldPartitionRuntimeSpatialHash::Draw2D);
 
 	TArray<const FSpatialHashStreamingGrid*> FilteredStreamingGrids = GetFilteredStreamingGrids();
 	if (FilteredStreamingGrids.Num() == 0 || Sources.Num() == 0)
 	{
-		return;
+		return false;
 	}
 
 	UWorldPartition* WorldPartition = GetOuterUWorldPartition();
@@ -1621,6 +1623,8 @@ void UWorldPartitionRuntimeSpatialHash::Draw2D(UCanvas* Canvas, const TArray<FWo
 
 		++GridIndex;
 	}
+
+	return true;
 }
 
 void UWorldPartitionRuntimeSpatialHash::Draw3D(const TArray<FWorldPartitionStreamingSource>& Sources) const
@@ -1650,13 +1654,18 @@ bool UWorldPartitionRuntimeSpatialHash::ContainsRuntimeHash(const FString& Name)
 
 TArray<const FSpatialHashStreamingGrid*> UWorldPartitionRuntimeSpatialHash::GetFilteredStreamingGrids() const
 {
+	const bool bShouldConsiderClientOnlyVisible = ShouldConsiderClientOnlyVisibleCells();
+
 	TArray<const FSpatialHashStreamingGrid*> FilteredStreamingGrids;
 	FilteredStreamingGrids.Reserve(GetNumGrids());
-	ForEachStreamingGrid([&FilteredStreamingGrids](const FSpatialHashStreamingGrid& StreamingGrid)
+	ForEachStreamingGrid([&FilteredStreamingGrids, bShouldConsiderClientOnlyVisible](const FSpatialHashStreamingGrid& StreamingGrid)
 	{
-		if (FWorldPartitionDebugHelper::IsDebugRuntimeHashGridShown(StreamingGrid.GridName))
+		if (!StreamingGrid.bClientOnlyVisible || bShouldConsiderClientOnlyVisible)
 		{
-			FilteredStreamingGrids.Add(&StreamingGrid);
+			if (FWorldPartitionDebugHelper::IsDebugRuntimeHashGridShown(StreamingGrid.GridName))
+			{
+				FilteredStreamingGrids.Add(&StreamingGrid);
+			}
 		}
 	});
 
diff --git a/Engine/Source/Runtime/Engine/Private/WorldPartition/WorldPartitionStreamingPolicy.cpp b/Engine/Source/Runtime/Engine/Private/WorldPartition/WorldPartitionStreamingPolicy.cpp
--- a/Engine/Source/Runtime/Engine/Private/WorldPartition/WorldPartitionStreamingPolicy.cpp
+++ b/Engine/Source/Runtime/Engine/Private/WorldPartition/WorldPartitionStreamingPolicy.cpp
@@ -58,27 +58,12 @@ static FAutoConsoleVariableRef CVarBlockOnSlowStreaming(
 	GBlockOnSlowStreaming,
 	TEXT("Set if streaming needs to block when to slow to catchup."));
 
-static int32 GEnableServerStreaming = 0;
-static FAutoConsoleVariableRef CVarEnableServerStreaming(
-	TEXT("wp.Runtime.EnableServerStreaming"),
-	GEnableServerStreaming,
-	TEXT("Set to 1 to enable server streaming, set to 2 to only enable it in PIE."));
-
-static inline bool IsServerStreamingEnabled(bool bIsPIE)
-{
-	switch (GEnableServerStreaming)
-	{
-	case 1:	return true;
-	case 2:	return bIsPIE;
-	}
-	return false;
-}
-
 UWorldPartitionStreamingPolicy::UWorldPartitionStreamingPolicy(const FObjectInitializer& ObjectInitializer)
 : Super(ObjectInitializer) 
 , bCriticalPerformanceRequestedBlockTillOnWorld(false)
 , CriticalPerformanceBlockTillLevelStreamingCompletedEpoch(0)
 , DataLayersStatesServerEpoch(INT_MIN)
+, ServerStreamingEnabledEpoch(INT_MIN)
 , StreamingPerformance(EWorldPartitionStreamingPerformance::Good)
 #if !UE_BUILD_SHIPPING
 , OnScreenMessageStartTime(0.0)
@@ -111,17 +96,11 @@ void UWorldPartitionStreamingPolicy::UpdateStreamingSources()
 	const FTransform WorldToLocal = WorldPartition->GetInstanceTransform().Inverse();
 	UWorld* World = WorldPartition->GetWorld();
 
-#if WITH_EDITOR
-	const bool bIsPIE = WorldPartition->bIsPIE;
-#else
-	const bool bIsPIE = false;
-#endif
-	const bool bIsServerStreamingEnabled = IsServerStreamingEnabled(bIsPIE);
+	const bool bIsServer = WorldPartition->IsServer();
+	const bool bIsServerStreamingEnabled = WorldPartition->IsServerStreamingEnabled();
 
 	if (!AWorldPartitionReplay::IsPlaybackEnabled(World) || !WorldPartition->Replay->GetReplayStreamingSources(StreamingSources))
 	{
-		const ENetMode NetMode = World->GetNetMode();
-
 #if WITH_EDITOR
 		if (GEnableSimulationStreamingSource && UWorldPartition::IsSimulating())
 		{
@@ -133,13 +112,13 @@ void UWorldPartitionStreamingPolicy::UpdateStreamingSources()
 		}
 		else
 #endif
-		if (NetMode == NM_Standalone || NetMode == NM_Client || bIsServerStreamingEnabled || AWorldPartitionReplay::IsRecordingEnabled(World))
+		if (!bIsServer || bIsServerStreamingEnabled || AWorldPartitionReplay::IsRecordingEnabled(World))
 		{
 			for (FConstPlayerControllerIterator It = World->GetPlayerControllerIterator(); It; ++It)
 			{
 				if (APlayerController* PlayerController = It->Get())
 				{
-					if (NetMode == NM_DedicatedServer || NetMode == NM_ListenServer || PlayerController->IsLocalController())
+					if (bIsServer || PlayerController->IsLocalController())
 					{
 						if (PlayerController->IsStreamingSourceEnabled())
 						{
@@ -257,17 +236,18 @@ void UWorldPartitionStreamingPolicy::UpdateStreamingState()
 	check(ActivateStreamingCells.IsEmpty());
 	check(LoadStreamingCells.IsEmpty());
 
-	const ENetMode NetMode = World->GetNetMode();
-	const bool bClient = NetMode == NM_Standalone || NetMode == NM_Client;
-	const bool bServer = NetMode == NM_ListenServer || NetMode == NM_DedicatedServer;
-#if WITH_EDITOR
-	const bool bIsPIE = WorldPartition->bIsPIE;
-#else
-	const bool bIsPIE = false;
-#endif
-	const bool bIsServerStreamingEnabled = IsServerStreamingEnabled(bIsPIE);
-	
-	if (bClient || bIsServerStreamingEnabled || AWorldPartitionReplay::IsPlaybackEnabled(World))
+	ON_SCOPE_EXIT
+	{
+		// Reset frame StreamingSourceCells (optimization to avoid reallocation at every call to UpdateStreamingState)
+		FrameActivateCells.Reset();
+		FrameLoadCells.Reset();
+	};
+
+	const bool bIsServer = WorldPartition->IsServer();
+	const bool bIsServerStreamingEnabled = WorldPartition->IsServerStreamingEnabled();
+	const int32 NewServerStreamingEnabledEpoch = bIsServer ? (bIsServerStreamingEnabled ? 1 : 0) : INT_MIN;
+
+	if (!bIsServer || bIsServerStreamingEnabled || AWorldPartitionReplay::IsPlaybackEnabled(World))
 	{
 		// When world partition can't stream, all cells must be unloaded
 		if (WorldPartition->CanStream())
@@ -283,10 +263,12 @@ void UWorldPartitionStreamingPolicy::UpdateStreamingState()
 		}
 	}
 
-	if (bServer)
+	if (bIsServer)
 	{
 		// Server will activate all non data layer cells at first and then load/activate/unload data layer cells only when the data layer states change
-		if (!bIsServerStreamingEnabled && DataLayersStatesServerEpoch == AWorldDataLayers::GetDataLayersStateEpoch())
+		if (!bIsServerStreamingEnabled && 
+			(ServerStreamingEnabledEpoch == NewServerStreamingEnabledEpoch) &&
+			(DataLayersStatesServerEpoch == AWorldDataLayers::GetDataLayersStateEpoch()))
 		{
 			// Server as nothing to do early out
 			return; 
@@ -358,9 +340,9 @@ void UWorldPartitionStreamingPolicy::UpdateStreamingState()
 	}
 		
 	// Remove cells state changes based on client visibility. If any cell gets removed, delay epoch change.
-	auto ProcessClientVisibility = [bClient, &ClientVisibleLevelNames, &bUpdateEpoch](TSet<const UWorldPartitionRuntimeCell*>& Cells)
+	auto ProcessClientVisibility = [bIsServer, &ClientVisibleLevelNames, &bUpdateEpoch](TSet<const UWorldPartitionRuntimeCell*>& Cells)
 	{
-		if (!bClient)
+		if (bIsServer)
 		{
 			for (TSet<const UWorldPartitionRuntimeCell*>::TIterator It = Cells.CreateIterator(); It; ++It)
 			{
@@ -428,14 +410,11 @@ void UWorldPartitionStreamingPolicy::UpdateStreamingState()
 	// Evaluate streaming performance based on cells that should be activated
 	UpdateStreamingPerformance(ActivateStreamingCells);
 	
-	// Reset frame StreamingSourceCells (optimization to avoid reallocation at every call to UpdateStreamingState)
-	FrameActivateCells.Reset();
-	FrameLoadCells.Reset();
-
 	// Update Epoch if we aren't waiting for clients anymore
 	if (bUpdateEpoch)
 	{
 		DataLayersStatesServerEpoch = AWorldDataLayers::GetDataLayersStateEpoch();
+		ServerStreamingEnabledEpoch = NewServerStreamingEnabledEpoch;
 		UE_LOG(LogWorldPartition, Verbose, TEXT("Server epoch updated"));
 	}
 }
@@ -502,8 +481,12 @@ void UWorldPartitionStreamingPolicy::UpdateStreamingPerformance(const TSet<const
 	
 	if (StreamingPerformance == EWorldPartitionStreamingPerformance::Critical)
 	{
-		// This is a first very simple standalone implementation of handling of critical streaming conditions.
-		if (GBlockOnSlowStreaming && !IsInBlockTillLevelStreamingCompleted() && World->GetNetMode() == NM_Standalone)
+		const bool bIsServer = WorldPartition->IsServer();
+		const bool bIsServerStreamingEnabled = WorldPartition->IsServerStreamingEnabled();
+		const bool bCanBlockOnSlowStreaming = GBlockOnSlowStreaming && (!bIsServer || bIsServerStreamingEnabled);
+
+		// This is a very simple implementation of handling of critical streaming conditions.
+		if (bCanBlockOnSlowStreaming && !IsInBlockTillLevelStreamingCompleted())
 		{
 			World->bRequestedBlockOnAsyncLoading = true;
 			bCriticalPerformanceRequestedBlockTillOnWorld = true;
@@ -542,9 +525,7 @@ bool UWorldPartitionStreamingPolicy::ShouldSkipCellForPerformance(const UWorldPa
 	// When performance is degrading start skipping non blocking cells
 	if (!Cell->GetBlockOnSlowLoading())
 	{
-		UWorld* World = WorldPartition->GetWorld();
-		const ENetMode NetMode = World->GetNetMode();
-		if (NetMode == NM_Standalone || NetMode == NM_Client)
+		if (!WorldPartition->IsServer())
 		{
 			return IsInBlockTillLevelStreamingCompleted(/*bIsCausedByBadStreamingPerformance*/true);
 		}
@@ -555,9 +536,7 @@ bool UWorldPartitionStreamingPolicy::ShouldSkipCellForPerformance(const UWorldPa
 int32 UWorldPartitionStreamingPolicy::GetMaxCellsToLoad() const
 {
 	// This policy limits the number of concurrent loading streaming cells, except if match hasn't started
-	UWorld* World = WorldPartition->GetWorld();
-	const ENetMode NetMode = World->GetNetMode();
-	if (NetMode == NM_Standalone || NetMode == NM_Client)
+	if (!WorldPartition->IsServer())
 	{
 		return !IsInBlockTillLevelStreamingCompleted() ? (GMaxLoadingStreamingCells - GetCellLoadingCount()) : MAX_int32;
 	}
@@ -685,10 +664,9 @@ bool UWorldPartitionStreamingPolicy::CanAddLoadedLevelToWorld(ULevel* InLevel) c
 	TRACE_CPUPROFILER_EVENT_SCOPE(UWorldPartitionStreamingPolicy::CanAddLoadedLevelToWorld);
 
 	check(WorldPartition->IsInitialized());
-	UWorld* World = WorldPartition->GetWorld();
 
-	// Always allow AddToWorld in DedicatedServer
-	if (World->GetNetMode() == NM_DedicatedServer)
+	// Always allow AddToWorld in Dedicated server and Listen server
+	if (WorldPartition->IsServer())
 	{
 		return true;
 	}
@@ -766,20 +744,18 @@ FVector2D UWorldPartitionStreamingPolicy::GetDrawRuntimeHash2DDesiredFootprint(c
 	return FVector2D::ZeroVector;
 }
 
-void UWorldPartitionStreamingPolicy::DrawRuntimeHash2D(class UCanvas* Canvas, const FVector2D& PartitionCanvasSize, const FVector2D& Offset)
+bool UWorldPartitionStreamingPolicy::DrawRuntimeHash2D(class UCanvas* Canvas, const FVector2D& PartitionCanvasSize, const FVector2D& Offset)
 {
 	if (StreamingSources.Num() > 0 && WorldPartition->RuntimeHash)
 	{
-		WorldPartition->RuntimeHash->Draw2D(Canvas, StreamingSources, PartitionCanvasSize, Offset);
+		return WorldPartition->RuntimeHash->Draw2D(Canvas, StreamingSources, PartitionCanvasSize, Offset);
 	}
+	return false;
 }
 
 void UWorldPartitionStreamingPolicy::DrawRuntimeHash3D()
 {
-	UWorld* World = WorldPartition->GetWorld();
-	if (World->GetNetMode() != NM_DedicatedServer
-		&& WorldPartition->IsInitialized()
-		&& WorldPartition->RuntimeHash)
+	if (WorldPartition->IsInitialized() && WorldPartition->RuntimeHash)
 	{
 		WorldPartition->RuntimeHash->Draw3D(StreamingSources);
 	}
diff --git a/Engine/Source/Runtime/Engine/Private/WorldPartition/WorldPartitionSubsystem.cpp b/Engine/Source/Runtime/Engine/Private/WorldPartition/WorldPartitionSubsystem.cpp
--- a/Engine/Source/Runtime/Engine/Private/WorldPartition/WorldPartitionSubsystem.cpp
+++ b/Engine/Source/Runtime/Engine/Private/WorldPartition/WorldPartitionSubsystem.cpp
@@ -176,10 +176,7 @@ void UWorldPartitionSubsystem::OnWorldPartitionInitialized(UWorldPartition* InWo
 {
 	if (RegisteredWorldPartitions.IsEmpty())
 	{
-		if (InWorldPartition->CanDrawRuntimeHash() && (GetWorld()->GetNetMode() != NM_DedicatedServer))
-		{
-			DrawHandle = UDebugDrawService::Register(TEXT("Game"), FDebugDrawDelegate::CreateUObject(this, &UWorldPartitionSubsystem::Draw));
-		}
+		DrawHandle = UDebugDrawService::Register(TEXT("Game"), FDebugDrawDelegate::CreateUObject(this, &UWorldPartitionSubsystem::Draw));
 
 		// Enforce some GC settings when using World Partition
 		if (GetWorld()->IsGameWorld())
@@ -236,7 +233,7 @@ void UWorldPartitionSubsystem::Tick(float DeltaSeconds)
 	{
 		RegisteredWorldPartition->Tick(DeltaSeconds);
 
-		if (GDrawRuntimeHash3D && RegisteredWorldPartition->CanDrawRuntimeHash())
+		if (GDrawRuntimeHash3D && RegisteredWorldPartition->CanDebugDraw())
 		{
 			RegisteredWorldPartition->DrawRuntimeHash3D();
 		}
@@ -305,14 +302,20 @@ void UWorldPartitionSubsystem::Draw(UCanvas* Canvas, class APlayerController* PC
 		return;
 	}
 
+	UWorldPartition* WorldPartition = RegisteredWorldPartitions[GDrawWorldPartitionIndex];
+	if (!WorldPartition->CanDebugDraw())
+	{
+		return;
+	}
+
 	// Filter out views that don't match our world
-	if (Canvas->SceneView->ViewActor != nullptr && 
+	if (!WorldPartition->IsServer() &&
+		Canvas->SceneView->ViewActor != nullptr &&
 		Canvas->SceneView->ViewActor->GetWorld() != GetWorld())
 	{
 		return;
 	}
 
-	UWorldPartition* WorldPartition = RegisteredWorldPartitions[GDrawWorldPartitionIndex];
 	const FVector2D CanvasTopLeftPadding(10.f, 10.f);
 
 	FVector2D CurrentOffset(CanvasTopLeftPadding);
@@ -325,9 +328,11 @@ void UWorldPartitionSubsystem::Draw(UCanvas* Canvas, class APlayerController* PC
 		const FVector2D CanvasMaxScreenSize = FVector2D::Max(MaxScreenRatio*FVector2D(Canvas->ClipX, Canvas->ClipY) - CanvasBottomRightPadding - CurrentOffset, CanvasMinimumSize);
 
 		FVector2D PartitionCanvasSize = FVector2D(CanvasMaxScreenSize.X, CanvasMaxScreenSize.Y);
-		WorldPartition->DrawRuntimeHash2D(Canvas, PartitionCanvasSize, CurrentOffset);
-		CurrentOffset.X = CanvasBottomRightPadding.X;
-		CurrentOffset.Y += CanvasMaxScreenSize.Y;
+		if (WorldPartition->DrawRuntimeHash2D(Canvas, PartitionCanvasSize, CurrentOffset))
+		{
+			CurrentOffset.X = CanvasBottomRightPadding.X;
+			CurrentOffset.Y += CanvasMaxScreenSize.Y;
+		}
 	}
 	
 	if (GDrawStreamingPerfs || GDrawRuntimeHash2D)
@@ -337,12 +342,12 @@ void UWorldPartitionSubsystem::Draw(UCanvas* Canvas, class APlayerController* PC
 			if (IsIncrementalPurgePending()) { StatusText += TEXT("(Purging) "); }
 			if (IsIncrementalUnhashPending()) { StatusText += TEXT("(Unhashing) "); }
 			if (IsAsyncLoading()) { StatusText += TEXT("(AsyncLoading) "); }
-			if (StatusText.IsEmpty()) { StatusText = TEXT("(Idle)"); }
+			if (StatusText.IsEmpty()) { StatusText = TEXT("(Idle) "); }
 
-			const ENetMode NetMode = GetWorld()->GetNetMode();
-			if (NetMode == NM_ListenServer)
+			StatusText += FString::Printf(TEXT("(%s) "), *GetDebugStringForWorld(GetWorld()));
+			if (WorldPartition->IsServer())
 			{
-				StatusText += TEXT("(ListenServer)");
+				StatusText += FString::Printf(TEXT("(Server Streaming %s)"), WorldPartition->IsServerStreamingEnabled() ? TEXT("Enabled") : TEXT("Disabled"));
 			}
 			
 			const FString Text = FString::Printf(TEXT("Streaming Status: %s"), *StatusText);
diff --git a/Engine/Source/Runtime/Engine/Public/WorldPartition/WorldPartition.h b/Engine/Source/Runtime/Engine/Public/WorldPartition/WorldPartition.h
--- a/Engine/Source/Runtime/Engine/Public/WorldPartition/WorldPartition.h
+++ b/Engine/Source/Runtime/Engine/Public/WorldPartition/WorldPartition.h
@@ -212,6 +212,8 @@ public:
 	bool SupportsStreaming() const;
 	bool IsStreamingEnabled() const;
 	bool CanStream() const;
+	bool IsServer() const;
+	bool IsServerStreamingEnabled() const;
 
 	bool IsMainWorldPartition() const;
 
@@ -223,9 +225,9 @@ public:
 	const TArray<FWorldPartitionStreamingSource>& GetStreamingSources() const;
 
 	// Debugging
-	bool CanDrawRuntimeHash() const;
+	bool CanDebugDraw() const;
 	FVector2D GetDrawRuntimeHash2DDesiredFootprint(const FVector2D& CanvasSize);
-	void DrawRuntimeHash2D(UCanvas* Canvas, const FVector2D& PartitionCanvasSize, const FVector2D& Offset);
+	bool DrawRuntimeHash2D(UCanvas* Canvas, const FVector2D& PartitionCanvasSize, const FVector2D& Offset);
 	void DrawRuntimeHash3D();
 	void DrawRuntimeCellsDetails(UCanvas* Canvas, FVector2D& Offset);
 	void DrawStreamingStatusLegend(UCanvas* Canvas, FVector2D& Offset);
diff --git a/Engine/Source/Runtime/Engine/Public/WorldPartition/WorldPartitionRuntimeHash.h b/Engine/Source/Runtime/Engine/Public/WorldPartition/WorldPartitionRuntimeHash.h
--- a/Engine/Source/Runtime/Engine/Public/WorldPartition/WorldPartitionRuntimeHash.h
+++ b/Engine/Source/Runtime/Engine/Public/WorldPartition/WorldPartitionRuntimeHash.h
@@ -90,7 +90,7 @@ class ENGINE_API UWorldPartitionRuntimeHash : public UObject
 	 */
 	virtual FVector2D GetDraw2DDesiredFootprint(const FVector2D& CanvasSize) const { return FVector2D::ZeroVector; }
 
-	virtual void Draw2D(class UCanvas* Canvas, const TArray<FWorldPartitionStreamingSource>& Sources, const FVector2D& PartitionCanvasSize, const FVector2D& Offset) const {}
+	virtual bool Draw2D(class UCanvas* Canvas, const TArray<FWorldPartitionStreamingSource>& Sources, const FVector2D& PartitionCanvasSize, const FVector2D& Offset) const { return false; }
 	virtual void Draw3D(const TArray<FWorldPartitionStreamingSource>& Sources) const {}
 	virtual bool ContainsRuntimeHash(const FString& Name) const { return false; }
 
diff --git a/Engine/Source/Runtime/Engine/Public/WorldPartition/WorldPartitionRuntimeSpatialHash.h b/Engine/Source/Runtime/Engine/Public/WorldPartition/WorldPartitionRuntimeSpatialHash.h
--- a/Engine/Source/Runtime/Engine/Public/WorldPartition/WorldPartitionRuntimeSpatialHash.h
+++ b/Engine/Source/Runtime/Engine/Public/WorldPartition/WorldPartitionRuntimeSpatialHash.h
@@ -315,7 +315,7 @@ private:
 	TArray<TWeakObjectPtr<URuntimeSpatialHashExternalStreamingObject>> ExternalStreamingObjects;
 	
 	virtual FVector2D GetDraw2DDesiredFootprint(const FVector2D& CanvasSize) const override;
-	virtual void Draw2D(class UCanvas* Canvas, const TArray<FWorldPartitionStreamingSource>& Sources, const FVector2D& PartitionCanvasSize, const FVector2D& Offset) const override;
+	virtual bool Draw2D(class UCanvas* Canvas, const TArray<FWorldPartitionStreamingSource>& Sources, const FVector2D& PartitionCanvasSize, const FVector2D& Offset) const override;
 	virtual void Draw3D(const TArray<FWorldPartitionStreamingSource>& Sources) const override;
 	virtual bool ContainsRuntimeHash(const FString& Name) const override;
 
diff --git a/Engine/Source/Runtime/Engine/Public/WorldPartition/WorldPartitionStreamingPolicy.h b/Engine/Source/Runtime/Engine/Public/WorldPartition/WorldPartitionStreamingPolicy.h
--- a/Engine/Source/Runtime/Engine/Public/WorldPartition/WorldPartitionStreamingPolicy.h
+++ b/Engine/Source/Runtime/Engine/Public/WorldPartition/WorldPartitionStreamingPolicy.h
@@ -46,7 +46,7 @@ public:
 	virtual void SetTargetStateForCells(EWorldPartitionRuntimeCellState TargetState, const TSet<const UWorldPartitionRuntimeCell*>& Cells);
 	virtual bool CanAddLoadedLevelToWorld(class ULevel* InLevel) const;
 	virtual FVector2D GetDrawRuntimeHash2DDesiredFootprint(const FVector2D& CanvasSize);
-	virtual void DrawRuntimeHash2D(class UCanvas* Canvas, const FVector2D& PartitionCanvasSize, const FVector2D& Offset);
+	virtual bool DrawRuntimeHash2D(class UCanvas* Canvas, const FVector2D& PartitionCanvasSize, const FVector2D& Offset);
 	virtual void DrawRuntimeHash3D();
 	virtual void DrawRuntimeCellsDetails(class UCanvas* Canvas, FVector2D& Offset) {}
 	virtual void DrawStreamingStatusLegend(class UCanvas* Canvas, FVector2D& Offset) {}
@@ -98,6 +98,7 @@ protected:
 	mutable TArray<const UWorldPartitionRuntimeCell*, TInlineAllocator<256>> SortedAddToWorldCells;
 
 	int32 DataLayersStatesServerEpoch;
+	int32 ServerStreamingEnabledEpoch;
 
 	EWorldPartitionStreamingPerformance StreamingPerformance;
 #if !UE_BUILD_SHIPPING
