OnRep replication:
* Fix ReplicatedUsing functions not receiving the proper old version of a dynamic array
* Now OnRep functions tagged on an array that set a parameter will behave the same as other non-array OnRep with param functions.
   * So the previous version of the array passed in parameter will always be the version the local object had before we set the new network properties on it.
   * Before this fix it was a mix of cdo and last net-serialized entries.

diff --git a/Engine/Source/Runtime/Engine/Private/RepLayout.cpp b/Engine/Source/Runtime/Engine/Private/RepLayout.cpp
--- a/Engine/Source/Runtime/Engine/Private/RepLayout.cpp
+++ b/Engine/Source/Runtime/Engine/Private/RepLayout.cpp
@@ -3237,6 +3237,7 @@ static bool ReceivePropertyHelper(
 	FRepShadowDataBuffer ShadowData,
 	FRepObjectDataBuffer Data,
 	TArray<FProperty*>* RepNotifies,
+	const bool bShadowDataCopied,
 	const TArray<FRepParentCmd>& Parents,
 	const TArray<FRepLayoutCmd>& Cmds,
 	const int32 CmdIndex,
@@ -3284,7 +3285,10 @@ static bool ReceivePropertyHelper(
 	if (RepNotifies != nullptr && INDEX_NONE != Parent.RepNotifyNumParams)
 	{
 		// Copy current value over so we can check to see if it changed
-		StoreProperty(Cmd, ShadowData + Cmd, Data + SwappedCmd);
+		if (!bShadowDataCopied)
+		{
+			StoreProperty(Cmd, ShadowData + Cmd, Data + SwappedCmd);
+		}
 
 		// Read the property
 		Cmd.Property->NetSerializeItem(Bunch, Bunch.PackageMap, Data + SwappedCmd);
@@ -3382,7 +3386,8 @@ static FGuidReferencesMap* PrepReceivedArray(
 	const int32 CmdIndex,
 	FRepShadowDataBuffer* OutShadowBaseData,
 	FRepObjectDataBuffer* OutBaseData,
-	TArray<FProperty*>* RepNotifies)
+	TArray<FProperty*>* RepNotifies,
+	bool& bOutShadowDataCopied)
 {
 	FGuidReferences* NewGuidReferencesArray = nullptr;
 
@@ -3405,7 +3410,7 @@ static FGuidReferencesMap* PrepReceivedArray(
 		check(NewGuidReferencesArray->CmdIndex == CmdIndex);
 	}
 
-	if (RepNotifies != nullptr && INDEX_NONE != Parent.RepNotifyNumParams)
+	if (RepNotifies && Parent.RepNotifyNumParams != INDEX_NONE)
 	{
 		if (DataArray->Num() != ArrayNum || Parent.RepNotifyCondition == REPNOTIFY_Always)
 		{
@@ -3413,42 +3418,27 @@ static FGuidReferencesMap* PrepReceivedArray(
 		}
 		else
 		{
-			UE_CLOG(LogSkippedRepNotifies> 0, LogRep, Display, TEXT("1 FReceivedPropertiesStackState Skipping RepNotify for property %s because local value has not changed."), *Cmd.Property->GetName());
+			UE_CLOG(LogSkippedRepNotifies > 0, LogRep, Display, TEXT("1 FReceivedPropertiesStackState Skipping RepNotify for property %s because local value has not changed."), *Cmd.Property->GetName());
 		}
 		check(ShadowArray != nullptr);
 
-		// This should only happen if
-		//	A. Clients have modified the array locally.
-		//	B. The Archetype has different values in its array than the actual object.
-		//		This can happen if we have an instance of an Actor / Object in a level,
-		//		the array property is EditAnywhere,
-		//		and items have been added to or removed from the array. 
-		//
-		//	In the case of b, it's possible that our ShadowArray has never allocated its dynamic
-		//	memory, and that will cause us problems detecting changes to items in the data array.
-		//	As a stopgap, we will resize the shadow array to match the size of the current,
-		//	pre-replicated size of the array.
-		//
-		//	We don't need to resize the ShadowArray to the newly replicated size of the array,
-		//	because if we detect changes (either in the size, like above, or in any elements),
-		//	then during CallRepNotifies, we will copy the entire state of the DataArray (live array) to the
-		//	ShadowArray.
-		//
-		//	This still leaves us with a bit of a problem, however.
-		//	If the Archetype's array is the same size as the instance's array, but the values are different,
-		//	AND the user has defined an OnRep function that takes the Old Array as a parameter,
-		//	when the OnRep is called, it will have correct "old" values for any of the properties that
-		//	were replicated (because we will copy them prior to applying the replicated data),
-		//	but any properties that were NOT replicated will instead have the Archetype's values.
-		//
-		//	This would only happen the first time the OnRep was called, however, because we will
-		//	copy the complete value of the array after every OnRep (see CallRepNotifies).
-		//
-		// Note, this is behind a CVar because it changes behavior, and that could cause subtle bugs
-		// in game logic.
-		if (GbWithArrayOnRepFix && ShadowArray->Num() != DataArray->Num())
-		{
-			Cmd.Property->CopyCompleteValue((uint8*)ShadowArray, (uint8*)DataArray);
+		// If a top level property already set the current data in ShadowBuffer, we don't need to redo it again.
+		if (!bOutShadowDataCopied)
+		{
+			// Does the OnRep function have a parameter to receive the previous version of the array
+			if (Parent.RepNotifyNumParams > 0)
+			{
+				// Copy the entire array into the shadow buffer before it gets overwritten by the network data.
+				// The OnRep callback will pass that array back as a function parameter and it needs to be the current local array before the network data was applied.
+				Cmd.Property->CopyCompleteValue((uint8*)ShadowArray, (uint8*)DataArray);
+				bOutShadowDataCopied = true;
+			}
+			else if (ShadowArray->Num() != DataArray->Num())
+			{
+				// When individual entries get netserialized, they will copy over the current entry into the shadow buffer so ensure the array has the size to do so. 
+				FScriptArrayHelper ShadowArrayHelper((FArrayProperty*)Cmd.Property, ShadowArray);
+				ShadowArrayHelper.Resize(DataArray->Num());
+			}
 		}
 
 		*OutShadowBaseData = ShadowArray->GetData();
@@ -3501,6 +3491,7 @@ struct FReceivePropertiesStackParams
 	TArray<FProperty*>* RepNotifies;
 	uint32 ArrayElementOffset = 0;
 	uint16 CurrentHandle = 0;
+    bool bShadowDataCopied = false;
 };
 
 static FORCEINLINE void ReadPropertyHandle(FReceivePropertiesSharedParams& Params)
@@ -3584,7 +3575,10 @@ static bool ReceiveProperties_r(FReceivePropertiesSharedParams& Params, FReceive
 					nullptr,
 					CmdIndex + 1,
 					Cmd.EndCmd - 1,
-					StackParams.RepNotifies
+					StackParams.RepNotifies,
+					0 /*ArrayElementOffset*/,
+					0 /*CurrentHandle*/,
+					StackParams.bShadowDataCopied
 				};
 
 				// These buffers will track the dynamic array memory.
@@ -3610,7 +3604,8 @@ static bool ReceiveProperties_r(FReceivePropertiesSharedParams& Params, FReceive
 					CmdIndex,
 					&ShadowArrayBuffer,
 					&ObjectArrayBuffer,
-					StackParams.RepNotifies);
+					StackParams.RepNotifies,
+					ArrayStackParams.bShadowDataCopied);
 
 				// Read the next array handle.
 				ReadPropertyHandle(Params);
@@ -3663,6 +3658,7 @@ static bool ReceiveProperties_r(FReceivePropertiesSharedParams& Params, FReceive
 					StackParams.ShadowData,
 					StackParams.ObjectData,
 					StackParams.RepNotifies,
+					StackParams.bShadowDataCopied,
 					Params.Parents,
 					Params.Cmds,
 					CmdIndex,
@@ -3968,6 +3964,8 @@ bool FRepLayout::ReceiveProperties_BackwardsCompatible_r(
 			FRepObjectDataBuffer LocalData = Data;
 			FRepShadowDataBuffer LocalShadowData = ShadowData;
 
+			bool bShadowDataCopied = false;
+
 			FGuidReferencesMap* NewGuidReferencesArray = PrepReceivedArray(
 				ArrayNum,
 				ShadowArray,
@@ -3979,7 +3977,9 @@ bool FRepLayout::ReceiveProperties_BackwardsCompatible_r(
 				CmdIndex,
 				&LocalShadowData,
 				&LocalData,
-				ShadowData ? &RepState->RepNotifies : nullptr);
+				ShadowData ? &RepState->RepNotifies : nullptr,
+				bShadowDataCopied
+				);
 
 			// Read until we read all array elements
 			while (true)
@@ -4061,6 +4061,7 @@ bool FRepLayout::ReceiveProperties_BackwardsCompatible_r(
 				ShadowData,
 				Data,
 				ShadowData ? &RepState->RepNotifies : nullptr,
+				false /*bShadowDataCopied*/,
 				Parents,
 				Cmds,
 				CmdIndex,
