Fixing up a bunch o tool tips and comments
Removed UE4 references

diff --git a/Engine/Plugins/Runtime/Steam/SteamAudio/Source/SteamAudio/Private/SteamAudioEnvironment.h b/Engine/Plugins/Runtime/Steam/SteamAudio/Source/SteamAudio/Private/SteamAudioEnvironment.h
--- a/Engine/Plugins/Runtime/Steam/SteamAudio/Source/SteamAudio/Private/SteamAudioEnvironment.h
+++ b/Engine/Plugins/Runtime/Steam/SteamAudio/Source/SteamAudio/Private/SteamAudioEnvironment.h
@@ -12,7 +12,7 @@
 namespace SteamAudio
 {
 	/**
-	 * Maps from UE4 FName identifiers stored on AudioComponents to a unique integer for the Phonon API.
+	 * Maps from UE FName identifiers stored on AudioComponents to a unique integer for the Phonon API.
 	 * The expected number of baked sources for a given scene is quite low (< 1k usually), so this simply
 	 * maps to random integers in (0, INT_MAX).
 	 */
diff --git a/Engine/Source/Runtime/AudioMixer/Classes/SubmixEffects/AudioMixerSubmixEffectReverb.h b/Engine/Source/Runtime/AudioMixer/Classes/SubmixEffects/AudioMixerSubmixEffectReverb.h
--- a/Engine/Source/Runtime/AudioMixer/Classes/SubmixEffects/AudioMixerSubmixEffectReverb.h
+++ b/Engine/Source/Runtime/AudioMixer/Classes/SubmixEffects/AudioMixerSubmixEffectReverb.h
@@ -65,7 +65,7 @@ struct AUDIOMIXER_API FSubmixEffectReverbSettings
 	UPROPERTY(EditAnywhere, BlueprintReadWrite, Category = LateReflections, meta = (ClampMin = "0.0", ClampMax = "1.0", EditCondition = "!bBypass && !bBypassLateReflections"))
 	float AirAbsorptionGainHF;
 
-	/** Decay High Frequency Ratio - 0.1 < 0.83 < 2.0 - how much the quicker or slower the high frequencies decay relative to the lower frequencies. */
+	/** Decay High Frequency Ratio - 0.1 < 0.83 < 2.0 - how much quicker or slower the high frequencies decay relative to the lower frequencies. */
 	UPROPERTY(EditAnywhere, BlueprintReadWrite, Category = LateReflections, meta = (ClampMin = "0.1", ClampMax = "2.0", EditCondition = "!bBypass && !bBypassLateReflections"))
 	float DecayHFRatio;
 
diff --git a/Engine/Source/Runtime/Engine/Classes/Sound/ReverbSettings.h b/Engine/Source/Runtime/Engine/Classes/Sound/ReverbSettings.h
--- a/Engine/Source/Runtime/Engine/Classes/Sound/ReverbSettings.h
+++ b/Engine/Source/Runtime/Engine/Classes/Sound/ReverbSettings.h
@@ -67,7 +67,7 @@ struct ENGINE_API FReverbSettings
 	UPROPERTY(EditAnywhere, BlueprintReadWrite, Category = ReverbSettings)
 	TObjectPtr<USoundEffectSubmixPreset> ReverbPluginEffect;
 
-	/** Volume level of the reverb affect. */
+	/** Volume level of the reverb effect. */
 	UPROPERTY(EditAnywhere, BlueprintReadWrite, Category = ReverbSettings)
 	float Volume;
 
diff --git a/Engine/Source/Runtime/Engine/Classes/Sound/SoundSubmix.h b/Engine/Source/Runtime/Engine/Classes/Sound/SoundSubmix.h
--- a/Engine/Source/Runtime/Engine/Classes/Sound/SoundSubmix.h
+++ b/Engine/Source/Runtime/Engine/Classes/Sound/SoundSubmix.h
@@ -80,7 +80,7 @@ enum class EFFTWindowType : uint8
 	// No window is applied. Technically a boxcar window.
 	None,
 
-	// Mainlobe width of -3 dB and sidelove attenuation of ~-40 dB. Good for COLA.
+	// Mainlobe width of -3 dB and sidelobe attenuation of ~-40 dB. Good for COLA.
 	Hamming,
 
 	// Mainlobe width of -3 dB and sidelobe attenuation of ~-30dB. Good for COLA.
@@ -353,7 +353,7 @@ public:
 	 *	@param	UpdateRate						How often to retrieve the data from the spectral analyzer and broadcast the event. Max is 30 times per second.
 	 *	@param  InterpMethod                    Method to used for band peak calculation.
 	 *	@param  SpectrumType                    Metric to use when returning spectrum values.
-	 *	@param  DecibelNoiseFloor               Decibel Noise Floor to consider as silence silence when using a Decibel Spectrum Type.
+	 *	@param  DecibelNoiseFloor               Decibel Noise Floor to consider as silence when using a Decibel Spectrum Type.
 	 *	@param  bDoNormalize                    If true, output band values will be normalized between zero and one.
 	 *	@param  bDoAutoRange                    If true, output band values will have their ranges automatically adjusted to the minimum and maximum values in the audio. Output band values will be normalized between zero and one.
 	 *	@param  AutoRangeAttackTime             The time (in seconds) it takes for the range to expand to 90% of a larger range.
@@ -435,7 +435,7 @@ public:
 	UPROPERTY(EditAnywhere, Category = Soundfield)
 	TArray<TObjectPtr<USoundfieldEffectBase>> SoundfieldEffectChain;
 
-	// Traverses parent submixes until we find a submix that doesn't inherit it's soundfield format.
+	// Traverses parent submixes until we find a submix that doesn't inherit its soundfield format.
 	FName GetSubmixFormat() const;
 
 	UPROPERTY()
@@ -444,7 +444,7 @@ public:
 	// Traverses parent submixes until we find a submix that specifies encoding settings.
 	const USoundfieldEncodingSettingsBase* GetEncodingSettings() const;
 
-	// This function goes through every child submix and the parent submix to ensure that they have a compatible format with this  submix's format.
+	// This function goes through every child submix and the parent submix to ensure that they have a compatible format with this submix's format.
 	void SanitizeLinks();
 
 protected:
diff --git a/Engine/Source/Runtime/Engine/Private/AudioDevice.cpp b/Engine/Source/Runtime/Engine/Private/AudioDevice.cpp
--- a/Engine/Source/Runtime/Engine/Private/AudioDevice.cpp
+++ b/Engine/Source/Runtime/Engine/Private/AudioDevice.cpp
@@ -189,7 +189,7 @@ FAutoConsoleVariableRef CVarSetAudioMaxDistance(
 
 static FAutoConsoleCommandWithWorld GListAvailableSpatialPluginsCommand(
 	TEXT("au.spatialization.ListAvailableSpatialPlugins"),
-	TEXT("This will output a list of currently availible/active spatialization plugins"),
+	TEXT("This will output a list of currently available/active spatialization plugins"),
 	FConsoleCommandWithWorldDelegate::CreateStatic(
 		[](UWorld* InWorld)
 	{
