Adding net.PushModelValidateSkipUpdate for push model CanSkipUpdate testing

diff --git a/Engine/Source/Runtime/Engine/Private/DataChannel.cpp b/Engine/Source/Runtime/Engine/Private/DataChannel.cpp
--- a/Engine/Source/Runtime/Engine/Private/DataChannel.cpp
+++ b/Engine/Source/Runtime/Engine/Private/DataChannel.cpp
@@ -143,6 +143,11 @@ namespace UE::Net
 		TEXT("Enables new SubObjects to set bNetInitial to true to make sure all replicated properties are replicated.")
 	);
 
+	static bool bPushModelValidateSkipUpdate = false;
+	static FAutoConsoleVariableRef CVarNetPushModelValidateSkipUpdate(
+		TEXT("net.PushModelValidateSkipUpdate"),
+		bPushModelValidateSkipUpdate,
+		TEXT("If enabled, detect when we thought we could skip an object replication based on push model state, but we sent data anyway."), ECVF_Default);
 }; // namespace UE::Net
 
 template<typename T>
@@ -3369,10 +3374,14 @@ int64 UActorChannel::ReplicateActor()
 		{
 			UE_NET_TRACE_OBJECT_SCOPE(ActorReplicator->ObjectNetGUID, Bunch, GetTraceCollector(Bunch), ENetTraceVerbosity::Trace);
 
-			if (!ActorReplicator->CanSkipUpdate(RepFlags))
+			const bool bCanSkipUpdate = ActorReplicator->CanSkipUpdate(RepFlags);
+
+			if (UE::Net::bPushModelValidateSkipUpdate || !bCanSkipUpdate)
 			{
 				bWroteSomethingImportant |= ActorReplicator->ReplicateProperties(Bunch, RepFlags);
 			}
+
+			ensureMsgf(!UE::Net::bPushModelValidateSkipUpdate || !bCanSkipUpdate || !bWroteSomethingImportant, TEXT("Actor wrote data but we thought it was skippable: %s"), *GetFullNameSafe(Actor));
 		}
 
 		bWroteSomethingImportant |= DoSubObjectReplication(Bunch, RepFlags);
@@ -4859,7 +4868,7 @@ bool UActorChannel::WriteSubObjectInBunch(UObject* Obj, FOutBunch& Bunch, FRepli
 
 	const bool bCanSkipUpdate = bFoundReplicator && !bNewToReplay && (*FoundReplicator)->CanSkipUpdate(RepFlags);
 
-	if (bCanSkipUpdate)
+	if (!UE::Net::bPushModelValidateSkipUpdate && bCanSkipUpdate)
 	{
 		return false;
 	}
@@ -4904,6 +4913,8 @@ bool UActorChannel::WriteSubObjectInBunch(UObject* Obj, FOutBunch& Bunch, FRepli
 		bWroteSomething = true;
 	}
 
+	ensureMsgf(!UE::Net::bPushModelValidateSkipUpdate || !bCanSkipUpdate || !bWroteSomething, TEXT("Subobject wrote data but we thought it was skippable: %s Actor: %s"), *GetFullNameSafe(Obj), *GetFullNameSafe(Actor));
+
 	return bWroteSomething;
 }
 
